New streaming system for SparseVolumeTextures. It is based on the Nanite FStreamingManager design and supports streaming both from disk and from DDC. This first version uses a single logical tile data texture for all frames of a single UStreamableSparseVolumeTexture and starts streaming out tiles of older mip levels once it runs out of space. The highest mip level (lowest resolution) is always resident in the tile data texture, so there is always something available for rendering. Each frame's page table texture is currently always fully resident. There is now also a buffer storing the lowest resident mip level index for each frame.

In a future version, additional logical tile data textures may be allocated to handle cases where more mip levels are being requested than can physically fit into a single texture. In addition, page table textures should also be resized so that streamed out frames take up less GPU memory. It might also be necessary to implement a blocking streaming option for MRQ and similar use cases. This feature probably depends on being able to spill to additional physical page table textures.

Also moved all of the SVT runtime classes into a shared namespace (UE::SVT), which is why this CL ended up touching almost all SVT related files.

#rnx

diff --git a/Engine/Plugins/FX/Niagara/Shaders/Private/NiagaraDataInterfaceSparseVolumeTextureTemplate.ush b/Engine/Plugins/FX/Niagara/Shaders/Private/NiagaraDataInterfaceSparseVolumeTextureTemplate.ush
--- a/Engine/Plugins/FX/Niagara/Shaders/Private/NiagaraDataInterfaceSparseVolumeTextureTemplate.ush
+++ b/Engine/Plugins/FX/Niagara/Shaders/Private/NiagaraDataInterfaceSparseVolumeTextureTemplate.ush
@@ -2,19 +2,20 @@
 
 #include "/Engine/Private/SparseVolumeTexture/SparseVolumeTextureCommon.ush"
 
-SamplerState	{ParameterName}_TileDataTextureSampler;
-Texture3D<uint>	{ParameterName}_PageTableTexture;
-Texture3D		{ParameterName}_PhysicalTileDataATexture;
-Texture3D		{ParameterName}_PhysicalTileDataBTexture;
-uint4			{ParameterName}_PackedUniforms0;
-uint4			{ParameterName}_PackedUniforms1;
-int3			{ParameterName}_TextureSize;
-int				{ParameterName}_MipLevels;
+SamplerState		{ParameterName}_TileDataTextureSampler;
+Texture3D<uint>		{ParameterName}_PageTableTexture;
+Texture3D			{ParameterName}_PhysicalTileDataATexture;
+Texture3D			{ParameterName}_PhysicalTileDataBTexture;
+ByteAddressBuffer	{ParameterName}_StreamingInfoBuffer;
+uint4				{ParameterName}_PackedUniforms0;
+uint4				{ParameterName}_PackedUniforms1;
+int3				{ParameterName}_TextureSize;
+int					{ParameterName}_MipLevels;
 
 
 void LoadSparseVolumeTexture_{ParameterName}(in int TexelX, in int TexelY, in int TexelZ, in int MipLevel, out float4 OutAttributesA, out float4 OutAttributesB)
 {
-	const FSparseVolumeTextureUniforms Uniforms = SparseVolumeTextureUnpackUniforms({ParameterName}_PackedUniforms0, {ParameterName}_PackedUniforms1);
+	const FSparseVolumeTextureUniforms Uniforms = SparseVolumeTextureUnpackUniforms({ParameterName}_PackedUniforms0, {ParameterName}_PackedUniforms1, {ParameterName}_StreamingInfoBuffer);
 	const int3 VoxelCoord = SparseVolumeTextureLoadPageTable({ParameterName}_PageTableTexture, Uniforms, int3(TexelX, TexelY, TexelZ), MipLevel);
 	OutAttributesA = {ParameterName}_PhysicalTileDataATexture.Load(int4(VoxelCoord, 0));
 	OutAttributesB = {ParameterName}_PhysicalTileDataBTexture.Load(int4(VoxelCoord, 0));
@@ -22,7 +23,7 @@ void LoadSparseVolumeTexture_{ParameterName}(in int TexelX, in int TexelY, in in
 
 void SampleSparseVolumeTexture_{ParameterName}(in float3 UVW, in int MipLevel, out float4 OutAttributesA, out float4 OutAttributesB)
 {
-	const FSparseVolumeTextureUniforms Uniforms = SparseVolumeTextureUnpackUniforms({ParameterName}_PackedUniforms0, {ParameterName}_PackedUniforms1);
+	const FSparseVolumeTextureUniforms Uniforms = SparseVolumeTextureUnpackUniforms({ParameterName}_PackedUniforms0, {ParameterName}_PackedUniforms1, {ParameterName}_StreamingInfoBuffer);
 	const float3 VoxelUVW = SparseVolumeTextureSamplePageTable({ParameterName}_PageTableTexture, Uniforms, UVW, SVTADDRESSMODE_WRAP, SVTADDRESSMODE_WRAP, SVTADDRESSMODE_WRAP, MipLevel);
 	OutAttributesA = SparseVolumeTextureSamplePhysicalTileData({ParameterName}_PhysicalTileDataATexture, {ParameterName}_PhysicalTileDataBTexture, {ParameterName}_TileDataTextureSampler, VoxelUVW, 0);
 	OutAttributesB = SparseVolumeTextureSamplePhysicalTileData({ParameterName}_PhysicalTileDataATexture, {ParameterName}_PhysicalTileDataBTexture, {ParameterName}_TileDataTextureSampler, VoxelUVW, 1);
diff --git a/Engine/Plugins/FX/Niagara/Shaders/Private/NiagaraSVTToDenseBuffer.usf b/Engine/Plugins/FX/Niagara/Shaders/Private/NiagaraSVTToDenseBuffer.usf
--- a/Engine/Plugins/FX/Niagara/Shaders/Private/NiagaraSVTToDenseBuffer.usf
+++ b/Engine/Plugins/FX/Niagara/Shaders/Private/NiagaraSVTToDenseBuffer.usf
@@ -4,21 +4,22 @@
 #include "/Engine/Private/SparseVolumeTexture/SparseVolumeTextureCommon.ush"
 RWTexture3D<float4>	DestinationBuffer;
 
-SamplerState	TileDataTextureSampler;
-Texture3D<uint>	SparseVolumeTexturePageTable;
-Texture3D		SparseVolumeTextureA;
-uint4			PackedSVTUniforms0;
-uint4			PackedSVTUniforms1;
-int3			TextureSize;
-int				MipLevels;
+SamplerState		TileDataTextureSampler;
+Texture3D<uint>		SparseVolumeTexturePageTable;
+Texture3D			SparseVolumeTextureA;
+ByteAddressBuffer	StreamingInfoBuffer;
+uint4				PackedSVTUniforms0;
+uint4				PackedSVTUniforms1;
+int3				TextureSize;
+int					MipLevels;
 
 [numthreads(4, 4, 4)]
 void PerformCopyCS(uint3 Index : SV_DispatchThreadID)
 {
-	
-	float4 SVTSample;
-	const FSparseVolumeTextureUniforms Uniforms = SparseVolumeTextureUnpackUniforms(PackedSVTUniforms0, PackedSVTUniforms1);
-	const int3 VoxelCoord = SparseVolumeTextureLoadPageTable(SparseVolumeTexturePageTable, Uniforms, Index, 0);	
+	const FSparseVolumeTextureUniforms Uniforms = SparseVolumeTextureUnpackUniforms(PackedSVTUniforms0, PackedSVTUniforms1, StreamingInfoBuffer);
+	const int MipLevel = Uniforms.LowestResidentMipLevel; // We want to sample level 0, but we also need to be able to fall back to coarser mips.
+	const int3 VirtualCoord = Index >> MipLevel;
+	const int3 VoxelCoord = SparseVolumeTextureLoadPageTable(SparseVolumeTexturePageTable, Uniforms, VirtualCoord, MipLevel);
 
 	DestinationBuffer[Index] = SparseVolumeTextureA.Load(int4(VoxelCoord, 0));
 }
diff --git a/Engine/Plugins/FX/Niagara/Source/Niagara/Classes/NiagaraDataInterfaceSparseVolumeTexture.h b/Engine/Plugins/FX/Niagara/Source/Niagara/Classes/NiagaraDataInterfaceSparseVolumeTexture.h
--- a/Engine/Plugins/FX/Niagara/Source/Niagara/Classes/NiagaraDataInterfaceSparseVolumeTexture.h
+++ b/Engine/Plugins/FX/Niagara/Source/Niagara/Classes/NiagaraDataInterfaceSparseVolumeTexture.h
@@ -20,6 +20,7 @@ class NIAGARA_API UNiagaraDataInterfaceSparseVolumeTexture : public UNiagaraData
 		SHADER_PARAMETER_TEXTURE(Texture3D<uint>, PageTableTexture)
 		SHADER_PARAMETER_TEXTURE(Texture3D, PhysicalTileDataATexture)
 		SHADER_PARAMETER_TEXTURE(Texture3D, PhysicalTileDataBTexture)
+		SHADER_PARAMETER_SRV(ByteAddressBuffer, StreamingInfoBuffer)
 		SHADER_PARAMETER(FUintVector4, PackedUniforms0)
 		SHADER_PARAMETER(FUintVector4, PackedUniforms1)
 		SHADER_PARAMETER(FIntVector3, TextureSize)
diff --git a/Engine/Plugins/FX/Niagara/Source/Niagara/Private/NiagaraDataInterfaceRenderTargetVolume.cpp b/Engine/Plugins/FX/Niagara/Source/Niagara/Private/NiagaraDataInterfaceRenderTargetVolume.cpp
--- a/Engine/Plugins/FX/Niagara/Source/Niagara/Private/NiagaraDataInterfaceRenderTargetVolume.cpp
+++ b/Engine/Plugins/FX/Niagara/Source/Niagara/Private/NiagaraDataInterfaceRenderTargetVolume.cpp
@@ -24,12 +24,12 @@
 #include "Misc/PathViews.h"
 #include "Misc/Paths.h"
 #include "DataDrivenShaderPlatformInfo.h"
+#include "GlobalRenderResources.h"
 
 #include "RenderGraphUtils.h"
 
 #include "SparseVolumeTexture/SparseVolumeTexture.h"
 #include "SparseVolumeTexture/SparseVolumeTextureData.h"
-#include "SparseVolumeTexture/SparseVolumeTextureSceneProxy.h"
 
 
 #include UE_INLINE_GENERATED_CPP_BY_NAME(NiagaraDataInterfaceRenderTargetVolume)
@@ -709,7 +709,7 @@ bool UNiagaraDataInterfaceRenderTargetVolume::SimCacheWriteFrame(UObject* Storag
 #if WITH_EDITOR
 				UAnimatedSparseVolumeTexture* CurrCache = CastChecked<UAnimatedSparseVolumeTexture>(StorageObject);
 				
-				FSparseVolumeTextureDataCreateInfo SVTCreateInfo;
+				UE::SVT::FTextureDataCreateInfo SVTCreateInfo;
 				SVTCreateInfo.VirtualVolumeAABBMin = FIntVector3::ZeroValue;
 				SVTCreateInfo.VirtualVolumeAABBMax = InstanceData_GT->Size;
 				SVTCreateInfo.FallbackValues[0] = FVector4f(0, 0, 0, 0);
@@ -717,7 +717,7 @@ bool UNiagaraDataInterfaceRenderTargetVolume::SimCacheWriteFrame(UObject* Storag
 				SVTCreateInfo.AttributesFormats[0] = PF_FloatRGBA;
 				SVTCreateInfo.AttributesFormats[1] = PF_Unknown;
 
-				FSparseVolumeTextureData SparseTextureData{};
+				UE::SVT::FTextureData SparseTextureData{};
 				SparseTextureData.CreateFromDense(SVTCreateInfo, MakeArrayView<uint8>((uint8*)TextureData.GetData(), TextureData.Num() * sizeof(TextureData[0])), TArrayView<uint8>());
 				
 				CurrCache->AppendFrame(SparseTextureData);
@@ -826,11 +826,8 @@ bool UNiagaraDataInterfaceRenderTargetVolume::SimCacheReadFrame(UObject* Storage
 	{			
 		UAnimatedSparseVolumeTexture* SVT = Cast<UAnimatedSparseVolumeTexture>(StorageObject);
 		
-		USparseVolumeTextureFrame *SVTFrame = USparseVolumeTextureFrame::GetFrame(SVT, FrameA);
-		
-		FUintVector4 CurrentPackedUniforms0 = FUintVector4();
-		FUintVector4 CurrentPackedUniforms1 = FUintVector4();
-		SVTFrame->GetPackedUniforms(CurrentPackedUniforms0, CurrentPackedUniforms1);
+		const int32 MipLevel = 0;
+		USparseVolumeTextureFrame *SVTFrame = USparseVolumeTextureFrame::GetFrameAndIssueStreamingRequest(SVT, FrameA + Interp, MipLevel);
 
 		FIntVector VolumeResolution = SVT->GetVolumeResolution();						
 
@@ -848,15 +845,18 @@ bool UNiagaraDataInterfaceRenderTargetVolume::SimCacheReadFrame(UObject* Storage
 		ENQUEUE_RENDER_COMMAND(NDIRenderTargetVolumeUpdate)
 			(
 				[RT_Proxy, RT_InstanceID = SystemInstance->GetId(),
-				RT_VolumeResolution = VolumeResolution, RT_SVTProxy = SVTFrame->GetSparseVolumeTextureSceneProxy(),
-				FeatureLevel = SystemInstance->GetFeatureLevel(),
-				CurrentPackedUniforms0, CurrentPackedUniforms1](FRHICommandListImmediate& RHICmdList)
+				RT_VolumeResolution = VolumeResolution, RT_SVTRenderResources = SVTFrame->GetTextureRenderResources(),
+				FeatureLevel = SystemInstance->GetFeatureLevel()](FRHICommandListImmediate& RHICmdList)
 				{
-					if (RT_SVTProxy == nullptr)
+					if (RT_SVTRenderResources == nullptr)
 					{
 						return;
 					}
 
+					FUintVector4 CurrentPackedUniforms0 = FUintVector4();
+					FUintVector4 CurrentPackedUniforms1 = FUintVector4();
+					RT_SVTRenderResources->GetPackedUniforms(CurrentPackedUniforms0, CurrentPackedUniforms1);
+
 					if (FRenderTargetVolumeRWInstanceData_RenderThread * InstanceData_RT = RT_Proxy->SystemInstancesToProxyData_RT.Find(RT_InstanceID))
 					{
 						FRDGBuilder GraphBuilder(RHICmdList);
@@ -896,12 +896,14 @@ bool UNiagaraDataInterfaceRenderTargetVolume::SimCacheReadFrame(UObject* Storage
 							);
 						}						
 
-						FRHITexture* PageTableTexture = RT_SVTProxy->GetPageTableTextureRHI();
-						FRHITexture* TextureA = RT_SVTProxy->GetPhysicalTileDataATextureRHI();						
+						FRHITexture* PageTableTexture = RT_SVTRenderResources->GetPageTableTextureRHI();
+						FRHITexture* TextureA = RT_SVTRenderResources->GetPhysicalTileDataATextureRHI();
+						FRHIShaderResourceView* StreamingInfoBufferSRV = RT_SVTRenderResources->GetStreamingInfoBufferSRVRHI();
 
 						PassParameters->TileDataTextureSampler = TStaticSamplerState<SF_Point, AM_Clamp, AM_Clamp, AM_Clamp>::GetRHI();
-						PassParameters->SparseVolumeTexturePageTable = PageTableTexture;
-						PassParameters->SparseVolumeTextureA = TextureA;
+						PassParameters->SparseVolumeTexturePageTable = PageTableTexture ? PageTableTexture : GBlackUintVolumeTexture->TextureRHI.GetReference();
+						PassParameters->SparseVolumeTextureA = TextureA ? TextureA : GBlackVolumeTexture->TextureRHI.GetReference();
+						PassParameters->StreamingInfoBuffer = StreamingInfoBufferSRV ? StreamingInfoBufferSRV : GEmptyVertexBufferWithUAV->ShaderResourceViewRHI.GetReference();
 							
 						PassParameters->PackedSVTUniforms0 = CurrentPackedUniforms0;
 						PassParameters->PackedSVTUniforms1 = CurrentPackedUniforms1;
diff --git a/Engine/Plugins/FX/Niagara/Source/Niagara/Private/NiagaraDataInterfaceSparseVolumeTexture.cpp b/Engine/Plugins/FX/Niagara/Source/Niagara/Private/NiagaraDataInterfaceSparseVolumeTexture.cpp
--- a/Engine/Plugins/FX/Niagara/Source/Niagara/Private/NiagaraDataInterfaceSparseVolumeTexture.cpp
+++ b/Engine/Plugins/FX/Niagara/Source/Niagara/Private/NiagaraDataInterfaceSparseVolumeTexture.cpp
@@ -2,7 +2,6 @@
 
 #include "NiagaraDataInterfaceSparseVolumeTexture.h"
 #include "SparseVolumeTexture/SparseVolumeTexture.h"
-#include "SparseVolumeTexture/SparseVolumeTextureSceneProxy.h"
 #include "NiagaraComputeExecutionContext.h"
 #include "NiagaraCustomVersion.h"
 #include "NiagaraGpuComputeDispatchInterface.h"
@@ -37,8 +36,7 @@ struct FNDISparseVolumeTextureFunctionVersion
 struct FNDISparseVolumeTextureInstanceData_GameThread
 {
 	TWeakObjectPtr<USparseVolumeTexture> CurrentTexture = nullptr;
-	FUintVector4 CurrentPackedUniforms0 = FUintVector4();
-	FUintVector4 CurrentPackedUniforms1 = FUintVector4();
+	const UE::SVT::FTextureRenderResources* CurrentRenderResources = nullptr;
 	FIntVector3 CurrentTextureSize = FIntVector3::ZeroValue;
 	int32 CurrentTextureMipLevels = 0;
 	FNiagaraParameterDirectBinding<UObject*> UserParamBinding;
@@ -46,11 +44,7 @@ struct FNDISparseVolumeTextureInstanceData_GameThread
 
 struct FNDISparseVolumeTextureInstanceData_RenderThread
 {
-	FTextureRHIRef PageTableTextureRHI = nullptr;
-	FTextureRHIRef PhysicalTileDataATextureRHI = nullptr;
-	FTextureRHIRef PhysicalTileDataBTextureRHI = nullptr;
-	FUintVector4 PackedUniforms0 = FUintVector4();
-	FUintVector4 PackedUniforms1 = FUintVector4();
+	const UE::SVT::FTextureRenderResources* RenderResources = nullptr;
 	FIntVector3 TextureSize = FIntVector3::ZeroValue;
 	int32 MipLevels = 0;
 };
@@ -189,60 +183,39 @@ bool UNiagaraDataInterfaceSparseVolumeTexture::PerInstanceTick(void* PerInstance
 	FNDISparseVolumeTextureInstanceData_GameThread* InstanceData = static_cast<FNDISparseVolumeTextureInstanceData_GameThread*>(PerInstanceData);
 
 	USparseVolumeTexture* CurrentTexture = InstanceData->UserParamBinding.GetValueOrDefault<USparseVolumeTexture>(SparseVolumeTexture);
-	FUintVector4 CurrentPackedUniforms0 = FUintVector4();
-	FUintVector4 CurrentPackedUniforms1 = FUintVector4();
+	const UE::SVT::FTextureRenderResources* CurrentRenderResources = nullptr;
 	FIntVector3 CurrentTextureSize = FIntVector3::ZeroValue;
 	int32 CurrentTextureMipLevels = 0;
 	if (CurrentTexture)
 	{
-		CurrentTexture->GetPackedUniforms(CurrentPackedUniforms0, CurrentPackedUniforms1);
+		CurrentRenderResources = CurrentTexture->GetTextureRenderResources();
 		CurrentTextureSize = FIntVector3(CurrentTexture->GetVolumeResolution());
 		CurrentTextureMipLevels = CurrentTexture->GetNumMipLevels();
 	}
 
 	if ((InstanceData->CurrentTexture != CurrentTexture) 
+		|| (InstanceData->CurrentRenderResources != CurrentRenderResources)
 		|| (InstanceData->CurrentTextureSize != CurrentTextureSize) 
-		|| (InstanceData->CurrentTextureMipLevels != CurrentTextureMipLevels)
-		|| (InstanceData->CurrentPackedUniforms0 != CurrentPackedUniforms0)
-		|| (InstanceData->CurrentPackedUniforms1 != CurrentPackedUniforms1))
+		|| (InstanceData->CurrentTextureMipLevels != CurrentTextureMipLevels))
 	{
 		InstanceData->CurrentTexture = CurrentTexture;
+		InstanceData->CurrentRenderResources = CurrentRenderResources;
 		InstanceData->CurrentTextureSize = CurrentTextureSize;
 		InstanceData->CurrentTextureMipLevels = CurrentTextureMipLevels;
-		InstanceData->CurrentPackedUniforms0 = CurrentPackedUniforms0;
-		InstanceData->CurrentPackedUniforms1 = CurrentPackedUniforms1;
 
 		ENQUEUE_RENDER_COMMAND(NDISparseVolumeTexture_UpdateInstance)
 		(
-			[RT_Proxy = GetProxyAs<FNiagaraDataInterfaceProxySparseVolumeTexture>(), RT_InstanceID = SystemInstance->GetId(), RT_Texture = CurrentTexture, RT_Packed0 = CurrentPackedUniforms0, 
-			RT_Packed1 = CurrentPackedUniforms1, RT_TextureSize = CurrentTextureSize, RT_MipLevels = CurrentTextureMipLevels] (FRHICommandListImmediate&)
+			[RT_Proxy = GetProxyAs<FNiagaraDataInterfaceProxySparseVolumeTexture>(),
+			RT_InstanceID = SystemInstance->GetId(),
+			RT_RenderResources = CurrentRenderResources,
+			RT_TextureSize = CurrentTextureSize, 
+			RT_MipLevels = CurrentTextureMipLevels] 
+		(FRHICommandListImmediate&)
 			{
 				FNDISparseVolumeTextureInstanceData_RenderThread& RTInstanceData = RT_Proxy->InstanceData_RT.FindOrAdd(RT_InstanceID);
-				const FSparseVolumeTextureSceneProxy* SceneProxy = RT_Texture ? RT_Texture->GetSparseVolumeTextureSceneProxy() : nullptr;
-				if (RT_Texture && SceneProxy)
-				{
-					FRHITexture* PageTableTexture = SceneProxy->GetPageTableTextureRHI();
-					FRHITexture* TextureA = SceneProxy->GetPhysicalTileDataATextureRHI();
-					FRHITexture* TextureB = SceneProxy->GetPhysicalTileDataBTextureRHI();
-
-					RTInstanceData.PageTableTextureRHI = PageTableTexture;
-					RTInstanceData.PhysicalTileDataATextureRHI = TextureA;
-					RTInstanceData.PhysicalTileDataBTextureRHI = TextureB;
-					RTInstanceData.PackedUniforms0 = RT_Packed0;
-					RTInstanceData.PackedUniforms1 = RT_Packed1;
-					RTInstanceData.TextureSize = RT_TextureSize;
-					RTInstanceData.MipLevels = RT_MipLevels;
-				}
-				else
-				{
-					RTInstanceData.PageTableTextureRHI = nullptr;
-					RTInstanceData.PhysicalTileDataATextureRHI = nullptr;
-					RTInstanceData.PhysicalTileDataBTextureRHI = nullptr;
-					RTInstanceData.PackedUniforms0 = FUintVector4();
-					RTInstanceData.PackedUniforms1 = FUintVector4();
-					RTInstanceData.TextureSize = FIntVector3::ZeroValue;
-					RTInstanceData.MipLevels = 0;
-				}
+				RTInstanceData.RenderResources = RT_RenderResources;
+				RTInstanceData.TextureSize = RT_TextureSize;
+				RTInstanceData.MipLevels = RT_MipLevels;
 			}
 		);
 	}
@@ -309,26 +282,30 @@ void UNiagaraDataInterfaceSparseVolumeTexture::SetShaderParameters(const FNiagar
 
 	FShaderParameters* Parameters = Context.GetParameterNestedStruct<FShaderParameters>();
 	Parameters->TileDataTextureSampler = TStaticSamplerState<SF_Bilinear, AM_Clamp, AM_Clamp, AM_Clamp>::GetRHI();
-	if (RTInstanceData && RTInstanceData->PageTableTextureRHI.IsValid() && (RTInstanceData->PhysicalTileDataATextureRHI.IsValid() || RTInstanceData->PhysicalTileDataBTextureRHI.IsValid()))
+	Parameters->PageTableTexture = GBlackUintVolumeTexture->TextureRHI;
+	Parameters->PhysicalTileDataATexture = GBlackVolumeTexture->TextureRHI;
+	Parameters->PhysicalTileDataBTexture = GBlackVolumeTexture->TextureRHI;
+	Parameters->StreamingInfoBuffer = GEmptyVertexBufferWithUAV->ShaderResourceViewRHI;
+	Parameters->PackedUniforms0 = FUintVector4();
+	Parameters->PackedUniforms1 = FUintVector4();
+	Parameters->TextureSize = FIntVector3::ZeroValue;
+	Parameters->MipLevels = 0;
+	
+	if (RTInstanceData && RTInstanceData->RenderResources)
 	{
-		Parameters->PageTableTexture = RTInstanceData->PageTableTextureRHI;
-		Parameters->PhysicalTileDataATexture = RTInstanceData->PhysicalTileDataATextureRHI.IsValid() ? RTInstanceData->PhysicalTileDataATextureRHI : GBlackVolumeTexture->TextureRHI;
-		Parameters->PhysicalTileDataBTexture = RTInstanceData->PhysicalTileDataBTextureRHI.IsValid() ? RTInstanceData->PhysicalTileDataBTextureRHI : GBlackVolumeTexture->TextureRHI;
-		Parameters->PackedUniforms0 = RTInstanceData->PackedUniforms0;
-		Parameters->PackedUniforms1 = RTInstanceData->PackedUniforms1;
+		FTextureRHIRef PageTableTexture = RTInstanceData->RenderResources->GetPageTableTextureRHI();
+		FTextureRHIRef PhysicalTileDataATexture = RTInstanceData->RenderResources->GetPhysicalTileDataATextureRHI();
+		FTextureRHIRef PhysicalTileDataBTexture = RTInstanceData->RenderResources->GetPhysicalTileDataBTextureRHI();
+		FShaderResourceViewRHIRef StreamingInfoBufferSRV = RTInstanceData->RenderResources->GetStreamingInfoBufferSRVRHI();
+
+		Parameters->PageTableTexture = PageTableTexture ? PageTableTexture.GetReference() : Parameters->PageTableTexture;
+		Parameters->PhysicalTileDataATexture = PhysicalTileDataATexture ? PhysicalTileDataATexture.GetReference() : Parameters->PhysicalTileDataATexture;
+		Parameters->PhysicalTileDataBTexture = PhysicalTileDataBTexture ? PhysicalTileDataBTexture.GetReference() : Parameters->PhysicalTileDataBTexture;
+		Parameters->StreamingInfoBuffer = StreamingInfoBufferSRV ? StreamingInfoBufferSRV.GetReference() : Parameters->StreamingInfoBuffer;
+		RTInstanceData->RenderResources->GetPackedUniforms(Parameters->PackedUniforms0, Parameters->PackedUniforms1);
 		Parameters->TextureSize = RTInstanceData->TextureSize;
 		Parameters->MipLevels = RTInstanceData->MipLevels;
 	}
-	else
-	{
-		Parameters->PageTableTexture = GBlackUintVolumeTexture->TextureRHI;
-		Parameters->PhysicalTileDataATexture = GBlackVolumeTexture->TextureRHI;
-		Parameters->PhysicalTileDataBTexture = GBlackVolumeTexture->TextureRHI;
-		Parameters->PackedUniforms0 = FUintVector4();
-		Parameters->PackedUniforms1 = FUintVector4();
-		Parameters->TextureSize = FIntVector3::ZeroValue;
-		Parameters->MipLevels = 0;
-	}
 }
 
 void UNiagaraDataInterfaceSparseVolumeTexture::VMGetTextureDimensions(FVectorVMExternalFunctionContext& Context)
diff --git a/Engine/Plugins/FX/Niagara/Source/NiagaraShader/Public/NiagaraSVTShaders.h b/Engine/Plugins/FX/Niagara/Source/NiagaraShader/Public/NiagaraSVTShaders.h
--- a/Engine/Plugins/FX/Niagara/Source/NiagaraShader/Public/NiagaraSVTShaders.h
+++ b/Engine/Plugins/FX/Niagara/Source/NiagaraShader/Public/NiagaraSVTShaders.h
@@ -16,7 +16,8 @@ class FNiagaraCopySVTToDenseBufferCS : public FGlobalShader
 		SHADER_PARAMETER_RDG_TEXTURE_UAV(RWTexture3D<float4>, DestinationBuffer)
 		SHADER_PARAMETER_SAMPLER(SamplerState, TileDataTextureSampler)
 		SHADER_PARAMETER_TEXTURE(Texture3D<uint>, SparseVolumeTexturePageTable)
-		SHADER_PARAMETER_TEXTURE(Texture3D, SparseVolumeTextureA)		
+		SHADER_PARAMETER_TEXTURE(Texture3D, SparseVolumeTextureA)	
+		SHADER_PARAMETER_SRV(ByteAddressBuffer, StreamingInfoBuffer)
 		SHADER_PARAMETER(FUintVector4, PackedSVTUniforms0)
 		SHADER_PARAMETER(FUintVector4, PackedSVTUniforms1)
 		SHADER_PARAMETER(FIntVector, TextureSize)
diff --git a/Engine/Shaders/Private/SparseVolumeTexture/SparseVolumeTextureCommon.ush b/Engine/Shaders/Private/SparseVolumeTexture/SparseVolumeTextureCommon.ush
--- a/Engine/Shaders/Private/SparseVolumeTexture/SparseVolumeTextureCommon.ush
+++ b/Engine/Shaders/Private/SparseVolumeTexture/SparseVolumeTextureCommon.ush
@@ -2,6 +2,11 @@
 
 #pragma once
 
+// SVT_TODO: Unify with macros in SparseVolumeTexture.h
+#define SPARSE_VOLUME_TILE_RES 16
+#define SPARSE_VOLUME_TILE_BORDER 1
+#define SPARSE_VOLUME_TILE_RES_PADDED (SPARSE_VOLUME_TILE_RES + 2 * SPARSE_VOLUME_TILE_BORDER)
+
 #define SVTADDRESSMODE_CLAMP 0u
 #define SVTADDRESSMODE_WRAP 1u
 #define SVTADDRESSMODE_MIRROR 2u
@@ -41,37 +46,34 @@ struct FSparseVolumeTextureUniforms
 	float3			VolumePageResolution;
 	float3			PageTableOffset;
 	float3			TileDataTexelSize;
-	int             MinMipLevel;
-	int             MaxMipLevel;
-	float			TileSize;
-	float			BorderSize;
+	int				FrameIndex;
+	int             HighestMipLevel;
+	int				LowestResidentMipLevel;
 };
 
-FSparseVolumeTextureUniforms SparseVolumeTextureUnpackUniforms(const uint4 Packed0, const uint4 Packed1)
+FSparseVolumeTextureUniforms SparseVolumeTextureUnpackUniforms(const uint4 Packed0, const uint4 Packed1, ByteAddressBuffer StreamingInfoBuffer)
 {
 	FSparseVolumeTextureUniforms Result;
 	Result.VolumePageResolution = asfloat(Packed0.xyz);
 	Result.PageTableOffset.x = float(Packed0.w & 0x7FFu);
-	Result.PageTableOffset.y = float((Packed0.w >> 11) & 0x7FF);
-	Result.PageTableOffset.z = float((Packed0.w >> 22) & 0x3FF);
+	Result.PageTableOffset.y = float((Packed0.w >> 11u) & 0x7FFu);
+	Result.PageTableOffset.z = float((Packed0.w >> 22u) & 0x3FFu);
 	Result.TileDataTexelSize = asfloat(Packed1.xyz);
-	Result.MinMipLevel = int(Packed1.w & 0xFFu);
-	Result.MaxMipLevel = int((Packed1.w >> 8u) & 0xFFu);
-	Result.TileSize = float((Packed1.w >> 16u) & 0xFFu);
-	Result.BorderSize = float((Packed1.w >> 24u) & 0xFFu);
+	Result.FrameIndex = int(Packed1.w & 0xFFFFu);
+	Result.HighestMipLevel = int((Packed1.w >> 16u) & 0xFFFFu);
+	Result.LowestResidentMipLevel = asint(StreamingInfoBuffer.Load(Result.FrameIndex << 2));
 
 	return Result;
 }
 
-float3 SparseVolumeTextureGetVoxelCoord(const uint PackedPhysicalTileCoord, const float TileSize, const float BorderSize, const float3 PageTableCoord)
+float3 SparseVolumeTextureGetVoxelCoord(const uint PackedPhysicalTileCoord, const float3 PageTableCoord)
 {
 	const int3 PhysicalTileCoord = int3(
-		PackedPhysicalTileCoord & 0x7FF,
-		(PackedPhysicalTileCoord >> 11) & 0x7FF,
-		(PackedPhysicalTileCoord >> 22) & 0x3FF);
+		PackedPhysicalTileCoord & 0xFF,
+		(PackedPhysicalTileCoord >> 8) & 0xFF,
+		(PackedPhysicalTileCoord >> 16) & 0xFF);
 
-	const float PaddedTileSize = BorderSize * 2.0f + TileSize;
-	const float3 VoxelCoord = float3(PhysicalTileCoord) * PaddedTileSize + (frac(PageTableCoord) * TileSize + BorderSize);
+	const float3 VoxelCoord = float3(PhysicalTileCoord) * float(SPARSE_VOLUME_TILE_RES_PADDED) + (frac(PageTableCoord) * float(SPARSE_VOLUME_TILE_RES) + float(SPARSE_VOLUME_TILE_BORDER));
 	return VoxelCoord;
 }
 
@@ -79,34 +81,34 @@ float3 SparseVolumeTextureSamplePageTable(Texture3D<uint> PageTable, FSparseVolu
 {
 	// Apply address mode to UVW and clamp mip level to resident levels
 	UVW = SparseVolumeTextureApplyAddressMode(UVW, AddressU, AddressV, AddressW);
-	MipLevel = clamp(MipLevel, Uniforms.MinMipLevel, Uniforms.MaxMipLevel);
+	MipLevel = clamp(MipLevel, Uniforms.LowestResidentMipLevel, Uniforms.HighestMipLevel);
 
 	const float RcpMipLevelFactor = rcp(float(1u << (uint)MipLevel));
 	const float3 VolumePageCoord = UVW * Uniforms.VolumePageResolution;
 	const float3 MipPageTableOffset = floor(Uniforms.PageTableOffset * RcpMipLevelFactor);
 	const float3 PageTableCoord = VolumePageCoord * RcpMipLevelFactor - MipPageTableOffset;
 
-	const uint PackedPhysicalTileCoord = PageTable.Load(int4(int3(PageTableCoord), (MipLevel - Uniforms.MinMipLevel))).x;
+	const uint PackedPhysicalTileCoord = PageTable.Load(int4(int3(PageTableCoord), MipLevel)).x;
 
-	const float3 VoxelCoord = SparseVolumeTextureGetVoxelCoord(PackedPhysicalTileCoord, Uniforms.TileSize, Uniforms.BorderSize, PageTableCoord);
+	const float3 VoxelCoord = SparseVolumeTextureGetVoxelCoord(PackedPhysicalTileCoord, PageTableCoord);
 	const float3 VoxelUVW = VoxelCoord * Uniforms.TileDataTexelSize;
 	return VoxelUVW;
 }
 
 int3 SparseVolumeTextureLoadPageTable(Texture3D<uint> PageTable, FSparseVolumeTextureUniforms Uniforms, int3 TexelCoord, int MipLevel = 0)
 {
-	if (MipLevel < Uniforms.MinMipLevel || MipLevel > Uniforms.MaxMipLevel)
+	if (MipLevel < Uniforms.LowestResidentMipLevel || MipLevel > Uniforms.HighestMipLevel)
 	{
 		return 0; // Point to null tile
 	}
 	const float RcpMipLevelFactor = rcp(float(1u << (uint)MipLevel));
-	const float3 VolumeMipPageCoord = (TexelCoord + 0.5f) / Uniforms.TileSize;
+	const float3 VolumeMipPageCoord = (TexelCoord + 0.5f) / float(SPARSE_VOLUME_TILE_RES);
 	const float3 MipPageTableOffset = floor(Uniforms.PageTableOffset * RcpMipLevelFactor);
 	const float3 PageTableCoord = VolumeMipPageCoord - MipPageTableOffset;
 
-	const uint PackedPhysicalTileCoord = PageTable.Load(int4(int3(PageTableCoord), (MipLevel - Uniforms.MinMipLevel))).x;
+	const uint PackedPhysicalTileCoord = PageTable.Load(int4(int3(PageTableCoord), MipLevel)).x;
 
-	const float3 VoxelCoord = SparseVolumeTextureGetVoxelCoord(PackedPhysicalTileCoord, Uniforms.TileSize, Uniforms.BorderSize, PageTableCoord);
+	const float3 VoxelCoord = SparseVolumeTextureGetVoxelCoord(PackedPhysicalTileCoord, PageTableCoord);
 	return int3(VoxelCoord);
 }
 
diff --git a/Engine/Shaders/Private/SparseVolumeTexture/UpdateSparseVolumeTexture.usf b/Engine/Shaders/Private/SparseVolumeTexture/UpdateSparseVolumeTexture.usf
new file mode 100644
--- /dev/null
+++ b/Engine/Shaders/Private/SparseVolumeTexture/UpdateSparseVolumeTexture.usf
@@ -0,0 +1,116 @@
+// Copyright Epic Games, Inc. All Rights Reserved.
+
+#pragma once
+
+#include "/Engine/Public/Platform.ush" 
+#include "SparseVolumeTextureCommon.ush"
+
+uint FlattenCoord(uint3 Coord, uint3 VolumeSize)
+{
+	return Coord.z * (VolumeSize.y * VolumeSize.x) + Coord.y * VolumeSize.x + Coord.x;
+}
+
+uint3 UnflattenCoord(uint Index, uint3 VolumeSize)
+{
+	const uint SliceIndex = Index / (VolumeSize.y * VolumeSize.x);
+	const uint SliceLocalIndex = Index - SliceIndex * (VolumeSize.y * VolumeSize.x);
+
+	uint3 Result;
+	Result.z = SliceIndex;
+	Result.y = SliceLocalIndex / VolumeSize.x;
+	Result.x = SliceLocalIndex - (Result.y * VolumeSize.x);
+
+	return Result;
+}
+
+#ifdef UPDATE_TILE_TEXTURE_FROM_BUFFER
+
+RWTexture3D<float4> DstPhysicalTileTextureA;
+RWTexture3D<float4> DstPhysicalTileTextureB;
+ByteAddressBuffer DstTileCoords;
+Buffer<float4> SrcPhysicalTileBufferA;
+Buffer<float4> SrcPhysicalTileBufferB;
+uint TileCoordsBufferOffset;
+uint TileDataBufferOffsetInTiles;
+uint NumTilesToCopy;
+uint NumDispatchedGroups;
+uint PaddedTileSize;
+uint bCopyTexureAOnlyUI;
+
+[numthreads(64, 1, 1)]
+void SparseVolumeTextureUpdateFromBufferCS(uint3 GroupId : SV_GroupID, uint GroupThreadIndex : SV_GroupIndex)
+{
+	const uint GroupIdx = GroupId.x;
+	const bool bCopyTextureAOnly = (bCopyTexureAOnlyUI != 0);
+	const uint NumVoxelsPerPaddedTile = PaddedTileSize * PaddedTileSize * PaddedTileSize;
+
+	// Process one tile per workgroup
+	for (uint TileIdx = GroupIdx; TileIdx < NumTilesToCopy; TileIdx += NumDispatchedGroups)
+	{
+		const uint PackedDstCoord = DstTileCoords.Load((TileIdx + TileCoordsBufferOffset) << 2u);
+		const uint3 DstBaseCoord = uint3(PackedDstCoord & 0xFFu, (PackedDstCoord >> 8u) & 0xFFu, (PackedDstCoord >> 16u) & 0xFFu) * PaddedTileSize;
+		const uint TileDataBufferBaseIndex = (TileIdx + TileDataBufferOffsetInTiles) * NumVoxelsPerPaddedTile;
+
+		// Iterate over all texels in the tile and copy them
+		for (uint VoxelIdx = GroupThreadIndex; VoxelIdx < NumVoxelsPerPaddedTile; VoxelIdx += 64)
+		{
+			const uint LoadIdx = TileDataBufferBaseIndex + VoxelIdx;
+			const uint3 DstCoord = DstBaseCoord + UnflattenCoord(VoxelIdx, PaddedTileSize);
+			DstPhysicalTileTextureA[DstCoord] = SrcPhysicalTileBufferA.Load(LoadIdx);
+			if (!bCopyTextureAOnly)
+			{
+				DstPhysicalTileTextureB[DstCoord] = SrcPhysicalTileBufferB.Load(LoadIdx);
+			}
+		}
+	}
+}
+
+#endif // UPDATE_TILE_TEXTURE_FROM_BUFFER
+
+#ifdef UPDATE_PAGE_TABLE
+
+RWTexture3D<uint> PageTable;
+ByteAddressBuffer PageTableUpdates;
+uint UpdateCoordOffset;
+uint UpdatePayloadOffset;
+uint NumUpdates;
+
+[numthreads(64, 1, 1)]
+void SparseVolumeTextureUpdatePageTableCS(uint3 DispatchThreadID : SV_DispatchThreadID)
+{
+	if (DispatchThreadID.x < NumUpdates)
+	{
+		// The page table coordinate indexes into the page table itself and needs to be able to address each individual voxel.
+		// As the maximum dimension of 3D textures is 2048, we are using a 11|11|10 bit packing scheme for the page table coordinate.
+		const uint PackedPageTableCoord = PageTableUpdates.Load((DispatchThreadID.x + UpdateCoordOffset) << 2);
+		const uint3 PageTableCoord = uint3(PackedPageTableCoord & 0x7FFu, (PackedPageTableCoord >> 11u) & 0x7FFu, (PackedPageTableCoord >> 22u) & 0x3FFu);
+
+		// The payload is the data that is written to the page table. It stores the coordinate of the pointed to tile
+		// in the physical tile data texture and is encoded as XYZ with 8 bit per component, taking up 24 bits and leaving 8 bit for other uses.
+		// 8 bit is enough for this as we only need to address entire tiles, not individual voxels of the physical tile data texture.
+		const uint Payload = PageTableUpdates.Load((DispatchThreadID.x + UpdatePayloadOffset) << 2);
+		
+		PageTable[PageTableCoord] = Payload;
+	}
+}
+
+#endif // UPDATE_PAGE_TABLE
+
+#ifdef UPDATE_STREAMING_INFO_BUFFER
+
+RWByteAddressBuffer StreamingInfoBuffer;
+ByteAddressBuffer StreamingInfoBufferUpdates;
+uint UpdateOffset;
+uint NumUpdates;
+
+[numthreads(64, 1, 1)]
+void SparseVolumeTextureUpdateStreamingInfoBufferCS(uint3 DispatchThreadID : SV_DispatchThreadID)
+{
+	if (DispatchThreadID.x < NumUpdates)
+	{
+		uint2 Update = StreamingInfoBufferUpdates.Load2(((DispatchThreadID.x + UpdateOffset) * 2) << 2);
+		StreamingInfoBuffer.Store(Update.x << 2, Update.y);
+	}
+}
+
+#endif // UPDATE_STREAMING_INFO_BUFFER
diff --git a/Engine/Shaders/Private/SparseVolumeTexture/VisualizeSparseVolumeTexture.usf b/Engine/Shaders/Private/SparseVolumeTexture/VisualizeSparseVolumeTexture.usf
--- a/Engine/Shaders/Private/SparseVolumeTexture/VisualizeSparseVolumeTexture.usf
+++ b/Engine/Shaders/Private/SparseVolumeTexture/VisualizeSparseVolumeTexture.usf
@@ -45,6 +45,7 @@ SamplerState TileDataTextureSampler;
 Texture3D<uint> SparseVolumeTexturePageTable;
 Texture3D<float4> SparseVolumeTextureA;
 Texture3D<float4> SparseVolumeTextureB;
+ByteAddressBuffer StreamingInfoBuffer;
 uint4 PackedSVTUniforms0;
 uint4 PackedSVTUniforms1;
 float3 SparseVolumeTextureResolution;
@@ -131,7 +132,7 @@ void VisualizeSparseVolumeTexturePS(
 		tMax.y *= tDelta.y;
 		tMax.z *= tDelta.z;
 
-		const FSparseVolumeTextureUniforms SVTUniforms = SparseVolumeTextureUnpackUniforms(PackedSVTUniforms0, PackedSVTUniforms1);
+		const FSparseVolumeTextureUniforms SVTUniforms = SparseVolumeTextureUnpackUniforms(PackedSVTUniforms0, PackedSVTUniforms1, StreamingInfoBuffer);
 
 		float StepLength = 1.0f;
 
@@ -140,7 +141,7 @@ void VisualizeSparseVolumeTexturePS(
 		{
 			const float3 UVW = P * RcpVolumeResolution;
 			const float3 VoxelUVW = SparseVolumeTextureSamplePageTable(SparseVolumeTexturePageTable, SVTUniforms, UVW, SVTADDRESSMODE_CLAMP, SVTADDRESSMODE_CLAMP, SVTADDRESSMODE_CLAMP, MipLevel);
-			const float3 UVWTileSize = (SVTUniforms.TileSize + 2.0f * SVTUniforms.BorderSize) * SVTUniforms.TileDataTexelSize;
+			const float3 UVWTileSize = float(SPARSE_VOLUME_TILE_RES_PADDED) * SVTUniforms.TileDataTexelSize;
 			if (any(VoxelUVW >= UVWTileSize)) // skip on null tile
 			{
 				const int PhysicalTileDataIndex = ComponentToVisualize < 4 ? 0 : 1;
diff --git a/Engine/Source/Editor/SparseVolumeTexture/Private/SparseVolumeTextureFactory.cpp b/Engine/Source/Editor/SparseVolumeTexture/Private/SparseVolumeTextureFactory.cpp
--- a/Engine/Source/Editor/SparseVolumeTexture/Private/SparseVolumeTextureFactory.cpp
+++ b/Engine/Source/Editor/SparseVolumeTexture/Private/SparseVolumeTextureFactory.cpp
@@ -498,7 +498,7 @@ UObject* USparseVolumeTextureFactory::FactoryCreateFile(UClass* InClass, UObject
 
 		ExpandVolumeBounds(ImportOptions, PreviewData.GridInfo, VolumeBoundsMin, VolumeBoundsMax);
 
-		FSparseVolumeTextureData TextureData{};
+		UE::SVT::FTextureData TextureData{};
 		const bool bConversionSuccess = ConvertOpenVDBToSparseVolumeTexture(PreviewData.LoadedFile, ImportOptions, VolumeBoundsMin, TextureData);
 
 		if (!bConversionSuccess)
@@ -508,7 +508,7 @@ UObject* USparseVolumeTextureFactory::FactoryCreateFile(UClass* InClass, UObject
 		}
 
 		UStaticSparseVolumeTexture* StaticSVTexture = NewObject<UStaticSparseVolumeTexture>(InParent, UStaticSparseVolumeTexture::StaticClass(), InName, Flags);
-		const bool bInitSuccess = StaticSVTexture->Initialize(MakeArrayView<FSparseVolumeTextureData>(&TextureData, 1));
+		const bool bInitSuccess = StaticSVTexture->Initialize(MakeArrayView<UE::SVT::FTextureData>(&TextureData, 1));
 		if (!bInitSuccess)
 		{
 			UE_LOG(LogSparseVolumeTextureFactory, Error, TEXT("Failed to initialize SparseVolumeTexture: %s"), *Filename);
@@ -537,7 +537,7 @@ UObject* USparseVolumeTextureFactory::FactoryCreateFile(UClass* InClass, UObject
 		ImportTask.MakeDialog(true);
 
 		// Allocate space for each frame
-		TArray<FSparseVolumeTextureData> UncookedFramesData;
+		TArray<UE::SVT::FTextureData> UncookedFramesData;
 		UncookedFramesData.SetNum(NumFrames);
 
 		std::atomic_bool bErrored = false;
@@ -659,7 +659,7 @@ UObject* USparseVolumeTextureFactory::FactoryCreateFile(UClass* InClass, UObject
 						return;
 					}
 
-					FSparseVolumeTextureData TextureData{};
+					UE::SVT::FTextureData TextureData{};
 					const bool bConversionSuccess = ConvertOpenVDBToSparseVolumeTexture(LoadedFrameFile, ImportOptions, VolumeBoundsMin, TextureData);
 
 					if (!bConversionSuccess)
diff --git a/Engine/Source/Editor/SparseVolumeTexture/Private/SparseVolumeTextureOpenVDBUtility.cpp b/Engine/Source/Editor/SparseVolumeTexture/Private/SparseVolumeTextureOpenVDBUtility.cpp
--- a/Engine/Source/Editor/SparseVolumeTexture/Private/SparseVolumeTextureOpenVDBUtility.cpp
+++ b/Engine/Source/Editor/SparseVolumeTexture/Private/SparseVolumeTextureOpenVDBUtility.cpp
@@ -281,7 +281,7 @@ static EPixelFormat GetMultiComponentFormat(ESparseVolumeAttributesFormat Format
 
 #if OPENVDB_AVAILABLE
 
-class FSparseVolumeTextureDataProviderOpenVDB : public ISparseVolumeTextureDataProvider
+class FSparseVolumeTextureDataProviderOpenVDB : public UE::SVT::ITextureDataProvider
 {
 public:
 
@@ -423,7 +423,7 @@ public:
 		return true;
 	}
 
-	virtual FSparseVolumeTextureDataCreateInfo GetCreateInfo() const override
+	virtual UE::SVT::FTextureDataCreateInfo GetCreateInfo() const override
 	{
 		return SVTCreateInfo;
 	}
@@ -462,13 +462,13 @@ private:
 	TArray<TArray<FSingleGridToComponentMapping, TInlineAllocator<4>>> GridToComponentMappings;
 	TStaticArray<FOpenVDBSparseVolumeAttributesDesc, NumAttributesDescs> Attributes;
 	TStaticArray<uint32, NumAttributesDescs> NumComponents;
-	FSparseVolumeTextureDataCreateInfo SVTCreateInfo;
+	UE::SVT::FTextureDataCreateInfo SVTCreateInfo;
 	FIntVector3 VolumeBoundsMin;
 };
 
 #endif // OPENVDB_AVAILABLE
 
-bool ConvertOpenVDBToSparseVolumeTexture(TArray64<uint8>& SourceFile, const FOpenVDBImportOptions& ImportOptions, const FIntVector3& VolumeBoundsMin, FSparseVolumeTextureData& OutResult)
+bool ConvertOpenVDBToSparseVolumeTexture(TArray64<uint8>& SourceFile, const FOpenVDBImportOptions& ImportOptions, const FIntVector3& VolumeBoundsMin, UE::SVT::FTextureData& OutResult)
 {
 #if OPENVDB_AVAILABLE
 	FSparseVolumeTextureDataProviderOpenVDB DataProvider;
diff --git a/Engine/Source/Editor/SparseVolumeTexture/Private/SparseVolumeTextureOpenVDBUtility.h b/Engine/Source/Editor/SparseVolumeTexture/Private/SparseVolumeTextureOpenVDBUtility.h
--- a/Engine/Source/Editor/SparseVolumeTexture/Private/SparseVolumeTextureOpenVDBUtility.h
+++ b/Engine/Source/Editor/SparseVolumeTexture/Private/SparseVolumeTextureOpenVDBUtility.h
@@ -6,6 +6,14 @@
 
 #include "CoreMinimal.h"
 
+namespace UE
+{
+	namespace SVT
+	{
+		struct FTextureData;
+	}
+}
+
 enum class EOpenVDBGridType : uint8
 {
 	Unknown = 0,
@@ -51,7 +59,7 @@ bool IsOpenVDBGridValid(const FOpenVDBGridInfo& GridInfo, const FString& Filenam
 
 bool GetOpenVDBGridInfo(TArray64<uint8>& SourceFile, bool bCreateStrings, TArray<FOpenVDBGridInfo>* OutGridInfo);
 
-bool ConvertOpenVDBToSparseVolumeTexture(TArray64<uint8>& SourceFile, const struct FOpenVDBImportOptions& ImportOptions, const FIntVector3& VolumeBoundsMin, struct FSparseVolumeTextureData& OutResult);
+bool ConvertOpenVDBToSparseVolumeTexture(TArray64<uint8>& SourceFile, const struct FOpenVDBImportOptions& ImportOptions, const FIntVector3& VolumeBoundsMin, UE::SVT::FTextureData& OutResult);
 
 const TCHAR* OpenVDBGridTypeToString(EOpenVDBGridType Type);
 
diff --git a/Engine/Source/Runtime/Engine/Classes/SparseVolumeTexture/SparseVolumeTexture.h b/Engine/Source/Runtime/Engine/Classes/SparseVolumeTexture/SparseVolumeTexture.h
--- a/Engine/Source/Runtime/Engine/Classes/SparseVolumeTexture/SparseVolumeTexture.h
+++ b/Engine/Source/Runtime/Engine/Classes/SparseVolumeTexture/SparseVolumeTexture.h
@@ -8,6 +8,7 @@
 
 #include "Serialization/EditorBulkData.h"
 #include "Containers/Array.h"
+#include "Containers/StaticArray.h"
 #include "UnrealClient.h"
 #include "UObject/ObjectSaveContext.h"
 #include "Engine/TextureDefines.h"
@@ -17,14 +18,20 @@
 namespace UE { namespace Shader	{ enum class EValueType : uint8; } }
 namespace UE { namespace DerivedData { class FRequestOwner; } }
 
+// SVT_TODO: Unify with macros in SparseVolumeTextureCommon.ush
 #define SPARSE_VOLUME_TILE_RES 16
 #define SPARSE_VOLUME_TILE_BORDER 1
 #define SPARSE_VOLUME_TILE_RES_PADDED (SPARSE_VOLUME_TILE_RES + 2 * SPARSE_VOLUME_TILE_BORDER)
 
-struct FSparseVolumeTextureData;
-class FSparseVolumeTextureSceneProxy;
+namespace UE
+{
+namespace SVT
+{
 
-struct ENGINE_API FSparseVolumeTextureHeader
+struct FTextureData;
+class FStreamingManager;
+
+struct ENGINE_API FHeader
 {
 	static const uint32 kVersion = 0; // The current data format version for the header.
 	uint32 Version = kVersion; // This version can be used to convert existing header to new version later.
@@ -36,14 +43,124 @@ struct ENGINE_API FSparseVolumeTextureHeader
 	FIntVector3 PageTableVolumeAABBMin = FIntVector3(INT32_MAX, INT32_MAX, INT32_MAX);
 	FIntVector3 PageTableVolumeAABBMax = FIntVector3(INT32_MIN, INT32_MIN, INT32_MIN);
 	TStaticArray<EPixelFormat, 2> AttributesFormats = TStaticArray<EPixelFormat, 2>(InPlace, PF_Unknown);
-	TStaticArray<FVector4f, 2> NullTileValues = TStaticArray<FVector4f, 2>(InPlace, FVector4f());
-	TStaticArray<FVector4f, 2> NullTileValuesQuantized = TStaticArray<FVector4f, 2>(InPlace, FVector4f()); // Values after converting/quantizing original fallback values to the given attribute formats
+	TStaticArray<FVector4f, 2> FallbackValues = TStaticArray<FVector4f, 2>(InPlace, FVector4f());
 
-	FSparseVolumeTextureHeader() = default;
-	FSparseVolumeTextureHeader(const FIntVector3& AABBMin, const FIntVector3& AABBMax, EPixelFormat FormatA, EPixelFormat FormatB, const FVector4f& FallbackValueA, const FVector4f& FallbackValueB);
+	FHeader() = default;
+	FHeader(const FIntVector3& AABBMin, const FIntVector3& AABBMax, EPixelFormat FormatA, EPixelFormat FormatB, const FVector4f& FallbackValueA, const FVector4f& FallbackValueB);
 	void Serialize(FArchive& Ar);
 };
 
+// Describes a mip level of a SVT frame in terms of the sizes and offsets of the data in the built bulk data.
+struct FMipLevelStreamingInfo
+{
+	int32 BulkOffset;
+	int32 BulkSize;
+	int32 PageTableOffset; // relative to BulkOffset
+	int32 PageTableSize;
+	int32 TileDataAOffset; // relative to BulkOffset
+	int32 TileDataASize;
+	int32 TileDataBOffset; // relative to BulkOffset
+	int32 TileDataBSize;
+	int32 NumPhysicalTiles;
+};
+
+enum EResourceFlag : uint32
+{
+	EResourceFlag_StreamingDataInDDC = 1 << 0u, // FResources was cached, so MipLevelStreamingInfo can be streamed from DDC
+};
+
+// Represents the derived data of a SVT that is needed by the streaming manager.
+struct FResources
+{
+public:
+	FHeader Header;
+	uint32 ResourceFlags = 0;
+	// Info about sizes and offsets into the streamable mip level data. The last entry refers to the root mip level which is stored in RootData, not StreamableMipLevels.
+	TArray<FMipLevelStreamingInfo> MipLevelStreamingInfo;
+	// Data for the highest/"root" mip level
+	TArray<uint8> RootData;
+	// Data for all streamable mip levels
+	FByteBulkData StreamableMipLevels;
+#if WITH_EDITORONLY_DATA
+	// FTextureData from which all the other data can be built with a call to Build()
+	UE::Serialization::FEditorBulkData SourceData;
+#endif
+
+	// These are used for logging and retrieving StreamableMipLevels from DDC in FStreamingManager
+#if WITH_EDITORONLY_DATA
+	FString ResourceName;
+	FIoHash DDCKeyHash;
+	FIoHash DDCRawHash;
+#endif
+
+	// Called when serializing to/from DDC buffers and when serializing the owning USparseVolumeTextureFrame.
+	void Serialize(FArchive& Ar, UObject* Owner, bool bCooked);
+	// Returns true if there are streamable mip levels.
+	bool HasStreamingData() const;
+#if WITH_EDITORONLY_DATA
+	// Removes the StreamableMipLevels bulk data if it was successfully cached to DDC.
+	void DropBulkData();
+	// Fills StreamableMipLevels with data from DDC. Returns true when done.
+	bool RebuildBulkDataFromCacheAsync(const UObject* Owner, bool& bFailed);
+	// Builds all the data from SourceData. Is called by Cache().
+	bool Build(USparseVolumeTextureFrame* Owner);
+	// Cache the built data to/from DDC.
+	void Cache(USparseVolumeTextureFrame* Owner);
+#endif
+
+private:
+#if WITH_EDITORONLY_DATA
+	enum class EDDCRebuildState : uint8
+	{
+		Initial,
+		Pending,
+		Succeeded,
+		Failed,
+	};
+	TDontCopy<TPimplPtr<UE::DerivedData::FRequestOwner>> DDCRequestOwner;
+	std::atomic<EDDCRebuildState> DDCRebuildState;
+	void BeginRebuildBulkDataFromCache(const UObject* Owner);
+	void EndRebuildBulkDataFromCache();
+#endif
+};
+
+// Encapsulates RHI resources needed to render a SparseVolumeTexture.
+class ENGINE_API FTextureRenderResources : public ::FRenderResource
+{
+	friend class FStreamingManager;
+public:
+	const FHeader& GetHeader() const								{ check(IsInParallelRenderingThread()); return Header; }
+	FIntVector3 GetTileDataTextureResolution() const				{ check(IsInParallelRenderingThread()); return TileDataTextureResolution; }
+	int32 GetFrameIndex() const										{ check(IsInParallelRenderingThread()); return FrameIndex; }
+	int32 GetNumLogicalMipLevels() const							{ check(IsInParallelRenderingThread()); return NumLogicalMipLevels; }
+	FTextureRHIRef GetPageTableTextureRHI() const					{ check(IsInParallelRenderingThread()); return PageTableTextureRHI; }
+	FTextureRHIRef GetPhysicalTileDataATextureRHI() const			{ check(IsInParallelRenderingThread()); return PhysicalTileDataATextureRHI; }
+	FTextureRHIRef GetPhysicalTileDataBTextureRHI() const			{ check(IsInParallelRenderingThread()); return PhysicalTileDataBTextureRHI; }
+	FShaderResourceViewRHIRef GetStreamingInfoBufferSRVRHI() const	{ check(IsInParallelRenderingThread()); return StreamingInfoBufferSRVRHI; }
+	void GetPackedUniforms(FUintVector4& OutPacked0, FUintVector4& OutPacked1) const;
+	// Updates the GlobalVolumeResolution member in a thread-safe way.
+	void SetGlobalVolumeResolution_GameThread(const FIntVector3& GlobalVolumeResolution);
+
+	//~ Begin FRenderResource Interface.
+	virtual void InitRHI() override { /* Managed by FStreamingManager */ }
+	virtual void ReleaseRHI() override { /* Managed by FStreamingManager */ }
+	//~ End FRenderResource Interface.
+
+private:
+	FHeader Header;
+	FIntVector3 GlobalVolumeResolution = FIntVector3::ZeroValue; // The virtual resolution of the union of the AABBs of all frames. Needed for GetPackedUniforms().
+	FIntVector3 TileDataTextureResolution = FIntVector3::ZeroValue;
+	int32 FrameIndex = INDEX_NONE;
+	int32 NumLogicalMipLevels = 0; // Might not all be resident in GPU memory
+	FTextureRHIRef PageTableTextureRHI;
+	FTextureRHIRef PhysicalTileDataATextureRHI;
+	FTextureRHIRef PhysicalTileDataBTextureRHI;
+	FShaderResourceViewRHIRef StreamingInfoBufferSRVRHI;
+};
+
+}
+}
+
 enum ESparseVolumeTextureShaderUniform
 {
 	ESparseVolumeTexture_TileSize,
@@ -80,21 +197,15 @@ public:
 	virtual int32 GetNumMipLevels() const { return 0; }
 
 	virtual FIntVector GetVolumeResolution() const { return FIntVector(); }
+	virtual EPixelFormat GetFormat(int32 AttributesIndex) const { return PF_Unknown; }
+	virtual FVector4f GetFallbackValue(int32 AttributesIndex) const { return FVector4f(); }
 	virtual TextureAddress GetTextureAddressX() const { return TA_Wrap; }
 	virtual TextureAddress GetTextureAddressY() const { return TA_Wrap; }
 	virtual TextureAddress GetTextureAddressZ() const { return TA_Wrap; }
-
-	virtual const FSparseVolumeTextureSceneProxy* GetSparseVolumeTextureSceneProxy() const { return nullptr; }
+	virtual const UE::SVT::FTextureRenderResources* GetTextureRenderResources() const { return nullptr; }
 
 	/** Getter for the shader uniform parameters with index as ESparseVolumeTextureShaderUniform. */
-	FVector4 GetUniformParameter(int32 Index) const;
-
-	void GetPackedUniforms(FUintVector4& OutPacked0, FUintVector4& OutPacked1) const;
-
-	/** In order to keep the contents of an animated SVT sequence stable in world space, we need to account for the fact that
-		different frames of the sequence have different AABBs. We solve this by scaling and biasing UVs that are relative to
-		the volume bounds into the UV space represented by the AABB of each animation frame.*/
-	void GetFrameUVScaleBias(FVector* OutScale, FVector* OutBias) const;
+	FVector4 GetUniformParameter(int32 Index) const { return FVector4(ForceInitToZero); } // SVT_TODO: This mechanism is no longer needed and can be removed
 
 	/** Getter for the shader uniform parameter type with index as ESparseVolumeTextureShaderUniform. */
 	static UE::Shader::EValueType GetUniformParameterType(int32 Index);
@@ -116,20 +227,29 @@ class ENGINE_API USparseVolumeTextureFrame : public USparseVolumeTexture
 {
 	GENERATED_UCLASS_BODY()
 
-	// SVT_TODO: Remove this once the new streaming manager has been implemented. For now we need this to keep the old system running.
-	friend struct FStreamingSparseVolumeTextureData;
+	friend class UE::SVT::FStreamingManager;
+	friend class UStreamableSparseVolumeTexture;
 
 public:
 
 	USparseVolumeTextureFrame();
 	virtual ~USparseVolumeTextureFrame() = default;
 
-	static USparseVolumeTextureFrame* GetFrame(USparseVolumeTexture* SparseVolumeTexture, int32 FrameIndex);
+	// Retrieves a frame from the given SparseVolumeTexture and also issues a streaming request for it. 
+	// FrameIndex is of float type so that the streaming system can use the fractional part to more easily keep track of playback speed and direction (forward/reverse playback).
+	// MipLevel is the lowest mip level that the caller intends to use but does not guarantee that the mip is actually resident.
+	static USparseVolumeTextureFrame* GetFrameAndIssueStreamingRequest(USparseVolumeTexture* SparseVolumeTexture, float FrameIndex, int32 MipLevel);
 
-	bool Initialize(USparseVolumeTexture* InOwner, int32 InFrameIndex, FSparseVolumeTextureData& UncookedFrame);
-	bool BuildDerivedData(FSparseVolumeTextureData* OutMippedTextureData);
+	bool Initialize(USparseVolumeTexture* InOwner, int32 InFrameIndex, UE::SVT::FTextureData& UncookedFrame);
 	int32 GetFrameIndex() const { return FrameIndex; }
-	void GenerateOrLoadDDCRuntimeData(UE::DerivedData::FRequestOwner& DDCRequestOwner);
+	UE::SVT::FResources* GetResources() { return &Resources; }
+	// Creates TextureRenderResources if they don't already exist. Returns false if they already existed.
+	bool CreateTextureRenderResources();
+
+#if WITH_EDITORONLY_DATA
+	// Caches the derived data (FResources) of this frame to/from DDC and ensures that FTextureRenderResources exists.
+	void Cache();
+#endif
 
 	//~ Begin UObject Interface.
 	virtual void PostLoad() override;
@@ -137,30 +257,33 @@ public:
 	virtual void BeginDestroy() override;
 	virtual void Serialize(FArchive& Ar) override;
 	virtual void GetResourceSizeEx(FResourceSizeEx& CumulativeResourceSize) override;
+#if WITH_EDITOR
+	virtual void BeginCacheForCookedPlatformData(const ITargetPlatform* TargetPlatform) override;
+	virtual bool IsCachedCookedPlatformDataLoaded(const ITargetPlatform* TargetPlatform) override;
+	virtual void WillNeverCacheCookedPlatformDataAgain() override;
+	virtual void ClearCachedCookedPlatformData(const ITargetPlatform* TargetPlatform) override;
+	virtual void ClearAllCachedCookedPlatformData() override;
+#endif
 	//~ End UObject Interface.
 
 	//~ Begin USparseVolumeTexture Interface.
 	virtual int32 GetNumFrames() const override { return 1; }
-	virtual int32 GetNumMipLevels() const { return Owner->GetNumMipLevels(); }
-	virtual FIntVector GetVolumeResolution() const { return Owner->GetVolumeResolution(); }
-	virtual TextureAddress GetTextureAddressX() const { return Owner->GetTextureAddressX(); }
-	virtual TextureAddress GetTextureAddressY() const { return Owner->GetTextureAddressY(); }
-	virtual TextureAddress GetTextureAddressZ() const { return Owner->GetTextureAddressZ(); }
-	virtual const FSparseVolumeTextureSceneProxy* GetSparseVolumeTextureSceneProxy() const { return SceneProxy; }
+	virtual int32 GetNumMipLevels() const override { return Owner->GetNumMipLevels(); }
+	virtual FIntVector GetVolumeResolution() const override { return Owner->GetVolumeResolution(); }
+	virtual EPixelFormat GetFormat(int32 AttributesIndex) const override { return Owner->GetFormat(AttributesIndex); }
+	virtual FVector4f GetFallbackValue(int32 AttributesIndex) const override { return Owner->GetFallbackValue(AttributesIndex); }
+	virtual TextureAddress GetTextureAddressX() const override { return Owner->GetTextureAddressX(); }
+	virtual TextureAddress GetTextureAddressY() const override { return Owner->GetTextureAddressY(); }
+	virtual TextureAddress GetTextureAddressZ() const override { return Owner->GetTextureAddressZ(); }
+	virtual const UE::SVT::FTextureRenderResources* GetTextureRenderResources() const override { return TextureRenderResources; }
 	//~ End USparseVolumeTexture Interface.
 
 private:
 	UPROPERTY()
 	TObjectPtr<USparseVolumeTexture> Owner;
 	int32 FrameIndex;
-
-	FSparseVolumeTextureSceneProxy* SceneProxy;
-	FByteBulkData StreamingData;
-
-#if WITH_EDITORONLY_DATA
-	/** The raw data that can be loaded when we want to update cook the data with different settings or updated code without re importing. */
-	UE::Serialization::FEditorBulkData	RawData;
-#endif
+	UE::SVT::FResources Resources;
+	UE::SVT::FTextureRenderResources* TextureRenderResources;
 };
 
 UCLASS(ClassGroup = Rendering, BlueprintType)//, hidecategories = (Object))
@@ -184,6 +307,18 @@ public:
 	UPROPERTY(VisibleAnywhere, Category = "Texture")
 	int32 NumMipLevels;
 
+	UPROPERTY(VisibleAnywhere, Category = "Texture")
+	TEnumAsByte<enum EPixelFormat> FormatA;
+
+	UPROPERTY(VisibleAnywhere, Category = "Texture")
+	TEnumAsByte<enum EPixelFormat> FormatB;
+
+	UPROPERTY(VisibleAnywhere, Category = "Texture")
+	FVector4f FallbackValueA;
+
+	UPROPERTY(VisibleAnywhere, Category = "Texture")
+	FVector4f FallbackValueB;
+
 	/** The addressing mode to use for the X axis.								*/
 	UPROPERTY(EditAnywhere, BlueprintReadWrite, Category = "Texture", meta = (DisplayName = "X-axis Tiling Method"), AssetRegistrySearchable, AdvancedDisplay)
 	TEnumAsByte<enum TextureAddress> AddressX;
@@ -203,13 +338,12 @@ public:
 	// The NumExpectedFrames parameter on BeginInitialize() just serves as a potential optimization to reserve memory for the frames to be appended
 	// and doesn't need to match the exact number if it is not known at the time.
 	virtual bool BeginInitialize(int32 NumExpectedFrames);
-	virtual bool AppendFrame(FSparseVolumeTextureData& UncookedFrame);
+	virtual bool AppendFrame(UE::SVT::FTextureData& UncookedFrame);
 	virtual bool EndInitialize(int32 NumMipLevels = INDEX_NONE /*Create entire mip chain by default*/);
 
 	// Convenience function wrapping the multi-phase initialization functions above
-	virtual bool Initialize(const TArrayView<FSparseVolumeTextureData>& UncookedData, int32 NumMipLevels = INDEX_NONE /*Create entire mip chain by default*/);
-
-	const FSparseVolumeTextureSceneProxy* GetStreamedFrameProxyOrFallback(int32 FrameIndex, int32 MipLevel) const;
+	virtual bool Initialize(const TArrayView<UE::SVT::FTextureData>& UncookedData, int32 NumMipLevels = INDEX_NONE /*Create entire mip chain by default*/);
+	// Consider using USparseVolumeTextureFrame::GetFrameAndIssueStreamingRequest() if the frame should have streaming requests issued.
 	USparseVolumeTextureFrame* GetFrame(int32 FrameIndex) const { return Frames.IsValidIndex(FrameIndex) ? Frames[FrameIndex] : nullptr; }
 
 	//~ Begin UObject Interface.
@@ -227,10 +361,12 @@ public:
 	virtual int32 GetNumFrames() const override { return Frames.Num(); }
 	virtual int32 GetNumMipLevels() const override { return NumMipLevels; }
 	virtual FIntVector GetVolumeResolution() const override { return VolumeResolution; };
+	virtual EPixelFormat GetFormat(int32 AttributesIndex) const override { check(AttributesIndex >= 0 && AttributesIndex < 2) return AttributesIndex == 0 ? FormatA : FormatB; }
+	virtual FVector4f GetFallbackValue(int32 AttributesIndex) const override { check(AttributesIndex >= 0 && AttributesIndex < 2) return AttributesIndex == 0 ? FallbackValueA : FallbackValueB; }
 	virtual TextureAddress GetTextureAddressX() const override { return AddressX; }
 	virtual TextureAddress GetTextureAddressY() const override { return AddressY; }
 	virtual TextureAddress GetTextureAddressZ() const override { return AddressZ; }
-	virtual const FSparseVolumeTextureSceneProxy* GetSparseVolumeTextureSceneProxy() const override { return GetStreamedFrameProxyOrFallback(0 /*FrameIndex*/, 0 /*MipLevel*/); };
+	virtual const UE::SVT::FTextureRenderResources* GetTextureRenderResources() const override { return Frames.IsEmpty() ? nullptr : Frames[0]->GetTextureRenderResources(); }
 	//~ End USparseVolumeTexture Interface.
 
 protected:
@@ -241,9 +377,11 @@ protected:
 	FIntVector VolumeBoundsMin;
 	FIntVector VolumeBoundsMax;
 	EInitState InitState = EInitState::Uninitialized;
-#endif // WITH_EDITORONLY_DATA
 
-	void GenerateOrLoadDDCRuntimeDataAndCreateSceneProxy();
+	// Ensures all frames have derived data (based on the source data and the current settings like TextureAddress modes etc.) cached to DDC and are ready for rendering.
+	// Disconnects this SVT from the streaming manager, calls Cache() on all frames and finally connects to FStreamingManager again.
+	void RecacheFrames();
+#endif // WITH_EDITORONLY_DATA
 };
 
 UCLASS(ClassGroup = Rendering, BlueprintType)//, hidecategories = (Object))
@@ -257,7 +395,7 @@ public:
 	virtual ~UStaticSparseVolumeTexture() = default;
 
 	// Override AppendFrame() to ensure that there is never more than a single frame in a static SVT
-	virtual bool AppendFrame(FSparseVolumeTextureData& UncookedFrame) override;
+	virtual bool AppendFrame(UE::SVT::FTextureData& UncookedFrame) override;
 
 	//~ Begin USparseVolumeTexture Interface.
 	int32 GetNumFrames() const override { return 1; }
@@ -278,12 +416,11 @@ public:
 	virtual ~UAnimatedSparseVolumeTexture() = default;
 
 	//~ Begin USparseVolumeTexture Interface.
-	const FSparseVolumeTextureSceneProxy* GetSparseVolumeTextureSceneProxy() const override;
+	virtual const UE::SVT::FTextureRenderResources* GetTextureRenderResources() const override { return Frames.IsValidIndex(PreviewFrameIndex) ? Frames[PreviewFrameIndex]->GetTextureRenderResources() : nullptr; }
 	//~ End USparseVolumeTexture Interface.
 
 private:
 	int32 PreviewFrameIndex;
-	int32 PreviewMipLevel;
 };
 
 UCLASS(ClassGroup = Rendering, BlueprintType)//, hidecategories = (Object))
diff --git a/Engine/Source/Runtime/Engine/Private/Components/HeterogeneousVolumeComponent.cpp b/Engine/Source/Runtime/Engine/Private/Components/HeterogeneousVolumeComponent.cpp
--- a/Engine/Source/Runtime/Engine/Private/Components/HeterogeneousVolumeComponent.cpp
+++ b/Engine/Source/Runtime/Engine/Private/Components/HeterogeneousVolumeComponent.cpp
@@ -289,9 +289,10 @@ void UHeterogeneousVolumeComponent::TickComponent(float DeltaTime, ELevelTick Ti
 			{
 				if (bAnimate)
 				{
-					int32 FrameIndex = int32(Time * Framerate) % DefaultSVT->GetNumFrames();
-					int32 MipLevel = 0;
-					USparseVolumeTextureFrame* SVTFrame = USparseVolumeTextureFrame::GetFrame(DefaultSVT, FrameIndex);
+					const int32 FrameCount = DefaultSVT->GetNumFrames();
+					const float FrameIndexF = FMath::Fmod(Time * Framerate, (float)FrameCount);
+					const int32 MipLevel = 0;
+					USparseVolumeTextureFrame* SVTFrame = USparseVolumeTextureFrame::GetFrameAndIssueStreamingRequest(DefaultSVT, FrameIndexF, MipLevel);
 					MaterialInstanceDynamic->SetSparseVolumeTextureParameterValue(SVTParameterInfo[MaterialIndex].Name, SVTFrame);
 
 					VolumeResolution = SVTFrame ? SVTFrame->GetVolumeResolution() : FIntVector(1);
diff --git a/Engine/Source/Runtime/Engine/Private/ContentStreaming.cpp b/Engine/Source/Runtime/Engine/Private/ContentStreaming.cpp
--- a/Engine/Source/Runtime/Engine/Private/ContentStreaming.cpp
+++ b/Engine/Source/Runtime/Engine/Private/ContentStreaming.cpp
@@ -19,7 +19,6 @@
 #include "AudioCompressionSettingsUtils.h"
 #include "VT/VirtualTextureChunkManager.h"
 #include "Rendering/NaniteCoarseMeshStreamingManager.h"
-#include "SparseVolumeTexture/SparseVolumeTextureStreaming.h"
 
 #if WITH_EDITOR
 #include "AudioDevice.h"
@@ -795,9 +794,6 @@ FStreamingManagerCollection::FStreamingManagerCollection()
 
 	VirtualTextureStreamingManager = new FVirtualTextureChunkStreamingManager();
 	AddStreamingManager(VirtualTextureStreamingManager);
-
-	SparseVolumeTextureStreamingManager = new FSparseVolumeTextureStreamingManager();
-	AddStreamingManager(SparseVolumeTextureStreamingManager);
 }
 
 FStreamingManagerCollection::~FStreamingManagerCollection()
@@ -809,10 +805,6 @@ FStreamingManagerCollection::~FStreamingManagerCollection()
 		NaniteCoarseMeshStreamingManager = nullptr;
 	}
 
-	RemoveStreamingManager(SparseVolumeTextureStreamingManager);
-	delete SparseVolumeTextureStreamingManager;
-	SparseVolumeTextureStreamingManager = nullptr;
-
 	RemoveStreamingManager(VirtualTextureStreamingManager);
 	delete VirtualTextureStreamingManager;
 	VirtualTextureStreamingManager = nullptr;
@@ -1077,12 +1069,6 @@ Nanite::FCoarseMeshStreamingManager* FStreamingManagerCollection::GetNaniteCoars
 	return NaniteCoarseMeshStreamingManager;
 }
 
-ENGINE_API ISparseVolumeTextureStreamingManager& FStreamingManagerCollection::GetSparseVolumeTextureStreamingManager() const
-{
-	check(SparseVolumeTextureStreamingManager);
-	return *SparseVolumeTextureStreamingManager;
-}
-
 /** Don't stream world resources for the next NumFrames. */
 void FStreamingManagerCollection::SetDisregardWorldResourcesForFrames(int32 NumFrames )
 {
diff --git a/Engine/Source/Runtime/Engine/Private/Materials/HLSLMaterialTranslator.cpp b/Engine/Source/Runtime/Engine/Private/Materials/HLSLMaterialTranslator.cpp
--- a/Engine/Source/Runtime/Engine/Private/Materials/HLSLMaterialTranslator.cpp
+++ b/Engine/Source/Runtime/Engine/Private/Materials/HLSLMaterialTranslator.cpp
@@ -3868,7 +3868,7 @@ int32 FHLSLMaterialTranslator::AccessUniformExpression(int32 Index)
 	else if (CodeChunk.Type == MCT_SparseVolumeTexture)
 	{
 		int32 TextureInputIndex = UniformTextureExpressions[(uint32)EMaterialTextureParameterType::SparseVolume].AddUnique(TextureUniformExpression);
-		FormattedCode.Appendf(TEXT("SparseVolumeTextureUnpackUniforms(Material.SVTPackedUniform[%d*2], Material.SVTPackedUniform[%d*2+1])"), TextureInputIndex, TextureInputIndex);
+		FormattedCode.Appendf(TEXT("SparseVolumeTextureUnpackUniforms(Material.SVTPackedUniform[%d*2], Material.SVTPackedUniform[%d*2+1], Material.SparseVolumeTextureStreamingInfoBuffer_%d)"), TextureInputIndex, TextureInputIndex, TextureInputIndex);
 	}
 	else
 	{
diff --git a/Engine/Source/Runtime/Engine/Private/Materials/MaterialUniformExpressions.cpp b/Engine/Source/Runtime/Engine/Private/Materials/MaterialUniformExpressions.cpp
--- a/Engine/Source/Runtime/Engine/Private/Materials/MaterialUniformExpressions.cpp
+++ b/Engine/Source/Runtime/Engine/Private/Materials/MaterialUniformExpressions.cpp
@@ -17,7 +17,6 @@
 #include "GlobalRenderResources.h"
 #include "Shader/PreshaderTypes.h"
 #include "SparseVolumeTexture/SparseVolumeTexture.h"
-#include "SparseVolumeTexture/SparseVolumeTextureSceneProxy.h"
 #include "VT/RuntimeVirtualTexture.h"
 #include "TextureResource.h"
 
@@ -407,6 +406,7 @@ FShaderParametersMetadata* FUniformExpressionSet::CreateBufferStruct()
 	static FString SparseVolumeTexturePageTableNames[128];
 	static FString SparseVolumeTexturePhysicalANames[128];
 	static FString SparseVolumeTexturePhysicalBNames[128];
+	static FString SparseVolumeTextureStreamingInfoBufferNames[128];
 	static FString SparseVolumeTexturePhysicalSamplerNames[128];
 	static bool bInitializedTextureNames = false;
 	if (!bInitializedTextureNames)
@@ -434,6 +434,7 @@ FShaderParametersMetadata* FUniformExpressionSet::CreateBufferStruct()
 			SparseVolumeTexturePageTableNames[i] = FString::Printf(TEXT("SparseVolumeTexturePageTable_%d"), i);
 			SparseVolumeTexturePhysicalANames[i] = FString::Printf(TEXT("SparseVolumeTexturePhysicalA_%d"), i);
 			SparseVolumeTexturePhysicalBNames[i] = FString::Printf(TEXT("SparseVolumeTexturePhysicalB_%d"), i);
+			SparseVolumeTextureStreamingInfoBufferNames[i] = FString::Printf(TEXT("SparseVolumeTextureStreamingInfoBuffer_%d"), i);
 			SparseVolumeTexturePhysicalSamplerNames[i] = FString::Printf(TEXT("SparseVolumeTexturePhysical_%dSampler"), i);
 		}
 	}
@@ -510,6 +511,11 @@ FShaderParametersMetadata* FUniformExpressionSet::CreateBufferStruct()
 		new(Members) FShaderParametersMetadata::FMember(*SparseVolumeTexturePhysicalBNames[i], TEXT("Texture3D"), __LINE__, NextMemberOffset, UBMT_TEXTURE, EShaderPrecisionModifier::Float, 1, 1, 0, NULL);
 		NextMemberOffset += SHADER_PARAMETER_POINTER_ALIGNMENT;
 
+		// Streaming Info Buffer
+		check((NextMemberOffset% SHADER_PARAMETER_POINTER_ALIGNMENT) == 0);
+		new(Members) FShaderParametersMetadata::FMember(*SparseVolumeTextureStreamingInfoBufferNames[i], TEXT("ByteAddressBuffer"), __LINE__, NextMemberOffset, UBMT_SRV, EShaderPrecisionModifier::Float, 1, 1, 0, NULL);
+		NextMemberOffset += SHADER_PARAMETER_POINTER_ALIGNMENT;
+
 		// Sampler
 		check((NextMemberOffset% SHADER_PARAMETER_POINTER_ALIGNMENT) == 0);
 		new(Members) FShaderParametersMetadata::FMember(*SparseVolumeTexturePhysicalSamplerNames[i], TEXT("SamplerState"), __LINE__, NextMemberOffset, UBMT_SAMPLER, EShaderPrecisionModifier::Float, 1, 1, 0, NULL);
@@ -829,9 +835,10 @@ void FUniformExpressionSet::FillUniformBuffer(const FMaterialRenderContext& Mate
 
 			const USparseVolumeTexture* SVTexture = nullptr;
 			GetTextureValue(ExpressionIndex, MaterialRenderContext, MaterialRenderContext.Material, SVTexture);
-			if (SVTexture != nullptr)
+			const UE::SVT::FTextureRenderResources* RenderResources = SVTexture != nullptr ? SVTexture->GetTextureRenderResources() : nullptr;
+			if (RenderResources)
 			{
-				SVTexture->GetPackedUniforms(SVTPackedUniform[0], SVTPackedUniform[1]);
+				RenderResources->GetPackedUniforms(SVTPackedUniform[0], SVTPackedUniform[1]);
 			}
 			else
 			{
@@ -991,7 +998,7 @@ void FUniformExpressionSet::FillUniformBuffer(const FMaterialRenderContext& Mate
 				+ UniformTextureParameters[(uint32)EMaterialTextureParameterType::Array2D].Num() * 2
 				+ UniformTextureParameters[(uint32)EMaterialTextureParameterType::ArrayCube].Num() * 2
 				+ UniformTextureParameters[(uint32)EMaterialTextureParameterType::Volume].Num() * 2
-				+ UniformTextureParameters[(uint32)EMaterialTextureParameterType::SparseVolume].Num() * 4
+				+ UniformTextureParameters[(uint32)EMaterialTextureParameterType::SparseVolume].Num() * 5
 				+ UniformExternalTextureParameters.Num() * 2
 				+ UniformTextureParameters[(uint32)EMaterialTextureParameterType::Virtual].Num() * 2
 				+ NumPageTableTextures
@@ -1264,8 +1271,9 @@ void FUniformExpressionSet::FillUniformBuffer(const FMaterialRenderContext& Mate
 			void** ResourceTableTexturePageTablePtr = (void**)((uint8*)BufferCursor + 0 * SHADER_PARAMETER_POINTER_ALIGNMENT);
 			void** ResourceTableTexturePhysicalAPtr = (void**)((uint8*)BufferCursor + 1 * SHADER_PARAMETER_POINTER_ALIGNMENT);
 			void** ResourceTableTexturePhysicalBPtr = (void**)((uint8*)BufferCursor + 2 * SHADER_PARAMETER_POINTER_ALIGNMENT);
-			void** ResourceTablePhysicalSamplerPtr = (void**)((uint8*)BufferCursor + 3 * SHADER_PARAMETER_POINTER_ALIGNMENT);
-			BufferCursor = ((uint8*)BufferCursor) + (SHADER_PARAMETER_POINTER_ALIGNMENT * 4);
+			void** ResourceTableStreamingInfoBufferPtr = (void**)((uint8*)BufferCursor + 3 * SHADER_PARAMETER_POINTER_ALIGNMENT);
+			void** ResourceTablePhysicalSamplerPtr = (void**)((uint8*)BufferCursor + 4 * SHADER_PARAMETER_POINTER_ALIGNMENT);
+			BufferCursor = ((uint8*)BufferCursor) + (SHADER_PARAMETER_POINTER_ALIGNMENT * 5);
 			check(BufferCursor <= TempBuffer + TempBufferSize);
 
 			check(GBlackVolumeTexture->TextureRHI);
@@ -1279,18 +1287,22 @@ void FUniformExpressionSet::FillUniformBuffer(const FMaterialRenderContext& Mate
 			GetTextureValue(ExpressionIndex, MaterialRenderContext, MaterialRenderContext.Material, SVTexture);
 			if (SVTexture != nullptr)
 			{
-				const FSparseVolumeTextureSceneProxy* SVTextureProxy = SVTexture->GetSparseVolumeTextureSceneProxy();
-				if (SVTextureProxy)
+				const UE::SVT::FTextureRenderResources* RenderResources = SVTexture->GetTextureRenderResources();
+				if (RenderResources)
 				{
-					check(SVTextureProxy->GetPageTableTextureRHI());
-					*ResourceTableTexturePageTablePtr = SVTextureProxy->GetPageTableTextureRHI();
+					// It's possible for RenderResources to be valid, but PageTableTexture still null, so we need to have a fallback
+					FTextureRHIRef PageTableTexture = RenderResources->GetPageTableTextureRHI();
+					*ResourceTableTexturePageTablePtr = PageTableTexture ? PageTableTexture : GBlackUintVolumeTexture->TextureRHI;
 
-					FTextureRHIRef TileDataATexture = SVTextureProxy->GetPhysicalTileDataATextureRHI();
+					FTextureRHIRef TileDataATexture = RenderResources->GetPhysicalTileDataATextureRHI();
 					*ResourceTableTexturePhysicalAPtr = TileDataATexture ? TileDataATexture : GBlackVolumeTexture->TextureRHI;
 
-					FTextureRHIRef TileDataBTexture = SVTextureProxy->GetPhysicalTileDataBTextureRHI();
+					FTextureRHIRef TileDataBTexture = RenderResources->GetPhysicalTileDataBTextureRHI();
 					*ResourceTableTexturePhysicalBPtr = TileDataBTexture ? TileDataBTexture : GBlackVolumeTexture->TextureRHI;
 
+					FShaderResourceViewRHIRef StreamingInfoBufferSRV = RenderResources->GetStreamingInfoBufferSRVRHI();
+					*ResourceTableStreamingInfoBufferPtr = StreamingInfoBufferSRV ? StreamingInfoBufferSRV : GEmptyVertexBufferWithUAV->ShaderResourceViewRHI;
+
 					*ResourceTablePhysicalSamplerPtr = TStaticSamplerState<SF_Bilinear, AM_Clamp, AM_Clamp, AM_Clamp>::GetRHI();
 				}
 			}
diff --git a/Engine/Source/Runtime/Engine/Private/SparseVolumeTexture/SparseVolumeTexture.cpp b/Engine/Source/Runtime/Engine/Private/SparseVolumeTexture/SparseVolumeTexture.cpp
--- a/Engine/Source/Runtime/Engine/Private/SparseVolumeTexture/SparseVolumeTexture.cpp
+++ b/Engine/Source/Runtime/Engine/Private/SparseVolumeTexture/SparseVolumeTexture.cpp
@@ -13,11 +13,14 @@
 #include "EngineUtils.h"
 #include "Shader/ShaderTypes.h"
 #include "RenderingThread.h"
-#include "SparseVolumeTexture/SparseVolumeTextureSceneProxy.h"
+#include "UObject/Package.h"
 #include "SparseVolumeTexture/SparseVolumeTextureData.h"
 #include "SparseVolumeTexture/SparseVolumeTextureUtility.h"
+#include "SparseVolumeTexture/SparseVolumeTextureStreamingManager.h"
 
 #if WITH_EDITORONLY_DATA
+#include "Misc/ScopedSlowTask.h"
+#include "DerivedDataCacheInterface.h"
 #include "DerivedDataCache.h"
 #include "DerivedDataRequestOwner.h"
 #endif
@@ -39,7 +42,28 @@ DEFINE_LOG_CATEGORY_STATIC(LogSparseVolumeTexture, Log, All);
 
 ////////////////////////////////////////////////////////////////////////////////////////////////
 
-static int32 SVTComputeNumMipLevels(const FIntVector3& InResolution)
+FArchive& operator<<(FArchive& Ar, UE::SVT::FMipLevelStreamingInfo& MipLevelStreamingInfo)
+{
+	Ar << MipLevelStreamingInfo.BulkOffset;
+	Ar << MipLevelStreamingInfo.BulkSize;
+	Ar << MipLevelStreamingInfo.PageTableOffset;
+	Ar << MipLevelStreamingInfo.PageTableSize;
+	Ar << MipLevelStreamingInfo.TileDataAOffset;
+	Ar << MipLevelStreamingInfo.TileDataASize;
+	Ar << MipLevelStreamingInfo.TileDataBOffset;
+	Ar << MipLevelStreamingInfo.TileDataBSize;
+	Ar << MipLevelStreamingInfo.NumPhysicalTiles;
+	return Ar;
+}
+
+////////////////////////////////////////////////////////////////////////////////////////////////
+
+namespace UE
+{
+namespace SVT
+{
+
+static int32 ComputeNumMipLevels(const FIntVector3& InResolution)
 {
 	int32 Levels = 1;
 	FIntVector3 Resolution = InResolution;
@@ -51,9 +75,13 @@ static int32 SVTComputeNumMipLevels(const FIntVector3& InResolution)
 	return Levels;
 };
 
-////////////////////////////////////////////////////////////////////////////////////////////////
+static const FString& GetDerivedDataVersion()
+{
+	static FString CachedVersionString = TEXT("381AE2A9-A903-4C8F-8486-891E24D6FC71");	// Bump this if you want to ignore all cached data so far.
+	return CachedVersionString;
+}
 
-FSparseVolumeTextureHeader::FSparseVolumeTextureHeader(const FIntVector3& AABBMin, const FIntVector3& AABBMax, EPixelFormat FormatA, EPixelFormat FormatB, const FVector4f& FallbackValueA, const FVector4f& FallbackValueB)
+FHeader::FHeader(const FIntVector3& AABBMin, const FIntVector3& AABBMax, EPixelFormat FormatA, EPixelFormat FormatB, const FVector4f& FallbackValueA, const FVector4f& FallbackValueB)
 {
 	VirtualVolumeAABBMin = AABBMin;
 	VirtualVolumeAABBMax = AABBMax;
@@ -72,13 +100,11 @@ FSparseVolumeTextureHeader::FSparseVolumeTextureHeader(const FIntVector3& AABBMi
 	AttributesFormats[0] = FormatA;
 	AttributesFormats[1] = FormatB;
 
-	NullTileValues[0] = FallbackValueA;
-	NullTileValues[1] = FallbackValueB;
-	NullTileValuesQuantized[0] = FallbackValueA;
-	NullTileValuesQuantized[1] = FallbackValueB;
+	FallbackValues[0] = FallbackValueA;
+	FallbackValues[1] = FallbackValueB;
 }
 
-void FSparseVolumeTextureHeader::Serialize(FArchive& Ar)
+void FHeader::Serialize(FArchive& Ar)
 {
 	Ar << Version;
 
@@ -92,81 +118,420 @@ void FSparseVolumeTextureHeader::Serialize(FArchive& Ar)
 		Ar << PageTableVolumeAABBMax;
 		UE::SVT::Private::SerializeEnumAs<uint8>(Ar, AttributesFormats[0]);
 		UE::SVT::Private::SerializeEnumAs<uint8>(Ar, AttributesFormats[1]);
-		Ar << NullTileValues[0];
-		Ar << NullTileValues[1];
-		Ar << NullTileValuesQuantized[0];
-		Ar << NullTileValuesQuantized[1];
+		Ar << FallbackValues[0];
+		Ar << FallbackValues[1];
 	}
 	else
 	{
-		// FSparseVolumeTextureHeader needs to account for new version
+		// FHeader needs to account for new version
 		check(false);
 	}
 }
 
 ////////////////////////////////////////////////////////////////////////////////////////////////
 
-USparseVolumeTexture::USparseVolumeTexture(const FObjectInitializer& ObjectInitializer)
-	: Super(ObjectInitializer)
+void FResources::Serialize(FArchive& Ar, UObject* Owner, bool bCooked)
+{
+	// Note: this is all derived data, native versioning is not needed, but be sure to bump GetDerivedDataVersion() when modifying!
+	FStripDataFlags StripFlags(Ar, 0);
+	if (!StripFlags.IsDataStrippedForServer())
+	{
+		Header.Serialize(Ar);
+
+		uint32 StoredResourceFlags;
+		if (Ar.IsSaving() && bCooked)
+		{
+			// Disable DDC store when saving out a cooked build
+			StoredResourceFlags = ResourceFlags & ~EResourceFlag_StreamingDataInDDC;
+			Ar << StoredResourceFlags;
+		}
+		else
+		{
+			Ar << ResourceFlags;
+			StoredResourceFlags = ResourceFlags;
+		}
+
+		Ar << MipLevelStreamingInfo;
+		Ar << RootData;
+
+		// StreamableMipLevels is only serialized in cooked builds and when caching to DDC failed in editor builds.
+		// If the data was successfully cached to DDC, we just query it from DDC on the next run or recreate it if that failed.
+		if (StoredResourceFlags & EResourceFlag_StreamingDataInDDC)
+		{
+#if !WITH_EDITORONLY_DATA
+			checkf(false, TEXT("UE::SVT::FResources was serialized with EResourceFlag_StreamingDataInDDC in a cooked build!"));
+#endif
+		}
+		else
+		{
+			StreamableMipLevels.Serialize(Ar, Owner, 0);
+		}
+
+#if WITH_EDITORONLY_DATA
+		if (!bCooked)
+		{
+			SourceData.Serialize(Ar, Owner);
+		}
+#endif
+
+#if !WITH_EDITORONLY_DATA
+		check(!HasStreamingData() || StreamableMipLevels.GetBulkDataSize() > 0);
+#endif
+	}
+}
+
+bool FResources::HasStreamingData() const
+{
+	return MipLevelStreamingInfo.Num() > 1;
+}
+
+#if WITH_EDITORONLY_DATA
+
+void FResources::DropBulkData()
 {
+	if (HasStreamingData() && (ResourceFlags & EResourceFlag_StreamingDataInDDC))
+	{
+		StreamableMipLevels.RemoveBulkData();
+	}
 }
 
-FVector4 USparseVolumeTexture::GetUniformParameter(int32 Index) const
+bool FResources::RebuildBulkDataFromCacheAsync(const UObject* Owner, bool& bFailed)
 {
-	const FSparseVolumeTextureSceneProxy* Proxy = GetSparseVolumeTextureSceneProxy();
-	if (Proxy)
+	bFailed = false;
+	if (!HasStreamingData() || (ResourceFlags & EResourceFlag_StreamingDataInDDC) == 0u)
 	{
-		const FSparseVolumeTextureRuntimeHeader& Header = Proxy->GetHeader();
-		switch (Index)
+		return true;
+	}
+	if (DDCRebuildState.load() == EDDCRebuildState::Initial)
+	{
+		if (StreamableMipLevels.IsBulkDataLoaded())
+		{
+			return true;
+		}
+		// Handle Initial state first so we can transition directly to Succeeded/Failed if the data was immediately available from the cache.
+		check(!(*DDCRequestOwner).IsValid());
+		BeginRebuildBulkDataFromCache(Owner);
+	}
+	switch (DDCRebuildState.load())
+	{
+	case EDDCRebuildState::Pending:
+		return false;
+	case EDDCRebuildState::Succeeded:
+		check(StreamableMipLevels.GetBulkDataSize() > 0);
+		EndRebuildBulkDataFromCache();
+		return true;
+	case EDDCRebuildState::Failed:
+		bFailed = true;
+		EndRebuildBulkDataFromCache();
+		return true;
+	default:
+		check(false);
+		return true;
+	}
+}
+
+bool FResources::Build(USparseVolumeTextureFrame* Owner)
+{
+	// Check if the virtualized bulk data payload is available
+	if (SourceData.HasPayloadData())
+	{
+		UE::Serialization::FEditorBulkDataReader SourceDataReader(SourceData);
+		FTextureData SourceTextureData;
+		SourceTextureData.Serialize(SourceDataReader);
+
+		FTextureDataAddressingInfo AddressingInfo{};
+		AddressingInfo.VolumeResolution = Owner->GetVolumeResolution();
+		AddressingInfo.AddressX = Owner->GetTextureAddressX();
+		AddressingInfo.AddressY = Owner->GetTextureAddressY();
+		AddressingInfo.AddressZ = Owner->GetTextureAddressZ();
+
+		const int32 NumMipLevelsRequested = -1; // generate entire mip chain
+		const bool bMoveMip0FromSource = true; // we have no need to keep SourceTextureData around
+		FTextureData DerivedTextureData;
+		if (!SourceTextureData.BuildDerivedData(AddressingInfo, NumMipLevelsRequested, bMoveMip0FromSource, DerivedTextureData))
+		{
+			return false;
+		}
+
+		// Now unload the source data
+		SourceData.UnloadData();
+
+		const int32 NumMipLevels = DerivedTextureData.MipMaps.Num();
+
+		Header = DerivedTextureData.Header;
+		RootData.Reset();
+		MipLevelStreamingInfo.SetNumZeroed(NumMipLevels);
+		ResourceFlags = 0;
+		ResourceName.Reset();
+		DDCKeyHash.Reset();
+		DDCRawHash.Reset();
+		DDCRebuildState.store(EDDCRebuildState::Initial);
+
+		// Stores page table into BulkData as two consecutive arrays of packed page coordinates and linear indices into the physical tiles array.
+		// Returns number of written/non-zero page table entries
+		auto CompressPageTable = [](const TArray<uint32>& PageTable, const FIntVector3& Resolution, TArray<uint8>& BulkData) -> int32
+		{
+			int32 NumNonZeroEntries = 0;
+			for (uint32 Entry : PageTable)
+			{
+				if (Entry)
+				{
+					++NumNonZeroEntries;
+				}
+			}
+
+			const int32 BaseOffset = BulkData.Num();
+			BulkData.SetNum(BulkData.Num() + NumNonZeroEntries * 2 * sizeof(uint32));
+			uint32* PackedCoords = reinterpret_cast<uint32*>(BulkData.GetData() + BaseOffset);
+			uint32* PageEntries = PackedCoords + NumNonZeroEntries;
+			
+			int32 NumWrittenEntries = 0;
+			for (int32 Z = 0; Z < Resolution.Z; ++Z)
+			{
+				for (int32 Y = 0; Y < Resolution.Y; ++Y)
+				{
+					for (int32 X = 0; X < Resolution.X; ++X)
+					{
+						const int32 LinearCoord = (Z * Resolution.Y * Resolution.X) + (Y * Resolution.X) + X;
+						const uint32 Entry = PageTable[LinearCoord];
+						if (Entry)
+						{
+							const uint32 Packed = (X & 0x7FFu) | ((Y & 0x7FFu) << 11u) | ((Z & 0x3FFu) << 22u);
+							PackedCoords[NumWrittenEntries] = Packed;
+							PageEntries[NumWrittenEntries] = Entry - 1;
+							++NumWrittenEntries;
+						}
+					}
+				}
+			}
+
+			return NumNonZeroEntries;
+		};
+
+		TArray<uint8> StreamableBulkData;
+		for (int32 MipLevelIdx = 0; MipLevelIdx < DerivedTextureData.MipMaps.Num(); ++MipLevelIdx)
 		{
-		case ESparseVolumeTexture_TileSize:
+			const FTextureData::FMipMap& Mip = DerivedTextureData.MipMaps[MipLevelIdx];
+			const bool bIsRootMipLevel = (MipLevelIdx == (DerivedTextureData.MipMaps.Num() - 1));
+			TArray<uint8>& BulkData = bIsRootMipLevel ? RootData : StreamableBulkData;
+
+			FIntVector3 MipPageTableResolution = DerivedTextureData.Header.PageTableVolumeResolution >> MipLevelIdx;
+			MipPageTableResolution = FIntVector3(FMath::Max(1, MipPageTableResolution.X), FMath::Max(1, MipPageTableResolution.Y), FMath::Max(1, MipPageTableResolution.Z));
+
+			FMipLevelStreamingInfo& MipStreamingInfo = MipLevelStreamingInfo[MipLevelIdx];
+			MipStreamingInfo.BulkOffset = BulkData.Num();
+
+			MipStreamingInfo.PageTableOffset = BulkData.Num() - MipStreamingInfo.BulkOffset;
+			const int32 NumNonZeroPageTableEntries = CompressPageTable(Mip.PageTable, MipPageTableResolution, BulkData);
+			MipStreamingInfo.PageTableSize = NumNonZeroPageTableEntries * 2 * sizeof(uint32);
+			
+			MipStreamingInfo.TileDataAOffset = BulkData.Num() - MipStreamingInfo.BulkOffset;
+			MipStreamingInfo.TileDataASize = Mip.PhysicalTileDataA.Num();
+			BulkData.Append(Mip.PhysicalTileDataA.GetData(), Mip.PhysicalTileDataA.Num());
+
+			MipStreamingInfo.TileDataBOffset = BulkData.Num() - MipStreamingInfo.BulkOffset;
+			MipStreamingInfo.TileDataBSize = Mip.PhysicalTileDataB.Num();
+			BulkData.Append(Mip.PhysicalTileDataB.GetData(), Mip.PhysicalTileDataB.Num());
+
+			MipStreamingInfo.BulkSize = BulkData.Num() - MipStreamingInfo.BulkOffset;
+			MipStreamingInfo.NumPhysicalTiles = Mip.NumPhysicalTiles;
+		}
+
+		// Store StreamableMipLevels
 		{
-			return FVector4(float(SPARSE_VOLUME_TILE_RES), 0.0f, 0.0f, 0.0f);
+			StreamableMipLevels.Lock(LOCK_READ_WRITE);
+			uint8* Ptr = (uint8*)StreamableMipLevels.Realloc(StreamableBulkData.Num());
+			FMemory::Memcpy(Ptr, StreamableBulkData.GetData(), StreamableBulkData.Num());
+			StreamableMipLevels.Unlock();
+			StreamableMipLevels.SetBulkDataFlags(BULKDATA_Force_NOT_InlinePayload);
 		}
-		case ESparseVolumeTexture_PageTableSize:
+
+		return true;
+	}
+	return false;
+}
+
+void FResources::Cache(USparseVolumeTextureFrame* Owner)
+{
+	if (Owner->GetPackage()->bIsCookedForEditor)
+	{
+		// Don't cache for cooked packages
+		return;
+	}
+
+	using namespace UE::DerivedData;
+
+	static const FValueId SVTDataId = FValueId::FromName("SparseVolumeTextureData");
+	static const FValueId SVTStreamingDataId = FValueId::FromName("SparseVolumeTextureStreamingData");
+	const FString KeySuffix = SourceData.GetIdentifier().ToString() + FString::Format(TEXT("{0}_{1}_{2}_{3}"), { Owner->GetNumMipLevels(), Owner->GetTextureAddressX(), Owner->GetTextureAddressY(), Owner->GetTextureAddressZ() });
+	FString DerivedDataKey = FDerivedDataCacheInterface::BuildCacheKey(TEXT("SPARSEVOLUMETEXTURE"), *GetDerivedDataVersion(), *KeySuffix);
+
+	FCacheKey CacheKey;
+	CacheKey.Bucket = FCacheBucket(TEXT("SparseVolumeTexture"));
+	CacheKey.Hash = FIoHash::HashBuffer(MakeMemoryView(FTCHARToUTF8(DerivedDataKey)));
+
+	// Check if the data already exists in DDC
+	FSharedBuffer ResourcesDataBuffer;
+	FIoHash SVTStreamingDataHash;
+	{
+		FCacheRecordPolicyBuilder PolicyBuilder(ECachePolicy::Default | ECachePolicy::KeepAlive);
+		PolicyBuilder.AddValuePolicy(SVTStreamingDataId, ECachePolicy::Default | ECachePolicy::SkipData);
+
+		FCacheGetRequest Request;
+		Request.Name = Owner->GetPathName();
+		Request.Key = CacheKey;
+		Request.Policy = PolicyBuilder.Build();
+
+		FRequestOwner RequestOwner(EPriority::Blocking);
+		GetCache().Get(MakeArrayView(&Request, 1), RequestOwner,
+			[&ResourcesDataBuffer, &SVTStreamingDataHash](FCacheGetResponse&& Response)
+			{
+				if (Response.Status == EStatus::Ok)
+				{
+					const FCompressedBuffer& CompressedBuffer = Response.Record.GetValue(SVTDataId).GetData();
+					ResourcesDataBuffer = CompressedBuffer.Decompress();
+
+					SVTStreamingDataHash = Response.Record.GetValue(SVTStreamingDataId).GetRawHash();
+				}
+			});
+		RequestOwner.Wait();
+	}
+
+	if (!ResourcesDataBuffer.IsNull())
+	{
+		// Found it!
+		// We can serialize the data from the DDC buffer and are done.
+		FMemoryReaderView Ar(ResourcesDataBuffer.GetView(), /*bIsPersistent=*/ true);
+		Serialize(Ar, Owner, /*bCooked=*/ false);
+
+		check(StreamableMipLevels.GetBulkDataSize() == 0);
+		if (ResourceFlags & EResourceFlag_StreamingDataInDDC)
 		{
-			return FVector4(Header.PageTableVolumeResolution.X, Header.PageTableVolumeResolution.Y, Header.PageTableVolumeResolution.Z, 0.0f);
+			DDCKeyHash = CacheKey.Hash;
+			DDCRawHash = SVTStreamingDataHash;
 		}
-		case ESparseVolumeTexture_UVScale: // fallthrough
-		case ESparseVolumeTexture_UVBias:
+	}
+	else
+	{
+		// DDC lookup failed! Build the data again.
+		const bool bBuiltSuccessfully = Build(Owner);
+		check(bBuiltSuccessfully);
+
+		FCacheRecordBuilder RecordBuilder(CacheKey);
+		if (!MipLevelStreamingInfo.IsEmpty())
 		{
-			FVector Scale;
-			FVector Bias;
-			GetFrameUVScaleBias(&Scale, &Bias);
-			return (Index == ESparseVolumeTexture_UVScale) ? FVector4(Scale) : FVector4(Bias);
+			if (HasStreamingData())
+			{
+				FByteBulkData& BulkData = StreamableMipLevels;
+
+				FValue Value = FValue::Compress(FSharedBuffer::MakeView(BulkData.LockReadOnly(), BulkData.GetBulkDataSize()));
+				RecordBuilder.AddValue(SVTStreamingDataId, Value);
+				BulkData.Unlock();
+				ResourceFlags |= EResourceFlag_StreamingDataInDDC;
+				DDCKeyHash = CacheKey.Hash;
+				DDCRawHash = Value.GetRawHash();
+			}
 		}
-		default:
+
+		// Serialize to a buffer and store into DDC.
+		FLargeMemoryWriter Ar(0, /*bIsPersistent=*/ true);
+		Serialize(Ar, Owner, /*bCooked=*/ false);
+
+		bool bSavedToDDC = false;
 		{
-			break;
+			FValue Value = FValue::Compress(FSharedBuffer::MakeView(Ar.GetData(), Ar.TotalSize()));
+			RecordBuilder.AddValue(SVTDataId, Value);
+
+			FRequestOwner RequestOwner(UE::DerivedData::EPriority::Blocking);
+			const FCachePutRequest PutRequest = { FSharedString(Owner->GetPathName()), RecordBuilder.Build(), ECachePolicy::Default | ECachePolicy::KeepAlive };
+			GetCache().Put(MakeArrayView(&PutRequest, 1), RequestOwner,
+				[&bSavedToDDC](FCachePutResponse&& Response)
+				{
+					if (Response.Status == EStatus::Ok)
+					{
+						bSavedToDDC = true;
+					}
+				});
+
+			RequestOwner.Wait();
+
+			if (bSavedToDDC && HasStreamingData())
+			{
+				// Drop streaming data from memory when it has been successfully committed to DDC
+				DropBulkData();
+			}
 		}
+
+		if (HasStreamingData() && !bSavedToDDC)
+		{
+			// Streaming data was not pushed to DDC. Disable DDC streaming flag.
+			check(StreamableMipLevels.GetBulkDataSize() > 0);
+			ResourceFlags &= ~EResourceFlag_StreamingDataInDDC;
 		}
-		checkNoEntry();
-		return FVector4(ForceInitToZero);
 	}
+}
 
-	// 0 while waiting for the proxy
-	return FVector4(ForceInitToZero);
+#endif // WITH_EDITORONLY_DATA
+
+#if WITH_EDITORONLY_DATA
+
+void FResources::BeginRebuildBulkDataFromCache(const UObject* Owner)
+{
+	check(DDCRebuildState.load() == EDDCRebuildState::Initial);
+	if (!HasStreamingData() || (ResourceFlags & EResourceFlag_StreamingDataInDDC) == 0u)
+	{
+		return;
+	}
+	using namespace UE::DerivedData;
+	FCacheKey Key;
+	Key.Bucket = FCacheBucket(TEXT("SparseVolumeTexture"));
+	Key.Hash = DDCKeyHash;
+	check(!DDCKeyHash.IsZero());
+	FCacheGetChunkRequest Request;
+	Request.Name = Owner->GetPathName();
+	Request.Id = FValueId::FromName("SparseVolumeTextureStreamingData");
+	Request.Key = Key;
+	Request.RawHash = DDCRawHash;
+	check(!DDCRawHash.IsZero());
+	FSharedBuffer SharedBuffer;
+	*DDCRequestOwner = MakePimpl<FRequestOwner>(EPriority::Normal);
+	DDCRebuildState.store(EDDCRebuildState::Pending);
+	GetCache().GetChunks(MakeArrayView(&Request, 1), **DDCRequestOwner,
+		[this](FCacheGetChunkResponse&& Response)
+		{
+			if (Response.Status == EStatus::Ok)
+			{
+				StreamableMipLevels.Lock(LOCK_READ_WRITE);
+				uint8* Ptr = (uint8*)StreamableMipLevels.Realloc(Response.RawData.GetSize());
+				FMemory::Memcpy(Ptr, Response.RawData.GetData(), Response.RawData.GetSize());
+				StreamableMipLevels.Unlock();
+				StreamableMipLevels.SetBulkDataFlags(BULKDATA_Force_NOT_InlinePayload);
+				DDCRebuildState.store(EDDCRebuildState::Succeeded);
+			}
+			else
+			{
+				DDCRebuildState.store(EDDCRebuildState::Failed);
+			}
+		});
 }
 
-void USparseVolumeTexture::GetPackedUniforms(FUintVector4& OutPacked0, FUintVector4& OutPacked1) const
+void FResources::EndRebuildBulkDataFromCache()
 {
-	FIntVector3 PageTableOffset = FIntVector3::ZeroValue;
-	FVector3f TileDataTexelSize = FVector3f(0.0f, 0.0f, 0.0f);
-	int32 MinMipLevel = 0;
-	int32 MaxMipLevel = 0;
-	const FSparseVolumeTextureSceneProxy* Proxy = GetSparseVolumeTextureSceneProxy();
-	if (Proxy)
+	if (*DDCRequestOwner)
 	{
-		const FSparseVolumeTextureRuntimeHeader& Header = Proxy->GetHeader();
-		PageTableOffset = Header.PageTableVolumeAABBMin;
-		TileDataTexelSize.X = 1.0f / Header.TileDataVolumeResolution.X;
-		TileDataTexelSize.Y = 1.0f / Header.TileDataVolumeResolution.Y;
-		TileDataTexelSize.Z = 1.0f / Header.TileDataVolumeResolution.Z;
-		MinMipLevel = Header.LowestResidentLevel;
-		MaxMipLevel = Header.HighestResidentLevel;
+		(*DDCRequestOwner)->Wait();
+		(*DDCRequestOwner).Reset();
 	}
-	const FIntVector3 VolumeResolution = GetVolumeResolution();
-	const FVector3f VolumePageResolution = FVector3f(VolumeResolution) / SPARSE_VOLUME_TILE_RES;
+	DDCRebuildState.store(EDDCRebuildState::Initial);
+}
+
+#endif // WITH_EDITORONLY_DATA
+
+////////////////////////////////////////////////////////////////////////////////////////////////
+
+void FTextureRenderResources::GetPackedUniforms(FUintVector4& OutPacked0, FUintVector4& OutPacked1) const
+{
+	check(IsInParallelRenderingThread());
 
 	auto AsUint = [](float X)
 	{
@@ -174,37 +539,42 @@ void USparseVolumeTexture::GetPackedUniforms(FUintVector4& OutPacked0, FUintVect
 		return FU.U;
 	};
 
+	const FIntVector3 PageTableOffset = Header.PageTableVolumeAABBMin;
+	const FVector3f TileDataTexelSize = FVector3f(
+		1.0f / TileDataTextureResolution.X,
+		1.0f / TileDataTextureResolution.Y,
+		1.0f / TileDataTextureResolution.Z);
+	const FVector3f VolumePageResolution = FVector3f(GlobalVolumeResolution) / SPARSE_VOLUME_TILE_RES;
+
 	OutPacked0.X = AsUint(VolumePageResolution.X);
 	OutPacked0.Y = AsUint(VolumePageResolution.Y);
 	OutPacked0.Z = AsUint(VolumePageResolution.Z);
-	OutPacked0.W = UE::SVT::PackPageTableEntry(PageTableOffset);
+	OutPacked0.W = (PageTableOffset.X & 0x7FFu) | ((PageTableOffset.Y & 0x7FFu) << 11u) | ((PageTableOffset.Z & 0x3FFu) << 22u);
 	OutPacked1.X = AsUint(TileDataTexelSize.X);
 	OutPacked1.Y = AsUint(TileDataTexelSize.Y);
 	OutPacked1.Z = AsUint(TileDataTexelSize.Z);
 	OutPacked1.W = 0;
-	OutPacked1.W |= (uint32)((MinMipLevel & 0xFF) << 0);
-	OutPacked1.W |= (uint32)((MaxMipLevel & 0xFF) << 8);
-	OutPacked1.W |= (uint32)((int32(SPARSE_VOLUME_TILE_RES) & 0xFF) << 16);
-	OutPacked1.W |= (uint32)((int32(SPARSE_VOLUME_TILE_BORDER) & 0xFF) << 24);
+	OutPacked1.W |= (uint32)((FrameIndex & 0xFFFF) << 0);
+	OutPacked1.W |= (uint32)(((NumLogicalMipLevels - 1) & 0xFFFF) << 16);
 }
 
-void USparseVolumeTexture::GetFrameUVScaleBias(FVector* OutScale, FVector* OutBias) const
+void FTextureRenderResources::SetGlobalVolumeResolution_GameThread(const FIntVector3& InGlobalVolumeResolution)
 {
-	*OutScale = FVector::One();
-	*OutBias = FVector::Zero();
-	const FSparseVolumeTextureSceneProxy* Proxy = GetSparseVolumeTextureSceneProxy();
-	if (Proxy)
-	{
-		const FSparseVolumeTextureRuntimeHeader& Header = Proxy->GetHeader();
-		const FVector GlobalVolumeRes = FVector(GetVolumeResolution());
-		check(GlobalVolumeRes.X > 0.0 && GlobalVolumeRes.Y > 0.0 && GlobalVolumeRes.Z > 0.0);
-		const FVector FrameBoundsPaddedMin = FVector(Header.PageTableVolumeAABBMin * SPARSE_VOLUME_TILE_RES); // padded to multiple of page size
-		const FVector FrameBoundsPaddedMax = FVector(Header.PageTableVolumeAABBMax * SPARSE_VOLUME_TILE_RES);
-		const FVector FramePaddedSize = FrameBoundsPaddedMax - FrameBoundsPaddedMin;
+	ENQUEUE_RENDER_COMMAND(FTextureRenderResources_UpdateGlobalVolumeResolution)(
+		[this, InGlobalVolumeResolution](FRHICommandListImmediate& RHICmdList)
+		{
+			GlobalVolumeResolution = InGlobalVolumeResolution;
+		});
+}
 
-		*OutScale = GlobalVolumeRes / FramePaddedSize; // scale from SVT UV space to frame (padded) local UV space
-		*OutBias = -(FrameBoundsPaddedMin / GlobalVolumeRes * *OutScale);
-	}
+}
+}
+
+////////////////////////////////////////////////////////////////////////////////////////////////
+
+USparseVolumeTexture::USparseVolumeTexture(const FObjectInitializer& ObjectInitializer)
+	: Super(ObjectInitializer)
+{
 }
 
 UE::Shader::EValueType USparseVolumeTexture::GetUniformParameterType(int32 Index)
@@ -273,23 +643,24 @@ USparseVolumeTextureFrame::USparseVolumeTextureFrame(const FObjectInitializer& O
 {
 }
 
-USparseVolumeTextureFrame* USparseVolumeTextureFrame::GetFrame(USparseVolumeTexture* SparseVolumeTexture, int32 FrameIndex)
+USparseVolumeTextureFrame* USparseVolumeTextureFrame::GetFrameAndIssueStreamingRequest(USparseVolumeTexture* SparseVolumeTexture, float FrameIndex, int32 MipLevel)
 {
 	if (UStreamableSparseVolumeTexture* StreamableSVT = Cast<UStreamableSparseVolumeTexture>(SparseVolumeTexture))
 	{
-		return StreamableSVT->GetFrame(FrameIndex);
+		UE::SVT::GStreamingManager.Request_GameThread(StreamableSVT, FrameIndex, MipLevel);
+		return StreamableSVT->GetFrame(static_cast<int32>(FrameIndex));
 	}
 	return nullptr;
 }
 
-bool USparseVolumeTextureFrame::Initialize(USparseVolumeTexture* InOwner, int32 InFrameIndex, FSparseVolumeTextureData& UncookedFrame)
+bool USparseVolumeTextureFrame::Initialize(USparseVolumeTexture* InOwner, int32 InFrameIndex, UE::SVT::FTextureData& UncookedFrame)
 {
 #if WITH_EDITORONLY_DATA
 	Owner = InOwner;
 	FrameIndex = InFrameIndex;
 	{
-		UE::Serialization::FEditorBulkDataWriter RawDataArchiveWriter(RawData);
-		UncookedFrame.Serialize(RawDataArchiveWriter);
+		UE::Serialization::FEditorBulkDataWriter SourceDataArchiveWriter(Resources.SourceData);
+		UncookedFrame.Serialize(SourceDataArchiveWriter);
 	}
 
 	return true;
@@ -298,42 +669,24 @@ bool USparseVolumeTextureFrame::Initialize(USparseVolumeTexture* InOwner, int32
 #endif
 }
 
-bool USparseVolumeTextureFrame::BuildDerivedData(FSparseVolumeTextureData* OutMippedTextureData)
+bool USparseVolumeTextureFrame::CreateTextureRenderResources()
 {
-#if WITH_EDITORONLY_DATA
-	// Check if the virtualized bulk data payload is available now
-	if (RawData.HasPayloadData())
-	{
-		// First, read the source data in from the raw data stored as bulk data
-		UE::Serialization::FEditorBulkDataReader RawDataArchiveReader(RawData);
-		FSparseVolumeTextureData TextureData;
-		TextureData.Serialize(RawDataArchiveReader);
-
-		FSparseVolumeTextureDataAddressingInfo AddressingInfo{};
-		AddressingInfo.VolumeResolution = GetVolumeResolution();
-		AddressingInfo.AddressX = GetTextureAddressX();
-		AddressingInfo.AddressY = GetTextureAddressY();
-		AddressingInfo.AddressZ = GetTextureAddressZ();
-
-		const int32 NumMipLevels = -1; // generate entire mip chain
-		const bool bMoveMip0FromSource = true; // we have no need to keep TextureData around
-		if (!TextureData.BuildDerivedData(AddressingInfo, NumMipLevels, bMoveMip0FromSource, *OutMippedTextureData))
-		{
-			return false;
-		}
-
-		// Now unload the raw data
-		RawData.UnloadData();
+	if (!TextureRenderResources)
+	{
+		TextureRenderResources = new UE::SVT::FTextureRenderResources();
+		TextureRenderResources->SetGlobalVolumeResolution_GameThread(Owner->GetVolumeResolution());
+		BeginInitResource(TextureRenderResources);
 
 		return true;
 	}
-#endif
 	return false;
 }
 
 void USparseVolumeTextureFrame::PostLoad()
 {
 	Super::PostLoad();
+
+	CreateTextureRenderResources();
 }
 
 void USparseVolumeTextureFrame::FinishDestroy()
@@ -343,15 +696,21 @@ void USparseVolumeTextureFrame::FinishDestroy()
 
 void USparseVolumeTextureFrame::BeginDestroy()
 {
-	if (SceneProxy)
+	// Ensure that the streamable SVT has been removed from the streaming manager
+	if (IsValid(Owner))
 	{
-		ENQUEUE_RENDER_COMMAND(USparseVolumeTextureFrame_DeleteSVTProxy)(
-			[Proxy = SceneProxy](FRHICommandListImmediate& RHICmdList)
+		UE::SVT::GStreamingManager.Remove_GameThread(CastChecked<UStreamableSparseVolumeTexture>(Owner));
+	}
+	
+	if (TextureRenderResources)
+	{
+		ENQUEUE_RENDER_COMMAND(USparseVolumeTextureFrame_DeleteTextureRenderResources)(
+			[Resources = TextureRenderResources](FRHICommandListImmediate& RHICmdList)
 			{
-				Proxy->ReleaseResource();
-				delete Proxy;
+				Resources->ReleaseResource();
+				delete Resources;
 			});
-		SceneProxy = nullptr;
+		TextureRenderResources = nullptr;
 	}
 
 	Super::BeginDestroy();
@@ -362,64 +721,11 @@ void USparseVolumeTextureFrame::Serialize(FArchive& Ar)
 	Super::Serialize(Ar);
 
 	FStripDataFlags StripFlags(Ar);
-
-	const bool bInlinePayload = (FrameIndex == 0);
-	StreamingData.SetBulkDataFlags(bInlinePayload ? BULKDATA_ForceInlinePayload : BULKDATA_Force_NOT_InlinePayload);
-
-	if (StripFlags.IsEditorDataStripped() && Ar.IsLoadingFromCookedPackage())
-	{
-		// In this case we are loading in game with a cooked build so we only need to load the runtime data.
-
-		// Read cooked bulk data from archive
-		StreamingData.Serialize(Ar, Owner);
-
-		if (bInlinePayload)
-		{
-			SceneProxy = new FSparseVolumeTextureSceneProxy();
-
-			// Create runtime data from cooked bulk data
-			{
-				FBulkDataReader BulkDataReader(StreamingData);
-				FSparseVolumeTextureData TextureData;
-				TextureData.Serialize(BulkDataReader);
-				bool bSuccess = SceneProxy->GetRuntimeData().Create(TextureData);
-				check(bSuccess); // SVT_TODO
-			}
-
-			// The bulk data is no longer needed
-			StreamingData.RemoveBulkData();
-
-			// Runtime data is now valid, initialize the render thread proxy
-			BeginInitResource(SceneProxy);
-		}
-	}
-	else if (Ar.IsCooking())
-	{
-		// We are cooking the game, serialize the asset out.
-
-		FSparseVolumeTextureData DerivedData;
-		const bool bBuiltDerivedData = BuildDerivedData(&DerivedData);
-		check(bBuiltDerivedData); // SVT_TODO: actual error handling
-
-		// Write derived data into StreamingData
-		{
-			FBulkDataWriter BulkDataWriter(StreamingData);
-			DerivedData.Serialize(BulkDataWriter);
-		}
-
-		// And now write the cooked bulk data to the archive
-		StreamingData.Serialize(Ar, Owner);
-	}
-	else if (!Ar.IsObjectReferenceCollector())
-	{
-#if WITH_EDITORONLY_DATA
-		// When in EDITOR:
-		//  - We only serialize raw data 
-		//  - The runtime data is fetched/put from/to DDC
-		//  - This EditorBulk data do not load the full and huge OpenVDB data. That is only done explicitly later.
-		RawData.Serialize(Ar, Owner);
-#endif
-	}
+	
+	bool bCooked = Ar.IsCooking();
+	Ar << bCooked;
+	Ar << FrameIndex;
+	Resources.Serialize(Ar, this, bCooked);
 }
 
 void USparseVolumeTextureFrame::GetResourceSizeEx(FResourceSizeEx& CumulativeResourceSize)
@@ -427,97 +733,54 @@ void USparseVolumeTextureFrame::GetResourceSizeEx(FResourceSizeEx& CumulativeRes
 	Super::GetResourceSizeEx(CumulativeResourceSize);
 }
 
-void USparseVolumeTextureFrame::GenerateOrLoadDDCRuntimeData(UE::DerivedData::FRequestOwner& DDCRequestOwner)
+#if WITH_EDITOR
+
+void USparseVolumeTextureFrame::BeginCacheForCookedPlatformData(const ITargetPlatform* TargetPlatform)
 {
-#if WITH_EDITORONLY_DATA
-	using namespace UE::DerivedData;
+}
 
-	// Release any previously allocated render thread proxy
-	if (SceneProxy)
+bool USparseVolumeTextureFrame::IsCachedCookedPlatformDataLoaded(const ITargetPlatform* TargetPlatform)
+{
+	bool bFailed = false;
+	if (!Resources.RebuildBulkDataFromCacheAsync(this, bFailed))
 	{
-		BeginReleaseResource(SceneProxy);
+		return false;
 	}
-	else
+
+	if (bFailed)
 	{
-		SceneProxy = new FSparseVolumeTextureSceneProxy();
+		UE_LOG(LogSparseVolumeTexture, Log, TEXT("Failed to recover SparseVolumeTexture streaming from DDC for '%s' Frame %i. Rebuilding and retrying."), *Owner->GetPathName(), FrameIndex);
+
+		Resources.Cache(this);
+		return false;
 	}
 
-	static const FString SparseVolumeTextureDDCVersion = TEXT("381AE2A9-A903-4C8F-8486-891E24D6EC70");	// Bump this if you want to ignore all cached data so far.
-	const FString DerivedDataKey = RawData.GetIdentifier().ToString()
-		+ FString::Format(TEXT("{0},{1},{2}"), { GetTextureAddressX(), GetTextureAddressY(), GetTextureAddressZ() })
-		+ SparseVolumeTextureDDCVersion;
+	return true;
+}
 
-	const FCacheKey Key = ConvertLegacyCacheKey(DerivedDataKey);
-	const FSharedString Name = MakeStringView(GetPathName());
+void USparseVolumeTextureFrame::WillNeverCacheCookedPlatformDataAgain()
+{
+}
 
-	UE::DerivedData::GetCache().GetValue({ {Name, Key} }, DDCRequestOwner,
-		[this, &DDCRequestOwner](FCacheGetValueResponse&& Response)
-		{
-			if (Response.Status == EStatus::Ok)
-			{
-				DDCRequestOwner.LaunchTask(TEXT("USparseVolumeTextureFrame_DerivedDataLoad"),
-					[this, Value = MoveTemp(Response.Value)]()
-					{
-						FSharedBuffer Data = Value.GetData().Decompress();
-						FMemoryReaderView Ar(Data, true /*bIsPersistent*/);
-						FSparseVolumeTextureData TextureData;
-						TextureData.Serialize(Ar);
-						bool bSuccess = SceneProxy->GetRuntimeData().Create(TextureData);
-						check(bSuccess); // SVT_TODO
-
-						// Runtime data is now valid, initialize the render thread proxy
-						BeginInitResource(SceneProxy);
-					});
-			}
-			else if (Response.Status == EStatus::Error)
-			{
-				DDCRequestOwner.LaunchTask(TEXT("USparseVolumeTextureFrame_DerivedDataBuild"),
-					[this, &DDCRequestOwner, Name = Response.Name, Key = Response.Key]()
-					{
-						FSparseVolumeTextureRuntime& RuntimeData = SceneProxy->GetRuntimeData();
+void USparseVolumeTextureFrame::ClearCachedCookedPlatformData(const ITargetPlatform* TargetPlatform)
+{
+	Resources.DropBulkData();
+}
 
-						// Check if the virtualized bulk data payload is available now
-						if (RawData.HasPayloadData())
-						{
-							FSparseVolumeTextureData TextureData;
-							bool bSuccess = BuildDerivedData(&TextureData);
-							ensure(bSuccess);
-
-							bSuccess = RuntimeData.Create(TextureData);
-							ensure(bSuccess);
-
-							// Using a LargeMemoryWriter for serialization since the data can be bigger than 2 GB
-							FLargeMemoryWriter LargeMemWriter(0, /*bIsPersistent=*/ true);
-							TextureData.Serialize(LargeMemWriter);
-
-							const int64 UncompressedSize = LargeMemWriter.TotalSize();
-
-							// Since the DDC doesn't support data bigger than 2 GB, we only cache for such uncompressed size.
-							constexpr int64 SizeThreshold = 2147483648LL;	// 2GB
-							const bool bIsCacheable = UncompressedSize < SizeThreshold;
-							if (bIsCacheable)
-							{
-								FValue Value = FValue::Compress(FSharedBuffer::MakeView(LargeMemWriter.GetView()));
-								UE::DerivedData::GetCache().PutValue({ {Name, Key, Value} }, DDCRequestOwner);
-							}
-							else
-							{
-								UE_LOG(LogSparseVolumeTexture, Error, TEXT("SparseVolumeTexture - the asset is too large to fit in Derived Data Cache %s"), *GetName());
-							}
-						}
-						else
-						{
-							UE_LOG(LogSparseVolumeTexture, Error, TEXT("SparseVolumeTexture - Raw source data is not available for %s. Using default data."), *GetName());
-							RuntimeData.SetAsDefaultTexture();
-						}
+void USparseVolumeTextureFrame::ClearAllCachedCookedPlatformData()
+{
+	Resources.DropBulkData();
+}
 
-						// Runtime data is now valid, initialize the render thread proxy
-						BeginInitResource(SceneProxy);
-					});
-			}
-		});
-#endif // WITH_EDITORONLY_DATA
+#endif // WITH_EDITOR
+
+#if WITH_EDITORONLY_DATA
+void USparseVolumeTextureFrame::Cache()
+{
+	Resources.Cache(this);
+	CreateTextureRenderResources();
 }
+#endif
 
 ////////////////////////////////////////////////////////////////////////////////////////////////
 
@@ -539,6 +802,8 @@ bool UStreamableSparseVolumeTexture::BeginInitialize(int32 NumExpectedFrames)
 	Frames.Empty(NumExpectedFrames);
 	VolumeBoundsMin = FIntVector(INT32_MAX, INT32_MAX, INT32_MAX);
 	VolumeBoundsMax = FIntVector(INT32_MIN, INT32_MIN, INT32_MIN);
+	check(FormatA == PF_Unknown);
+	check(FormatB == PF_Unknown);
 
 	InitState = EInitState::Pending;
 
@@ -548,7 +813,7 @@ bool UStreamableSparseVolumeTexture::BeginInitialize(int32 NumExpectedFrames)
 #endif
 }
 
-bool UStreamableSparseVolumeTexture::AppendFrame(FSparseVolumeTextureData& UncookedFrame)
+bool UStreamableSparseVolumeTexture::AppendFrame(UE::SVT::FTextureData& UncookedFrame)
 {
 #if WITH_EDITORONLY_DATA
 	if (InitState != EInitState::Pending)
@@ -566,6 +831,29 @@ bool UStreamableSparseVolumeTexture::AppendFrame(FSparseVolumeTextureData& Uncoo
 		return false;
 	}
 
+	// SVT_TODO: Valide formats against list of supported formats
+	if (Frames.IsEmpty())
+	{
+		FormatA = UncookedFrame.Header.AttributesFormats[0];
+		FormatB = UncookedFrame.Header.AttributesFormats[1];
+		FallbackValueA = UncookedFrame.Header.FallbackValues[0];
+		FallbackValueB = UncookedFrame.Header.FallbackValues[1];
+	}
+	else
+	{
+		if (UncookedFrame.Header.AttributesFormats[0] != FormatA || UncookedFrame.Header.AttributesFormats[1] != FormatB)
+		{
+			UE_LOG(LogSparseVolumeTexture, Error, TEXT("Tried to add a frame to a SparseVolumeTexture without matching formats! Expected: (%i, %i), Actual: (%i %i)"),
+				(int)FormatA.GetIntValue(), (int)FormatB.GetIntValue(), (int)UncookedFrame.Header.AttributesFormats[0], (int)UncookedFrame.Header.AttributesFormats[1]);
+			return false;
+		}
+		if (UncookedFrame.Header.FallbackValues[0] != FallbackValueA || UncookedFrame.Header.FallbackValues[1] != FallbackValueB)
+		{
+			UE_LOG(LogSparseVolumeTexture, Error, TEXT("Tried to add a frame to a SparseVolumeTexture without matching fallback/null tile values!"));
+			return false;
+		}
+	}
+
 	// Compute union of all frame AABBs
 	VolumeBoundsMin.X = FMath::Min(VolumeBoundsMin.X, UncookedFrame.Header.VirtualVolumeAABBMin.X);
 	VolumeBoundsMin.Y = FMath::Min(VolumeBoundsMin.Y, UncookedFrame.Header.VirtualVolumeAABBMin.Y);
@@ -602,13 +890,14 @@ bool UStreamableSparseVolumeTexture::EndInitialize(int32 InNumMipLevels)
 	if (Frames.IsEmpty())
 	{
 		UE_LOG(LogSparseVolumeTexture, Warning, TEXT("SVT has zero frames! Adding a dummy frame. SVT: %s"), *GetName());
-		FSparseVolumeTextureData DummyFrame;
+		UE::SVT::FTextureData DummyFrame;
 		DummyFrame.CreateDefault();
 		AppendFrame(DummyFrame);
 	}
 
 	check(VolumeResolution.X > 0 && VolumeResolution.Y > 0 && VolumeResolution.Z > 0);
 	check(VolumeBoundsMin.X >= 0 && VolumeBoundsMin.Y >= 0 && VolumeBoundsMin.Z >= 0);
+	check(FormatA != PF_Unknown || FormatB != PF_Unknown);
 
 	if (VolumeBoundsMin.X > 0 || VolumeBoundsMin.Y > 0 || VolumeBoundsMin.Z > 0)
 	{
@@ -616,11 +905,16 @@ bool UStreamableSparseVolumeTexture::EndInitialize(int32 InNumMipLevels)
 			VolumeBoundsMin.X, VolumeBoundsMin.Y, VolumeBoundsMin.Z);
 	}
 
-	const int32 NumMipLevelsFullMipChain = SVTComputeNumMipLevels(VolumeResolution);
+	const int32 NumMipLevelsFullMipChain = UE::SVT::ComputeNumMipLevels(VolumeResolution);
 	check(NumMipLevelsFullMipChain > 0);
 
 	NumMipLevels = (InNumMipLevels <= INDEX_NONE) ? NumMipLevelsFullMipChain : FMath::Clamp(InNumMipLevels, 1, NumMipLevelsFullMipChain);
 
+	for (USparseVolumeTextureFrame* Frame : Frames)
+	{
+		Frame->PostLoad();
+	}
+
 	InitState = EInitState::Done;
 
 	return true;
@@ -629,7 +923,7 @@ bool UStreamableSparseVolumeTexture::EndInitialize(int32 InNumMipLevels)
 #endif
 }
 
-bool UStreamableSparseVolumeTexture::Initialize(const TArrayView<FSparseVolumeTextureData>& InUncookedData, int32 InNumMipLevels)
+bool UStreamableSparseVolumeTexture::Initialize(const TArrayView<UE::SVT::FTextureData>& InUncookedData, int32 InNumMipLevels)
 {
 	if (InUncookedData.IsEmpty())
 	{
@@ -641,7 +935,7 @@ bool UStreamableSparseVolumeTexture::Initialize(const TArrayView<FSparseVolumeTe
 	{
 		return false;
 	}
-	for (FSparseVolumeTextureData& UncookedFrame : InUncookedData)
+	for (UE::SVT::FTextureData& UncookedFrame : InUncookedData)
 	{
 		if (!AppendFrame(UncookedFrame))
 		{
@@ -656,65 +950,30 @@ bool UStreamableSparseVolumeTexture::Initialize(const TArrayView<FSparseVolumeTe
 	return true;
 }
 
-const FSparseVolumeTextureSceneProxy* UStreamableSparseVolumeTexture::GetStreamedFrameProxyOrFallback(int32 FrameIndex, int32 MipLevel) const
-{
-#if WITH_EDITORONLY_DATA
-	if (InitState != EInitState::Done)
-	{
-		UE_LOG(LogSparseVolumeTexture, Warning, TEXT("Tried to call GetStreamedFrameProxyOrFallback() on uninitialized SVT: %s"), *GetName());
-		return nullptr;
-	}
-#endif
-	if (Frames.IsEmpty())
-	{
-		UE_LOG(LogSparseVolumeTexture, Warning, TEXT("SVT is empty and has no frames at all! SVT: %s"), *GetName());
-		return nullptr;
-	}
-	FrameIndex = FMath::Clamp(FrameIndex, 0, Frames.Num() - 1);
-
-#if WITH_EDITORONLY_DATA
-	return Frames[FrameIndex]->GetSparseVolumeTextureSceneProxy();
-#else
-	ISparseVolumeTextureStreamingManager& StreamingManager = IStreamingManager::Get().GetSparseVolumeTextureStreamingManager();
-	const FSparseVolumeTextureSceneProxy* Proxy = StreamingManager.GetSparseVolumeTextureSceneProxy(this, FrameIndex, MipLevel, true);
-
-	int32 FallbackFrameIndex = FrameIndex;
-	while (!Proxy)
-	{
-		FallbackFrameIndex = FallbackFrameIndex > 0 ? (FallbackFrameIndex - 1) : (Frames.Num() - 1);
-		if (FallbackFrameIndex == FrameIndex)
-		{
-			UE_LOG(LogSparseVolumeTexture, Warning, TEXT("Failed to get ANY streamed SparseVolumeTexture frame  SVT: %s, FrameIndex: %i"), *GetName(), FrameIndex);
-			return nullptr;
-		}
-		Proxy = StreamingManager.GetSparseVolumeTextureSceneProxy(this, FallbackFrameIndex, MipLevel, false);
-	}
-
-	return Proxy;
-#endif
-}
-
 void UStreamableSparseVolumeTexture::PostLoad()
 {
 	Super::PostLoad();
 
 #if WITH_EDITORONLY_DATA
-	GenerateOrLoadDDCRuntimeDataAndCreateSceneProxy();
+	RecacheFrames();
 #else
-	IStreamingManager::Get().GetSparseVolumeTextureStreamingManager().AddSparseVolumeTexture(this); // GenerateOrLoadDDCRuntimeDataAndCreateSceneProxy() handles this in editor builds
+	for (USparseVolumeTextureFrame* Frame : Frames)
+	{
+		Frame->CreateTextureRenderResources();
+	}
+	UE::SVT::GStreamingManager.Add_GameThread(this); // RecacheFrames() handles this in editor builds
 #endif
 }
 
 void UStreamableSparseVolumeTexture::FinishDestroy()
 {
 	Super::FinishDestroy();
-
-	IStreamingManager::Get().GetSparseVolumeTextureStreamingManager().RemoveSparseVolumeTexture(this);
 }
 
 void UStreamableSparseVolumeTexture::BeginDestroy()
 {
 	Super::BeginDestroy();
+	UE::SVT::GStreamingManager.Remove_GameThread(this);
 }
 
 void UStreamableSparseVolumeTexture::Serialize(FArchive& Ar)
@@ -765,51 +1024,50 @@ void UStreamableSparseVolumeTexture::PostEditChangeProperty(FPropertyChangedEven
 
 	Super::PostEditChangeProperty(PropertyChangedEvent);
 
-	GenerateOrLoadDDCRuntimeDataAndCreateSceneProxy();
+	RecacheFrames();
 }
 #endif // WITH_EDITOR
 
 void UStreamableSparseVolumeTexture::GetResourceSizeEx(FResourceSizeEx& CumulativeResourceSize)
 {
 	Super::GetResourceSizeEx(CumulativeResourceSize);
-	SIZE_T SizeCPU = sizeof(*this);
+	SIZE_T SizeCPU = sizeof(*this) - sizeof(Super);
 	SIZE_T SizeGPU = 0;
 	SizeCPU += Frames.GetAllocatedSize();
 	for (USparseVolumeTextureFrame* Frame : Frames)
 	{
 		Frame->GetResourceSizeEx(CumulativeResourceSize);
 	}
-	ISparseVolumeTextureStreamingManager& StreamingManager = IStreamingManager::Get().GetSparseVolumeTextureStreamingManager();
-	StreamingManager.GetMemorySizeForSparseVolumeTexture(this, &SizeCPU, &SizeGPU);
 	CumulativeResourceSize.AddDedicatedSystemMemoryBytes(SizeCPU);
 	CumulativeResourceSize.AddDedicatedVideoMemoryBytes(SizeGPU);
 }
 
-void UStreamableSparseVolumeTexture::GenerateOrLoadDDCRuntimeDataAndCreateSceneProxy()
-{
 #if WITH_EDITORONLY_DATA
-
+void UStreamableSparseVolumeTexture::RecacheFrames()
+{
 	if (InitState != EInitState::Done)
 	{
 		UE_LOG(LogSparseVolumeTexture, Warning, TEXT("Tried to cache derived data of an uninitialized SVT: %s"), *GetName());
 		return;
 	}
 
-	UE::DerivedData::FRequestOwner DDCRequestOwner(UE::DerivedData::EPriority::Normal);
+	// SVT_TODO: This shows the user that something is actually happening and the editor did not freeze. The cancel button is currently ignored.
+	FScopedSlowTask RecacheTask(static_cast<float>(Frames.Num() + 2), LOCTEXT("SparseVolumeTextureCacheFrames", "Caching SparseVolumeTexture frames in Derived Data Cache"));
+	RecacheTask.MakeDialog(true);
+
+	UE::SVT::GStreamingManager.Remove_GameThread(this);
+	RecacheTask.EnterProgressFrame(1.0f);
+
+	for (USparseVolumeTextureFrame* Frame : Frames)
 	{
-		UE::DerivedData::FRequestBarrier DDCRequestBarrier(DDCRequestOwner);
-		for (USparseVolumeTextureFrame* Frame : Frames)
-		{
-			Frame->GenerateOrLoadDDCRuntimeData(DDCRequestOwner);
-		}
+		RecacheTask.EnterProgressFrame(1.0f);
+		Frame->Cache();
 	}
 
-	// Wait for all DDC requests to complete
-	DDCRequestOwner.Wait();
-
-	IStreamingManager::Get().GetSparseVolumeTextureStreamingManager().AddSparseVolumeTexture(this);
-#endif
+	RecacheTask.EnterProgressFrame(1.0f);
+	UE::SVT::GStreamingManager.Add_GameThread(this);
 }
+#endif
 
 ////////////////////////////////////////////////////////////////////////////////////////////////
 
@@ -818,7 +1076,7 @@ UStaticSparseVolumeTexture::UStaticSparseVolumeTexture(const FObjectInitializer&
 {
 }
 
-bool UStaticSparseVolumeTexture::AppendFrame(FSparseVolumeTextureData& UncookedFrame)
+bool UStaticSparseVolumeTexture::AppendFrame(UE::SVT::FTextureData& UncookedFrame)
 {
 	if (!Frames.IsEmpty())
 	{
@@ -835,15 +1093,6 @@ UAnimatedSparseVolumeTexture::UAnimatedSparseVolumeTexture(const FObjectInitiali
 {
 }
 
-const FSparseVolumeTextureSceneProxy* UAnimatedSparseVolumeTexture::GetSparseVolumeTextureSceneProxy() const
-{
-	// When an AnimatedSparseVolumeTexture is used as SparseVolumeTexture, it can only be previewed using a single preview frame.
-	check(!Frames.IsEmpty());
-	const int32 FrameIndex = PreviewFrameIndex % Frames.Num();
-	const int32 MipLevel = FMath::Clamp(PreviewMipLevel, 0, GetNumMipLevels() - 1);
-	return GetStreamedFrameProxyOrFallback(FrameIndex, MipLevel);
-}
-
 ////////////////////////////////////////////////////////////////////////////////////////////////
 
 UAnimatedSparseVolumeTextureController::UAnimatedSparseVolumeTextureController(const FObjectInitializer& ObjectInitializer)
@@ -901,7 +1150,7 @@ USparseVolumeTextureFrame* UAnimatedSparseVolumeTextureController::GetFrameByInd
 		return nullptr;
 	}
 
-	return USparseVolumeTextureFrame::GetFrame(SparseVolumeTexture, FrameIndex);
+	return USparseVolumeTextureFrame::GetFrameAndIssueStreamingRequest(SparseVolumeTexture, FrameIndex, MipLevel);
 }
 
 USparseVolumeTextureFrame* UAnimatedSparseVolumeTextureController::GetCurrentFrame()
@@ -913,9 +1162,8 @@ USparseVolumeTextureFrame* UAnimatedSparseVolumeTextureController::GetCurrentFra
 
 	// Compute (fractional) index of frame to sample
 	const float FrameIndexF = GetFractionalFrameIndex();
-	const int32 FrameIndex = (int32)FrameIndexF;
 
-	return USparseVolumeTextureFrame::GetFrame(SparseVolumeTexture, FrameIndex);
+	return USparseVolumeTextureFrame::GetFrameAndIssueStreamingRequest(SparseVolumeTexture, FrameIndexF, MipLevel);
 }
 
 void UAnimatedSparseVolumeTextureController::GetCurrentFramesForInterpolation(USparseVolumeTextureFrame*& Frame0, USparseVolumeTextureFrame*& Frame1, float& LerpAlpha)
@@ -930,8 +1178,8 @@ void UAnimatedSparseVolumeTextureController::GetCurrentFramesForInterpolation(US
 	const int32 FrameIndex = (int32)FrameIndexF;
 	LerpAlpha = FMath::Frac(FrameIndexF);
 
-	Frame0 = USparseVolumeTextureFrame::GetFrame(SparseVolumeTexture, FrameIndex);
-	Frame1 = USparseVolumeTextureFrame::GetFrame(SparseVolumeTexture, (FrameIndex + 1) % SparseVolumeTexture->GetNumFrames());
+	Frame0 = USparseVolumeTextureFrame::GetFrameAndIssueStreamingRequest(SparseVolumeTexture, FrameIndexF, MipLevel);
+	Frame1 = USparseVolumeTextureFrame::GetFrameAndIssueStreamingRequest(SparseVolumeTexture, (FrameIndex + 1) % SparseVolumeTexture->GetNumFrames(), MipLevel);
 }
 
 float UAnimatedSparseVolumeTextureController::GetDuration()
diff --git a/Engine/Source/Runtime/Engine/Private/SparseVolumeTexture/SparseVolumeTextureData.cpp b/Engine/Source/Runtime/Engine/Private/SparseVolumeTexture/SparseVolumeTextureData.cpp
--- a/Engine/Source/Runtime/Engine/Private/SparseVolumeTexture/SparseVolumeTextureData.cpp
+++ b/Engine/Source/Runtime/Engine/Private/SparseVolumeTexture/SparseVolumeTextureData.cpp
@@ -16,7 +16,12 @@ DEFINE_LOG_CATEGORY_STATIC(LogSparseVolumeTextureData, Log, All);
 
 ////////////////////////////////////////////////////////////////////////////////////////////////
 
-void FSparseVolumeTextureData::Serialize(FArchive& Ar)
+namespace UE
+{
+namespace SVT
+{
+
+void FTextureData::Serialize(FArchive& Ar)
 {
 	Header.Serialize(Ar);
 
@@ -42,16 +47,16 @@ void FSparseVolumeTextureData::Serialize(FArchive& Ar)
 	}
 	else
 	{
-		// FSparseVolumeTextureData needs to account for new version
+		// FTextureData needs to account for new version
 		check(false);
 	}
 }
 
-bool FSparseVolumeTextureData::Create(const ISparseVolumeTextureDataProvider& DataProvider)
+bool FTextureData::Create(const ITextureDataProvider& DataProvider)
 {
-	const FSparseVolumeTextureDataCreateInfo CreateInfo = DataProvider.GetCreateInfo();
+	const FTextureDataCreateInfo CreateInfo = DataProvider.GetCreateInfo();
 
-	Header = FSparseVolumeTextureHeader(CreateInfo.VirtualVolumeAABBMin, CreateInfo.VirtualVolumeAABBMax, CreateInfo.AttributesFormats[0], CreateInfo.AttributesFormats[1], CreateInfo.FallbackValues[0], CreateInfo.FallbackValues[1]);
+	Header = UE::SVT::FHeader(CreateInfo.VirtualVolumeAABBMin, CreateInfo.VirtualVolumeAABBMax, CreateInfo.AttributesFormats[0], CreateInfo.AttributesFormats[1], CreateInfo.FallbackValues[0], CreateInfo.FallbackValues[1]);
 
 	if (Header.PageTableVolumeResolution.X > UE::SVT::SVTMaxVolumeTextureDim
 		|| Header.PageTableVolumeResolution.Y > UE::SVT::SVTMaxVolumeTextureDim
@@ -160,15 +165,15 @@ bool FSparseVolumeTextureData::Create(const ISparseVolumeTextureDataProvider& Da
 							for (int32 X = 0; X < SPARSE_VOLUME_TILE_RES_PADDED; ++X)
 							{
 								const FIntVector3 WriteCoord = FIntVector3(X - 1, Y - 1, Z - 1);
-								WriteTileDataVoxel(TileIndex, WriteCoord, 0 /*MipLevel*/, 0 /*AttributesIndex*/, Header.NullTileValues[0]);
-								WriteTileDataVoxel(TileIndex, WriteCoord, 0 /*MipLevel*/, 1 /*AttributesIndex*/, Header.NullTileValues[1]);
+								WriteTileDataVoxel(TileIndex, WriteCoord, 0 /*MipLevel*/, 0 /*AttributesIndex*/, Header.FallbackValues[0]);
+								WriteTileDataVoxel(TileIndex, WriteCoord, 0 /*MipLevel*/, 1 /*AttributesIndex*/, Header.FallbackValues[1]);
 							}
 						}
 					}
 				});
 
-			Header.NullTileValuesQuantized[0] = ReadTileDataVoxel(0, FIntVector::ZeroValue, 0 /*MipLevel*/, 0 /*AttributesIndex*/);
-			Header.NullTileValuesQuantized[1] = ReadTileDataVoxel(0, FIntVector::ZeroValue, 0 /*MipLevel*/, 1 /*AttributesIndex*/);
+			FallbackValuesQuantized[0] = ReadTileDataVoxel(0, FIntVector::ZeroValue, 0 /*MipLevel*/, 0 /*AttributesIndex*/);
+			FallbackValuesQuantized[1] = ReadTileDataVoxel(0, FIntVector::ZeroValue, 0 /*MipLevel*/, 1 /*AttributesIndex*/);
 		}
 
 		// Set up the page table entries
@@ -205,9 +210,9 @@ bool FSparseVolumeTextureData::Create(const ISparseVolumeTextureDataProvider& Da
 	return true;
 }
 
-bool FSparseVolumeTextureData::CreateFromDense(const FSparseVolumeTextureDataCreateInfo& CreateInfo, const TArrayView<uint8>& VoxelDataA, const TArrayView<uint8>& VoxelDataB)
+bool FTextureData::CreateFromDense(const FTextureDataCreateInfo& CreateInfo, const TArrayView<uint8>& VoxelDataA, const TArrayView<uint8>& VoxelDataB)
 {	
-	Header = FSparseVolumeTextureHeader(CreateInfo.VirtualVolumeAABBMin, CreateInfo.VirtualVolumeAABBMax, CreateInfo.AttributesFormats[0], CreateInfo.AttributesFormats[1], CreateInfo.FallbackValues[0], CreateInfo.FallbackValues[1]);
+	Header = UE::SVT::FHeader(CreateInfo.VirtualVolumeAABBMin, CreateInfo.VirtualVolumeAABBMax, CreateInfo.AttributesFormats[0], CreateInfo.AttributesFormats[1], CreateInfo.FallbackValues[0], CreateInfo.FallbackValues[1]);
 
 	if (Header.PageTableVolumeResolution.X > UE::SVT::SVTMaxVolumeTextureDim
 		|| Header.PageTableVolumeResolution.Y > UE::SVT::SVTMaxVolumeTextureDim
@@ -227,7 +232,7 @@ bool FSparseVolumeTextureData::CreateFromDense(const FSparseVolumeTextureDataCre
 		if (CreateInfo.AttributesFormats[i] != PF_Unknown)
 		{
 			UE::SVT::WriteVoxel(0, NullTileValuesU8[i], CreateInfo.AttributesFormats[i], CreateInfo.FallbackValues[i]);
-			Header.NullTileValuesQuantized[i] = UE::SVT::ReadVoxel(0, NullTileValuesU8[i], CreateInfo.AttributesFormats[i]);
+			FallbackValuesQuantized[i] = UE::SVT::ReadVoxel(0, NullTileValuesU8[i], CreateInfo.AttributesFormats[i]);
 		}
 	}
 
@@ -372,10 +377,10 @@ bool FSparseVolumeTextureData::CreateFromDense(const FSparseVolumeTextureDataCre
 	return true;
 }
 
-void FSparseVolumeTextureData::CreateDefault()
+void FTextureData::CreateDefault()
 {
 	// Store a single 1x1x1 mip level with the page table pointing to a (zero-valued) null tile.
-	Header = FSparseVolumeTextureHeader(FIntVector3::ZeroValue, FIntVector3(1, 1, 1), PF_R8, PF_Unknown, FVector4f(), FVector4f());
+	Header = UE::SVT::FHeader(FIntVector3::ZeroValue, FIntVector3(1, 1, 1), PF_R8, PF_Unknown, FVector4f(), FVector4f());
 
 	MipMaps.SetNum(1);
 	MipMaps[0].PageTable.SetNumZeroed(1);
@@ -384,7 +389,7 @@ void FSparseVolumeTextureData::CreateDefault()
 	MipMaps[0].NumPhysicalTiles = 0;
 }
 
-uint32 FSparseVolumeTextureData::ReadPageTable(const FIntVector3& PageTableCoord, int32 MipLevel) const
+uint32 FTextureData::ReadPageTable(const FIntVector3& PageTableCoord, int32 MipLevel) const
 {
 	if (!MipMaps.IsValidIndex(MipLevel))
 	{
@@ -410,7 +415,7 @@ uint32 FSparseVolumeTextureData::ReadPageTable(const FIntVector3& PageTableCoord
 	return INDEX_NONE;
 }
 
-FVector4f FSparseVolumeTextureData::ReadTileDataVoxel(int32 TileIndex, const FIntVector3& TileDataCoord, int32 MipLevel, int32 AttributesIdx) const
+FVector4f FTextureData::ReadTileDataVoxel(int32 TileIndex, const FIntVector3& TileDataCoord, int32 MipLevel, int32 AttributesIdx) const
 {
 	check(AttributesIdx >= 0 && AttributesIdx <= 1);
 	if (TileIndex == INDEX_NONE || (AttributesIdx == 0 && MipMaps[MipLevel].PhysicalTileDataA.IsEmpty()) || (AttributesIdx == 1 && MipMaps[MipLevel].PhysicalTileDataB.IsEmpty()))
@@ -435,7 +440,7 @@ FVector4f FSparseVolumeTextureData::ReadTileDataVoxel(int32 TileIndex, const FIn
 	return UE::SVT::ReadVoxel(VoxelIndex, TileData, Format);
 }
 
-FVector4f FSparseVolumeTextureData::Load(const FIntVector3& VolumeCoord, int32 MipLevel, int32 AttributesIdx, const FSparseVolumeTextureDataAddressingInfo& AddressingInfo) const
+FVector4f FTextureData::Load(const FIntVector3& VolumeCoord, int32 MipLevel, int32 AttributesIdx, const FTextureDataAddressingInfo& AddressingInfo) const
 {
 	auto ApplyAddressMode = [](int32 x, int32 Width, TextureAddress Mode)
 	{
@@ -464,14 +469,14 @@ FVector4f FSparseVolumeTextureData::Load(const FIntVector3& VolumeCoord, int32 M
 	const uint32 TileIndex = ReadPageTable(PageTableCoord, MipLevel);
 	if (TileIndex == 0)
 	{
-		return Header.NullTileValuesQuantized[AttributesIdx];
+		return FallbackValuesQuantized[AttributesIdx];
 	}
 	const FIntVector3 VoxelCoord = AddressModeVolumeCoord % SPARSE_VOLUME_TILE_RES;
 	const FVector4f Sample = ReadTileDataVoxel(TileIndex - 1, VoxelCoord, MipLevel, AttributesIdx); // -1 because we don't store the null tile in memory.
 	return Sample;
 }
 
-void FSparseVolumeTextureData::WriteTileDataVoxel(int32 TileIndex, const FIntVector3& TileDataCoord, int32 MipLevel, int32 AttributesIdx, const FVector4f& Value, int32 DstComponent)
+void FTextureData::WriteTileDataVoxel(int32 TileIndex, const FIntVector3& TileDataCoord, int32 MipLevel, int32 AttributesIdx, const FVector4f& Value, int32 DstComponent)
 {
 	check(AttributesIdx >= 0 && AttributesIdx <= 1);
 	if ((AttributesIdx == 0 && MipMaps[MipLevel].PhysicalTileDataA.IsEmpty()) || (AttributesIdx == 1 && MipMaps[MipLevel].PhysicalTileDataB.IsEmpty()))
@@ -500,7 +505,7 @@ void FSparseVolumeTextureData::WriteTileDataVoxel(int32 TileIndex, const FIntVec
 	UE::SVT::WriteVoxel(VoxelIndex, TileData, Format, Value, DstComponent);
 }
 
-bool FSparseVolumeTextureData::GenerateMipMaps(const FSparseVolumeTextureDataAddressingInfo& AddressingInfo, int32 NumMipLevels)
+bool FTextureData::GenerateMipMaps(const FTextureDataAddressingInfo& AddressingInfo, int32 NumMipLevels)
 {
 	const int32 FormatSize[] = { GPixelFormats[Header.AttributesFormats[0]].BlockBytes, GPixelFormats[Header.AttributesFormats[1]].BlockBytes };
 
@@ -643,7 +648,7 @@ bool FSparseVolumeTextureData::GenerateMipMaps(const FSparseVolumeTextureDataAdd
 	return true;
 }
 
-bool FSparseVolumeTextureData::GenerateBorderVoxels(const FSparseVolumeTextureDataAddressingInfo& AddressingInfo, int32 MipLevel, const TArray<FIntVector3>& PageCoords)
+bool FTextureData::GenerateBorderVoxels(const FTextureDataAddressingInfo& AddressingInfo, int32 MipLevel, const TArray<FIntVector3>& PageCoords)
 {
 	const FIntVector3 PageTableVolumeResolution = FIntVector3(
 		FMath::Max(1, Header.PageTableVolumeResolution.X >> MipLevel), 
@@ -679,7 +684,7 @@ bool FSparseVolumeTextureData::GenerateBorderVoxels(const FSparseVolumeTextureDa
 	return true;
 }
 
-bool FSparseVolumeTextureData::DeduplicateTiles()
+bool FTextureData::DeduplicateTiles()
 {
 	const int32 NumMipLevels = MipMaps.Num();
 	const int32 FormatSize[] = { GPixelFormats[Header.AttributesFormats[0]].BlockBytes, GPixelFormats[Header.AttributesFormats[1]].BlockBytes };
@@ -782,9 +787,9 @@ bool FSparseVolumeTextureData::DeduplicateTiles()
 	return true;
 }
 
-bool FSparseVolumeTextureData::BuildDerivedData(const FSparseVolumeTextureDataAddressingInfo& AddressingInfo, int32 NumMipLevels, bool bMoveMip0FromThis, FSparseVolumeTextureData& OutDerivedData)
+bool FTextureData::BuildDerivedData(const FTextureDataAddressingInfo& AddressingInfo, int32 NumMipLevels, bool bMoveMip0FromThis, FTextureData& OutDerivedData)
 {
-	OutDerivedData = FSparseVolumeTextureData{};
+	OutDerivedData = FTextureData{};
 	OutDerivedData.Header = Header;
 	OutDerivedData.MipMaps.SetNum(1);
 	if (bMoveMip0FromThis)
@@ -838,4 +843,7 @@ bool FSparseVolumeTextureData::BuildDerivedData(const FSparseVolumeTextureDataAd
 	return true;
 }
 
+}
+}
+
 ////////////////////////////////////////////////////////////////////////////////////////////////
diff --git a/Engine/Source/Runtime/Engine/Private/SparseVolumeTexture/SparseVolumeTextureSceneProxy.cpp b/Engine/Source/Runtime/Engine/Private/SparseVolumeTexture/SparseVolumeTextureSceneProxy.cpp
deleted file mode 100644
diff --git a/Engine/Source/Runtime/Engine/Private/SparseVolumeTexture/SparseVolumeTextureStreaming.cpp b/Engine/Source/Runtime/Engine/Private/SparseVolumeTexture/SparseVolumeTextureStreaming.cpp
deleted file mode 100644
diff --git a/Engine/Source/Runtime/Engine/Private/SparseVolumeTexture/SparseVolumeTextureStreamingManager.cpp b/Engine/Source/Runtime/Engine/Private/SparseVolumeTexture/SparseVolumeTextureStreamingManager.cpp
new file mode 100644
--- /dev/null
+++ b/Engine/Source/Runtime/Engine/Private/SparseVolumeTexture/SparseVolumeTextureStreamingManager.cpp
@@ -0,0 +1,2102 @@
+// Copyright Epic Games, Inc. All Rights Reserved.
+
+#include "SparseVolumeTexture/SparseVolumeTextureStreamingManager.h"
+#include "SparseVolumeTexture/SparseVolumeTextureUtility.h"
+#include "HAL/IConsoleManager.h"
+#include "SparseVolumeTexture/SparseVolumeTexture.h"
+#include "RenderCore.h"
+#include "RenderGraph.h"
+#include "GlobalShader.h"
+#include "ShaderCompilerCore.h" // AllowGlobalShaderLoad()
+#include "Async/ParallelFor.h"
+
+#if WITH_EDITORONLY_DATA
+#include "DerivedDataCache.h"
+#include "DerivedDataRequestOwner.h"
+#endif
+
+DEFINE_LOG_CATEGORY_STATIC(LogSparseVolumeTextureStreamingManager, Log, All);
+
+#ifndef SVT_STREAMING_LOG_VERBOSE
+#define SVT_STREAMING_LOG_VERBOSE 0
+#endif
+
+static int32 GSVTStreamingNumPrefetchFrames = 3;
+static FAutoConsoleVariableRef CVarSVTStreamingNumPrefetchFrames(
+	TEXT("r.SparseVolumeTexture.Streaming.NumPrefetchFrames"),
+	GSVTStreamingNumPrefetchFrames,
+	TEXT("Number of frames to prefetch when a frame is requested."),
+	ECVF_RenderThreadSafe
+);
+
+static int32 GSVTStreamingPrefetchMipLevelBias = -1;
+static FAutoConsoleVariableRef CVarSVTStreamingPrefetchMipLevelBias(
+	TEXT("r.SparseVolumeTexture.Streaming.PrefetchMipLevelBias"),
+	GSVTStreamingPrefetchMipLevelBias,
+	TEXT("Bias to apply to the mip level of prefetched frames. Prefetching is done at increasingly higher mip levels (lower resolution), so setting a negative value here will increase the requested mip level resolution."),
+	ECVF_RenderThreadSafe
+);
+
+static int32 GSVTStreamingAsyncThread = 1;
+static FAutoConsoleVariableRef CVarSVTStreamingAsync(
+	TEXT("r.SparseVolumeTexture.Streaming.AsyncThread"),
+	GSVTStreamingAsyncThread,
+	TEXT("Perform most of the SVT streaming on an asynchronous worker thread instead of the rendering thread."),
+	ECVF_RenderThreadSafe
+);
+
+static int32 GSVTStreamingAsyncCompute = 1;
+static FAutoConsoleVariableRef CVarSVTStreamingAsyncCompute(
+	TEXT("r.SparseVolumeTexture.Streaming.AsyncCompute"),
+	GSVTStreamingAsyncCompute,
+	TEXT("Schedule GPU work in async compute queue."),
+	ECVF_RenderThreadSafe
+);
+
+static int32 GSVTStreamingEmptyPhysicalTileTextures = 0;
+static FAutoConsoleVariableRef CVarSVTStreamingEmptyPhysicalTileTextures(
+	TEXT("r.SparseVolumeTexture.Streaming.EmptyPhysicalTileTextures"),
+	GSVTStreamingEmptyPhysicalTileTextures,
+	TEXT("Streams out all streamable tiles of all physical tile textures."),
+	ECVF_RenderThreadSafe
+);
+
+static int32 GSVTStreamingMaxPendingMipLevels = 128;
+static FAutoConsoleVariableRef CVarSVTStreamingMaxPendingMipLevels(
+	TEXT("r.SparseVolumeTexture.Streaming.MaxPendingMipLevels"),
+	GSVTStreamingMaxPendingMipLevels,
+	TEXT("Maximum number of mip levels that can be pending for installation."),
+	ECVF_RenderThreadSafe | ECVF_ReadOnly
+);
+
+namespace UE
+{
+namespace SVT
+{
+
+static bool DoesPlatformSupportSparseVolumeTexture(EShaderPlatform Platform)
+{
+	// SVT_TODO: This is a bit of a hack: FStreamingManager::Add_GameThread() issues a rendering thread lambda for creating the RHI resources and uploading root tile data.
+	// Uploading root tile data involves access to the global shader map, which is empty under certain circumstances. By checking AllowGlobalShaderLoad(), we disallow streaming completely.
+	return AllowGlobalShaderLoad();
+}
+
+static FIntVector3 ComputeTileDataVolumeResolution2(int32 NumAllocatedPages)
+{
+	int32 TileVolumeResolutionCube = 1;
+	while (TileVolumeResolutionCube * TileVolumeResolutionCube * TileVolumeResolutionCube < NumAllocatedPages)
+	{
+		TileVolumeResolutionCube++;				// We use a simple loop to compute the minimum resolution of a cube to store all the tile data
+	}
+	FIntVector3 TileDataVolumeResolution = FIntVector3(TileVolumeResolutionCube, TileVolumeResolutionCube, TileVolumeResolutionCube);
+	
+	// Trim volume to reclaim some space
+	while ((TileDataVolumeResolution.X * TileDataVolumeResolution.Y * (TileDataVolumeResolution.Z - 1)) > NumAllocatedPages)
+	{
+		TileDataVolumeResolution.Z--;
+	}
+	while ((TileDataVolumeResolution.X * (TileDataVolumeResolution.Y - 1) * TileDataVolumeResolution.Z) > NumAllocatedPages)
+	{
+		TileDataVolumeResolution.Y--;
+	}
+	while (((TileDataVolumeResolution.X - 1) * TileDataVolumeResolution.Y * TileDataVolumeResolution.Z) > NumAllocatedPages)
+	{
+		TileDataVolumeResolution.X--;
+	}
+
+	return TileDataVolumeResolution * SPARSE_VOLUME_TILE_RES_PADDED;
+}
+
+static FIntVector3 ComputeLargestPossibleTileDataVolumeResolution(int32 VoxelMemSize)
+{
+	const int64 TileMemSize = SVTNumVoxelsPerPaddedTile * VoxelMemSize;
+	const int64 NumMaxTiles = int64(INT32_MAX) / TileMemSize;
+	int64 ResourceSize = NumMaxTiles * TileMemSize;
+
+	// Find a cube with a volume as close to NumMaxTiles as possible
+	int32 TileVolumeResolutionCube = 1;
+	while (((TileVolumeResolutionCube + 1) * (TileVolumeResolutionCube + 1) * (TileVolumeResolutionCube + 1)) <= NumMaxTiles)
+	{
+		++TileVolumeResolutionCube;
+	}
+
+	// Try to add to the sides to get closer to NumMaxTiles
+	FIntVector3 ResolutionInTiles = FIntVector3(TileVolumeResolutionCube, TileVolumeResolutionCube, TileVolumeResolutionCube);
+	if (((ResolutionInTiles.X + 1) * ResolutionInTiles.Y * ResolutionInTiles.Z) <= NumMaxTiles)
+	{
+		++ResolutionInTiles.X;
+	}
+	if ((ResolutionInTiles.X * (ResolutionInTiles.Y + 1) * ResolutionInTiles.Z) <= NumMaxTiles)
+	{
+		++ResolutionInTiles.Y;
+	}
+	if ((ResolutionInTiles.X * ResolutionInTiles.Y * (ResolutionInTiles.Z + 1)) <= NumMaxTiles)
+	{
+		++ResolutionInTiles.Z;
+	}
+
+	const FIntVector3 Resolution = ResolutionInTiles * SPARSE_VOLUME_TILE_RES_PADDED;
+	check(Resolution.X <= SVTMaxVolumeTextureDim && Resolution.Y <= SVTMaxVolumeTextureDim && Resolution.Z <= SVTMaxVolumeTextureDim);
+	check(((int64)Resolution.X * (int64)Resolution.Y * (int64)Resolution.Z) < int64(INT32_MAX));
+
+	return Resolution;
+}
+
+class SparseVolumeTextureUpdateFromBufferCS : public FGlobalShader
+{
+	DECLARE_GLOBAL_SHADER(SparseVolumeTextureUpdateFromBufferCS);
+	SHADER_USE_PARAMETER_STRUCT(SparseVolumeTextureUpdateFromBufferCS, FGlobalShader)
+
+	BEGIN_SHADER_PARAMETER_STRUCT(FParameters, )
+		SHADER_PARAMETER_RDG_TEXTURE_UAV(RWTexture3D, DstPhysicalTileTextureA)
+		SHADER_PARAMETER_RDG_TEXTURE_UAV(RWTexture3D, DstPhysicalTileTextureB)
+		SHADER_PARAMETER_RDG_BUFFER_SRV(ByteAddressBuffer, DstTileCoords)
+		SHADER_PARAMETER_RDG_BUFFER_SRV(Buffer<float4>, SrcPhysicalTileBufferA)
+		SHADER_PARAMETER_RDG_BUFFER_SRV(Buffer<float4>, SrcPhysicalTileBufferB)
+		SHADER_PARAMETER(uint32, TileCoordsBufferOffset)
+		SHADER_PARAMETER(uint32, TileDataBufferOffsetInTiles)
+		SHADER_PARAMETER(uint32, NumTilesToCopy)
+		SHADER_PARAMETER(uint32, NumDispatchedGroups)
+		SHADER_PARAMETER(uint32, PaddedTileSize)
+		SHADER_PARAMETER(uint32, bCopyTexureAOnlyUI)
+	END_SHADER_PARAMETER_STRUCT()
+
+	static bool ShouldCompilePermutation(const FGlobalShaderPermutationParameters& Parameters)
+	{
+		return true;
+	}
+
+	static void ModifyCompilationEnvironment(const FGlobalShaderPermutationParameters& Parameters, FShaderCompilerEnvironment& OutEnvironment)
+	{
+		FGlobalShader::ModifyCompilationEnvironment(Parameters, OutEnvironment);
+		OutEnvironment.SetDefine(TEXT("UPDATE_TILE_TEXTURE_FROM_BUFFER"), 1);
+	}
+};
+IMPLEMENT_GLOBAL_SHADER(SparseVolumeTextureUpdateFromBufferCS, "/Engine/Private/SparseVolumeTexture/UpdateSparseVolumeTexture.usf", "SparseVolumeTextureUpdateFromBufferCS", SF_Compute);
+
+class SparseVolumeTextureUpdatePageTableCS : public FGlobalShader
+{
+	DECLARE_GLOBAL_SHADER(SparseVolumeTextureUpdatePageTableCS);
+	SHADER_USE_PARAMETER_STRUCT(SparseVolumeTextureUpdatePageTableCS, FGlobalShader)
+
+	BEGIN_SHADER_PARAMETER_STRUCT(FParameters, )
+		SHADER_PARAMETER_RDG_TEXTURE_UAV(RWTexture3D, PageTable)
+		SHADER_PARAMETER_RDG_BUFFER_SRV(ByteAddressBuffer, PageTableUpdates)
+		SHADER_PARAMETER(uint32, UpdateCoordOffset)
+		SHADER_PARAMETER(uint32, UpdatePayloadOffset)
+		SHADER_PARAMETER(uint32, NumUpdates)
+	END_SHADER_PARAMETER_STRUCT()
+
+	static bool ShouldCompilePermutation(const FGlobalShaderPermutationParameters& Parameters)
+	{
+		return true;
+	}
+
+	static void ModifyCompilationEnvironment(const FGlobalShaderPermutationParameters& Parameters, FShaderCompilerEnvironment& OutEnvironment)
+	{
+		FGlobalShader::ModifyCompilationEnvironment(Parameters, OutEnvironment);
+		OutEnvironment.SetDefine(TEXT("UPDATE_PAGE_TABLE"), 1);
+	}
+};
+IMPLEMENT_GLOBAL_SHADER(SparseVolumeTextureUpdatePageTableCS, "/Engine/Private/SparseVolumeTexture/UpdateSparseVolumeTexture.usf", "SparseVolumeTextureUpdatePageTableCS", SF_Compute);
+
+class SparseVolumeTextureUpdateStreamingInfoBufferCS : public FGlobalShader
+{
+	DECLARE_GLOBAL_SHADER(SparseVolumeTextureUpdateStreamingInfoBufferCS);
+	SHADER_USE_PARAMETER_STRUCT(SparseVolumeTextureUpdateStreamingInfoBufferCS, FGlobalShader)
+
+	BEGIN_SHADER_PARAMETER_STRUCT(FParameters, )
+		SHADER_PARAMETER_RDG_BUFFER_UAV(RWByteAddressBuffer, StreamingInfoBuffer)
+		SHADER_PARAMETER_RDG_BUFFER_SRV(ByteAddressBuffer, StreamingInfoBufferUpdates)
+		SHADER_PARAMETER(uint32, UpdateOffset)
+		SHADER_PARAMETER(uint32, NumUpdates)
+	END_SHADER_PARAMETER_STRUCT()
+
+	static bool ShouldCompilePermutation(const FGlobalShaderPermutationParameters& Parameters)
+	{
+		return true;
+	}
+
+	static void ModifyCompilationEnvironment(const FGlobalShaderPermutationParameters& Parameters, FShaderCompilerEnvironment& OutEnvironment)
+	{
+		FGlobalShader::ModifyCompilationEnvironment(Parameters, OutEnvironment);
+		OutEnvironment.SetDefine(TEXT("UPDATE_STREAMING_INFO_BUFFER"), 1);
+	}
+};
+IMPLEMENT_GLOBAL_SHADER(SparseVolumeTextureUpdateStreamingInfoBufferCS, "/Engine/Private/SparseVolumeTexture/UpdateSparseVolumeTexture.usf", "SparseVolumeTextureUpdateStreamingInfoBufferCS", SF_Compute);
+
+// Utility class for uploading tiles to a physical tile data texture
+class FTileUploader
+{
+public:
+	FTileUploader()
+	{
+		ResetState();
+	}
+
+	void Init(FRDGBuilder& GraphBuilder, int32 InMaxNumTiles, EPixelFormat InFormatA, EPixelFormat InFormatB)
+	{
+		check(InFormatA != PF_Unknown || InFormatB != PF_Unknown);
+		ResetState();
+		MaxNumTiles = InMaxNumTiles;
+		FormatA = InFormatA;
+		FormatB = InFormatB;
+		FormatSizeA = GPixelFormats[FormatA].BlockBytes;
+		FormatSizeB = GPixelFormats[FormatB].BlockBytes;
+
+		// Create a new set of buffers if the old set is already queued into RDG.
+		if (IsRegistered(GraphBuilder, DstTileCoordsUploadBuffer))
+		{
+			DstTileCoordsUploadBuffer = nullptr;
+			TileDataAUploadBuffer = nullptr;
+			TileDataBUploadBuffer = nullptr;
+		}
+
+		if (MaxNumTiles > 0)
+		{
+			// TileCoords
+			{
+				// Add EBufferUsageFlags::Dynamic to skip the unneeded copy from upload to VRAM resource on d3d12 RHI
+				FRDGBufferDesc BufferDesc = FRDGBufferDesc::CreateByteAddressUploadDesc(MaxNumTiles * sizeof(uint32));
+				BufferDesc.Usage |= EBufferUsageFlags::Dynamic;
+				AllocatePooledBuffer(BufferDesc, DstTileCoordsUploadBuffer, TEXT("SparseVolumeTexture.TileCoordsUploadBuffer"));
+
+				TileCoordsPtr = (uint8*)RHILockBuffer(DstTileCoordsUploadBuffer->GetRHI(), 0, MaxNumTiles * sizeof(uint32), RLM_WriteOnly);
+			}
+
+			const int32 NumVoxels = MaxNumTiles * UE::SVT::SVTNumVoxelsPerPaddedTile;
+
+			// TileData
+			if (FormatSizeA > 0)
+			{
+				// Add EBufferUsageFlags::Dynamic to skip the unneeded copy from upload to VRAM resource on d3d12 RHI
+				FRDGBufferDesc BufferDesc = FRDGBufferDesc::CreateUploadDesc(FormatSizeA, NumVoxels);
+				BufferDesc.Usage |= EBufferUsageFlags::Dynamic;
+				AllocatePooledBuffer(BufferDesc, TileDataAUploadBuffer, TEXT("SparseVolumeTexture.TileDataAUploadBuffer"));
+
+				TileDataAPtr = (uint8*)RHILockBuffer(TileDataAUploadBuffer->GetRHI(), 0, NumVoxels * FormatSizeA, RLM_WriteOnly);
+			}
+			if (FormatSizeB > 0)
+			{
+				// Add EBufferUsageFlags::Dynamic to skip the unneeded copy from upload to VRAM resource on d3d12 RHI
+				FRDGBufferDesc BufferDesc = FRDGBufferDesc::CreateUploadDesc(FormatSizeB, NumVoxels);
+				BufferDesc.Usage |= EBufferUsageFlags::Dynamic;
+				AllocatePooledBuffer(BufferDesc, TileDataBUploadBuffer, TEXT("SparseVolumeTexture.TileDataBUploadBuffer"));
+
+				TileDataBPtr = (uint8*)RHILockBuffer(TileDataBUploadBuffer->GetRHI(), 0, NumVoxels * FormatSizeB, RLM_WriteOnly);
+			}
+		}
+	}
+
+	void Add_GetRef(int32 NumTiles, uint8*& OutPackedPhysicalTileCoordsPtr, uint8*& OutPtrA, uint8*& OutPtrB)
+	{
+		check((NumWrittenTiles + NumTiles) <= MaxNumTiles);
+		check(TileCoordsPtr);
+		check(FormatSizeA <= 0 || TileDataAPtr);
+		check(FormatSizeB <= 0 || TileDataBPtr);
+
+		OutPackedPhysicalTileCoordsPtr = TileCoordsPtr + NumWrittenTiles * sizeof(uint32);
+		OutPtrA = TileDataAPtr ? TileDataAPtr + (NumWrittenTiles * UE::SVT::SVTNumVoxelsPerPaddedTile * FormatSizeA) : nullptr;
+		OutPtrB = TileDataBPtr ? TileDataBPtr + (NumWrittenTiles * UE::SVT::SVTNumVoxelsPerPaddedTile * FormatSizeB) : nullptr;
+
+		NumWrittenTiles += NumTiles;
+	}
+
+	void Release()
+	{
+		DstTileCoordsUploadBuffer.SafeRelease();
+		TileDataAUploadBuffer.SafeRelease();
+		TileDataBUploadBuffer.SafeRelease();
+		ResetState();
+	}
+
+	void ResourceUploadTo(FRDGBuilder& GraphBuilder, FRHITexture* DstTextureA, FRHITexture* DstTextureB)
+	{
+		check(DstTextureA || FormatSizeA <= 0);
+		check(DstTextureB || FormatSizeB <= 0);
+		if (MaxNumTiles > 0)
+		{
+			RHIUnlockBuffer(DstTileCoordsUploadBuffer->GetRHI());
+			if (TileDataAPtr)
+			{
+				RHIUnlockBuffer(TileDataAUploadBuffer->GetRHI());
+			}
+			if (TileDataBPtr)
+			{
+				RHIUnlockBuffer(TileDataBUploadBuffer->GetRHI());
+			}
+
+			if (NumWrittenTiles > 0)
+			{
+				FRDGTexture* DstTextureARDG = DstTextureA ? GraphBuilder.RegisterExternalTexture(CreateRenderTarget(DstTextureA, TEXT("SparseVolumeTexture.TileDataTextureA"))) : nullptr;
+				FRDGTexture* DstTextureBRDG = DstTextureB ? GraphBuilder.RegisterExternalTexture(CreateRenderTarget(DstTextureB, TEXT("SparseVolumeTexture.TileDataTextureB"))) : nullptr;
+
+				FRDGBufferSRV* DstTileCoordsBufferSRV = GraphBuilder.CreateSRV(GraphBuilder.RegisterExternalBuffer(DstTileCoordsUploadBuffer));
+				FRDGBufferSRV* TileDataABufferSRV = FormatSizeA ? GraphBuilder.CreateSRV(GraphBuilder.RegisterExternalBuffer(TileDataAUploadBuffer), FormatA) : nullptr;
+				FRDGBufferSRV* TileDataBBufferSRV = FormatSizeB ? GraphBuilder.CreateSRV(GraphBuilder.RegisterExternalBuffer(TileDataBUploadBuffer), FormatB) : nullptr;
+				FRDGTextureUAV* DstTextureAUAV = DstTextureARDG ? GraphBuilder.CreateUAV(DstTextureARDG) : nullptr;
+				FRDGTextureUAV* DstTextureBUAV = DstTextureBRDG ? GraphBuilder.CreateUAV(DstTextureBRDG) : nullptr;
+
+				auto ComputeShader = GetGlobalShaderMap(GMaxRHIFeatureLevel)->GetShader<SparseVolumeTextureUpdateFromBufferCS>();
+
+				SparseVolumeTextureUpdateFromBufferCS::FParameters* PassParameters = GraphBuilder.AllocParameters<SparseVolumeTextureUpdateFromBufferCS::FParameters>();
+				if (FormatSizeA > 0) // TextureA exists
+				{
+					PassParameters->DstPhysicalTileTextureA = DstTextureAUAV;
+					PassParameters->SrcPhysicalTileBufferA = TileDataABufferSRV;
+					if (FormatSizeB == 0) // B doesn't exist: fill B params with A
+					{
+						PassParameters->DstPhysicalTileTextureB = DstTextureAUAV;
+						PassParameters->SrcPhysicalTileBufferB = TileDataABufferSRV;
+					}
+				}
+				if (FormatSizeB > 0) // TextureB exists
+				{
+					PassParameters->DstPhysicalTileTextureB = DstTextureBUAV;
+					PassParameters->SrcPhysicalTileBufferB = TileDataBBufferSRV;
+					if (FormatSizeA == 0) // A doesnt't exist: fill A params with B
+					{
+						PassParameters->DstPhysicalTileTextureA = DstTextureBUAV;
+						PassParameters->SrcPhysicalTileBufferA = TileDataBBufferSRV;
+					}
+				}
+				PassParameters->DstTileCoords = DstTileCoordsBufferSRV;
+				PassParameters->TileCoordsBufferOffset = 0;
+				PassParameters->TileDataBufferOffsetInTiles = 0;
+				PassParameters->NumTilesToCopy = NumWrittenTiles;
+				PassParameters->NumDispatchedGroups = FMath::Min(NumWrittenTiles, 1024);
+				PassParameters->PaddedTileSize = SPARSE_VOLUME_TILE_RES_PADDED;
+				PassParameters->bCopyTexureAOnlyUI = (FormatSizeA == 0 || FormatSizeB == 0);
+
+				const bool bAsyncCompute = GSupportsEfficientAsyncCompute && (GSVTStreamingAsyncCompute != 0);
+
+				FComputeShaderUtils::AddPass(
+					GraphBuilder,
+					RDG_EVENT_NAME("Upload SVT Tiles (TileCount: %u)", NumWrittenTiles),
+					bAsyncCompute ? ERDGPassFlags::AsyncCompute : ERDGPassFlags::Compute,
+					ComputeShader,
+					PassParameters,
+					FIntVector3(PassParameters->NumDispatchedGroups, 1, 1)
+				);
+			}
+		}
+		Release();
+	}
+
+private:
+	TRefCountPtr<FRDGPooledBuffer> DstTileCoordsUploadBuffer;
+	TRefCountPtr<FRDGPooledBuffer> TileDataAUploadBuffer;
+	TRefCountPtr<FRDGPooledBuffer> TileDataBUploadBuffer;
+	uint8* TileCoordsPtr = nullptr;
+	uint8* TileDataAPtr = nullptr;
+	uint8* TileDataBPtr = nullptr;
+	int32 MaxNumTiles = 0;
+	EPixelFormat FormatA = PF_Unknown;
+	EPixelFormat FormatB = PF_Unknown;
+	int32 FormatSizeA = 0;
+	int32 FormatSizeB = 0;
+	int32 NumWrittenTiles = 0;
+
+	void ResetState()
+	{
+		TileCoordsPtr = nullptr;
+		TileDataAPtr = nullptr;
+		TileDataBPtr = nullptr;
+		MaxNumTiles = 0;
+		FormatA = PF_Unknown;
+		FormatB = PF_Unknown;
+		FormatSizeA = 0;
+		FormatSizeB = 0;
+		NumWrittenTiles = 0;
+	}
+};
+
+// Utility class for writing page table entries
+class FPageTableUpdater
+{
+public:
+	FPageTableUpdater()
+	{
+		ResetState();
+	}
+
+	void Init(FRDGBuilder& GraphBuilder, int32 InMaxNumUpdates, int32 InEstimatedNumBatches)
+	{
+		ResetState();
+		MaxNumUpdates = InMaxNumUpdates;
+		Batches.Reserve(InEstimatedNumBatches);
+
+		// Create a new buffer if the old one is already queued into RDG.
+		if (IsRegistered(GraphBuilder, UpdatesUploadBuffer))
+		{
+			UpdatesUploadBuffer = nullptr;
+		}
+
+		if (MaxNumUpdates > 0)
+		{
+			// Add EBufferUsageFlags::Dynamic to skip the unneeded copy from upload to VRAM resource on d3d12 RHI
+			FRDGBufferDesc BufferDesc = FRDGBufferDesc::CreateByteAddressUploadDesc(MaxNumUpdates * 2 * sizeof(uint32));
+			BufferDesc.Usage |= EBufferUsageFlags::Dynamic;
+			AllocatePooledBuffer(BufferDesc, UpdatesUploadBuffer, TEXT("SparseVolumeTexture.PageTableUpdatesUploadBuffer"));
+
+			DataPtr = (uint8*)RHILockBuffer(UpdatesUploadBuffer->GetRHI(), 0, MaxNumUpdates * 2 * sizeof(uint32), RLM_WriteOnly);
+		}
+	}
+
+	void Add_GetRef(FRHITexture* PageTable, int32 MipLevel, int32 NumUpdates, uint8*& OutCoordsPtr, uint8*& OutPayloadPtr)
+	{
+		check((NumWrittenUpdates + NumUpdates) <= MaxNumUpdates);
+		check(DataPtr);
+		FBatch* Batch = Batches.IsEmpty() ? nullptr : &Batches.Last();
+		if (!Batch || Batch->PageTable != PageTable || Batch->MipLevel != MipLevel)
+		{
+			Batch = &Batches.Add_GetRef(FBatch(PageTable, MipLevel));
+		}
+
+		OutCoordsPtr = DataPtr + NumWrittenUpdates * sizeof(uint32);
+		OutPayloadPtr = DataPtr + (MaxNumUpdates + NumWrittenUpdates) * sizeof(uint32);
+
+		Batch->NumUpdates += NumUpdates;
+		NumWrittenUpdates += NumUpdates;
+	}
+
+	void Release()
+	{
+		UpdatesUploadBuffer.SafeRelease();
+		ResetState();
+	}
+
+	void Apply(FRDGBuilder& GraphBuilder)
+	{
+		if (MaxNumUpdates > 0)
+		{
+			RHIUnlockBuffer(UpdatesUploadBuffer->GetRHI());
+
+			if (NumWrittenUpdates > 0)
+			{
+				auto ComputeShader = GetGlobalShaderMap(GMaxRHIFeatureLevel)->GetShader<SparseVolumeTextureUpdatePageTableCS>();
+				const bool bAsyncCompute = GSupportsEfficientAsyncCompute && (GSVTStreamingAsyncCompute != 0);
+
+				uint32 UpdatesOffset = 0;
+				for (const FBatch& Batch : Batches)
+				{
+					FRDGTexture* PageTableRDG = GraphBuilder.RegisterExternalTexture(CreateRenderTarget(Batch.PageTable, TEXT("SparseVolumeTexture.PageTableTexture")));
+					FRDGTextureUAV* PageTableUAV = GraphBuilder.CreateUAV(FRDGTextureUAVDesc(PageTableRDG, Batch.MipLevel, PF_R32_UINT));
+					FRDGBufferSRV* UpdatesBufferSRV = GraphBuilder.CreateSRV(GraphBuilder.RegisterExternalBuffer(UpdatesUploadBuffer));
+
+					SparseVolumeTextureUpdatePageTableCS::FParameters* PassParameters = GraphBuilder.AllocParameters<SparseVolumeTextureUpdatePageTableCS::FParameters>();
+					PassParameters->PageTable = PageTableUAV;
+					PassParameters->PageTableUpdates = UpdatesBufferSRV;
+					PassParameters->UpdateCoordOffset = UpdatesOffset;
+					PassParameters->UpdatePayloadOffset = MaxNumUpdates + UpdatesOffset;
+					PassParameters->NumUpdates = Batch.NumUpdates;
+
+					FComputeShaderUtils::AddPass(
+						GraphBuilder,
+						RDG_EVENT_NAME("Update SVT PageTable (UpdateCount: %u)", Batch.NumUpdates),
+						bAsyncCompute ? ERDGPassFlags::AsyncCompute : ERDGPassFlags::Compute,
+						ComputeShader,
+						PassParameters,
+						FComputeShaderUtils::GetGroupCount(Batch.NumUpdates, 64)
+					);
+
+					UpdatesOffset += Batch.NumUpdates;
+				}
+			}
+		}
+		
+		Release();
+	}
+
+private:
+	struct FBatch
+	{
+		FRHITexture* PageTable;
+		int32 MipLevel;
+		int32 NumUpdates;
+
+		FBatch() = default;
+		FBatch(FRHITexture* InPageTable, int32 InMipLevel) : PageTable(InPageTable), MipLevel(InMipLevel), NumUpdates(0) {}
+	};
+
+	TRefCountPtr<FRDGPooledBuffer> UpdatesUploadBuffer;
+	TArray<FBatch> Batches;
+	uint8* DataPtr = nullptr;
+	int32 NumWrittenUpdates = 0;
+	int32 MaxNumUpdates = 0;
+
+
+	void ResetState()
+	{
+		Batches.Reset();
+		DataPtr = nullptr;
+		NumWrittenUpdates = 0;
+		MaxNumUpdates = 0;
+	}
+};
+
+// Updates entries in the streaming info buffers of multiple SVTs
+class FStreamingInfoBufferUpdater
+{
+public:
+	FStreamingInfoBufferUpdater()
+	{
+		ResetState();
+	}
+
+	void Add(TRefCountPtr<FRDGPooledBuffer> StreamingInfoBuffer, int32 FrameIndex, int32 LowestResidentMipLevel)
+	{
+		FBatch* Batch = Batches.IsEmpty() ? nullptr : &Batches.Last();
+		if (!Batch || Batch->StreamingInfoBuffer != StreamingInfoBuffer)
+		{
+			Batch = &Batches.Add_GetRef(FBatch(StreamingInfoBuffer, Updates.Num()));
+		}
+
+		Updates.Add(FrameIndex);
+		Updates.Add(LowestResidentMipLevel);
+
+		++Batch->NumUpdates;
+	}
+
+	void Apply(FRDGBuilder& GraphBuilder)
+	{
+		if (!Updates.IsEmpty())
+		{
+			TRefCountPtr<FRDGPooledBuffer> UpdatesUploadBuffer;
+			{
+				// Add EBufferUsageFlags::Dynamic to skip the unneeded copy from upload to VRAM resource on d3d12 RHI
+				FRDGBufferDesc BufferDesc = FRDGBufferDesc::CreateByteAddressUploadDesc(Updates.Num() * sizeof(uint32));
+				BufferDesc.Usage |= EBufferUsageFlags::Dynamic;
+				AllocatePooledBuffer(BufferDesc, UpdatesUploadBuffer, TEXT("SparseVolumeTexture.StreamingInfoUploadBuffer"));
+
+				void* DataPtr = RHILockBuffer(UpdatesUploadBuffer->GetRHI(), 0, Updates.Num() * sizeof(uint32), RLM_WriteOnly);
+				FMemory::Memcpy(DataPtr, Updates.GetData(), Updates.Num() * sizeof(uint32));
+				RHIUnlockBuffer(UpdatesUploadBuffer->GetRHI());
+			}
+
+			auto ComputeShader = GetGlobalShaderMap(GMaxRHIFeatureLevel)->GetShader<SparseVolumeTextureUpdateStreamingInfoBufferCS>();
+			const bool bAsyncCompute = GSupportsEfficientAsyncCompute && (GSVTStreamingAsyncCompute != 0);
+
+			uint32 UpdatesOffset = 0;
+			for (const FBatch& Batch : Batches)
+			{
+				FRDGBufferUAV* StreamingInfoBufferUAV = GraphBuilder.CreateUAV(GraphBuilder.RegisterExternalBuffer(Batch.StreamingInfoBuffer), PF_R32_UINT);
+				FRDGBufferSRV* UpdatesBufferSRV = GraphBuilder.CreateSRV(GraphBuilder.RegisterExternalBuffer(UpdatesUploadBuffer));
+
+				SparseVolumeTextureUpdateStreamingInfoBufferCS::FParameters* PassParameters = GraphBuilder.AllocParameters<SparseVolumeTextureUpdateStreamingInfoBufferCS::FParameters>();
+				PassParameters->StreamingInfoBuffer = StreamingInfoBufferUAV;
+				PassParameters->StreamingInfoBufferUpdates = UpdatesBufferSRV;
+				PassParameters->UpdateOffset = UpdatesOffset;
+				PassParameters->NumUpdates = Batch.NumUpdates;
+
+				FComputeShaderUtils::AddPass(
+					GraphBuilder,
+					RDG_EVENT_NAME("Update SVT StreamingInfoBuffer (UpdateCount: %u)", Batch.NumUpdates),
+					bAsyncCompute ? ERDGPassFlags::AsyncCompute : ERDGPassFlags::Compute,
+					ComputeShader,
+					PassParameters,
+					FComputeShaderUtils::GetGroupCount(Batch.NumUpdates, 64)
+				);
+
+				UpdatesOffset += Batch.NumUpdates;
+			}
+		}
+		
+		ResetState();
+	}
+
+private:
+	struct FBatch
+	{
+		TRefCountPtr<FRDGPooledBuffer> StreamingInfoBuffer;
+		int32 UpdatesOffset;
+		int32 NumUpdates;
+
+		FBatch() = default;
+		FBatch(TRefCountPtr<FRDGPooledBuffer> InStreamingInfoBuffer, int32 InUpdatesOffset) : StreamingInfoBuffer(InStreamingInfoBuffer), UpdatesOffset(InUpdatesOffset), NumUpdates(0) {}
+	};
+
+	TArray<FBatch> Batches;
+	TArray<uint32> Updates;
+
+	void ResetState()
+	{
+		Batches.Reset();
+		Updates.Reset();
+	}
+};
+
+struct FStreamingUpdateParameters
+{
+	FStreamingManager* StreamingManager = nullptr;
+};
+
+class FStreamingUpdateTask
+{
+public:
+	explicit FStreamingUpdateTask(const FStreamingUpdateParameters& InParams) : Parameters(InParams) {}
+
+	FStreamingUpdateParameters Parameters;
+
+	void DoTask(ENamedThreads::Type CurrentThread, const FGraphEventRef& MyCompletionGraphEvent)
+	{
+		Parameters.StreamingManager->InstallReadyMipLevels();
+	}
+
+	static ESubsequentsMode::Type	GetSubsequentsMode() { return ESubsequentsMode::TrackSubsequents; }
+	ENamedThreads::Type				GetDesiredThread() { return ENamedThreads::AnyNormalThreadNormalTask; }
+	FORCEINLINE TStatId				GetStatId() const { return TStatId(); }
+};
+
+FStreamingManager::FStreamingManager()
+{
+
+}
+
+void FStreamingManager::InitRHI()
+{
+	using namespace UE::DerivedData;
+
+	if (!DoesPlatformSupportSparseVolumeTexture(GMaxRHIShaderPlatform))
+	{
+		return;
+	}
+
+	MaxPendingMipLevels = GSVTStreamingMaxPendingMipLevels;
+	PendingMipLevels.SetNum(MaxPendingMipLevels);
+	PageTableUpdater = new FPageTableUpdater();
+	StreamingInfoBufferUpdater = new FStreamingInfoBufferUpdater();
+
+#if WITH_EDITORONLY_DATA
+	RequestOwner = new FRequestOwner(EPriority::Normal);
+#endif
+}
+
+void FStreamingManager::ReleaseRHI()
+{
+	if (!DoesPlatformSupportSparseVolumeTexture(GMaxRHIShaderPlatform))
+	{
+		return;
+	}
+
+#if WITH_EDITORONLY_DATA
+	delete RequestOwner;
+	RequestOwner = nullptr;
+#endif
+
+	delete PageTableUpdater;
+	PageTableUpdater = nullptr;
+	delete StreamingInfoBufferUpdater;
+	StreamingInfoBufferUpdater = nullptr;
+}
+
+void FStreamingManager::Add_GameThread(UStreamableSparseVolumeTexture* SparseVolumeTexture)
+{
+	if (!DoesPlatformSupportSparseVolumeTexture(GMaxRHIShaderPlatform) || !SparseVolumeTexture)
+	{
+		return;
+	}
+
+	FNewSparseVolumeTextureInfo NewSVTInfo{};
+	const int32 NumFrames = SparseVolumeTexture->GetNumFrames();
+	NewSVTInfo.SVT = SparseVolumeTexture;
+	NewSVTInfo.FormatA = SparseVolumeTexture->GetFormat(0);
+	NewSVTInfo.FormatB = SparseVolumeTexture->GetFormat(1);
+	NewSVTInfo.FallbackValueA = SparseVolumeTexture->GetFallbackValue(0);
+	NewSVTInfo.FallbackValueB = SparseVolumeTexture->GetFallbackValue(1);
+	NewSVTInfo.NumMipLevelsGlobal = SparseVolumeTexture->GetNumMipLevels();
+	NewSVTInfo.FrameInfo.SetNum(NumFrames);
+
+	for (int32 FrameIdx = 0; FrameIdx < NumFrames; ++FrameIdx)
+	{
+		USparseVolumeTextureFrame* SVTFrame = SparseVolumeTexture->GetFrame(FrameIdx);
+		FFrameInfo& FrameInfo = NewSVTInfo.FrameInfo[FrameIdx];
+		FrameInfo.Resources = SVTFrame->GetResources();
+		FrameInfo.TextureRenderResources = SVTFrame->TextureRenderResources;
+		check(FrameInfo.TextureRenderResources);
+	}
+
+
+	ENQUEUE_RENDER_COMMAND(SVTAdd)(
+		[this, NewSVTInfoCaptured = MoveTemp(NewSVTInfo), SVTName = SparseVolumeTexture->GetName()](FRHICommandListImmediate& RHICmdList) mutable /* Required to be able to move from NewSVTInfoCaptured inside the lambda */
+		{
+			// We need to fully initialize the SVT streaming state (including resource creation) to ensure that valid resources exist before FillUniformBuffers() is called.
+			// This is why we can't defer resource creation until BeginAsyncUpdate() is called.
+			FRDGBuilder GraphBuilder(RHICmdList, RDG_EVENT_NAME("SVT::FStreamingManager::Add(%s)", *SVTName));
+			AddInternal(GraphBuilder, MoveTemp(NewSVTInfoCaptured));
+			GraphBuilder.Execute();
+		});
+}
+
+void FStreamingManager::Remove_GameThread(UStreamableSparseVolumeTexture* SparseVolumeTexture)
+{
+	if (!DoesPlatformSupportSparseVolumeTexture(GMaxRHIShaderPlatform) || !SparseVolumeTexture)
+	{
+		return;
+	}
+	ENQUEUE_RENDER_COMMAND(SVTRemove)(
+		[this, SparseVolumeTexture](FRHICommandListImmediate& RHICmdList)
+		{
+			RemoveInternal(SparseVolumeTexture);
+		});
+}
+
+void FStreamingManager::BeginAsyncUpdate(FRDGBuilder& GraphBuilder)
+{
+	check(IsInRenderingThread());
+	check(!AsyncState.bUpdateActive);
+	if (!DoesPlatformSupportSparseVolumeTexture(GMaxRHIShaderPlatform) || StreamingInfo.IsEmpty())
+	{
+		return;
+	}
+
+	TRACE_CPUPROFILER_EVENT_SCOPE(SVT::FStreamingManager::BeginAsyncUpdate);
+
+#if SVT_STREAMING_LOG_VERBOSE
+	UE_LOG(LogSparseVolumeTextureStreamingManager, Display, TEXT("SVT Streaming Update %i"), NextUpdateIndex);
+#endif
+
+	AsyncState = {};
+	AsyncState.bUpdateActive = true;
+
+	// For debugging, we can stream out ALL tiles
+	if (GSVTStreamingEmptyPhysicalTileTextures != 0)
+	{
+		TArray<FLRUNode*> MipLevelsToFree;
+		for (auto& Pair : StreamingInfo)
+		{
+			MipLevelsToFree.Reset();
+			FStreamingInfo& SVTInfo = Pair.Value;
+			const int32 NumFrames = SVTInfo.PerFrameInfo.Num();
+			const int32 NumMipLevelsGlobal = SVTInfo.NumMipLevelsGlobal;
+			
+			for (int32 MipLevel = 0; MipLevel < NumMipLevelsGlobal; ++MipLevel)
+			{
+				for (auto& Node : SVTInfo.PerMipLRULists[MipLevel])
+				{
+					MipLevelsToFree.Add(&Node);
+				}
+			}
+			for (FLRUNode* Node : MipLevelsToFree)
+			{
+				StreamOutMipLevel(SVTInfo, Node);
+			}
+		}
+
+		GSVTStreamingEmptyPhysicalTileTextures = 0;
+	}
+
+	AddParentRequests();
+	const int32 MaxSelectedMipLevels = MaxPendingMipLevels - NumPendingMipLevels;
+	SelectHighestPriorityRequestsAndUpdateLRU(MaxSelectedMipLevels);
+	IssueRequests(MaxSelectedMipLevels);
+	AsyncState.NumReadyMipLevels = DetermineReadyMipLevels();
+
+	// Do a first pass over all the mips to be uploaded to compute the upload buffer size requirements.
+	int32 NumPageTableUpdatesTotal = 0;
+	TileDataTexturesToUpdate.Reset();
+	{
+		const int32 StartPendingMipLevelIndex = (NextPendingMipLevelIndex + MaxPendingMipLevels - NumPendingMipLevels) % MaxPendingMipLevels;
+		for (int32 i = 0; i < AsyncState.NumReadyMipLevels; ++i)
+		{
+			const int32 PendingMipLevelIndex = (StartPendingMipLevelIndex + i) % MaxPendingMipLevels;
+			FPendingMipLevel& PendingMipLevel = PendingMipLevels[PendingMipLevelIndex];
+
+			FStreamingInfo* SVTInfo = StreamingInfo.Find(PendingMipLevel.SparseVolumeTexture);
+			if (!SVTInfo || (SVTInfo->PerFrameInfo[PendingMipLevel.FrameIndex].LowestRequestedMipLevel > PendingMipLevel.MipLevelIndex))
+			{
+				continue; // Skip mip level install. SVT no longer exists or mip level was "streamed out" before it was even installed in the first place.
+			}
+
+			const FResources* Resources = SVTInfo->PerFrameInfo[PendingMipLevel.FrameIndex].Resources;
+			SVTInfo->TileDataTexture->NumTilesToUpload += Resources->MipLevelStreamingInfo[PendingMipLevel.MipLevelIndex].NumPhysicalTiles;
+			TileDataTexturesToUpdate.Add(SVTInfo->TileDataTexture);
+			NumPageTableUpdatesTotal += Resources->MipLevelStreamingInfo[PendingMipLevel.MipLevelIndex].PageTableSize / (2 * sizeof(uint32));
+		}
+
+		PageTableUpdater->Init(GraphBuilder, NumPageTableUpdatesTotal, 1);
+
+		for (FTileDataTexture* TileDataTexture : TileDataTexturesToUpdate)
+		{
+			TileDataTexture->TileUploader->Init(GraphBuilder, TileDataTexture->NumTilesToUpload, TileDataTexture->FormatA, TileDataTexture->FormatB);
+		}
+	}
+
+	// Start async processing
+	FStreamingUpdateParameters Parameters;
+	Parameters.StreamingManager = this;
+	
+	check(AsyncTaskEvents.IsEmpty());
+	if (GSVTStreamingAsyncThread)
+	{
+		AsyncTaskEvents.Add(TGraphTask<FStreamingUpdateTask>::CreateTask().ConstructAndDispatchWhenReady(Parameters));
+	}
+	else
+	{
+		InstallReadyMipLevels();
+	}
+}
+
+void FStreamingManager::EndAsyncUpdate(FRDGBuilder& GraphBuilder)
+{
+	check(IsInRenderingThread());
+	if (!DoesPlatformSupportSparseVolumeTexture(GMaxRHIShaderPlatform) || StreamingInfo.IsEmpty())
+	{
+		return;
+	}
+	check(AsyncState.bUpdateActive);
+
+	TRACE_CPUPROFILER_EVENT_SCOPE(SVT::FStreamingManager::EndAsyncUpdate);
+
+	// Wait for async processing to finish
+	if (GSVTStreamingAsyncThread)
+	{
+		check(!AsyncTaskEvents.IsEmpty());
+		FTaskGraphInterface::Get().WaitUntilTasksComplete(AsyncTaskEvents, ENamedThreads::GetRenderThread_Local());
+	}
+	AsyncTaskEvents.Empty();
+
+	// Clear unused mip levels to 0. SVT_TODO: We can probably skip this because the page table lookup in the shader is clamped
+	if (!PageTableClears.IsEmpty())
+	{
+		for (auto& Clear : PageTableClears)
+		{
+			FRDGTexture* PageTableTextureRDG = GraphBuilder.RegisterExternalTexture(CreateRenderTarget(Clear.PageTableTexture, TEXT("SparseVolumeTexture.PageTableTexture")));
+			FRDGTextureUAVDesc UAVDesc(PageTableTextureRDG, static_cast<uint8>(Clear.MipLevel), PF_R32_UINT);
+			AddClearUAVPass(GraphBuilder, GraphBuilder.CreateUAV(UAVDesc), FUintVector4(ForceInitToZero));
+		}
+		PageTableClears.Reset();
+	}
+
+	// Issue the actual data uploads
+	for (FTileDataTexture* TileDataTexture : TileDataTexturesToUpdate)
+	{
+		TileDataTexture->TileUploader->ResourceUploadTo(GraphBuilder, TileDataTexture->TileDataTextureARHIRef, TileDataTexture->TileDataTextureBRHIRef);
+		TileDataTexture->NumTilesToUpload = 0;
+	}
+
+	// Update streaming info buffers
+	for (FStreamingInfo* SVTInfo : InvalidatedStreamingInfos)
+	{
+#if DO_CHECK
+		bool bSVTInfoExists = false;
+		for (const auto& Pair : StreamingInfo)
+		{
+			if (&Pair.Value == SVTInfo)
+			{
+				bSVTInfoExists = true;
+				break;
+			}
+		}
+		check(bSVTInfoExists);
+#endif
+		
+		for (TConstSetBitIterator It(SVTInfo->DirtyStreamingInfoData); It; ++It)
+		{
+			const int32 FrameIndex = It.GetIndex();
+			StreamingInfoBufferUpdater->Add(SVTInfo->StreamingInfoBuffer, FrameIndex, SVTInfo->PerFrameInfo[FrameIndex].LowestResidentMipLevel);
+		}
+	}
+	InvalidatedStreamingInfos.Reset();
+	StreamingInfoBufferUpdater->Apply(GraphBuilder);
+
+	PageTableUpdater->Apply(GraphBuilder);
+
+	check(AsyncState.NumReadyMipLevels <= NumPendingMipLevels);
+	NumPendingMipLevels -= AsyncState.NumReadyMipLevels;
+	++NextUpdateIndex;
+	AsyncState.bUpdateActive = false;
+
+#if DO_CHECK
+	for (const auto& Pair : StreamingInfo)
+	{
+#if SVT_STREAMING_LOG_VERBOSE
+		FString ResidentMipLevelsStr = TEXT("");
+#endif
+		const int32 NumFrames = Pair.Value.PerFrameInfo.Num();
+		for (int32 FrameIndex = 0; FrameIndex < NumFrames; ++FrameIndex)
+		{
+			const auto& FrameInfo = Pair.Value.PerFrameInfo[FrameIndex];
+			check(FrameInfo.LowestResidentMipLevel <= (FrameInfo.NumMipLevels - 1));
+			check(FrameInfo.LowestRequestedMipLevel <= FrameInfo.LowestResidentMipLevel);
+			check(FrameInfo.TextureRenderResources->GetNumLogicalMipLevels() == FrameInfo.NumMipLevels);
+
+#if SVT_STREAMING_LOG_VERBOSE
+			ResidentMipLevelsStr += FString::Printf(TEXT("%i"), FrameInfo.LowestResidentMipLevel);
+#endif
+		}
+#if SVT_STREAMING_LOG_VERBOSE
+		UE_LOG(LogSparseVolumeTextureStreamingManager, Display, TEXT("%s"), *ResidentMipLevelsStr);
+#endif
+	}
+#endif // DO_CHECK
+}
+
+void FStreamingManager::Request_GameThread(UStreamableSparseVolumeTexture* SparseVolumeTexture, float FrameIndex, int32 MipLevel)
+{
+	if (!DoesPlatformSupportSparseVolumeTexture(GMaxRHIShaderPlatform) || !SparseVolumeTexture)
+	{
+		return;
+	}
+	ENQUEUE_RENDER_COMMAND(SVTRequest)(
+		[this, SparseVolumeTexture, FrameIndex, MipLevel](FRHICommandListImmediate& RHICmdList)
+		{
+			Request(SparseVolumeTexture, FrameIndex, MipLevel);
+		});
+}
+
+void FStreamingManager::Request(UStreamableSparseVolumeTexture* SparseVolumeTexture, float FrameIndex, int32 MipLevel)
+{
+	check(IsInRenderingThread());
+	FStreamingInfo* SVTInfo = StreamingInfo.Find(SparseVolumeTexture);
+	if (SVTInfo)
+	{
+		const int32 NumFrames = SVTInfo->PerFrameInfo.Num();
+		const int32 FrameIndexI32 = static_cast<int32>(FrameIndex);
+		if (FrameIndexI32 < 0 || FrameIndexI32 >= NumFrames)
+		{
+			return;
+		}
+
+		// Try to find a FStreamingWindow around the requested frame index. This will inform us about which direction we need to prefetch into.
+		FStreamingWindow* StreamingWindow = nullptr;
+		for (FStreamingWindow& Window : SVTInfo->StreamingWindows)
+		{
+			if (FMath::Abs(FrameIndex - Window.CenterFrame) <= FStreamingWindow::WindowSize)
+			{
+				StreamingWindow = &Window;
+				break;
+			}
+		}
+		// Found an existing window!
+		if (StreamingWindow)
+		{
+			const bool bForward = StreamingWindow->LastCenterFrame <= FrameIndex;
+			if (StreamingWindow->LastRequested < NextUpdateIndex)
+			{
+				StreamingWindow->LastCenterFrame = StreamingWindow->CenterFrame;
+				StreamingWindow->CenterFrame = FrameIndex;
+				StreamingWindow->NumRequestsThisUpdate = 1;
+				StreamingWindow->LastRequested = NextUpdateIndex;
+				StreamingWindow->bPlayForward = bForward;
+				StreamingWindow->bPlayBackward = !bForward;
+			}
+			else
+			{
+				// Update the average center frame
+				StreamingWindow->CenterFrame = (StreamingWindow->CenterFrame * StreamingWindow->NumRequestsThisUpdate + FrameIndex) / (StreamingWindow->NumRequestsThisUpdate + 1.0f);
+				++StreamingWindow->NumRequestsThisUpdate;
+				StreamingWindow->bPlayForward |= bForward;
+				StreamingWindow->bPlayBackward |= !bForward;
+			}
+		}
+		// No existing window. Create a new one.
+		else
+		{
+			StreamingWindow = &SVTInfo->StreamingWindows.AddDefaulted_GetRef();
+			StreamingWindow->CenterFrame = FrameIndex;
+			StreamingWindow->LastCenterFrame = FrameIndex;
+			StreamingWindow->NumRequestsThisUpdate = 1;
+			StreamingWindow->LastRequested = NextUpdateIndex;
+			StreamingWindow->bPlayForward = true; // No prior data, so just take a guess that playback is forwards
+			StreamingWindow->bPlayBackward = false;
+		}
+
+		check(StreamingWindow);
+		const int32 OffsetMagnitude = GSVTStreamingNumPrefetchFrames;
+		const int32 LowerFrameOffset = StreamingWindow->bPlayBackward ? -OffsetMagnitude : 0;
+		const int32 UpperFrameOffset = StreamingWindow->bPlayForward ? OffsetMagnitude : 0;
+
+		for (int32 i = LowerFrameOffset; i <= UpperFrameOffset; ++i)
+		{
+			const int32 RequestFrameIndex = (static_cast<int32>(FrameIndex) + i + NumFrames) % NumFrames;
+			const int32 RequestMipLevelOffset = FMath::Abs(i) + GSVTStreamingPrefetchMipLevelBias;
+			FStreamingRequest Request;
+			Request.Key.SVT = SparseVolumeTexture;
+			Request.Key.FrameIndex = RequestFrameIndex;
+			Request.Key.MipLevelIndex = FMath::Clamp(MipLevel + RequestMipLevelOffset, 0, SVTInfo->PerFrameInfo[RequestFrameIndex].NumMipLevels);
+			Request.Priority = FMath::Max(0, OffsetMagnitude - FMath::Abs(i));
+			AddRequest(Request);
+		}
+
+		// Clean up unused streaming windows
+		SVTInfo->StreamingWindows.RemoveAll([&](const FStreamingWindow& Window) { return (NextUpdateIndex - Window.LastRequested) > 5; });
+	}
+}
+
+void FStreamingManager::AddInternal(FRDGBuilder& GraphBuilder, FNewSparseVolumeTextureInfo&& NewSVTInfo)
+{
+	check(IsInRenderingThread());
+	check(!AsyncState.bUpdateActive);
+	if (!ensure(!StreamingInfo.Contains(NewSVTInfo.SVT)))
+	{
+		return;
+	}
+
+	TRACE_CPUPROFILER_EVENT_SCOPE(SVT::FStreamingManager::AddInternal);
+
+	const int32 NumFrames = NewSVTInfo.FrameInfo.Num();
+
+	FStreamingInfo& SVTInfo = StreamingInfo.Add(NewSVTInfo.SVT);
+	SVTInfo.FormatA = NewSVTInfo.FormatA;
+	SVTInfo.FormatB = NewSVTInfo.FormatB;
+	SVTInfo.FallbackValueA = NewSVTInfo.FallbackValueA;
+	SVTInfo.FallbackValueB = NewSVTInfo.FallbackValueB;
+	SVTInfo.NumMipLevelsGlobal = NewSVTInfo.NumMipLevelsGlobal;
+	SVTInfo.LastRequested = 0;
+	SVTInfo.PerFrameInfo = MoveTemp(NewSVTInfo.FrameInfo);
+	SVTInfo.LRUNodes.SetNum(NumFrames * SVTInfo.NumMipLevelsGlobal);
+	SVTInfo.PerMipLRULists.SetNum(SVTInfo.NumMipLevelsGlobal);
+
+	int32 NumRootPhysicalTiles = 0;
+	int32 MaxNumPhysicalTiles = 0;
+	for (int32 FrameIdx = 0; FrameIdx < NumFrames; ++FrameIdx)
+	{
+		FFrameInfo& FrameInfo = SVTInfo.PerFrameInfo[FrameIdx];
+		check(FrameInfo.TextureRenderResources && FrameInfo.TextureRenderResources->IsInitialized());
+		const FResources* Resources = FrameInfo.Resources;
+
+		FrameInfo.NumMipLevels = Resources->MipLevelStreamingInfo.Num();
+		FrameInfo.LowestRequestedMipLevel = FrameInfo.NumMipLevels - 1;
+		FrameInfo.LowestResidentMipLevel = FrameInfo.NumMipLevels - 1;
+		FrameInfo.TileAllocations.SetNum(FrameInfo.NumMipLevels);
+		for (int32 MipLevel = 0; MipLevel < FrameInfo.NumMipLevels; ++MipLevel)
+		{
+			FrameInfo.TileAllocations[MipLevel].SetNumZeroed(Resources->MipLevelStreamingInfo[MipLevel].NumPhysicalTiles);
+		}
+		
+		int32 NumPhysicalTiles = 0;
+		for (const FMipLevelStreamingInfo& MipLevelStreamingInfo : Resources->MipLevelStreamingInfo)
+		{
+			NumPhysicalTiles += MipLevelStreamingInfo.NumPhysicalTiles;
+		}
+
+		MaxNumPhysicalTiles = FMath::Max(NumPhysicalTiles, MaxNumPhysicalTiles);
+		if (NumPhysicalTiles > 0)
+		{
+			++NumRootPhysicalTiles;
+		}
+
+		for (int32 MipIdx = 0; MipIdx < SVTInfo.NumMipLevelsGlobal; ++MipIdx)
+		{
+			FLRUNode& LRUNode = SVTInfo.LRUNodes[FrameIdx * SVTInfo.NumMipLevelsGlobal + MipIdx];
+			LRUNode.Reset();
+			LRUNode.FrameIndex = FrameIdx;
+			LRUNode.MipLevelIndex = MipIdx < FrameInfo.NumMipLevels ? MipIdx : INDEX_NONE;
+
+			if ((MipIdx + 1) < FrameInfo.NumMipLevels)
+			{
+				LRUNode.NextHigherMipLevel = &SVTInfo.LRUNodes[FrameIdx * SVTInfo.NumMipLevelsGlobal + (MipIdx + 1)];
+			}
+		}
+	}
+
+	// Create RHI resources and upload root tile data
+	{
+		const int32 TileFactor = NumFrames <= 1 ? 1 : 3;
+		const int32 NumPhysicalTilesCapacity = NumRootPhysicalTiles + (TileFactor * MaxNumPhysicalTiles);
+		const FIntVector3 TileDataVolumeResolution = ComputeTileDataVolumeResolution2(NumPhysicalTilesCapacity);
+		const FIntVector3 TileDataVolumeResolutionInTiles = TileDataVolumeResolution / SPARSE_VOLUME_TILE_RES_PADDED;
+
+		SVTInfo.TileDataTexture = new FTileDataTexture(TileDataVolumeResolutionInTiles, SVTInfo.FormatA, SVTInfo.FormatB);
+		SVTInfo.TileDataTexture->InitResource();
+
+		// Create streaming info buffer
+		{
+			FRDGBufferRef StreamingInfoBufferRDG = GraphBuilder.CreateBuffer(FRDGBufferDesc::CreateByteAddressDesc(sizeof(uint32) * NumFrames), TEXT("SparseVolumeTexture.StreamingInfo.RHIBuffer"));
+			SVTInfo.StreamingInfoBuffer = GraphBuilder.ConvertToExternalBuffer(StreamingInfoBufferRDG);
+			SVTInfo.StreamingInfoBufferSRVRHIRef = SVTInfo.StreamingInfoBuffer->GetSRV();
+			SVTInfo.DirtyStreamingInfoData.Init(true, NumFrames);
+		}
+
+		FTileUploader RootTileUploader;
+		RootTileUploader.Init(GraphBuilder, NumRootPhysicalTiles + 1 /*null tile*/, SVTInfo.FormatA, SVTInfo.FormatB);
+
+		// Allocate null tile
+		{
+			const uint32 NullTileCoord = SVTInfo.TileDataTexture->Allocate();
+			check(NullTileCoord == 0);
+			uint8* TileCoordsPtr = nullptr;
+			uint8* DataAPtr = nullptr;
+			uint8* DataBPtr = nullptr;
+			RootTileUploader.Add_GetRef(1, TileCoordsPtr, DataAPtr, DataBPtr);
+			FMemory::Memcpy(TileCoordsPtr, &NullTileCoord, sizeof(NullTileCoord));
+			for (int32 VoxelIdx = 0; VoxelIdx < UE::SVT::SVTNumVoxelsPerPaddedTile; ++VoxelIdx)
+			{
+				if (SVTInfo.FormatA != PF_Unknown)
+				{
+					WriteVoxel(VoxelIdx, DataAPtr, SVTInfo.FormatA, SVTInfo.FallbackValueA);
+				}
+				if (SVTInfo.FormatB != PF_Unknown)
+				{
+					WriteVoxel(VoxelIdx, DataBPtr, SVTInfo.FormatB, SVTInfo.FallbackValueB);
+				}
+			}
+		}
+
+		// Process frames
+		for (int32 FrameIdx = 0; FrameIdx < NumFrames; ++FrameIdx)
+		{
+			FFrameInfo& FrameInfo = SVTInfo.PerFrameInfo[FrameIdx];
+			const FResources* Resources = FrameInfo.Resources;
+			const int32 NumMipLevels = Resources->MipLevelStreamingInfo.Num();
+
+			FrameInfo.LowestRequestedMipLevel = NumMipLevels - 1;
+			FrameInfo.LowestResidentMipLevel = NumMipLevels - 1;
+
+			// Initialize TextureRenderResources
+			FrameInfo.TextureRenderResources->PhysicalTileDataATextureRHI = SVTInfo.TileDataTexture->TileDataTextureARHIRef;
+			FrameInfo.TextureRenderResources->PhysicalTileDataBTextureRHI = SVTInfo.TileDataTexture->TileDataTextureBRHIRef;
+			FrameInfo.TextureRenderResources->StreamingInfoBufferSRVRHI = SVTInfo.StreamingInfoBufferSRVRHIRef;
+			FrameInfo.TextureRenderResources->Header = Resources->Header;
+			FrameInfo.TextureRenderResources->TileDataTextureResolution = SVTInfo.TileDataTexture->ResolutionInTiles * SPARSE_VOLUME_TILE_RES_PADDED;
+			FrameInfo.TextureRenderResources->FrameIndex = FrameIdx;
+			FrameInfo.TextureRenderResources->NumLogicalMipLevels = NumMipLevels;
+
+			// Create page table
+			{
+				// SVT_TODO: Currently we keep all mips of the page table resident. It would be better to stream in/out page table mips.
+				const int32 NumResidentMipLevels = NumMipLevels;
+				FIntVector3 PageTableResolution = Resources->Header.PageTableVolumeResolution;
+				PageTableResolution = FIntVector3(FMath::Max(1, PageTableResolution.X), FMath::Max(1, PageTableResolution.Y), FMath::Max(1, PageTableResolution.Z));
+
+				const EPixelFormat PageEntryFormat = PF_R32_UINT;
+				const FRHITextureCreateDesc Desc =
+					FRHITextureCreateDesc::Create3D(TEXT("SparseVolumeTexture.PageTable.RHITexture"), PageTableResolution.X, PageTableResolution.Y, PageTableResolution.Z, PageEntryFormat)
+					.SetFlags(ETextureCreateFlags::ShaderResource | ETextureCreateFlags::UAV)
+					.SetNumMips((uint8)NumResidentMipLevels);
+
+				FrameInfo.TextureRenderResources->PageTableTextureRHI = RHICreateTexture(Desc);
+			}
+
+			// Upload root mip data and update page tables
+			const FMipLevelStreamingInfo* RootStreamingInfo = !Resources->MipLevelStreamingInfo.IsEmpty() ? &Resources->MipLevelStreamingInfo.Last() : nullptr;
+			if (!Resources->RootData.IsEmpty() && RootStreamingInfo)
+			{
+				check((RootStreamingInfo->TileDataASize > 0) || (RootStreamingInfo->TileDataBSize > 0));
+
+				const uint32 TileCoord = SVTInfo.TileDataTexture->Allocate();
+				check(TileCoord != INDEX_NONE);
+				FrameInfo.TileAllocations.Last()[0] = TileCoord;
+				uint8* TileCoordsPtr = nullptr;
+				uint8* DataAPtr = nullptr;
+				uint8* DataBPtr = nullptr;
+				RootTileUploader.Add_GetRef(1, TileCoordsPtr, DataAPtr, DataBPtr);
+
+				FMemory::Memcpy(TileCoordsPtr, &TileCoord, sizeof(TileCoord));
+				if (RootStreamingInfo->TileDataASize > 0)
+				{
+					const uint8* Src = Resources->RootData.GetData() + RootStreamingInfo->TileDataAOffset;
+					check(DataAPtr);
+					FMemory::Memcpy(DataAPtr, Src, RootStreamingInfo->TileDataASize);
+				}
+				if (RootStreamingInfo->TileDataBSize > 0)
+				{
+					const uint8* Src = Resources->RootData.GetData() + RootStreamingInfo->TileDataBOffset;
+					check(DataBPtr);
+					FMemory::Memcpy(DataBPtr, Src, RootStreamingInfo->TileDataBSize);
+				}
+
+				// Update highest mip (1x1x1) in page table
+				const FUpdateTextureRegion3D UpdateRegion(0, 0, 0, 0, 0, 0, 1, 1, 1);
+				RHIUpdateTexture3D(FrameInfo.TextureRenderResources->PageTableTextureRHI, FrameInfo.TextureRenderResources->PageTableTextureRHI->GetDesc().NumMips - 1, UpdateRegion, sizeof(uint32), sizeof(uint32), (uint8*)&TileCoord);
+			}
+		}
+
+		RootTileUploader.ResourceUploadTo(GraphBuilder, SVTInfo.TileDataTexture->TileDataTextureARHIRef, SVTInfo.TileDataTexture->TileDataTextureBRHIRef);
+	}
+
+	InvalidatedStreamingInfos.Add(&SVTInfo);
+
+	// Add requests for all mips the first frame. This is necessary for cases where UAnimatedSparseVolumeTexture or UStaticSparseVolumeTexture
+	// are directly bound to the material without getting a specific frame through USparseVolumeTextureFrame::GetFrameAndIssueStreamingRequest().
+	const int32 NumMipLevelsFrame0 = SVTInfo.PerFrameInfo[0].NumMipLevels;
+	for (int32 MipLevel = 0; (MipLevel + 1) < NumMipLevelsFrame0; ++MipLevel)
+	{
+		FStreamingRequest Request;
+		Request.Key.SVT = NewSVTInfo.SVT;
+		Request.Key.FrameIndex = 0;
+		Request.Key.MipLevelIndex = MipLevel;
+		Request.Priority = MipLevel;
+		AddRequest(Request);
+	}
+}
+
+void FStreamingManager::RemoveInternal(UStreamableSparseVolumeTexture* SparseVolumeTexture)
+{
+	check(IsInRenderingThread());
+	check(!AsyncState.bUpdateActive);
+	FStreamingInfo* SVTInfo = StreamingInfo.Find(SparseVolumeTexture);
+	if (SVTInfo)
+	{
+		// The RHI resources in FTextureRenderResources are managed by FStreamingManager, so release them now.
+		for (FFrameInfo& FrameInfo : SVTInfo->PerFrameInfo)
+		{
+			FrameInfo.TextureRenderResources->PageTableTextureRHI.SafeRelease();
+			FrameInfo.TextureRenderResources->PhysicalTileDataATextureRHI.SafeRelease();
+			FrameInfo.TextureRenderResources->PhysicalTileDataBTextureRHI.SafeRelease();
+		}
+		if (SVTInfo->TileDataTexture)
+		{
+			SVTInfo->TileDataTexture->ReleaseResource();
+			delete SVTInfo->TileDataTexture;
+		}
+
+		StreamingInfo.Remove(SparseVolumeTexture);
+	}
+}
+
+bool FStreamingManager::AddRequest(const FStreamingRequest& Request)
+{
+	uint32* ExistingRequestPriority = RequestsHashTable.Find(Request.Key);
+	if (ExistingRequestPriority)
+	{
+		if (Request.Priority > *ExistingRequestPriority)
+		{
+			*ExistingRequestPriority = Request.Priority;
+			return true;
+		}
+		else
+		{
+			return false;
+		}
+	}
+	else
+	{
+		RequestsHashTable.Add(Request.Key, Request.Priority);
+		return true;
+	}
+}
+
+void FStreamingManager::AddParentRequests()
+{
+	TRACE_CPUPROFILER_EVENT_SCOPE(SVT::AddParentRequests);
+
+	ParentRequestsToAdd.Reset();
+	for (const auto& Request : RequestsHashTable)
+	{
+		FStreamingInfo* SVTInfo = StreamingInfo.Find(Request.Key.SVT);
+		check(SVTInfo);
+		const int32 NumStreamableMipLevels = SVTInfo->PerFrameInfo[Request.Key.FrameIndex].NumMipLevels - 1;
+		uint32 Priority = Request.Value + 1;
+		for (int32 MipLevelIndex = Request.Key.MipLevelIndex + 1; MipLevelIndex < NumStreamableMipLevels; ++MipLevelIndex)
+		{
+			FMipLevelKey ParentKey = Request.Key;
+			ParentKey.MipLevelIndex = MipLevelIndex;
+			
+			uint32* ExistingParentRequestPriority = RequestsHashTable.Find(ParentKey);
+			if (ExistingParentRequestPriority && Priority > *ExistingParentRequestPriority)
+			{
+				*ExistingParentRequestPriority = Priority;
+			}
+			else
+			{
+				ParentRequestsToAdd.Add(FStreamingRequest{ ParentKey, Priority });
+			}
+
+			++Priority;
+		}
+	}
+
+	for (const FStreamingRequest& Request : ParentRequestsToAdd)
+	{
+		AddRequest(Request);
+	}
+}
+
+void FStreamingManager::SelectHighestPriorityRequestsAndUpdateLRU(int32 MaxSelectedMipLevels)
+{
+	PrioritizedRequestsHeap.Reset();
+	SelectedMipLevels.Reset();
+
+	if (!RequestsHashTable.IsEmpty())
+	{
+		TRACE_CPUPROFILER_EVENT_SCOPE(SVT::SelectHighestPriorityRequestsAndUpdateLRU);
+
+		for (const auto& Request : RequestsHashTable)
+		{
+			FStreamingInfo* SVTInfo = StreamingInfo.Find(Request.Key.SVT);
+			check(SVTInfo);
+
+			// Discard invalid requests
+			if (Request.Key.FrameIndex < 0
+				|| Request.Key.FrameIndex >= SVTInfo->PerFrameInfo.Num()
+				|| Request.Key.MipLevelIndex >= (SVTInfo->PerFrameInfo[Request.Key.FrameIndex].NumMipLevels - 1))
+			{
+				continue;
+			}
+
+			const int32 LRUNodeIndex = Request.Key.FrameIndex * SVTInfo->NumMipLevelsGlobal + Request.Key.MipLevelIndex;
+			FLRUNode* LRUNode = &SVTInfo->LRUNodes[LRUNodeIndex];
+#if DO_CHECK
+			bool bFoundNodeInList = false;
+			for (auto& Node : SVTInfo->PerMipLRULists[Request.Key.MipLevelIndex])
+			{
+				if (&Node == LRUNode)
+				{
+					bFoundNodeInList = true;
+					break;
+				}
+			}
+#endif
+
+			const bool bIsAlreadyStreaming = Request.Key.MipLevelIndex >= SVTInfo->PerFrameInfo[Request.Key.FrameIndex].LowestRequestedMipLevel;
+			if (bIsAlreadyStreaming)
+			{
+				check(bFoundNodeInList);
+				// Update LastRequested and move to front of LRU
+				LRUNode->LastRequested = NextUpdateIndex;
+
+				// Unlink
+				LRUNode->Remove();
+
+				// Insert at the end of the LRU list
+				SVTInfo->PerMipLRULists[Request.Key.MipLevelIndex].AddTail(LRUNode);
+			}
+			else
+			{
+				check(!bFoundNodeInList);
+				PrioritizedRequestsHeap.Add(FStreamingRequest{ Request.Key, Request.Value });
+			}
+		}
+
+		auto PriorityPredicate = [](const auto& A, const auto& B) { return A.Priority > B.Priority; };
+		PrioritizedRequestsHeap.Heapify(PriorityPredicate);
+
+		while (SelectedMipLevels.Num() < MaxSelectedMipLevels && PrioritizedRequestsHeap.Num() > 0)
+		{
+			FStreamingRequest SelectedRequest;
+			PrioritizedRequestsHeap.HeapPop(SelectedRequest, PriorityPredicate, false /*bAllowShrinking*/);
+
+			FStreamingInfo* SVTInfo = StreamingInfo.Find(SelectedRequest.Key.SVT);
+			if (SVTInfo)
+			{
+				check(SelectedRequest.Key.FrameIndex < SVTInfo->PerFrameInfo.Num());
+				check(SelectedRequest.Key.MipLevelIndex < SVTInfo->PerFrameInfo[SelectedRequest.Key.FrameIndex].NumMipLevels);
+				SelectedMipLevels.Push(SelectedRequest.Key);
+			}
+		}
+
+		RequestsHashTable.Reset();
+	}
+}
+
+void FStreamingManager::IssueRequests(int32 MaxSelectedMipLevels)
+{
+	using namespace UE::DerivedData;
+
+	if (SelectedMipLevels.IsEmpty())
+	{
+		return;
+	}
+
+	TRACE_CPUPROFILER_EVENT_SCOPE(SVT::IssueRequests);
+
+#if WITH_EDITORONLY_DATA
+	TArray<FCacheGetChunkRequest> DDCRequests;
+	DDCRequests.Reserve(MaxSelectedMipLevels);
+#endif
+
+	FBulkDataBatchRequest::FBatchBuilder Batch = FBulkDataBatchRequest::NewBatch(SelectedMipLevels.Num());
+	bool bIssueIOBatch = false;
+
+	for (const FMipLevelKey& SelectedKey : SelectedMipLevels)
+	{
+		FStreamingInfo* SVTInfo = StreamingInfo.Find(SelectedKey.SVT);
+		check(SVTInfo);
+		check(SVTInfo->PerFrameInfo.Num() > SelectedKey.FrameIndex && SelectedKey.FrameIndex >= 0);
+		check(SVTInfo->PerFrameInfo[SelectedKey.FrameIndex].LowestRequestedMipLevel > SelectedKey.MipLevelIndex);
+		const FResources* Resources = SVTInfo->PerFrameInfo[SelectedKey.FrameIndex].Resources;
+		check((SelectedKey.MipLevelIndex + 1) < Resources->MipLevelStreamingInfo.Num()); // The lowest/last mip level is always resident and does not stream.
+		const FMipLevelStreamingInfo& MipLevelStreamingInfo = Resources->MipLevelStreamingInfo[SelectedKey.MipLevelIndex];
+
+		FTileDataTexture* TileDataTexture = SVTInfo->TileDataTexture;
+		check(TileDataTexture);
+
+		// Ensure that enough tiles are available in the tile texture
+		const int32 TileDataTextureCapacity = TileDataTexture->PhysicalTilesCapacity;
+		const int32 NumAvailableTiles = TileDataTexture->GetNumAvailableTiles();
+		const int32 NumRequiredTiles = MipLevelStreamingInfo.NumPhysicalTiles;
+		if (NumAvailableTiles < NumRequiredTiles)
+		{
+#if SVT_STREAMING_LOG_VERBOSE
+			UE_LOG(LogSparseVolumeTextureStreamingManager, Display, TEXT("(%i)%i IssueRequests() Frame %i Mip %i: Not enough tiles available (%i) to fit mip level (%i)"), 
+				NextUpdateIndex, NextPendingMipLevelIndex, SelectedKey.FrameIndex, SelectedKey.MipLevelIndex, NumAvailableTiles, NumRequiredTiles);
+#endif
+
+			// Try to free old mip levels, starting at higher resolution mips and going up the mip chain
+			TArray<FLRUNode*, TInlineAllocator<16>> MipLevelsToFree;
+			int32 NumNewlyAvailableTiles = 0;
+			const int32 NumMipLevelsGlobal = SVTInfo->NumMipLevelsGlobal;
+			for (int32 MipLevel = 0; MipLevel < NumMipLevelsGlobal && (NumAvailableTiles + NumNewlyAvailableTiles) < NumRequiredTiles; ++MipLevel)
+			{
+				for (auto& Node : SVTInfo->PerMipLRULists[MipLevel])
+				{
+					// Only free "leaf" mip levels with no higher resolution mip levels resident. Don't free mip levels requested this frame.
+					if (Node.RefCount == 0 && Node.LastRequested < NextUpdateIndex)
+					{
+						MipLevelsToFree.Add(&Node);
+						NumNewlyAvailableTiles += SVTInfo->PerFrameInfo[Node.FrameIndex].Resources->MipLevelStreamingInfo[Node.MipLevelIndex].NumPhysicalTiles;
+
+						// Decrement ref count of mip levels higher up the chain
+						FLRUNode* Dependency = Node.NextHigherMipLevel;
+						while (Dependency)
+						{
+							check(Dependency->RefCount > 0);
+							--Dependency->RefCount;
+							Dependency = Dependency->NextHigherMipLevel;
+						}
+					}
+
+					// Exit once we freed enough tiles
+					if ((NumAvailableTiles + NumNewlyAvailableTiles) >= NumRequiredTiles)
+					{
+						break;
+					}
+				}
+			}
+
+			// Free mip levels
+			for (FLRUNode* MipLevelToFree : MipLevelsToFree)
+			{
+				StreamOutMipLevel(*SVTInfo, MipLevelToFree);
+			}
+
+			// Couldn't free enough tiles, so skip this mip level
+			if ((NumAvailableTiles + NumNewlyAvailableTiles) < NumRequiredTiles)
+			{
+				UE_LOG(LogSparseVolumeTextureStreamingManager, Warning, TEXT("IssueRequests() Frame %i Mip %i: Not enough tiles available (%i) to fit mip level (%i) even after freeing"),
+					SelectedKey.FrameIndex, SelectedKey.MipLevelIndex, (NumAvailableTiles + NumNewlyAvailableTiles), NumRequiredTiles);
+				continue;
+			}
+		}
+
+#if DO_CHECK
+		for (auto& Pending : PendingMipLevels)
+		{
+			check(Pending.FrameIndex != SelectedKey.FrameIndex || Pending.MipLevelIndex != SelectedKey.MipLevelIndex);
+		}
+#endif
+
+		const int32 PendingMipLevelIndex = NextPendingMipLevelIndex;
+		FPendingMipLevel& PendingMipLevel = PendingMipLevels[PendingMipLevelIndex];
+		PendingMipLevel.Reset();
+		PendingMipLevel.SparseVolumeTexture = SelectedKey.SVT;
+		PendingMipLevel.FrameIndex = SelectedKey.FrameIndex;
+		PendingMipLevel.MipLevelIndex = SelectedKey.MipLevelIndex;
+		PendingMipLevel.IssuedInFrame = NextUpdateIndex;
+
+		const FByteBulkData& BulkData = Resources->StreamableMipLevels;
+#if WITH_EDITORONLY_DATA
+		const bool bDiskRequest = (!(Resources->ResourceFlags & EResourceFlag_StreamingDataInDDC) && !BulkData.IsBulkDataLoaded());
+#else
+		const bool bDiskRequest = true;
+#endif
+
+#if WITH_EDITORONLY_DATA
+		if (!bDiskRequest)
+		{
+			if (Resources->ResourceFlags & EResourceFlag_StreamingDataInDDC)
+			{
+				DDCRequests.Add(BuildDDCRequest(*Resources, MipLevelStreamingInfo, NextPendingMipLevelIndex));
+				PendingMipLevel.State = FPendingMipLevel::EState::DDC_Pending;
+			}
+			else
+			{
+				PendingMipLevel.State = FPendingMipLevel::EState::Memory;
+			}
+		}
+		else
+#endif
+		{
+			PendingMipLevel.RequestBuffer = FIoBuffer(MipLevelStreamingInfo.BulkSize); // SVT_TODO: Use FIoBuffer::Wrap with preallocated memory
+			Batch.Read(BulkData, MipLevelStreamingInfo.BulkOffset, MipLevelStreamingInfo.BulkSize, AIOP_Low, PendingMipLevel.RequestBuffer, PendingMipLevel.Request);
+			bIssueIOBatch = true;
+
+#if WITH_EDITORONLY_DATA
+			PendingMipLevel.State = FPendingMipLevel::EState::Disk;
+#endif
+		}
+
+		NextPendingMipLevelIndex = (NextPendingMipLevelIndex + 1) % MaxPendingMipLevels;
+		check(NumPendingMipLevels < MaxPendingMipLevels);
+		++NumPendingMipLevels;
+
+		FFrameInfo& FrameInfo = SVTInfo->PerFrameInfo[SelectedKey.FrameIndex];
+
+		// Allocate tiles in the tile data texture
+		{
+			TArray<uint32>& TileAllocations = FrameInfo.TileAllocations[SelectedKey.MipLevelIndex];
+			check(TileAllocations.Num() == NumRequiredTiles);
+			for (int32 TileIdx = 0; TileIdx < NumRequiredTiles; ++TileIdx)
+			{
+				const int32 TileCoord = TileDataTexture->Allocate();
+				check(TileCoord != INDEX_NONE);
+				TileAllocations[TileIdx] = TileCoord;
+			}
+		}
+
+		// Add to tail of LRU list
+		{
+			const int32 LRUNodeIndex = SelectedKey.FrameIndex * SVTInfo->NumMipLevelsGlobal + SelectedKey.MipLevelIndex;
+			FLRUNode* LRUNode = &SVTInfo->LRUNodes[LRUNodeIndex];
+			check(!LRUNode->IsInList());
+			LRUNode->LastRequested = NextUpdateIndex;
+			LRUNode->PendingMipLevelIndex = PendingMipLevelIndex;
+
+			FLRUNode* Dependency = LRUNode->NextHigherMipLevel;
+			while (Dependency)
+			{
+				++Dependency->RefCount;
+				Dependency = Dependency->NextHigherMipLevel;
+			}
+
+			SVTInfo->PerMipLRULists[SelectedKey.MipLevelIndex].AddTail(LRUNode);
+		}
+
+#if SVT_STREAMING_LOG_VERBOSE
+		UE_LOG(LogSparseVolumeTextureStreamingManager, Display, TEXT("(%i)%i StreamIn Frame %i OldReqMip %i, NewReqMip %i, ResMip %i"),
+			PendingMipLevel.IssuedInFrame, PendingMipLevelIndex,
+			SelectedKey.FrameIndex, 
+			FrameInfo.LowestRequestedMipLevel, SelectedKey.MipLevelIndex, 
+			FrameInfo.LowestResidentMipLevel);
+#endif
+
+		check(FrameInfo.LowestRequestedMipLevel == (SelectedKey.MipLevelIndex + 1));
+		FrameInfo.LowestRequestedMipLevel = SelectedKey.MipLevelIndex;
+	}
+
+#if WITH_EDITORONLY_DATA
+	if (!DDCRequests.IsEmpty())
+	{
+		RequestDDCData(DDCRequests);
+		DDCRequests.Empty();
+	}
+#endif
+
+	if (bIssueIOBatch)
+	{
+		(void)Batch.Issue();
+	}
+}
+
+void FStreamingManager::StreamOutMipLevel(FStreamingInfo& SVTInfo, FLRUNode* LRUNode)
+{
+	const int32 FrameIndex = LRUNode->FrameIndex;
+	const int32 MipLevelIndex = LRUNode->MipLevelIndex;
+
+	FFrameInfo& FrameInfo = SVTInfo.PerFrameInfo[FrameIndex];
+
+	check(FrameInfo.LowestResidentMipLevel >= MipLevelIndex); // mip might not have streamed in yet, so use >= instead of ==
+	check(FrameInfo.LowestRequestedMipLevel == MipLevelIndex);
+
+	// Cancel potential IO request
+	check((MipLevelIndex < FrameInfo.LowestResidentMipLevel) == (LRUNode->PendingMipLevelIndex != INDEX_NONE));
+	if (LRUNode->PendingMipLevelIndex != INDEX_NONE)
+	{
+		PendingMipLevels[LRUNode->PendingMipLevelIndex].Reset();
+		LRUNode->PendingMipLevelIndex = INDEX_NONE;
+	}
+
+	const int32 NewLowestRequestedMipLevel = MipLevelIndex + 1;
+	const int32 NewLowestResidentMipLevel = FMath::Max(MipLevelIndex + 1, FrameInfo.LowestResidentMipLevel);
+#if SVT_STREAMING_LOG_VERBOSE
+	UE_LOG(LogSparseVolumeTextureStreamingManager, Display, TEXT("(%i)%i StreamOut Frame %i OldReqMip %i, NewReqMip %i, OldResMip %i, NewResMip %i"),
+		NextUpdateIndex, NextPendingMipLevelIndex,
+		FrameIndex,
+		FrameInfo.LowestRequestedMipLevel, NewLowestRequestedMipLevel,
+		FrameInfo.LowestResidentMipLevel, NewLowestResidentMipLevel);
+#endif
+
+	// Only clear page table mip if the mip to be freed is actually resident.
+	const bool bNeedsPageTableClear = MipLevelIndex >= FrameInfo.LowestResidentMipLevel;
+	FrameInfo.LowestRequestedMipLevel = NewLowestRequestedMipLevel;
+	FrameInfo.LowestResidentMipLevel = NewLowestResidentMipLevel;
+	
+	// Update the streaming info buffer data
+	SVTInfo.DirtyStreamingInfoData[FrameIndex] = true;
+	InvalidatedStreamingInfos.Add(&SVTInfo);
+
+	// Unlink
+	LRUNode->Remove();
+	LRUNode->LastRequested = INDEX_NONE;
+
+	if (bNeedsPageTableClear)
+	{
+		PageTableClears.Push({ FrameInfo.TextureRenderResources->PageTableTextureRHI, MipLevelIndex });
+	}
+
+	// Free allocated tiles
+	for (uint32& TileCoord : FrameInfo.TileAllocations[MipLevelIndex])
+	{
+		SVTInfo.TileDataTexture->Free(TileCoord);
+		TileCoord = 0;
+	}
+}
+
+int32 FStreamingManager::DetermineReadyMipLevels()
+{
+	using namespace UE::DerivedData;
+
+	TRACE_CPUPROFILER_EVENT_SCOPE(SVT::DetermineReadyMipLevels);
+
+	const int32 StartPendingMipLevelIndex = (NextPendingMipLevelIndex + MaxPendingMipLevels - NumPendingMipLevels) % MaxPendingMipLevels;
+	int32 NumReadyMipLevels = 0;
+
+	for (int32 i = 0; i < NumPendingMipLevels; ++i)
+	{
+		const int32 PendingMipLevelIndex = (StartPendingMipLevelIndex + i) % MaxPendingMipLevels;
+		FPendingMipLevel& PendingMipLevel = PendingMipLevels[PendingMipLevelIndex];
+
+		FStreamingInfo* SVTInfo = StreamingInfo.Find(PendingMipLevel.SparseVolumeTexture);
+		if (!SVTInfo)
+		{
+#if WITH_EDITORONLY_DATA
+			// Resource is no longer there. Just mark as ready so it will be skipped later
+			PendingMipLevel.State = FPendingMipLevel::EState::DDC_Ready;
+#endif
+			continue; 
+		}
+
+		const FResources* Resources = SVTInfo->PerFrameInfo[PendingMipLevel.FrameIndex].Resources;
+
+#if WITH_EDITORONLY_DATA
+		if (PendingMipLevel.State == FPendingMipLevel::EState::DDC_Ready)
+		{
+			if (PendingMipLevel.RetryCount > 0)
+			{
+				check(SVTInfo);
+				UE_LOG(LogSparseVolumeTextureStreamingManager, Display, TEXT("SVT DDC retry succeeded for '%s' (frame %i, mip %i) on %i attempt."), 
+					*Resources->ResourceName, PendingMipLevel.FrameIndex, PendingMipLevel.MipLevelIndex, PendingMipLevel.RetryCount);
+			}
+		}
+		else if (PendingMipLevel.State == FPendingMipLevel::EState::DDC_Pending)
+		{
+			break;
+		}
+		else if (PendingMipLevel.State == FPendingMipLevel::EState::DDC_Failed)
+		{
+			PendingMipLevel.State = FPendingMipLevel::EState::DDC_Pending;
+
+			if (PendingMipLevel.RetryCount == 0) // Only warn on first retry to prevent spam
+			{
+				UE_LOG(LogSparseVolumeTextureStreamingManager, Warning, TEXT("SVT DDC request failed for '%s' (frame %i, mip %i). Retrying..."),
+					*Resources->ResourceName, PendingMipLevel.FrameIndex, PendingMipLevel.MipLevelIndex);
+			}
+
+			const FMipLevelStreamingInfo& MipLevelStreamingInfo = Resources->MipLevelStreamingInfo[PendingMipLevel.MipLevelIndex];
+			FCacheGetChunkRequest Request = BuildDDCRequest(*Resources, MipLevelStreamingInfo, PendingMipLevelIndex);
+			RequestDDCData(MakeArrayView(&Request, 1));
+
+			++PendingMipLevel.RetryCount;
+			break;
+		}
+		else if (PendingMipLevel.State == FPendingMipLevel::EState::Memory)
+		{
+			// Memory is always ready
+		}
+		else
+#endif // WITH_EDITORONLY_DATA
+		{
+#if WITH_EDITORONLY_DATA
+			check(PendingMipLevel.State == FPendingMipLevel::EState::Disk);
+#endif
+			if (PendingMipLevel.Request.IsCompleted())
+			{
+				if (!PendingMipLevel.Request.IsOk())
+				{
+					// Retry if IO request failed for some reason
+					const FMipLevelStreamingInfo& MipLevelStreamingInfo = Resources->MipLevelStreamingInfo[PendingMipLevel.MipLevelIndex];
+					UE_LOG(LogSparseVolumeTextureStreamingManager, Warning, TEXT("SVT IO request failed for '%p' (frame %i, mip %i, offset %i, size %i). Retrying..."),
+						PendingMipLevel.SparseVolumeTexture, PendingMipLevel.FrameIndex, PendingMipLevel.MipLevelIndex, MipLevelStreamingInfo.BulkOffset, MipLevelStreamingInfo.BulkSize);
+					
+					FBulkDataBatchRequest::FBatchBuilder Batch = FBulkDataBatchRequest::NewBatch(1);
+					Batch.Read(Resources->StreamableMipLevels, MipLevelStreamingInfo.BulkOffset, MipLevelStreamingInfo.BulkSize, AIOP_Low, PendingMipLevel.RequestBuffer, PendingMipLevel.Request);
+					(void)Batch.Issue();
+					break;
+				}
+			}
+			else
+			{
+				break;
+			}
+		}
+
+		++NumReadyMipLevels;
+	}
+
+	return NumReadyMipLevels;
+}
+
+void FStreamingManager::InstallReadyMipLevels()
+{
+	check(AsyncState.bUpdateActive);
+	check(AsyncState.NumReadyMipLevels <= PendingMipLevels.Num())
+	if (AsyncState.NumReadyMipLevels <= 0)
+	{
+		return;
+	}
+
+	TRACE_CPUPROFILER_EVENT_SCOPE(SVT::InstallReadyMipLevels);
+
+	const int32 StartPendingMipLevelIndex = (NextPendingMipLevelIndex + MaxPendingMipLevels - NumPendingMipLevels) % MaxPendingMipLevels;
+
+	// Do a first pass over all the mips to be uploaded to compute the upload buffer size requirements.
+	int32 NumPageTableUpdatesTotal = 0;
+	for (int32 i = 0; i < AsyncState.NumReadyMipLevels; ++i)
+	{
+		const int32 PendingMipLevelIndex = (StartPendingMipLevelIndex + i) % MaxPendingMipLevels;
+		FPendingMipLevel& PendingMipLevel = PendingMipLevels[PendingMipLevelIndex];
+
+		FStreamingInfo* SVTInfo = StreamingInfo.Find(PendingMipLevel.SparseVolumeTexture);
+		if (!SVTInfo || (SVTInfo->PerFrameInfo[PendingMipLevel.FrameIndex].LowestRequestedMipLevel > PendingMipLevel.MipLevelIndex))
+		{
+			continue; // Skip mip level install. SVT no longer exists or mip level was "streamed out" before it was even installed in the first place.
+		}
+
+		const FResources* Resources = SVTInfo->PerFrameInfo[PendingMipLevel.FrameIndex].Resources;
+		SVTInfo->TileDataTexture->NumTilesToUpload += Resources->MipLevelStreamingInfo[PendingMipLevel.MipLevelIndex].NumPhysicalTiles;
+		NumPageTableUpdatesTotal += Resources->MipLevelStreamingInfo[PendingMipLevel.MipLevelIndex].PageTableSize / (2 * sizeof(uint32));
+	}
+
+	UploadTasks.Reset();
+	UploadTasks.Reserve(AsyncState.NumReadyMipLevels * 2 /*slack for splitting large uploads*/);
+	UploadCleanupTasks.Reset();
+	
+
+#if WITH_EDITORONLY_DATA
+	TMap<const FResources*, const uint8*> ResourceToBulkPointer;
+#endif
+
+	// Do a second pass over all ready mip levels, claiming memory in the upload buffers and creating FUploadTasks
+	for (int32 i = 0; i < AsyncState.NumReadyMipLevels; ++i)
+	{
+		const int32 PendingMipLevelIndex = (StartPendingMipLevelIndex + i) % MaxPendingMipLevels;
+		FPendingMipLevel& PendingMipLevel = PendingMipLevels[PendingMipLevelIndex];
+
+		FStreamingInfo* SVTInfo = StreamingInfo.Find(PendingMipLevel.SparseVolumeTexture);
+		if (!SVTInfo || (SVTInfo->PerFrameInfo[PendingMipLevel.FrameIndex].LowestRequestedMipLevel > PendingMipLevel.MipLevelIndex))
+		{
+			PendingMipLevel.Reset();
+			continue; // Skip mip level install. SVT no longer exists or mip level was "streamed out" before it was even installed in the first place.
+		}
+
+		FFrameInfo& FrameInfo = SVTInfo->PerFrameInfo[PendingMipLevel.FrameIndex];
+		const FResources* Resources = FrameInfo.Resources;
+		const FMipLevelStreamingInfo& MipLevelStreamingInfo = Resources->MipLevelStreamingInfo[PendingMipLevel.MipLevelIndex];
+
+		const uint8* SrcPtr = nullptr;
+
+#if WITH_EDITORONLY_DATA
+		if (PendingMipLevel.State == FPendingMipLevel::EState::DDC_Ready)
+		{
+			check(Resources->ResourceFlags & EResourceFlag_StreamingDataInDDC);
+			SrcPtr = (const uint8*)PendingMipLevel.SharedBuffer.GetData();
+		}
+		else if (PendingMipLevel.State == FPendingMipLevel::EState::Memory)
+		{
+			const uint8** BulkDataPtrPtr = ResourceToBulkPointer.Find(Resources);
+			if (BulkDataPtrPtr)
+			{
+				SrcPtr = *BulkDataPtrPtr + MipLevelStreamingInfo.BulkOffset;
+			}
+			else
+			{
+				const FByteBulkData& BulkData = Resources->StreamableMipLevels;
+				check(BulkData.IsBulkDataLoaded() && BulkData.GetBulkDataSize() > 0);
+				const uint8* BulkDataPtr = (const uint8*)BulkData.LockReadOnly();
+				ResourceToBulkPointer.Add(Resources, BulkDataPtr);
+				SrcPtr = BulkDataPtr + MipLevelStreamingInfo.BulkOffset;
+			}
+		}
+		else
+#endif
+		{
+#if WITH_EDITORONLY_DATA
+			check(PendingMipLevel.State == FPendingMipLevel::EState::Disk);
+#endif
+			SrcPtr = PendingMipLevel.RequestBuffer.GetData();
+		}
+
+		check(SrcPtr);
+
+		const int32 NumPhysicalTiles = MipLevelStreamingInfo.NumPhysicalTiles;
+		TArray<uint32>& TileAllocations = FrameInfo.TileAllocations[PendingMipLevel.MipLevelIndex];
+		check(TileAllocations.Num() == NumPhysicalTiles);
+		check((MipLevelStreamingInfo.PageTableSize % (sizeof(uint32) * 2)) == 0);
+		const int32 NumPageTableUpdates = MipLevelStreamingInfo.PageTableSize / (sizeof(uint32) * 2);
+
+		uint8* DstPhysicalTileCoords = nullptr;
+		uint8* DstPhysicalTileDataA = nullptr;
+		uint8* DstPhysicalTileDataB = nullptr;
+		SVTInfo->TileDataTexture->TileUploader->Add_GetRef(NumPhysicalTiles, DstPhysicalTileCoords, DstPhysicalTileDataA, DstPhysicalTileDataB);
+
+		uint8* DstPageCoords = nullptr;
+		uint8* DstPageEntries = nullptr;
+		PageTableUpdater->Add_GetRef(FrameInfo.TextureRenderResources->PageTableTextureRHI, PendingMipLevel.MipLevelIndex, NumPageTableUpdates, DstPageCoords, DstPageEntries);
+
+		// Tile data
+		{
+			check(MipLevelStreamingInfo.TileDataASize == (SVTNumVoxelsPerPaddedTile * GPixelFormats[SVTInfo->TileDataTexture->FormatA].BlockBytes * NumPhysicalTiles));
+			check(MipLevelStreamingInfo.TileDataBSize == (SVTNumVoxelsPerPaddedTile * GPixelFormats[SVTInfo->TileDataTexture->FormatB].BlockBytes * NumPhysicalTiles));
+
+			FUploadTask& Task = UploadTasks.AddDefaulted_GetRef();
+			Task.TaskType = FUploadTask::ETaskType::TileData;
+			Task.TileDataTask.DstA = DstPhysicalTileDataA;
+			Task.TileDataTask.DstB = DstPhysicalTileDataB;
+			Task.TileDataTask.DstPhysicalTileCoords = DstPhysicalTileCoords;
+			Task.TileDataTask.SrcA = SrcPtr + MipLevelStreamingInfo.TileDataAOffset;
+			Task.TileDataTask.SrcB = SrcPtr + MipLevelStreamingInfo.TileDataBOffset;
+			Task.TileDataTask.SrcPhysicalTileCoords = reinterpret_cast<const uint8*>(TileAllocations.GetData());
+			Task.TileDataTask.SizeA = MipLevelStreamingInfo.TileDataASize;
+			Task.TileDataTask.SizeB = MipLevelStreamingInfo.TileDataBSize;
+			Task.TileDataTask.NumPhysicalTiles = NumPhysicalTiles;
+		}
+
+		// Page table
+		{
+			FUploadTask& Task = UploadTasks.AddDefaulted_GetRef();
+			Task.TaskType = FUploadTask::ETaskType::PageTable;
+			Task.PageTableTask.PendingMipLevel = &PendingMipLevel;
+			Task.PageTableTask.DstPageCoords = DstPageCoords;
+			Task.PageTableTask.DstPageEntries = DstPageEntries;
+			Task.PageTableTask.SrcPageCoords = SrcPtr + MipLevelStreamingInfo.PageTableOffset;
+			Task.PageTableTask.SrcPageEntries = SrcPtr + MipLevelStreamingInfo.PageTableOffset + NumPageTableUpdates * sizeof(uint32);
+			Task.PageTableTask.NumPageTableUpdates = NumPageTableUpdates;
+		}
+
+		// Cleanup
+		{
+			UploadCleanupTasks.Add(&PendingMipLevel);
+		}
+	
+#if SVT_STREAMING_LOG_VERBOSE
+		UE_LOG(LogSparseVolumeTextureStreamingManager, Display, TEXT("(%i)%i Install Frame %i OldResMip %i, NewResMip %i, ReqMip %i"),
+			PendingMipLevel.IssuedInFrame, PendingMipLevelIndex,
+			PendingMipLevel.FrameIndex, 
+			FrameInfo.LowestResidentMipLevel, PendingMipLevel.MipLevelIndex,
+			FrameInfo.LowestRequestedMipLevel);
+#endif
+
+		check(FrameInfo.LowestResidentMipLevel == (PendingMipLevel.MipLevelIndex + 1));
+		FrameInfo.LowestResidentMipLevel = PendingMipLevel.MipLevelIndex;
+
+		// Update the streaming info buffer data
+		SVTInfo->DirtyStreamingInfoData[PendingMipLevel.FrameIndex] = true;
+		InvalidatedStreamingInfos.Add(SVTInfo);
+
+		const int32 LRUNodeIndex = PendingMipLevel.FrameIndex * SVTInfo->NumMipLevelsGlobal + PendingMipLevel.MipLevelIndex;
+		SVTInfo->LRUNodes[LRUNodeIndex].PendingMipLevelIndex = INDEX_NONE;
+	}
+
+	// Do all the memcpy's in parallel
+	ParallelFor(UploadTasks.Num(), [&](int32 TaskIndex)
+		{
+			TRACE_CPUPROFILER_EVENT_SCOPE(SVT::FUploadTask);
+
+			FUploadTask& Task = UploadTasks[TaskIndex];
+
+			switch (Task.TaskType)
+			{
+			case FUploadTask::ETaskType::PageTable:
+			{
+				TRACE_CPUPROFILER_EVENT_SCOPE(SVT::PageTableUpload);
+				if (Task.PageTableTask.NumPageTableUpdates > 0)
+				{
+					FMemory::Memcpy(Task.PageTableTask.DstPageCoords, Task.PageTableTask.SrcPageCoords, Task.PageTableTask.NumPageTableUpdates * sizeof(uint32));
+
+					FStreamingInfo* SVTInfo = StreamingInfo.Find(Task.PageTableTask.PendingMipLevel->SparseVolumeTexture);
+					TArray<uint32>& TileAllocations = SVTInfo->PerFrameInfo[Task.PageTableTask.PendingMipLevel->FrameIndex].TileAllocations[Task.PageTableTask.PendingMipLevel->MipLevelIndex];
+					const uint32* SrcEntries = reinterpret_cast<const uint32*>(Task.PageTableTask.SrcPageEntries);
+					uint32* DstEntries = reinterpret_cast<uint32*>(Task.PageTableTask.DstPageEntries);
+					for (int32 i = 0; i < Task.PageTableTask.NumPageTableUpdates; ++i)
+					{
+						DstEntries[i] = TileAllocations[SrcEntries[i]];
+					}
+				}
+				break;
+			}
+			case FUploadTask::ETaskType::TileData:
+			{
+				TRACE_CPUPROFILER_EVENT_SCOPE(SVT::TileDataUpload);
+				FMemory::Memcpy(Task.TileDataTask.DstPhysicalTileCoords, Task.TileDataTask.SrcPhysicalTileCoords, Task.TileDataTask.NumPhysicalTiles * sizeof(uint32));
+				if (Task.TileDataTask.SizeA > 0)
+				{
+					FMemory::Memcpy(Task.TileDataTask.DstA, Task.TileDataTask.SrcA, Task.TileDataTask.SizeA);
+				}
+				if (Task.TileDataTask.SizeB > 0)
+				{
+					FMemory::Memcpy(Task.TileDataTask.DstB, Task.TileDataTask.SrcB, Task.TileDataTask.SizeB);
+				}
+				break;
+			}
+			default:
+				checkNoEntry();
+			}
+		});
+
+	ParallelFor(UploadCleanupTasks.Num(), [&](int32 TaskIndex)
+		{
+			TRACE_CPUPROFILER_EVENT_SCOPE(SVT::FUploadCleanupTask);
+
+			FPendingMipLevel* PendingMipLevel = UploadCleanupTasks[TaskIndex];
+#if WITH_EDITORONLY_DATA
+			PendingMipLevel->SharedBuffer.Reset();
+#endif
+			if (!PendingMipLevel->Request.IsNone())
+			{
+				check(PendingMipLevel->Request.IsCompleted());
+				PendingMipLevel->Request.Reset();
+			}
+		});
+
+#if DO_CHECK // Clear processed pending mip levels for better debugging
+	for (int32 i = 0; i < AsyncState.NumReadyMipLevels; ++i)
+	{
+		const int32 PendingMipLevelIndex = (StartPendingMipLevelIndex + i) % MaxPendingMipLevels;
+		PendingMipLevels[PendingMipLevelIndex].Reset();
+	}
+#endif
+
+#if WITH_EDITORONLY_DATA
+	// Unlock BulkData
+	for (auto& Pair : ResourceToBulkPointer)
+	{
+		Pair.Key->StreamableMipLevels.Unlock();
+	}
+#endif
+}
+
+#if WITH_EDITORONLY_DATA
+
+UE::DerivedData::FCacheGetChunkRequest FStreamingManager::BuildDDCRequest(const FResources& Resources, const FMipLevelStreamingInfo& MipLevelStreamingInfo, const uint32 PendingMipLevelIndex)
+{
+	using namespace UE::DerivedData;
+
+	FCacheKey Key;
+	Key.Bucket = FCacheBucket(TEXT("SparseVolumeTexture"));
+	Key.Hash = Resources.DDCKeyHash;
+	check(!Resources.DDCRawHash.IsZero());
+
+	FCacheGetChunkRequest Request;
+	Request.Id = FValueId::FromName("SparseVolumeTextureStreamingData");
+	Request.Key = Key;
+	Request.RawOffset = MipLevelStreamingInfo.BulkOffset;
+	Request.RawSize = MipLevelStreamingInfo.BulkSize;
+	Request.RawHash = Resources.DDCRawHash;
+	Request.UserData = (((uint64)PendingMipLevelIndex) << uint64(32)) | (uint64)PendingMipLevels[PendingMipLevelIndex].RequestVersion;
+	return Request;
+}
+
+void FStreamingManager::RequestDDCData(TConstArrayView<UE::DerivedData::FCacheGetChunkRequest> DDCRequests)
+{
+	using namespace UE::DerivedData;
+
+	FRequestBarrier Barrier(*RequestOwner);	// This is a critical section on the owner. It does not constrain ordering
+	GetCache().GetChunks(DDCRequests, *RequestOwner,
+		[this](FCacheGetChunkResponse&& Response)
+		{
+			const uint32 PendingMipLevelIndex = (uint32)(Response.UserData >> uint64(32));
+			const uint32 RequestVersion = (uint32)Response.UserData;
+			
+			// In case the request returned after the mip level was already streamed out again we need to abort so that we do not overwrite data in the FPendingMipLevel slot.
+			if (RequestVersion < PendingMipLevels[PendingMipLevelIndex].RequestVersion)
+			{
+				return;
+			}
+			
+			FPendingMipLevel& PendingMipLevel = PendingMipLevels[PendingMipLevelIndex];
+			check(PendingMipLevel.SparseVolumeTexture); // A valid PendingMipLevel should have a non-nullptr here
+
+			if (Response.Status == EStatus::Ok)
+			{
+				PendingMipLevel.SharedBuffer = MoveTemp(Response.RawData);
+				PendingMipLevel.State = FPendingMipLevel::EState::DDC_Ready;
+			}
+			else
+			{
+				PendingMipLevel.State = FPendingMipLevel::EState::DDC_Failed;
+			}
+		});
+}
+
+#endif // WITH_EDITORONLY_DATA
+
+FStreamingManager::FTileDataTexture::FTileDataTexture(const FIntVector3& InResolutionInTiles, EPixelFormat InFormatA, EPixelFormat InFormatB)
+	: ResolutionInTiles(InResolutionInTiles), 
+	PhysicalTilesCapacity(InResolutionInTiles.X * InResolutionInTiles.Y * InResolutionInTiles.Z), 
+	FormatA(InFormatA),
+	FormatB(InFormatB),
+	TileUploader(MakeUnique<FTileUploader>()),
+	NumTilesToUpload(0)
+{
+	const int64 MaxFormatSize = FMath::Max(GPixelFormats[FormatA].BlockBytes, GPixelFormats[FormatB].BlockBytes);
+	const FIntVector3 LargestPossibleResolution = ComputeLargestPossibleTileDataVolumeResolution(MaxFormatSize);
+	const int32 LargestPossiblePhysicalTilesCapacity = LargestPossibleResolution.X * LargestPossibleResolution.Y * LargestPossibleResolution.Z;
+	
+	// Ensure that the tile data texture(s) do not exceed the memory size and resolution limits.
+	if (PhysicalTilesCapacity > LargestPossiblePhysicalTilesCapacity
+		|| (ResolutionInTiles.X * SPARSE_VOLUME_TILE_RES_PADDED) > SVTMaxVolumeTextureDim
+		|| (ResolutionInTiles.Y * SPARSE_VOLUME_TILE_RES_PADDED) > SVTMaxVolumeTextureDim
+		|| (ResolutionInTiles.Z * SPARSE_VOLUME_TILE_RES_PADDED) > SVTMaxVolumeTextureDim)
+	{
+		ResolutionInTiles = LargestPossibleResolution;
+		PhysicalTilesCapacity = LargestPossiblePhysicalTilesCapacity;
+
+		UE_LOG(LogSparseVolumeTextureStreamingManager, Warning, TEXT("Requested SparseVolumeTexture tile data texture resolution (in tiles) (%i, %i, %i) exceeds the resource size limit. Using the maximum value of (%i, %i. %i) instead."),
+			InResolutionInTiles.X, InResolutionInTiles.Y, InResolutionInTiles.Z,
+			LargestPossibleResolution.X, LargestPossibleResolution.Y, LargestPossibleResolution.Z);
+	}
+
+	const FIntVector3 Resolution = ResolutionInTiles * SPARSE_VOLUME_TILE_RES_PADDED;
+	check(Resolution.X <= SVTMaxVolumeTextureDim && Resolution.Y <= SVTMaxVolumeTextureDim && Resolution.Z <= SVTMaxVolumeTextureDim);
+	check(((int64)Resolution.X * (int64)Resolution.Y * (int64)Resolution.Z * (int64)GPixelFormats[FormatA].BlockBytes) <= int64(INT32_MAX));
+	check(((int64)Resolution.X * (int64)Resolution.Y * (int64)Resolution.Z * (int64)GPixelFormats[FormatB].BlockBytes) <= int64(INT32_MAX));
+	
+	TileCoords.SetNum(PhysicalTilesCapacity);
+
+	int32 TileCoordsIndex = 0;
+	for (int32 Z = 0; Z < ResolutionInTiles.Z; ++Z)
+	{
+		for (int32 Y = 0; Y < ResolutionInTiles.Y; ++Y)
+		{
+			for (int32 X = 0; X < ResolutionInTiles.X; ++X)
+			{
+				uint32 PackedCoord = 0;
+				PackedCoord |= (X & 0xFFu);
+				PackedCoord |= (Y & 0xFFu) << 8u;
+				PackedCoord |= (Z & 0xFFu) << 16u;
+				TileCoords[TileCoordsIndex++] = PackedCoord;
+			}
+		}
+	}
+	check(TileCoordsIndex == PhysicalTilesCapacity);
+}
+
+void FStreamingManager::FTileDataTexture::InitRHI()
+{
+	const FIntVector3 Resolution = ResolutionInTiles * SPARSE_VOLUME_TILE_RES_PADDED;
+	if (FormatA != PF_Unknown)
+	{
+		const FRHITextureCreateDesc Desc =
+			FRHITextureCreateDesc::Create3D(TEXT("SparseVolumeTexture.PhysicalTileDataA.RHITexture"), Resolution.X, Resolution.Y, Resolution.Z, FormatA)
+			.SetFlags(ETextureCreateFlags::ShaderResource | ETextureCreateFlags::UAV);
+		TileDataTextureARHIRef = RHICreateTexture(Desc);
+	}
+	if (FormatB != PF_Unknown)
+	{
+		const FRHITextureCreateDesc Desc =
+			FRHITextureCreateDesc::Create3D(TEXT("SparseVolumeTexture.PhysicalTileDataB.RHITexture"), Resolution.X, Resolution.Y, Resolution.Z, FormatB)
+			.SetFlags(ETextureCreateFlags::ShaderResource | ETextureCreateFlags::UAV);
+		TileDataTextureBRHIRef = RHICreateTexture(Desc);
+	}
+}
+
+void FStreamingManager::FTileDataTexture::ReleaseRHI()
+{
+}
+
+TGlobalResource<FStreamingManager> GStreamingManager;
+
+}
+}
diff --git a/Engine/Source/Runtime/Engine/Private/SparseVolumeTexture/SparseVolumeTextureUtility.cpp b/Engine/Source/Runtime/Engine/Private/SparseVolumeTexture/SparseVolumeTextureUtility.cpp
--- a/Engine/Source/Runtime/Engine/Private/SparseVolumeTexture/SparseVolumeTextureUtility.cpp
+++ b/Engine/Source/Runtime/Engine/Private/SparseVolumeTexture/SparseVolumeTextureUtility.cpp
@@ -20,19 +20,16 @@ namespace Private
 
 uint32 UE::SVT::PackPageTableEntry(const FIntVector3& Coord)
 {
-	// A page encodes the physical tile coord as unsigned int of 11 11 10 bits
-	// This means a page coord cannot be larger than 2047 for x and y and 1023 for z
-	// which mean we cannot have more than 2048*2048*1024 = 4 Giga tiles of 16^3 tiles.
-	uint32 Result = (Coord.X & 0x7FF) | ((Coord.Y & 0x7FF) << 11) | ((Coord.Z & 0x3FF) << 22);
+	uint32 Result = (Coord.X & 0xFFu) | ((Coord.Y & 0xFFu) << 8u) | ((Coord.Z & 0xFFu) << 16u);
 	return Result;
 }
 
 FIntVector3 UE::SVT::UnpackPageTableEntry(uint32 Packed)
 {
 	FIntVector3 Result;
-	Result.X = Packed & 0x7FF;
-	Result.Y = (Packed >> 11) & 0x7FF;
-	Result.Z = (Packed >> 22) & 0x3FF;
+	Result.X = Packed & 0xFFu;
+	Result.Y = (Packed >> 8u) & 0xFFu;
+	Result.Z = (Packed >> 16u) & 0xFFu;
 	return Result;
 }
 
diff --git a/Engine/Source/Runtime/Engine/Public/ContentStreaming.h b/Engine/Source/Runtime/Engine/Public/ContentStreaming.h
--- a/Engine/Source/Runtime/Engine/Public/ContentStreaming.h
+++ b/Engine/Source/Runtime/Engine/Public/ContentStreaming.h
@@ -32,7 +32,6 @@ class UAnimStreamable;
 enum class EStreamableRenderAssetType : uint8;
 struct FCompressedAnimSequence;
 class UStreamableSparseVolumeTexture;
-class FSparseVolumeTextureSceneProxy;
 
 using FSoundWaveProxyPtr = TSharedPtr<FSoundWaveProxy, ESPMode::ThreadSafe>;
 
@@ -632,32 +631,6 @@ struct IAnimationStreamingManager : public IStreamingManager
 	virtual const FCompressedAnimSequence* GetLoadedChunk(const UAnimStreamable* Anim, uint32 ChunkIndex, bool bRequestNextChunk) const = 0;
 };
 
-/**
- * Interface to add functions specifically related to SparseVolumeTexture streaming
- */
-struct ISparseVolumeTextureStreamingManager : public IStreamingManager
-{
-	/** Adds a new StreamableSparseVolumeTexture to the streaming manager. */
-	virtual void AddSparseVolumeTexture(UStreamableSparseVolumeTexture* SparseVolumeTexture) = 0;
-
-	/** Removes a StreamableSparseVolumeTexture from the streaming manager. */
-	virtual bool RemoveSparseVolumeTexture(UStreamableSparseVolumeTexture* SparseVolumeTexture) = 0;
-
-	/** Returns the memory a StreamableSparseVolumeTexture is currently using */
-	virtual void GetMemorySizeForSparseVolumeTexture(const UStreamableSparseVolumeTexture* SparseVolumeTexture, SIZE_T* SizeCPU, SIZE_T* SizeGPU) const = 0;
-
-	/**
-	 * Gets a pointer to a proxy of a frame of a SparseVolumeTexture
-	 *
-	 * @param SparseVolumeTexture	SparseVolumeTexture we want a frame from
-	 * @param FrameIndex			Index of the frame we want
-	 * @param MipLevel				MipLevel of the frame we want
-	 * @return Either the desired frame proxy or NULL if it's not loaded
-	 */
-	virtual const FSparseVolumeTextureSceneProxy* GetSparseVolumeTextureSceneProxy(const UStreamableSparseVolumeTexture* SparseVolumeTexture, int32 FrameIndex, int32 MipLevel, bool bTrackAsRequested) = 0;
-};
-
-
 /**
  * Streaming manager collection, routing function calls to streaming managers that have been added
  * via AddStreamingManager.
@@ -767,11 +740,6 @@ struct FStreamingManagerCollection : public IStreamingManager
 	*/
 	ENGINE_API Nanite::FCoarseMeshStreamingManager* GetNaniteCoarseMeshStreamingManager() const;
 
-	/**
-	 * Gets a reference to the SparseVolumeTexture Streaming Manager interface
-	 */
-	ENGINE_API ISparseVolumeTextureStreamingManager& GetSparseVolumeTextureStreamingManager() const;
-
 	/**
 	 * Adds a streaming manager to the array of managers to route function calls to.
 	 *
@@ -876,9 +844,6 @@ protected:
 	/** The nanite coarse mesh streaming manager, should always exist */
 	Nanite::FCoarseMeshStreamingManager* NaniteCoarseMeshStreamingManager;
 
-	/** The SparseVolumeTexture streaming manager, should always exist */
-	ISparseVolumeTextureStreamingManager* SparseVolumeTextureStreamingManager;
-
 #if WITH_EDITOR
 	// Locks out any audio streaming manager call when we are re-initializing the audio streaming manager.
 	mutable FCriticalSection AudioStreamingManagerCriticalSection;
diff --git a/Engine/Source/Runtime/Engine/Public/SparseVolumeTexture/SparseVolumeTextureData.h b/Engine/Source/Runtime/Engine/Public/SparseVolumeTexture/SparseVolumeTextureData.h
--- a/Engine/Source/Runtime/Engine/Public/SparseVolumeTexture/SparseVolumeTextureData.h
+++ b/Engine/Source/Runtime/Engine/Public/SparseVolumeTexture/SparseVolumeTextureData.h
@@ -5,7 +5,12 @@
 #include "CoreMinimal.h"
 #include "SparseVolumeTexture/SparseVolumeTexture.h"
 
-struct ENGINE_API FSparseVolumeTextureDataCreateInfo
+namespace UE
+{
+namespace SVT
+{
+
+struct ENGINE_API FTextureDataCreateInfo
 {
 	FIntVector3 VirtualVolumeAABBMin = FIntVector3(INT32_MAX, INT32_MAX, INT32_MAX);
 	FIntVector3 VirtualVolumeAABBMax = FIntVector3(INT32_MIN, INT32_MIN, INT32_MIN);
@@ -13,15 +18,15 @@ struct ENGINE_API FSparseVolumeTextureDataCreateInfo
 	TStaticArray<FVector4f, 2> FallbackValues = TStaticArray<FVector4f, 2>(InPlace, FVector4f());
 };
 
-class ENGINE_API ISparseVolumeTextureDataProvider
+class ENGINE_API ITextureDataProvider
 {
 public:
-	virtual FSparseVolumeTextureDataCreateInfo GetCreateInfo() const = 0;
+	virtual FTextureDataCreateInfo GetCreateInfo() const = 0;
 	virtual void IteratePhysicalSource(TFunctionRef<void(const FIntVector3& Coord, int32 AttributesIdx, int32 ComponentIdx, float VoxelValue)> OnVisit) const = 0;
-	virtual ~ISparseVolumeTextureDataProvider() = default;
+	virtual ~ITextureDataProvider() = default;
 };
 
-struct ENGINE_API FSparseVolumeTextureDataAddressingInfo
+struct ENGINE_API FTextureDataAddressingInfo
 {
 	FIntVector3 VolumeResolution;
 	TextureAddress AddressX;
@@ -32,7 +37,7 @@ struct ENGINE_API FSparseVolumeTextureDataAddressingInfo
 // Holds the data for a SparseVolumeTexture that is stored on disk. It only has a single mip after importing a source asset. The mip chain is built during cooking.
 // Tiles are addressed by a flat index; unlike the runtime representation, this one stores all tiles in a 1D array (per mip level) and doesn't have the concept of a 3D physical tile texture.
 // The page table itself is 3D though.
-struct ENGINE_API FSparseVolumeTextureData
+struct ENGINE_API FTextureData
 {
 	struct FMipMap
 	{
@@ -45,24 +50,28 @@ struct ENGINE_API FSparseVolumeTextureData
 	static const uint32 kVersion = 0; // The current data format version for the raw source data.
 	uint32 Version = kVersion; // This version can be used to convert existing source data to new version later.
 
-	FSparseVolumeTextureHeader Header = {};
+	UE::SVT::FHeader Header = {};
+	TStaticArray<FVector4f, 2> FallbackValuesQuantized = TStaticArray<FVector4f, 2>(InPlace, FVector4f());
 	TArray<FMipMap> MipMaps;
 
 	void Serialize(FArchive& Ar);
 	
-	bool Create(const ISparseVolumeTextureDataProvider& DataProvider);
-	bool CreateFromDense(const FSparseVolumeTextureDataCreateInfo& CreateInfo, const TArrayView<uint8>& VoxelDataA, const TArrayView<uint8>& VoxelDataB);
+	bool Create(const ITextureDataProvider& DataProvider);
+	bool CreateFromDense(const FTextureDataCreateInfo& CreateInfo, const TArrayView<uint8>& VoxelDataA, const TArrayView<uint8>& VoxelDataB);
 	void CreateDefault();
-	FVector4f Load(const FIntVector3& VolumeCoord, int32 MipLevel, int32 AttributesIdx, const FSparseVolumeTextureDataAddressingInfo& AddressingInfo) const;
-	bool BuildDerivedData(const FSparseVolumeTextureDataAddressingInfo& AddressingInfo, int32 NumMipLevels, bool bMoveMip0FromThis, FSparseVolumeTextureData& OutDerivedData);
+	FVector4f Load(const FIntVector3& VolumeCoord, int32 MipLevel, int32 AttributesIdx, const FTextureDataAddressingInfo& AddressingInfo) const;
+	bool BuildDerivedData(const FTextureDataAddressingInfo& AddressingInfo, int32 NumMipLevels, bool bMoveMip0FromThis, FTextureData& OutDerivedData);
 
 private:
 
-	bool GenerateMipMaps(const FSparseVolumeTextureDataAddressingInfo& AddressingInfo, int32 NumMipLevels = -1);
-	bool GenerateBorderVoxels(const FSparseVolumeTextureDataAddressingInfo& AddressingInfo, int32 MipLevel, const TArray<FIntVector3>& PageCoords);
+	bool GenerateMipMaps(const FTextureDataAddressingInfo& AddressingInfo, int32 NumMipLevels = -1);
+	bool GenerateBorderVoxels(const FTextureDataAddressingInfo& AddressingInfo, int32 MipLevel, const TArray<FIntVector3>& PageCoords);
 	bool DeduplicateTiles();
 
 	uint32 ReadPageTable(const FIntVector3& PageTableCoord, int32 MipLevel) const;
 	FVector4f ReadTileDataVoxel(int32 TileIndex, const FIntVector3& TileDataCoord, int32 MipLevel, int32 AttributesIdx) const;
 	void WriteTileDataVoxel(int32 TileIndex, const FIntVector3& TileDataCoord, int32 MipLevel, int32 AttributesIdx, const FVector4f& Value, int32 DstComponent = -1);
 };
+
+}
+}
diff --git a/Engine/Source/Runtime/Engine/Public/SparseVolumeTexture/SparseVolumeTextureSceneProxy.h b/Engine/Source/Runtime/Engine/Public/SparseVolumeTexture/SparseVolumeTextureSceneProxy.h
deleted file mode 100644
diff --git a/Engine/Source/Runtime/Engine/Public/SparseVolumeTexture/SparseVolumeTextureStreaming.h b/Engine/Source/Runtime/Engine/Public/SparseVolumeTexture/SparseVolumeTextureStreaming.h
deleted file mode 100644
diff --git a/Engine/Source/Runtime/Engine/Public/SparseVolumeTexture/SparseVolumeTextureStreamingManager.h b/Engine/Source/Runtime/Engine/Public/SparseVolumeTexture/SparseVolumeTextureStreamingManager.h
new file mode 100644
--- /dev/null
+++ b/Engine/Source/Runtime/Engine/Public/SparseVolumeTexture/SparseVolumeTextureStreamingManager.h
@@ -0,0 +1,359 @@
+// Copyright Epic Games, Inc. All Rights Reserved.
+
+#pragma once
+
+#include "CoreMinimal.h"
+#include "Engine/EngineTypes.h"
+#include "IO/IoDispatcher.h"
+#include "Memory/SharedBuffer.h"
+#include "Serialization/BulkData.h"
+#include "RenderResource.h"
+#include "Containers/IntrusiveDoubleLinkedList.h"
+#include "Containers/Map.h"
+#include "RenderGraphBuilder.h"
+
+class UStreamableSparseVolumeTexture;
+class USparseVolumeTextureFrame;
+
+namespace UE
+{
+	namespace DerivedData
+	{
+		class FRequestOwner; // Can't include DDC headers from here, so we have to forward declare
+		struct FCacheGetChunkRequest;
+	}
+}
+
+namespace UE
+{
+namespace SVT
+{
+
+struct FResources;
+struct FMipLevelStreamingInfo;
+class FTextureRenderResources;
+
+// Uniquely identifies the mip level of a frame in a static or animated SparseVolumeTexture
+struct FMipLevelKey
+{
+	UStreamableSparseVolumeTexture* SVT = nullptr; // This struct is used on the rendering thread. Do not dereference!
+	uint16 FrameIndex = INDEX_NONE;
+	uint16 MipLevelIndex = INDEX_NONE;
+
+	friend FORCEINLINE uint32 GetTypeHash(const FMipLevelKey& Key)
+	{
+		return HashCombine(GetTypeHash(Key.SVT), GetTypeHash(((uint32)Key.FrameIndex << 16u) | (uint32)Key.MipLevelIndex));
+	}
+
+	FORCEINLINE bool operator==(const FMipLevelKey& Other) const
+	{
+		return (SVT == Other.SVT) && (FrameIndex == Other.FrameIndex) && (MipLevelIndex == Other.MipLevelIndex);
+	}
+
+	FORCEINLINE bool operator!=(const FMipLevelKey& Other) const
+	{
+		return !(*this == Other);
+	}
+
+	FORCEINLINE bool operator<(const FMipLevelKey& Other) const
+	{
+		return SVT != Other.SVT ? SVT < Other.SVT : FrameIndex != Other.FrameIndex ? FrameIndex < Other.FrameIndex : MipLevelIndex > Other.MipLevelIndex;
+	}
+};
+
+struct FStreamingRequest
+{
+	FMipLevelKey Key;
+	uint32 Priority; // A higher value means a higher priority
+
+	FORCEINLINE bool operator<(const FStreamingRequest& Other) const
+	{
+		return Key != Other.Key ? Key < Other.Key : Priority > Other.Priority;
+	}
+};
+
+class FStreamingManager : public FRenderResource
+{
+public:
+	FStreamingManager();
+
+	virtual void InitRHI() override;
+	virtual void ReleaseRHI() override;
+
+	ENGINE_API void Add_GameThread(UStreamableSparseVolumeTexture* SparseVolumeTexture);
+	ENGINE_API void Remove_GameThread(UStreamableSparseVolumeTexture* SparseVolumeTexture);
+	ENGINE_API void BeginAsyncUpdate(FRDGBuilder& GraphBuilder);
+	ENGINE_API void EndAsyncUpdate(FRDGBuilder& GraphBuilder);
+
+	// Request a frame to be streamed in. FrameIndex is of float type so that the fractional part can be used to better track the playback speed/direction.
+	// This function automatically also requests all higher mip levels and adds prefetch requests for upcoming frames.
+	ENGINE_API void Request_GameThread(UStreamableSparseVolumeTexture* SparseVolumeTexture, float FrameIndex, int32 MipLevel);
+	ENGINE_API void Request(UStreamableSparseVolumeTexture* SparseVolumeTexture, float FrameIndex, int32 MipLevel);
+
+private:
+	friend class FStreamingUpdateTask;
+
+	// SVT_TODO: actually implement streaming in/out of page table mips greater than this value.
+	static constexpr int32 NumAlwaysResidentPageTableMipLevels = 6;
+
+	// Represents the physical tile data texture that serves as backing memory for the streamed in tiles. While this is treated as a single logical texture,
+	// it currently supports up to two actual RHI textures.
+	struct FTileDataTexture : public FRenderResource
+	{
+	public:
+		static constexpr uint32 PhysicalCoordMask = (1u << 24u) - 1u; // Lower 24 bits are used for storing XYZ in 8 bit each. Upper 8 bit can be used by the caller. 
+
+		FIntVector3 ResolutionInTiles;
+		int32 PhysicalTilesCapacity;
+		EPixelFormat FormatA;
+		EPixelFormat FormatB;
+		FTextureRHIRef TileDataTextureARHIRef;
+		FTextureRHIRef TileDataTextureBRHIRef;
+		TUniquePtr<class FTileUploader> TileUploader;
+		int32 NumTilesToUpload;
+
+		// Constructor. May change the requested ResolutionInTiles (and resulting PhysicalTilesCapacity) if it exceeds hardware limits.
+		FTileDataTexture(const FIntVector3& ResolutionInTiles, EPixelFormat FormatA, EPixelFormat FormatB);
+
+		// Allocate a tile slot in the texture. The resulting value is a packed coordinate (8 bit per component) of the allocated slot or INDEX_NONE if the allocation failed.
+		// The upper 8 bit are free to be used by the caller.
+		uint32 Allocate() 
+		{ 
+			check(PhysicalTilesCapacity == TileCoords.Num());
+			return TileCoords.IsValidIndex(NextFreeTileCoordIndex) ? TileCoords[NextFreeTileCoordIndex++] : INDEX_NONE; 
+		}
+
+		// Frees a previously allocated tile slot. The upper 8 bit (user data) are automatically cleared by this function.
+		void Free(uint32 PackedPhysicalCoord) 
+		{ 
+			check(PackedPhysicalCoord != INDEX_NONE)
+			PackedPhysicalCoord &= PhysicalCoordMask;
+			check(PhysicalTilesCapacity == TileCoords.Num());
+			check(NextFreeTileCoordIndex > 0);
+#if DO_GUARD_SLOW
+			for (int32 i = NextFreeTileCoordIndex; i < PhysicalTilesCapacity; ++i)
+			{
+				check(TileCoords[i] != PackedPhysicalCoord);
+			}
+#endif
+			TileCoords[--NextFreeTileCoordIndex] = PackedPhysicalCoord;
+		}
+
+		// Number of tiles available for allocation.
+		int32 GetNumAvailableTiles() { return PhysicalTilesCapacity - NextFreeTileCoordIndex; }
+		virtual void InitRHI() override;
+		virtual void ReleaseRHI() override;
+
+	private:
+
+		TArray<uint32> TileCoords;
+		int32 NextFreeTileCoordIndex = 0;
+		uint32 NumRefs = 0;
+	};
+
+	// Represents a context or "window" into the frame sequence that moves along the playback direction. It is used to cache the prefetch direction.
+	struct FStreamingWindow
+	{
+		static constexpr int32 WindowSize = 5;
+		float CenterFrame = -1.0f; // Frame index this window is centered around
+		float LastCenterFrame = -1.0f;
+		int32 NumRequestsThisUpdate = 0;
+		uint32 LastRequested = 0;
+		bool bPlayForward = false;
+		bool bPlayBackward = false;
+	};
+
+	// One per frame of each SVT
+	struct FFrameInfo
+	{
+		const FResources* Resources;						// Initialized on the game thread
+		FTextureRenderResources* TextureRenderResources;	// Initialized on the game thread
+
+		int32 NumMipLevels; // Number of actual mip levels of this frame. Depending on virtual volume extents, frames can have different numbers of levels.
+		int32 LowestRequestedMipLevel; // Lowest mip level that should be resident. Can be lower than LowestResidentMipLevel when streaming in new mips. Stream-out is instant, so it should never be higher.
+		int32 LowestResidentMipLevel; // Actually resident on the GPU
+		TArray<TArray<uint32>> TileAllocations; // TileAllocations[MipLevel][PhysicalTileIndex]
+	};
+
+	// Used to keep track of the least-recently-used order of mip levels. This is done per mip level and not per frame, so that lower resolution mip levels are kept in memory for longer.
+	struct FLRUNode : public TIntrusiveDoubleLinkedListNode<FLRUNode>
+	{
+		FLRUNode* NextHigherMipLevel = nullptr;
+		uint32 LastRequested = INDEX_NONE;
+		uint32 RefCount = 0; // Keep track of how many lower mip levels have a dependency on this one
+		uint32 PendingMipLevelIndex = INDEX_NONE;
+		int16 FrameIndex = INDEX_NONE;
+		int16 MipLevelIndex = INDEX_NONE;
+	};
+
+	// Streaming manager internal data for each SVT
+	struct FStreamingInfo
+	{
+		EPixelFormat FormatA;
+		EPixelFormat FormatB;
+		FVector4f FallbackValueA;
+		FVector4f FallbackValueB;
+		int32 NumMipLevelsGlobal; // Maximum number of mip levels in the sequence
+		uint32 LastRequested; // Last update index when some frame in this SVT was requested
+		TArray<FFrameInfo> PerFrameInfo;
+		TArray<FLRUNode> LRUNodes;
+		TArray<TIntrusiveDoubleLinkedList<FLRUNode>> PerMipLRULists; // PerMipLRULists[MipLevel]
+		TArray<FStreamingWindow> StreamingWindows;
+
+		FTileDataTexture* TileDataTexture;
+		TRefCountPtr<FRDGPooledBuffer> StreamingInfoBuffer; // One uint32 per frame storing the LowestResidentMipLevel
+		FShaderResourceViewRHIRef StreamingInfoBufferSRVRHIRef;
+		TBitArray<> DirtyStreamingInfoData; // One bit per frame, potentially marking the streaming info data as dirty/in need of an update
+	};
+
+	// Represents an IO request for a mip level
+	struct FPendingMipLevel
+	{
+		UStreamableSparseVolumeTexture* SparseVolumeTexture = nullptr; // Do not dereference!
+		int32 FrameIndex = INDEX_NONE;
+		int32 MipLevelIndex = INDEX_NONE;
+#if WITH_EDITORONLY_DATA
+		FSharedBuffer SharedBuffer;
+		enum class EState
+		{
+			None,
+			DDC_Pending,
+			DDC_Ready,
+			DDC_Failed,
+			Memory,
+			Disk,
+		} State = EState::None;
+		uint32 RetryCount = 0;
+		uint32 RequestVersion = 0;
+#endif
+		FIoBuffer RequestBuffer;
+		FBulkDataBatchReadRequest Request;
+		uint32 IssuedInFrame = 0;
+
+		void Reset()
+		{
+			SparseVolumeTexture = nullptr;
+			FrameIndex = INDEX_NONE;
+			MipLevelIndex = INDEX_NONE;
+#if WITH_EDITORONLY_DATA
+			SharedBuffer.Reset();
+			State = EState::None;
+			RetryCount = 0;
+			++RequestVersion;
+#endif
+			RequestBuffer = {};
+			Request.Reset();
+			IssuedInFrame = 0;
+		}
+	};
+
+	// Encapsulates all the src/dst pointers for memcpy-ing the streamed data into GPU memory.
+	struct FUploadTask
+	{
+		enum class ETaskType
+		{
+			TileData, PageTable
+		};
+
+		struct FTileDataTask
+		{
+			uint8* DstA;
+			uint8* DstB;
+			uint8* DstPhysicalTileCoords;
+			const uint8* SrcA;
+			const uint8* SrcB;
+			const uint8* SrcPhysicalTileCoords;
+			int32 SizeA;
+			int32 SizeB;
+			int32 NumPhysicalTiles;
+		};
+
+		struct FPageTableTask
+		{
+			FPendingMipLevel* PendingMipLevel;
+			uint8* DstPageCoords;
+			uint8* DstPageEntries;
+			const uint8* SrcPageCoords;
+			const uint8* SrcPageEntries;
+			int32 NumPageTableUpdates;
+		};
+
+		ETaskType TaskType = ETaskType::TileData;
+		union
+		{
+			FTileDataTask TileDataTask;
+			FPageTableTask PageTableTask;
+		};
+	};
+
+	struct FPageTableClear
+	{
+		FTextureRHIRef PageTableTexture;
+		int32 MipLevel;
+	};
+
+	struct FAsyncState
+	{
+		int32 NumReadyMipLevels = 0;
+		bool bUpdateActive = false;
+	};
+
+	// Stores info extracted from the UStreamableSparseVolumeTexture and its USparseVolumeTextureFrames so that we don't need to access these uobjects on the rendering thread
+	struct FNewSparseVolumeTextureInfo
+	{
+		UStreamableSparseVolumeTexture* SVT = nullptr; // Do not dereference!
+		EPixelFormat FormatA = PF_Unknown;
+		EPixelFormat FormatB = PF_Unknown;
+		FVector4f FallbackValueA = FVector4f(0.0f, 0.0f, 0.0f, 0.0f);
+		FVector4f FallbackValueB = FVector4f(0.0f, 0.0f, 0.0f, 0.0f);
+		int32 NumMipLevelsGlobal = 0;
+		TArray<FFrameInfo> FrameInfo; // Only Resources and TextureRenderResources are initialized
+	};
+
+	TMap<UStreamableSparseVolumeTexture*, FStreamingInfo> StreamingInfo; // Do not dereference the key! We just read the pointer itself.
+	TMap<FMipLevelKey, uint32> RequestsHashTable;
+	TArray<FPendingMipLevel> PendingMipLevels;
+#if WITH_EDITORONLY_DATA
+	UE::DerivedData::FRequestOwner* RequestOwner = nullptr;
+#endif
+
+	class FPageTableUpdater* PageTableUpdater = nullptr;
+	class FStreamingInfoBufferUpdater* StreamingInfoBufferUpdater = nullptr;
+	FGraphEventArray AsyncTaskEvents;
+	FAsyncState AsyncState;
+	int32 MaxPendingMipLevels = 0;
+	int32 NumPendingMipLevels = 0;
+	int32 NextPendingMipLevelIndex = 0;
+	uint32 NextUpdateIndex = 1;
+
+	// Transient lifetime
+	TArray<FStreamingRequest> ParentRequestsToAdd;
+	TSet<FTileDataTexture*> TileDataTexturesToUpdate;
+	TSet<FStreamingInfo*> InvalidatedStreamingInfos;
+	TArray<FStreamingRequest> PrioritizedRequestsHeap;
+	TArray<FMipLevelKey> SelectedMipLevels;
+	TArray<FPageTableClear> PageTableClears;
+	TArray<FUploadTask> UploadTasks; // accessed on the async thread
+	TArray<FPendingMipLevel*> UploadCleanupTasks; // accessed on the async thread
+
+	void AddInternal(FRDGBuilder& GraphBuilder, FNewSparseVolumeTextureInfo&& NewSVTInfo);
+	void RemoveInternal(UStreamableSparseVolumeTexture* SparseVolumeTexture);
+	bool AddRequest(const FStreamingRequest& Request); // Returns true when the request is new or it has been updated with a higher priority
+	void AsyncUpdate();
+	void AddParentRequests(); // Add requests for all parent mip levels
+	void SelectHighestPriorityRequestsAndUpdateLRU(int32 MaxSelectedMipLevels);
+	void IssueRequests(int32 MaxSelectedMipLevels);
+	void StreamOutMipLevel(FStreamingInfo& SVTInfo, FLRUNode* LRUNode);
+	int32 DetermineReadyMipLevels();
+	void InstallReadyMipLevels();
+
+#if WITH_EDITORONLY_DATA
+	UE::DerivedData::FCacheGetChunkRequest BuildDDCRequest(const FResources& Resources, const FMipLevelStreamingInfo& MipLevelStreamingInfo, const uint32 PendingMipLevelIndex);
+	void RequestDDCData(TConstArrayView<UE::DerivedData::FCacheGetChunkRequest> DDCRequests);
+#endif
+};
+
+extern ENGINE_API TGlobalResource<FStreamingManager> GStreamingManager;
+
+}
+}
diff --git a/Engine/Source/Runtime/Renderer/Private/DeferredShadingRenderer.cpp b/Engine/Source/Runtime/Renderer/Private/DeferredShadingRenderer.cpp
--- a/Engine/Source/Runtime/Renderer/Private/DeferredShadingRenderer.cpp
+++ b/Engine/Source/Runtime/Renderer/Private/DeferredShadingRenderer.cpp
@@ -91,6 +91,7 @@
 #include "LocalHeightFogRendering.h"
 #include "Shadows/ShadowScene.h"
 #include "Lumen/LumenHardwareRayTracingCommon.h"
+#include "SparseVolumeTexture/SparseVolumeTextureStreamingManager.h"
 
 #if !UE_BUILD_SHIPPING
 #include "RenderCaptureInterface.h"
@@ -2647,6 +2648,8 @@ void FDeferredShadingSceneRenderer::Render(FRDGBuilder& GraphBuilder)
 #endif
 	}
 
+	UE::SVT::GStreamingManager.BeginAsyncUpdate(GraphBuilder);
+
 	bool bUpdateNaniteStreaming = false;
 	bool bVisualizeNanite = false;
 	if (bNaniteEnabled)
@@ -3010,6 +3013,8 @@ void FDeferredShadingSceneRenderer::Render(FRDGBuilder& GraphBuilder)
 		}
 	}
 
+	UE::SVT::GStreamingManager.EndAsyncUpdate(GraphBuilder);
+
 	FHairStrandsBookmarkParameters& HairStrandsBookmarkParameters = *GraphBuilder.AllocObject<FHairStrandsBookmarkParameters>();
 	if (IsHairStrandsEnabled(EHairStrandsShaderType::All, Scene->GetShaderPlatform()) && RendererOutput == ERendererOutput::FinalSceneColor)
 	{
diff --git a/Engine/Source/Runtime/Renderer/Private/SparseVolumeTexture/SparseVolumeTextureViewerComponent.cpp b/Engine/Source/Runtime/Renderer/Private/SparseVolumeTexture/SparseVolumeTextureViewerComponent.cpp
--- a/Engine/Source/Runtime/Renderer/Private/SparseVolumeTexture/SparseVolumeTextureViewerComponent.cpp
+++ b/Engine/Source/Runtime/Renderer/Private/SparseVolumeTexture/SparseVolumeTextureViewerComponent.cpp
@@ -149,14 +149,11 @@ void USparseVolumeTextureViewerComponent::SendRenderTransformCommand()
 	{
 		FVector VolumeExtent = FVector(SVTViewerDefaultVolumeExtent);
 		FVector VolumeResolution = FVector(SVTViewerDefaultVolumeExtent * 2.0);
-		FUintVector4 PackedSVTUniforms0 = FUintVector4();
-		FUintVector4 PackedSVTUniforms1 = FUintVector4();
 		if (SparseVolumeTextureFrame)
 		{
 			VolumeResolution = FVector(SparseVolumeTextureFrame->GetVolumeResolution());
 			const double MaxBoundsDim = FMath::Max(FMath::Max(VolumeResolution.X, VolumeResolution.Y), VolumeResolution.Z);
 			VolumeExtent = VolumeResolution / MaxBoundsDim * SVTViewerDefaultVolumeExtent;
-			SparseVolumeTextureFrame->GetPackedUniforms(PackedSVTUniforms0, PackedSVTUniforms1);
 		}
 
 		const FTransform ToWorldTransform = GetComponentTransform();
@@ -178,13 +175,11 @@ void USparseVolumeTextureViewerComponent::SendRenderTransformCommand()
 
 		FSparseVolumeTextureViewerSceneProxy* SVTViewerSceneProxy = SparseVolumeTextureViewerSceneProxy;
 		ENQUEUE_RENDER_COMMAND(FUpdateSparseVolumeTextureViewerProxyTransformCommand)(
-			[SVTViewerSceneProxy, ToLocalMat, ToLocalMatNoScale, CompIdx = (uint32)ComponentToVisualize, Ext = Extinction, Mip = MipLevel, PackedSVTUniforms0, PackedSVTUniforms1, VolumeRes3f]
+			[SVTViewerSceneProxy, ToLocalMat, ToLocalMatNoScale, CompIdx = (uint32)ComponentToVisualize, Ext = Extinction, Mip = MipLevel, VolumeRes3f]
 		(FRHICommandList& RHICmdList)
 			{
 				SVTViewerSceneProxy->WorldToLocal = ToLocalMat;
 				SVTViewerSceneProxy->WorldToLocalNoScale = ToLocalMatNoScale;
-				SVTViewerSceneProxy->PackedSVTUniforms0 = PackedSVTUniforms0;
-				SVTViewerSceneProxy->PackedSVTUniforms1 = PackedSVTUniforms1;
 				SVTViewerSceneProxy->VolumeResolution = VolumeRes3f;
 				SVTViewerSceneProxy->MipLevel = Mip;
 				SVTViewerSceneProxy->ComponentToVisualize = CompIdx;
@@ -197,19 +192,21 @@ void USparseVolumeTextureViewerComponent::TickComponent(float DeltaTime, enum EL
 {
 	if (SparseVolumeTexturePreview)
 	{
+		const int32 NumFrames = SparseVolumeTexturePreview->GetNumFrames();
+		float FrameIndexF = 0.0f;
 		if (bAnimate)
 		{
-			const int32 NumFrames = SparseVolumeTexturePreview->GetNumFrames();
-			const float AnimationDuration = NumFrames / (FrameRate + UE_SMALL_NUMBER);
-			AnimationTime = FMath::Fmod(AnimationTime + DeltaTime, (AnimationDuration + UE_SMALL_NUMBER));
-			FrameIndex = int32(AnimationTime * FrameRate);
+			const float AnimationDuration = (NumFrames / (FrameRate + UE_SMALL_NUMBER)) + UE_SMALL_NUMBER;
+			AnimationTime = FMath::Fmod(AnimationTime + AnimationDuration + (bReversePlayback ? -DeltaTime : DeltaTime), AnimationDuration);
+			FrameIndexF = AnimationTime * FrameRate;
 		}
 		else
 		{
-			FrameIndex = int32(AnimationFrame * float(SparseVolumeTexturePreview->GetNumFrames()));
+			FrameIndexF = AnimationFrame * float(NumFrames);
 		}
-
-		SparseVolumeTextureFrame = USparseVolumeTextureFrame::GetFrame(SparseVolumeTexturePreview, FrameIndex);
+		FrameIndexF = FMath::Clamp(FrameIndexF, 0, static_cast<float>(NumFrames - 1));
+		FrameIndex = FrameIndexF;
+		SparseVolumeTextureFrame = USparseVolumeTextureFrame::GetFrameAndIssueStreamingRequest(SparseVolumeTexturePreview, FrameIndexF, MipLevel);
 	}
 	else
 	{
diff --git a/Engine/Source/Runtime/Renderer/Private/SparseVolumeTexture/SparseVolumeTextureViewerComponent.h b/Engine/Source/Runtime/Renderer/Private/SparseVolumeTexture/SparseVolumeTextureViewerComponent.h
--- a/Engine/Source/Runtime/Renderer/Private/SparseVolumeTexture/SparseVolumeTextureViewerComponent.h
+++ b/Engine/Source/Runtime/Renderer/Private/SparseVolumeTexture/SparseVolumeTextureViewerComponent.h
@@ -32,6 +32,9 @@ class USparseVolumeTextureViewerComponent : public UPrimitiveComponent
 	UPROPERTY(EditAnywhere, BlueprintReadOnly, Category = "Asset Preview")
 	uint32 bAnimate : 1;
 
+	UPROPERTY(EditAnywhere, BlueprintReadOnly, Category = "Asset Preview")
+	uint32 bReversePlayback : 1;
+
 	UPROPERTY(EditAnywhere, BlueprintReadWrite, Category = "Asset Preview", meta = (UIMin = 0.0, UIMax = 1.0, ClampMin = 0.0, ClampMax = 1.0, EditCondition = "!bAnimate"))
 	float AnimationFrame;
 
diff --git a/Engine/Source/Runtime/Renderer/Private/SparseVolumeTexture/SparseVolumeTextureViewerRendering.cpp b/Engine/Source/Runtime/Renderer/Private/SparseVolumeTexture/SparseVolumeTextureViewerRendering.cpp
--- a/Engine/Source/Runtime/Renderer/Private/SparseVolumeTexture/SparseVolumeTextureViewerRendering.cpp
+++ b/Engine/Source/Runtime/Renderer/Private/SparseVolumeTexture/SparseVolumeTextureViewerRendering.cpp
@@ -3,7 +3,6 @@
 #include "SparseVolumeTexture/SparseVolumeTextureViewerRendering.h"
 #include "SparseVolumeTexture/SparseVolumeTextureViewerSceneProxy.h"
 #include "SparseVolumeTexture/SparseVolumeTexture.h"
-#include "SparseVolumeTexture/SparseVolumeTextureSceneProxy.h"
 
 #include "DataDrivenShaderPlatformInfo.h"
 #include "SceneView.h"
@@ -61,6 +60,7 @@ class FVisualizeSparseVolumeTexturePS : public FGlobalShader
 		SHADER_PARAMETER_TEXTURE(Texture3D<uint>, SparseVolumeTexturePageTable)
 		SHADER_PARAMETER_TEXTURE(Texture3D, SparseVolumeTextureA)
 		SHADER_PARAMETER_TEXTURE(Texture3D, SparseVolumeTextureB)
+		SHADER_PARAMETER_SRV(ByteAddressBuffer, StreamingInfoBuffer)
 		SHADER_PARAMETER(FUintVector4, PackedSVTUniforms0)
 		SHADER_PARAMETER(FUintVector4, PackedSVTUniforms1)
 		SHADER_PARAMETER(FVector3f, SparseVolumeTextureResolution)
@@ -117,10 +117,10 @@ void AddSparseVolumeTextureViewerRenderPass(FRDGBuilder& GraphBuilder, FSceneRen
 	{
 		const FIntRect& ViewportRect = View.ViewRect;
 
-		for (auto& SVTProxy : Scene->SparseVolumeTextureViewers)
+		for (auto& SVTSceneProxy : Scene->SparseVolumeTextureViewers)
 		{
-			const FMatrix44f& WorldToLocal = SVTProxy->WorldToLocal;
-			const FMatrix44f& WorldToLocalNoScale = SVTProxy->WorldToLocalNoScale;
+			const FMatrix44f& WorldToLocal = SVTSceneProxy->WorldToLocal;
+			const FMatrix44f& WorldToLocalNoScale = SVTSceneProxy->WorldToLocalNoScale;
 			
 			FGlobalShaderMap* GlobalShaderMap = GetGlobalShaderMap(GMaxRHIFeatureLevel);
 			FVisualizeSparseVolumeTextureVS::FPermutationDomain VsPermutationVector;
@@ -135,36 +135,45 @@ void AddSparseVolumeTextureViewerRenderPass(FRDGBuilder& GraphBuilder, FSceneRen
 			PsPassParameters->SparseVolumeTexturePageTable = GBlackUintVolumeTexture->TextureRHI;
 			PsPassParameters->SparseVolumeTextureA = GBlackVolumeTexture->TextureRHI;
 			PsPassParameters->SparseVolumeTextureB = GBlackVolumeTexture->TextureRHI;
-			PsPassParameters->PackedSVTUniforms0 = SVTProxy->PackedSVTUniforms0;
-			PsPassParameters->PackedSVTUniforms1 = SVTProxy->PackedSVTUniforms1;
-			PsPassParameters->SparseVolumeTextureResolution = SVTProxy->VolumeResolution;
-			PsPassParameters->MipLevel = SVTProxy->MipLevel;
+			PsPassParameters->StreamingInfoBuffer = GEmptyVertexBufferWithUAV->ShaderResourceViewRHI;
+			PsPassParameters->PackedSVTUniforms0 = FUintVector4(0);
+			PsPassParameters->PackedSVTUniforms1 = FUintVector4(0);
+			PsPassParameters->SparseVolumeTextureResolution = SVTSceneProxy->VolumeResolution;
+			PsPassParameters->MipLevel = SVTSceneProxy->MipLevel;
 			PsPassParameters->WorldToLocal0 = FVector4f(WorldToLocal.M[0][0], WorldToLocal.M[1][0], WorldToLocal.M[2][0], WorldToLocal.M[3][0]);
 			PsPassParameters->WorldToLocal1 = FVector4f(WorldToLocal.M[0][1], WorldToLocal.M[1][1], WorldToLocal.M[2][1], WorldToLocal.M[3][1]);
 			PsPassParameters->WorldToLocal2 = FVector4f(WorldToLocal.M[0][2], WorldToLocal.M[1][2], WorldToLocal.M[2][2], WorldToLocal.M[3][2]);
 			PsPassParameters->WorldToLocalNoScale0 = FVector3f(WorldToLocalNoScale.M[0][0], WorldToLocalNoScale.M[1][0], WorldToLocalNoScale.M[2][0]);
 			PsPassParameters->WorldToLocalNoScale1 = FVector3f(WorldToLocalNoScale.M[0][1], WorldToLocalNoScale.M[1][1], WorldToLocalNoScale.M[2][1]);
 			PsPassParameters->WorldToLocalNoScale2 = FVector3f(WorldToLocalNoScale.M[0][2], WorldToLocalNoScale.M[1][2], WorldToLocalNoScale.M[2][2]);
-			PsPassParameters->ComponentToVisualize = SVTProxy->ComponentToVisualize;
-			PsPassParameters->Extinction = SVTProxy->Extinction;
+			PsPassParameters->ComponentToVisualize = SVTSceneProxy->ComponentToVisualize;
+			PsPassParameters->Extinction = SVTSceneProxy->Extinction;
 			
-			if (SVTProxy->SparseVolumeTextureSceneProxy)
+			const UE::SVT::FTextureRenderResources* RenderResources = SVTSceneProxy->TextureRenderResources;
+			if (RenderResources)
 			{
-				FRHITexture* PageTableTexture = SVTProxy->SparseVolumeTextureSceneProxy->GetPageTableTextureRHI();
+				SVTSceneProxy->TextureRenderResources->GetPackedUniforms(PsPassParameters->PackedSVTUniforms0, PsPassParameters->PackedSVTUniforms1);
+
+				FRHITexture* PageTableTexture = RenderResources->GetPageTableTextureRHI();
 				if (PageTableTexture)
 				{
 					PsPassParameters->SparseVolumeTexturePageTable = PageTableTexture;
 				}
-				FRHITexture* TextureA = SVTProxy->SparseVolumeTextureSceneProxy->GetPhysicalTileDataATextureRHI();
+				FRHITexture* TextureA = RenderResources->GetPhysicalTileDataATextureRHI();
 				if (TextureA)
 				{
 					PsPassParameters->SparseVolumeTextureA = TextureA;
 				}
-				FRHITexture* TextureB = SVTProxy->SparseVolumeTextureSceneProxy->GetPhysicalTileDataBTextureRHI();
+				FRHITexture* TextureB = RenderResources->GetPhysicalTileDataBTextureRHI();
 				if (TextureB)
 				{
 					PsPassParameters->SparseVolumeTextureB = TextureB;
 				}
+				FRHIShaderResourceView* StreamingInfoBufferSRV = RenderResources->GetStreamingInfoBufferSRVRHI();
+				if (StreamingInfoBufferSRV)
+				{
+					PsPassParameters->StreamingInfoBuffer = StreamingInfoBufferSRV;
+				}
 			}
 
 			ClearUnusedGraphResources(PixelShader, PsPassParameters);
diff --git a/Engine/Source/Runtime/Renderer/Private/SparseVolumeTexture/SparseVolumeTextureViewerSceneProxy.cpp b/Engine/Source/Runtime/Renderer/Private/SparseVolumeTexture/SparseVolumeTextureViewerSceneProxy.cpp
--- a/Engine/Source/Runtime/Renderer/Private/SparseVolumeTexture/SparseVolumeTextureViewerSceneProxy.cpp
+++ b/Engine/Source/Runtime/Renderer/Private/SparseVolumeTexture/SparseVolumeTextureViewerSceneProxy.cpp
@@ -38,12 +38,11 @@ void FScene::RemoveSparseVolumeTextureViewer(FSparseVolumeTextureViewerSceneProx
 
 FSparseVolumeTextureViewerSceneProxy::FSparseVolumeTextureViewerSceneProxy(const USparseVolumeTextureViewerComponent* InComponent, int32 FrameIndex, FName ResourceName)
 	: FPrimitiveSceneProxy((UPrimitiveComponent*)InComponent, ResourceName)
-	, SparseVolumeTextureSceneProxy(nullptr)
+	, TextureRenderResources(nullptr)
 {
 	if (InComponent->SparseVolumeTextureFrame)
 	{
-		check(InComponent->SparseVolumeTextureFrame);
-		SparseVolumeTextureSceneProxy = InComponent->SparseVolumeTextureFrame->GetSparseVolumeTextureSceneProxy();
+		TextureRenderResources = InComponent->SparseVolumeTextureFrame->GetTextureRenderResources();
 	}
 }
 
diff --git a/Engine/Source/Runtime/Renderer/Private/SparseVolumeTexture/SparseVolumeTextureViewerSceneProxy.h b/Engine/Source/Runtime/Renderer/Private/SparseVolumeTexture/SparseVolumeTextureViewerSceneProxy.h
--- a/Engine/Source/Runtime/Renderer/Private/SparseVolumeTexture/SparseVolumeTextureViewerSceneProxy.h
+++ b/Engine/Source/Runtime/Renderer/Private/SparseVolumeTexture/SparseVolumeTextureViewerSceneProxy.h
@@ -8,7 +8,14 @@
 #include "SceneViewExtension.h"
 
 class USparseVolumeTextureViewerComponent;
-class FSparseVolumeTextureSceneProxy;
+namespace UE
+{
+	namespace SVT
+	{
+		class FTextureRenderResources;
+	}
+}
+
 
 // Proxy representing the rendering of the SparseVolumeTextureViewer component on the render thread.
 class FSparseVolumeTextureViewerSceneProxy : public FPrimitiveSceneProxy
@@ -17,11 +24,9 @@ public:
 	FSparseVolumeTextureViewerSceneProxy(const USparseVolumeTextureViewerComponent* InComponent, int32 FrameIndex, FName ResourceName = NAME_None);
 	virtual ~FSparseVolumeTextureViewerSceneProxy() = default;
 
-	const FSparseVolumeTextureSceneProxy* SparseVolumeTextureSceneProxy = nullptr;
+	const UE::SVT::FTextureRenderResources* TextureRenderResources;
 	FMatrix44f WorldToLocal;
 	FMatrix44f WorldToLocalNoScale;
-	FUintVector4 PackedSVTUniforms0;
-	FUintVector4 PackedSVTUniforms1;
 	FVector3f VolumeResolution;
 	int32 MipLevel;
 	uint32 ComponentToVisualize;
