Moved SVT streaming manager behind an interface to hide the implementation details present in the classes header.

#rnx

diff --git a/Engine/Source/Runtime/Engine/Private/SparseVolumeTexture/SparseVolumeTexture.cpp b/Engine/Source/Runtime/Engine/Private/SparseVolumeTexture/SparseVolumeTexture.cpp
--- a/Engine/Source/Runtime/Engine/Private/SparseVolumeTexture/SparseVolumeTexture.cpp
+++ b/Engine/Source/Runtime/Engine/Private/SparseVolumeTexture/SparseVolumeTexture.cpp
@@ -16,7 +16,7 @@
 #include "UObject/Package.h"
 #include "SparseVolumeTexture/SparseVolumeTextureData.h"
 #include "SparseVolumeTexture/SparseVolumeTextureUtility.h"
-#include "SparseVolumeTexture/SparseVolumeTextureStreamingManager.h"
+#include "SparseVolumeTexture/ISparseVolumeTextureStreamingManager.h"
 
 #if WITH_EDITORONLY_DATA
 #include "Misc/ScopedSlowTask.h"
@@ -647,7 +647,7 @@ USparseVolumeTextureFrame* USparseVolumeTextureFrame::GetFrameAndIssueStreamingR
 {
 	if (UStreamableSparseVolumeTexture* StreamableSVT = Cast<UStreamableSparseVolumeTexture>(SparseVolumeTexture))
 	{
-		UE::SVT::GStreamingManager.Request_GameThread(StreamableSVT, FrameIndex, MipLevel);
+		UE::SVT::GetStreamingManager().Request_GameThread(StreamableSVT, FrameIndex, MipLevel);
 		return StreamableSVT->GetFrame(static_cast<int32>(FrameIndex));
 	}
 	return nullptr;
@@ -699,7 +699,7 @@ void USparseVolumeTextureFrame::BeginDestroy()
 	// Ensure that the streamable SVT has been removed from the streaming manager
 	if (IsValid(Owner))
 	{
-		UE::SVT::GStreamingManager.Remove_GameThread(CastChecked<UStreamableSparseVolumeTexture>(Owner));
+		UE::SVT::GetStreamingManager().Remove_GameThread(CastChecked<UStreamableSparseVolumeTexture>(Owner));
 	}
 	
 	if (TextureRenderResources)
@@ -961,7 +961,7 @@ void UStreamableSparseVolumeTexture::PostLoad()
 	{
 		Frame->CreateTextureRenderResources();
 	}
-	UE::SVT::GStreamingManager.Add_GameThread(this); // RecacheFrames() handles this in editor builds
+	UE::SVT::GetStreamingManager().Add_GameThread(this); // RecacheFrames() handles this in editor builds
 #endif
 }
 
@@ -973,7 +973,7 @@ void UStreamableSparseVolumeTexture::FinishDestroy()
 void UStreamableSparseVolumeTexture::BeginDestroy()
 {
 	Super::BeginDestroy();
-	UE::SVT::GStreamingManager.Remove_GameThread(this);
+	UE::SVT::GetStreamingManager().Remove_GameThread(this);
 }
 
 void UStreamableSparseVolumeTexture::Serialize(FArchive& Ar)
@@ -1055,7 +1055,7 @@ void UStreamableSparseVolumeTexture::RecacheFrames()
 	FScopedSlowTask RecacheTask(static_cast<float>(Frames.Num() + 2), LOCTEXT("SparseVolumeTextureCacheFrames", "Caching SparseVolumeTexture frames in Derived Data Cache"));
 	RecacheTask.MakeDialog(true);
 
-	UE::SVT::GStreamingManager.Remove_GameThread(this);
+	UE::SVT::GetStreamingManager().Remove_GameThread(this);
 	RecacheTask.EnterProgressFrame(1.0f);
 
 	for (USparseVolumeTextureFrame* Frame : Frames)
@@ -1065,7 +1065,7 @@ void UStreamableSparseVolumeTexture::RecacheFrames()
 	}
 
 	RecacheTask.EnterProgressFrame(1.0f);
-	UE::SVT::GStreamingManager.Add_GameThread(this);
+	UE::SVT::GetStreamingManager().Add_GameThread(this);
 }
 #endif
 
diff --git a/Engine/Source/Runtime/Engine/Private/SparseVolumeTexture/SparseVolumeTextureStreamingManager.cpp b/Engine/Source/Runtime/Engine/Private/SparseVolumeTexture/SparseVolumeTextureStreamingManager.cpp
--- a/Engine/Source/Runtime/Engine/Private/SparseVolumeTexture/SparseVolumeTextureStreamingManager.cpp
+++ b/Engine/Source/Runtime/Engine/Private/SparseVolumeTexture/SparseVolumeTextureStreamingManager.cpp
@@ -74,6 +74,13 @@
 namespace SVT
 {
 
+TGlobalResource<FStreamingManager> GStreamingManager;
+
+IStreamingManager& GetStreamingManager()
+{
+	return GStreamingManager;
+}
+
 static bool DoesPlatformSupportSparseVolumeTexture(EShaderPlatform Platform)
 {
 	// SVT_TODO: This is a bit of a hack: FStreamingManager::Add_GameThread() issues a rendering thread lambda for creating the RHI resources and uploading root tile data.
@@ -741,6 +748,103 @@
 		});
 }
 
+void FStreamingManager::Request_GameThread(UStreamableSparseVolumeTexture* SparseVolumeTexture, float FrameIndex, int32 MipLevel)
+{
+	if (!DoesPlatformSupportSparseVolumeTexture(GMaxRHIShaderPlatform) || !SparseVolumeTexture)
+	{
+		return;
+	}
+	ENQUEUE_RENDER_COMMAND(SVTRequest)(
+		[this, SparseVolumeTexture, FrameIndex, MipLevel](FRHICommandListImmediate& RHICmdList)
+		{
+			Request(SparseVolumeTexture, FrameIndex, MipLevel);
+		});
+}
+
+void FStreamingManager::Request(UStreamableSparseVolumeTexture* SparseVolumeTexture, float FrameIndex, int32 MipLevel)
+{
+	check(IsInRenderingThread());
+	if (!DoesPlatformSupportSparseVolumeTexture(GMaxRHIShaderPlatform) || !SparseVolumeTexture)
+	{
+		return;
+	}
+
+	FStreamingInfo* SVTInfo = StreamingInfo.Find(SparseVolumeTexture);
+	if (SVTInfo)
+	{
+		const int32 NumFrames = SVTInfo->PerFrameInfo.Num();
+		const int32 FrameIndexI32 = static_cast<int32>(FrameIndex);
+		if (FrameIndexI32 < 0 || FrameIndexI32 >= NumFrames)
+		{
+			return;
+		}
+
+		// Try to find a FStreamingWindow around the requested frame index. This will inform us about which direction we need to prefetch into.
+		FStreamingWindow* StreamingWindow = nullptr;
+		for (FStreamingWindow& Window : SVTInfo->StreamingWindows)
+		{
+			if (FMath::Abs(FrameIndex - Window.CenterFrame) <= FStreamingWindow::WindowSize)
+			{
+				StreamingWindow = &Window;
+				break;
+			}
+		}
+		// Found an existing window!
+		if (StreamingWindow)
+		{
+			const bool bForward = StreamingWindow->LastCenterFrame <= FrameIndex;
+			if (StreamingWindow->LastRequested < NextUpdateIndex)
+			{
+				StreamingWindow->LastCenterFrame = StreamingWindow->CenterFrame;
+				StreamingWindow->CenterFrame = FrameIndex;
+				StreamingWindow->NumRequestsThisUpdate = 1;
+				StreamingWindow->LastRequested = NextUpdateIndex;
+				StreamingWindow->bPlayForward = bForward;
+				StreamingWindow->bPlayBackward = !bForward;
+			}
+			else
+			{
+				// Update the average center frame
+				StreamingWindow->CenterFrame = (StreamingWindow->CenterFrame * StreamingWindow->NumRequestsThisUpdate + FrameIndex) / (StreamingWindow->NumRequestsThisUpdate + 1.0f);
+				++StreamingWindow->NumRequestsThisUpdate;
+				StreamingWindow->bPlayForward |= bForward;
+				StreamingWindow->bPlayBackward |= !bForward;
+			}
+		}
+		// No existing window. Create a new one.
+		else
+		{
+			StreamingWindow = &SVTInfo->StreamingWindows.AddDefaulted_GetRef();
+			StreamingWindow->CenterFrame = FrameIndex;
+			StreamingWindow->LastCenterFrame = FrameIndex;
+			StreamingWindow->NumRequestsThisUpdate = 1;
+			StreamingWindow->LastRequested = NextUpdateIndex;
+			StreamingWindow->bPlayForward = true; // No prior data, so just take a guess that playback is forwards
+			StreamingWindow->bPlayBackward = false;
+		}
+
+		check(StreamingWindow);
+		const int32 OffsetMagnitude = GSVTStreamingNumPrefetchFrames;
+		const int32 LowerFrameOffset = StreamingWindow->bPlayBackward ? -OffsetMagnitude : 0;
+		const int32 UpperFrameOffset = StreamingWindow->bPlayForward ? OffsetMagnitude : 0;
+
+		for (int32 i = LowerFrameOffset; i <= UpperFrameOffset; ++i)
+		{
+			const int32 RequestFrameIndex = (static_cast<int32>(FrameIndex) + i + NumFrames) % NumFrames;
+			const int32 RequestMipLevelOffset = FMath::Abs(i) + GSVTStreamingPrefetchMipLevelBias;
+			FStreamingRequest Request;
+			Request.Key.SVT = SparseVolumeTexture;
+			Request.Key.FrameIndex = RequestFrameIndex;
+			Request.Key.MipLevelIndex = FMath::Clamp(MipLevel + RequestMipLevelOffset, 0, SVTInfo->PerFrameInfo[RequestFrameIndex].NumMipLevels);
+			Request.Priority = FMath::Max(0, OffsetMagnitude - FMath::Abs(i));
+			AddRequest(Request);
+		}
+
+		// Clean up unused streaming windows
+		SVTInfo->StreamingWindows.RemoveAll([&](const FStreamingWindow& Window) { return (NextUpdateIndex - Window.LastRequested) > 5; });
+	}
+}
+
 void FStreamingManager::BeginAsyncUpdate(FRDGBuilder& GraphBuilder)
 {
 	check(IsInRenderingThread());
@@ -932,98 +1036,6 @@
 #endif // DO_CHECK
 }
 
-void FStreamingManager::Request_GameThread(UStreamableSparseVolumeTexture* SparseVolumeTexture, float FrameIndex, int32 MipLevel)
-{
-	if (!DoesPlatformSupportSparseVolumeTexture(GMaxRHIShaderPlatform) || !SparseVolumeTexture)
-	{
-		return;
-	}
-	ENQUEUE_RENDER_COMMAND(SVTRequest)(
-		[this, SparseVolumeTexture, FrameIndex, MipLevel](FRHICommandListImmediate& RHICmdList)
-		{
-			Request(SparseVolumeTexture, FrameIndex, MipLevel);
-		});
-}
-
-void FStreamingManager::Request(UStreamableSparseVolumeTexture* SparseVolumeTexture, float FrameIndex, int32 MipLevel)
-{
-	check(IsInRenderingThread());
-	FStreamingInfo* SVTInfo = StreamingInfo.Find(SparseVolumeTexture);
-	if (SVTInfo)
-	{
-		const int32 NumFrames = SVTInfo->PerFrameInfo.Num();
-		const int32 FrameIndexI32 = static_cast<int32>(FrameIndex);
-		if (FrameIndexI32 < 0 || FrameIndexI32 >= NumFrames)
-		{
-			return;
-		}
-
-		// Try to find a FStreamingWindow around the requested frame index. This will inform us about which direction we need to prefetch into.
-		FStreamingWindow* StreamingWindow = nullptr;
-		for (FStreamingWindow& Window : SVTInfo->StreamingWindows)
-		{
-			if (FMath::Abs(FrameIndex - Window.CenterFrame) <= FStreamingWindow::WindowSize)
-			{
-				StreamingWindow = &Window;
-				break;
-			}
-		}
-		// Found an existing window!
-		if (StreamingWindow)
-		{
-			const bool bForward = StreamingWindow->LastCenterFrame <= FrameIndex;
-			if (StreamingWindow->LastRequested < NextUpdateIndex)
-			{
-				StreamingWindow->LastCenterFrame = StreamingWindow->CenterFrame;
-				StreamingWindow->CenterFrame = FrameIndex;
-				StreamingWindow->NumRequestsThisUpdate = 1;
-				StreamingWindow->LastRequested = NextUpdateIndex;
-				StreamingWindow->bPlayForward = bForward;
-				StreamingWindow->bPlayBackward = !bForward;
-			}
-			else
-			{
-				// Update the average center frame
-				StreamingWindow->CenterFrame = (StreamingWindow->CenterFrame * StreamingWindow->NumRequestsThisUpdate + FrameIndex) / (StreamingWindow->NumRequestsThisUpdate + 1.0f);
-				++StreamingWindow->NumRequestsThisUpdate;
-				StreamingWindow->bPlayForward |= bForward;
-				StreamingWindow->bPlayBackward |= !bForward;
-			}
-		}
-		// No existing window. Create a new one.
-		else
-		{
-			StreamingWindow = &SVTInfo->StreamingWindows.AddDefaulted_GetRef();
-			StreamingWindow->CenterFrame = FrameIndex;
-			StreamingWindow->LastCenterFrame = FrameIndex;
-			StreamingWindow->NumRequestsThisUpdate = 1;
-			StreamingWindow->LastRequested = NextUpdateIndex;
-			StreamingWindow->bPlayForward = true; // No prior data, so just take a guess that playback is forwards
-			StreamingWindow->bPlayBackward = false;
-		}
-
-		check(StreamingWindow);
-		const int32 OffsetMagnitude = GSVTStreamingNumPrefetchFrames;
-		const int32 LowerFrameOffset = StreamingWindow->bPlayBackward ? -OffsetMagnitude : 0;
-		const int32 UpperFrameOffset = StreamingWindow->bPlayForward ? OffsetMagnitude : 0;
-
-		for (int32 i = LowerFrameOffset; i <= UpperFrameOffset; ++i)
-		{
-			const int32 RequestFrameIndex = (static_cast<int32>(FrameIndex) + i + NumFrames) % NumFrames;
-			const int32 RequestMipLevelOffset = FMath::Abs(i) + GSVTStreamingPrefetchMipLevelBias;
-			FStreamingRequest Request;
-			Request.Key.SVT = SparseVolumeTexture;
-			Request.Key.FrameIndex = RequestFrameIndex;
-			Request.Key.MipLevelIndex = FMath::Clamp(MipLevel + RequestMipLevelOffset, 0, SVTInfo->PerFrameInfo[RequestFrameIndex].NumMipLevels);
-			Request.Priority = FMath::Max(0, OffsetMagnitude - FMath::Abs(i));
-			AddRequest(Request);
-		}
-
-		// Clean up unused streaming windows
-		SVTInfo->StreamingWindows.RemoveAll([&](const FStreamingWindow& Window) { return (NextUpdateIndex - Window.LastRequested) > 5; });
-	}
-}
-
 void FStreamingManager::AddInternal(FRDGBuilder& GraphBuilder, FNewSparseVolumeTextureInfo&& NewSVTInfo)
 {
 	check(IsInRenderingThread());
@@ -2096,7 +2108,5 @@
 {
 }
 
-TGlobalResource<FStreamingManager> GStreamingManager;
-
 }
 }
diff --git a/Engine/Source/Runtime/Engine/Private/SparseVolumeTexture/SparseVolumeTextureStreamingManager.h b/Engine/Source/Runtime/Engine/Private/SparseVolumeTexture/SparseVolumeTextureStreamingManager.h
new file mode 100644
--- /dev/null
+++ b/Engine/Source/Runtime/Engine/Private/SparseVolumeTexture/SparseVolumeTextureStreamingManager.h
@@ -0,0 +1,362 @@
+// Copyright Epic Games, Inc. All Rights Reserved.
+
+#pragma once
+
+#include "CoreMinimal.h"
+#include "Engine/EngineTypes.h"
+#include "SparseVolumeTexture/ISparseVolumeTextureStreamingManager.h"
+#include "IO/IoDispatcher.h"
+#include "Memory/SharedBuffer.h"
+#include "Serialization/BulkData.h"
+#include "RenderResource.h"
+#include "Containers/IntrusiveDoubleLinkedList.h"
+#include "Containers/Map.h"
+#include "RenderGraphBuilder.h"
+
+class UStreamableSparseVolumeTexture;
+class USparseVolumeTextureFrame;
+
+namespace UE
+{
+	namespace DerivedData
+	{
+		class FRequestOwner; // Can't include DDC headers from here, so we have to forward declare
+		struct FCacheGetChunkRequest;
+	}
+}
+
+namespace UE
+{
+namespace SVT
+{
+
+struct FResources;
+struct FMipLevelStreamingInfo;
+class FTextureRenderResources;
+
+// Uniquely identifies the mip level of a frame in a static or animated SparseVolumeTexture
+struct FMipLevelKey
+{
+	UStreamableSparseVolumeTexture* SVT = nullptr; // This struct is used on the rendering thread. Do not dereference!
+	uint16 FrameIndex = INDEX_NONE;
+	uint16 MipLevelIndex = INDEX_NONE;
+
+	friend FORCEINLINE uint32 GetTypeHash(const FMipLevelKey& Key)
+	{
+		return HashCombine(GetTypeHash(Key.SVT), GetTypeHash(((uint32)Key.FrameIndex << 16u) | (uint32)Key.MipLevelIndex));
+	}
+
+	FORCEINLINE bool operator==(const FMipLevelKey& Other) const
+	{
+		return (SVT == Other.SVT) && (FrameIndex == Other.FrameIndex) && (MipLevelIndex == Other.MipLevelIndex);
+	}
+
+	FORCEINLINE bool operator!=(const FMipLevelKey& Other) const
+	{
+		return !(*this == Other);
+	}
+
+	FORCEINLINE bool operator<(const FMipLevelKey& Other) const
+	{
+		return SVT != Other.SVT ? SVT < Other.SVT : FrameIndex != Other.FrameIndex ? FrameIndex < Other.FrameIndex : MipLevelIndex > Other.MipLevelIndex;
+	}
+};
+
+struct FStreamingRequest
+{
+	FMipLevelKey Key;
+	uint32 Priority; // A higher value means a higher priority
+
+	FORCEINLINE bool operator<(const FStreamingRequest& Other) const
+	{
+		return Key != Other.Key ? Key < Other.Key : Priority > Other.Priority;
+	}
+};
+
+class FStreamingManager : public FRenderResource, public IStreamingManager
+{
+public:
+	FStreamingManager();
+
+	//~ Begin FRenderResource Interface.
+	virtual void InitRHI() override;
+	virtual void ReleaseRHI() override;
+	//~ End FRenderResource Interface.
+
+	//~ Begin IStreamingManager Interface.
+	virtual void Add_GameThread(UStreamableSparseVolumeTexture* SparseVolumeTexture) override;
+	virtual void Remove_GameThread(UStreamableSparseVolumeTexture* SparseVolumeTexture) override;
+	// Request a frame to be streamed in. FrameIndex is of float type so that the fractional part can be used to better track the playback speed/direction.
+	// This function automatically also requests all higher mip levels and adds prefetch requests for upcoming frames.
+	virtual void Request_GameThread(UStreamableSparseVolumeTexture* SparseVolumeTexture, float FrameIndex, int32 MipLevel) override;
+	
+	virtual void Request(UStreamableSparseVolumeTexture* SparseVolumeTexture, float FrameIndex, int32 MipLevel) override;
+	virtual void BeginAsyncUpdate(FRDGBuilder& GraphBuilder) override;
+	virtual void EndAsyncUpdate(FRDGBuilder& GraphBuilder) override;
+	//~ End IStreamingManager Interface.
+
+private:
+	friend class FStreamingUpdateTask;
+
+	// SVT_TODO: actually implement streaming in/out of page table mips greater than this value.
+	static constexpr int32 NumAlwaysResidentPageTableMipLevels = 6;
+
+	// Represents the physical tile data texture that serves as backing memory for the streamed in tiles. While this is treated as a single logical texture,
+	// it currently supports up to two actual RHI textures.
+	struct FTileDataTexture : public FRenderResource
+	{
+	public:
+		static constexpr uint32 PhysicalCoordMask = (1u << 24u) - 1u; // Lower 24 bits are used for storing XYZ in 8 bit each. Upper 8 bit can be used by the caller. 
+
+		FIntVector3 ResolutionInTiles;
+		int32 PhysicalTilesCapacity;
+		EPixelFormat FormatA;
+		EPixelFormat FormatB;
+		FTextureRHIRef TileDataTextureARHIRef;
+		FTextureRHIRef TileDataTextureBRHIRef;
+		TUniquePtr<class FTileUploader> TileUploader;
+		int32 NumTilesToUpload;
+
+		// Constructor. May change the requested ResolutionInTiles (and resulting PhysicalTilesCapacity) if it exceeds hardware limits.
+		FTileDataTexture(const FIntVector3& ResolutionInTiles, EPixelFormat FormatA, EPixelFormat FormatB);
+
+		// Allocate a tile slot in the texture. The resulting value is a packed coordinate (8 bit per component) of the allocated slot or INDEX_NONE if the allocation failed.
+		// The upper 8 bit are free to be used by the caller.
+		uint32 Allocate() 
+		{ 
+			check(PhysicalTilesCapacity == TileCoords.Num());
+			return TileCoords.IsValidIndex(NextFreeTileCoordIndex) ? TileCoords[NextFreeTileCoordIndex++] : INDEX_NONE; 
+		}
+
+		// Frees a previously allocated tile slot. The upper 8 bit (user data) are automatically cleared by this function.
+		void Free(uint32 PackedPhysicalCoord) 
+		{ 
+			check(PackedPhysicalCoord != INDEX_NONE)
+			PackedPhysicalCoord &= PhysicalCoordMask;
+			check(PhysicalTilesCapacity == TileCoords.Num());
+			check(NextFreeTileCoordIndex > 0);
+#if DO_GUARD_SLOW
+			for (int32 i = NextFreeTileCoordIndex; i < PhysicalTilesCapacity; ++i)
+			{
+				check(TileCoords[i] != PackedPhysicalCoord);
+			}
+#endif
+			TileCoords[--NextFreeTileCoordIndex] = PackedPhysicalCoord;
+		}
+
+		// Number of tiles available for allocation.
+		int32 GetNumAvailableTiles() { return PhysicalTilesCapacity - NextFreeTileCoordIndex; }
+		virtual void InitRHI() override;
+		virtual void ReleaseRHI() override;
+
+	private:
+
+		TArray<uint32> TileCoords;
+		int32 NextFreeTileCoordIndex = 0;
+		uint32 NumRefs = 0;
+	};
+
+	// Represents a context or "window" into the frame sequence that moves along the playback direction. It is used to cache the prefetch direction.
+	struct FStreamingWindow
+	{
+		static constexpr int32 WindowSize = 5;
+		float CenterFrame = -1.0f; // Frame index this window is centered around
+		float LastCenterFrame = -1.0f;
+		int32 NumRequestsThisUpdate = 0;
+		uint32 LastRequested = 0;
+		bool bPlayForward = false;
+		bool bPlayBackward = false;
+	};
+
+	// One per frame of each SVT
+	struct FFrameInfo
+	{
+		const FResources* Resources;						// Initialized on the game thread
+		FTextureRenderResources* TextureRenderResources;	// Initialized on the game thread
+
+		int32 NumMipLevels; // Number of actual mip levels of this frame. Depending on virtual volume extents, frames can have different numbers of levels.
+		int32 LowestRequestedMipLevel; // Lowest mip level that should be resident. Can be lower than LowestResidentMipLevel when streaming in new mips. Stream-out is instant, so it should never be higher.
+		int32 LowestResidentMipLevel; // Actually resident on the GPU
+		TArray<TArray<uint32>> TileAllocations; // TileAllocations[MipLevel][PhysicalTileIndex]
+	};
+
+	// Used to keep track of the least-recently-used order of mip levels. This is done per mip level and not per frame, so that lower resolution mip levels are kept in memory for longer.
+	struct FLRUNode : public TIntrusiveDoubleLinkedListNode<FLRUNode>
+	{
+		FLRUNode* NextHigherMipLevel = nullptr;
+		uint32 LastRequested = INDEX_NONE;
+		uint32 RefCount = 0; // Keep track of how many lower mip levels have a dependency on this one
+		uint32 PendingMipLevelIndex = INDEX_NONE;
+		int16 FrameIndex = INDEX_NONE;
+		int16 MipLevelIndex = INDEX_NONE;
+	};
+
+	// Streaming manager internal data for each SVT
+	struct FStreamingInfo
+	{
+		EPixelFormat FormatA;
+		EPixelFormat FormatB;
+		FVector4f FallbackValueA;
+		FVector4f FallbackValueB;
+		int32 NumMipLevelsGlobal; // Maximum number of mip levels in the sequence
+		uint32 LastRequested; // Last update index when some frame in this SVT was requested
+		TArray<FFrameInfo> PerFrameInfo;
+		TArray<FLRUNode> LRUNodes;
+		TArray<TIntrusiveDoubleLinkedList<FLRUNode>> PerMipLRULists; // PerMipLRULists[MipLevel]
+		TArray<FStreamingWindow> StreamingWindows;
+
+		FTileDataTexture* TileDataTexture;
+		TRefCountPtr<FRDGPooledBuffer> StreamingInfoBuffer; // One uint32 per frame storing the LowestResidentMipLevel
+		FShaderResourceViewRHIRef StreamingInfoBufferSRVRHIRef;
+		TBitArray<> DirtyStreamingInfoData; // One bit per frame, potentially marking the streaming info data as dirty/in need of an update
+	};
+
+	// Represents an IO request for a mip level
+	struct FPendingMipLevel
+	{
+		UStreamableSparseVolumeTexture* SparseVolumeTexture = nullptr; // Do not dereference!
+		int32 FrameIndex = INDEX_NONE;
+		int32 MipLevelIndex = INDEX_NONE;
+#if WITH_EDITORONLY_DATA
+		FSharedBuffer SharedBuffer;
+		enum class EState
+		{
+			None,
+			DDC_Pending,
+			DDC_Ready,
+			DDC_Failed,
+			Memory,
+			Disk,
+		} State = EState::None;
+		uint32 RetryCount = 0;
+		uint32 RequestVersion = 0;
+#endif
+		FIoBuffer RequestBuffer;
+		FBulkDataBatchReadRequest Request;
+		uint32 IssuedInFrame = 0;
+
+		void Reset()
+		{
+			SparseVolumeTexture = nullptr;
+			FrameIndex = INDEX_NONE;
+			MipLevelIndex = INDEX_NONE;
+#if WITH_EDITORONLY_DATA
+			SharedBuffer.Reset();
+			State = EState::None;
+			RetryCount = 0;
+			++RequestVersion;
+#endif
+			RequestBuffer = {};
+			Request.Reset();
+			IssuedInFrame = 0;
+		}
+	};
+
+	// Encapsulates all the src/dst pointers for memcpy-ing the streamed data into GPU memory.
+	struct FUploadTask
+	{
+		enum class ETaskType
+		{
+			TileData, PageTable
+		};
+
+		struct FTileDataTask
+		{
+			uint8* DstA;
+			uint8* DstB;
+			uint8* DstPhysicalTileCoords;
+			const uint8* SrcA;
+			const uint8* SrcB;
+			const uint8* SrcPhysicalTileCoords;
+			int32 SizeA;
+			int32 SizeB;
+			int32 NumPhysicalTiles;
+		};
+
+		struct FPageTableTask
+		{
+			FPendingMipLevel* PendingMipLevel;
+			uint8* DstPageCoords;
+			uint8* DstPageEntries;
+			const uint8* SrcPageCoords;
+			const uint8* SrcPageEntries;
+			int32 NumPageTableUpdates;
+		};
+
+		ETaskType TaskType = ETaskType::TileData;
+		union
+		{
+			FTileDataTask TileDataTask;
+			FPageTableTask PageTableTask;
+		};
+	};
+
+	struct FPageTableClear
+	{
+		FTextureRHIRef PageTableTexture;
+		int32 MipLevel;
+	};
+
+	struct FAsyncState
+	{
+		int32 NumReadyMipLevels = 0;
+		bool bUpdateActive = false;
+	};
+
+	// Stores info extracted from the UStreamableSparseVolumeTexture and its USparseVolumeTextureFrames so that we don't need to access these uobjects on the rendering thread
+	struct FNewSparseVolumeTextureInfo
+	{
+		UStreamableSparseVolumeTexture* SVT = nullptr; // Do not dereference!
+		EPixelFormat FormatA = PF_Unknown;
+		EPixelFormat FormatB = PF_Unknown;
+		FVector4f FallbackValueA = FVector4f(0.0f, 0.0f, 0.0f, 0.0f);
+		FVector4f FallbackValueB = FVector4f(0.0f, 0.0f, 0.0f, 0.0f);
+		int32 NumMipLevelsGlobal = 0;
+		TArray<FFrameInfo> FrameInfo; // Only Resources and TextureRenderResources are initialized
+	};
+
+	TMap<UStreamableSparseVolumeTexture*, FStreamingInfo> StreamingInfo; // Do not dereference the key! We just read the pointer itself.
+	TMap<FMipLevelKey, uint32> RequestsHashTable;
+	TArray<FPendingMipLevel> PendingMipLevels;
+#if WITH_EDITORONLY_DATA
+	UE::DerivedData::FRequestOwner* RequestOwner = nullptr;
+#endif
+
+	class FPageTableUpdater* PageTableUpdater = nullptr;
+	class FStreamingInfoBufferUpdater* StreamingInfoBufferUpdater = nullptr;
+	FGraphEventArray AsyncTaskEvents;
+	FAsyncState AsyncState;
+	int32 MaxPendingMipLevels = 0;
+	int32 NumPendingMipLevels = 0;
+	int32 NextPendingMipLevelIndex = 0;
+	uint32 NextUpdateIndex = 1;
+
+	// Transient lifetime
+	TArray<FStreamingRequest> ParentRequestsToAdd;
+	TSet<FTileDataTexture*> TileDataTexturesToUpdate;
+	TSet<FStreamingInfo*> InvalidatedStreamingInfos;
+	TArray<FStreamingRequest> PrioritizedRequestsHeap;
+	TArray<FMipLevelKey> SelectedMipLevels;
+	TArray<FPageTableClear> PageTableClears;
+	TArray<FUploadTask> UploadTasks; // accessed on the async thread
+	TArray<FPendingMipLevel*> UploadCleanupTasks; // accessed on the async thread
+
+	void AddInternal(FRDGBuilder& GraphBuilder, FNewSparseVolumeTextureInfo&& NewSVTInfo);
+	void RemoveInternal(UStreamableSparseVolumeTexture* SparseVolumeTexture);
+	bool AddRequest(const FStreamingRequest& Request); // Returns true when the request is new or it has been updated with a higher priority
+	void AsyncUpdate();
+	void AddParentRequests(); // Add requests for all parent mip levels
+	void SelectHighestPriorityRequestsAndUpdateLRU(int32 MaxSelectedMipLevels);
+	void IssueRequests(int32 MaxSelectedMipLevels);
+	void StreamOutMipLevel(FStreamingInfo& SVTInfo, FLRUNode* LRUNode);
+	int32 DetermineReadyMipLevels();
+	void InstallReadyMipLevels();
+
+#if WITH_EDITORONLY_DATA
+	UE::DerivedData::FCacheGetChunkRequest BuildDDCRequest(const FResources& Resources, const FMipLevelStreamingInfo& MipLevelStreamingInfo, const uint32 PendingMipLevelIndex);
+	void RequestDDCData(TConstArrayView<UE::DerivedData::FCacheGetChunkRequest> DDCRequests);
+#endif
+};
+
+}
+}
diff --git a/Engine/Source/Runtime/Engine/Public/SparseVolumeTexture/ISparseVolumeTextureStreamingManager.h b/Engine/Source/Runtime/Engine/Public/SparseVolumeTexture/ISparseVolumeTextureStreamingManager.h
new file mode 100644
--- /dev/null
+++ b/Engine/Source/Runtime/Engine/Public/SparseVolumeTexture/ISparseVolumeTextureStreamingManager.h
@@ -0,0 +1,40 @@
+// Copyright Epic Games, Inc. All Rights Reserved.
+
+#pragma once
+
+#include "CoreMinimal.h"
+#include "Engine/EngineTypes.h"
+
+class UStreamableSparseVolumeTexture;
+class FRDGBuilder;
+
+namespace UE
+{
+namespace SVT
+{
+
+// Interface for the SparseVolumeTexture streaming manager
+class ENGINE_API IStreamingManager
+{
+public:
+	//~ Begin game thread functions.
+	virtual void Add_GameThread(UStreamableSparseVolumeTexture* SparseVolumeTexture) = 0;
+	virtual void Remove_GameThread(UStreamableSparseVolumeTexture* SparseVolumeTexture) = 0;
+	// Request a frame to be streamed in. FrameIndex is of float type so that the fractional part can be used to better track the playback speed/direction.
+	// This function automatically also requests all higher mip levels and adds prefetch requests for upcoming frames.
+	virtual void Request_GameThread(UStreamableSparseVolumeTexture* SparseVolumeTexture, float FrameIndex, int32 MipLevel) = 0;
+	//~ End game thread functions.
+	
+	//~ Begin rendering thread functions.
+	virtual void Request(UStreamableSparseVolumeTexture* SparseVolumeTexture, float FrameIndex, int32 MipLevel) = 0;
+	virtual void BeginAsyncUpdate(FRDGBuilder& GraphBuilder) = 0;
+	virtual void EndAsyncUpdate(FRDGBuilder& GraphBuilder) = 0;
+	//~ End rendering thread functions.
+
+	virtual ~IStreamingManager() = default;
+};
+
+ENGINE_API IStreamingManager& GetStreamingManager();
+
+}
+}
diff --git a/Engine/Source/Runtime/Engine/Public/SparseVolumeTexture/SparseVolumeTextureStreamingManager.h b/Engine/Source/Runtime/Engine/Public/SparseVolumeTexture/SparseVolumeTextureStreamingManager.h
deleted file mode 100644
diff --git a/Engine/Source/Runtime/Renderer/Private/DeferredShadingRenderer.cpp b/Engine/Source/Runtime/Renderer/Private/DeferredShadingRenderer.cpp
--- a/Engine/Source/Runtime/Renderer/Private/DeferredShadingRenderer.cpp
+++ b/Engine/Source/Runtime/Renderer/Private/DeferredShadingRenderer.cpp
@@ -91,7 +91,7 @@
 #include "LocalHeightFogRendering.h"
 #include "Shadows/ShadowScene.h"
 #include "Lumen/LumenHardwareRayTracingCommon.h"
-#include "SparseVolumeTexture/SparseVolumeTextureStreamingManager.h"
+#include "SparseVolumeTexture/ISparseVolumeTextureStreamingManager.h"
 
 #if !UE_BUILD_SHIPPING
 #include "RenderCaptureInterface.h"
@@ -2648,7 +2648,7 @@ void FDeferredShadingSceneRenderer::Render(FRDGBuilder& GraphBuilder)
 #endif
 	}
 
-	UE::SVT::GStreamingManager.BeginAsyncUpdate(GraphBuilder);
+	UE::SVT::GetStreamingManager().BeginAsyncUpdate(GraphBuilder);
 
 	bool bUpdateNaniteStreaming = false;
 	bool bVisualizeNanite = false;
@@ -3013,7 +3013,7 @@ void FDeferredShadingSceneRenderer::Render(FRDGBuilder& GraphBuilder)
 		}
 	}
 
-	UE::SVT::GStreamingManager.EndAsyncUpdate(GraphBuilder);
+	UE::SVT::GetStreamingManager().EndAsyncUpdate(GraphBuilder);
 
 	FHairStrandsBookmarkParameters& HairStrandsBookmarkParameters = *GraphBuilder.AllocObject<FHairStrandsBookmarkParameters>();
 	if (IsHairStrandsEnabled(EHairStrandsShaderType::All, Scene->GetShaderPlatform()) && RendererOutput == ERendererOutput::FinalSceneColor)
