- Moved initialization of game instanced world partition to happen when the level starts to be added to the world instead of when it is loaded.
- Moved the uninitialization of game instanced world partition to happen when the level starts to be removed from the world instead of when it is done being removed from the level.
- Centralized the initialization/uninitialization logic in WorldPartitionSubsystem
- Fixes the case where an instanced world partition was added to the world, then removed from the world, then re-added to the world without being unloaded first.
- Instanced world partition IsServerStreamingEnabled/IsServerStreamingOutEnabled now forward the logic to their parent world partition (if any) when using project default modes.

diff --git a/Engine/Source/Runtime/Engine/Private/Level.cpp b/Engine/Source/Runtime/Engine/Private/Level.cpp
--- a/Engine/Source/Runtime/Engine/Private/Level.cpp
+++ b/Engine/Source/Runtime/Engine/Private/Level.cpp
@@ -3028,21 +3028,10 @@ void ULevel::OnLevelLoaded()
 			UE_LOG(LogWorldPartition, Log, TEXT("ULevel::OnLevelLoaded(%s)(bIsOwningWorldGameWorld=%d, bIsOwningWorldPartitioned=%d, bIsValidLevelInstance=%d, InitializeForMainWorld=%d, InitializeForEditor=%d, InitializeForGame=%d)"), 
 				*GetTypedOuter<UWorld>()->GetName(), bIsOwningWorldGameWorld ? 1 : 0, bIsOwningWorldPartitioned ? 1 : 0, bIsValidLevelInstance ? 1 : 0, bIsMainWorldLevel ? 1 : 0, bInitializeForEditor ? 1 : 0, bInitializeForGame ? 1 : 0);
 
-			if (bIsMainWorldLevel || bInitializeForEditor || bInitializeForGame)
+			if (bIsMainWorldLevel || bInitializeForEditor)
 			{
 				FTransform Transform = LevelStreaming ? LevelStreaming->LevelTransform : FTransform::Identity;
-
-				// It is allowed for a WorldPartition to already be initialized in the case where a partitioned sub-level 
-				// was streamed-out and is streamed-in again, without a CleanupLevel (GC).
-				if (!WorldPartition->IsInitialized())
-				{
-					WorldPartition->Initialize(OwningWorld, Transform);
-				}
-				else
-				{
-					check(OwningWorld->IsGameWorld());
-					check(WorldPartition->GetInstanceTransform().Equals(Transform));
-				}
+				WorldPartition->Initialize(OwningWorld, Transform);
 			}
 		}
 	}
diff --git a/Engine/Source/Runtime/Engine/Private/WorldPartition/WorldPartition.cpp b/Engine/Source/Runtime/Engine/Private/WorldPartition/WorldPartition.cpp
--- a/Engine/Source/Runtime/Engine/Private/WorldPartition/WorldPartition.cpp
+++ b/Engine/Source/Runtime/Engine/Private/WorldPartition/WorldPartition.cpp
@@ -490,15 +490,36 @@ void UWorldPartition::Initialize(UWorld* InWorld, const FTransform& InTransform)
 		AWorldPartitionReplay::Initialize(World);
 	}
 
-#if WITH_EDITOR
 	const bool bIsGame = IsRunningGame();
 	const bool bIsEditor = !World->IsGameWorld();
 	const bool bIsCooking = IsRunningCookCommandlet();
+	const bool bIsPIEWorldTravel = (World->WorldType == EWorldType::PIE) && !StreamingPolicy;
 	const bool bIsDedicatedServer = IsRunningDedicatedServer();
-	const bool bPIEWorldTravel = (World->WorldType == EWorldType::PIE) && !StreamingPolicy;
 
-	UE_LOG(LogWorldPartition, Log, TEXT("UWorldPartition::Initialize(Asset=%s, IsEditor=%d, bPIEWorldTravel=%d IsGame=%d, IsCooking=%d)"), *OuterWorld->GetName(), bIsEditor ? 1 : 0, bPIEWorldTravel ? 1 : 0, bIsGame ? 1 : 0, bIsCooking ? 1 : 0);
+	UE_LOG(LogWorldPartition, Log, TEXT("UWorldPartition::Initialize : World = %s, World Type = %s, IsMainWorldPartition = %d, Location = %s, Rotation = %s, IsEditor = %d, IsGame = %d, IsPIEWorldTravel = %d, IsCooking = %d"),
+		*OuterWorld->GetName(),
+		LexToString(World->WorldType),
+		IsMainWorldPartition() ? 1 : 0,
+		*InTransform.GetLocation().ToCompactString(),
+		*InTransform.Rotator().ToCompactString(),
+		bIsEditor,
+		bIsGame,
+		bIsPIEWorldTravel,
+		bIsCooking);
+
+	if (World->IsGameWorld())
+	{
+		UE_LOG(LogWorldPartition, Log, TEXT("UWorldPartition::Initialize Context : World NetMode = %s, IsServer = %d, IsDedicatedServer = %d, IsServerStreamingEnabled = %d, IsServerStreamingOutEnabled = %d, IsUsingMakingVisibleTransaction = %d, IsUsingMakingInvisibleTransaction = %d"),
+			*ToString(World->GetNetMode()),
+			IsServer() ? 1 : 0, 
+			bIsDedicatedServer ? 1 : 0, 
+			IsServerStreamingEnabled() ? 1 : 0, 
+			IsServerStreamingOutEnabled() ? 1 : 0, 
+			UseMakingVisibleTransactionRequests() ? 1 : 0, 
+			UseMakingInvisibleTransactionRequests() ? 1 : 0);
+	}
 
+#if WITH_EDITOR
 	if (bEnableStreaming)
 	{
 		bStreamingWasEnabled = true;
@@ -534,7 +555,7 @@ void UWorldPartition::Initialize(UWorld* InWorld, const FTransform& InTransform)
 	RuntimeHash->SetFlags(RF_Transactional);
 
 	TArray<FGuid> ForceLoadedActorGuids;
-	if (bIsEditor || bIsGame || bPIEWorldTravel || bIsDedicatedServer)
+	if (bIsEditor || bIsGame || bIsPIEWorldTravel || bIsDedicatedServer)
 	{
 		UPackage* LevelPackage = OuterWorld->PersistentLevel->GetOutermost();
 
@@ -657,9 +678,9 @@ void UWorldPartition::Initialize(UWorld* InWorld, const FTransform& InTransform)
 #if WITH_EDITOR
 	if (!bIsEditor)
 	{
-		if (bIsGame || bPIEWorldTravel || bIsDedicatedServer)
+		if (bIsGame || bIsPIEWorldTravel || bIsDedicatedServer)
 		{
-			if (bPIEWorldTravel)
+			if (bIsPIEWorldTravel)
 			{
 				check(!bIsPIE);
 				bIsPIE = true;
@@ -891,10 +912,6 @@ void UWorldPartition::RegisterDelegates()
 			FCoreDelegates::OnGetOnScreenMessages.AddUObject(this, &UWorldPartition::GetOnScreenMessages);
 #endif
 		}
-		else
-		{
-			FWorldDelegates::LevelRemovedFromWorld.AddUObject(this, &UWorldPartition::OnLevelRemovedFromWorld);
-		}
 	}
 }
 
@@ -939,20 +956,6 @@ void UWorldPartition::UnregisterDelegates()
 			FCoreDelegates::OnGetOnScreenMessages.RemoveAll(this);
 #endif
 		}
-		else
-		{
-			FWorldDelegates::LevelRemovedFromWorld.RemoveAll(this);
-		}
-	}
-}
-
-void UWorldPartition::OnLevelRemovedFromWorld(ULevel* InLevel, UWorld* InWorld)
-{
-	check(!IsMainWorldPartition());
-	if ((World == InWorld) && (InLevel == GetTypedOuter<UWorld>()->PersistentLevel))
-	{
-		check(!CanStream());
-		Uninitialize();
 	}
 }
 
@@ -1097,30 +1100,41 @@ bool UWorldPartition::IsServerStreamingEnabled() const
 	if (!bCachedIsServerStreamingEnabled.IsSet())
 	{
 		bool bIsEnabled = false;
-		if (ServerStreamingMode == EWorldPartitionServerStreamingMode::ProjectDefault)
+		UWorld* OwningWorld = GetWorld();
+		if (OwningWorld && OwningWorld->IsGameWorld())
 		{
-			switch (UWorldPartition::GlobalEnableServerStreaming)
+			if (ServerStreamingMode == EWorldPartitionServerStreamingMode::ProjectDefault)
 			{
-			case 1:	bIsEnabled = true;
+				UWorldPartition* MainWorldPartition = OwningWorld->GetWorldPartition();
+				if (MainWorldPartition && (this != MainWorldPartition))
+				{
+					bIsEnabled = MainWorldPartition->IsServerStreamingEnabled();
+				}
+				else
+				{
+					switch (UWorldPartition::GlobalEnableServerStreaming)
+					{
+					case 1:	bIsEnabled = true;
 #if WITH_EDITOR
-			case 2: bIsEnabled = bIsPIE;
+					case 2: bIsEnabled = bIsPIE;
 #endif
+					}
+				}
 			}
-		}
-		else
-		{
-			if ((ServerStreamingMode == EWorldPartitionServerStreamingMode::Enabled)
+			else
+			{
+				if ((ServerStreamingMode == EWorldPartitionServerStreamingMode::Enabled)
 #if WITH_EDITOR
-				|| (bIsPIE && (ServerStreamingMode == EWorldPartitionServerStreamingMode::EnabledInPIE))
+					|| (bIsPIE && (ServerStreamingMode == EWorldPartitionServerStreamingMode::EnabledInPIE))
 #endif
-				)
-			{
-				bIsEnabled = true;
+					)
+				{
+					bIsEnabled = true;
+				}
 			}
 		}
 
-		UWorld* OwningWorld = GetWorld();
-		bCachedIsServerStreamingEnabled = (OwningWorld && OwningWorld->IsGameWorld() && bIsEnabled);
+		bCachedIsServerStreamingEnabled = bIsEnabled;
 	}
 	
 	return bCachedIsServerStreamingEnabled.Get(false);
@@ -1131,9 +1145,28 @@ bool UWorldPartition::IsServerStreamingOutEnabled() const
 	// Resolve once (we don't allow changing the state at runtime)
 	if (!bCachedIsServerStreamingOutEnabled.IsSet())
 	{
+		bool bEnableServerStreamingOut = false;
 		UWorld* OwningWorld = GetWorld();
-		const bool bEnableServerStreamingOut = (ServerStreamingOutMode == EWorldPartitionServerStreamingOutMode::ProjectDefault) ? UWorldPartition::bGlobalEnableServerStreamingOut : (ServerStreamingOutMode == EWorldPartitionServerStreamingOutMode::Enabled);
-		bCachedIsServerStreamingOutEnabled = OwningWorld && OwningWorld->IsGameWorld() && IsServerStreamingEnabled() && bEnableServerStreamingOut;
+		if (OwningWorld && OwningWorld->IsGameWorld() && IsServerStreamingEnabled())
+		{
+			if (ServerStreamingMode == EWorldPartitionServerStreamingMode::ProjectDefault)
+			{
+				UWorldPartition* MainWorldPartition = OwningWorld->GetWorldPartition();
+				if (MainWorldPartition && (this != MainWorldPartition))
+				{
+					bEnableServerStreamingOut = MainWorldPartition->IsServerStreamingOutEnabled();
+				}
+				else
+				{
+					bEnableServerStreamingOut = UWorldPartition::bGlobalEnableServerStreamingOut;
+				}
+			}
+			else
+			{
+				bEnableServerStreamingOut = (ServerStreamingOutMode == EWorldPartitionServerStreamingOutMode::Enabled);
+			}
+		}
+		bCachedIsServerStreamingOutEnabled = bEnableServerStreamingOut;
 	}
 
 	return bCachedIsServerStreamingOutEnabled.Get(false);
diff --git a/Engine/Source/Runtime/Engine/Private/WorldPartition/WorldPartitionSubsystem.cpp b/Engine/Source/Runtime/Engine/Private/WorldPartition/WorldPartitionSubsystem.cpp
--- a/Engine/Source/Runtime/Engine/Private/WorldPartition/WorldPartitionSubsystem.cpp
+++ b/Engine/Source/Runtime/Engine/Private/WorldPartition/WorldPartitionSubsystem.cpp
@@ -431,6 +431,8 @@ void UWorldPartitionSubsystem::Initialize(FSubsystemCollectionBase& Collection)
 	GetWorld()->OnWorldPartitionUninitialized().AddUObject(this, &UWorldPartitionSubsystem::OnWorldPartitionUninitialized);
 	if (GetWorld()->IsGameWorld())
 	{
+		FLevelStreamingDelegates::OnLevelBeginMakingVisible.AddUObject(this, &UWorldPartitionSubsystem::OnLevelBeginMakingVisible);
+		FLevelStreamingDelegates::OnLevelBeginMakingInvisible.AddUObject(this, &UWorldPartitionSubsystem::OnLevelBeginMakingInvisible);
 		FLevelStreamingDelegates::OnLevelStreamingTargetStateChanged.AddUObject(this, &UWorldPartitionSubsystem::OnLevelStreamingTargetStateChanged);
 	}
 }
@@ -449,6 +451,8 @@ void UWorldPartitionSubsystem::Deinitialize()
 	GetWorld()->OnWorldPartitionUninitialized().RemoveAll(this);
 	if (GetWorld()->IsGameWorld())
 	{
+		FLevelStreamingDelegates::OnLevelBeginMakingVisible.RemoveAll(this);
+		FLevelStreamingDelegates::OnLevelBeginMakingInvisible.RemoveAll(this);
 		FLevelStreamingDelegates::OnLevelStreamingTargetStateChanged.RemoveAll(this);
 	}
 
@@ -522,21 +526,51 @@ void UWorldPartitionSubsystem::OnWorldPartitionUninitialized(UWorldPartition* In
 	}
 }
 
-void UWorldPartitionSubsystem::OnLevelStreamingTargetStateChanged(UWorld* World, const ULevelStreaming* StreamingLevel, ULevel* LevelIfLoaded, ELevelStreamingState CurrentState, ELevelStreamingTargetState PrevTarget, ELevelStreamingTargetState NewTarget)
+void UWorldPartitionSubsystem::OnLevelBeginMakingVisible(UWorld* InWorld, const ULevelStreaming* InStreamingLevel, ULevel* InLoadedLevel)
 {
-	if (World != GetWorld())
+	if ((InWorld != GetWorld()) || !GetWorld()->IsGameWorld())
+	{
+		return;
+	}
+
+	if (UWorldPartition* WorldPartition = InLoadedLevel->GetWorldPartition())
+	{
+		check(WorldPartition != InWorld->GetWorldPartition());
+		WorldPartition->Initialize(InWorld, InStreamingLevel->LevelTransform);
+	}
+}
+
+void UWorldPartitionSubsystem::OnLevelBeginMakingInvisible(UWorld* InWorld, const ULevelStreaming* InStreamingLevel, ULevel* InLoadedLevel)
+{
+	if ((InWorld != GetWorld()) || !GetWorld()->IsGameWorld())
+	{
+		return;
+	}
+
+	if (UWorldPartition* WorldPartition = InLoadedLevel->GetWorldPartition())
+	{
+		check(WorldPartition != InWorld->GetWorldPartition());
+		check(WorldPartition->GetWorld() == GetWorld());
+		check(InLoadedLevel == WorldPartition->GetTypedOuter<UWorld>()->PersistentLevel);
+		WorldPartition->Uninitialize();
+	}
+}
+
+void UWorldPartitionSubsystem::OnLevelStreamingTargetStateChanged(UWorld* InWorld, const ULevelStreaming* InStreamingLevel, ULevel* InLevelIfLoaded, ELevelStreamingState InCurrentState, ELevelStreamingTargetState InPrevTarget, ELevelStreamingTargetState InNewTarget)
+{
+	if (InWorld != GetWorld())
 	{
 		return;
 	}
 
 	// Make sure when a WorldPartition is LevelStreamed that changing its state to remove it from world will update the target states of its Cells right away.
-	if(LevelIfLoaded && NewTarget != ELevelStreamingTargetState::LoadedVisible)
+	if(InLevelIfLoaded && InNewTarget != ELevelStreamingTargetState::LoadedVisible)
 	{
-		UWorldPartition* WorldPartition = LevelIfLoaded->GetTypedOuter<UWorld>()->GetWorldPartition();
+		UWorldPartition* WorldPartition = InLevelIfLoaded->GetTypedOuter<UWorld>()->GetWorldPartition();
 		if (WorldPartition && WorldPartition->IsInitialized() && !WorldPartition->CanStream())
 		{
 			TRACE_CPUPROFILER_EVENT_SCOPE(UWorldPartitionSubsystem::OnLevelStreamingTargetStateChanged);
-			UWorldPartitionSubsystem::UpdateStreamingStateInternal(World, { WorldPartition });
+			UWorldPartitionSubsystem::UpdateStreamingStateInternal(InWorld, { WorldPartition });
 		}
 	}
 }
diff --git a/Engine/Source/Runtime/Engine/Public/WorldPartition/WorldPartition.h b/Engine/Source/Runtime/Engine/Public/WorldPartition/WorldPartition.h
--- a/Engine/Source/Runtime/Engine/Public/WorldPartition/WorldPartition.h
+++ b/Engine/Source/Runtime/Engine/Public/WorldPartition/WorldPartition.h
@@ -477,7 +477,6 @@ private:
 	static FAutoConsoleVariableRef CVarUseMakingInvisibleTransactionRequests;
 
 	void OnWorldMatchStarting();
-	void OnLevelRemovedFromWorld(ULevel* InLevel, UWorld* InWorld);
 	void OnPostBugItGoCalled(const FVector& Loc, const FRotator& Rot);
 
 	// Delegates registration
diff --git a/Engine/Source/Runtime/Engine/Public/WorldPartition/WorldPartitionSubsystem.h b/Engine/Source/Runtime/Engine/Public/WorldPartition/WorldPartitionSubsystem.h
--- a/Engine/Source/Runtime/Engine/Public/WorldPartition/WorldPartitionSubsystem.h
+++ b/Engine/Source/Runtime/Engine/Public/WorldPartition/WorldPartitionSubsystem.h
@@ -165,7 +165,9 @@ private:
 
 	void OnWorldPartitionInitialized(UWorldPartition* InWorldPartition);
 	void OnWorldPartitionUninitialized(UWorldPartition* InWorldPartition);
-	void OnLevelStreamingTargetStateChanged(UWorld* World, const ULevelStreaming* StreamingLevel, ULevel* LevelIfLoaded, ELevelStreamingState CurrentState, ELevelStreamingTargetState PrevTarget, ELevelStreamingTargetState NewTarget);
+	void OnLevelStreamingTargetStateChanged(UWorld* InWorld, const ULevelStreaming* InStreamingLevel, ULevel* InLevelIfLoaded, ELevelStreamingState InCurrentState, ELevelStreamingTargetState InPrevTarget, ELevelStreamingTargetState InNewTarget);
+	void OnLevelBeginMakingVisible(UWorld* InWorld, const ULevelStreaming* InStreamingLevel, ULevel* InLoadedLevel);
+	void OnLevelBeginMakingInvisible(UWorld* InWorld, const ULevelStreaming* InStreamingLevel, ULevel* InLoadedLevel);
 
 	static void UpdateStreamingStateInternal(const UWorld* InWorld, const TArray<TObjectPtr<UWorldPartition>>& InWorldPartitions);
 
