- Server streaming making visible/invisible transaction fixes
- Added server's client visible level names in update streaming state hash.

diff --git a/Engine/Source/Runtime/Engine/Classes/Engine/LevelStreaming.h b/Engine/Source/Runtime/Engine/Classes/Engine/LevelStreaming.h
--- a/Engine/Source/Runtime/Engine/Classes/Engine/LevelStreaming.h
+++ b/Engine/Source/Runtime/Engine/Classes/Engine/LevelStreaming.h
@@ -373,7 +373,7 @@ private:
 	/** If true client will wait for acknowledgment from server before making streaming levels visible */
 	bool ShouldClientUseMakingVisibleTransactionRequest() const;
 	/** Returns whether the streaming level should wait for the server ack before changing its visibility. */
-	bool ShouldWaitForServerAckBeforeChangingVisibilityState(ENetLevelVisibilityRequest InRequestType, bool bInShouldBeVisible);
+	bool ShouldWaitForServerAckBeforeChangingVisibilityState(ENetLevelVisibilityRequest InRequestType);
 	/** Ack a client instigated visibility/streaming transaction */
 	void AckNetVisibilityTransaction(FNetLevelVisibilityTransactionId AckedClientTransactionId, bool bClientAckCanMakeVisible);
 
diff --git a/Engine/Source/Runtime/Engine/Private/LevelStreaming.cpp b/Engine/Source/Runtime/Engine/Private/LevelStreaming.cpp
--- a/Engine/Source/Runtime/Engine/Private/LevelStreaming.cpp
+++ b/Engine/Source/Runtime/Engine/Private/LevelStreaming.cpp
@@ -717,13 +717,11 @@ void ULevelStreaming::ServerUpdateLevelVisibility(bool bIsVisible, bool bTryMake
 	}
 }
 
-bool ULevelStreaming::ShouldWaitForServerAckBeforeChangingVisibilityState(ENetLevelVisibilityRequest InRequestType, bool bInShouldBeVisible)
+bool ULevelStreaming::ShouldWaitForServerAckBeforeChangingVisibilityState(ENetLevelVisibilityRequest InRequestType)
 {
 	if (IsWaitingForNetVisibilityTransactionAck(InRequestType))
 	{
-		const bool bTargetShouldBeVisible = (NetVisibilityState.PendingRequestType == ENetLevelVisibilityRequest::MakingVisible);
-		const bool bTargetStateMatches = (bInShouldBeVisible == bTargetShouldBeVisible);
-		if (bTargetStateMatches && NetVisibilityState.bHasClientPendingRequest)
+		if (NetVisibilityState.bHasClientPendingRequest)
 		{
 			// We have a pending request, IncrementTransactionIndex and send ServerUpdateLevelVisibility request to server
 			FNetLevelVisibilityTransactionId TransactionId;
@@ -785,7 +783,7 @@ bool ULevelStreaming::CanMakeInvisible()
 
 	if (ShouldClientUseMakingInvisibleTransactionRequest())
 	{
-		if (ShouldWaitForServerAckBeforeChangingVisibilityState(ENetLevelVisibilityRequest::MakingInvisible, bShouldBeVisible))
+		if (ShouldWaitForServerAckBeforeChangingVisibilityState(ENetLevelVisibilityRequest::MakingInvisible))
 		{
 			return false;
 		}
@@ -808,7 +806,7 @@ bool ULevelStreaming::CanMakeVisible()
 
 	if (ShouldClientUseMakingVisibleTransactionRequest())
 	{
-		if (ShouldWaitForServerAckBeforeChangingVisibilityState(ENetLevelVisibilityRequest::MakingVisible, true))
+		if (ShouldWaitForServerAckBeforeChangingVisibilityState(ENetLevelVisibilityRequest::MakingVisible))
 		{
 			return false;
 		}
@@ -949,10 +947,7 @@ void ULevelStreaming::UpdateStreamingState(bool& bOutUpdateAgain, bool& bOutRede
 			{
 				// Only respond with ServerTransactionId if the is the target visibility state is supposed to be visible
 				FNetLevelVisibilityTransactionId TransactionId;
-				if (bShouldBeVisible)
-				{
-					TransactionId.SetTransactionIndex(NetVisibilityState.ServerRequestIndex);
-				}
+				TransactionId.SetTransactionIndex(NetVisibilityState.ServerRequestIndex);
 
 				// Calling CanMakeVisible will trigger a request for visibility if necessary
 				if (!CanMakeVisible())
@@ -1004,11 +999,7 @@ void ULevelStreaming::UpdateStreamingState(bool& bOutUpdateAgain, bool& bOutRede
 			}
 
 			FNetLevelVisibilityTransactionId TransactionId;
-			// Only respond with ServerTransactionId if the is the target visibility state is supposed to be not visible
-			if (bShouldBeVisible == false)
-			{
-				TransactionId.SetTransactionIndex(NetVisibilityState.ServerRequestIndex);
-			}
+			TransactionId.SetTransactionIndex(NetVisibilityState.ServerRequestIndex);
 
 			// Hide loaded level, incrementally if necessary
 			World->RemoveFromWorld(LoadedLevel, !ShouldBlockOnUnload() && World->IsGameWorld(), TransactionId, this);
@@ -1068,7 +1059,7 @@ void ULevelStreaming::UpdateStreamingState(bool& bOutUpdateAgain, bool& bOutRede
 			SetCurrentState(ELevelStreamingState::MakingVisible);
 			// Make sure client pending visibility request (if any) matches current state
 			NetVisibilityState.InvalidateClientPendingRequest();
-			BeginClientNetVisibilityRequest(ShouldBeVisible());
+			BeginClientNetVisibilityRequest(true);
 			bOutUpdateAgain = true;
 			break;
 
@@ -1116,7 +1107,7 @@ void ULevelStreaming::UpdateStreamingState(bool& bOutUpdateAgain, bool& bOutRede
 			SetCurrentState(ELevelStreamingState::MakingInvisible);
 			// Make sure client pending visibility request (if any) matches current state
 			NetVisibilityState.InvalidateClientPendingRequest();
-			BeginClientNetVisibilityRequest(ShouldBeVisible());
+			BeginClientNetVisibilityRequest(false);
 			bOutUpdateAgain = true;
 			break;
 
diff --git a/Engine/Source/Runtime/Engine/Private/LevelUtils.cpp b/Engine/Source/Runtime/Engine/Private/LevelUtils.cpp
--- a/Engine/Source/Runtime/Engine/Private/LevelUtils.cpp
+++ b/Engine/Source/Runtime/Engine/Private/LevelUtils.cpp
@@ -114,12 +114,20 @@ bool FLevelUtils::IsServerStreamingLevelVisible(UWorld* InWorld, const FName& In
 		return true;
 	}
 
-	if (const AServerStreamingLevelsVisibility* ServerStreamingLevelsVisibility = InWorld->GetServerStreamingLevelsVisibility())
+	const AServerStreamingLevelsVisibility* ServerStreamingLevelsVisibility = InWorld->GetServerStreamingLevelsVisibility();
+	return ServerStreamingLevelsVisibility && ServerStreamingLevelsVisibility->Contains(InPackageName);
+}
+
+ULevelStreaming* FLevelUtils::GetServerVisibleStreamingLevel(UWorld* InWorld, const FName& InPackageName)
+{
+	if (SupportsMakingVisibleTransactionRequests(InWorld))
 	{
-		return ServerStreamingLevelsVisibility && ServerStreamingLevelsVisibility->Contains(InPackageName);
+		if (const AServerStreamingLevelsVisibility* ServerStreamingLevelsVisibility = InWorld->GetServerStreamingLevelsVisibility())
+		{
+			return ServerStreamingLevelsVisibility->GetVisibleStreamingLevel(InPackageName);
+		}
 	}
-
-	return false;
+	return nullptr;
 }
 
 /////////////////////////////////////////////////////////////////////////////////////////
diff --git a/Engine/Source/Runtime/Engine/Private/NetConnection.cpp b/Engine/Source/Runtime/Engine/Private/NetConnection.cpp
--- a/Engine/Source/Runtime/Engine/Private/NetConnection.cpp
+++ b/Engine/Source/Runtime/Engine/Private/NetConnection.cpp
@@ -1662,7 +1662,9 @@ void UNetConnection::UpdateLevelVisibilityInternal(const FUpdateLevelVisibilityL
 			const FNetLevelVisibilityTransactionId VisibilityRequestId = LevelVisibility.VisibilityRequestId;
 			check(VisibilityRequestId.IsValid() && VisibilityRequestId.IsClientTransaction());
 
-			if (FLevelUtils::IsServerStreamingLevelVisible(GetWorld(), LevelVisibility.PackageName))
+			// Only consider visible levels with their streaming level in the LoadedVisible state and returning ShouldBeVisible()
+			ULevelStreaming* ServerVisibleStreamingLevel = FLevelUtils::GetServerVisibleStreamingLevel(GetWorld(), LevelVisibility.PackageName);
+			if (ServerVisibleStreamingLevel && ServerVisibleStreamingLevel->ShouldBeVisible() && (ServerVisibleStreamingLevel->GetLevelStreamingState() == ELevelStreamingState::LoadedVisible))
 			{
 				ClientMakingVisibleLevelNames.Add(LevelVisibility.PackageName);
 				check(!ClientVisibleLevelNames.Contains(LevelVisibility.PackageName));
diff --git a/Engine/Source/Runtime/Engine/Private/Streaming/ServerStreamingLevelsVisibility.cpp b/Engine/Source/Runtime/Engine/Private/Streaming/ServerStreamingLevelsVisibility.cpp
--- a/Engine/Source/Runtime/Engine/Private/Streaming/ServerStreamingLevelsVisibility.cpp
+++ b/Engine/Source/Runtime/Engine/Private/Streaming/ServerStreamingLevelsVisibility.cpp
@@ -1,7 +1,9 @@
 // Copyright Epic Games, Inc. All Rights Reserved.
 
 #include "Streaming/ServerStreamingLevelsVisibility.h"
+#include "UObject/Package.h"
 #include "Engine/LevelStreaming.h"
+#include "Engine/Level.h"
 #include "Engine/World.h"
 
 #include UE_INLINE_GENERATED_CPP_BY_NAME(ServerStreamingLevelsVisibility)
@@ -29,7 +31,7 @@ AServerStreamingLevelsVisibility* AServerStreamingLevelsVisibility::SpawnServerA
 			{
 				if (StreamingLevel && StreamingLevel->IsLevelVisible())
 				{
-					ServerStreamingLevelVisibility->SetIsVisible(StreamingLevel->GetWorldAssetPackageFName(), true);
+					ServerStreamingLevelVisibility->SetIsVisible(StreamingLevel, true);
 				}
 			}
 			return ServerStreamingLevelVisibility;
@@ -40,20 +42,30 @@ AServerStreamingLevelsVisibility* AServerStreamingLevelsVisibility::SpawnServerA
 
 bool AServerStreamingLevelsVisibility::Contains(const FName& InLevelPackageName) const
 {
-	return ServerVisibleLevelNames.Contains(InLevelPackageName);
+	return ServerVisibleStreamingLevels.Contains(InLevelPackageName);
 }
 
-void AServerStreamingLevelsVisibility::SetIsVisible(const FName& InLevelPackageName, bool bInIsVisible)
+ULevelStreaming* AServerStreamingLevelsVisibility::GetVisibleStreamingLevel(const FName& InPackageName) const
+{
+	const TWeakObjectPtr<ULevelStreaming>* StreamingLevel = ServerVisibleStreamingLevels.Find(InPackageName);
+	return StreamingLevel ? StreamingLevel->Get() : nullptr;
+}
+
+void AServerStreamingLevelsVisibility::SetIsVisible(ULevelStreaming* InStreamingLevel, bool bInIsVisible)
 {
 	if (ensure(GetLocalRole() == ROLE_Authority))
 	{
-		if (bInIsVisible)
-		{
-			ServerVisibleLevelNames.Add(InLevelPackageName);
-		}
-		else
+		if (InStreamingLevel)
 		{
-			ServerVisibleLevelNames.Remove(InLevelPackageName);
+			const FName PackageName = InStreamingLevel->GetWorldAssetPackageFName();
+			if (bInIsVisible)
+			{
+				ServerVisibleStreamingLevels.Add(PackageName, InStreamingLevel);
+			}
+			else
+			{
+				ServerVisibleStreamingLevels.Remove(PackageName);
+			}
 		}
 	}
 }
diff --git a/Engine/Source/Runtime/Engine/Private/World.cpp b/Engine/Source/Runtime/Engine/Private/World.cpp
--- a/Engine/Source/Runtime/Engine/Private/World.cpp
+++ b/Engine/Source/Runtime/Engine/Private/World.cpp
@@ -3065,7 +3065,7 @@ const AServerStreamingLevelsVisibility* UWorld::GetServerStreamingLevelsVisibili
 	return ServerStreamingLevelsVisibility;
 }
 
-void UWorld::AddToWorld( ULevel* Level, const FTransform& LevelTransform, bool bConsiderTimeLimit, FNetLevelVisibilityTransactionId TransactionId, ULevelStreaming* OwningLevelStreaming)
+void UWorld::AddToWorld( ULevel* Level, const FTransform& LevelTransform, bool bConsiderTimeLimit, FNetLevelVisibilityTransactionId TransactionId, ULevelStreaming* InOwningLevelStreaming)
 {
 	TRACE_CPUPROFILER_EVENT_SCOPE(AddToWorld);
 	SCOPE_CYCLE_COUNTER(STAT_AddToWorldTime);
@@ -3076,6 +3076,9 @@ void UWorld::AddToWorld( ULevel* Level, const FTransform& LevelTransform, bool b
 
 	FScopeCycleCounterUObject ContextScope(Level);
 
+	// If not provided, find the owning streaming level
+	ULevelStreaming* OwningLevelStreaming = InOwningLevelStreaming ? InOwningLevelStreaming : FLevelUtils::FindStreamingLevel(Level);
+
 	// Set flags to indicate that we are associating a level with the world to e.g. perform slower/ better octree insertion 
 	// and such, as opposed to the fast path taken for run-time/ gameplay objects.
 	Level->bIsAssociatingLevel = true;
@@ -3371,7 +3374,7 @@ void UWorld::AddToWorld( ULevel* Level, const FTransform& LevelTransform, bool b
 			if (IsValid(ServerStreamingLevelsVisibility))
 			{
 				check(IsNetMode(NM_ListenServer) || IsNetMode(NM_DedicatedServer));
-				ServerStreamingLevelsVisibility->SetIsVisible(UnmappedPackageName, true);
+				ServerStreamingLevelsVisibility->SetIsVisible(OwningLevelStreaming, true);
 			}
 		}
 
@@ -3459,7 +3462,7 @@ void UWorld::BeginTearingDown()
 // Cumulated time doing IncrementalUnregisterComponents in UWorld::RemoveFromWorld since last call to UWorld::UpdateLevelStreaming.
 static double GRemoveFromWorldUnregisterComponentTimeCumul = 0.0;
 
-void UWorld::RemoveFromWorld( ULevel* Level, bool bAllowIncrementalRemoval, FNetLevelVisibilityTransactionId TransactionId, ULevelStreaming* OwningLevelStreaming)
+void UWorld::RemoveFromWorld( ULevel* Level, bool bAllowIncrementalRemoval, FNetLevelVisibilityTransactionId TransactionId, ULevelStreaming* InOwningLevelStreaming)
 {
 	TRACE_CPUPROFILER_EVENT_SCOPE(RemoveFromWorld);
 	SCOPE_CYCLE_COUNTER(STAT_RemoveFromWorldTime);
@@ -3473,7 +3476,10 @@ void UWorld::RemoveFromWorld( ULevel* Level, bool bAllowIncrementalRemoval, FNet
 
 	Level->bIsDisassociatingLevel = true;
 
-	auto BeginRemoval = [Level, this, OwningLevelStreaming]()
+	// If not provided, find the owning streaming level
+	ULevelStreaming* OwningLevelStreaming = InOwningLevelStreaming ? InOwningLevelStreaming : FLevelUtils::FindStreamingLevel(Level);
+
+	auto BeginRemoval = [Level, bIsGameWorld, this, OwningLevelStreaming]()
 	{
 		FWorldDelegates::PreLevelRemovedFromWorld.Broadcast(Level, this);
 		Level->bIsBeingRemoved = true;
@@ -3488,6 +3494,16 @@ void UWorld::RemoveFromWorld( ULevel* Level, bool bAllowIncrementalRemoval, FNet
 			Level->bIsVisible = false;
 			ULevelStreaming::BroadcastLevelVisibleStatus(this, Level->GetOutermost()->GetFName(), false);
 		}
+
+		// Set server level as not visible right away so that clients trying to make visible
+		// will fail for this level as it is currently removed from world on the server
+		if (bIsGameWorld && !Level->bClientOnlyVisible && IsValid(ServerStreamingLevelsVisibility))
+		{
+			// Update server's visible levels
+			check(IsNetMode(NM_ListenServer) || IsNetMode(NM_DedicatedServer));
+			check(OwningLevelStreaming || !ServerStreamingLevelsVisibility->Contains(Level->GetPackage()->GetFName()));
+			ServerStreamingLevelsVisibility->SetIsVisible(OwningLevelStreaming, false);
+		}
 	};
 
 	// To be removed from the world a world must be visible and not pending being made visible (this may be redundent, but for safety)
@@ -3608,12 +3624,6 @@ void UWorld::RemoveFromWorld( ULevel* Level, bool bAllowIncrementalRemoval, FNet
 						LocalPlayerController->ServerUpdateLevelVisibility(LevelVisibility);
 					}
 				}
-				// update server's visible levels
-				if (IsValid(ServerStreamingLevelsVisibility))
-				{
-					check(IsNetMode(NM_ListenServer) || IsNetMode(NM_DedicatedServer));
-					ServerStreamingLevelsVisibility->SetIsVisible(UnmappedPackageName, false);
-				}
 			}
 
 			// We expect level that use the bRequireFullVisibilityToRender flag to already be !visible at this point
diff --git a/Engine/Source/Runtime/Engine/Private/WorldPartition/WorldPartitionStreamingPolicy.cpp b/Engine/Source/Runtime/Engine/Private/WorldPartition/WorldPartitionStreamingPolicy.cpp
--- a/Engine/Source/Runtime/Engine/Private/WorldPartition/WorldPartitionStreamingPolicy.cpp
+++ b/Engine/Source/Runtime/Engine/Private/WorldPartition/WorldPartitionStreamingPolicy.cpp
@@ -170,6 +170,14 @@ uint32 UWorldPartitionStreamingPolicy::ComputeUpdateStreamingHash(bool bCanOptim
 		{
 			HashBuilder << Source.GetHash(bIsStreaming3D);
 		}
+
+		if (WorldPartition->IsServer())
+		{
+			TArray<FName> ClientsVisibleLevelNames = ServerClientsVisibleLevelNames.Array();
+			ClientsVisibleLevelNames.Sort(FNameFastLess());
+			HashBuilder << ClientsVisibleLevelNames;
+		}
+
 		return HashBuilder.GetHash();
 	}
 
@@ -275,6 +283,7 @@ void UWorldPartitionStreamingPolicy::UpdateStreamingState()
 		CriticalPerformanceBlockTillLevelStreamingCompletedEpoch = World->GetBlockTillLevelStreamingCompletedEpoch();
 	}
 
+	const bool bIsServer = WorldPartition->IsServer();
 	const bool bCanStream = WorldPartition->CanStream();
 	const bool bForceFrameUpdate = (UWorldPartitionStreamingPolicy::ForceUpdateFrameCount > 0) ? ((UpdateStreamingStateCalls % UWorldPartitionStreamingPolicy::ForceUpdateFrameCount) == 0) : false;
 	const bool bCanOptimizeUpdate =
@@ -289,6 +298,21 @@ void UWorldPartitionStreamingPolicy::UpdateStreamingState()
 	// Update streaming sources
 	UpdateStreamingSources(bCanOptimizeUpdate);
 
+	if (bIsServer)
+	{
+		// Gather Client visible level names
+		ServerClientsVisibleLevelNames.Reset();
+		if (const UNetDriver* NetDriver = GetWorld()->GetNetDriver())
+		{
+			for (UNetConnection* Connection : NetDriver->ClientConnections)
+			{
+				ServerClientsVisibleLevelNames.Add(Connection->GetClientWorldPackageName());
+				ServerClientsVisibleLevelNames.Append(Connection->ClientVisibleLevelNames);
+				ServerClientsVisibleLevelNames.Append(Connection->GetClientMakingVisibleLevelNames());
+			}
+		}
+	}
+
 	// Detect if nothing relevant changed and early out
 	const uint32 NewUpdateStreamingHash = ComputeUpdateStreamingHash(bCanOptimizeUpdate);
 	const bool bShouldSkipUpdate = NewUpdateStreamingHash && (UpdateStreamingHash == NewUpdateStreamingHash);
@@ -300,7 +324,6 @@ void UWorldPartitionStreamingPolicy::UpdateStreamingState()
 	// Update new streaming sources hash
 	UpdateStreamingHash = NewUpdateStreamingHash;
 
-	TSet<FName> ClientVisibleLevelNames;
 	bool bUpdateEpoch = false;
 
 	check(FrameActivateCells.IsEmpty());
@@ -313,7 +336,6 @@ void UWorldPartitionStreamingPolicy::UpdateStreamingState()
 		FrameLoadCells.Reset();
 	};
 
-	const bool bIsServer = WorldPartition->IsServer();
 	const bool bIsServerStreamingEnabled = WorldPartition->IsServerStreamingEnabled();
 	const int32 NewServerStreamingEnabledEpoch = ComputeServerStreamingEnabledEpoch();
 
@@ -369,17 +391,6 @@ void UWorldPartitionStreamingPolicy::UpdateStreamingState()
 
 			bUpdateEpoch = true;
 
-			// Gather Client visible level names
-			if (const UNetDriver* NetDriver = World->GetNetDriver())
-			{
-				for (UNetConnection* Connection : NetDriver->ClientConnections)
-				{
-					ClientVisibleLevelNames.Add(Connection->GetClientWorldPackageName());
-					ClientVisibleLevelNames.Append(Connection->ClientVisibleLevelNames);
-					ClientVisibleLevelNames.Append(Connection->GetClientMakingVisibleLevelNames());
-				}
-			}
-
 			const UDataLayerManager* DataLayerManager = WorldPartition->GetDataLayerManager();
 			TSet<FName> EffectiveActiveDataLayerNames = DataLayerManager->GetEffectiveActiveDataLayerNames();
 			TSet<FName> EffectiveLoadedDataLayerNames = DataLayerManager->GetEffectiveLoadedDataLayerNames();
@@ -457,12 +468,12 @@ void UWorldPartitionStreamingPolicy::UpdateStreamingState()
 		}
 	}
 
-	auto ShouldWaitForClientVisibility = [bIsServer, &ClientVisibleLevelNames, &bUpdateEpoch](const UWorldPartitionRuntimeCell* Cell)
+	auto ShouldWaitForClientVisibility = [bIsServer, this, &bUpdateEpoch](const UWorldPartitionRuntimeCell* Cell)
 	{
 		check(bIsServer);
 		if (ULevel* Level = Cell->GetLevel())
 		{
-			if (ClientVisibleLevelNames.Contains(Cell->GetLevel()->GetPackage()->GetFName()))
+			if (ServerClientsVisibleLevelNames.Contains(Cell->GetLevel()->GetPackage()->GetFName()))
 			{
 				UE_CLOG(bUpdateEpoch, LogWorldPartition, Verbose, TEXT("Server epoch update delayed by client visibility"));
 				bUpdateEpoch = false;
diff --git a/Engine/Source/Runtime/Engine/Public/LevelUtils.h b/Engine/Source/Runtime/Engine/Public/LevelUtils.h
--- a/Engine/Source/Runtime/Engine/Public/LevelUtils.h
+++ b/Engine/Source/Runtime/Engine/Public/LevelUtils.h
@@ -63,6 +63,15 @@ public:
 	 */
 	static bool IsServerStreamingLevelVisible(UWorld* InWorld, const FName& InPackageName);
 
+	/**
+	 * Returns the streaming level by package name if visible on server, or NULL if none exists.
+	 *
+	 * @param		InWorld			World to look in for the streaming level
+	 * @param		PackageName		Name of the package containing the ULevel to query
+	 * @return						The level's streaming level if visible on server, or NULL if none exists.
+	 */
+	static ULevelStreaming* GetServerVisibleStreamingLevel(UWorld* InWorld, const FName& InPackageName);
+
 	/** Returns whether the world supports for a client to use "making visible" transaction requests to the server. */
 	static bool SupportsMakingVisibleTransactionRequests(UWorld* InWorld);
 
diff --git a/Engine/Source/Runtime/Engine/Public/Streaming/ServerStreamingLevelsVisibility.h b/Engine/Source/Runtime/Engine/Public/Streaming/ServerStreamingLevelsVisibility.h
--- a/Engine/Source/Runtime/Engine/Public/Streaming/ServerStreamingLevelsVisibility.h
+++ b/Engine/Source/Runtime/Engine/Public/Streaming/ServerStreamingLevelsVisibility.h
@@ -7,6 +7,8 @@
 #include "GameFramework/Actor.h"
 #include "ServerStreamingLevelsVisibility.generated.h"
 
+class ULevelStreaming;
+
 /**
  * Actor used to replicate server's visible level streaming
  */
@@ -18,8 +20,8 @@ class ENGINE_API AServerStreamingLevelsVisibility : public AActor
 public:
 	static AServerStreamingLevelsVisibility* SpawnServerActor(UWorld* World);
 	bool Contains(const FName& InPackageName) const;
-	void SetIsVisible(const FName& InLevelPackageName, bool bInIsVisible);
-
+	void SetIsVisible(ULevelStreaming* InStreamingLevel, bool bInIsVisible);
+	ULevelStreaming* GetVisibleStreamingLevel(const FName& InPackageName) const;
 private:
-	TSet<FName> ServerVisibleLevelNames;
+	TMap<FName, TWeakObjectPtr<ULevelStreaming>> ServerVisibleStreamingLevels;
 };
\ No newline at end of file
diff --git a/Engine/Source/Runtime/Engine/Public/WorldPartition/WorldPartitionStreamingPolicy.h b/Engine/Source/Runtime/Engine/Public/WorldPartition/WorldPartitionStreamingPolicy.h
--- a/Engine/Source/Runtime/Engine/Public/WorldPartition/WorldPartitionStreamingPolicy.h
+++ b/Engine/Source/Runtime/Engine/Public/WorldPartition/WorldPartitionStreamingPolicy.h
@@ -132,6 +132,7 @@ private:
 	int32 ServerDataLayersStatesEpoch;
 	int32 ServerStreamingStateEpoch;
 	int32 ServerStreamingEnabledEpoch;
+	TSet<FName> ServerClientsVisibleLevelNames;
 	uint32 UpdateStreamingHash;
 	uint32 UpdateStreamingSourcesHash;
 	uint32 UpdateStreamingStateCalls;
