Fix assert from AWorldPartitionHLOD::PreRegisterAllComponents() if WP isn't yet initialized

diff --git a/Engine/Source/Runtime/Engine/Private/WorldPartition/HLOD/HLODActor.cpp b/Engine/Source/Runtime/Engine/Private/WorldPartition/HLOD/HLODActor.cpp
--- a/Engine/Source/Runtime/Engine/Private/WorldPartition/HLOD/HLODActor.cpp
+++ b/Engine/Source/Runtime/Engine/Private/WorldPartition/HLOD/HLODActor.cpp
@@ -246,7 +246,8 @@ void AWorldPartitionHLOD::PreRegisterAllComponents()
 	// If world is instanced, we need to recompute our bounds since they are in the instanced-world space
 	if (UWorldPartition* WorldPartition = FWorldPartitionHelpers::GetWorldPartition(this))
 	{
-		if (!WorldPartition->IsMainWorldPartition())
+		const bool bIsInstancedLevel = WorldPartition->GetTypedOuter<ULevel>()->IsInstancedLevel();
+		if (bIsInstancedLevel)
 		{
 			ForEachComponent<USceneComponent>(false, [](USceneComponent* SceneComponent)
 			{
