Fixes FPostProcessSettings::DepthOfFieldFocalDistance interpolation

diff --git a/Engine/Source/Runtime/Engine/Private/SceneView.cpp b/Engine/Source/Runtime/Engine/Private/SceneView.cpp
--- a/Engine/Source/Runtime/Engine/Private/SceneView.cpp
+++ b/Engine/Source/Runtime/Engine/Private/SceneView.cpp
@@ -1569,7 +1569,18 @@ void FSceneView::OverridePostProcessSettings(const FPostProcessSettings& Src, fl
 		LERP_PP(AmbientOcclusionTemporalBlendWeight);
 		LERP_PP(IndirectLightingColor);
 		LERP_PP(IndirectLightingIntensity);
-		LERP_PP(DepthOfFieldFocalDistance);
+
+		if (Src.bOverride_DepthOfFieldFocalDistance)
+		{
+			if (Dest.DepthOfFieldFocalDistance == 0.0f || Src.DepthOfFieldFocalDistance == 0.0f)
+			{
+				Dest.DepthOfFieldFocalDistance = Src.DepthOfFieldFocalDistance;
+			}
+			else
+			{
+				Dest.DepthOfFieldFocalDistance = FMath::Lerp(Dest.DepthOfFieldFocalDistance, Src.DepthOfFieldFocalDistance, Weight);
+			}
+		}
 		LERP_PP(DepthOfFieldFstop);
 		LERP_PP(DepthOfFieldMinFstop);
 		LERP_PP(DepthOfFieldSensorWidth);
