Fix for trying to stream in zero sized mip level data in the SVT streaming manager.

#rnx

diff --git a/Engine/Source/Runtime/Engine/Private/SparseVolumeTexture/SparseVolumeTextureStreamingManager.cpp b/Engine/Source/Runtime/Engine/Private/SparseVolumeTexture/SparseVolumeTextureStreamingManager.cpp
--- a/Engine/Source/Runtime/Engine/Private/SparseVolumeTexture/SparseVolumeTextureStreamingManager.cpp
+++ b/Engine/Source/Runtime/Engine/Private/SparseVolumeTexture/SparseVolumeTextureStreamingManager.cpp
@@ -88,7 +88,7 @@ static bool DoesPlatformSupportSparseVolumeTexture(EShaderPlatform Platform)
 	return AllowGlobalShaderLoad();
 }
 
-static FIntVector3 ComputeTileDataVolumeResolution2(int32 NumAllocatedPages)
+static FIntVector3 ComputeTileDataVolumeResolution(int32 NumAllocatedPages)
 {
 	int32 TileVolumeResolutionCube = 1;
 	while (TileVolumeResolutionCube * TileVolumeResolutionCube * TileVolumeResolutionCube < NumAllocatedPages)
@@ -1106,8 +1106,8 @@ void FStreamingManager::AddInternal(FRDGBuilder& GraphBuilder, FNewSparseVolumeT
 	// Create RHI resources and upload root tile data
 	{
 		const int32 TileFactor = NumFrames <= 1 ? 1 : 3;
-		const int32 NumPhysicalTilesCapacity = NumRootPhysicalTiles + (TileFactor * MaxNumPhysicalTiles);
-		const FIntVector3 TileDataVolumeResolution = ComputeTileDataVolumeResolution2(NumPhysicalTilesCapacity);
+		const int32 NumPhysicalTilesCapacity = FMath::Max(1, NumRootPhysicalTiles + (TileFactor * MaxNumPhysicalTiles)); // Ensure a minimum size of 1
+		const FIntVector3 TileDataVolumeResolution = ComputeTileDataVolumeResolution(NumPhysicalTilesCapacity);
 		const FIntVector3 TileDataVolumeResolutionInTiles = TileDataVolumeResolution / SPARSE_VOLUME_TILE_RES_PADDED;
 
 		SVTInfo.TileDataTexture = new FTileDataTexture(TileDataVolumeResolutionInTiles, SVTInfo.FormatA, SVTInfo.FormatB);
@@ -1330,10 +1330,13 @@ void FStreamingManager::SelectHighestPriorityRequestsAndUpdateLRU(int32 MaxSelec
 			FStreamingInfo* SVTInfo = StreamingInfo.Find(Request.Key.SVT);
 			check(SVTInfo);
 
-			// Discard invalid requests
+			// Discard invalid requests: frame index out of bounds, mip level index out of bounds (or root mip level) and mip levels without any data.
+			// There can never be lower mip levels with data depending on higher mip levels without any data, so discarding such requests is ok.
 			if (Request.Key.FrameIndex < 0
 				|| Request.Key.FrameIndex >= SVTInfo->PerFrameInfo.Num()
-				|| Request.Key.MipLevelIndex >= (SVTInfo->PerFrameInfo[Request.Key.FrameIndex].NumMipLevels - 1))
+				|| Request.Key.MipLevelIndex < 0
+				|| Request.Key.MipLevelIndex >= (SVTInfo->PerFrameInfo[Request.Key.FrameIndex].NumMipLevels - 1)
+				|| SVTInfo->PerFrameInfo[Request.Key.FrameIndex].Resources->MipLevelStreamingInfo[Request.Key.MipLevelIndex].BulkSize == 0)
 			{
 				continue;
 			}
