Fix for SVT::FTextureData::CreateFromDense() not accepting input data > 2GB.

#rnx

diff --git a/Engine/Plugins/FX/Niagara/Source/Niagara/Private/NiagaraDataInterfaceRenderTargetVolume.cpp b/Engine/Plugins/FX/Niagara/Source/Niagara/Private/NiagaraDataInterfaceRenderTargetVolume.cpp
--- a/Engine/Plugins/FX/Niagara/Source/Niagara/Private/NiagaraDataInterfaceRenderTargetVolume.cpp
+++ b/Engine/Plugins/FX/Niagara/Source/Niagara/Private/NiagaraDataInterfaceRenderTargetVolume.cpp
@@ -718,7 +718,7 @@ bool UNiagaraDataInterfaceRenderTargetVolume::SimCacheWriteFrame(UObject* Storag
 				SVTCreateInfo.AttributesFormats[1] = PF_Unknown;
 
 				UE::SVT::FTextureData SparseTextureData{};
-				SparseTextureData.CreateFromDense(SVTCreateInfo, MakeArrayView<uint8>((uint8*)TextureData.GetData(), TextureData.Num() * sizeof(TextureData[0])), TArrayView<uint8>());
+				SparseTextureData.CreateFromDense(SVTCreateInfo, TArrayView<uint8, int64>((uint8*)TextureData.GetData(), (int64)TextureData.Num() * sizeof(TextureData[0])), TArrayView<uint8>());
 				
 				CurrCache->AppendFrame(SparseTextureData);
 #endif
diff --git a/Engine/Source/Runtime/Engine/Private/SparseVolumeTexture/SparseVolumeTextureData.cpp b/Engine/Source/Runtime/Engine/Private/SparseVolumeTexture/SparseVolumeTextureData.cpp
--- a/Engine/Source/Runtime/Engine/Private/SparseVolumeTexture/SparseVolumeTextureData.cpp
+++ b/Engine/Source/Runtime/Engine/Private/SparseVolumeTexture/SparseVolumeTextureData.cpp
@@ -56,14 +56,14 @@ bool FTextureData::Create(const ITextureDataProvider& DataProvider)
 {
 	const FTextureDataCreateInfo CreateInfo = DataProvider.GetCreateInfo();
 
-	Header = UE::SVT::FHeader(CreateInfo.VirtualVolumeAABBMin, CreateInfo.VirtualVolumeAABBMax, CreateInfo.AttributesFormats[0], CreateInfo.AttributesFormats[1], CreateInfo.FallbackValues[0], CreateInfo.FallbackValues[1]);
+	Header = FHeader(CreateInfo.VirtualVolumeAABBMin, CreateInfo.VirtualVolumeAABBMax, CreateInfo.AttributesFormats[0], CreateInfo.AttributesFormats[1], CreateInfo.FallbackValues[0], CreateInfo.FallbackValues[1]);
 
-	if (Header.PageTableVolumeResolution.X > UE::SVT::SVTMaxVolumeTextureDim
-		|| Header.PageTableVolumeResolution.Y > UE::SVT::SVTMaxVolumeTextureDim
-		|| Header.PageTableVolumeResolution.Z > UE::SVT::SVTMaxVolumeTextureDim)
+	if (Header.PageTableVolumeResolution.X > SVT::MaxVolumeTextureDim
+		|| Header.PageTableVolumeResolution.Y > SVT::MaxVolumeTextureDim
+		|| Header.PageTableVolumeResolution.Z > SVT::MaxVolumeTextureDim)
 	{
 		UE_LOG(LogSparseVolumeTextureData, Warning, TEXT("SparseVolumeTexture page table texture dimensions exceed limit (%ix%ix%i): %ix%ix%i"),
-			UE::SVT::SVTMaxVolumeTextureDim, UE::SVT::SVTMaxVolumeTextureDim, UE::SVT::SVTMaxVolumeTextureDim,
+			SVT::MaxVolumeTextureDim, SVT::MaxVolumeTextureDim, SVT::MaxVolumeTextureDim,
 			Header.PageTableVolumeResolution.X, Header.PageTableVolumeResolution.Y, Header.PageTableVolumeResolution.Z);
 		return false;
 	}
@@ -146,8 +146,8 @@ bool FTextureData::Create(const ITextureDataProvider& DataProvider)
 	// Initialize the SparseVolumeTexture page and tile.
 	const int32 FormatSize[] = { GPixelFormats[Header.AttributesFormats[0]].BlockBytes, GPixelFormats[Header.AttributesFormats[1]].BlockBytes };
 	MipMaps[0].PageTable.SetNumZeroed(Header.PageTableVolumeResolution.X * Header.PageTableVolumeResolution.Y * Header.PageTableVolumeResolution.Z);
-	MipMaps[0].PhysicalTileDataA.SetNum((int64)NumAllocatedPages * (int64)UE::SVT::SVTNumVoxelsPerPaddedTile * (int64)FormatSize[0]);
-	MipMaps[0].PhysicalTileDataB.SetNum((int64)NumAllocatedPages * (int64)UE::SVT::SVTNumVoxelsPerPaddedTile * (int64)FormatSize[1]);
+	MipMaps[0].PhysicalTileDataA.SetNum((int64)NumAllocatedPages * (int64)SVT::NumVoxelsPerPaddedTile * (int64)FormatSize[0]);
+	MipMaps[0].PhysicalTileDataB.SetNum((int64)NumAllocatedPages * (int64)SVT::NumVoxelsPerPaddedTile * (int64)FormatSize[1]);
 	MipMaps[0].NumPhysicalTiles = NumAllocatedPages;
 
 	// Generate page table and tile volume data by splatting the data
@@ -210,16 +210,16 @@ bool FTextureData::Create(const ITextureDataProvider& DataProvider)
 	return true;
 }
 
-bool FTextureData::CreateFromDense(const FTextureDataCreateInfo& CreateInfo, const TArrayView<uint8>& VoxelDataA, const TArrayView<uint8>& VoxelDataB)
+bool FTextureData::CreateFromDense(const FTextureDataCreateInfo& CreateInfo, const TArrayView64<uint8>& VoxelDataA, const TArrayView64<uint8>& VoxelDataB)
 {	
-	Header = UE::SVT::FHeader(CreateInfo.VirtualVolumeAABBMin, CreateInfo.VirtualVolumeAABBMax, CreateInfo.AttributesFormats[0], CreateInfo.AttributesFormats[1], CreateInfo.FallbackValues[0], CreateInfo.FallbackValues[1]);
+	Header = FHeader(CreateInfo.VirtualVolumeAABBMin, CreateInfo.VirtualVolumeAABBMax, CreateInfo.AttributesFormats[0], CreateInfo.AttributesFormats[1], CreateInfo.FallbackValues[0], CreateInfo.FallbackValues[1]);
 
-	if (Header.PageTableVolumeResolution.X > UE::SVT::SVTMaxVolumeTextureDim
-		|| Header.PageTableVolumeResolution.Y > UE::SVT::SVTMaxVolumeTextureDim
-		|| Header.PageTableVolumeResolution.Z > UE::SVT::SVTMaxVolumeTextureDim)
+	if (Header.PageTableVolumeResolution.X > SVT::MaxVolumeTextureDim
+		|| Header.PageTableVolumeResolution.Y > SVT::MaxVolumeTextureDim
+		|| Header.PageTableVolumeResolution.Z > SVT::MaxVolumeTextureDim)
 	{
 		UE_LOG(LogSparseVolumeTextureData, Warning, TEXT("SparseVolumeTexture page table texture dimensions exceed limit (%ix%ix%i): %ix%ix%i"),
-			UE::SVT::SVTMaxVolumeTextureDim, UE::SVT::SVTMaxVolumeTextureDim, UE::SVT::SVTMaxVolumeTextureDim,
+			SVT::MaxVolumeTextureDim, SVT::MaxVolumeTextureDim, SVT::MaxVolumeTextureDim,
 			Header.PageTableVolumeResolution.X, Header.PageTableVolumeResolution.Y, Header.PageTableVolumeResolution.Z);
 		return false;
 	}
@@ -231,8 +231,8 @@ bool FTextureData::CreateFromDense(const FTextureDataCreateInfo& CreateInfo, con
 	{
 		if (CreateInfo.AttributesFormats[i] != PF_Unknown)
 		{
-			UE::SVT::WriteVoxel(0, NullTileValuesU8[i], CreateInfo.AttributesFormats[i], CreateInfo.FallbackValues[i]);
-			FallbackValuesQuantized[i] = UE::SVT::ReadVoxel(0, NullTileValuesU8[i], CreateInfo.AttributesFormats[i]);
+			SVT::WriteVoxel(0, NullTileValuesU8[i], CreateInfo.AttributesFormats[i], CreateInfo.FallbackValues[i]);
+			FallbackValuesQuantized[i] = SVT::ReadVoxel(0, NullTileValuesU8[i], CreateInfo.AttributesFormats[i]);
 		}
 	}
 
@@ -287,7 +287,7 @@ bool FTextureData::CreateFromDense(const FTextureDataCreateInfo& CreateInfo, con
 
 						// Compute flat source and destination voxel index. Make sure to account for the fact that the voxels in the dst buffer have padding/a border.
 						const FInt64Vector3 LocalCoordPadded = FInt64Vector3(GridCoord % SPARSE_VOLUME_TILE_RES) + FInt64Vector3(SPARSE_VOLUME_TILE_BORDER);
-						const int64 DstVoxelIndex = ((int64)(TileIndex - 1) * (int64)UE::SVT::SVTNumVoxelsPerPaddedTile) 
+						const int64 DstVoxelIndex = ((int64)(TileIndex - 1) * (int64)SVT::NumVoxelsPerPaddedTile) 
 							+ (LocalCoordPadded.Z * (SPARSE_VOLUME_TILE_RES_PADDED * SPARSE_VOLUME_TILE_RES_PADDED))
 							+ (LocalCoordPadded.Y * SPARSE_VOLUME_TILE_RES_PADDED)
 							+ LocalCoordPadded.X;
@@ -346,8 +346,8 @@ bool FTextureData::CreateFromDense(const FTextureDataCreateInfo& CreateInfo, con
 #endif // SVT_DETERMINISTIC_PAGE_TABLE_GENERATION
 
 	// Allocate tile data
-	MipMaps[0].PhysicalTileDataA.SetNum((int64)NumAllocatedPages * (int64)UE::SVT::SVTNumVoxelsPerPaddedTile * (int64)FormatSize[0]);
-	MipMaps[0].PhysicalTileDataB.SetNum((int64)NumAllocatedPages * (int64)UE::SVT::SVTNumVoxelsPerPaddedTile * (int64)FormatSize[1]);
+	MipMaps[0].PhysicalTileDataA.SetNum((int64)NumAllocatedPages * (int64)SVT::NumVoxelsPerPaddedTile * (int64)FormatSize[0]);
+	MipMaps[0].PhysicalTileDataB.SetNum((int64)NumAllocatedPages * (int64)SVT::NumVoxelsPerPaddedTile * (int64)FormatSize[1]);
 	MipMaps[0].NumPhysicalTiles = NumAllocatedPages;
 
 	// Clear all physical tiles to the fallback value. We can't guarantee that every voxel will be written to in the next step, so we need to ensure that tiles are always fully written to with valid data.
@@ -356,9 +356,9 @@ bool FTextureData::CreateFromDense(const FTextureDataCreateInfo& CreateInfo, con
 	{
 		ParallelFor(NumJobs, [&](int32 TileIndex)
 			{
-				for (int32 VoxelIndex = 0; VoxelIndex < UE::SVT::SVTNumVoxelsPerPaddedTile; ++VoxelIndex)
+				for (int32 VoxelIndex = 0; VoxelIndex < SVT::NumVoxelsPerPaddedTile; ++VoxelIndex)
 				{
-					const int64 VoxelIndexGlobal = (int64)TileIndex * (int64)UE::SVT::SVTNumVoxelsPerPaddedTile + (int64)VoxelIndex;
+					const int64 VoxelIndexGlobal = (int64)TileIndex * (int64)SVT::NumVoxelsPerPaddedTile + (int64)VoxelIndex;
 					if (FormatSize[0])
 					{
 						FMemory::Memcpy(MipMaps[0].PhysicalTileDataA.GetData() + VoxelIndexGlobal * FormatSize[0], NullTileValuesU8[0], FormatSize[0]);
@@ -380,7 +380,7 @@ bool FTextureData::CreateFromDense(const FTextureDataCreateInfo& CreateInfo, con
 void FTextureData::CreateDefault()
 {
 	// Store a single 1x1x1 mip level with the page table pointing to a (zero-valued) null tile.
-	Header = UE::SVT::FHeader(FIntVector3::ZeroValue, FIntVector3(1, 1, 1), PF_R8, PF_Unknown, FVector4f(), FVector4f());
+	Header = FHeader(FIntVector3::ZeroValue, FIntVector3(1, 1, 1), PF_R8, PF_Unknown, FVector4f(), FVector4f());
 
 	MipMaps.SetNum(1);
 	MipMaps[0].PageTable.SetNumZeroed(1);
@@ -434,10 +434,10 @@ FVector4f FTextureData::ReadTileDataVoxel(int32 TileIndex, const FIntVector3& Ti
 	}
 
 	const FInt64Vector3 CoordPadded = FInt64Vector3(TileDataCoord) + FInt64Vector3(SPARSE_VOLUME_TILE_BORDER);
-	const int64 VoxelIndex = (int64)TileIndex * (int64)UE::SVT::SVTNumVoxelsPerPaddedTile + CoordPadded.Z * (SPARSE_VOLUME_TILE_RES_PADDED * SPARSE_VOLUME_TILE_RES_PADDED) + CoordPadded.Y * SPARSE_VOLUME_TILE_RES_PADDED + CoordPadded.X;
+	const int64 VoxelIndex = (int64)TileIndex * (int64)SVT::NumVoxelsPerPaddedTile + CoordPadded.Z * (SPARSE_VOLUME_TILE_RES_PADDED * SPARSE_VOLUME_TILE_RES_PADDED) + CoordPadded.Y * SPARSE_VOLUME_TILE_RES_PADDED + CoordPadded.X;
 	const uint8* TileData = AttributesIdx == 0 ? MipMaps[MipLevel].PhysicalTileDataA.GetData() : MipMaps[MipLevel].PhysicalTileDataB.GetData();
 	const EPixelFormat Format = Header.AttributesFormats[AttributesIdx];
-	return UE::SVT::ReadVoxel(VoxelIndex, TileData, Format);
+	return SVT::ReadVoxel(VoxelIndex, TileData, Format);
 }
 
 FVector4f FTextureData::Load(const FIntVector3& VolumeCoord, int32 MipLevel, int32 AttributesIdx, const FTextureDataAddressingInfo& AddressingInfo) const
@@ -499,10 +499,10 @@ void FTextureData::WriteTileDataVoxel(int32 TileIndex, const FIntVector3& TileDa
 	}
 
 	const FInt64Vector3 CoordPadded = FInt64Vector3(TileDataCoord) + FInt64Vector3(SPARSE_VOLUME_TILE_BORDER);
-	const int64 VoxelIndex = (int64)TileIndex * (int64)UE::SVT::SVTNumVoxelsPerPaddedTile + CoordPadded.Z * (SPARSE_VOLUME_TILE_RES_PADDED * SPARSE_VOLUME_TILE_RES_PADDED) + CoordPadded.Y * SPARSE_VOLUME_TILE_RES_PADDED + CoordPadded.X;
+	const int64 VoxelIndex = (int64)TileIndex * (int64)SVT::NumVoxelsPerPaddedTile + CoordPadded.Z * (SPARSE_VOLUME_TILE_RES_PADDED * SPARSE_VOLUME_TILE_RES_PADDED) + CoordPadded.Y * SPARSE_VOLUME_TILE_RES_PADDED + CoordPadded.X;
 	uint8* TileData = AttributesIdx == 0 ? MipMaps[MipLevel].PhysicalTileDataA.GetData() : MipMaps[MipLevel].PhysicalTileDataB.GetData();
 	const EPixelFormat Format = Header.AttributesFormats[AttributesIdx];
-	UE::SVT::WriteVoxel(VoxelIndex, TileData, Format, Value, DstComponent);
+	SVT::WriteVoxel(VoxelIndex, TileData, Format, Value, DstComponent);
 }
 
 bool FTextureData::GenerateMipMaps(const FTextureDataAddressingInfo& AddressingInfo, int32 NumMipLevels)
@@ -514,8 +514,8 @@ bool FTextureData::GenerateMipMaps(const FTextureDataAddressingInfo& AddressingI
 		&& FMath::IsPowerOfTwo(Header.PageTableVolumeResolution.Y)
 		&& FMath::IsPowerOfTwo(Header.PageTableVolumeResolution.Z));
 	const int32 NumMip0Tiles = MipMaps[0].NumPhysicalTiles;
-	check(MipMaps[0].PhysicalTileDataA.Num() >= ((int64)NumMip0Tiles * (int64)UE::SVT::SVTNumVoxelsPerPaddedTile * (int64)FormatSize[0]));
-	check(MipMaps[0].PhysicalTileDataB.Num() >= ((int64)NumMip0Tiles * (int64)UE::SVT::SVTNumVoxelsPerPaddedTile * (int64)FormatSize[1]));
+	check(MipMaps[0].PhysicalTileDataA.Num() >= ((int64)NumMip0Tiles * (int64)SVT::NumVoxelsPerPaddedTile * (int64)FormatSize[0]));
+	check(MipMaps[0].PhysicalTileDataB.Num() >= ((int64)NumMip0Tiles * (int64)SVT::NumVoxelsPerPaddedTile * (int64)FormatSize[1]));
 	check(MipMaps[0].PageTable.Num() == (Header.PageTableVolumeResolution.X * Header.PageTableVolumeResolution.Y * Header.PageTableVolumeResolution.Z));
 	
 	if (NumMipLevels == -1)
@@ -530,8 +530,8 @@ bool FTextureData::GenerateMipMaps(const FTextureDataAddressingInfo& AddressingI
 	}
 
 	MipMaps.SetNum(1);
-	MipMaps[0].PhysicalTileDataA.SetNum((int64)NumMip0Tiles * (int64)UE::SVT::SVTNumVoxelsPerPaddedTile * (int64)FormatSize[0]);
-	MipMaps[0].PhysicalTileDataB.SetNum((int64)NumMip0Tiles * (int64)UE::SVT::SVTNumVoxelsPerPaddedTile * (int64)FormatSize[1]);
+	MipMaps[0].PhysicalTileDataA.SetNum((int64)NumMip0Tiles * (int64)SVT::NumVoxelsPerPaddedTile * (int64)FormatSize[0]);
+	MipMaps[0].PhysicalTileDataB.SetNum((int64)NumMip0Tiles * (int64)SVT::NumVoxelsPerPaddedTile * (int64)FormatSize[1]);
 
 	for (int32 MipLevel = 1; MipLevel < NumMipLevels; ++MipLevel)
 	{
@@ -586,8 +586,8 @@ bool FTextureData::GenerateMipMaps(const FTextureDataAddressingInfo& AddressingI
 
 		// Initialize the SparseVolumeTexture page and tile.
 		MipMaps[MipLevel].PageTable.SetNumZeroed(PageTableVolumeResolution.X * PageTableVolumeResolution.Y * PageTableVolumeResolution.Z);
-		MipMaps[MipLevel].PhysicalTileDataA.SetNum((int64)NumAllocatedPages * (int64)UE::SVT::SVTNumVoxelsPerPaddedTile * (int64)FormatSize[0]);
-		MipMaps[MipLevel].PhysicalTileDataB.SetNum((int64)NumAllocatedPages * (int64)UE::SVT::SVTNumVoxelsPerPaddedTile * (int64)FormatSize[1]);
+		MipMaps[MipLevel].PhysicalTileDataA.SetNum((int64)NumAllocatedPages * (int64)SVT::NumVoxelsPerPaddedTile * (int64)FormatSize[0]);
+		MipMaps[MipLevel].PhysicalTileDataB.SetNum((int64)NumAllocatedPages * (int64)SVT::NumVoxelsPerPaddedTile * (int64)FormatSize[1]);
 		MipMaps[MipLevel].NumPhysicalTiles = NumAllocatedPages;
 
 		// Generate page table and tile volume data by splatting the data
@@ -688,7 +688,7 @@ bool FTextureData::DeduplicateTiles()
 {
 	const int32 NumMipLevels = MipMaps.Num();
 	const int32 FormatSize[] = { GPixelFormats[Header.AttributesFormats[0]].BlockBytes, GPixelFormats[Header.AttributesFormats[1]].BlockBytes };
-	const int32 TileByteSize[] = { UE::SVT::SVTNumVoxelsPerPaddedTile * FormatSize[0], UE::SVT::SVTNumVoxelsPerPaddedTile * FormatSize[1] };
+	const int32 TileByteSize[] = { SVT::NumVoxelsPerPaddedTile * FormatSize[0], SVT::NumVoxelsPerPaddedTile * FormatSize[1] };
 
 	int32 NumProcessedTiles = 0;
 	int32 NumSurvivingTiles = 0;
@@ -761,7 +761,7 @@ bool FTextureData::DeduplicateTiles()
 		}
 
 		const uint64 DeduplicatedTileSizeBytes[] = { (uint64)TileWritePos * (uint64)TileByteSize[0], (uint64)TileWritePos * (uint64)TileByteSize[1] };
-		if (DeduplicatedTileSizeBytes[0] > INT32_MAX || DeduplicatedTileSizeBytes[1] > INT32_MAX)
+		if (DeduplicatedTileSizeBytes[0] > (uint64)SVT::MaxResourceSize || DeduplicatedTileSizeBytes[1] > (uint64)SVT::MaxResourceSize)
 		{
 			UE_LOG(LogSparseVolumeTextureData, Error, TEXT("SparseVolumeTexture still requires too much memory after tile deduplication! Physical tile data A: %llu bytes. Physical tile data B: %llu bytes"),
 				(unsigned long long)DeduplicatedTileSizeBytes[0], (unsigned long long)DeduplicatedTileSizeBytes[1]);
diff --git a/Engine/Source/Runtime/Engine/Private/SparseVolumeTexture/SparseVolumeTextureStreamingManager.cpp b/Engine/Source/Runtime/Engine/Private/SparseVolumeTexture/SparseVolumeTextureStreamingManager.cpp
--- a/Engine/Source/Runtime/Engine/Private/SparseVolumeTexture/SparseVolumeTextureStreamingManager.cpp
+++ b/Engine/Source/Runtime/Engine/Private/SparseVolumeTexture/SparseVolumeTextureStreamingManager.cpp
@@ -116,8 +116,8 @@ static FIntVector3 ComputeTileDataVolumeResolution(int32 NumAllocatedPages)
 
 static FIntVector3 ComputeLargestPossibleTileDataVolumeResolution(int32 VoxelMemSize)
 {
-	const int64 TileMemSize = SVTNumVoxelsPerPaddedTile * VoxelMemSize;
-	const int64 NumMaxTiles = int64(INT32_MAX) / TileMemSize;
+	const int64 TileMemSize = SVT::NumVoxelsPerPaddedTile * VoxelMemSize;
+	const int64 NumMaxTiles = SVT::MaxResourceSize / TileMemSize;
 	int64 ResourceSize = NumMaxTiles * TileMemSize;
 
 	// Find a cube with a volume as close to NumMaxTiles as possible
@@ -143,8 +143,8 @@ static FIntVector3 ComputeLargestPossibleTileDataVolumeResolution(int32 VoxelMem
 	}
 
 	const FIntVector3 Resolution = ResolutionInTiles * SPARSE_VOLUME_TILE_RES_PADDED;
-	check(Resolution.X <= SVTMaxVolumeTextureDim && Resolution.Y <= SVTMaxVolumeTextureDim && Resolution.Z <= SVTMaxVolumeTextureDim);
-	check(((int64)Resolution.X * (int64)Resolution.Y * (int64)Resolution.Z) < int64(INT32_MAX));
+	check(Resolution.X <= SVT::MaxVolumeTextureDim && Resolution.Y <= SVT::MaxVolumeTextureDim && Resolution.Z <= SVT::MaxVolumeTextureDim);
+	check(((int64)Resolution.X * (int64)Resolution.Y * (int64)Resolution.Z) < SVT::MaxResourceSize);
 
 	return Resolution;
 }
@@ -271,7 +271,7 @@ public:
 				TileCoordsPtr = (uint8*)RHILockBuffer(DstTileCoordsUploadBuffer->GetRHI(), 0, MaxNumTiles * sizeof(uint32), RLM_WriteOnly);
 			}
 
-			const int32 NumVoxels = MaxNumTiles * UE::SVT::SVTNumVoxelsPerPaddedTile;
+			const int32 NumVoxels = MaxNumTiles * SVT::NumVoxelsPerPaddedTile;
 
 			// TileData
 			if (FormatSizeA > 0)
@@ -303,8 +303,8 @@ public:
 		check(FormatSizeB <= 0 || TileDataBPtr);
 
 		OutPackedPhysicalTileCoordsPtr = TileCoordsPtr + NumWrittenTiles * sizeof(uint32);
-		OutPtrA = TileDataAPtr ? TileDataAPtr + (NumWrittenTiles * UE::SVT::SVTNumVoxelsPerPaddedTile * FormatSizeA) : nullptr;
-		OutPtrB = TileDataBPtr ? TileDataBPtr + (NumWrittenTiles * UE::SVT::SVTNumVoxelsPerPaddedTile * FormatSizeB) : nullptr;
+		OutPtrA = TileDataAPtr ? TileDataAPtr + (NumWrittenTiles * UE::SVT::NumVoxelsPerPaddedTile * FormatSizeA) : nullptr;
+		OutPtrB = TileDataBPtr ? TileDataBPtr + (NumWrittenTiles * UE::SVT::NumVoxelsPerPaddedTile * FormatSizeB) : nullptr;
 
 		NumWrittenTiles += NumTiles;
 	}
@@ -1133,7 +1133,7 @@ void FStreamingManager::AddInternal(FRDGBuilder& GraphBuilder, FNewSparseVolumeT
 			uint8* DataBPtr = nullptr;
 			RootTileUploader.Add_GetRef(1, TileCoordsPtr, DataAPtr, DataBPtr);
 			FMemory::Memcpy(TileCoordsPtr, &NullTileCoord, sizeof(NullTileCoord));
-			for (int32 VoxelIdx = 0; VoxelIdx < UE::SVT::SVTNumVoxelsPerPaddedTile; ++VoxelIdx)
+			for (int32 VoxelIdx = 0; VoxelIdx < SVT::NumVoxelsPerPaddedTile; ++VoxelIdx)
 			{
 				if (SVTInfo.FormatA != PF_Unknown)
 				{
@@ -1852,8 +1852,8 @@ void FStreamingManager::InstallReadyMipLevels()
 
 		// Tile data
 		{
-			check(MipLevelStreamingInfo.TileDataASize == (SVTNumVoxelsPerPaddedTile * GPixelFormats[SVTInfo->TileDataTexture->FormatA].BlockBytes * NumPhysicalTiles));
-			check(MipLevelStreamingInfo.TileDataBSize == (SVTNumVoxelsPerPaddedTile * GPixelFormats[SVTInfo->TileDataTexture->FormatB].BlockBytes * NumPhysicalTiles));
+			check(MipLevelStreamingInfo.TileDataASize == (SVT::NumVoxelsPerPaddedTile * GPixelFormats[SVTInfo->TileDataTexture->FormatA].BlockBytes * NumPhysicalTiles));
+			check(MipLevelStreamingInfo.TileDataBSize == (SVT::NumVoxelsPerPaddedTile * GPixelFormats[SVTInfo->TileDataTexture->FormatB].BlockBytes * NumPhysicalTiles));
 
 			FUploadTask& Task = UploadTasks.AddDefaulted_GetRef();
 			Task.TaskType = FUploadTask::ETaskType::TileData;
@@ -2051,9 +2051,9 @@ FStreamingManager::FTileDataTexture::FTileDataTexture(const FIntVector3& InResol
 	
 	// Ensure that the tile data texture(s) do not exceed the memory size and resolution limits.
 	if (PhysicalTilesCapacity > LargestPossiblePhysicalTilesCapacity
-		|| (ResolutionInTiles.X * SPARSE_VOLUME_TILE_RES_PADDED) > SVTMaxVolumeTextureDim
-		|| (ResolutionInTiles.Y * SPARSE_VOLUME_TILE_RES_PADDED) > SVTMaxVolumeTextureDim
-		|| (ResolutionInTiles.Z * SPARSE_VOLUME_TILE_RES_PADDED) > SVTMaxVolumeTextureDim)
+		|| (ResolutionInTiles.X * SPARSE_VOLUME_TILE_RES_PADDED) > SVT::MaxVolumeTextureDim
+		|| (ResolutionInTiles.Y * SPARSE_VOLUME_TILE_RES_PADDED) > SVT::MaxVolumeTextureDim
+		|| (ResolutionInTiles.Z * SPARSE_VOLUME_TILE_RES_PADDED) > SVT::MaxVolumeTextureDim)
 	{
 		ResolutionInTiles = LargestPossibleResolution;
 		PhysicalTilesCapacity = LargestPossiblePhysicalTilesCapacity;
@@ -2064,9 +2064,9 @@ FStreamingManager::FTileDataTexture::FTileDataTexture(const FIntVector3& InResol
 	}
 
 	const FIntVector3 Resolution = ResolutionInTiles * SPARSE_VOLUME_TILE_RES_PADDED;
-	check(Resolution.X <= SVTMaxVolumeTextureDim && Resolution.Y <= SVTMaxVolumeTextureDim && Resolution.Z <= SVTMaxVolumeTextureDim);
-	check(((int64)Resolution.X * (int64)Resolution.Y * (int64)Resolution.Z * (int64)GPixelFormats[FormatA].BlockBytes) <= int64(INT32_MAX));
-	check(((int64)Resolution.X * (int64)Resolution.Y * (int64)Resolution.Z * (int64)GPixelFormats[FormatB].BlockBytes) <= int64(INT32_MAX));
+	check(Resolution.X <= SVT::MaxVolumeTextureDim && Resolution.Y <= SVT::MaxVolumeTextureDim && Resolution.Z <= SVT::MaxVolumeTextureDim);
+	check(((int64)Resolution.X * (int64)Resolution.Y * (int64)Resolution.Z * (int64)GPixelFormats[FormatA].BlockBytes) <= SVT::MaxResourceSize);
+	check(((int64)Resolution.X * (int64)Resolution.Y * (int64)Resolution.Z * (int64)GPixelFormats[FormatB].BlockBytes) <= SVT::MaxResourceSize);
 	
 	TileCoords.SetNum(PhysicalTilesCapacity);
 
diff --git a/Engine/Source/Runtime/Engine/Private/SparseVolumeTexture/SparseVolumeTextureUtility.h b/Engine/Source/Runtime/Engine/Private/SparseVolumeTexture/SparseVolumeTextureUtility.h
--- a/Engine/Source/Runtime/Engine/Private/SparseVolumeTexture/SparseVolumeTextureUtility.h
+++ b/Engine/Source/Runtime/Engine/Private/SparseVolumeTexture/SparseVolumeTextureUtility.h
@@ -11,9 +11,10 @@ namespace SVT
 {
 	// We are using this instead of GMaxVolumeTextureDimensions to be independent of the platform that
 	// the asset is imported on. 2048 should be a safe value that should be supported by all our platforms.
-	static constexpr int32 SVTMaxVolumeTextureDim = 2048;
-	static constexpr int32 SVTNumVoxelsPerTile = SPARSE_VOLUME_TILE_RES * SPARSE_VOLUME_TILE_RES * SPARSE_VOLUME_TILE_RES;
-	static constexpr int32 SVTNumVoxelsPerPaddedTile = SPARSE_VOLUME_TILE_RES_PADDED * SPARSE_VOLUME_TILE_RES_PADDED * SPARSE_VOLUME_TILE_RES_PADDED;
+	static constexpr int32 MaxVolumeTextureDim = 2048;
+	static constexpr int64 MaxResourceSize = 2048LL * 1024LL * 1024LL;
+	static constexpr int32 NumVoxelsPerTile = SPARSE_VOLUME_TILE_RES * SPARSE_VOLUME_TILE_RES * SPARSE_VOLUME_TILE_RES;
+	static constexpr int32 NumVoxelsPerPaddedTile = SPARSE_VOLUME_TILE_RES_PADDED * SPARSE_VOLUME_TILE_RES_PADDED * SPARSE_VOLUME_TILE_RES_PADDED;
 
 	uint32 PackPageTableEntry(const FIntVector3& Coord);
 	FIntVector3 UnpackPageTableEntry(uint32 Packed);
diff --git a/Engine/Source/Runtime/Engine/Public/SparseVolumeTexture/SparseVolumeTextureData.h b/Engine/Source/Runtime/Engine/Public/SparseVolumeTexture/SparseVolumeTextureData.h
--- a/Engine/Source/Runtime/Engine/Public/SparseVolumeTexture/SparseVolumeTextureData.h
+++ b/Engine/Source/Runtime/Engine/Public/SparseVolumeTexture/SparseVolumeTextureData.h
@@ -57,7 +57,7 @@ struct ENGINE_API FTextureData
 	void Serialize(FArchive& Ar);
 	
 	bool Create(const ITextureDataProvider& DataProvider);
-	bool CreateFromDense(const FTextureDataCreateInfo& CreateInfo, const TArrayView<uint8>& VoxelDataA, const TArrayView<uint8>& VoxelDataB);
+	bool CreateFromDense(const FTextureDataCreateInfo& CreateInfo, const TArrayView64<uint8>& VoxelDataA, const TArrayView64<uint8>& VoxelDataB);
 	void CreateDefault();
 	FVector4f Load(const FIntVector3& VolumeCoord, int32 MipLevel, int32 AttributesIdx, const FTextureDataAddressingInfo& AddressingInfo) const;
 	bool BuildDerivedData(const FTextureDataAddressingInfo& AddressingInfo, int32 NumMipLevels, bool bMoveMip0FromThis, FTextureData& OutDerivedData);
