Move Audio DDC key(s) somewhere more obvious and give (Audio) a propper guid (this will bump all IAudioFormats too)

diff --git a/Engine/Source/Runtime/Engine/Private/AudioDerivedData.cpp b/Engine/Source/Runtime/Engine/Private/AudioDerivedData.cpp
--- a/Engine/Source/Runtime/Engine/Private/AudioDerivedData.cpp
+++ b/Engine/Source/Runtime/Engine/Private/AudioDerivedData.cpp
@@ -158,7 +158,6 @@ FName FStreamedAudioPlatformData::GetAudioFormat() const
 	return AudioFormat;
 }
 
-#if WITH_EDITORONLY_DATA
 
 /*------------------------------------------------------------------------------
 Derived data key generation.
@@ -166,7 +165,10 @@ Derived data key generation.
 
 // If you want to bump this version, generate a new guid using
 // VS->Tools->Create GUID and paste it here. https://www.guidgen.com works too.
-#define STREAMEDAUDIO_DERIVEDDATA_VER		TEXT("F6F8FA984D96440B8A28C73F2ABAE692")
+#define AUDIO_DERIVEDDATA_VER				TEXT("4CC71E548D764B6989612C03CBF8E376")				// Change this if you want to bump all audio formats at once
+#define STREAMEDAUDIO_DERIVEDDATA_VER		TEXT("F6F8FA984D96440B8A28C73F2ABAE692")				// This depends on the above key, but will regenerate all streaming chunk data derived from the compressed audio
+
+#if WITH_EDITORONLY_DATA
 
 #ifndef CASE_ENUM_TO_TEXT
 #define CASE_ENUM_TO_TEXT(X) case X: return TEXT(#X);
@@ -1917,6 +1919,11 @@ FDerivedAudioDataCompressor::FDerivedAudioDataCompressor(USoundWave* InSoundNode
 {
 }
 
+const TCHAR* FDerivedAudioDataCompressor::GetVersionString() const
+{	
+	return AUDIO_DERIVEDDATA_VER;
+}
+
 FString FDerivedAudioDataCompressor::GetPluginSpecificCacheKeySuffix() const
 {
 	int32 FormatVersion = 0xffff; // if the compressor is NULL, this will be used as the version...and in that case we expect everything to fail anyway
diff --git a/Engine/Source/Runtime/Engine/Public/AudioDerivedData.h b/Engine/Source/Runtime/Engine/Public/AudioDerivedData.h
--- a/Engine/Source/Runtime/Engine/Public/AudioDerivedData.h
+++ b/Engine/Source/Runtime/Engine/Public/AudioDerivedData.h
@@ -23,14 +23,7 @@ public:
 		return TEXT("Audio");
 	}
 
-	virtual const TCHAR* GetVersionString() const override
-	{
-		// This is a version string that mimics the old versioning scheme. If you
-		// want to bump this version, generate a new guid using VS->Tools->Create GUID and
-		// return it here. Ex.
-		// return TEXT("855EE5B3574C43ABACC6700C4ADC62E6");
-		return TEXT("0005_0000");
-	}
+	virtual const TCHAR* GetVersionString() const override;
 
 	virtual FString GetPluginSpecificCacheKeySuffix() const override;
 	
