Check if any wave insances on an active sound implement a parameter interface instead of just its main sound.

diff --git a/Engine/Source/Runtime/Engine/Private/ActiveSound.cpp b/Engine/Source/Runtime/Engine/Private/ActiveSound.cpp
--- a/Engine/Source/Runtime/Engine/Private/ActiveSound.cpp
+++ b/Engine/Source/Runtime/Engine/Private/ActiveSound.cpp
@@ -603,10 +603,28 @@ void FActiveSound::UpdateInterfaceParameters(const TArray<FListener>& InListener
 	FParameterInterfacePtr SourceOrientationInterface = SourceOrientationInterface::GetInterface();
 	FParameterInterfacePtr ListenerOrientationInterface = ListenerOrientationInterface::GetInterface();
 
-	const bool bImplementsAttenuation = Sound->ImplementsParameterInterface(AttenuationInterface);
-	const bool bImplementsSpatialization = Sound->ImplementsParameterInterface(SpatializationInterface);
-	const bool bImplementsSourceOrientation = Sound->ImplementsParameterInterface(SourceOrientationInterface);
-	const bool bImplementsListenerOrientation = Sound->ImplementsParameterInterface(ListenerOrientationInterface);
+	bool bImplementsAttenuation = Sound->ImplementsParameterInterface(AttenuationInterface);
+	bool bImplementsSpatialization = Sound->ImplementsParameterInterface(SpatializationInterface);
+	bool bImplementsSourceOrientation = Sound->ImplementsParameterInterface(SourceOrientationInterface);
+	bool bImplementsListenerOrientation = Sound->ImplementsParameterInterface(ListenerOrientationInterface);
+	
+	for (const TPair<UPTRINT, FWaveInstance*>& InstancePair : WaveInstances)
+	{
+		const FWaveInstance* Instance = InstancePair.Value;
+
+		if (Instance == nullptr || Instance->WaveData == Sound)
+		{
+			continue;
+		}
+
+		if (const USoundBase* SoundInstance = Instance->WaveData)
+		{
+			bImplementsAttenuation |= SoundInstance->ImplementsParameterInterface(AttenuationInterface);
+			bImplementsSpatialization |= SoundInstance->ImplementsParameterInterface(SpatializationInterface);
+			bImplementsSourceOrientation |= SoundInstance->ImplementsParameterInterface(SourceOrientationInterface);
+			bImplementsListenerOrientation |= SoundInstance->ImplementsParameterInterface(ListenerOrientationInterface);
+		}
+	}
 
 	if (false == (bImplementsAttenuation || bImplementsSpatialization || bImplementsSourceOrientation || bImplementsListenerOrientation))
 	{
