Fix eye legacy material conversion with Substrate.

diff --git a/Engine/Source/Runtime/Engine/Private/Materials/MaterialExpressions.cpp b/Engine/Source/Runtime/Engine/Private/Materials/MaterialExpressions.cpp
--- a/Engine/Source/Runtime/Engine/Private/Materials/MaterialExpressions.cpp
+++ b/Engine/Source/Runtime/Engine/Private/Materials/MaterialExpressions.cpp
@@ -22886,7 +22886,7 @@ static int32 CompileWithDefaultFloat3(class FMaterialCompiler* Compiler, FExpres
 	}
 	return CodeChunk == INDEX_NONE ? Compiler->Constant3(X, Y, Z) : CodeChunk;
 }
-static int32 CompileWithDefaultNormalWS(class FMaterialCompiler* Compiler, FExpressionInput& Input, bool* bDefaultIsUsed = nullptr)
+static int32 CompileWithDefaultNormalWS(class FMaterialCompiler* Compiler, FExpressionInput& Input, bool bConvertToRequestedSpace=true)
 {
 	if (Input.GetTracedInput().Expression != nullptr)
 	{
@@ -22894,24 +22894,16 @@ static int32 CompileWithDefaultNormalWS(class FMaterialCompiler* Compiler, FExpr
 
 		if (NormalCodeChunk == INDEX_NONE)
 		{
-			if (bDefaultIsUsed)
-			{
-				*bDefaultIsUsed = true;
-			}
 			// Nothing is plug in from the linked input, so specify world space normal the BSDF node expects.
 			return Compiler->VertexNormal();
 		}
 		// Transform into world space normal if needed. BSDF nodes always expects world space normal as input.
-		return Compiler->TransformNormalFromRequestedBasisToWorld(NormalCodeChunk);
-	}
-	if (bDefaultIsUsed)
-	{
-		*bDefaultIsUsed = true;
+		return bConvertToRequestedSpace ? Compiler->TransformNormalFromRequestedBasisToWorld(NormalCodeChunk) : NormalCodeChunk;
 	}
 	// Nothing is plug in on the BSDF node, so specify world space normal the node expects.
 	return Compiler->VertexNormal();
 }
-static int32 CompileWithDefaultTangentWS(class FMaterialCompiler* Compiler, FExpressionInput& Input, bool* bDefaultIsUsed = nullptr)
+static int32 CompileWithDefaultTangentWS(class FMaterialCompiler* Compiler, FExpressionInput& Input, bool bConvertToRequestedSpace=true)
 {
 	if (Input.GetTracedInput().Expression != nullptr)
 	{
@@ -22919,19 +22911,11 @@ static int32 CompileWithDefaultTangentWS(class FMaterialCompiler* Compiler, FExp
 
 		if (TangentCodeChunk == INDEX_NONE)
 		{
-			if (bDefaultIsUsed)
-			{
-				*bDefaultIsUsed = true;
-			}
 			// Nothing is plug in from the linked input, so specify world space tangent the BSDF node expects.
 			return Compiler->VertexTangent();
 		}
 		// Transform into world space tangent if needed. BSDF nodes always expects world space tangent as input.
-		return Compiler->TransformNormalFromRequestedBasisToWorld(TangentCodeChunk);
-	}
-	if (bDefaultIsUsed)
-	{
-		*bDefaultIsUsed = true;
+		return bConvertToRequestedSpace ? Compiler->TransformNormalFromRequestedBasisToWorld(TangentCodeChunk) : TangentCodeChunk;
 	}
 	// Nothing is plug in on the BSDF node, so specify world space tangent the node expects.
 	return Compiler->VertexTangent();
@@ -23002,7 +22986,8 @@ int32 UMaterialExpressionStrataLegacyConversion::Compile(class FMaterialCompiler
 	const bool bHasCustomTangent = CustomTangent.IsConnected();
 	if (bHasCustomTangent)
 	{
-		CustomTangent_TangentCodeChunk = CompileWithDefaultNormalWS(Compiler, CustomTangent);
+		// Legacy code doesn't do tangent <-> world basis conversion on tangent output, when provided.
+		CustomTangent_TangentCodeChunk = CompileWithDefaultNormalWS(Compiler, CustomTangent, false /*bConvertToRequestedSpace*/);
 	}
 	else
 	{
