Only init default parameters once - fixes object proxies being cleared on metasounds

#rnx

diff --git a/Engine/Source/Runtime/Engine/Private/AudioDevice.cpp b/Engine/Source/Runtime/Engine/Private/AudioDevice.cpp
--- a/Engine/Source/Runtime/Engine/Private/AudioDevice.cpp
+++ b/Engine/Source/Runtime/Engine/Private/AudioDevice.cpp
@@ -5061,7 +5061,6 @@ void FAudioDevice::AddNewActiveSound(const FActiveSound& NewActiveSound, TArray<
 	if (USoundBase* Sound = NewActiveSound.GetSound())
 	{
 		Sound->InitResources();
-		Sound->InitParameters(InDefaultParams, "NewSoundRequest");
 	}
 
 	AddNewActiveSoundInternal(NewActiveSound, MoveTemp(InDefaultParams));
