fix issue where cluster union would not get added into external SQ properly in sync bodies

diff --git a/Engine/Source/Runtime/Engine/Classes/PhysicsEngine/ClusterUnionComponent.h b/Engine/Source/Runtime/Engine/Classes/PhysicsEngine/ClusterUnionComponent.h
--- a/Engine/Source/Runtime/Engine/Classes/PhysicsEngine/ClusterUnionComponent.h
+++ b/Engine/Source/Runtime/Engine/Classes/PhysicsEngine/ClusterUnionComponent.h
@@ -190,6 +190,8 @@ public:
 	// Lambda returns whether or not iteration should continue;
 	void VisitAllCurrentChildComponents(const TFunction<bool(UPrimitiveComponent*)>& Lambda) const;
 
+	int32 NumChildClusterComponents() const;
+
 	friend class UClusterUnionReplicatedProxyComponent;
 protected:
 
diff --git a/Engine/Source/Runtime/Engine/Private/PhysicsEngine/ClusterUnionComponent.cpp b/Engine/Source/Runtime/Engine/Private/PhysicsEngine/ClusterUnionComponent.cpp
--- a/Engine/Source/Runtime/Engine/Private/PhysicsEngine/ClusterUnionComponent.cpp
+++ b/Engine/Source/Runtime/Engine/Private/PhysicsEngine/ClusterUnionComponent.cpp
@@ -1202,3 +1202,8 @@ void UClusterUnionComponent::AddReferencedObjects(UObject* InThis, FReferenceCol
 		}
 	}
 }
+
+int32 UClusterUnionComponent::NumChildClusterComponents() const
+{
+	return PendingComponentSync.Num() + ComponentToPhysicsObjects.Num();
+}
diff --git a/Engine/Source/Runtime/Engine/Private/PhysicsEngine/Experimental/PhysScene_Chaos.cpp b/Engine/Source/Runtime/Engine/Private/PhysicsEngine/Experimental/PhysScene_Chaos.cpp
--- a/Engine/Source/Runtime/Engine/Private/PhysicsEngine/Experimental/PhysScene_Chaos.cpp
+++ b/Engine/Source/Runtime/Engine/Private/PhysicsEngine/Experimental/PhysScene_Chaos.cpp
@@ -2172,10 +2172,14 @@ void FPhysScene_Chaos::OnSyncBodies(Chaos::FPhysicsSolverBase* Solver)
 
 				// make sure we have at least a child to be added to the acceleration structure 
 				// this avoid the invalid bounds to cause the particle to be added to the global acceleration structure array
-				if (ParentComponent->GetNumChildrenComponents() > 0)
+				if (ParentComponent->NumChildClusterComponents() > 0)
 				{
 					Interface->AddToSpatialAcceleration({ &Handle, 1 }, Outer->GetSpacialAcceleration());
 				}
+				else
+				{
+					Interface->RemoveFromSpatialAcceleration({ &Handle, 1 }, Outer->GetSpacialAcceleration());
+				}
 				ParentComponent->SyncClusterUnionFromProxy();
 				DirtyParticle->ClearEvents();
 			}
