Optimizations in Material translation code amounting to 50% speed up.

Optimized function UMaterialExpression::GetInputs() with highest exclusive run-time during material translation from 19.79% to 1.79% exclusive run time.

This change caches the inputs upon expression creation by default and return a TArrayView instead of always creating a new transient dynamic array. The optimization also simplifies code for some Expression types.

Also optimized ContainsInputLoopInternal() to use a linked list of stack allocated nodes instead of Pushing/Popping from a dynamic array.

diff --git a/Engine/Plugins/Editor/SpeedTreeImporter/Source/SpeedTreeImporter/Private/SpeedTreeImportFactory.cpp b/Engine/Plugins/Editor/SpeedTreeImporter/Source/SpeedTreeImporter/Private/SpeedTreeImportFactory.cpp
--- a/Engine/Plugins/Editor/SpeedTreeImporter/Source/SpeedTreeImporter/Private/SpeedTreeImportFactory.cpp
+++ b/Engine/Plugins/Editor/SpeedTreeImporter/Source/SpeedTreeImporter/Private/SpeedTreeImportFactory.cpp
@@ -575,7 +575,7 @@ void LayoutMaterial(UMaterialInterface* MaterialInterface, bool bOffsetOddColumn
 		{
 			Expression->MaterialExpressionEditorX = FMath::Min(Expression->MaterialExpressionEditorX, -Width);
 
-			TArray<FExpressionInput*> Inputs = Expression->GetInputs();
+			TArrayView<FExpressionInput*> Inputs = Expression->GetInputsView();
 			for (int32 InputIndex = 0; InputIndex < Inputs.Num(); ++InputIndex)
 			{
 				UMaterialExpression* Input = Inputs[InputIndex]->Expression;
@@ -619,7 +619,7 @@ void LayoutMaterial(UMaterialInterface* MaterialInterface, bool bOffsetOddColumn
 				// all the outputs to other expressions
 				for (UMaterialExpression* OtherExpression : Material->GetExpressions())
 				{
-					TArray<FExpressionInput*> Inputs = OtherExpression->GetInputs();
+					TArrayView<FExpressionInput*> Inputs = OtherExpression->GetInputsView();
 					for (int32 InputIndex = 0; InputIndex < Inputs.Num(); ++InputIndex)
 					{
 						if (Inputs[InputIndex]->Expression == Expression)
diff --git a/Engine/Plugins/Enterprise/DatasmithImporter/Source/DatasmithImporter/Private/DatasmithMaterialExpressions.cpp b/Engine/Plugins/Enterprise/DatasmithImporter/Source/DatasmithImporter/Private/DatasmithMaterialExpressions.cpp
--- a/Engine/Plugins/Enterprise/DatasmithImporter/Source/DatasmithImporter/Private/DatasmithMaterialExpressions.cpp
+++ b/Engine/Plugins/Enterprise/DatasmithImporter/Source/DatasmithImporter/Private/DatasmithMaterialExpressions.cpp
@@ -276,7 +276,7 @@ void FDatasmithMaterialExpressions::GetSamplersRecursive(UMaterialExpression* Ex
 		TextureSamplers.AddUnique((UMaterialExpressionTextureSample*)Expression);
 	}
 
-	TArray<FExpressionInput*> Inputs = Expression->GetInputs();
+	TArrayView<FExpressionInput*> Inputs = Expression->GetInputsView();
 	for (int32 InputIndex = 0; InputIndex < Inputs.Num(); InputIndex++)
 	{
 		UMaterialExpression* Input = Inputs[InputIndex]->Expression;
@@ -3144,7 +3144,7 @@ void FDatasmithMaterialExpressions::CreateUEPbrMaterialGraph(const TSharedPtr< I
 			{
 				for ( int32 ExpressionInput = 0; ExpressionInput < DatasmithExpression->GetInputCount(); ++ExpressionInput )
 				{
-					if ( MaterialOutputExpression->GetInputs().IsValidIndex( ExpressionInput ) )
+					if ( MaterialOutputExpression->GetInputsView().IsValidIndex( ExpressionInput ) )
 					{
 						ConnectExpression( MaterialElement.ToSharedRef(), MaterialExpressions, DatasmithExpression->GetInput( ExpressionInput )->GetExpression(),
 							MaterialOutputExpression->GetInput( ExpressionInput ), DatasmithExpression->GetInput( ExpressionInput )->GetOutputIndex() );
diff --git a/Engine/Plugins/FX/Niagara/Source/NiagaraEditor/Private/ViewModels/Stack/NiagaraStackFunctionInput.cpp b/Engine/Plugins/FX/Niagara/Source/NiagaraEditor/Private/ViewModels/Stack/NiagaraStackFunctionInput.cpp
--- a/Engine/Plugins/FX/Niagara/Source/NiagaraEditor/Private/ViewModels/Stack/NiagaraStackFunctionInput.cpp
+++ b/Engine/Plugins/FX/Niagara/Source/NiagaraEditor/Private/ViewModels/Stack/NiagaraStackFunctionInput.cpp
@@ -1051,7 +1051,7 @@ FString UNiagaraStackFunctionInput::ResolveDisplayNameArgument(const FString& In
 		Material->GetAllExpressionsInMaterialAndFunctionsOfType<UMaterialExpression>(Expressions);
 		for (UMaterialExpression* Expression : Expressions)
 		{
-			const TArray<FExpressionInput*>& ExpressionInputs = Expression->GetInputs();
+			TArrayView<FExpressionInput*> ExpressionInputs = Expression->GetInputsView();
 			ExpressionInputsToProcess.Append(ExpressionInputs);
 		}
 	}
diff --git a/Engine/Plugins/Interchange/Runtime/Source/Import/Private/Material/InterchangeMaterialFactory.cpp b/Engine/Plugins/Interchange/Runtime/Source/Import/Private/Material/InterchangeMaterialFactory.cpp
--- a/Engine/Plugins/Interchange/Runtime/Source/Import/Private/Material/InterchangeMaterialFactory.cpp
+++ b/Engine/Plugins/Interchange/Runtime/Source/Import/Private/Material/InterchangeMaterialFactory.cpp
@@ -84,7 +84,7 @@ namespace UE
 				{
 					int32 ExpressionInputIndex = 0;
 
-					for (const FExpressionInput* ExpressionInput : MaterialExpression.GetInputs())
+					for (const FExpressionInput* ExpressionInput : MaterialExpression.GetInputsView())
 					{
 						// MaterialFuncCall appends the type to the input name when calling GetInputName
 						// and the InputName in FExpressionInput is optional so we'll check both here to be safe
diff --git a/Engine/Plugins/Runtime/RenderTrace/Source/Private/Materials/MaterialExpressionPhysicalMaterialOutput.cpp b/Engine/Plugins/Runtime/RenderTrace/Source/Private/Materials/MaterialExpressionPhysicalMaterialOutput.cpp
--- a/Engine/Plugins/Runtime/RenderTrace/Source/Private/Materials/MaterialExpressionPhysicalMaterialOutput.cpp
+++ b/Engine/Plugins/Runtime/RenderTrace/Source/Private/Materials/MaterialExpressionPhysicalMaterialOutput.cpp
@@ -62,14 +62,15 @@ void UMaterialExpressionPhysicalMaterialOutput::GetCaption(TArray<FString>& OutC
 	OutCaptions.Add(FString(TEXT("Physical Material Output")));
 }
 
-const TArray<FExpressionInput*> UMaterialExpressionPhysicalMaterialOutput::GetInputs()
+TArrayView<FExpressionInput*> UMaterialExpressionPhysicalMaterialOutput::GetInputsView()
 {
-	TArray<FExpressionInput*> OutInputs;
+	CachedInputs.Empty();
+	CachedInputs.Reserve(Inputs.Num());
 	for (auto& Input : Inputs)
 	{
-		OutInputs.Add(&Input.Input);
+		CachedInputs.Add(&Input.Input);
 	}
-	return OutInputs;
+	return CachedInputs;
 }
 
 FExpressionInput* UMaterialExpressionPhysicalMaterialOutput::GetInput(int32 InputIndex)
diff --git a/Engine/Plugins/Runtime/RenderTrace/Source/Private/Materials/MaterialExpressionPhysicalMaterialOutput.h b/Engine/Plugins/Runtime/RenderTrace/Source/Private/Materials/MaterialExpressionPhysicalMaterialOutput.h
--- a/Engine/Plugins/Runtime/RenderTrace/Source/Private/Materials/MaterialExpressionPhysicalMaterialOutput.h
+++ b/Engine/Plugins/Runtime/RenderTrace/Source/Private/Materials/MaterialExpressionPhysicalMaterialOutput.h
@@ -50,7 +50,7 @@ class RENDERTRACE_API UMaterialExpressionPhysicalMaterialOutput : public UMateri
 
 	//~ Begin UMaterialExpression Interface
 	virtual void GetCaption(TArray<FString>& OutCaptions) const override;
-	virtual const TArray<FExpressionInput*> GetInputs() override;
+	virtual TArrayView<FExpressionInput*> GetInputsView() override;
 	virtual FExpressionInput* GetInput(int32 InputIndex) override;
 	virtual FName GetInputName(int32 InputIndex) const override;
 	virtual uint32 GetInputType(int32 InputIndex) override { return MCT_Float; }
diff --git a/Engine/Plugins/Tests/InterchangeTests/Source/InterchangeTests/Private/ImportTestFunctions/MaterialXTestFunctions.cpp b/Engine/Plugins/Tests/InterchangeTests/Source/InterchangeTests/Private/ImportTestFunctions/MaterialXTestFunctions.cpp
--- a/Engine/Plugins/Tests/InterchangeTests/Source/InterchangeTests/Private/ImportTestFunctions/MaterialXTestFunctions.cpp
+++ b/Engine/Plugins/Tests/InterchangeTests/Source/InterchangeTests/Private/ImportTestFunctions/MaterialXTestFunctions.cpp
@@ -20,7 +20,7 @@ namespace UE::Interchange::Tests::Private
 		const UMaterialEditorOnlyData* MaterialEditorOnly = Material.GetEditorOnlyData();
 		if (UMaterialExpressionMaterialFunctionCall* FunctionCall = Cast<UMaterialExpressionMaterialFunctionCall>(MaterialEditorOnly->BaseColor.Expression))
 		{
-			const TArray<FExpressionInput*>& Inputs = FunctionCall->GetInputs();
+			TArrayView<FExpressionInput*> Inputs = FunctionCall->GetInputsView();
 			InputNames.Empty(Inputs.Num());
 			for (const FExpressionInput* InputExpression : Inputs)
 			{
diff --git a/Engine/Source/Editor/MaterialEditor/Private/MaterialEditingLibrary.cpp b/Engine/Source/Editor/MaterialEditor/Private/MaterialEditingLibrary.cpp
--- a/Engine/Source/Editor/MaterialEditor/Private/MaterialEditingLibrary.cpp
+++ b/Engine/Source/Editor/MaterialEditor/Private/MaterialEditingLibrary.cpp
@@ -45,7 +45,7 @@ static FExpressionInput* GetExpressionInputByName(UMaterialExpression* Expressio
 	check(Expression);
 	FExpressionInput* Result = nullptr;
 
-	TArray<FExpressionInput*> Inputs = Expression->GetInputs();
+	TArrayView<FExpressionInput*> Inputs = Expression->GetInputsView();
 
 	// Return first input if no name specified
 	if (InputName.IsNone())
@@ -184,7 +184,7 @@ namespace MaterialEditingLibraryImpl
 
 		MaterialExpressionsToLayout.Add( MaterialExpression ) = MoveTemp( LayoutInfo );
 
-		for ( FExpressionInput* ExpressionInput : MaterialExpression->GetInputs() )
+		for ( FExpressionInput* ExpressionInput : MaterialExpression->GetInputsView() )
 		{
 			LayoutMaterialExpression( ExpressionInput->Expression, MaterialExpression, MaterialExpressionsToLayout, Row, Depth + 1 );
 		}
@@ -420,7 +420,7 @@ static void BreakLinksToExpression(TConstArrayView<TObjectPtr<UMaterialExpressio
 		// Don't check myself, though that shouldn't really matter...
 		if (TestExp != Expression)
 		{
-			TArray<FExpressionInput*> Inputs = TestExp->GetInputs();
+			TArrayView<FExpressionInput*> Inputs = TestExp->GetInputsView();
 			for (FExpressionInput* Input : Inputs)
 			{
 				if (Input->Expression == Expression)
@@ -854,7 +854,7 @@ TArray<UMaterialExpression*> UMaterialEditingLibrary::GetInputsForMaterialExpres
 	TArray<UMaterialExpression*> MaterialExpressions;
 	if (Material)
 	{
-		for (const FExpressionInput* Input : MaterialExpression->GetInputs())
+		for (const FExpressionInput* Input : MaterialExpression->GetInputsView())
 		{
 			MaterialExpressions.Add(Input->Expression);
 		}
@@ -866,7 +866,7 @@ TArray<UMaterialExpression*> UMaterialEditingLibrary::GetInputsForMaterialExpres
 bool UMaterialEditingLibrary::GetInputNodeOutputNameForMaterialExpression(UMaterialExpression* MaterialExpression, UMaterialExpression* InputNode, FString& OutputName)
 {
 	OutputName = TEXT("");
-	for (const FExpressionInput* Input : MaterialExpression->GetInputs())
+	for (const FExpressionInput* Input : MaterialExpression->GetInputsView())
 	{
 		if (Input->Expression == InputNode)
 		{
diff --git a/Engine/Source/Editor/MaterialEditor/Private/MaterialEditor.cpp b/Engine/Source/Editor/MaterialEditor/Private/MaterialEditor.cpp
--- a/Engine/Source/Editor/MaterialEditor/Private/MaterialEditor.cpp
+++ b/Engine/Source/Editor/MaterialEditor/Private/MaterialEditor.cpp
@@ -4704,7 +4704,7 @@ UClass* FMaterialEditor::GetOnPromoteToParameterClass(const UEdGraphPin* TargetP
 	}
 	else if (OtherPinNode)
 	{
-		const TArray<FExpressionInput*> ExpressionInputs = OtherPinNode->MaterialExpression->GetInputs();
+		TArrayView<FExpressionInput*> ExpressionInputs = OtherPinNode->MaterialExpression->GetInputsView();
 		FName TargetPinName = OtherPinNode->GetShortenPinName(TargetPin->PinName);
 
 		for (int32 Index = 0; Index < ExpressionInputs.Num(); ++Index)
@@ -4839,7 +4839,7 @@ void FMaterialEditor::OnCreateStrataNodeForPin(const FToolMenuContext& InMenuCon
 	{
 		// Link manually
 		UMaterialGraphNode* NewNode = Cast<UMaterialGraphNode>(Action.PerformAction(GraphObj, nullptr, NewNodePos));
-		const TArray<FExpressionInput*> NewNodeExpressionInputs = NewNode->MaterialExpression->GetInputs();
+		TArrayView<FExpressionInput*> NewNodeExpressionInputs = NewNode->MaterialExpression->GetInputsView();
 
 		// From that direction, the node is never going to be a root node (a root node has no output we can connect from).
 		UMaterialGraphNode* TargetPinNode = Cast<UMaterialGraphNode>(TargetPin->GetOwningNode());
@@ -4878,7 +4878,7 @@ bool FMaterialEditor::OnCanCreateStrataNodeForPin(const FToolMenuContext& InMenu
 		}
 		else if (OtherPinNode)
 		{
-			const TArray<FExpressionInput*> ExpressionInputs = OtherPinNode->MaterialExpression->GetInputs();
+			TArrayView<FExpressionInput*> ExpressionInputs = OtherPinNode->MaterialExpression->GetInputsView();
 			FName TargetPinName = OtherPinNode->GetShortenPinName(TargetPin->PinName);
 
 			for (int32 Index = 0; Index < ExpressionInputs.Num(); ++Index)
diff --git a/Engine/Source/Editor/MaterialEditor/Private/MaterialEditorUtilities.cpp b/Engine/Source/Editor/MaterialEditor/Private/MaterialEditorUtilities.cpp
--- a/Engine/Source/Editor/MaterialEditor/Private/MaterialEditorUtilities.cpp
+++ b/Engine/Source/Editor/MaterialEditor/Private/MaterialEditorUtilities.cpp
@@ -621,7 +621,7 @@ void FMaterialEditorUtilities::GetVisibleMaterialParametersFromExpression(
 		}
 		else
 		{
-			const TArray<FExpressionInput*>& ExpressionInputs = MaterialExpressionKey.Expression->GetInputs();
+			TArrayView<FExpressionInput*> ExpressionInputs = MaterialExpressionKey.Expression->GetInputsView();
 			for (int32 ExpressionInputIndex = 0; ExpressionInputIndex < ExpressionInputs.Num(); ExpressionInputIndex++)
 			{
 				//retrieve the expression input and then start parsing its children
@@ -729,7 +729,7 @@ bool FMaterialEditorUtilities::HasCompatibleConnection(UClass* ExpressionClass,
 		UMaterialExpression* DefaultExpression = CastChecked<UMaterialExpression>(ExpressionClass->GetDefaultObject());
 		if (TestDirection == EGPD_Output)
 		{
-			int32 NumInputs = DefaultExpression->GetInputs().Num();
+			int32 NumInputs = DefaultExpression->GetInputsView().Num();
 			for (int32 Index = 0; Index < NumInputs; ++Index)
 			{
 				uint32 InputType = DefaultExpression->GetInputType(Index);
diff --git a/Engine/Source/Editor/UnrealEd/Private/MaterialGraph.cpp b/Engine/Source/Editor/UnrealEd/Private/MaterialGraph.cpp
--- a/Engine/Source/Editor/UnrealEd/Private/MaterialGraph.cpp
+++ b/Engine/Source/Editor/UnrealEd/Private/MaterialGraph.cpp
@@ -497,7 +497,7 @@ void UMaterialGraph::LinkGraphNodesFromMaterial()
 			continue;
 		}
 
-		const TArray<FExpressionInput*> ExpressionInputs = Expression->GetInputs();
+		TArrayView<FExpressionInput*> ExpressionInputs = Expression->GetInputsView();
 
 		TArray<FExpressionExecOutputEntry> ExecOutputs;
 		Expression->GetExecOutputs(ExecOutputs);
@@ -631,7 +631,7 @@ void UMaterialGraph::LinkMaterialExpressionsFromGraph()
 						Expression->Desc = GraphNode->NodeComment;
 					}
 
-					const TArray<FExpressionInput*> ExpressionInputs = Expression->GetInputs();
+					TArrayView<FExpressionInput*> ExpressionInputs = Expression->GetInputsView();
 
 					TArray<FExpressionExecOutputEntry> ExecOutputs;
 					Expression->GetExecOutputs(ExecOutputs);
diff --git a/Engine/Source/Editor/UnrealEd/Private/MaterialGraphNode.cpp b/Engine/Source/Editor/UnrealEd/Private/MaterialGraphNode.cpp
--- a/Engine/Source/Editor/UnrealEd/Private/MaterialGraphNode.cpp
+++ b/Engine/Source/Editor/UnrealEd/Private/MaterialGraphNode.cpp
@@ -671,7 +671,7 @@ void UMaterialGraphNode::CreateInputPins()
 		NewPin->PinFriendlyName = SpaceText;
 	}
 
-	TArray<FExpressionInput*> ExpressionInputs = MaterialExpression->GetInputs();
+	TArrayView<FExpressionInput*> ExpressionInputs = MaterialExpression->GetInputsView();
 	const int32 NumExpressionInputs = ExpressionInputs.Num();
 	for (int32 Index = 0; Index < NumExpressionInputs; ++Index)
 	{
diff --git a/Engine/Source/Runtime/Engine/Classes/Materials/MaterialExpression.h b/Engine/Source/Runtime/Engine/Classes/Materials/MaterialExpression.h
--- a/Engine/Source/Runtime/Engine/Classes/Materials/MaterialExpression.h
+++ b/Engine/Source/Runtime/Engine/Classes/Materials/MaterialExpression.h
@@ -345,8 +345,13 @@ class ENGINE_API UMaterialExpression : public UObject
 	 *	@param	Outputs		The TArray of outputs to fill in.
 	 */
 	virtual TArray<FExpressionOutput>& GetOutputs();
+	
 	/** Get the expression inputs supported by this expression (Note: property inputs NOT included). */
-	virtual const TArray<FExpressionInput*> GetInputs();
+	virtual TArrayView<FExpressionInput*> GetInputsView();
+	
+	UE_DEPRECATED(5.2, "Use GetInputsView() instead.")
+	const TArray<FExpressionInput*> GetInputs() { return TArray<FExpressionInput*>{ GetInputsView() }; }
+
 	virtual FExpressionInput* GetInput(int32 InputIndex);
 	virtual FName GetInputName(int32 InputIndex) const;
 	virtual bool IsInputConnectionRequired(int32 InputIndex) const;
@@ -570,6 +575,16 @@ class ENGINE_API UMaterialExpression : public UObject
 	bool ContainsInputLoop(const bool bStopOnFunctionCall = true);
 
 protected:
+	/** Caches the list of expression inputs this expression has. */
+	TArray<FExpressionInput*> CachedInputs;
+
+private:
+	/**
+	 * Helper struct to represent nodes the trail of nodes we're coming from when visiting a new
+	 * expression input node.
+	 */
+	struct FContainsInputLoopInternalExpressionStack;
+
 	/**
 	 * Checks whether any inputs to this expression create a loop by recursively
 	 * calling itself and keeping a list of inputs as expression keys.
@@ -577,7 +592,8 @@ protected:
 	 * @param ExpressionStack    List of expression keys that have been checked already in the current stack
 	 * @param VisitedExpressions List of all expression keys that have been visited
 	 */
-	bool ContainsInputLoopInternal(TArray<class FMaterialExpressionKey>& ExpressionStack, TSet<class FMaterialExpressionKey>& VisitedExpressions, const bool bStopOnFunctionCall);
+	bool ContainsInputLoopInternal(const FContainsInputLoopInternalExpressionStack& ExpressionStack, TSet<class FMaterialExpressionKey>& VisitedExpressions, const bool bStopOnFunctionCall);
+
 #endif // WITH_EDITOR
 };
 
diff --git a/Engine/Source/Runtime/Engine/Classes/Materials/MaterialExpressionBlendMaterialAttributes.h b/Engine/Source/Runtime/Engine/Classes/Materials/MaterialExpressionBlendMaterialAttributes.h
--- a/Engine/Source/Runtime/Engine/Classes/Materials/MaterialExpressionBlendMaterialAttributes.h
+++ b/Engine/Source/Runtime/Engine/Classes/Materials/MaterialExpressionBlendMaterialAttributes.h
@@ -46,8 +46,6 @@ class UMaterialExpressionBlendMaterialAttributes : public UMaterialExpression
 #if WITH_EDITOR
 	virtual int32 Compile(class FMaterialCompiler* Compiler, int32 OutputIndex) override;
 	virtual void GetCaption(TArray<FString>& OutCaptions) const override;
-	virtual const TArray<FExpressionInput*> GetInputs()override;
-	virtual FExpressionInput* GetInput(int32 InputIndex)override;
 	virtual FName GetInputName(int32 InputIndex) const override;
 	virtual bool IsInputConnectionRequired(int32 InputIndex) const override {return true;}
 	virtual bool IsResultMaterialAttributes(int32 OutputIndex) override {return true;}
diff --git a/Engine/Source/Runtime/Engine/Classes/Materials/MaterialExpressionBreakMaterialAttributes.h b/Engine/Source/Runtime/Engine/Classes/Materials/MaterialExpressionBreakMaterialAttributes.h
--- a/Engine/Source/Runtime/Engine/Classes/Materials/MaterialExpressionBreakMaterialAttributes.h
+++ b/Engine/Source/Runtime/Engine/Classes/Materials/MaterialExpressionBreakMaterialAttributes.h
@@ -26,8 +26,6 @@ class UMaterialExpressionBreakMaterialAttributes : public UMaterialExpression
 #if WITH_EDITOR
 	virtual int32 Compile(class FMaterialCompiler* Compiler, int32 OutputIndex) override;
 	virtual void GetCaption(TArray<FString>& OutCaptions) const override;
-	virtual const TArray<FExpressionInput*> GetInputs()override;
-	virtual FExpressionInput* GetInput(int32 InputIndex)override;
 	virtual FName GetInputName(int32 InputIndex) const override;
 	virtual bool IsInputConnectionRequired(int32 InputIndex) const override;
 	virtual uint32 GetInputType(int32 InputIndex) override {return MCT_MaterialAttributes;}
diff --git a/Engine/Source/Runtime/Engine/Classes/Materials/MaterialExpressionComposite.h b/Engine/Source/Runtime/Engine/Classes/Materials/MaterialExpressionComposite.h
--- a/Engine/Source/Runtime/Engine/Classes/Materials/MaterialExpressionComposite.h
+++ b/Engine/Source/Runtime/Engine/Classes/Materials/MaterialExpressionComposite.h
@@ -73,7 +73,7 @@ class UMaterialExpressionComposite : public UMaterialExpression
 	virtual void SetEditableName(const FString& NewName);
 
 	virtual TArray<FExpressionOutput>& GetOutputs() override;
-	virtual const TArray<FExpressionInput*> GetInputs() override;
+	virtual TArrayView<FExpressionInput*> GetInputsView() override;
 	virtual FExpressionInput* GetInput(int32 InputIndex) override;
 	virtual FName GetInputName(int32 InputIndex) const override;
 	virtual uint32 GetInputType(int32 InputIndex) override;
diff --git a/Engine/Source/Runtime/Engine/Classes/Materials/MaterialExpressionCustom.h b/Engine/Source/Runtime/Engine/Classes/Materials/MaterialExpressionCustom.h
--- a/Engine/Source/Runtime/Engine/Classes/Materials/MaterialExpressionCustom.h
+++ b/Engine/Source/Runtime/Engine/Classes/Materials/MaterialExpressionCustom.h
@@ -98,7 +98,7 @@ class UMaterialExpressionCustom : public UMaterialExpression
 	virtual int32 Compile(class FMaterialCompiler* Compiler, int32 OutputIndex) override;
 	virtual void GetCaption(TArray<FString>& OutCaptions) const override;
 	virtual FText GetCreationName() const override { return FText::FromString(TEXT("Custom")); }
-	virtual const TArray<FExpressionInput*> GetInputs() override;
+	virtual TArrayView<FExpressionInput*> GetInputsView() override;
 	virtual FExpressionInput* GetInput(int32 InputIndex) override;
 	virtual FName GetInputName(int32 InputIndex) const override;
 	virtual uint32 GetInputType(int32 InputIndex) override {return MCT_Unknown;}
diff --git a/Engine/Source/Runtime/Engine/Classes/Materials/MaterialExpressionDataDrivenShaderPlatformInfoSwitch.h b/Engine/Source/Runtime/Engine/Classes/Materials/MaterialExpressionDataDrivenShaderPlatformInfoSwitch.h
--- a/Engine/Source/Runtime/Engine/Classes/Materials/MaterialExpressionDataDrivenShaderPlatformInfoSwitch.h
+++ b/Engine/Source/Runtime/Engine/Classes/Materials/MaterialExpressionDataDrivenShaderPlatformInfoSwitch.h
@@ -61,8 +61,6 @@ class UMaterialExpressionDataDrivenShaderPlatformInfoSwitch : public UMaterialEx
 	virtual int32 Compile(class FMaterialCompiler* Compiler, int32 OutputIndex) override;
 	virtual void GetCaption(TArray<FString>& OutCaptions) const override;
 	virtual bool IsInputConnectionRequired(int32 InputIndex) const override;
-	virtual const TArray<FExpressionInput*> GetInputs() override;
-	virtual FExpressionInput* GetInput(int32 InputIndex) override;
 	virtual FName GetInputName(int32 InputIndex) const override;
 	virtual bool IsResultMaterialAttributes(int32 OutputIndex) override;
 	virtual uint32 GetInputType(int32 InputIndex) override { return MCT_Unknown; }
diff --git a/Engine/Source/Runtime/Engine/Classes/Materials/MaterialExpressionFeatureLevelSwitch.h b/Engine/Source/Runtime/Engine/Classes/Materials/MaterialExpressionFeatureLevelSwitch.h
--- a/Engine/Source/Runtime/Engine/Classes/Materials/MaterialExpressionFeatureLevelSwitch.h
+++ b/Engine/Source/Runtime/Engine/Classes/Materials/MaterialExpressionFeatureLevelSwitch.h
@@ -32,8 +32,6 @@ class UMaterialExpressionFeatureLevelSwitch : public UMaterialExpression
 	virtual int32 Compile(class FMaterialCompiler* Compiler, int32 OutputIndex) override;
 	virtual void GetCaption(TArray<FString>& OutCaptions) const override;
 
-	virtual const TArray<FExpressionInput*> GetInputs() override;
-	virtual FExpressionInput* GetInput(int32 InputIndex) override;
 	virtual FName GetInputName(int32 InputIndex) const override;
 	virtual bool IsInputConnectionRequired(int32 InputIndex) const override;
 	virtual bool IsResultMaterialAttributes(int32 OutputIndex) override;
diff --git a/Engine/Source/Runtime/Engine/Classes/Materials/MaterialExpressionGetMaterialAttributes.h b/Engine/Source/Runtime/Engine/Classes/Materials/MaterialExpressionGetMaterialAttributes.h
--- a/Engine/Source/Runtime/Engine/Classes/Materials/MaterialExpressionGetMaterialAttributes.h
+++ b/Engine/Source/Runtime/Engine/Classes/Materials/MaterialExpressionGetMaterialAttributes.h
@@ -38,10 +38,7 @@ class UMaterialExpressionGetMaterialAttributes : public UMaterialExpression
 #if WITH_EDITOR
 	virtual int32 Compile(class FMaterialCompiler* Compiler, int32 OutputIndex) override;
 	virtual bool GenerateHLSLExpression(FMaterialHLSLGenerator& Generator, UE::HLSLTree::FScope& Scope, int32 OutputIndex, UE::HLSLTree::FExpression const*& OutExpression) const override;
-
 	virtual void GetCaption(TArray<FString>& OutCaptions) const override;
-	virtual const TArray<FExpressionInput*> GetInputs()override;
-	virtual FExpressionInput* GetInput(int32 InputIndex)override;
 	virtual FName GetInputName(int32 InputIndex) const override;
 	virtual uint32 GetOutputType(int32 OutputIndex) override;
 	virtual bool IsInputConnectionRequired(int32 InputIndex) const override {return true;}
diff --git a/Engine/Source/Runtime/Engine/Classes/Materials/MaterialExpressionMaterialAttributeLayers.h b/Engine/Source/Runtime/Engine/Classes/Materials/MaterialExpressionMaterialAttributeLayers.h
--- a/Engine/Source/Runtime/Engine/Classes/Materials/MaterialExpressionMaterialAttributeLayers.h
+++ b/Engine/Source/Runtime/Engine/Classes/Materials/MaterialExpressionMaterialAttributeLayers.h
@@ -105,8 +105,6 @@ class UMaterialExpressionMaterialAttributeLayers : public UMaterialExpression
 	virtual int32 Compile(FMaterialCompiler* Compiler, int32 OutputIndex) override;
 	virtual void GetCaption(TArray<FString>& OutCaptions) const override;
 	virtual void GetExpressionToolTip(TArray<FString>& OutToolTip) override;
-	virtual const TArray<FExpressionInput*> GetInputs()override;
-	virtual FExpressionInput* GetInput(int32 InputIndex)override;
 	virtual FName GetInputName(int32 InputIndex) const override;
 	virtual bool IsInputConnectionRequired(int32 InputIndex) const override {return false;}
 	virtual uint32 GetInputType(int32 InputIndex) override;
diff --git a/Engine/Source/Runtime/Engine/Classes/Materials/MaterialExpressionMaterialFunctionCall.h b/Engine/Source/Runtime/Engine/Classes/Materials/MaterialExpressionMaterialFunctionCall.h
--- a/Engine/Source/Runtime/Engine/Classes/Materials/MaterialExpressionMaterialFunctionCall.h
+++ b/Engine/Source/Runtime/Engine/Classes/Materials/MaterialExpressionMaterialFunctionCall.h
@@ -119,7 +119,7 @@ class UMaterialExpressionMaterialFunctionCall : public UMaterialExpression
 #if WITH_EDITOR
 	virtual int32 Compile(class FMaterialCompiler* Compiler, int32 OutputIndex) override;
 	virtual void GetCaption(TArray<FString>& OutCaptions) const override;
-	virtual const TArray<FExpressionInput*> GetInputs() override;
+	virtual TArrayView<FExpressionInput*> GetInputsView() override;
 	virtual FExpressionInput* GetInput(int32 InputIndex) override;
 	virtual FName GetInputName(int32 InputIndex) const override;
 	virtual bool IsInputConnectionRequired(int32 InputIndex) const override;
diff --git a/Engine/Source/Runtime/Engine/Classes/Materials/MaterialExpressionPinBase.h b/Engine/Source/Runtime/Engine/Classes/Materials/MaterialExpressionPinBase.h
--- a/Engine/Source/Runtime/Engine/Classes/Materials/MaterialExpressionPinBase.h
+++ b/Engine/Source/Runtime/Engine/Classes/Materials/MaterialExpressionPinBase.h
@@ -77,7 +77,7 @@ class UMaterialExpressionPinBase : public UMaterialExpression
 
 	//~ Begin UMaterialExpression Interface
 	virtual TArray<FExpressionOutput>& GetOutputs() override;
-	virtual const TArray<FExpressionInput*> GetInputs() override;
+	virtual TArrayView<FExpressionInput*> GetInputsView() override;
 	virtual FExpressionInput* GetInput(int32 InputIndex) override;
 	virtual FName GetInputName(int32 InputIndex) const override;
 	virtual uint32 GetInputType(int32 InputIndex) override;
diff --git a/Engine/Source/Runtime/Engine/Classes/Materials/MaterialExpressionQualitySwitch.h b/Engine/Source/Runtime/Engine/Classes/Materials/MaterialExpressionQualitySwitch.h
--- a/Engine/Source/Runtime/Engine/Classes/Materials/MaterialExpressionQualitySwitch.h
+++ b/Engine/Source/Runtime/Engine/Classes/Materials/MaterialExpressionQualitySwitch.h
@@ -27,8 +27,6 @@ class UMaterialExpressionQualitySwitch : public UMaterialExpression
 	virtual int32 Compile(class FMaterialCompiler* Compiler, int32 OutputIndex) override;
 	virtual void GetCaption(TArray<FString>& OutCaptions) const override;
 
-	virtual const TArray<FExpressionInput*> GetInputs() override;
-	virtual FExpressionInput* GetInput(int32 InputIndex) override;
 	virtual FName GetInputName(int32 InputIndex) const override;
 	virtual bool IsInputConnectionRequired(int32 InputIndex) const override;
 	virtual bool IsResultMaterialAttributes(int32 OutputIndex) override;
diff --git a/Engine/Source/Runtime/Engine/Classes/Materials/MaterialExpressionSetMaterialAttributes.h b/Engine/Source/Runtime/Engine/Classes/Materials/MaterialExpressionSetMaterialAttributes.h
--- a/Engine/Source/Runtime/Engine/Classes/Materials/MaterialExpressionSetMaterialAttributes.h
+++ b/Engine/Source/Runtime/Engine/Classes/Materials/MaterialExpressionSetMaterialAttributes.h
@@ -36,7 +36,7 @@ class UMaterialExpressionSetMaterialAttributes : public UMaterialExpression
 #if WITH_EDITOR
 	virtual int32 Compile(class FMaterialCompiler* Compiler, int32 OutputIndex) override;
 	virtual void GetCaption(TArray<FString>& OutCaptions) const override;
-	virtual const TArray<FExpressionInput*> GetInputs()override;
+	virtual TArrayView<FExpressionInput*> GetInputsView()override;
 	virtual FExpressionInput* GetInput(int32 InputIndex)override;
 	virtual FName GetInputName(int32 InputIndex) const override;
 	virtual uint32 GetInputType(int32 InputIndex) override;
diff --git a/Engine/Source/Runtime/Engine/Classes/Materials/MaterialExpressionShadingPathSwitch.h b/Engine/Source/Runtime/Engine/Classes/Materials/MaterialExpressionShadingPathSwitch.h
--- a/Engine/Source/Runtime/Engine/Classes/Materials/MaterialExpressionShadingPathSwitch.h
+++ b/Engine/Source/Runtime/Engine/Classes/Materials/MaterialExpressionShadingPathSwitch.h
@@ -28,8 +28,6 @@ class UMaterialExpressionShadingPathSwitch : public UMaterialExpression
 	virtual int32 Compile(class FMaterialCompiler* Compiler, int32 OutputIndex) override;
 	virtual void GetCaption(TArray<FString>& OutCaptions) const override;
 
-	virtual const TArray<FExpressionInput*> GetInputs() override;
-	virtual FExpressionInput* GetInput(int32 InputIndex) override;
 	virtual FName GetInputName(int32 InputIndex) const override;
 	virtual bool IsInputConnectionRequired(int32 InputIndex) const override;
 	virtual bool IsResultMaterialAttributes(int32 OutputIndex) override;
diff --git a/Engine/Source/Runtime/Engine/Classes/Materials/MaterialExpressionSparseVolumeTextureObject.h b/Engine/Source/Runtime/Engine/Classes/Materials/MaterialExpressionSparseVolumeTextureObject.h
--- a/Engine/Source/Runtime/Engine/Classes/Materials/MaterialExpressionSparseVolumeTextureObject.h
+++ b/Engine/Source/Runtime/Engine/Classes/Materials/MaterialExpressionSparseVolumeTextureObject.h
@@ -38,7 +38,6 @@ class UMaterialExpressionSparseVolumeTextureTextureObjectParameter : public UMat
 	virtual void GetCaption(TArray<FString>& OutCaptions) const override;
 	virtual int32 Compile(class FMaterialCompiler* Compiler, int32 OutputIndex) override;
 	virtual uint32 GetOutputType(int32 OutputIndex) override;
-	virtual const TArray<FExpressionInput*> GetInputs() override;
 #endif
 	//~ End UMaterialExpression Interface
 };
diff --git a/Engine/Source/Runtime/Engine/Classes/Materials/MaterialExpressionStrata.h b/Engine/Source/Runtime/Engine/Classes/Materials/MaterialExpressionStrata.h
--- a/Engine/Source/Runtime/Engine/Classes/Materials/MaterialExpressionStrata.h
+++ b/Engine/Source/Runtime/Engine/Classes/Materials/MaterialExpressionStrata.h
@@ -186,7 +186,6 @@ class UMaterialExpressionStrataLegacyConversion : public UMaterialExpressionStra
 	virtual FName GetInputName(int32 InputIndex) const override;
 	virtual void GetConnectorToolTip(int32 InputIndex, int32 OutputIndex, TArray<FString>& OutToolTip) override;
 	virtual void PostEditChangeProperty(FPropertyChangedEvent& PropertyChangedEvent) override;
-	virtual const TArray<FExpressionInput*> GetInputs() override;
 
 	bool HasSSS() const;
 	bool HasAnisotropy() const;
@@ -331,7 +330,6 @@ class UMaterialExpressionStrataSlabBSDF : public UMaterialExpressionStrataBSDF
 	virtual FStrataOperator* StrataGenerateMaterialTopologyTree(class FMaterialCompiler* Compiler, class UMaterialExpression* Parent, int32 OutputIndex) override;
 	virtual FName GetInputName(int32 InputIndex) const override;
 	virtual void GetConnectorToolTip(int32 InputIndex, int32 OutputIndex, TArray<FString>& OutToolTip) override;
-	virtual const TArray<FExpressionInput*> GetInputs() override;
 
 	bool HasEdgeColor() const;
 	bool HasFuzz() const;
@@ -404,7 +402,6 @@ class UMaterialExpressionStrataSimpleClearCoatBSDF : public UMaterialExpressionS
 	virtual void GatherStrataMaterialInfo(FStrataMaterialInfo& StrataMaterialInfo, int32 OutputIndex) override;
 	virtual FStrataOperator* StrataGenerateMaterialTopologyTree(class FMaterialCompiler* Compiler, class UMaterialExpression* Parent, int32 OutputIndex) override;
 	virtual FName GetInputName(int32 InputIndex) const override;
-	virtual const TArray<FExpressionInput*> GetInputs() override;
 #endif
 	//~ End UMaterialExpression Interface
 };
diff --git a/Engine/Source/Runtime/Engine/Classes/Materials/MaterialExpressionSwitch.h b/Engine/Source/Runtime/Engine/Classes/Materials/MaterialExpressionSwitch.h
--- a/Engine/Source/Runtime/Engine/Classes/Materials/MaterialExpressionSwitch.h
+++ b/Engine/Source/Runtime/Engine/Classes/Materials/MaterialExpressionSwitch.h
@@ -72,7 +72,7 @@ class UMaterialExpressionSwitch : public UMaterialExpression
 	virtual int32 Compile(class FMaterialCompiler* Compiler, int32 OutputIndex) override;
 	virtual void GetCaption(TArray<FString>& OutCaptions) const override;
 	virtual FText GetCreationName() const override { return FText::FromString(TEXT("Switch")); }
-	virtual const TArray<FExpressionInput*> GetInputs() override;
+	virtual TArrayView<FExpressionInput*> GetInputsView() override;
 	virtual FExpressionInput* GetInput(int32 InputIndex) override;
 	virtual FName GetInputName(int32 InputIndex) const override;
 	virtual uint32 GetInputType(int32 InputIndex) override {return MCT_Unknown;}
diff --git a/Engine/Source/Runtime/Engine/Classes/Materials/MaterialExpressionTextureObjectParameter.h b/Engine/Source/Runtime/Engine/Classes/Materials/MaterialExpressionTextureObjectParameter.h
--- a/Engine/Source/Runtime/Engine/Classes/Materials/MaterialExpressionTextureObjectParameter.h
+++ b/Engine/Source/Runtime/Engine/Classes/Materials/MaterialExpressionTextureObjectParameter.h
@@ -30,7 +30,6 @@ class UMaterialExpressionTextureObjectParameter : public UMaterialExpressionText
 	virtual int32 Compile(class FMaterialCompiler* Compiler, int32 OutputIndex) override;
 	virtual int32 CompilePreview(class FMaterialCompiler* Compiler, int32 OutputIndex) override;
 	virtual uint32 GetOutputType(int32 InputIndex) override {return MCT_Texture;}
-	virtual const TArray<FExpressionInput*> GetInputs() override;
 	virtual bool GenerateHLSLExpression(FMaterialHLSLGenerator& Generator, UE::HLSLTree::FScope& Scope, int32 OutputIndex, UE::HLSLTree::FExpression const*& OutExpression) const override;
 #endif
 	//~ End UMaterialExpression Interface
diff --git a/Engine/Source/Runtime/Engine/Classes/Materials/MaterialExpressionTextureSample.h b/Engine/Source/Runtime/Engine/Classes/Materials/MaterialExpressionTextureSample.h
--- a/Engine/Source/Runtime/Engine/Classes/Materials/MaterialExpressionTextureSample.h
+++ b/Engine/Source/Runtime/Engine/Classes/Materials/MaterialExpressionTextureSample.h
@@ -88,7 +88,7 @@ public:
 
 	//~ Begin UMaterialExpression Interface
 #if WITH_EDITOR
-	virtual const TArray<FExpressionInput*> GetInputs() override;
+	virtual TArrayView<FExpressionInput*> GetInputsView() override;
 	virtual FExpressionInput* GetInput(int32 InputIndex) override;
 	virtual FName GetInputName(int32 InputIndex) const override;
 	virtual int32 GetWidth() const override;
diff --git a/Engine/Source/Runtime/Engine/Private/Materials/Material.cpp b/Engine/Source/Runtime/Engine/Private/Materials/Material.cpp
--- a/Engine/Source/Runtime/Engine/Private/Materials/Material.cpp
+++ b/Engine/Source/Runtime/Engine/Private/Materials/Material.cpp
@@ -6074,7 +6074,7 @@ bool UMaterial::RecursiveGetExpressionChain(
 			else
 			{
 				// Follow all properties.
-				Inputs = InExpression->GetInputs();
+				Inputs = InExpression->GetInputsView();
 				InputsFrequency.Init(InShaderFrequency, Inputs.Num());
 			}
 		}
@@ -6092,7 +6092,7 @@ bool UMaterial::RecursiveGetExpressionChain(
 	}
 	else if ((SetMaterialAttributesExp = Cast<UMaterialExpressionSetMaterialAttributes>(InExpression)) != nullptr)
 	{
-		checkf(!SetMaterialAttributesExp->GetInputs().IsEmpty() && (SetMaterialAttributesExp->GetInputType(0) == MCT_MaterialAttributes), TEXT("There must always be one input at least : the material attribute pin"));
+		checkf(!SetMaterialAttributesExp->GetInputsView().IsEmpty() && (SetMaterialAttributesExp->GetInputType(0) == MCT_MaterialAttributes), TEXT("There must always be one input at least : the material attribute pin"));
 		// Always add the material attribute input, so that we keep on traversing up the property chain : 
 		Inputs.Add(SetMaterialAttributesExp->GetInput(0));
 		InputsFrequency.Add(InShaderFrequency);
@@ -6122,7 +6122,7 @@ bool UMaterial::RecursiveGetExpressionChain(
 			else
 			{
 				// Follow all properties.
-				Inputs = InExpression->GetInputs();
+				Inputs = InExpression->GetInputsView();
 				InputsFrequency.Init(InShaderFrequency, Inputs.Num());
 			}
 		}
@@ -6176,7 +6176,7 @@ bool UMaterial::RecursiveGetExpressionChain(
 		}
 
 		// here we assume ALL inputs to the MaterialFunctionCall are active
-		auto ExprInputs = InExpression->GetInputs();
+		auto ExprInputs = InExpression->GetInputsView();
 		for (int i = 0; i < ExprInputs.Num(); i++)
 		{
 			Inputs.Add(ExprInputs[i]);
@@ -6185,7 +6185,7 @@ bool UMaterial::RecursiveGetExpressionChain(
 	}
 	else
 	{
-		Inputs = InExpression->GetInputs();
+		Inputs = InExpression->GetInputsView();
 		InputsFrequency.Init(InShaderFrequency, Inputs.Num());
 	}
 
@@ -6279,7 +6279,7 @@ void UMaterial::RecursiveUpdateRealtimePreview( UMaterialExpression* InExpressio
 	}
 
 	// We need to examine our inputs. If any of them need realtime preview, so do we.
-	TArray<FExpressionInput*> Inputs = InExpression->GetInputs();
+	TArrayView<FExpressionInput*> Inputs = InExpression->GetInputsView();
 	for (int32 InputIdx = 0; InputIdx < Inputs.Num(); InputIdx++)
 	{
 		FExpressionInput* InnerInput = Inputs[InputIdx];
diff --git a/Engine/Source/Runtime/Engine/Private/Materials/MaterialCachedData.cpp b/Engine/Source/Runtime/Engine/Private/Materials/MaterialCachedData.cpp
--- a/Engine/Source/Runtime/Engine/Private/Materials/MaterialCachedData.cpp
+++ b/Engine/Source/Runtime/Engine/Private/Materials/MaterialCachedData.cpp
@@ -762,7 +762,7 @@ struct FMaterialConnectedPropertiesAnalyzer
 			// If otherwise specified, explore all input expressions to this node.
 			if (bVisitAllInputExpressions)
 			{
-				for (FExpressionInput* ExpressionInput : Expression->GetInputs())
+				for (FExpressionInput* ExpressionInput : Expression->GetInputsView())
 				{
 					PushUnexploredExpression(ExpressionInput->Expression);
 				}
diff --git a/Engine/Source/Runtime/Engine/Private/Materials/MaterialExpressions.cpp b/Engine/Source/Runtime/Engine/Private/Materials/MaterialExpressions.cpp
--- a/Engine/Source/Runtime/Engine/Private/Materials/MaterialExpressions.cpp
+++ b/Engine/Source/Runtime/Engine/Private/Materials/MaterialExpressions.cpp
@@ -737,6 +737,22 @@ UMaterialExpression::UMaterialExpression(const FObjectInitializer& ObjectInitial
 	bCollapsed = true;
 	bShowMaskColorsOnPin = true;
 #endif // WITH_EDITORONLY_DATA
+
+#if WITH_EDITOR
+	// Iterate over the properties of derived expression struct, searching for properties of type FExpressionInput, and add them to the list of cached inputs.
+	for (TFieldIterator<FStructProperty> InputIt(GetClass(), EFieldIteratorFlags::IncludeSuper, EFieldIteratorFlags::ExcludeDeprecated); InputIt; ++InputIt)
+	{
+		FStructProperty* StructProp = *InputIt;
+		if (StructProp->Struct->GetFName() == NAME_ExpressionInput)
+		{
+			for (int32 ArrayIndex = 0; ArrayIndex < StructProp->ArrayDim; ArrayIndex++)
+			{
+				CachedInputs.Add(StructProp->ContainerPtrToValuePtr<FExpressionInput>(this, ArrayIndex));
+			}
+		}
+	}
+	CachedInputs.Shrink();
+#endif
 }
 
 UObject* UMaterialExpression::GetAssetOwner() const
@@ -826,7 +842,7 @@ void UMaterialExpression::CopyMaterialExpressions(const TArray<UMaterialExpressi
 	for( int32 NewExpressionIndex = 0 ; NewExpressionIndex < OutNewExpressions.Num() ; ++NewExpressionIndex )
 	{
 		UMaterialExpression* NewExpression = OutNewExpressions[NewExpressionIndex];
-		const TArray<FExpressionInput*>& ExpressionInputs = NewExpression->GetInputs();
+		TArrayView<FExpressionInput*> ExpressionInputs = NewExpression->GetInputsView();
 		for ( int32 ExpressionInputIndex = 0 ; ExpressionInputIndex < ExpressionInputs.Num() ; ++ExpressionInputIndex )
 		{
 			FExpressionInput* Input = ExpressionInputs[ExpressionInputIndex];
@@ -879,7 +895,7 @@ void UMaterialExpression::Serialize(FStructuredArchive::FRecord Record)
 	const int32 RenderVer = Archive.CustomVer(FRenderingObjectVersion::GUID);
 	const int32 UE5Ver = Archive.CustomVer(FUE5MainStreamObjectVersion::GUID);
 
-	const TArray<FExpressionInput*> Inputs = GetInputs();
+	TArrayView<FExpressionInput*> Inputs = GetInputsView();
 	for (int32 InputIndex = 0; InputIndex < Inputs.Num(); ++InputIndex)
 	{
 		FExpressionInput* Input = Inputs[InputIndex];
@@ -949,7 +965,7 @@ TArray<FProperty*> UMaterialExpression::GetInputPinProperty(int32 PinIndex)
 	TArray<FProperty*> Properties;
 	// Explicit input pins are before property input pins
 	TArray<FProperty*> PropertyInputs = GetPropertyInputs();
-	const int32 NumInputs = GetInputs().Num();
+	const int32 NumInputs = GetInputsView().Num();
 	if (PinIndex < NumInputs)
 	{
 		FExpressionInput* Input = GetInput(PinIndex);
@@ -1463,7 +1479,7 @@ void UMaterialExpression::PostEditChangeProperty(FPropertyChangedEvent& Property
 	if (MemberPropertyThatChanged != nullptr && GraphNode)
 	{
 		int32 PinIndex = -1;
-		const TArray<FExpressionInput*> AllInputs = GetInputs();
+		TArrayView<FExpressionInput*> AllInputs = GetInputsView();
 
 		// Find the expression input this UPROPERTY points to with OverridingInputProperty meta data
 		static FName OverridingInputPropertyMetaData(TEXT("OverridingInputProperty"));
@@ -1654,58 +1670,30 @@ TArray<FExpressionOutput>& UMaterialExpression::GetOutputs()
 }
 
 
-const TArray<FExpressionInput*> UMaterialExpression::GetInputs()
+TArrayView<FExpressionInput*> UMaterialExpression::GetInputsView()
 {
-	TArray<FExpressionInput*> Result;
-	for( TFieldIterator<FStructProperty> InputIt(GetClass(), EFieldIteratorFlags::IncludeSuper,  EFieldIteratorFlags::ExcludeDeprecated) ; InputIt ; ++InputIt )
-	{
-		FStructProperty* StructProp = *InputIt;
-		if( StructProp->Struct->GetFName() == NAME_ExpressionInput)
-		{
-			for (int32 ArrayIndex = 0; ArrayIndex < StructProp->ArrayDim; ArrayIndex++)
-			{
-				Result.Add(StructProp->ContainerPtrToValuePtr<FExpressionInput>(this, ArrayIndex));
-			}
-		}
-	}
-	return Result;
+	return CachedInputs;
 }
 
 
 FExpressionInput* UMaterialExpression::GetInput(int32 InputIndex)
 {
-	int32 Index = 0;
-	for( TFieldIterator<FStructProperty> InputIt(GetClass(), EFieldIteratorFlags::IncludeSuper,  EFieldIteratorFlags::ExcludeDeprecated) ; InputIt ; ++InputIt )
-	{
-		FStructProperty* StructProp = *InputIt;
-		if( StructProp->Struct->GetFName() == NAME_ExpressionInput)
-		{
-			for (int32 ArrayIndex = 0; ArrayIndex < StructProp->ArrayDim; ArrayIndex++)
-			{
-			if( Index == InputIndex )
-			{
-					return StructProp->ContainerPtrToValuePtr<FExpressionInput>(this, ArrayIndex);
-			}
-			Index++;
-		}
-	}
-	}
-
-	return nullptr;
+	TArrayView<FExpressionInput*> Inputs = GetInputsView();
+	return InputIndex < Inputs.Num() ? Inputs[InputIndex] : nullptr;
 }
 
 
 FName UMaterialExpression::GetInputName(int32 InputIndex) const
 {
 	int32 Index = 0;
-	for( TFieldIterator<FStructProperty> InputIt(GetClass(),EFieldIteratorFlags::IncludeSuper,  EFieldIteratorFlags::ExcludeDeprecated) ; InputIt ; ++InputIt )
+	for (TFieldIterator<FStructProperty> InputIt(GetClass(), EFieldIteratorFlags::IncludeSuper, EFieldIteratorFlags::ExcludeDeprecated); InputIt; ++InputIt)
 	{
 		FStructProperty* StructProp = *InputIt;
-		if( StructProp->Struct->GetFName() == NAME_ExpressionInput)
+		if (StructProp->Struct->GetFName() == NAME_ExpressionInput)
 		{
 			for (int32 ArrayIndex = 0; ArrayIndex < StructProp->ArrayDim; ArrayIndex++)
 			{
-				if( Index == InputIndex )
+				if (Index == InputIndex)
 				{
 					FExpressionInput const* Input = StructProp->ContainerPtrToValuePtr<FExpressionInput>(this, ArrayIndex);
 
@@ -2115,7 +2103,7 @@ bool UMaterialExpression::GetAllInputExpressions(TArray<UMaterialExpression*>& I
 		bool bFoundRepeat = false;
 		InputExpressions.Add(this);
 
-		const TArray<FExpressionInput*> Inputs = GetInputs();
+		TArrayView<FExpressionInput*> Inputs = GetInputsView();
 
 		for (int32 Index = 0; Index < Inputs.Num(); Index++)
 		{
@@ -2197,46 +2185,74 @@ bool UMaterialExpression::HasConnectedOutputs() const
 	return bIsConnected;
 }
 
+struct UMaterialExpression::FContainsInputLoopInternalExpressionStack
+{
+	const FMaterialExpressionKey* Key;
+	const FContainsInputLoopInternalExpressionStack* Previous;
+
+	FContainsInputLoopInternalExpressionStack(const FMaterialExpressionKey* Key, const FContainsInputLoopInternalExpressionStack* Previous)
+		: Key{ Key }
+		, Previous{ Previous }
+	{}
+
+	bool Contains(const FMaterialExpressionKey& OtherKey) const
+	{
+		const FContainsInputLoopInternalExpressionStack* Node = this;
+		while (Node->Key)
+		{
+			if (*Node->Key == OtherKey)
+			{
+				return true;
+			}
+			Node = Node->Previous;
+		}
+		return false;
+	}
+};
+
 bool UMaterialExpression::ContainsInputLoop(const bool bStopOnFunctionCall /*= true*/)
 {
-	TArray<FMaterialExpressionKey> ExpressionStack;
+	FContainsInputLoopInternalExpressionStack ExpressionStack{ nullptr, nullptr };
 	TSet<FMaterialExpressionKey> VisitedExpressions;
 	return ContainsInputLoopInternal(ExpressionStack, VisitedExpressions, bStopOnFunctionCall);
 }
 
-bool UMaterialExpression::ContainsInputLoopInternal(TArray<FMaterialExpressionKey>& ExpressionStack, TSet<FMaterialExpressionKey>& VisitedExpressions, const bool bStopOnFunctionCall)
+bool UMaterialExpression::ContainsInputLoopInternal(const FContainsInputLoopInternalExpressionStack& ExpressionStack, TSet<FMaterialExpressionKey>& VisitedExpressions, const bool bStopOnFunctionCall)
 {
-	const TArray<FExpressionInput*> Inputs = GetInputs();
-	for (int32 Index = 0; Index < Inputs.Num(); ++Index)
+	for (FExpressionInput* Input : GetInputsView())
 	{
-		FExpressionInput* Input = Inputs[Index];
-		if (Input->Expression)
+		UMaterialExpression* InputExpression = Input->Expression;
+		if (!InputExpression)
 		{
-			// ContainsInputLoop primarily used to detect safe traversal path for IsResultMaterialAttributes.
-			// In those cases we can bail on a function as the inputs are strongly typed
-			UMaterialExpressionMaterialFunctionCall* FunctionCall = Cast<UMaterialExpressionMaterialFunctionCall>(Input->Expression);
-			UMaterialExpressionMaterialAttributeLayers* Layers = Cast<UMaterialExpressionMaterialAttributeLayers>(Input->Expression);
-			if (bStopOnFunctionCall && (FunctionCall || Layers))
-			{
-				continue;
-			}
+			continue;
+		}
 
-			FMaterialExpressionKey InputExpressionKey(Input->Expression, Input->OutputIndex);
-			if (ExpressionStack.Contains(InputExpressionKey))
+		// ContainsInputLoop primarily used to detect safe traversal path for IsResultMaterialAttributes.
+		// In those cases we can bail on a function as the inputs are strongly typed
+		UMaterialExpressionMaterialFunctionCall* FunctionCall = Cast<UMaterialExpressionMaterialFunctionCall>(InputExpression);
+		UMaterialExpressionMaterialAttributeLayers* Layers = Cast<UMaterialExpressionMaterialAttributeLayers>(InputExpression);
+		if (bStopOnFunctionCall && (FunctionCall || Layers))
+		{
+			continue;
+		}
+
+		// A cycle is detected if one of this node's inputs leads back to a node we're coming from.
+		FMaterialExpressionKey InputExpressionKey(InputExpression, Input->OutputIndex);
+		if (ExpressionStack.Contains(InputExpressionKey))
+		{
+			return true;
+		}
+
+		// Prevent recurring visits to expressions we've already checked
+		if (!VisitedExpressions.Contains(InputExpressionKey))
+		{
+			VisitedExpressions.Add(InputExpressionKey);
+
+			FContainsInputLoopInternalExpressionStack ExpressionStackWithThisInput{ &InputExpressionKey, &ExpressionStack };
+			if (InputExpression->ContainsInputLoopInternal(ExpressionStackWithThisInput, VisitedExpressions, bStopOnFunctionCall))
 			{
 				return true;
 			}
-			// prevent recurring visits to expressions we've already checked
-			else if (!VisitedExpressions.Contains(InputExpressionKey))
-			{
-				VisitedExpressions.Add(InputExpressionKey);
-				ExpressionStack.Add(InputExpressionKey);
-				if (Input->Expression->ContainsInputLoopInternal(ExpressionStack, VisitedExpressions, bStopOnFunctionCall))
-				{
-					return true;
-				}
-				ExpressionStack.Pop();
-			}
 		}
 	}
 
@@ -2510,18 +2526,16 @@ void UMaterialExpressionTextureSample::PostLoad()
 	}
 }
 
-const TArray<FExpressionInput*> UMaterialExpressionTextureSample::GetInputs()
+TArrayView<FExpressionInput*> UMaterialExpressionTextureSample::GetInputsView()
 {
-	TArray<FExpressionInput*> OutInputs;
-
+	CachedInputs.Empty();
 	// todo: we should remove GetInputs() and make this the common code for all expressions
 	uint32 InputIndex = 0;
-	while(FExpressionInput* Ptr = GetInput(InputIndex++))
+	while (FExpressionInput* Ptr = GetInput(InputIndex++))
 	{
-		OutInputs.Add(Ptr);
+		CachedInputs.Add(Ptr);
 	}
-
-	return OutInputs;
+	return CachedInputs;
 }
 
 // this define is only used for the following function
@@ -3525,7 +3539,7 @@ int32 UMaterialExpressionRuntimeVirtualTextureReplace::Compile(class FMaterialCo
 
 bool UMaterialExpressionRuntimeVirtualTextureReplace::IsResultMaterialAttributes(int32 OutputIndex)
 {
-	for (FExpressionInput* ExpressionInput : GetInputs())
+	for (FExpressionInput* ExpressionInput : GetInputsView())
 	{
 		if (ExpressionInput->GetTracedInput().Expression && !ExpressionInput->Expression->ContainsInputLoop() && ExpressionInput->Expression->IsResultMaterialAttributes(ExpressionInput->OutputIndex))
 		{
@@ -3585,7 +3599,7 @@ int32 UMaterialExpressionVirtualTextureFeatureSwitch::Compile(class FMaterialCom
 
 bool UMaterialExpressionVirtualTextureFeatureSwitch::IsResultMaterialAttributes(int32 OutputIndex)
 {
-	for (FExpressionInput* ExpressionInput : GetInputs())
+	for (FExpressionInput* ExpressionInput : GetInputsView())
 	{
 		if (ExpressionInput->GetTracedInput().Expression && !ExpressionInput->Expression->ContainsInputLoop() && ExpressionInput->Expression->IsResultMaterialAttributes(ExpressionInput->OutputIndex))
 		{
@@ -3783,7 +3797,13 @@ UMaterialExpressionTextureObjectParameter::UMaterialExpressionTextureObjectParam
 
 	Outputs.Empty();
 	Outputs.Add(FExpressionOutput(TEXT("")));
+
 #endif // WITH_EDITORONLY_DATA
+
+#if WITH_EDITOR
+	// Hide the texture coordinate input
+	CachedInputs.Empty();
+#endif
 }
 
 #if WITH_EDITOR
@@ -3804,12 +3824,6 @@ void UMaterialExpressionTextureObjectParameter::GetCaption(TArray<FString>& OutC
 	OutCaptions.Add(FString::Printf(TEXT("'%s'"), *ParameterName.ToString()));
 }
 
-const TArray<FExpressionInput*> UMaterialExpressionTextureObjectParameter::GetInputs()
-{
-	// Hide the texture coordinate input
-	return TArray<FExpressionInput*>();
-}
-
 int32 UMaterialExpressionTextureObjectParameter::Compile(class FMaterialCompiler* Compiler, int32 OutputIndex)
 {
 	FString ErrorMessage;
@@ -6775,6 +6789,11 @@ UMaterialExpressionBreakMaterialAttributes::UMaterialExpressionBreakMaterialAttr
 	Outputs.Add(FExpressionOutput(TEXT("ShadingModel"), 0, 0, 0, 0, 0));
 	Outputs.Add(FExpressionOutput(TEXT("Displacement"), 1, 1, 0, 0, 0));
 #endif
+
+#if WITH_EDITOR
+	CachedInputs.Empty();
+	CachedInputs.Add(&MaterialAttributes);
+#endif
 }
 
 void UMaterialExpressionBreakMaterialAttributes::Serialize(FStructuredArchive::FRecord Record)
@@ -6878,23 +6897,6 @@ void UMaterialExpressionBreakMaterialAttributes::GetCaption(TArray<FString>& Out
 	OutCaptions.Add(TEXT("BreakMaterialAttributes"));
 }
 
-const TArray<FExpressionInput*> UMaterialExpressionBreakMaterialAttributes::GetInputs()
-{
-	TArray<FExpressionInput*> Result;
-	Result.Add(&MaterialAttributes);
-	return Result;
-}
-
-
-FExpressionInput* UMaterialExpressionBreakMaterialAttributes::GetInput(int32 InputIndex)
-{
-	if( 0 == InputIndex )
-	{
-		return &MaterialAttributes;
-	}
-
-	return nullptr;
-}
 
 FName UMaterialExpressionBreakMaterialAttributes::GetInputName(int32 InputIndex) const
 {
@@ -6952,6 +6954,9 @@ UMaterialExpressionGetMaterialAttributes::UMaterialExpressionGetMaterialAttribut
 	// Add default output pins
 	Outputs.Reset();
 	Outputs.Add(FExpressionOutput(TEXT("MaterialAttributes"), 0, 0, 0, 0, 0));
+
+	CachedInputs.Empty();
+	CachedInputs.Push(&MaterialAttributes);
 #endif
 }
 
@@ -6998,23 +7003,6 @@ void UMaterialExpressionGetMaterialAttributes::GetCaption(TArray<FString>& OutCa
 	OutCaptions.Add(TEXT("GetMaterialAttributes"));
 }
 
-const TArray<FExpressionInput*> UMaterialExpressionGetMaterialAttributes::GetInputs()
-{
-	TArray<FExpressionInput*> Result;
-	Result.Add(&MaterialAttributes);
-	return Result;
-}
-
-FExpressionInput* UMaterialExpressionGetMaterialAttributes::GetInput(int32 InputIndex)
-{
-	if (InputIndex == 0)
-	{
-		return &MaterialAttributes;
-	}
-
-	return nullptr;
-}
-
 FName UMaterialExpressionGetMaterialAttributes::GetInputName(int32 InputIndex) const
 {
 	return NAME_None;
@@ -7285,14 +7273,15 @@ void UMaterialExpressionSetMaterialAttributes::GetCaption(TArray<FString>& OutCa
 	OutCaptions.Add(TEXT("SetMaterialAttributes"));
 }
 
-const TArray<FExpressionInput*> UMaterialExpressionSetMaterialAttributes::GetInputs()
+TArrayView<FExpressionInput*> UMaterialExpressionSetMaterialAttributes::GetInputsView()
 {
-	TArray<FExpressionInput*> Result;
+	CachedInputs.Empty();
+	CachedInputs.Reserve(Inputs.Num());
 	for (FExpressionInput& Input : Inputs)
 	{
-		Result.Add(&Input);
+		CachedInputs.Push(&Input);
 	}
-	return Result;
+	return CachedInputs;
 }
 
 FExpressionInput* UMaterialExpressionSetMaterialAttributes::GetInput(int32 InputIndex)
@@ -7461,6 +7450,14 @@ UMaterialExpressionBlendMaterialAttributes::UMaterialExpressionBlendMaterialAttr
 	Outputs.Reset();
 	Outputs.Add(FExpressionOutput(TEXT(""), 0, 0, 0, 0, 0));
 #endif
+
+#if WITH_EDITOR
+	CachedInputs.Empty();
+	CachedInputs.Reserve(3);
+	CachedInputs.Add(&A);
+	CachedInputs.Add(&B);
+	CachedInputs.Add(&Alpha);
+#endif
 }
 
 #if WITH_EDITOR
@@ -7525,33 +7522,6 @@ void UMaterialExpressionBlendMaterialAttributes::GetCaption(TArray<FString>& Out
 	OutCaptions.Add(TEXT("BlendMaterialAttributes"));
 }
 
-const TArray<FExpressionInput*> UMaterialExpressionBlendMaterialAttributes::GetInputs()
-{
-	TArray<FExpressionInput*> Result;
-	Result.Add(&A);
-	Result.Add(&B);
-	Result.Add(&Alpha);
-	return Result;
-}
-
-FExpressionInput* UMaterialExpressionBlendMaterialAttributes::GetInput(int32 InputIndex)
-{
-	if (InputIndex == 0)
-	{
-		return &A;
-	}
-	else if (InputIndex == 1)
-	{
-		return &B;
-	}
-	else if (InputIndex == 2)
-	{
-		return &Alpha;
-	}
-
-	return nullptr;
-}
-
 FName UMaterialExpressionBlendMaterialAttributes::GetInputName(int32 InputIndex) const
 {
 	FName Name;
@@ -7646,6 +7616,9 @@ UMaterialExpressionMaterialAttributeLayers::UMaterialExpressionMaterialAttribute
 
 #if WITH_EDITOR
 	DefaultLayers.AddDefaultBackgroundLayer();
+	
+	CachedInputs.Empty();
+	CachedInputs.Push(&Input);
 #endif
 #if WITH_EDITORONLY_DATA
 	MenuCategories.Add(ConstructorStatics.NAME_MaterialAttributes);
@@ -8076,23 +8049,6 @@ void UMaterialExpressionMaterialAttributeLayers::GetExpressionToolTip(TArray<FSt
 	ConvertToMultilineToolTip(TEXT("Evaluates the active material layer stack and outputs the merged attributes."), 40, OutToolTip);
 }
 
-const TArray<FExpressionInput*> UMaterialExpressionMaterialAttributeLayers::GetInputs()
-{
-	TArray<FExpressionInput*> Result;
-	Result.Add(&Input);
-	return Result;
-}
-
-FExpressionInput* UMaterialExpressionMaterialAttributeLayers::GetInput(int32 InputIndex)
-{
-	if (InputIndex == 0)
-	{
-		return &Input;
-	}
-
-	return nullptr;
-}
-
 FName UMaterialExpressionMaterialAttributeLayers::GetInputName(int32 InputIndex) const
 {
 	return NAME_None;
@@ -9491,30 +9447,6 @@ void UMaterialExpressionQualitySwitch::GetCaption(TArray<FString>& OutCaptions)
 	OutCaptions.Add(FString(TEXT("Quality Switch")));
 }
 
-const TArray<FExpressionInput*> UMaterialExpressionQualitySwitch::GetInputs()
-{
-	TArray<FExpressionInput*> OutInputs;
-
-	OutInputs.Add(&Default);
-
-	for (int32 InputIndex = 0; InputIndex < UE_ARRAY_COUNT(Inputs); InputIndex++)
-	{
-		OutInputs.Add(&Inputs[InputIndex]);
-	}
-
-	return OutInputs;
-}
-
-FExpressionInput* UMaterialExpressionQualitySwitch::GetInput(int32 InputIndex)
-{
-	if (InputIndex == 0)
-	{
-		return &Default;
-	}
-
-	return &Inputs[InputIndex - 1];
-}
-
 FName UMaterialExpressionQualitySwitch::GetInputName(int32 InputIndex) const
 {
 	if (InputIndex == 0)
@@ -9533,7 +9465,7 @@ bool UMaterialExpressionQualitySwitch::IsInputConnectionRequired(int32 InputInde
 bool UMaterialExpressionQualitySwitch::IsResultMaterialAttributes(int32 OutputIndex)
 {
 	check(OutputIndex == 0);
-	TArray<FExpressionInput*> ExpressionInputs = GetInputs();
+	TArrayView<FExpressionInput*> ExpressionInputs = GetInputsView();
 
 	for (FExpressionInput* ExpressionInput : ExpressionInputs)
 	{
@@ -9637,30 +9569,6 @@ void UMaterialExpressionFeatureLevelSwitch::GetCaption(TArray<FString>& OutCapti
 	OutCaptions.Add(FString(TEXT("Feature Level Switch")));
 }
 
-const TArray<FExpressionInput*> UMaterialExpressionFeatureLevelSwitch::GetInputs()
-{
-	TArray<FExpressionInput*> OutInputs;
-
-	OutInputs.Add(&Default);
-
-	for (int32 InputIndex = 0; InputIndex < UE_ARRAY_COUNT(Inputs); InputIndex++)
-	{
-		OutInputs.Add(&Inputs[InputIndex]);
-	}
-
-	return OutInputs;
-}
-
-FExpressionInput* UMaterialExpressionFeatureLevelSwitch::GetInput(int32 InputIndex)
-{
-	if (InputIndex == 0)
-	{
-		return &Default;
-	}
-
-	return &Inputs[InputIndex - 1];
-}
-
 FName UMaterialExpressionFeatureLevelSwitch::GetInputName(int32 InputIndex) const
 {
 	if (InputIndex == 0)
@@ -9682,7 +9590,7 @@ bool UMaterialExpressionFeatureLevelSwitch::IsInputConnectionRequired(int32 Inpu
 bool UMaterialExpressionFeatureLevelSwitch::IsResultMaterialAttributes(int32 OutputIndex)
 {
 	check(OutputIndex == 0);
-	TArray<FExpressionInput*> ExpressionInputs = GetInputs();
+	TArrayView<FExpressionInput*> ExpressionInputs = GetInputsView();
 
 	for (FExpressionInput* ExpressionInput : ExpressionInputs)
 	{
@@ -9828,27 +9736,6 @@ bool UMaterialExpressionDataDrivenShaderPlatformInfoSwitch::IsInputConnectionReq
 	return true;
 }
 
-const TArray<FExpressionInput*> UMaterialExpressionDataDrivenShaderPlatformInfoSwitch::GetInputs()
-{
-	TArray<FExpressionInput*> Result;
-	Result.Add(&InputTrue);
-	Result.Add(&InputFalse);
-	return Result;
-}
-
-FExpressionInput* UMaterialExpressionDataDrivenShaderPlatformInfoSwitch::GetInput(int32 InputIndex)
-{
-	if (InputIndex == 0)
-	{
-		return &InputTrue;
-	}
-	else if (InputIndex == 1)
-	{
-		return &InputFalse;
-	}
-	return nullptr;
-}
-
 FName UMaterialExpressionDataDrivenShaderPlatformInfoSwitch::GetInputName(int32 InputIndex) const
 {
 	TStringBuilder<128> Condition;
@@ -9891,7 +9778,7 @@ FName UMaterialExpressionDataDrivenShaderPlatformInfoSwitch::GetInputName(int32
 bool UMaterialExpressionDataDrivenShaderPlatformInfoSwitch::IsResultMaterialAttributes(int32 OutputIndex)
 {
 	check(OutputIndex == 0);
-	TArray<FExpressionInput*> ExpressionInputs = GetInputs();
+	TArrayView<FExpressionInput*> ExpressionInputs = GetInputsView();
 
 	for (FExpressionInput* ExpressionInput : ExpressionInputs)
 	{
@@ -9987,30 +9874,6 @@ void UMaterialExpressionShadingPathSwitch::GetCaption(TArray<FString>& OutCaptio
 	OutCaptions.Add(FString(TEXT("Shading Path Switch")));
 }
 
-const TArray<FExpressionInput*> UMaterialExpressionShadingPathSwitch::GetInputs()
-{
-	TArray<FExpressionInput*> OutInputs;
-
-	OutInputs.Add(&Default);
-
-	for (int32 InputIndex = 0; InputIndex < UE_ARRAY_COUNT(Inputs); InputIndex++)
-	{
-		OutInputs.Add(&Inputs[InputIndex]);
-	}
-
-	return OutInputs;
-}
-
-FExpressionInput* UMaterialExpressionShadingPathSwitch::GetInput(int32 InputIndex)
-{
-	if (InputIndex == 0)
-	{
-		return &Default;
-	}
-
-	return &Inputs[InputIndex - 1];
-}
-
 FName UMaterialExpressionShadingPathSwitch::GetInputName(int32 InputIndex) const
 {
 	if (InputIndex == 0)
@@ -10031,9 +9894,7 @@ bool UMaterialExpressionShadingPathSwitch::IsInputConnectionRequired(int32 Input
 bool UMaterialExpressionShadingPathSwitch::IsResultMaterialAttributes(int32 OutputIndex)
 {
 	check(OutputIndex == 0);
-	TArray<FExpressionInput*> ExpressionInputs = GetInputs();
-
-	for (FExpressionInput* ExpressionInput : ExpressionInputs)
+	for (FExpressionInput* ExpressionInput : GetInputsView())
 	{
 		// If there is a loop anywhere in this expression's inputs then we can't risk checking them
 		if (ExpressionInput->Expression && !ExpressionInput->Expression->ContainsInputLoop() && ExpressionInput->Expression->IsResultMaterialAttributes(ExpressionInput->OutputIndex))
@@ -11983,22 +11844,22 @@ TArray<FExpressionOutput>& UMaterialExpressionComposite::GetOutputs()
 	return Outputs;
 }
 
-const TArray<FExpressionInput*> UMaterialExpressionComposite::GetInputs()
+TArrayView<FExpressionInput*> UMaterialExpressionComposite::GetInputsView()
 {
-	TArray<FExpressionInput*> ExpressionInputs;
-
 	// InputExpressions may be null if we are using the default object
+	CachedInputs.Empty();
 	if (InputExpressions)
 	{
+		CachedInputs.Reserve(InputExpressions->ReroutePins.Num());
 		for (const FCompositeReroute& ReroutePin : InputExpressions->ReroutePins)
 		{
 			if (ReroutePin.Expression)
 			{
-				ExpressionInputs.Add(ReroutePin.Expression->GetInput(0));
+				CachedInputs.Add(ReroutePin.Expression->GetInput(0));
 			}
 		}
 	}
-	return ExpressionInputs;
+	return CachedInputs;
 }
 
 FExpressionInput* UMaterialExpressionComposite::GetInput(int32 InputIndex)
@@ -12233,20 +12094,21 @@ TArray<FExpressionOutput>& UMaterialExpressionPinBase::GetOutputs()
 	return Outputs;
 }
 
-const TArray<FExpressionInput*> UMaterialExpressionPinBase::GetInputs()
+TArrayView<FExpressionInput*> UMaterialExpressionPinBase::GetInputsView()
 {
-	TArray<FExpressionInput*> ExpressionInputs;
+	CachedInputs.Empty();
 	if (PinDirection == EGPD_Input)
 	{
+		CachedInputs.Reserve(ReroutePins.Num());
 		for (const FCompositeReroute& ReroutePin : ReroutePins)
 		{
 			if (ReroutePin.Expression)
 			{
-				ExpressionInputs.Add(ReroutePin.Expression->GetInput(0));
+				CachedInputs.Add(ReroutePin.Expression->GetInput(0));
 			}
 		}
 	}
-	return ExpressionInputs;
+	return CachedInputs;
 }
 
 FExpressionInput* UMaterialExpressionPinBase::GetInput(int32 InputIndex)
@@ -13069,7 +12931,7 @@ int32 UMaterialExpressionDistanceFieldsRenderingSwitch::Compile(class FMaterialC
 
 bool UMaterialExpressionDistanceFieldsRenderingSwitch::IsResultMaterialAttributes(int32 OutputIndex)
 {
-	for (FExpressionInput* ExpressionInput : GetInputs())
+	for (FExpressionInput* ExpressionInput : GetInputsView())
 	{
 		if (ExpressionInput->GetTracedInput().Expression && !ExpressionInput->Expression->ContainsInputLoop() && ExpressionInput->Expression->IsResultMaterialAttributes(ExpressionInput->OutputIndex))
 		{
@@ -13293,14 +13155,15 @@ void UMaterialExpressionCustom::GetCaption(TArray<FString>& OutCaptions) const
 }
 
 
-const TArray<FExpressionInput*> UMaterialExpressionCustom::GetInputs()
+TArrayView<FExpressionInput*> UMaterialExpressionCustom::GetInputsView()
 {
-	TArray<FExpressionInput*> Result;
+	CachedInputs.Empty();
+	CachedInputs.Reserve(Inputs.Num());
 	for( int32 i = 0; i < Inputs.Num(); i++ )
 	{
-		Result.Add(&Inputs[i].Input);
+		CachedInputs.Add(&Inputs[i].Input);
 	}
-	return Result;
+	return CachedInputs;
 }
 
 FExpressionInput* UMaterialExpressionCustom::GetInput(int32 InputIndex)
@@ -13687,16 +13550,17 @@ void UMaterialExpressionSwitch::GetCaption(TArray<FString>& OutCaptions) const
 }
 
 
-const TArray<FExpressionInput*> UMaterialExpressionSwitch::GetInputs()
+TArrayView<FExpressionInput*> UMaterialExpressionSwitch::GetInputsView()
 {
-	TArray<FExpressionInput*> Result;
-	Result.Add(&SwitchValue);
-	Result.Add(&Default);
+	CachedInputs.Empty();
+	CachedInputs.Reserve(2 + Inputs.Num());
+	CachedInputs.Add(&SwitchValue);
+	CachedInputs.Add(&Default);
 	for (int32 i = 0; i < Inputs.Num(); i++)
 	{
-		Result.Add(&Inputs[i].Input);
+		CachedInputs.Add(&Inputs[i].Input);
 	}
-	return Result;
+	return CachedInputs;
 }
 
 FExpressionInput* UMaterialExpressionSwitch::GetInput(int32 InputIndex)
@@ -16199,14 +16063,15 @@ void UMaterialExpressionMaterialFunctionCall::GetCaption(TArray<FString>& OutCap
 	}
 }
 
-const TArray<FExpressionInput*> UMaterialExpressionMaterialFunctionCall::GetInputs()
+TArrayView<FExpressionInput*> UMaterialExpressionMaterialFunctionCall::GetInputsView()
 {
-	TArray<FExpressionInput*> Result;
+	CachedInputs.Empty();
+	CachedInputs.Reserve(FunctionInputs.Num());
 	for (int32 i = 0; i < FunctionInputs.Num(); i++)
 	{
-		Result.Add(&FunctionInputs[i].Input);
+		CachedInputs.Add(&FunctionInputs[i].Input);
 	}
-	return Result;
+	return CachedInputs;
 }
 
 FExpressionInput* UMaterialExpressionMaterialFunctionCall::GetInput(int32 InputIndex)
@@ -16556,7 +16421,7 @@ void UMaterialExpressionMaterialFunctionCall::UpdateFromFunctionResource(bool bR
 static void FixupReferencingInputs(
 	const TArray<FFunctionExpressionOutput>& NewOutputs,
 	const TArray<FFunctionExpressionOutput>& OriginalOutputs,
-	const TArray<FExpressionInput*>& Inputs, 
+	TArrayView<FExpressionInput*> Inputs, 
 	UMaterialExpressionMaterialFunctionCall* FunctionExpression,
 	bool bMatchByName)
 {
@@ -16607,7 +16472,7 @@ void UMaterialExpressionMaterialFunctionCall::FixupReferencingExpressions(
 		UMaterialExpression* CurrentExpression = Expressions[ExpressionIndex];
 		if (CurrentExpression)
 		{
-			TArray<FExpressionInput*> Inputs = CurrentExpression->GetInputs();
+			TArrayView<FExpressionInput*> Inputs = CurrentExpression->GetInputsView();
 			FixupReferencingInputs(NewOutputs, OriginalOutputs, Inputs, this, bMatchByName);
 		}
 	}
@@ -17754,7 +17619,7 @@ int32 UMaterialExpressionMaterialProxyReplace::Compile(class FMaterialCompiler*
 
 bool UMaterialExpressionMaterialProxyReplace::IsResultMaterialAttributes(int32 OutputIndex)
 {
-	for (FExpressionInput* ExpressionInput : GetInputs())
+	for (FExpressionInput* ExpressionInput : GetInputsView())
 	{
 		if (ExpressionInput->GetTracedInput().Expression && !ExpressionInput->Expression->ContainsInputLoop() && ExpressionInput->Expression->IsResultMaterialAttributes(ExpressionInput->OutputIndex))
 		{
@@ -23057,31 +22922,6 @@ int32 UMaterialExpressionStrataLegacyConversion::Compile(class FMaterialCompiler
 	return OutputCodeChunk;
 }
 
-const TArray<FExpressionInput*> UMaterialExpressionStrataLegacyConversion::GetInputs()
-{
-	TArray<FExpressionInput*> Result;
-	Result.Add(&BaseColor);
-	Result.Add(&Metallic);
-	Result.Add(&Specular);
-	Result.Add(&Roughness);
-	Result.Add(&Anisotropy);
-	Result.Add(&EmissiveColor);
-	Result.Add(&Normal);
-	Result.Add(&Tangent);
-	Result.Add(&SubSurfaceColor);
-	Result.Add(&ClearCoat);
-	Result.Add(&ClearCoatRoughness);
-	Result.Add(&Opacity);
-	Result.Add(&TransmittanceColor);
-	Result.Add(&WaterScatteringCoefficients);
-	Result.Add(&WaterAbsorptionCoefficients);
-	Result.Add(&WaterPhaseG);
-	Result.Add(&ColorScaleBehindWater);
-	Result.Add(&ClearCoatNormal);
-	Result.Add(&CustomTangent);
-	Result.Add(&ShadingModel);
-	return Result;
-}
 
 void UMaterialExpressionStrataLegacyConversion::PostEditChangeProperty(FPropertyChangedEvent& PropertyChangedEvent)
 {
@@ -23396,31 +23236,6 @@ int32 UMaterialExpressionStrataSlabBSDF::Compile(class FMaterialCompiler* Compil
 	return OutputCodeChunk;
 }
 
-
-const TArray<FExpressionInput*> UMaterialExpressionStrataSlabBSDF::GetInputs()
-{
-	TArray<FExpressionInput*> Result;
-	Result.Add(&DiffuseAlbedo);
-	Result.Add(&F0);
-	Result.Add(&F90);
-	Result.Add(&Roughness);
-	Result.Add(&Anisotropy);
-	Result.Add(&Normal);
-	Result.Add(&Tangent);
-	Result.Add(&SSSMFP);
-	Result.Add(&SSSMFPScale);
-	Result.Add(&SSSPhaseAnisotropy);
-	Result.Add(&EmissiveColor);
-	Result.Add(&SecondRoughness);
-	Result.Add(&SecondRoughnessWeight);
-	Result.Add(&FuzzRoughness);
-	Result.Add(&FuzzAmount);
-	Result.Add(&FuzzColor);
-	Result.Add(&GlintValue);
-	Result.Add(&GlintUV);
-	return Result;
-}
-
 void UMaterialExpressionStrataSlabBSDF::GetCaption(TArray<FString>& OutCaptions) const
 {
 	OutCaptions.Add(TEXT("Substrate Slab BSDF"));
@@ -23770,20 +23585,6 @@ int32 UMaterialExpressionStrataSimpleClearCoatBSDF::Compile(class FMaterialCompi
 	return OutputCodeChunk;
 }
 
-
-const TArray<FExpressionInput*> UMaterialExpressionStrataSimpleClearCoatBSDF::GetInputs()
-{
-	TArray<FExpressionInput*> Result;
-	Result.Add(&DiffuseAlbedo);
-	Result.Add(&F0);
-	Result.Add(&Roughness);
-	Result.Add(&ClearCoatCoverage);
-	Result.Add(&ClearCoatRoughness);
-	Result.Add(&Normal);
-	Result.Add(&EmissiveColor);
-	return Result;
-}
-
 void UMaterialExpressionStrataSimpleClearCoatBSDF::GetCaption(TArray<FString>& OutCaptions) const
 {
 	OutCaptions.Add(TEXT("Substrate Simple Clear Coat"));
@@ -26377,7 +26178,12 @@ UMaterialExpressionSparseVolumeTextureTextureObjectParameter::UMaterialExpressio
 
 	Outputs.Empty();
 	Outputs.Add(FExpressionOutput(TEXT("")));
-#endif // WITH_EDITORONLY_DATA
+#endif
+
+#if WITH_EDITOR
+	// Hide the texture coordinate input
+	CachedInputs.Empty();
+#endif
 }
 
 #if WITH_EDITOR
@@ -26404,12 +26210,6 @@ uint32 UMaterialExpressionSparseVolumeTextureTextureObjectParameter::GetOutputTy
 	return MCT_SparseVolumeTexture;
 }
 
-const TArray<FExpressionInput*> UMaterialExpressionSparseVolumeTextureTextureObjectParameter::GetInputs()
-{
-	// Hide the texture coordinate input
-	return TArray<FExpressionInput*>();
-}
-
 #endif // WITH_EDITOR
 
 
diff --git a/Engine/Source/Runtime/Engine/Private/Materials/MaterialHLSLGenerator.cpp b/Engine/Source/Runtime/Engine/Private/Materials/MaterialHLSLGenerator.cpp
--- a/Engine/Source/Runtime/Engine/Private/Materials/MaterialHLSLGenerator.cpp
+++ b/Engine/Source/Runtime/Engine/Private/Materials/MaterialHLSLGenerator.cpp
@@ -385,7 +385,7 @@ int32 FMaterialHLSLGenerator::FindInputIndex(const FExpressionInput* Input) cons
 	int32 Index = INDEX_NONE;
 	if (OwnerMaterialExpression)
 	{
-		const TArray<FExpressionInput*> Inputs = OwnerMaterialExpression->GetInputs();
+		TArrayView<FExpressionInput*> Inputs = OwnerMaterialExpression->GetInputsView();
 		Index = Inputs.Find(const_cast<FExpressionInput*>(Input));
 	}
 	return Index;
diff --git a/Engine/Source/Runtime/Landscape/Classes/Materials/MaterialExpressionLandscapeGrassOutput.h b/Engine/Source/Runtime/Landscape/Classes/Materials/MaterialExpressionLandscapeGrassOutput.h
--- a/Engine/Source/Runtime/Landscape/Classes/Materials/MaterialExpressionLandscapeGrassOutput.h
+++ b/Engine/Source/Runtime/Landscape/Classes/Materials/MaterialExpressionLandscapeGrassOutput.h
@@ -48,7 +48,7 @@ class LANDSCAPE_API UMaterialExpressionLandscapeGrassOutput : public UMaterialEx
 #if WITH_EDITOR
 	virtual int32 Compile(class FMaterialCompiler* Compiler, int32 OutputIndex) override;
 	virtual void GetCaption(TArray<FString>& OutCaptions) const override;
-	virtual const TArray<FExpressionInput*> GetInputs() override;
+	virtual TArrayView<FExpressionInput*> GetInputsView() override;
 	virtual FExpressionInput* GetInput(int32 InputIndex) override;
 	virtual FName GetInputName(int32 InputIndex) const override;
 	void ValidateInputName(FGrassInput& Input) const;
diff --git a/Engine/Source/Runtime/Landscape/Classes/Materials/MaterialExpressionLandscapeLayerBlend.h b/Engine/Source/Runtime/Landscape/Classes/Materials/MaterialExpressionLandscapeLayerBlend.h
--- a/Engine/Source/Runtime/Landscape/Classes/Materials/MaterialExpressionLandscapeLayerBlend.h
+++ b/Engine/Source/Runtime/Landscape/Classes/Materials/MaterialExpressionLandscapeLayerBlend.h
@@ -80,7 +80,7 @@ class LANDSCAPE_API UMaterialExpressionLandscapeLayerBlend : public UMaterialExp
 	virtual bool IsResultMaterialAttributes(int32 OutputIndex) override;
 	virtual int32 Compile(class FMaterialCompiler* Compiler, int32 OutputIndex) override;
 	virtual void GetCaption(TArray<FString>& OutCaptions) const override;
-	virtual const TArray<FExpressionInput*> GetInputs() override;
+	virtual TArrayView<FExpressionInput*> GetInputsView() override;
 	virtual FExpressionInput* GetInput(int32 InputIndex) override;
 	virtual FName GetInputName(int32 InputIndex) const override;
 
diff --git a/Engine/Source/Runtime/Landscape/Classes/Materials/MaterialExpressionLandscapePhysicalMaterialOutput.h b/Engine/Source/Runtime/Landscape/Classes/Materials/MaterialExpressionLandscapePhysicalMaterialOutput.h
--- a/Engine/Source/Runtime/Landscape/Classes/Materials/MaterialExpressionLandscapePhysicalMaterialOutput.h
+++ b/Engine/Source/Runtime/Landscape/Classes/Materials/MaterialExpressionLandscapePhysicalMaterialOutput.h
@@ -45,7 +45,7 @@ class LANDSCAPE_API UMaterialExpressionLandscapePhysicalMaterialOutput : public
 
 	//~ Begin UMaterialExpression Interface
 	virtual void GetCaption(TArray<FString>& OutCaptions) const override;
-	virtual const TArray<FExpressionInput*> GetInputs() override;
+	virtual TArrayView<FExpressionInput*> GetInputsView() override;
 	virtual FExpressionInput* GetInput(int32 InputIndex) override;
 	virtual FName GetInputName(int32 InputIndex) const override;
 	virtual uint32 GetInputType(int32 InputIndex) override { return MCT_Float; }
diff --git a/Engine/Source/Runtime/Landscape/Private/LandscapeGrass.cpp b/Engine/Source/Runtime/Landscape/Private/LandscapeGrass.cpp
--- a/Engine/Source/Runtime/Landscape/Private/LandscapeGrass.cpp
+++ b/Engine/Source/Runtime/Landscape/Private/LandscapeGrass.cpp
@@ -1834,14 +1834,15 @@ void UMaterialExpressionLandscapeGrassOutput::GetCaption(TArray<FString>& OutCap
 	OutCaptions.Add(TEXT("Landscape Grass"));
 }
 
-const TArray<FExpressionInput*> UMaterialExpressionLandscapeGrassOutput::GetInputs()
+TArrayView<FExpressionInput*> UMaterialExpressionLandscapeGrassOutput::GetInputsView()
 {
-	TArray<FExpressionInput*> OutInputs;
+	CachedInputs.Empty();
+	CachedInputs.Reserve(GrassTypes.Num());
 	for (auto& GrassType : GrassTypes)
 	{
-		OutInputs.Add(&GrassType.Input);
+		CachedInputs.Add(&GrassType.Input);
 	}
-	return OutInputs;
+	return CachedInputs;
 }
 
 FExpressionInput* UMaterialExpressionLandscapeGrassOutput::GetInput(int32 InputIndex)
diff --git a/Engine/Source/Runtime/Landscape/Private/Materials/MaterialExpressionLandscapeLayerBlend.cpp b/Engine/Source/Runtime/Landscape/Private/Materials/MaterialExpressionLandscapeLayerBlend.cpp
--- a/Engine/Source/Runtime/Landscape/Private/Materials/MaterialExpressionLandscapeLayerBlend.cpp
+++ b/Engine/Source/Runtime/Landscape/Private/Materials/MaterialExpressionLandscapeLayerBlend.cpp
@@ -54,18 +54,19 @@ void UMaterialExpressionLandscapeLayerBlend::Serialize(FStructuredArchive::FReco
 
 #if WITH_EDITOR
 
-const TArray<FExpressionInput*> UMaterialExpressionLandscapeLayerBlend::GetInputs()
+TArrayView<FExpressionInput*> UMaterialExpressionLandscapeLayerBlend::GetInputsView()
 {
-	TArray<FExpressionInput*> Result;
+	CachedInputs.Empty();
+	CachedInputs.Reserve(Layers.Num() * 2);
 	for (int32 LayerIdx = 0; LayerIdx<Layers.Num(); LayerIdx++)
 	{
-		Result.Add(&Layers[LayerIdx].LayerInput);
+		CachedInputs.Add(&Layers[LayerIdx].LayerInput);
 		if (Layers[LayerIdx].BlendType == LB_HeightBlend)
 		{
-			Result.Add(&Layers[LayerIdx].HeightInput);
+			CachedInputs.Add(&Layers[LayerIdx].HeightInput);
 		}
 	}
-	return Result;
+	return CachedInputs;
 }
 
 
diff --git a/Engine/Source/Runtime/Landscape/Private/Materials/MaterialExpressionLandscapePhysicalMaterialOutput.cpp b/Engine/Source/Runtime/Landscape/Private/Materials/MaterialExpressionLandscapePhysicalMaterialOutput.cpp
--- a/Engine/Source/Runtime/Landscape/Private/Materials/MaterialExpressionLandscapePhysicalMaterialOutput.cpp
+++ b/Engine/Source/Runtime/Landscape/Private/Materials/MaterialExpressionLandscapePhysicalMaterialOutput.cpp
@@ -41,14 +41,15 @@ void UMaterialExpressionLandscapePhysicalMaterialOutput::GetCaption(TArray<FStri
 	OutCaptions.Add(FString(TEXT("Landscape Physical Material Output")));
 }
 
-const TArray<FExpressionInput*> UMaterialExpressionLandscapePhysicalMaterialOutput::GetInputs()
+TArrayView<FExpressionInput*> UMaterialExpressionLandscapePhysicalMaterialOutput::GetInputsView()
 {
-	TArray<FExpressionInput*> OutInputs;
+	CachedInputs.Empty();
+	CachedInputs.Reserve(Inputs.Num());
 	for (auto& Input : Inputs)
 	{
-		OutInputs.Add(&Input.Input);
+		CachedInputs.Add(&Input.Input);
 	}
-	return OutInputs;
+	return CachedInputs;
 }
 
 FExpressionInput* UMaterialExpressionLandscapePhysicalMaterialOutput::GetInput(int32 InputIndex)
