ActorDescContainerCollection templated to support const constainer
GenerateStreaming now use a ContainerCollection to generate an ActorDescViewMap perlevel containing multiple ActorDescContainers

#test EngineTests, PIE, cook, wp cl validation

diff --git a/Engine/Source/Runtime/Engine/Private/WorldPartition/ActorDescContainerCollection.cpp b/Engine/Source/Runtime/Engine/Private/WorldPartition/ActorDescContainerCollection.cpp
deleted file mode 100644
diff --git a/Engine/Source/Runtime/Engine/Private/WorldPartition/RuntimeSpatialHash/RuntimeSpatialHashHLOD.cpp b/Engine/Source/Runtime/Engine/Private/WorldPartition/RuntimeSpatialHash/RuntimeSpatialHashHLOD.cpp
--- a/Engine/Source/Runtime/Engine/Private/WorldPartition/RuntimeSpatialHash/RuntimeSpatialHashHLOD.cpp
+++ b/Engine/Source/Runtime/Engine/Private/WorldPartition/RuntimeSpatialHash/RuntimeSpatialHashHLOD.cpp
@@ -442,7 +442,7 @@ bool UWorldPartitionRuntimeSpatialHash::GenerateHLOD(ISourceControlHelper* Sourc
 
 	UWorldPartition* WorldPartition = GetOuterUWorldPartition();
 	IStreamingGenerationContext::FActorSetContainer* MainActorSetContainer = const_cast<IStreamingGenerationContext::FActorSetContainer*>(StreamingGenerationContext->GetMainWorldContainer());
-	const UActorDescContainer* ActorDescContainer = MainActorSetContainer->ActorDescContainer;
+	const FStreamingGenerationActorDescCollection* MainContainerCollection = MainActorSetContainer->ActorDescCollection;
 
 	// Find all used HLOD layers
 	TMap<UHLODLayer*, int32> HLODLayersLevels = GatherHLODLayers(StreamingGenerationContext, WorldPartition);
@@ -465,7 +465,7 @@ bool UWorldPartitionRuntimeSpatialHash::GenerateHLOD(ISourceControlHelper* Sourc
 
 	// HLOD creation context
 	FHLODCreationContext Context;
-	for (FActorDescList::TConstIterator<AWorldPartitionHLOD> HLODIterator(ActorDescContainer); HLODIterator; ++HLODIterator)
+	for (FActorDescList::TConstIterator<AWorldPartitionHLOD> HLODIterator(MainContainerCollection->GetMainActorDescContainer()); HLODIterator; ++HLODIterator)
 	{
 		FWorldPartitionHandle HLODActorHandle(WorldPartition, HLODIterator->GetGuid());
 		Context.HLODActorDescs.Emplace(HLODIterator->GetActorName(), MoveTemp(HLODActorHandle));
@@ -556,7 +556,7 @@ bool UWorldPartitionRuntimeSpatialHash::GenerateHLOD(ISourceControlHelper* Sourc
 			{
 				NewHLODActorSetInstance.DataLayers = DataLayerManager->GetRuntimeDataLayerInstances(HLODActorDescView.GetRuntimeDataLayerInstanceNames());
 			}
-			NewHLODActorSetInstance.ContentBundleID = MainActorSetContainer->ActorDescContainer->GetContentBundleGuid();
+			NewHLODActorSetInstance.ContentBundleID = MainContainerCollection->GetContentBundleGuid();
 			NewHLODActorSetInstance.ContainerInstance = MainActorSetContainer;
 			NewHLODActorSetInstance.ActorSet = MainActorSetContainer->ActorSets.Add_GetRef(MakeUnique<IStreamingGenerationContext::FActorSet>()).Get();
 			const_cast<IStreamingGenerationContext::FActorSet*>(NewHLODActorSetInstance.ActorSet)->Actors.Add(HLODActorDescView.GetGuid());
@@ -584,7 +584,7 @@ bool UWorldPartitionRuntimeSpatialHash::GenerateHLOD(ISourceControlHelper* Sourc
 	}
 
 	// Create/destroy HLOD grid actors
-	UpdateHLODGridsActors(GetWorld(), ActorDescContainer, HLODGrids, SourceControlHelper);
+	UpdateHLODGridsActors(GetWorld(), MainContainerCollection->GetMainActorDescContainer(), HLODGrids, SourceControlHelper);
 
 	return true;
 }
diff --git a/Engine/Source/Runtime/Engine/Private/WorldPartition/WorldPartitionHandle.cpp b/Engine/Source/Runtime/Engine/Private/WorldPartition/WorldPartitionHandle.cpp
--- a/Engine/Source/Runtime/Engine/Private/WorldPartition/WorldPartitionHandle.cpp
+++ b/Engine/Source/Runtime/Engine/Private/WorldPartition/WorldPartitionHandle.cpp
@@ -24,11 +24,6 @@ UActorDescContainer* FWorldPartitionImplBase::GetActorDescContainer(TUniquePtr<F
 	return ActorDesc ? ActorDesc->Get()->GetContainer() : nullptr;
 }
 
-UActorDescContainer* FWorldPartitionImplBase::GetActorDescContainer(FActorDescContainerCollection* ContainerCollection, const FGuid& ActorGuid)
-{
-	return ContainerCollection ? ContainerCollection->GetActorDescContainer(ActorGuid) : nullptr;
-}
-
 bool FWorldPartitionImplBase::IsActorDescLoaded(FWorldPartitionActorDesc* ActorDesc)
 {
 	return ActorDesc->IsLoaded();
diff --git a/Engine/Source/Runtime/Engine/Private/WorldPartition/WorldPartitionStreamingGeneration.cpp b/Engine/Source/Runtime/Engine/Private/WorldPartition/WorldPartitionStreamingGeneration.cpp
--- a/Engine/Source/Runtime/Engine/Private/WorldPartition/WorldPartitionStreamingGeneration.cpp
+++ b/Engine/Source/Runtime/Engine/Private/WorldPartition/WorldPartitionStreamingGeneration.cpp
@@ -6,8 +6,11 @@
 #include "WorldPartition/WorldPartitionStreamingGenerationContext.h"
 
 #include "Editor.h"
+#include "Algo/ForEach.h"
+#include "Algo/Sort.h"
 #include "Algo/Transform.h"
 #include "Algo/Unique.h"
+#include "Containers/ArrayView.h"
 #include "Engine/LevelScriptBlueprint.h"
 #include "ActorReferencesUtils.h"
 #include "Misc/PackageName.h"
@@ -83,21 +86,21 @@ class FWorldPartitionStreamingGenerator
 	class FStreamingGenerationContext : public IStreamingGenerationContext
 	{
 	public:
-		FStreamingGenerationContext(const FWorldPartitionStreamingGenerator* StreamingGenerator, const UActorDescContainer* MainWorldContainer)
+		FStreamingGenerationContext(const FWorldPartitionStreamingGenerator* StreamingGenerator, const FStreamingGenerationActorDescCollection& TopLevelActorDescCollection)
 		{
 			// Create the dataset required for IStreamingGenerationContext interface
 			MainWorldActorSetContainerIndex = INDEX_NONE;
-			ActorSetContainers.Empty(StreamingGenerator->ContainerDescriptorsMap.Num());
+			ActorSetContainers.Empty(StreamingGenerator->ContainerCollectionDescriptorsMap.Num());
 
-			TMap<const UActorDescContainer*, int32> ActorSetContainerMap;
-			for (const auto& [ActorDescContainer, ContainerDescriptor] : StreamingGenerator->ContainerDescriptorsMap)
+			TMap<TWeakPtr<FStreamingGenerationActorDescCollection>, int32> ActorSetContainerMap;
+			for (const auto& [LevelName, ContainerDescriptor] : StreamingGenerator->ContainerCollectionDescriptorsMap)
 			{
 				const int32 ContainerIndex = ActorSetContainers.AddDefaulted();
-				ActorSetContainerMap.Add(ContainerDescriptor.Container, ContainerIndex);
+				ActorSetContainerMap.Add(ContainerDescriptor.ActorDescCollection, ContainerIndex);
 
 				FActorSetContainer& ActorSetContainer = ActorSetContainers[ContainerIndex];
 				ActorSetContainer.ActorDescViewMap = &ContainerDescriptor.ActorDescViewMap;
-				ActorSetContainer.ActorDescContainer = ContainerDescriptor.Container;
+				ActorSetContainer.ActorDescCollection = ContainerDescriptor.ActorDescCollection.Get();
 
 				ActorSetContainer.ActorSets.Empty(ContainerDescriptor.Clusters.Num());
 				for (const TArray<FGuid>& Cluster : ContainerDescriptor.Clusters)
@@ -106,7 +109,7 @@ class FWorldPartitionStreamingGenerator
 					ActorSet.Actors = Cluster;
 				}
 
-				if (ContainerDescriptor.Container == MainWorldContainer)
+				if (ContainerDescriptor.ActorDescCollection->GetMainContainerPackageName() == TopLevelActorDescCollection.GetMainContainerPackageName())
 				{
 					check(MainWorldActorSetContainerIndex == INDEX_NONE);
 					MainWorldActorSetContainerIndex = ContainerIndex;
@@ -114,16 +117,16 @@ class FWorldPartitionStreamingGenerator
 			}
 
 			TMap<FActorContainerID, TSet<FGuid>> FilteredOutActorsPerContainer;
-			for (const auto& [ContainerID, ContainerInstanceDescriptor] : StreamingGenerator->ContainerInstanceDescriptorsMap)
+			for (const auto& [ContainerID, ContainerCollectionInstanceDescriptor] : StreamingGenerator->ContainerCollectionInstanceDescriptorsMap)
 			{
-				FilteredOutActorsPerContainer.Append(ContainerInstanceDescriptor.FilteredActors);
+				FilteredOutActorsPerContainer.Append(ContainerCollectionInstanceDescriptor.FilteredActors);
 			}
 
 			ActorSetInstances.Empty();
-			for (const auto& [ContainerID, ContainerInstanceDescriptor] : StreamingGenerator->ContainerInstanceDescriptorsMap)
+			for (const auto& [ContainerID, ContainerCollectionInstanceDescriptor] : StreamingGenerator->ContainerCollectionInstanceDescriptorsMap)
 			{
-				const FContainerDescriptor& ContainerDescriptor = StreamingGenerator->ContainerDescriptorsMap.FindChecked(ContainerInstanceDescriptor.Container);
-				const FActorSetContainer& ActorSetContainer = ActorSetContainers[ActorSetContainerMap.FindChecked(ContainerInstanceDescriptor.Container)];
+				const FContainerCollectionDescriptor& ContainerCollectionDescriptor = StreamingGenerator->ContainerCollectionDescriptorsMap.FindChecked(ContainerCollectionInstanceDescriptor.ActorDescCollection->GetMainContainerPackageName());
+				const FActorSetContainer& ActorSetContainer = ActorSetContainers[ActorSetContainerMap.FindChecked(ContainerCollectionDescriptor.ActorDescCollection)];
 
 				for (const TUniquePtr<FActorSet>& ActorSetPtr : ActorSetContainer.ActorSets)
 				{
@@ -162,14 +165,14 @@ class FWorldPartitionStreamingGenerator
 					}
 
 					FActorSetInstance& ActorSetInstance = ActorSetInstances.Emplace_GetRef();
-					const FContainerInstanceDescriptor::FPerInstanceData& PerInstanceData = ContainerInstanceDescriptor.GetInstanceData(ReferenceActorDescView.GetGuid());
+					const FContainerCollectionInstanceDescriptor::FPerInstanceData& PerInstanceData = ContainerCollectionInstanceDescriptor.GetInstanceData(ReferenceActorDescView.GetGuid());
 				
 					ActorSetInstance.ContainerInstance = &ActorSetContainer;
 					ActorSetInstance.ActorSet = &ActorSet;
-					ActorSetInstance.ContainerID = ContainerInstanceDescriptor.ID;
-					ActorSetInstance.Transform = ContainerInstanceDescriptor.Transform;
+					ActorSetInstance.ContainerID = ContainerCollectionInstanceDescriptor.ID;
+					ActorSetInstance.Transform = ContainerCollectionInstanceDescriptor.Transform;
 					ActorSetInstance.bIsSpatiallyLoaded = PerInstanceData.bIsSpatiallyLoaded;
-					ActorSetInstance.ContentBundleID = ContainerInstanceDescriptor.ContentBundleID;
+					ActorSetInstance.ContentBundleID = ContainerCollectionInstanceDescriptor.ContentBundleID;
 					ActorSetInstance.RuntimeGrid = PerInstanceData.RuntimeGrid;
 					ActorSetInstance.DataLayers = StreamingGenerator->GetRuntimeDataLayerInstances(PerInstanceData.DataLayers);
 
@@ -180,13 +183,13 @@ class FWorldPartitionStreamingGenerator
 						const FBox RuntimeBounds = ActorDescView.GetRuntimeBounds();
 						if (RuntimeBounds.IsValid)
 						{
-							ActorSetInstance.Bounds += RuntimeBounds.TransformBy(ContainerInstanceDescriptor.Transform);
+							ActorSetInstance.Bounds += RuntimeBounds.TransformBy(ContainerCollectionInstanceDescriptor.Transform);
 						}
 					}
 				}
 			}
 
-			WorldBounds = StreamingGenerator->ContainerInstanceDescriptorsMap.FindChecked(FActorContainerID::GetMainContainerID()).Bounds;
+			WorldBounds = StreamingGenerator->ContainerCollectionInstanceDescriptorsMap.FindChecked(FActorContainerID::GetMainContainerID()).Bounds;
 		}
 
 		virtual ~FStreamingGenerationContext()
@@ -230,14 +233,10 @@ class FWorldPartitionStreamingGenerator
 	/** 
 	 * An actor container descriptor, one for the main world and one for every unique level used by level instances.
 	 */
-	struct FContainerDescriptor
+	struct FContainerCollectionDescriptor
 	{
-		FContainerDescriptor()
-			: Container(nullptr)
-		{}
-
-		/** The unique actor descriptor container for this descriptor */
-		const UActorDescContainer* Container;
+		/** The unique actor descriptor container collection for this descriptor */
+		TSharedPtr<FStreamingGenerationActorDescCollection> ActorDescCollection;
 
 		/** The actor descriptor views for for this descriptor */
 		FActorDescViewMap ActorDescViewMap;
@@ -246,7 +245,7 @@ class FWorldPartitionStreamingGenerator
 		TSet<FGuid> EditorOnlyActorDescMap;
 
 		/** List of actor descriptor views that are containers (mainly level instances) */
-		TArray<FWorldPartitionActorDescView> ContainerInstanceViews;
+		TArray<FWorldPartitionActorDescView> ContainerCollectionInstanceViews;
 
 		/** List of actor clusters for this descriptor */
 		TArray<TArray<FGuid>> Clusters;
@@ -255,16 +254,15 @@ class FWorldPartitionStreamingGenerator
 	/** 
 	 * An actor container instance descriptor, one for the main world and one for every actor container instance
 	 */
-	struct FContainerInstanceDescriptor
+	struct FContainerCollectionInstanceDescriptor
 	{
-		FContainerInstanceDescriptor()
+		FContainerCollectionInstanceDescriptor()
 			: Bounds(ForceInit)
-			, Container(nullptr)
 		{}
 
 		FBox Bounds;
 		FTransform Transform;
-		const UActorDescContainer* Container;
+		TSharedPtr<FStreamingGenerationActorDescCollection> ActorDescCollection;
 		EContainerClusterMode ClusterMode;
 		FString OwnerName;
 		FActorContainerID ID;
@@ -362,22 +360,25 @@ class FWorldPartitionStreamingGenerator
 		}
 	}
 
-	void CreateActorDescViewMap(const UActorDescContainer* InContainer, FActorDescViewMap& OutActorDescViewMap, TSet<FGuid>& OutEditorOnlyActorDescMap, const FActorContainerID& InContainerID, TArray<FWorldPartitionActorDescView>& OutContainerInstances)
+	void CreateActorDescViewMap(const FStreamingGenerationActorDescCollection& InActorDescCollection, FActorDescViewMap& OutActorDescViewMap, TSet<FGuid>& OutEditorOnlyActorDescMap, const FActorContainerID& InContainerID, TArray<FWorldPartitionActorDescView>& OutContainerInstances)
 	{
 		// Should we handle unsaved or newly created actors?
 		const bool bHandleUnsavedActors = ModifiedActorsDescList && InContainerID.IsMainContainer();
 
 		// Consider all actors of a /Temp/ container package as Unsaved because loading them from disk will fail (Outer world name mismatch)
-		const bool bIsTempContainerPackage = FPackageName::IsTempPackage(InContainer->GetPackage()->GetName());
+		const bool bIsTempContainerPackage = FPackageName::IsTempPackage((InActorDescCollection.GetMainContainerPackageName().ToString()));
 
 		// Create an actor descriptor view for the specified actor (modified or unsaved actors)
-		auto GetModifiedActorDesc = [this](AActor* InActor, const UActorDescContainer* InContainer) -> FWorldPartitionActorDesc*
+		auto GetModifiedActorDesc = [this](AActor* InActor, const FStreamingGenerationActorDescCollection& InActorDescCollection) -> FWorldPartitionActorDesc*
 		{
 			FWorldPartitionActorDesc* ModifiedActorDesc = ModifiedActorsDescList->AddActor(InActor);
 
+			const UActorDescContainer* HandlingContainer = InActorDescCollection.FindHandlingContainer(InActor);
+			check(HandlingContainer != nullptr);
+
 			// Pretend that this actor descriptor belongs to the original container, even if it's not present. It's essentially a proxy
 			// descriptor on top an existing one and at this point no code should require to access the container to resolve it anyways.
-			ModifiedActorDesc->SetContainer(const_cast<UActorDescContainer*>(InContainer), InActor->GetWorld());
+			ModifiedActorDesc->SetContainer(const_cast<UActorDescContainer*>(HandlingContainer), InActor->GetWorld());
 
 			return ModifiedActorDesc;
 		};
@@ -402,7 +403,7 @@ class FWorldPartitionStreamingGenerator
 		};
 
 		// Register the actor descriptor view
-		auto RegisterActorDescView = [this, InContainer, &OutActorDescViewMap, &OutContainerInstances](const FGuid& ActorGuid, FWorldPartitionActorDescView& InActorDescView)
+		auto RegisterActorDescView = [this, &OutActorDescViewMap, &OutContainerInstances](const FGuid& ActorGuid, FWorldPartitionActorDescView& InActorDescView)
 		{
 			if (InActorDescView.IsContainerInstance())
 			{
@@ -415,7 +416,7 @@ class FWorldPartitionStreamingGenerator
 		};
 		
 		TMap<FGuid, FGuid> ContainerGuidsRemap;
-		for (FActorDescList::TConstIterator<> ActorDescIt(InContainer); ActorDescIt; ++ActorDescIt)
+		for (FStreamingGenerationActorDescCollection::TConstIterator<> ActorDescIt(&InActorDescCollection); ActorDescIt; ++ActorDescIt)
 		{
 			if (ShouldRegisterActorDesc(*ActorDescIt, InContainerID))
 			{
@@ -432,7 +433,7 @@ class FWorldPartitionStreamingGenerator
 					if (bHandleUnsavedActors && (bIsTempContainerPackage || Actor->GetPackage()->IsDirty()))
 					{
 						// Dirty, unsaved actor for PIE
-						FWorldPartitionActorDescView ModifiedActorDescView = GetModifiedActorDesc(Actor, InContainer);
+						FWorldPartitionActorDescView ModifiedActorDescView = GetModifiedActorDesc(Actor, InActorDescCollection);
 						RegisterActorDescView(ActorDescIt->GetGuid(), ModifiedActorDescView);
 						continue;
 					}
@@ -451,11 +452,13 @@ class FWorldPartitionStreamingGenerator
 		// Append new unsaved actors for the persistent level
 		if (bHandleUnsavedActors)
 		{
-			for (AActor* Actor : InContainer->GetWorldPartition()->GetWorld()->PersistentLevel->Actors)
+			for (AActor* Actor : InActorDescCollection.GetWorld()->PersistentLevel->Actors)
 			{
-				if (IsValid(Actor) && Actor->IsPackageExternal() && Actor->IsMainPackageActor() && !Actor->IsEditorOnly() && (Actor->GetContentBundleGuid() == InContainer->GetContentBundleGuid()) && !InContainer->GetActorDesc(Actor->GetActorGuid()))
+				if (IsValid(Actor) && Actor->IsPackageExternal() && Actor->IsMainPackageActor() && !Actor->IsEditorOnly()
+					&& (Actor->GetContentBundleGuid() == InActorDescCollection.GetContentBundleGuid())
+					&& !InActorDescCollection.GetActorDesc(Actor->GetActorGuid()))
 				{
-					FWorldPartitionActorDescView ModifiedActorDescView = GetModifiedActorDesc(Actor, InContainer);
+					FWorldPartitionActorDescView ModifiedActorDescView = GetModifiedActorDesc(Actor, InActorDescCollection);
 					if (ShouldRegisterActorDesc(ModifiedActorDescView.GetActorDesc(), InContainerID))
 					{
 						RegisterActorDescView(Actor->GetActorGuid(), ModifiedActorDescView);
@@ -465,18 +468,18 @@ class FWorldPartitionStreamingGenerator
 		}
 	}
 
-	void CreateActorDescriptorViewsRecursive(const FContainerInstanceDescriptor& InContainerInstanceDescriptor)
+	void CreateActorDescriptorViewsRecursive(const FContainerCollectionInstanceDescriptor& InContainerCollectionInstanceDescriptor)
 	{
 		// Inherited parent per-instance data logic
-		auto InheritParentContainerPerInstanceData = [&InContainerInstanceDescriptor](const FContainerInstanceDescriptor::FPerInstanceData& InParentPerInstanceData, const FWorldPartitionActorDescView& InActorDescView)
+		auto InheritParentContainerPerInstanceData = [&InContainerCollectionInstanceDescriptor](const FContainerCollectionInstanceDescriptor::FPerInstanceData& InParentPerInstanceData, const FWorldPartitionActorDescView& InActorDescView)
 		{
-			FContainerInstanceDescriptor::FPerInstanceData ResultPerInstanceData;
+			FContainerCollectionInstanceDescriptor::FPerInstanceData ResultPerInstanceData;
 
 			// Apply AND logic on spatially loaded flag
 			ResultPerInstanceData.bIsSpatiallyLoaded = InActorDescView.GetIsSpatiallyLoaded() && InParentPerInstanceData.bIsSpatiallyLoaded;
 
 			// Runtime grid is only inherited from the main world, since level instance doesn't support setting this value on actors
-			ResultPerInstanceData.RuntimeGrid = InContainerInstanceDescriptor.ID.IsMainContainer() ? InActorDescView.GetRuntimeGrid() : InParentPerInstanceData.RuntimeGrid;
+			ResultPerInstanceData.RuntimeGrid = InContainerCollectionInstanceDescriptor.ID.IsMainContainer() ? InActorDescView.GetRuntimeGrid() : InParentPerInstanceData.RuntimeGrid;
 
 			// Data layers are accumulated down the hierarchy chain, since level instances supports data layers assignation on actors
 			ResultPerInstanceData.DataLayers = InActorDescView.GetRuntimeDataLayerInstanceNames();
@@ -495,54 +498,54 @@ class FWorldPartitionStreamingGenerator
 		// ContainerDescriptor may be reallocated after this scope
 		{
 			// Create or resolve container descriptor
-			FContainerDescriptor& ContainerDescriptor = ContainerDescriptorsMap.FindOrAdd(InContainerInstanceDescriptor.Container);
+			FContainerCollectionDescriptor& ContainerCollectionDescriptor = ContainerCollectionDescriptorsMap.FindOrAdd(InContainerCollectionInstanceDescriptor.ActorDescCollection->GetMainContainerPackageName());
 		
 			// ContainerInstanceDescriptor may be reallocated after this scope
 			{
 				// Create container instance descriptor
-				check(!ContainerInstanceDescriptorsMap.Contains(InContainerInstanceDescriptor.ID));
+				check(!ContainerCollectionInstanceDescriptorsMap.Contains(InContainerCollectionInstanceDescriptor.ID));
 
-				FContainerInstanceDescriptor& ContainerInstanceDescriptor = ContainerInstanceDescriptorsMap.Add(InContainerInstanceDescriptor.ID, InContainerInstanceDescriptor);
+				FContainerCollectionInstanceDescriptor& ContainerCollectionInstanceDescriptor = ContainerCollectionInstanceDescriptorsMap.Add(InContainerCollectionInstanceDescriptor.ID, InContainerCollectionInstanceDescriptor);
 
-				if (!ContainerDescriptor.Container)
+				if (!ContainerCollectionDescriptor.ActorDescCollection)
 				{
-					ContainerDescriptor.Container = ContainerInstanceDescriptor.Container;
+					ContainerCollectionDescriptor.ActorDescCollection = ContainerCollectionInstanceDescriptor.ActorDescCollection;
 
 					// Gather actor descriptor views for this container
-					CreateActorDescViewMap(ContainerInstanceDescriptor.Container, ContainerDescriptor.ActorDescViewMap, ContainerDescriptor.EditorOnlyActorDescMap, ContainerInstanceDescriptor.ID, ContainerDescriptor.ContainerInstanceViews);
+					CreateActorDescViewMap(*ContainerCollectionInstanceDescriptor.ActorDescCollection, ContainerCollectionDescriptor.ActorDescViewMap, ContainerCollectionDescriptor.EditorOnlyActorDescMap, ContainerCollectionInstanceDescriptor.ID, ContainerCollectionDescriptor.ContainerCollectionInstanceViews);
 
 					// Resolve actor descriptor views before validation
-					ResolveContainerDescriptor(ContainerDescriptor);
+					ResolveContainerDescriptor(ContainerCollectionDescriptor);
 
 					// Validate container, fixing anything illegal, etc.
-					ValidateContainerDescriptor(ContainerDescriptor, ContainerInstanceDescriptor.ID.IsMainContainer());
+					ValidateContainerDescriptor(ContainerCollectionDescriptor, ContainerCollectionInstanceDescriptor.ID.IsMainContainer());
 
 					// Update container, computing cluster, bounds, etc.
-					UpdateContainerDescriptor(ContainerDescriptor);
+					UpdateContainerDescriptor(ContainerCollectionDescriptor);
 				}
 
 				// Calculate Bounds of Non-container ActorDescViews
-				check(!ContainerInstanceDescriptor.Bounds.IsValid);
-				ContainerDescriptor.ActorDescViewMap.ForEachActorDescView([&ContainerInstanceDescriptor](const FWorldPartitionActorDescView& ActorDescView)
+				check(!ContainerCollectionInstanceDescriptor.Bounds.IsValid);
+				ContainerCollectionDescriptor.ActorDescViewMap.ForEachActorDescView([&ContainerCollectionInstanceDescriptor](const FWorldPartitionActorDescView& ActorDescView)
 				{
 					if (ActorDescView.GetIsSpatiallyLoaded())
 					{
 						const FBox RuntimeBounds = ActorDescView.GetRuntimeBounds();
 						check(RuntimeBounds.IsValid);
 
-						ContainerInstanceDescriptor.Bounds += RuntimeBounds.TransformBy(ContainerInstanceDescriptor.Transform);
+						ContainerCollectionInstanceDescriptor.Bounds += RuntimeBounds.TransformBy(ContainerCollectionInstanceDescriptor.Transform);
 					}
 				});
 			}
 
 			// Parse actor containers
-			TArray<FWorldPartitionActorDescView> ContainerInstanceViews = ContainerDescriptor.ContainerInstanceViews;
-			for (const FWorldPartitionActorDescView& ContainerInstanceView : ContainerInstanceViews)
+			TArray<FWorldPartitionActorDescView> ContainerCollectionInstanceViews = ContainerCollectionDescriptor.ContainerCollectionInstanceViews;
+			for (const FWorldPartitionActorDescView& ContainerCollectionInstanceView : ContainerCollectionInstanceViews)
 			{
 				FWorldPartitionActorDesc::FContainerInstance SubContainerInstance;
-				if (!ContainerInstanceView.GetContainerInstance(SubContainerInstance) || !SubContainerInstance.Container)
+				if (!ContainerCollectionInstanceView.GetContainerInstance(SubContainerInstance) || !SubContainerInstance.Container)
 				{
-					ErrorHandler->OnLevelInstanceInvalidWorldAsset(ContainerInstanceView, ContainerInstanceView.GetContainerPackage(), IStreamingGenerationErrorHandler::ELevelInstanceInvalidReason::WorldAssetHasInvalidContainer);
+					ErrorHandler->OnLevelInstanceInvalidWorldAsset(ContainerCollectionInstanceView, ContainerCollectionInstanceView.GetContainerPackage(), IStreamingGenerationErrorHandler::ELevelInstanceInvalidReason::WorldAssetHasInvalidContainer);
 					continue;
 				}
 
@@ -551,26 +554,26 @@ class FWorldPartitionStreamingGenerator
 
 				if (bContainerWasAlreadyInSet)
 				{
-					ErrorHandler->OnLevelInstanceInvalidWorldAsset(ContainerInstanceView, ContainerInstanceView.GetContainerPackage(), IStreamingGenerationErrorHandler::ELevelInstanceInvalidReason::CirculalReference);
+					ErrorHandler->OnLevelInstanceInvalidWorldAsset(ContainerCollectionInstanceView, ContainerCollectionInstanceView.GetContainerPackage(), IStreamingGenerationErrorHandler::ELevelInstanceInvalidReason::CirculalReference);
 					continue;
 				}
 
-				FContainerInstanceDescriptor SubContainerInstanceDescriptor;
-				SubContainerInstanceDescriptor.ID = FActorContainerID(InContainerInstanceDescriptor.ID, ContainerInstanceView.GetGuid());
-				SubContainerInstanceDescriptor.Container = SubContainerInstance.Container;
-				SubContainerInstanceDescriptor.Transform = SubContainerInstance.Transform * InContainerInstanceDescriptor.Transform;
-				SubContainerInstanceDescriptor.ParentID = InContainerInstanceDescriptor.ID;
-				SubContainerInstanceDescriptor.OwnerName = *ContainerInstanceView.GetActorLabelOrName().ToString();
+				FContainerCollectionInstanceDescriptor SubContainerInstanceDescriptor;
+				SubContainerInstanceDescriptor.ID = FActorContainerID(InContainerCollectionInstanceDescriptor.ID, ContainerCollectionInstanceView.GetGuid());
+				SubContainerInstanceDescriptor.ActorDescCollection = MakeShared<FStreamingGenerationActorDescCollection>(FStreamingGenerationActorDescCollection{SubContainerInstance.Container});
+				SubContainerInstanceDescriptor.Transform = SubContainerInstance.Transform * InContainerCollectionInstanceDescriptor.Transform;
+				SubContainerInstanceDescriptor.ParentID = InContainerCollectionInstanceDescriptor.ID;
+				SubContainerInstanceDescriptor.OwnerName = *ContainerCollectionInstanceView.GetActorLabelOrName().ToString();
 				// Since Content Bundles streaming generation happens in its own context, all actor set instances must have the same content bundle GUID for now, so Level Instances
 				// placed inside a Content Bundle will propagate their Content Bundle GUID to child instances.
-				SubContainerInstanceDescriptor.ContentBundleID = InContainerInstanceDescriptor.ContentBundleID;
-				SubContainerInstanceDescriptor.InstanceData = InheritParentContainerPerInstanceData(InContainerInstanceDescriptor.InstanceData, ContainerInstanceView);
+				SubContainerInstanceDescriptor.ContentBundleID = InContainerCollectionInstanceDescriptor.ContentBundleID;
+				SubContainerInstanceDescriptor.InstanceData = InheritParentContainerPerInstanceData(InContainerCollectionInstanceDescriptor.InstanceData, ContainerCollectionInstanceView);
 
-				if (WorldPartitionSubsystem && InContainerInstanceDescriptor.ID.IsMainContainer() && ContainerInstanceView.GetContainerFilterType() == EWorldPartitionActorFilterType::Loading)
+				if (WorldPartitionSubsystem && InContainerCollectionInstanceDescriptor.ID.IsMainContainer() && ContainerCollectionInstanceView.GetContainerFilterType() == EWorldPartitionActorFilterType::Loading)
 				{
-					if (const FWorldPartitionActorFilter* ContainerFilter = ContainerInstanceView.GetContainerFilter())
+					if (const FWorldPartitionActorFilter* ContainerFilter = ContainerCollectionInstanceView.GetContainerFilter())
 					{						
-						SubContainerInstanceDescriptor.FilteredActors = WorldPartitionSubsystem->GetFilteredActorsPerContainer(SubContainerInstanceDescriptor.ID, ContainerInstanceView.GetContainerPackage().ToString(), *ContainerFilter);
+						SubContainerInstanceDescriptor.FilteredActors = WorldPartitionSubsystem->GetFilteredActorsPerContainer(SubContainerInstanceDescriptor.ID, ContainerCollectionInstanceView.GetContainerPackage().ToString(), *ContainerFilter);
 					}
 				}
 
@@ -581,70 +584,70 @@ class FWorldPartitionStreamingGenerator
 		}
 
 		// Fetch the versions stored in the map as it can have been reallocated during recursion
-		FContainerDescriptor& ContainerDescriptor = ContainerDescriptorsMap.FindOrAdd(InContainerInstanceDescriptor.Container);
-		FContainerInstanceDescriptor& ContainerInstanceDescriptor = ContainerInstanceDescriptorsMap.FindChecked(InContainerInstanceDescriptor.ID);
+		FContainerCollectionDescriptor& ContainerCollectionDescriptor = ContainerCollectionDescriptorsMap.FindOrAdd(InContainerCollectionInstanceDescriptor.ActorDescCollection->GetMainContainerPackageName());
+		FContainerCollectionInstanceDescriptor& ContainerCollectionInstanceDescriptor = ContainerCollectionInstanceDescriptorsMap.FindChecked(InContainerCollectionInstanceDescriptor.ID);
 
-		if (!InContainerInstanceDescriptor.ID.IsMainContainer())
+		if (!InContainerCollectionInstanceDescriptor.ID.IsMainContainer())
 		{
-			FContainerInstanceDescriptor& ParentContainer = ContainerInstanceDescriptorsMap.FindChecked(InContainerInstanceDescriptor.ParentID);
-			ParentContainer.Bounds += ContainerInstanceDescriptor.Bounds;
+			FContainerCollectionInstanceDescriptor& ParentContainerCollection = ContainerCollectionInstanceDescriptorsMap.FindChecked(InContainerCollectionInstanceDescriptor.ParentID);
+			ParentContainerCollection.Bounds += ContainerCollectionInstanceDescriptor.Bounds;
 		}
 
 		// Apply per-instance data
-		ContainerInstanceDescriptor.PerInstanceData.Reserve(ContainerDescriptor.ActorDescViewMap.Num());
-		ContainerDescriptor.ActorDescViewMap.ForEachActorDescView([&ContainerInstanceDescriptor, &InheritParentContainerPerInstanceData](FWorldPartitionActorDescView& ActorDescView)
+		ContainerCollectionInstanceDescriptor.PerInstanceData.Reserve(ContainerCollectionDescriptor.ActorDescViewMap.Num());
+		ContainerCollectionDescriptor.ActorDescViewMap.ForEachActorDescView([&ContainerCollectionInstanceDescriptor, &InheritParentContainerPerInstanceData](FWorldPartitionActorDescView& ActorDescView)
 		{
-			const FContainerInstanceDescriptor::FPerInstanceData PerInstanceData = InheritParentContainerPerInstanceData(ContainerInstanceDescriptor.InstanceData, ActorDescView);
+			const FContainerCollectionInstanceDescriptor::FPerInstanceData PerInstanceData = InheritParentContainerPerInstanceData(ContainerCollectionInstanceDescriptor.InstanceData, ActorDescView);
 
-			if (PerInstanceData != ContainerInstanceDescriptor.InstanceData)
+			if (PerInstanceData != ContainerCollectionInstanceDescriptor.InstanceData)
 			{
-				const FSetElementId UniquePerInstanceDataId = ContainerInstanceDescriptor.UniquePerInstanceData.Add(PerInstanceData);
-				ContainerInstanceDescriptor.PerInstanceData.Emplace(ActorDescView.GetGuid(), UniquePerInstanceDataId);
+				const FSetElementId UniquePerInstanceDataId = ContainerCollectionInstanceDescriptor.UniquePerInstanceData.Add(PerInstanceData);
+				ContainerCollectionInstanceDescriptor.PerInstanceData.Emplace(ActorDescView.GetGuid(), UniquePerInstanceDataId);
 			}
 		});
 
 		// Validate container instance, fixing anything illegal, etc.
-		ValidateContainerInstanceDescriptor(ContainerInstanceDescriptor, ContainerInstanceDescriptor.ID.IsMainContainer());
+		ValidateContainerInstanceDescriptor(ContainerCollectionInstanceDescriptor, ContainerCollectionInstanceDescriptor.ID.IsMainContainer());
 	}
 
 	/** 
 	 * Creates the actor descriptor views for the specified container.
 	 */
-	void CreateActorContainers(const UActorDescContainer* InContainer)
+	void CreateActorContainers(const FStreamingGenerationActorDescCollection& InActorDescCollection)
 	{
 		TRACE_CPUPROFILER_EVENT_SCOPE(FWorldPartitionStreamingGenerator::CreateActorContainers);
 
-		FContainerInstanceDescriptor MainContainerInstance;
-		MainContainerInstance.Container = InContainer;
-		MainContainerInstance.ClusterMode = EContainerClusterMode::Partitioned;
-		MainContainerInstance.OwnerName = TEXT("MainContainer");
-		MainContainerInstance.ContentBundleID = InContainer->GetContentBundleGuid();
-		MainContainerInstance.InstanceData.bIsSpatiallyLoaded = true; // Since we apply AND logic on spatially loaded flag recursively, startup value must be true
+		FContainerCollectionInstanceDescriptor MainContainerCollectionInstance;
+		MainContainerCollectionInstance.ActorDescCollection = MakeShared<FStreamingGenerationActorDescCollection>(InActorDescCollection);
+		MainContainerCollectionInstance.ClusterMode = EContainerClusterMode::Partitioned;
+		MainContainerCollectionInstance.OwnerName = TEXT("MainContainer");
+		MainContainerCollectionInstance.ContentBundleID = InActorDescCollection.GetContentBundleGuid();
+		MainContainerCollectionInstance.InstanceData.bIsSpatiallyLoaded = true; // Since we apply AND logic on spatially loaded flag recursively, startup value must be true
 
-		CreateActorDescriptorViewsRecursive(MainContainerInstance);
+		CreateActorDescriptorViewsRecursive(MainContainerCollectionInstance);
 	}
 
 	/** 
 	 * Perform various validations on the container descriptor, and adjust it based on different requirements. This needs to happen before updating
 	 * containers bounds because some actor descriptor views might change grid placement, etc.
 	 */
-	void ResolveContainerDescriptor(FContainerDescriptor& ContainerDescriptor)
+	void ResolveContainerDescriptor(FContainerCollectionDescriptor& ContainerCollectionDescriptor)
 	{
-		auto ResolveActorDescView = [this, &ContainerDescriptor](FWorldPartitionActorDescView& ActorDescView)
+		auto ResolveActorDescView = [this, &ContainerCollectionDescriptor](FWorldPartitionActorDescView& ActorDescView)
 		{
 			ResolveRuntimeSpatiallyLoaded(ActorDescView);
 			ResolveRuntimeGrid(ActorDescView);
-			ResolveRuntimeDataLayers(ActorDescView, ContainerDescriptor.ActorDescViewMap);
+			ResolveRuntimeDataLayers(ActorDescView, ContainerCollectionDescriptor.ActorDescViewMap);
 		};
 
-		ContainerDescriptor.ActorDescViewMap.ForEachActorDescView([this, &ResolveActorDescView](FWorldPartitionActorDescView& ActorDescView)
+		ContainerCollectionDescriptor.ActorDescViewMap.ForEachActorDescView([this, &ResolveActorDescView](FWorldPartitionActorDescView& ActorDescView)
 		{
 			ResolveActorDescView(ActorDescView);
 		});
 
-		for (FWorldPartitionActorDescView& ContainerInstanceView : ContainerDescriptor.ContainerInstanceViews)
+		for (FWorldPartitionActorDescView& ContainerCollectionInstanceView : ContainerCollectionDescriptor.ContainerCollectionInstanceViews)
 		{
-			ResolveActorDescView(ContainerInstanceView);
+			ResolveActorDescView(ContainerCollectionInstanceView);
 		}
 	}
 
@@ -652,7 +655,7 @@ class FWorldPartitionStreamingGenerator
 	 * Perform various validations on the container descriptor, and adjust it based on different requirements. This needs to happen before updating
 	 * containers bounds because some actor descriptor views might change grid placement, etc.
 	 */
-	void ValidateContainerDescriptor(FContainerDescriptor& ContainerDescriptor, bool bIsMainContainer)
+	void ValidateContainerDescriptor(FContainerCollectionDescriptor& ContainerCollectionDescriptor, bool bIsMainContainer)
 	{
 		if (bIsMainContainer)
 		{
@@ -680,12 +683,12 @@ class FWorldPartitionStreamingGenerator
 			}
 			else
 			{
-				ULevel::GetLevelScriptExternalActorsReferencesFromPackage(ContainerDescriptor.Container->GetContainerPackage(), LevelScriptReferences);
+				ULevel::GetLevelScriptExternalActorsReferencesFromPackage(ContainerCollectionDescriptor.ActorDescCollection->GetMainContainerPackageName(), LevelScriptReferences);
 			}
 
 			for (const FGuid& LevelScriptReferenceActorGuid : LevelScriptReferences)
 			{
-				if (FWorldPartitionActorDescView* ActorDescView = ContainerDescriptor.ActorDescViewMap.FindByGuid(LevelScriptReferenceActorGuid))
+				if (FWorldPartitionActorDescView* ActorDescView = ContainerCollectionDescriptor.ActorDescViewMap.FindByGuid(LevelScriptReferenceActorGuid))
 				{
 					if (ActorDescView->GetIsSpatiallyLoaded())
 					{
@@ -703,14 +706,14 @@ class FWorldPartitionStreamingGenerator
 		}
 
 		// Route standard CheckForErrors calls which should not modify actor descriptors in any ways
-		ContainerDescriptor.ActorDescViewMap.ForEachActorDescView([this](FWorldPartitionActorDescView& ActorDescView)
+		ContainerCollectionDescriptor.ActorDescViewMap.ForEachActorDescView([this](FWorldPartitionActorDescView& ActorDescView)
 		{
 			ActorDescView.CheckForErrors(ErrorHandler);
 		});
 
-		for (FWorldPartitionActorDescView& ContainerInstanceView : ContainerDescriptor.ContainerInstanceViews)
+		for (FWorldPartitionActorDescView& ContainerCollectionInstanceView : ContainerCollectionDescriptor.ContainerCollectionInstanceViews)
 		{
-			ContainerInstanceView.CheckForErrors(ErrorHandler);
+			ContainerCollectionInstanceView.CheckForErrors(ErrorHandler);
 		}
 
 		// Perform various adjustements based on validations and report errors
@@ -732,7 +735,7 @@ class FWorldPartitionStreamingGenerator
 
 			NbErrorsDetected = 0;
 
-			ContainerDescriptor.ActorDescViewMap.ForEachActorDescView([this, &ContainerDescriptor, &NbErrorsDetected, PassType](FWorldPartitionActorDescView& ActorDescView)
+			ContainerCollectionDescriptor.ActorDescViewMap.ForEachActorDescView([this, &ContainerCollectionDescriptor, &NbErrorsDetected, PassType](FWorldPartitionActorDescView& ActorDescView)
 			{
 				// Validate grid placement
 				auto IsReferenceGridPlacementValid = [](const FWorldPartitionActorDescView& RefererActorDescView, const FWorldPartitionActorDescView& ReferenceActorDescView)
@@ -779,7 +782,7 @@ class FWorldPartitionStreamingGenerator
 					if (ReferenceGuid != ActorDescView.GetParentActor()) // References to the parent are inversed in their handling 
 					{
 						// Filter out parent back references
-						FWorldPartitionActorDescView* ReferenceActorDesc = ContainerDescriptor.ActorDescViewMap.FindByGuid(ReferenceGuid);
+						FWorldPartitionActorDescView* ReferenceActorDesc = ContainerCollectionDescriptor.ActorDescViewMap.FindByGuid(ReferenceGuid);
 						if (!ReferenceActorDesc || (ReferenceActorDesc->GetParentActor() != ActorDescView.GetGuid()))
 						{
 							References.Emplace(FActorReferenceInfo { ActorDescView.GetGuid(), &ActorDescView, ReferenceGuid, ReferenceActorDesc });
@@ -795,7 +798,7 @@ class FWorldPartitionStreamingGenerator
 
 					while (ParentGuid.IsValid())
 					{
-						FWorldPartitionActorDescView* ParentDescView = ContainerDescriptor.ActorDescViewMap.FindByGuid(ParentGuid);
+						FWorldPartitionActorDescView* ParentDescView = ContainerCollectionDescriptor.ActorDescViewMap.FindByGuid(ParentGuid);
 					
 						if (ParentDescView)
 						{
@@ -902,7 +905,7 @@ class FWorldPartitionStreamingGenerator
 					}
 					else
 					{
-						if (!ContainerDescriptor.EditorOnlyActorDescMap.Contains(Info.ReferenceGuid))
+						if (!ContainerCollectionDescriptor.EditorOnlyActorDescMap.Contains(Info.ReferenceGuid))
 						{
 							if (PassType == EPassType::ErrorReporting)
 							{
@@ -946,9 +949,9 @@ class FWorldPartitionStreamingGenerator
 	 * Perform various validations on the container descriptor instance, and adjust it based on different requirements. This needs to happen before updating
 	 * containers bounds because some actor descriptor views might change grid placement, etc.
 	 */
-	void ValidateContainerInstanceDescriptor(FContainerInstanceDescriptor& ContainerInstanceDescriptor, bool bIsMainContainer)
+	void ValidateContainerInstanceDescriptor(FContainerCollectionInstanceDescriptor& ContainerCollectionInstanceDescriptor, bool bIsMainContainer)
 	{
-		FContainerDescriptor& ContainerDescriptor = ContainerDescriptorsMap.FindChecked(ContainerInstanceDescriptor.Container);
+		FContainerCollectionDescriptor& ContainerCollectionDescriptor = ContainerCollectionDescriptorsMap.FindChecked(ContainerCollectionInstanceDescriptor.ActorDescCollection->GetMainContainerPackageName());
 
 		// Perform various adjustements based on validations and report errors
 		//
@@ -969,9 +972,9 @@ class FWorldPartitionStreamingGenerator
 
 			NbErrorsDetected = 0;
 
-			ContainerDescriptor.ActorDescViewMap.ForEachActorDescView([this, &ContainerInstanceDescriptor, &NbErrorsDetected, PassType](FWorldPartitionActorDescView& ActorDescView)
+			ContainerCollectionDescriptor.ActorDescViewMap.ForEachActorDescView([this, &ContainerCollectionInstanceDescriptor, &NbErrorsDetected, PassType](FWorldPartitionActorDescView& ActorDescView)
 			{
-				FContainerInstanceDescriptor::FPerInstanceData& PerInstanceData = ContainerInstanceDescriptor.GetInstanceData(ActorDescView.GetGuid());
+				FContainerCollectionInstanceDescriptor::FPerInstanceData& PerInstanceData = ContainerCollectionInstanceDescriptor.GetInstanceData(ActorDescView.GetGuid());
 
 PRAGMA_DISABLE_DEPRECATION_WARNINGS
 				if (ActorDescView.ShouldValidateRuntimeGrid())
@@ -998,12 +1001,12 @@ PRAGMA_ENABLE_DEPRECATION_WARNINGS
 	/** 
 	 * Update the container descriptor containers to adjust their bounds from actor descriptor views.
 	 */
-	void UpdateContainerDescriptor(FContainerDescriptor& ContainerDescriptor)
+	void UpdateContainerDescriptor(FContainerCollectionDescriptor& ContainerCollectionDescriptor)
 	{
 		// Build clusters for this container - at this point, all actors references should be in the same data layers, grid, etc because of actor descriptors validation.
 		TArray<TPair<FGuid, TArray<FGuid>>> ActorsWithRefs;
-		ContainerDescriptor.ActorDescViewMap.ForEachActorDescView([&ActorsWithRefs](const FWorldPartitionActorDescView& ActorDescView) { ActorsWithRefs.Emplace(ActorDescView.GetGuid(), ActorDescView.GetReferences()); });
-		ContainerDescriptor.Clusters = GenerateObjectsClusters(ActorsWithRefs);
+		ContainerCollectionDescriptor.ActorDescViewMap.ForEachActorDescView([&ActorsWithRefs](const FWorldPartitionActorDescView& ActorDescView) { ActorsWithRefs.Emplace(ActorDescView.GetGuid(), ActorDescView.GetReferences()); });
+		ContainerCollectionDescriptor.Clusters = GenerateObjectsClusters(ActorsWithRefs);
 	}
 
 public:
@@ -1049,12 +1052,12 @@ public:
 		DataLayerManager = OwningWorld ? WorldPartitionContext->GetDataLayerManager() : nullptr;
 	}
 
-	void PreparationPhase(const UActorDescContainer* Container)
+	void PreparationPhase(const FStreamingGenerationActorDescCollection& ActorDescCollection)
 	{
-		CreateActorContainers(Container);
+		CreateActorContainers(ActorDescCollection);
 
 		// Construct the streaming generation context
-		StreamingGenerationContext = MakeUnique<FStreamingGenerationContext>(this, Container);
+		StreamingGenerationContext = MakeUnique<FStreamingGenerationContext>(this, ActorDescCollection);
 	}
 
 	static TUniquePtr<FArchive> CreateDumpStateLogArchive(const TCHAR* Suffix, bool bTimeStamped = true)
@@ -1076,24 +1079,24 @@ public:
 
 		// Build the containers tree representation
 		TMultiMap<FActorContainerID, FActorContainerID> InvertedContainersHierarchy;
-		for (auto& [ContainerID, ContainerInstanceDescriptor] : ContainerInstanceDescriptorsMap)
+		for (auto& [ContainerID, ContainerCollectionInstanceDescriptor] : ContainerCollectionInstanceDescriptorsMap)
 		{
 			if (!ContainerID.IsMainContainer())
 			{
-				InvertedContainersHierarchy.Add(ContainerInstanceDescriptor.ParentID, ContainerID);
+				InvertedContainersHierarchy.Add(ContainerCollectionInstanceDescriptor.ParentID, ContainerID);
 			}
 		}
 
 		UE_SCOPED_INDENT_LOG_ARCHIVE(Ar.PrintfIndent(TEXT("Containers:")));
-		for (auto& [ActorDescContainer, ContainerDescriptor] : ContainerDescriptorsMap)
+		for (auto& [LevelName, ContainerCollectionDescriptor] : ContainerCollectionDescriptorsMap)
 		{
-			UE_SCOPED_INDENT_LOG_ARCHIVE(Ar.PrintfIndent(TEXT("Container: %s"), *ContainerDescriptor.Container->GetContainerPackage().ToString()));
+			UE_SCOPED_INDENT_LOG_ARCHIVE(Ar.PrintfIndent(TEXT("Container: %s"), *ContainerCollectionDescriptor.ActorDescCollection->GetMainContainerPackageName().ToString()));
 
-			if (ContainerDescriptor.ActorDescViewMap.Num())
+			if (ContainerCollectionDescriptor.ActorDescViewMap.Num())
 			{
 				UE_SCOPED_INDENT_LOG_ARCHIVE(Ar.PrintfIndent(TEXT("ActorDescs:")));
 
-				TMap<FGuid, FWorldPartitionActorDescView*> SortedActorDescViewMap = ContainerDescriptor.ActorDescViewMap.ActorDescViewsByGuid;
+				TMap<FGuid, FWorldPartitionActorDescView*> SortedActorDescViewMap = ContainerCollectionDescriptor.ActorDescViewMap.ActorDescViewsByGuid;
 				SortedActorDescViewMap.KeySort([](const FGuid& GuidA, const FGuid& GuidB) { return GuidA < GuidB; });
 
 				for (auto& [ActorGuid, ActorDescView] : SortedActorDescViewMap)
@@ -1102,11 +1105,11 @@ public:
 				}
 			}
 
-			if (ContainerDescriptor.Clusters.Num())
+			if (ContainerCollectionDescriptor.Clusters.Num())
 			{
 				UE_SCOPED_INDENT_LOG_ARCHIVE(Ar.PrintfIndent(TEXT("Clusters:")));
 
-				TArray<TArray<FGuid>> SortedClusters = ContainerDescriptor.Clusters;
+				TArray<TArray<FGuid>> SortedClusters = ContainerCollectionDescriptor.Clusters;
 				for (TArray<FGuid>& ActorGuids : SortedClusters)
 				{
 					ActorGuids.Sort([](const FGuid& GuidA, const FGuid& GuidB) { return GuidA < GuidB; });
@@ -1119,7 +1122,7 @@ public:
 					UE_SCOPED_INDENT_LOG_ARCHIVE(Ar.PrintfIndent(TEXT("[%3d]"), ClusterIndex++));
 					for (const FGuid& ActorGuid : ActorGuids)
 					{
-						const FWorldPartitionActorDescView& ActorDescView = ContainerDescriptor.ActorDescViewMap.FindByGuidChecked(ActorGuid);
+						const FWorldPartitionActorDescView& ActorDescView = ContainerCollectionDescriptor.ActorDescViewMap.FindByGuidChecked(ActorGuid);
 						Ar.Print(*ActorDescView.ToString());
 					}
 				}
@@ -1131,7 +1134,7 @@ public:
 		{
 			auto DumpContainerInstancesRecursive = [this, &InvertedContainersHierarchy, &Ar](const FActorContainerID& ContainerID, auto& RecursiveFunc) -> void
 			{
-				const FContainerInstanceDescriptor& ContainerInstanceDescriptor = ContainerInstanceDescriptorsMap.FindChecked(ContainerID);
+				const FContainerCollectionInstanceDescriptor& ContainerInstanceDescriptor = ContainerCollectionInstanceDescriptorsMap.FindChecked(ContainerID);
 				
 				{
 					UE_SCOPED_INDENT_LOG_ARCHIVE(Ar.PrintfIndent(TEXT("%s:"), *ContainerInstanceDescriptor.OwnerName));
@@ -1140,7 +1143,7 @@ public:
 					Ar.Printf(TEXT(" ParentID: %s"), *ContainerInstanceDescriptor.ParentID.ToString());
 					Ar.Printf(TEXT("   Bounds: %s"), *ContainerInstanceDescriptor.Bounds.ToString());
 					Ar.Printf(TEXT("Transform: %s"), *ContainerInstanceDescriptor.Transform.ToString());
-					Ar.Printf(TEXT("Container: %s"), *ContainerInstanceDescriptor.Container->GetContainerPackage().ToString());
+					Ar.Printf(TEXT("Container: %s"), *ContainerInstanceDescriptor.ActorDescCollection->GetMainContainerPackageName().ToString());
 				}
 
 				TArray<FActorContainerID> ChildContainersIDs;
@@ -1184,11 +1187,11 @@ private:
 	TFunction<bool(FName)> IsValidGrid;
 	IStreamingGenerationErrorHandler* ErrorHandler;
 
-	/** Maps containers to their container descriptor */
-	TMap<const UActorDescContainer*, FContainerDescriptor> ContainerDescriptorsMap;
+	// Maps a level to its descriptor
+	TMap<FName, FContainerCollectionDescriptor> ContainerCollectionDescriptorsMap;
 	
-	/** Maps containers IDs to their container instance descriptor */
-	TMap<FActorContainerID, FContainerInstanceDescriptor> ContainerInstanceDescriptorsMap;
+	/** Maps containers IDs to their container collection instance descriptor */
+	TMap<FActorContainerID, FContainerCollectionInstanceDescriptor> ContainerCollectionInstanceDescriptorsMap;
 
 	/** Data required for streaming generation interface */
 	TUniquePtr<FStreamingGenerationContext> StreamingGenerationContext;
@@ -1216,8 +1219,9 @@ bool UWorldPartition::GenerateContainerStreaming(const UActorDescContainer* InAc
 
 bool UWorldPartition::GenerateStreaming(const FGenerateStreamingParams& InParams, FGenerateStreamingContext& InContext)
 {
-	check(!InParams.ActorDescContainer);
-	FGenerateStreamingParams Params = FGenerateStreamingParams(InParams).SetActorDescContainer(ActorDescContainer);
+	check(InParams.ActorDescCollection.IsEmpty());
+	FGenerateStreamingParams Params = FGenerateStreamingParams(InParams);
+	Params.ActorDescCollection.Append(*this);
 
 	OnPreGenerateStreaming.Broadcast(InContext.PackagesToGenerate);
 
@@ -1242,7 +1246,7 @@ bool UWorldPartition::GenerateContainerStreaming(const FGenerateStreamingParams&
 		ErrorHandler = &MapCheckErrorHandler;
 	}
 
-	const FString ContainerPackageName = InParams.ActorDescContainer->ContainerPackageName.ToString();
+	const FString ContainerPackageName = InParams.ActorDescCollection.GetMainContainerPackageName().ToString();
 	FString ContainerShortName = FPackageName::GetShortName(ContainerPackageName);
 	if (!ContainerPackageName.StartsWith(TEXT("/Game/")))
 	{
@@ -1282,7 +1286,7 @@ bool UWorldPartition::GenerateContainerStreaming(const FGenerateStreamingParams&
 	FWorldPartitionStreamingGenerator StreamingGenerator(StreamingGeneratorParams);
 
 	// Preparation Phase
-	StreamingGenerator.PreparationPhase(InParams.ActorDescContainer);
+	StreamingGenerator.PreparationPhase(InParams.ActorDescCollection);
 
 	if (HierarchicalLogAr.IsValid())
 	{
@@ -1337,7 +1341,8 @@ void UWorldPartition::GenerateHLOD(ISourceControlHelper* SourceControlHelper, bo
 			.SetIsValidGrid([this](FName GridName) { return RuntimeHash->IsValidGrid(GridName); });
 
 		FWorldPartitionStreamingGenerator StreamingGenerator(StreamingGeneratorParams);
-		StreamingGenerator.PreparationPhase(InActorDescContainer);
+		FStreamingGenerationActorDescCollection ActorDescCollection{InActorDescContainer};
+		StreamingGenerator.PreparationPhase(ActorDescCollection);
 
 		TUniquePtr<FArchive> LogFileAr = FWorldPartitionStreamingGenerator::CreateDumpStateLogArchive(TEXT("HLOD"));
 		FHierarchicalLogArchive HierarchicalLogAr(*LogFileAr);
@@ -1371,8 +1376,110 @@ void UWorldPartition::CheckForErrors(IStreamingGenerationErrorHandler* ErrorHand
 	{
 		check(StreamingGeneratorParams.WorldPartitionContext == InActorDescContainer->GetWorldPartition());
 		FWorldPartitionStreamingGenerator StreamingGenerator(StreamingGeneratorParams);
-		StreamingGenerator.PreparationPhase(InActorDescContainer);
+		FStreamingGenerationActorDescCollection Collection{InActorDescContainer};
+		StreamingGenerator.PreparationPhase(Collection);
+	});
+	
+}
+
+FStreamingGenerationActorDescCollection::FStreamingGenerationActorDescCollection(std::initializer_list<const UActorDescContainer*> ActorDescContainerArray)
+	: TActorDescContainerCollection<const UActorDescContainer*>(ActorDescContainerArray)
+{
+	SortCollection();
+}
+
+UWorld* FStreamingGenerationActorDescCollection::GetWorld() const
+{
+	if (const UActorDescContainer* MainContainer = GetMainActorDescContainer())
+	{
+		if (UWorldPartition* WorldPartition = MainContainer->GetWorldPartition())
+		{
+			return WorldPartition->GetWorld();
+		}
+	}
+
+	return nullptr;
+}
+
+FGuid FStreamingGenerationActorDescCollection::GetContentBundleGuid() const
+{
+	if (const UActorDescContainer* MainContainer = GetMainActorDescContainer())
+	{
+		return MainContainer->GetContentBundleGuid();
+	}
+
+	return FGuid();
+}
+
+const UActorDescContainer* FStreamingGenerationActorDescCollection::GetMainActorDescContainer() const
+{
+	if (!IsEmpty())
+	{
+		return ActorDescContainerCollection[MainContainerIdx];
+	}
+
+	return nullptr;
+}
+
+FName FStreamingGenerationActorDescCollection::GetMainContainerPackageName() const
+{
+	if (const UActorDescContainer* MainContainer = GetMainActorDescContainer())
+	{
+		return MainContainer->GetContainerPackage();
+	}
+
+	return FName();
+}
+
+TArrayView<const UActorDescContainer*> FStreamingGenerationActorDescCollection::GetExternalDataLayerContainers()
+{
+	if (ActorDescContainerCollection.Num() <= 1)
+	{
+		return TArrayView<const UActorDescContainer*>();
+	}
+
+	uint32 NumExternalDataLayerContainers = ActorDescContainerCollection.Num() - ExternalDataLayerContainerStartIdx;
+	return MakeArrayView(&ActorDescContainerCollection[ExternalDataLayerContainerStartIdx], NumExternalDataLayerContainers);
+}
+
+void FStreamingGenerationActorDescCollection::OnCollectionChanged()
+{
+	SortCollection();
+}
+
+void FStreamingGenerationActorDescCollection::SortCollection()
+{
+	if (IsEmpty())
+	{
+		return;
+	}
+
+	auto IsMainPartitionContainer = [this](const UActorDescContainer* ActorDescContainer)
+	{
+		bool bIsContentBundleContainer = ActorDescContainer->GetContentBundleGuid().IsValid() && GetActorDescContainerCount() == 1 && !ActorDescContainer->GetContainerPackage().ToString().StartsWith(TEXT("/Game/"), ESearchCase::IgnoreCase);
+		return !ActorDescContainer->GetContentBundleGuid().IsValid() || bIsContentBundleContainer;
+	};
+
+	Algo::Sort(ActorDescContainerCollection, [&IsMainPartitionContainer](const UActorDescContainer* A, const UActorDescContainer* B)
+	{
+		if (IsMainPartitionContainer(A) && !IsMainPartitionContainer(B))
+		{
+			return true;
+		}
+		else if (IsMainPartitionContainer(B) && !IsMainPartitionContainer(A))
+		{
+			return false;
+		}
+
+		return A->GetContainerPackage().Compare(B->GetContainerPackage()) < 0;
 	});
+
+	check(IsMainPartitionContainer(ActorDescContainerCollection[MainContainerIdx]));
+
+#if DO_CHECK
+	TArrayView<const UActorDescContainer*> ExternalDataLayerContainerView = GetExternalDataLayerContainers();
+	Algo::ForEach(ExternalDataLayerContainerView, [&IsMainPartitionContainer](const UActorDescContainer* ActorDescContainer) { check(!IsMainPartitionContainer(ActorDescContainer)); });
+#endif
 }
 
 /* Deprecated */
@@ -1400,7 +1507,8 @@ void UWorldPartition::CheckForErrors(const FCheckForErrorsParams& Params)
 
 	FWorldPartitionStreamingGenerator StreamingGenerator(StreamingGeneratorParams);
 
-	StreamingGenerator.PreparationPhase(Params.ActorDescContainer);
+	FStreamingGenerationActorDescCollection Collection{Params.ActorDescContainer};
+	StreamingGenerator.PreparationPhase(Collection);
 }
 
 #endif // WITH_EDITOR
diff --git a/Engine/Source/Runtime/Engine/Private/WorldPartition/WorldPartitionStreamingGenerationContext.cpp b/Engine/Source/Runtime/Engine/Private/WorldPartition/WorldPartitionStreamingGenerationContext.cpp
--- a/Engine/Source/Runtime/Engine/Private/WorldPartition/WorldPartitionStreamingGenerationContext.cpp
+++ b/Engine/Source/Runtime/Engine/Private/WorldPartition/WorldPartitionStreamingGenerationContext.cpp
@@ -21,6 +21,6 @@ const FTransform& IStreamingGenerationContext::FActorInstance::GetTransform() co
 
 const UActorDescContainer* IStreamingGenerationContext::FActorInstance::GetActorDescContainer() const
 {
-	return ActorSetInstance->ContainerInstance->ActorDescContainer;
+	return GetActorDescView().GetActorDesc()->GetContainer();
 }
 #endif
diff --git a/Engine/Source/Runtime/Engine/Public/WorldPartition/ActorDescContainerCollection.h b/Engine/Source/Runtime/Engine/Public/WorldPartition/ActorDescContainerCollection.h
--- a/Engine/Source/Runtime/Engine/Public/WorldPartition/ActorDescContainerCollection.h
+++ b/Engine/Source/Runtime/Engine/Public/WorldPartition/ActorDescContainerCollection.h
@@ -8,17 +8,28 @@
 class FWorldPartitionActorDesc;
 class FReferenceCollector;
 
-class ENGINE_API FActorDescContainerCollection
+template<class ActorDescContPtrType>
+class TActorDescContainerCollection
 {
 #if WITH_EDITOR
 public:
-	~FActorDescContainerCollection();
+	template<class> friend class TActorDescContainerCollection;
 
-	void AddContainer(UActorDescContainer* Container);
-	bool RemoveContainer(UActorDescContainer* Container);
+	TActorDescContainerCollection() = default;
+
+	template<class U>
+	TActorDescContainerCollection(std::initializer_list<U> ActorDescContainerArray);
+
+	virtual ~TActorDescContainerCollection();
+
+	void AddContainer(ActorDescContPtrType Container);
+	bool RemoveContainer(ActorDescContPtrType Container);
 	bool Contains(const FName& ContainerPackageName) const;
-	UActorDescContainer* FindContainer(const FName& ContainerPackageName) const;
+	ActorDescContPtrType FindContainer(const FName& ContainerPackageName) const;
 
+	template<class U>
+	void Append(const TActorDescContainerCollection<U>& OtherCollection);
+	
 	void Empty();
 	bool IsEmpty() const { return ActorDescContainerCollection.IsEmpty(); }
 	uint32 GetActorDescContainerCount() const { return ActorDescContainerCollection.Num(); }
@@ -32,12 +43,16 @@ public:
 	const FWorldPartitionActorDesc* GetActorDescByName(const FString& ActorPath) const;
 	const FWorldPartitionActorDesc* GetActorDescByName(const FSoftObjectPath& ActorPath) const;
 
+	template<typename Dummy = void, typename = typename TEnableIf<!TIsConst<TRemovePointer<ActorDescContPtrType>>::Value, Dummy>::Type>
 	bool RemoveActor(const FGuid& ActorGuid);
 
+	template<typename Dummy = void, typename = typename TEnableIf<!TIsConst<TRemovePointer<ActorDescContPtrType>>::Value, Dummy>::Type>
 	void OnPackageDeleted(UPackage* Package);
 
+	template<typename Dummy = void, typename = typename TEnableIf<!TIsConst<TRemovePointer<ActorDescContPtrType>>::Value, Dummy>::Type>
 	bool IsActorDescHandled(const AActor* Actor) const;
 
+	template<typename Dummy = void, typename = typename TEnableIf<!TIsConst<TRemovePointer<ActorDescContPtrType>>::Value, Dummy>::Type>
 	void LoadAllActors(TArray<FWorldPartitionReference>& OutReferences);
 
 	DECLARE_EVENT_OneParam(UWorldPartition, FActorDescAddedEvent, FWorldPartitionActorDesc*);
@@ -45,23 +60,26 @@ public:
 
 	DECLARE_EVENT_OneParam(UWorldPartition, FActorDescRemovedEvent, FWorldPartitionActorDesc*);
 	FActorDescRemovedEvent OnActorDescRemovedEvent;
-
-	UActorDescContainer* GetActorDescContainer(const FGuid& ActorGuid);
-	const UActorDescContainer* GetActorDescContainer(const FGuid& ActorGuid) const;
-
-	void ForEachActorDescContainerBreakable(TFunctionRef<bool(UActorDescContainer*)> Func);
-	void ForEachActorDescContainerBreakable(TFunctionRef<bool(UActorDescContainer*)> Func) const;
-
-	void ForEachActorDescContainer(TFunctionRef<void(UActorDescContainer*)> Func);
-	void ForEachActorDescContainer(TFunctionRef<void(UActorDescContainer*)> Func) const;
+	
+	ActorDescContPtrType GetActorDescContainer(const FGuid& ActorGuid) const;
+	
+	ActorDescContPtrType FindHandlingContainer(const AActor* Actor) const;
+	
+	void ForEachActorDescContainerBreakable(TFunctionRef<bool(ActorDescContPtrType)> Func);
+	void ForEachActorDescContainerBreakable(TFunctionRef<bool(ActorDescContPtrType)> Func) const;
+	
+	void ForEachActorDescContainer(TFunctionRef<void(ActorDescContPtrType)> Func);
+	void ForEachActorDescContainer(TFunctionRef<void(ActorDescContPtrType)> Func) const;
 
 protected:
-	void RegisterDelegates(UActorDescContainer* Container);
-	void UnregisterDelegates(UActorDescContainer* Container);
+	void RegisterDelegates(ActorDescContPtrType Container);
+	void UnregisterDelegates(ActorDescContPtrType Container);
 
-	TArray<TObjectPtr<UActorDescContainer>> ActorDescContainerCollection;
+	virtual void OnCollectionChanged() {};
 
-public:
+	TArray<ActorDescContPtrType> ActorDescContainerCollection;
+
+private:
 	void OnActorDescAdded(FWorldPartitionActorDesc* ActorDesc);
 	void OnActorDescRemoved(FWorldPartitionActorDesc* ActorDesc);
 
@@ -73,8 +91,8 @@ public:
 
 	protected:
 		typedef UActorDescContainer ContainerType;
-		typedef TArray<TObjectPtr<ContainerType>> ContainerCollectionType;
-		typedef typename TChooseClass<bConst, ContainerCollectionType::TConstIterator, ContainerCollectionType::TIterator>::Result ContainerIteratorType;
+		typedef TArray<ActorDescContPtrType> ContainerCollectionType;
+		typedef typename TChooseClass<bConst, typename ContainerCollectionType::TConstIterator, typename ContainerCollectionType::TIterator>::Result ContainerIteratorType;
 		typedef typename TChooseClass<bConst, typename ContainerType::TConstIterator<ActorType>, typename ContainerType::TIterator<ActorType>>::Result ActDescIteratorType;
 
 		typedef typename FWorldPartitionActorDescType<ActorType>::Type ValueType;
@@ -176,3 +194,299 @@ public:
 	};
 #endif
 };
+
+#if WITH_EDITOR
+
+template<class ActorDescContPtrType>
+template<class U>
+TActorDescContainerCollection<ActorDescContPtrType>::TActorDescContainerCollection(std::initializer_list<U> ActorDescContainerArray)
+	:ActorDescContainerCollection(ActorDescContainerArray)
+{
+	ForEachActorDescContainer([this](ActorDescContPtrType ActorDescContainer)
+	{
+		RegisterDelegates(ActorDescContainer);
+	});
+}
+
+template<class ActorDescContPtrType>
+TActorDescContainerCollection<ActorDescContPtrType>::~TActorDescContainerCollection()
+{
+	Empty();
+}
+
+template<class ActorDescContPtrType>
+void TActorDescContainerCollection<ActorDescContPtrType>::AddContainer(ActorDescContPtrType Container)
+{
+	ActorDescContainerCollection.Add(Container);
+	RegisterDelegates(Container);
+	OnCollectionChanged();
+}
+
+template<class ActorDescContPtrType>
+bool TActorDescContainerCollection<ActorDescContPtrType>::RemoveContainer(ActorDescContPtrType Container)
+{
+	if (ActorDescContainerCollection.RemoveSwap(Container) > 0)
+	{
+		UnregisterDelegates(Container);
+		OnCollectionChanged();
+		return true;
+	}
+	return false;
+}
+
+template<class ActorDescContPtrType>
+template<class U>
+void TActorDescContainerCollection<ActorDescContPtrType>::Append(const TActorDescContainerCollection<U>& OtherCollection)
+{
+	if (OtherCollection.IsEmpty())
+	{
+		return;
+	}
+
+	ActorDescContainerCollection.Append(OtherCollection.ActorDescContainerCollection);
+	int32 AppendedContainerIndex = ActorDescContainerCollection.Num() - OtherCollection.ActorDescContainerCollection.Num();
+	for (; AppendedContainerIndex < ActorDescContainerCollection.Num(); ++AppendedContainerIndex)
+	{
+		RegisterDelegates(ActorDescContainerCollection[AppendedContainerIndex]);
+	}
+
+	OnCollectionChanged();
+}
+
+template<class ActorDescContPtrType>
+bool TActorDescContainerCollection<ActorDescContPtrType>::Contains(const FName& ContainerPackageName) const
+{
+	return FindContainer(ContainerPackageName) != nullptr;
+}
+
+template<class ActorDescContPtrType>
+ActorDescContPtrType TActorDescContainerCollection<ActorDescContPtrType>::FindContainer(const FName& ContainerPackageName) const
+{
+	auto ContainerPtr = ActorDescContainerCollection.FindByPredicate([&ContainerPackageName](ActorDescContPtrType ActorDescContainer) { return ActorDescContainer->GetContainerPackage() == ContainerPackageName; });
+	return ContainerPtr != nullptr ? *ContainerPtr : nullptr;
+}
+
+template<class ActorDescContPtrType>
+void TActorDescContainerCollection<ActorDescContPtrType>::Empty()
+{
+	ForEachActorDescContainer([this](ActorDescContPtrType ActorDescContainer)
+	{		
+		UnregisterDelegates(ActorDescContainer);
+	});
+
+	ActorDescContainerCollection.Empty();
+	OnCollectionChanged();
+}
+
+template<class ActorDescContPtrType>
+void TActorDescContainerCollection<ActorDescContPtrType>::RegisterDelegates(ActorDescContPtrType Container)
+{
+	using ActorDescContNonConstPtrType = typename std::remove_const<typename TRemovePointer<ActorDescContPtrType>::Type>::type*;
+	const_cast<ActorDescContNonConstPtrType>(Container)->OnActorDescAddedEvent.AddRaw(this, &TActorDescContainerCollection<ActorDescContPtrType>::OnActorDescAdded);
+	const_cast<ActorDescContNonConstPtrType>(Container)->OnActorDescRemovedEvent.AddRaw(this, &TActorDescContainerCollection::OnActorDescRemoved);
+}
+
+template<class ActorDescContPtrType>
+void TActorDescContainerCollection<ActorDescContPtrType>::UnregisterDelegates(ActorDescContPtrType Container)
+{
+	using ActorDescContNonConstPtrType = typename std::remove_const<typename TRemovePointer<ActorDescContPtrType>::Type>::type*;
+	const_cast<ActorDescContNonConstPtrType>(Container)->OnActorDescAddedEvent.RemoveAll(this);
+	const_cast<ActorDescContNonConstPtrType>(Container)->OnActorDescRemovedEvent.RemoveAll(this);
+}
+
+template<class ActorDescContPtrType>
+const FWorldPartitionActorDesc* TActorDescContainerCollection<ActorDescContPtrType>::GetActorDesc(const FGuid& Guid) const
+{
+	const FWorldPartitionActorDesc* ActorDesc = nullptr;
+	ForEachActorDescContainerBreakable([&Guid, &ActorDesc](ActorDescContPtrType ActorDescContainer)
+	{
+		ActorDesc = ActorDescContainer->GetActorDesc(Guid);
+		return ActorDesc == nullptr;
+	});
+
+	return ActorDesc;
+}
+
+template<class ActorDescContPtrType>
+FWorldPartitionActorDesc* TActorDescContainerCollection<ActorDescContPtrType>::GetActorDesc(const FGuid& Guid)
+{
+	return const_cast<FWorldPartitionActorDesc*>(const_cast<const TActorDescContainerCollection*>(this)->GetActorDesc(Guid));
+}
+
+template<class ActorDescContPtrType>
+const FWorldPartitionActorDesc& TActorDescContainerCollection<ActorDescContPtrType>::GetActorDescChecked(const FGuid& Guid) const
+{
+	const FWorldPartitionActorDesc* ActorDesc = GetActorDesc(Guid);
+	check(ActorDesc != nullptr);
+
+	static FWorldPartitionActorDesc EmptyDescriptor;
+	return ActorDesc != nullptr ? *ActorDesc : EmptyDescriptor;
+}
+
+template<class ActorDescContPtrType>
+FWorldPartitionActorDesc& TActorDescContainerCollection<ActorDescContPtrType>::GetActorDescChecked(const FGuid& Guid)
+{
+	return const_cast<FWorldPartitionActorDesc&>(const_cast<const TActorDescContainerCollection*>(this)->GetActorDescChecked(Guid));
+}
+
+template<class ActorDescContPtrType>
+const FWorldPartitionActorDesc* TActorDescContainerCollection<ActorDescContPtrType>::GetActorDescByName(const FString& ActorPath) const
+{
+	const FWorldPartitionActorDesc* ActorDesc = nullptr;
+	ForEachActorDescContainerBreakable([&ActorPath, &ActorDesc](ActorDescContPtrType ActorDescContainer)
+	{
+		ActorDesc = ActorDescContainer->GetActorDescByName(ActorPath);
+		return ActorDesc == nullptr;
+	});
+
+	return ActorDesc;
+}
+
+template<class ActorDescContPtrType>
+const FWorldPartitionActorDesc* TActorDescContainerCollection<ActorDescContPtrType>::GetActorDescByName(const FSoftObjectPath& ActorPath) const
+{
+	const FWorldPartitionActorDesc* ActorDesc = nullptr;
+	ForEachActorDescContainerBreakable([&ActorPath, &ActorDesc](ActorDescContPtrType ActorDescContainer)
+	{
+		ActorDesc = ActorDescContainer->GetActorDescByName(ActorPath);
+		return ActorDesc == nullptr;
+	});
+
+	return ActorDesc;
+}
+
+template<class ActorDescContPtrType>
+template<typename, typename>
+bool TActorDescContainerCollection<ActorDescContPtrType>::RemoveActor(const FGuid& ActorGuid)
+{
+	bool bRemoved = false;
+	ForEachActorDescContainerBreakable([&ActorGuid, &bRemoved](ActorDescContPtrType ActorDescContainer)
+	{
+		bRemoved = ActorDescContainer->RemoveActor(ActorGuid);
+		return !bRemoved;
+	});
+
+	return bRemoved;
+}
+
+template<class ActorDescContPtrType>
+template<typename, typename>
+void TActorDescContainerCollection<ActorDescContPtrType>::OnPackageDeleted(UPackage* Package)
+{
+	ForEachActorDescContainer([Package](ActorDescContPtrType ActorDescContainer)
+	{
+		ActorDescContainer->OnPackageDeleted(Package);
+	});
+}
+
+template<class ActorDescContPtrType>
+template<typename, typename>
+bool TActorDescContainerCollection<ActorDescContPtrType>::IsActorDescHandled(const AActor* Actor) const
+{
+	bool bIsHandled = false;
+	ForEachActorDescContainerBreakable([Actor, &bIsHandled](ActorDescContPtrType ActorDescContainer)
+	{
+		bIsHandled = ActorDescContainer->IsActorDescHandled(Actor);
+		return !bIsHandled;
+	});
+
+	return bIsHandled;
+}
+
+template<class ActorDescContPtrType>
+template<typename, typename>
+void TActorDescContainerCollection<ActorDescContPtrType>::LoadAllActors(TArray<FWorldPartitionReference>& OutReferences)
+{
+	ForEachActorDescContainer([&OutReferences](ActorDescContPtrType ActorDescContainer)
+	{
+		ActorDescContainer->LoadAllActors(OutReferences);
+	});
+}
+
+template<class ActorDescContPtrType>
+ActorDescContPtrType TActorDescContainerCollection<ActorDescContPtrType>::GetActorDescContainer(const FGuid& ActorGuid) const
+{
+	ActorDescContPtrType ActorDescContainer = nullptr;
+	ForEachActorDescContainerBreakable([&ActorGuid, &ActorDescContainer](ActorDescContPtrType InActorDescContainer)
+	{
+		if (InActorDescContainer->GetActorDesc(ActorGuid) != nullptr)
+		{
+			ActorDescContainer = InActorDescContainer;
+		}
+		return ActorDescContainer == nullptr;
+	});
+
+	return ActorDescContainer;
+}
+
+template<class ActorDescContPtrType>
+ActorDescContPtrType TActorDescContainerCollection<ActorDescContPtrType>::FindHandlingContainer(const AActor* Actor) const
+{
+	// Actor is already in a container
+	if (ActorDescContPtrType ActorDescContainer = GetActorDescContainer(Actor->GetActorGuid()))
+	{
+		return ActorDescContainer;
+	}
+
+	// Actor is not yet stored in a container. Find which one should handle it.
+	ActorDescContPtrType ActorDescContainer = nullptr;
+	ForEachActorDescContainerBreakable([&ActorDescContainer, &Actor](ActorDescContPtrType InActorDescContainer)
+	{
+		if (InActorDescContainer->IsActorDescHandled(Actor))
+		{
+			ActorDescContainer = InActorDescContainer;
+		}
+		return ActorDescContainer == nullptr;
+	});
+
+	return ActorDescContainer;
+}
+
+template<class ActorDescContPtrType>
+void TActorDescContainerCollection<ActorDescContPtrType>::ForEachActorDescContainerBreakable(TFunctionRef<bool(ActorDescContPtrType)> Func) const
+{
+	for (ActorDescContPtrType ActorDescContainer : ActorDescContainerCollection)
+	{
+		if (!Func(ActorDescContainer))
+		{
+			break;
+		}
+	}
+}
+
+template<class ActorDescContPtrType>
+void TActorDescContainerCollection<ActorDescContPtrType>::ForEachActorDescContainerBreakable(TFunctionRef<bool(ActorDescContPtrType)> Func)
+{
+	const_cast<const TActorDescContainerCollection*>(this)->ForEachActorDescContainerBreakable(Func);
+}
+
+template<class ActorDescContPtrType>
+void TActorDescContainerCollection<ActorDescContPtrType>::ForEachActorDescContainer(TFunctionRef<void(ActorDescContPtrType)> Func) const
+{
+	for (ActorDescContPtrType ActorDescContainer : ActorDescContainerCollection)
+	{
+		Func(ActorDescContainer);
+	}
+}
+
+template<class ActorDescContPtrType>
+void TActorDescContainerCollection<ActorDescContPtrType>::ForEachActorDescContainer(TFunctionRef<void(ActorDescContPtrType)> Func)
+{
+	const_cast<const TActorDescContainerCollection*>(this)->ForEachActorDescContainer(Func);
+}
+
+template<class ActorDescContPtrType>
+void TActorDescContainerCollection<ActorDescContPtrType>::OnActorDescAdded(FWorldPartitionActorDesc* ActorDesc)
+{
+	OnActorDescAddedEvent.Broadcast(ActorDesc);
+}
+
+template<class ActorDescContPtrType>
+void TActorDescContainerCollection<ActorDescContPtrType>::OnActorDescRemoved(FWorldPartitionActorDesc* ActorDesc)
+{
+	OnActorDescRemovedEvent.Broadcast(ActorDesc);
+}
+
+#endif 
+
+using FActorDescContainerCollection = TActorDescContainerCollection<UActorDescContainer*>;
\ No newline at end of file
diff --git a/Engine/Source/Runtime/Engine/Public/WorldPartition/WorldPartition.h b/Engine/Source/Runtime/Engine/Public/WorldPartition/WorldPartition.h
--- a/Engine/Source/Runtime/Engine/Public/WorldPartition/WorldPartition.h
+++ b/Engine/Source/Runtime/Engine/Public/WorldPartition/WorldPartition.h
@@ -17,6 +17,7 @@
 #include "WorldPartition/Cook/WorldPartitionCookPackageGenerator.h"
 
 #if WITH_EDITOR
+#include "WorldPartition/WorldPartitionStreamingGeneration.h"
 #include "WorldPartition/WorldPartitionActorLoaderInterface.h"
 #include "WorldPartition/WorldPartitionEditorLoaderAdapter.h"
 #include "PackageSourceControlHelper.h"
@@ -187,10 +188,14 @@ public:
 		FGenerateStreamingParams()
 		{}
 
-		const UActorDescContainer* ActorDescContainer = nullptr;
+		FStreamingGenerationActorDescCollection ActorDescCollection;
 		TOptional<const FString> OutputLogPath;
 
-		FGenerateStreamingParams& SetActorDescContainer(const UActorDescContainer* InActorDescContainer) { ActorDescContainer = InActorDescContainer; return *this; }
+		FGenerateStreamingParams& SetActorDescContainer(const UActorDescContainer* InActorDescContainer) 
+		{ 
+			ActorDescCollection.AddContainer(InActorDescContainer);
+			return *this;
+		}
 		FGenerateStreamingParams& SetOutputLogPath(const FString& InOutputLogPath) { OutputLogPath = InOutputLogPath; return *this; }
 	};
 
diff --git a/Engine/Source/Runtime/Engine/Public/WorldPartition/WorldPartitionActorDesc.h b/Engine/Source/Runtime/Engine/Public/WorldPartition/WorldPartitionActorDesc.h
--- a/Engine/Source/Runtime/Engine/Public/WorldPartition/WorldPartitionActorDesc.h
+++ b/Engine/Source/Runtime/Engine/Public/WorldPartition/WorldPartitionActorDesc.h
@@ -70,13 +70,13 @@ class ENGINE_API FWorldPartitionActorDesc
 {
 	friend class AActor;
 	friend class UWorldPartition;
-	friend class FActorDescContainerCollection;
 	friend struct FWorldPartitionHandleImpl;
 	friend struct FWorldPartitionReferenceImpl;
 	friend struct FWorldPartitionActorDescUtils;
 	friend struct FWorldPartitionActorDescUnitTestAcccessor;
 	friend class FAssetRootPackagePatcher;
 	friend class FActorDescArchive;
+	template<class U> friend class TActorDescContainerCollection;
 
 public:
 	struct FContainerInstance
diff --git a/Engine/Source/Runtime/Engine/Public/WorldPartition/WorldPartitionHandle.h b/Engine/Source/Runtime/Engine/Public/WorldPartition/WorldPartitionHandle.h
--- a/Engine/Source/Runtime/Engine/Public/WorldPartition/WorldPartitionHandle.h
+++ b/Engine/Source/Runtime/Engine/Public/WorldPartition/WorldPartitionHandle.h
@@ -8,11 +8,13 @@
 
 #if WITH_EDITOR
 class UActorDescContainer;
-class FActorDescContainerCollection;
 class FWorldPartitionActorDesc;
 struct FWorldPartitionHandleImpl;
 struct FWorldPartitionReferenceImpl;
 
+template<class U>
+class TActorDescContainerCollection;
+
 class ENGINE_API FWorldPartitionLoadingContext
 {
 public:
@@ -119,7 +121,8 @@ public:
 		}
 	}
 
-	FORCEINLINE TWorldPartitionHandle(FActorDescContainerCollection* ContainerCollection, const FGuid& ActorGuid)
+	template<class U>
+	FORCEINLINE TWorldPartitionHandle(TActorDescContainerCollection<U>* ContainerCollection, const FGuid& ActorGuid)
 		: Container(Impl::GetActorDescContainer(ContainerCollection, ActorGuid))
 		, ActorDesc(nullptr)
 	{
@@ -339,8 +342,13 @@ struct ENGINE_API FWorldPartitionImplBase
 {
 	static TUniquePtr<FWorldPartitionActorDesc>* GetActorDesc(UActorDescContainer* Container, const FGuid& ActorGuid);
 	static UActorDescContainer* GetActorDescContainer(TUniquePtr<FWorldPartitionActorDesc>* ActorDesc);
-	static UActorDescContainer* GetActorDescContainer(FActorDescContainerCollection* ContainerCollection, const FGuid& ActorGuid);
 	static bool IsActorDescLoaded(FWorldPartitionActorDesc* ActorDesc);
+
+	template<class U>
+	static UActorDescContainer* GetActorDescContainer(TActorDescContainerCollection<U>* ContainerCollection, const FGuid& ActorGuid)
+	{
+		return ContainerCollection ? ContainerCollection->GetActorDescContainer(ActorGuid) : nullptr;
+	}
 };
 
 struct ENGINE_API FWorldPartitionHandleImpl : FWorldPartitionImplBase
diff --git a/Engine/Source/Runtime/Engine/Public/WorldPartition/WorldPartitionStreamingGeneration.h b/Engine/Source/Runtime/Engine/Public/WorldPartition/WorldPartitionStreamingGeneration.h
--- a/Engine/Source/Runtime/Engine/Public/WorldPartition/WorldPartitionStreamingGeneration.h
+++ b/Engine/Source/Runtime/Engine/Public/WorldPartition/WorldPartitionStreamingGeneration.h
@@ -4,6 +4,7 @@
 #if WITH_EDITOR
 #include "CoreMinimal.h"
 #include "OverrideVoidReturnInvoker.h"
+#include "WorldPartition/ActorDescContainerCollection.h"
 #include "WorldPartition/WorldPartitionActorDescView.h"
 
 class ENGINE_API FActorDescViewMap
@@ -100,4 +101,29 @@ protected:
 	TMultiMap<FName, const FWorldPartitionActorDescView*> ActorDescViewsByClass;
 };
 
+class ENGINE_API FStreamingGenerationActorDescCollection : public TActorDescContainerCollection<const UActorDescContainer*>
+{
+public:
+	FStreamingGenerationActorDescCollection() = default;
+	FStreamingGenerationActorDescCollection(std::initializer_list<const UActorDescContainer*> ActorDescContainerArray);
+
+	UWorld* GetWorld() const;
+
+	// @todo_ow : Remove once conversion to ExternalDataLayer is complete. 
+	// It is present to handle content bundles streaming generation via the same code path. 
+	FGuid GetContentBundleGuid() const;
+
+	const UActorDescContainer* GetMainActorDescContainer() const;
+	FName GetMainContainerPackageName() const;
+	TArrayView<const UActorDescContainer*> GetExternalDataLayerContainers();
+
+	virtual void OnCollectionChanged() override;
+
+private:
+	void SortCollection();
+
+	static constexpr int MainContainerIdx = 0;
+	static constexpr int ExternalDataLayerContainerStartIdx = MainContainerIdx + 1;
+};
+
 #endif // WITH_EDITOR
\ No newline at end of file
diff --git a/Engine/Source/Runtime/Engine/Public/WorldPartition/WorldPartitionStreamingGenerationContext.h b/Engine/Source/Runtime/Engine/Public/WorldPartition/WorldPartitionStreamingGenerationContext.h
--- a/Engine/Source/Runtime/Engine/Public/WorldPartition/WorldPartitionStreamingGenerationContext.h
+++ b/Engine/Source/Runtime/Engine/Public/WorldPartition/WorldPartitionStreamingGenerationContext.h
@@ -2,9 +2,12 @@
 #pragma once
 
 #include "CoreMinimal.h"
-#include "WorldPartition/WorldPartitionStreamingGeneration.h"
 #include "WorldPartition/WorldPartitionActorContainerID.h"
 
+class FActorDescViewMap;
+class FStreamingGenerationActorDescCollection;
+class FWorldPartitionActorDescView;
+class UActorDescContainer;
 class UDataLayerInstance;
 
 #if WITH_EDITOR
@@ -20,7 +23,7 @@ public:
 	{
 		FActorSetContainer()
 			: ActorDescViewMap(nullptr)
-			, ActorDescContainer(nullptr)
+			, ActorDescCollection(nullptr)
 		{}
 
 		// Non-copyable
@@ -28,7 +31,7 @@ public:
 		FActorSetContainer& operator=(const FActorSetContainer&) = delete;
 
 		const FActorDescViewMap* ActorDescViewMap;
-		const UActorDescContainer* ActorDescContainer;
+		const FStreamingGenerationActorDescCollection* ActorDescCollection;
 		TArray<TUniquePtr<FActorSet>> ActorSets;
 	};
 
