World Partition : Fixed rare case where a WP streaming cell would contain an actor and this actor is also an editor-only reference of another actor of this cell. In this case, the editor-only reference must not be part of the WP streaming cell.

diff --git a/Engine/Source/Runtime/Engine/Private/WorldPartition/WorldPartitionRuntimeHash.cpp b/Engine/Source/Runtime/Engine/Private/WorldPartition/WorldPartitionRuntimeHash.cpp
--- a/Engine/Source/Runtime/Engine/Private/WorldPartition/WorldPartitionRuntimeHash.cpp
+++ b/Engine/Source/Runtime/Engine/Private/WorldPartition/WorldPartitionRuntimeHash.cpp
@@ -238,6 +238,7 @@ void UWorldPartitionRuntimeHash::PopulateRuntimeCell(UWorldPartitionRuntimeCell*
 	}
 
 	RuntimeCell->RuntimeCellData->ContentBounds = CellContentBounds;
+	RuntimeCell->Fixup();
 
 	// Always loaded cell actors are transfered to World's Persistent Level (see UWorldPartitionRuntimeSpatialHash::PopulateGeneratorPackageForCook)
 	if (OutPackagesToGenerate && RuntimeCell->GetActorCount() && !RuntimeCell->IsAlwaysLoaded())
diff --git a/Engine/Source/Runtime/Engine/Private/WorldPartition/WorldPartitionRuntimeLevelStreamingCell.cpp b/Engine/Source/Runtime/Engine/Private/WorldPartition/WorldPartitionRuntimeLevelStreamingCell.cpp
--- a/Engine/Source/Runtime/Engine/Private/WorldPartition/WorldPartitionRuntimeLevelStreamingCell.cpp
+++ b/Engine/Source/Runtime/Engine/Private/WorldPartition/WorldPartitionRuntimeLevelStreamingCell.cpp
@@ -12,9 +12,6 @@
 
 #include UE_INLINE_GENERATED_CPP_BY_NAME(WorldPartitionRuntimeLevelStreamingCell)
 
-#if WITH_EDITOR
-#endif
-
 UWorldPartitionRuntimeLevelStreamingCell::UWorldPartitionRuntimeLevelStreamingCell(const FObjectInitializer& ObjectInitializer)
 	: Super(ObjectInitializer)
 	, LevelStreaming(nullptr)
@@ -262,6 +259,23 @@ void UWorldPartitionRuntimeLevelStreamingCell::AddActorToCell(const FWorldPartit
 	);
 }
 
+void UWorldPartitionRuntimeLevelStreamingCell::Fixup()
+{
+	TMap<FGuid, FWorldPartitionRuntimeCellObjectMapping> UniquePackages;
+	UniquePackages.Reserve(Packages.Num());
+
+	for (FWorldPartitionRuntimeCellObjectMapping& Package : Packages)
+	{
+		FWorldPartitionRuntimeCellObjectMapping& UniquePackage = UniquePackages.FindOrAdd(Package.ActorInstanceGuid, Package);
+		UniquePackage.bIsEditorOnly &= Package.bIsEditorOnly;
+	};
+
+	if (UniquePackages.Num() != Packages.Num())
+	{
+		UniquePackages.GenerateValueArray(Packages);
+	}
+}
+
 bool UWorldPartitionRuntimeLevelStreamingCell::PopulateGeneratorPackageForCook(TArray<UPackage*>& OutModifiedPackages)
 {
 	check(IsAlwaysLoaded());
diff --git a/Engine/Source/Runtime/Engine/Public/WorldPartition/WorldPartitionRuntimeCell.h b/Engine/Source/Runtime/Engine/Public/WorldPartition/WorldPartitionRuntimeCell.h
--- a/Engine/Source/Runtime/Engine/Public/WorldPartition/WorldPartitionRuntimeCell.h
+++ b/Engine/Source/Runtime/Engine/Public/WorldPartition/WorldPartitionRuntimeCell.h
@@ -234,6 +234,7 @@ class ENGINE_API UWorldPartitionRuntimeCell : public UObject, public IWorldParti
 	//~End IWorldPartitionCell Interface
 
 	virtual void AddActorToCell(const FWorldPartitionActorDescView& ActorDescView, const FActorContainerID& InContainerID, const FTransform& InContainerTransform, const UActorDescContainer* InContainer) PURE_VIRTUAL(UWorldPartitionRuntimeCell::AddActorToCell,);
+	virtual void Fixup() PURE_VIRTUAL(UWorldPartitionRuntimeCell::Fixup, );
 	virtual int32 GetActorCount() const PURE_VIRTUAL(UWorldPartitionRuntimeCell::GetActorCount, return 0;);
 
 	// Cook methods
diff --git a/Engine/Source/Runtime/Engine/Public/WorldPartition/WorldPartitionRuntimeLevelStreamingCell.h b/Engine/Source/Runtime/Engine/Public/WorldPartition/WorldPartitionRuntimeLevelStreamingCell.h
--- a/Engine/Source/Runtime/Engine/Public/WorldPartition/WorldPartitionRuntimeLevelStreamingCell.h
+++ b/Engine/Source/Runtime/Engine/Public/WorldPartition/WorldPartitionRuntimeLevelStreamingCell.h
@@ -49,6 +49,7 @@ class ENGINE_API UWorldPartitionRuntimeLevelStreamingCell : public UWorldPartiti
 #if WITH_EDITOR
 	//~Begin UWorldPartitionRuntimeCell Interface
 	virtual void AddActorToCell(const FWorldPartitionActorDescView& ActorDescView, const FActorContainerID& InContainerID, const FTransform& InContainerTransform, const UActorDescContainer* InContainer) override;
+	virtual void Fixup() override;
 	virtual int32 GetActorCount() const override;
 	virtual void DumpStateLog(FHierarchicalLogArchive& Ar) override;
 	// Cook methods
