Don't consider objects part of a PIE package as assets.
Backout CLs 25587512 + CL25420856.

diff --git a/Engine/Source/Runtime/AssetRegistry/Private/AssetRegistry.cpp b/Engine/Source/Runtime/AssetRegistry/Private/AssetRegistry.cpp
--- a/Engine/Source/Runtime/AssetRegistry/Private/AssetRegistry.cpp
+++ b/Engine/Source/Runtime/AssetRegistry/Private/AssetRegistry.cpp
@@ -1997,6 +1997,9 @@ void EnumerateMemoryAssetsHelper(const FARCompiledFilter& InFilter, TSet<FName>&
 				return;
 			}
 
+			check(!Obj->GetPackage()->HasAnyPackageFlags(PKG_PlayInEditor));
+			check(!Obj->GetOutermostObject()->GetPackage()->HasAnyPackageFlags(PKG_PlayInEditor));
+
 			UPackage* InMemoryPackage = Obj->GetOutermost();
 
 			// Skip assets with any of the specified 'without' package flags 
@@ -7042,7 +7045,7 @@ void UAssetRegistryImpl::ReadLockEnumerateTagToAssetDatas(TFunctionRef<void(FNam
 void UAssetRegistryImpl::Broadcast(UE::AssetRegistry::Impl::FEventContext& EventContext)
 {
 	using namespace UE::AssetRegistry::Impl;
-	if (!IsInGameThread())
+	if (!IsInGameThread() || FUObjectThreadContext::Get().IsRoutingPostLoad)
 	{
 		// By contract events (and PackageLoads) can only be sent on the game thread; some legacy systems depend on 
 		// this and are not threadsafe. If we're not in the game thread, defer all events in the EventContext
diff --git a/Engine/Source/Runtime/CoreUObject/Private/UObject/Obj.cpp b/Engine/Source/Runtime/CoreUObject/Private/UObject/Obj.cpp
--- a/Engine/Source/Runtime/CoreUObject/Private/UObject/Obj.cpp
+++ b/Engine/Source/Runtime/CoreUObject/Private/UObject/Obj.cpp
@@ -2363,8 +2363,8 @@ bool UObject::IsAsset() const
 		// Don't count objects embedded in other objects (e.g. font textures, sequences, material expressions)
 		if ( UPackage* LocalOuterPackage = dynamic_cast<UPackage*>(GetOuter()) )
 		{
-			// Also exclude any objects found in the transient package, or in a package that is transient.
-			return LocalOuterPackage != GetTransientPackage() && !LocalOuterPackage->HasAnyFlags(RF_Transient);
+			// Also exclude any objects found in the transient package, in a package that is transient or in a play in editor package.
+			return LocalOuterPackage != GetTransientPackage() && !LocalOuterPackage->HasAnyFlags(RF_Transient) && !LocalOuterPackage->HasAnyPackageFlags(PKG_PlayInEditor);
 		}
 	}
 
diff --git a/Engine/Source/Runtime/Engine/Private/Actor.cpp b/Engine/Source/Runtime/Engine/Private/Actor.cpp
--- a/Engine/Source/Runtime/Engine/Private/Actor.cpp
+++ b/Engine/Source/Runtime/Engine/Private/Actor.cpp
@@ -427,7 +427,7 @@ bool AActor::IsAsset() const
 {
 	// External actors are considered assets, to allow using the asset logic for save dialogs, etc.
 	// Also, they return true even if pending kill, in order to show up as deleted in these dialogs.
-	return IsPackageExternal() && !GetPackage()->HasAnyFlags(RF_Transient) && !HasAnyFlags(RF_Transient | RF_ClassDefaultObject) && IsMainPackageActor();
+	return IsPackageExternal() && !GetPackage()->HasAnyFlags(RF_Transient) && !HasAnyFlags(RF_Transient | RF_ClassDefaultObject) && IsMainPackageActor() && !GetPackage()->HasAnyPackageFlags(PKG_PlayInEditor);
 }
 
 bool AActor::PreSaveRoot(const TCHAR* InFilename)
diff --git a/Engine/Source/Runtime/Engine/Private/ActorFolder.cpp b/Engine/Source/Runtime/Engine/Private/ActorFolder.cpp
--- a/Engine/Source/Runtime/Engine/Private/ActorFolder.cpp
+++ b/Engine/Source/Runtime/Engine/Private/ActorFolder.cpp
@@ -48,7 +48,7 @@ bool UActorFolder::IsAsset() const
 {
 	// Actor Folders are considered assets to allow using the asset logic for save dialogs, etc.
 	// Also, they return true even if pending kill, in order to show up as deleted in these dialogs.
-	return IsPackageExternal() && !GetPackage()->HasAnyFlags(RF_Transient) && !HasAnyFlags(RF_Transient | RF_ClassDefaultObject);
+	return IsPackageExternal() && !GetPackage()->HasAnyFlags(RF_Transient) && !HasAnyFlags(RF_Transient | RF_ClassDefaultObject) && !GetPackage()->HasAnyPackageFlags(PKG_PlayInEditor);
 }
 
 namespace ActorFolder
diff --git a/Engine/Source/Runtime/Engine/Private/Level.cpp b/Engine/Source/Runtime/Engine/Private/Level.cpp
--- a/Engine/Source/Runtime/Engine/Private/Level.cpp
+++ b/Engine/Source/Runtime/Engine/Private/Level.cpp
@@ -58,7 +58,6 @@ Level.cpp: Level-related functions
 #include "EditorActorFolders.h"
 #include "UObject/MetaData.h"
 #include "WorldPartition/WorldPartitionHelpers.h"
-#include "HAL/FileManager.h"
 #include "HAL/PlatformFileManager.h"
 #include "Misc/PathViews.h"
 #include "Selection.h"
@@ -3955,19 +3954,20 @@ TArray<FString> ULevel::GetOnDiskExternalActorPackages(const FString& ExternalAc
 	TArray<FString> ActorPackageNames;
 	if (!ExternalActorsPath.IsEmpty())
 	{
-		IFileManager::Get().IterateDirectoryRecursively(*FPackageName::LongPackageNameToFilename(ExternalActorsPath), [&ActorPackageNames](const TCHAR* FilenameOrDirectory, bool bIsDirectory)
-		{
-			if (!bIsDirectory)
-			{
-				FString Filename(FilenameOrDirectory);
-				if (Filename.EndsWith(FPackageName::GetAssetPackageExtension()))
-				{
-					ActorPackageNames.Add(FPackageName::FilenameToLongPackageName(Filename));
-				}
-			}
-			return true;
-		});
+		FARFilter Filter;
+		Filter.bIncludeOnlyOnDiskAssets = true;
+		Filter.PackagePaths.Add(*ExternalActorsPath);
+		Filter.bRecursivePaths = true;
+	
+		TArray<FAssetData> ActorAssets;
+		IAssetRegistry& AssetRegistry = FModuleManager::LoadModuleChecked<FAssetRegistryModule>(TEXT("AssetRegistry")).Get();
+		AssetRegistry.ScanSynchronous({ ExternalActorsPath }, TArray<FString>());
+		AssetRegistry.GetAssets(Filter, ActorAssets);
+
+		ActorPackageNames.Reserve(ActorAssets.Num());
+		Algo::Transform(ActorAssets, ActorPackageNames, [](const FAssetData& ActorAssetData) { return ActorAssetData.PackageName.ToString(); });
 	}
+
 	return ActorPackageNames;
 }
 
@@ -4573,3 +4573,5 @@ void ULevel::OnMultipleLSAsPopupDismissed()
 #endif // WITH_EDITOR
 
 #undef LOCTEXT_NAMESPACE
+
+UE_ENABLE_OPTIMIZATION
\ No newline at end of file
diff --git a/Engine/Source/Runtime/Engine/Public/ExternalPackageHelper.h b/Engine/Source/Runtime/Engine/Public/ExternalPackageHelper.h
--- a/Engine/Source/Runtime/Engine/Public/ExternalPackageHelper.h
+++ b/Engine/Source/Runtime/Engine/Public/ExternalPackageHelper.h
@@ -7,7 +7,6 @@
 #include "UObject/LinkerLoad.h"
 #include "UObject/Package.h"
 #include "UObject/UObjectThreadContext.h"
-#include "UObject/ObjectMacros.h"
 #include "AssetRegistry/AssetData.h"
 #include "AssetRegistry/AssetRegistryModule.h"
 #include "Misc/ArchiveMD5.h"
@@ -100,8 +99,6 @@ void FExternalPackageHelper::LoadObjectsFromExternalPackages(UObject* InOuter, T
 	Filter.ClassPaths.Add(T::StaticClass()->GetClassPathName());
 	Filter.bRecursiveClasses = true;
 	Filter.PackagePaths.Add(*ExternalObjectsPath);
-	// @todo_ow: remove once bIncludeOnlyOnDiskAssets=true doesn't return PIE Packages anymore
-	Filter.WithoutPackageFlags = PKG_PlayInEditor;
 	TArray<FAssetData> Assets;
 	AssetRegistry.GetAssets(Filter, Assets);
 
