Generalize list assets commands:
- MediaTexture now properly reports its size
- Added missing constructor overrides for AutoConsoleCommands
- Unified ListTexture reporting data

#tests PIE

diff --git a/Engine/Source/Runtime/Core/Public/HAL/IConsoleManager.h b/Engine/Source/Runtime/Core/Public/HAL/IConsoleManager.h
--- a/Engine/Source/Runtime/Core/Public/HAL/IConsoleManager.h
+++ b/Engine/Source/Runtime/Core/Public/HAL/IConsoleManager.h
@@ -1756,6 +1756,32 @@ public:
 	{
 	}
 
+	/**
+	* Register a console command that takes an output device
+	*
+	* @param	Name		The name of this command (must not be nullptr)
+	* @param	Help		Help text for this command
+	* @param	Command		The user function to call when this command is executed
+	* @param	Flags		Optional flags bitmask
+	*/
+	FAutoConsoleCommand(const TCHAR* Name, const TCHAR* Help, const FConsoleCommandWithOutputDeviceDelegate& Command, uint32 Flags = ECVF_Default)
+		: FAutoConsoleObject(IConsoleManager::Get().RegisterConsoleCommand(Name, Help, Command, Flags))
+	{
+	}
+
+	/**
+	* Register a console command that takes a world argument
+	*
+	* @param	Name		The name of this command (must not be nullptr)
+	* @param	Help		Help text for this command
+	* @param	Command		The user function to call when this command is executed
+	* @param	Flags		Optional flags bitmask
+	*/
+	FAutoConsoleCommand(const TCHAR* Name, const TCHAR* Help, const FConsoleCommandWithWorldDelegate& Command, uint32 Flags = ECVF_Default)
+		: FAutoConsoleObject(IConsoleManager::Get().RegisterConsoleCommand(Name, Help, Command, Flags))
+	{
+	}
+
 	/**
 	* Register a console command that takes arguments, a world argument and an output device
 	*
diff --git a/Engine/Source/Runtime/Engine/Classes/Sound/SoundGroups.h b/Engine/Source/Runtime/Engine/Classes/Sound/SoundGroups.h
--- a/Engine/Source/Runtime/Engine/Classes/Sound/SoundGroups.h
+++ b/Engine/Source/Runtime/Engine/Classes/Sound/SoundGroups.h
@@ -38,7 +38,7 @@ enum ESoundGroup : int
 };
 
 USTRUCT()
-struct FSoundGroup
+struct ENGINE_API FSoundGroup
 {
 	GENERATED_USTRUCT_BODY()
 
@@ -72,7 +72,7 @@ struct FSoundGroup
 
 // This class is a singleton initialized from the ini
 UCLASS(config=Engine, abstract)
-class USoundGroups : public UObject
+class ENGINE_API USoundGroups : public UObject
 {
 	GENERATED_UCLASS_BODY()
 
diff --git a/Engine/Source/Runtime/Engine/Private/UnrealEngine.cpp b/Engine/Source/Runtime/Engine/Private/UnrealEngine.cpp
--- a/Engine/Source/Runtime/Engine/Private/UnrealEngine.cpp
+++ b/Engine/Source/Runtime/Engine/Private/UnrealEngine.cpp
@@ -6010,7 +6010,7 @@ bool UEngine::HandleListTexturesCommand( const TCHAR* Cmd, FOutputDevice& Ar )
 		bListUnused ? TEXT(" unused") : TEXT(""));
 
 	// Find out how many times a texture is referenced by primitive components.
-	TMap<UTexture2D*,int32> TextureToUsageMap;
+	TMap<UTexture*,int32> TextureToUsageMap;
 	for( TObjectIterator<UPrimitiveComponent> It; It; ++It )
 	{
 		UPrimitiveComponent* PrimitiveComponent = *It;
@@ -6025,7 +6025,7 @@ bool UEngine::HandleListTexturesCommand( const TCHAR* Cmd, FOutputDevice& Ar )
 		// Increase usage count for all referenced textures
 		for( int32 TextureIndex=0; TextureIndex<StreamingTextures.Num(); TextureIndex++ )
 		{
-			UTexture2D* Texture = Cast<UTexture2D>(StreamingTextures[TextureIndex].RenderAsset);
+			UTexture* Texture = Cast<UTexture>(StreamingTextures[TextureIndex].RenderAsset);
 			if( Texture )
 			{
 				// Initializes UsageCount to 0 if texture is not found.
@@ -6054,21 +6054,21 @@ bool UEngine::HandleListTexturesCommand( const TCHAR* Cmd, FOutputDevice& Ar )
 		UTextureRenderTarget2D* RenderTexture	= Cast<UTextureRenderTarget2D>(Texture);
 
 		int32				LODGroup			= Texture->LODGroup;
-		int32				NumMips				= 0;
-		int32				MaxResLODBias		= 0;
-		int32				MaxAllowedSizeX		= 0;
-		int32				MaxAllowedSizeY		= 0;
+		int32				NumMips				= 1;
+		int32				MaxResLODBias		= Texture->GetCachedLODBias();;
+		int32				MaxAllowedSizeX		= FMath::Max<int32>(static_cast<int32>(Texture->GetSurfaceWidth()) >> MaxResLODBias, 1);
+		int32				MaxAllowedSizeY		= FMath::Max<int32>(static_cast<int32>(Texture->GetSurfaceHeight()) >> MaxResLODBias, 1);
 		EPixelFormat		Format				= PF_Unknown;
-		int32				DroppedMips			= 0;
-		int32				CurSizeX			= 0;
-		int32				CurSizeY			= 0;
-		bool				bIsStreamingTexture = false;
+		int32				DroppedMips			= MaxResLODBias;
+		int32				CurSizeX			= FMath::Max<int32>(static_cast<int32>(Texture->GetSurfaceWidth()) >> DroppedMips, 1);
+		int32				CurSizeY			= FMath::Max<int32>(static_cast<int32>(Texture->GetSurfaceHeight()) >> DroppedMips, 1);
+		bool				bIsStreamingTexture = Texture->GetStreamingIndex() != INDEX_NONE;
 		int32				MaxAllowedSize		= Texture->CalcTextureMemorySizeEnum(TMC_AllMipsBiased);
 		int32				CurrentSize			= Texture->CalcTextureMemorySizeEnum(TMC_ResidentMips);
-		int32				UsageCount			= 0;
-		bool				bIsForced			= false;
-		bool				bIsVirtual			= false;
-		bool				bIsUncompressed		= false;
+		int32				UsageCount			= TextureToUsageMap.FindRef(Texture);
+		bool				bIsForced			= Texture->ShouldMipLevelsBeForcedResident() && bIsStreamingTexture;
+		bool				bIsVirtual			= Texture->IsCurrentlyVirtualTextured();
+		bool				bIsUncompressed		= Texture->IsUncompressed();
 
 		bool bUnknownRef = false;
 		if (Streamer)
@@ -6095,11 +6095,6 @@ bool UEngine::HandleListTexturesCommand( const TCHAR* Cmd, FOutputDevice& Ar )
 			DroppedMips			= Texture2D->GetNumMips() - Texture2D->GetNumResidentMips();
 			CurSizeX			= FMath::Max<int32>(Texture2D->GetSizeX() >> DroppedMips, 1);
 			CurSizeY			= FMath::Max<int32>(Texture2D->GetSizeY() >> DroppedMips, 1);
-			bIsStreamingTexture = Texture2D->GetStreamingIndex() != INDEX_NONE;
-			UsageCount			= TextureToUsageMap.FindRef(Texture2D);
-			bIsForced			= Texture2D->ShouldMipLevelsBeForcedResident() && bIsStreamingTexture;
-			bIsVirtual			= Texture2D->IsCurrentlyVirtualTextured();
-			bIsUncompressed		= Texture2D->IsUncompressed();
 
 			if ((NumMips >= Texture2D->GetMinTextureResidentMipCount()) && bIsStreamingTexture)
 			{
@@ -6109,41 +6104,17 @@ bool UEngine::HandleListTexturesCommand( const TCHAR* Cmd, FOutputDevice& Ar )
 		else if (TextureCube != nullptr)
 		{
 			NumMips				= TextureCube->GetNumMips();
-			MaxResLODBias		= TextureCube->GetCachedLODBias();
-			MaxAllowedSizeX		= FMath::Max<int32>(TextureCube->GetSizeX() >> MaxResLODBias, 1);
-			MaxAllowedSizeY		= FMath::Max<int32>(TextureCube->GetSizeY() >> MaxResLODBias, 1);
 			Format				= TextureCube->GetPixelFormat();
-			DroppedMips			= MaxResLODBias;
-			CurSizeX			= FMath::Max<int32>(TextureCube->GetSizeX() >> DroppedMips, 1);
-			CurSizeY			= FMath::Max<int32>(TextureCube->GetSizeY() >> DroppedMips, 1);
-			bIsUncompressed		= TextureCube->IsUncompressed();
 		}
 		else if (Texture3D != nullptr)
 		{
 			NumMips				= Texture3D->GetNumMips();
-			MaxResLODBias		= Texture3D->GetCachedLODBias();
-			MaxAllowedSizeX		= FMath::Max<int32>(Texture3D->GetSizeX() >> MaxResLODBias, 1);
-			MaxAllowedSizeY		= FMath::Max<int32>(Texture3D->GetSizeY() >> MaxResLODBias, 1);
 			Format				= Texture3D->GetPixelFormat();
-			DroppedMips			= MaxResLODBias;
-			CurSizeX			= FMath::Max<int32>(Texture3D->GetSizeX() >> DroppedMips, 1);
-			CurSizeY			= FMath::Max<int32>(Texture3D->GetSizeY() >> DroppedMips, 1);
-			bIsUncompressed		= Texture3D->IsUncompressed();
 		}
 		else if (RenderTexture != nullptr)
 		{
 			NumMips = RenderTexture->GetNumMips();
-			MaxResLODBias = RenderTexture->GetCachedLODBias();
-			MaxAllowedSizeX = FMath::Max<int32>(RenderTexture->SizeX >> MaxResLODBias, 1);
-			MaxAllowedSizeY = FMath::Max<int32>(RenderTexture->SizeY >> MaxResLODBias, 1);
 			Format = RenderTexture->GetFormat();
-			DroppedMips = MaxResLODBias;
-			CurSizeX = FMath::Max<int32>(RenderTexture->SizeX >> DroppedMips, 1);
-			CurSizeY = FMath::Max<int32>(RenderTexture->SizeY >> DroppedMips, 1);
-			bIsStreamingTexture = RenderTexture->GetStreamingIndex() != INDEX_NONE;
-			bIsForced = RenderTexture->ShouldMipLevelsBeForcedResident() && bIsStreamingTexture;
-			bIsVirtual = RenderTexture->IsCurrentlyVirtualTextured();
-			bIsUncompressed = RenderTexture->IsUncompressed();
 		}
 
 		if (bListUnused && UsageCount != 0)
diff --git a/Engine/Source/Runtime/Engine/Public/Rendering/SkeletalMeshRenderData.h b/Engine/Source/Runtime/Engine/Public/Rendering/SkeletalMeshRenderData.h
--- a/Engine/Source/Runtime/Engine/Public/Rendering/SkeletalMeshRenderData.h
+++ b/Engine/Source/Runtime/Engine/Public/Rendering/SkeletalMeshRenderData.h
@@ -71,7 +71,7 @@ public:
 	ENGINE_API void ReleaseResources();
 
 	/** Return the resource size */
-	void GetResourceSizeEx(FResourceSizeEx& CumulativeResourceSize);
+	ENGINE_API void GetResourceSizeEx(FResourceSizeEx& CumulativeResourceSize);
 
 	/** Get the estimated memory overhead of buffers marked as NeedsCPUAccess. */
 	SIZE_T GetCPUAccessMemoryOverhead() const;
diff --git a/Engine/Source/Runtime/Engine/Public/StaticMeshResources.h b/Engine/Source/Runtime/Engine/Public/StaticMeshResources.h
--- a/Engine/Source/Runtime/Engine/Public/StaticMeshResources.h
+++ b/Engine/Source/Runtime/Engine/Public/StaticMeshResources.h
@@ -718,7 +718,7 @@ public:
 	ENGINE_API void ReleaseResources();
 
 	/** Compute the size of this resource. */
-	void GetResourceSizeEx(FResourceSizeEx& CumulativeResourceSize) const;
+	ENGINE_API void GetResourceSizeEx(FResourceSizeEx& CumulativeResourceSize) const;
 
 	/** Get the estimated memory overhead of buffers marked as NeedsCPUAccess. */
 	SIZE_T GetCPUAccessMemoryOverhead() const;
diff --git a/Engine/Source/Runtime/MediaAssets/Private/Assets/MediaTexture.cpp b/Engine/Source/Runtime/MediaAssets/Private/Assets/MediaTexture.cpp
--- a/Engine/Source/Runtime/MediaAssets/Private/Assets/MediaTexture.cpp
+++ b/Engine/Source/Runtime/MediaAssets/Private/Assets/MediaTexture.cpp
@@ -220,6 +220,12 @@ void UMediaTexture::SetRenderedExternalTextureGuid(const FGuid& InNewGuid)
 	CurrentRenderedGuid = InNewGuid;
 }
 
+
+uint32 UMediaTexture::CalcTextureMemorySizeEnum(ETextureMipCount Enum) const
+{
+	return Size;
+}
+
 /* UObject interface
  *****************************************************************************/
 
diff --git a/Engine/Source/Runtime/MediaAssets/Public/MediaTexture.h b/Engine/Source/Runtime/MediaAssets/Public/MediaTexture.h
--- a/Engine/Source/Runtime/MediaAssets/Public/MediaTexture.h
+++ b/Engine/Source/Runtime/MediaAssets/Public/MediaTexture.h
@@ -256,7 +256,8 @@ public:
 	virtual uint32 GetSurfaceArraySize() const override { return 0; }
 	virtual FGuid GetExternalTextureGuid() const override;
 	void SetRenderedExternalTextureGuid(const FGuid& InNewGuid);
-	virtual ETextureClass GetTextureClass() const { return ETextureClass::Other2DNoSource; }
+	virtual ETextureClass GetTextureClass() const override { return ETextureClass::Other2DNoSource; }
+	virtual uint32 CalcTextureMemorySizeEnum(ETextureMipCount Enum) const override;
 
 public:
 
