Add a virtual to FContentBundleClient before its requests unregistration so child class can cleanup any additionnal data before unregistering.
Fix crash occuring when unregistering content bundles via the UI with client preventing to remove content bundles content

diff --git a/Engine/Source/Runtime/Engine/Private/WorldPartition/ContentBundle/ContentBundleClient.cpp b/Engine/Source/Runtime/Engine/Private/WorldPartition/ContentBundle/ContentBundleClient.cpp
--- a/Engine/Source/Runtime/Engine/Private/WorldPartition/ContentBundle/ContentBundleClient.cpp
+++ b/Engine/Source/Runtime/Engine/Private/WorldPartition/ContentBundle/ContentBundleClient.cpp
@@ -62,6 +62,7 @@ void FContentBundleClient::RequestUnregister()
 {
 	if (State != EContentBundleClientState::Unregistered)
 	{
+		DoOnClientToUnregister();
 		GEngine->GetEngineSubsystem<UContentBundleEngineSubsystem>()->UnregisterContentBundle(*this);
 		SetState(EContentBundleClientState::Unregistered);
 	}
diff --git a/Engine/Source/Runtime/Engine/Public/WorldPartition/ContentBundle/ContentBundleClient.h b/Engine/Source/Runtime/Engine/Public/WorldPartition/ContentBundle/ContentBundleClient.h
--- a/Engine/Source/Runtime/Engine/Public/WorldPartition/ContentBundle/ContentBundleClient.h
+++ b/Engine/Source/Runtime/Engine/Public/WorldPartition/ContentBundle/ContentBundleClient.h
@@ -61,6 +61,8 @@ protected:
 	virtual void DoOnContentInjectedInWorld(EContentBundleStatus InjectionStatus, UWorld* InjectedWorld) {};
 	virtual void DoOnContentRemovedFromWorld(UWorld* InjectedWorld) {};
 
+	virtual void DoOnClientToUnregister() {};
+
 private:
 	bool HasContentToRemove() const;
 
