Remove unnoticed UE_ENABLE_OPTIMIZATION cpp footer.

diff --git a/Engine/Source/Runtime/Engine/Private/Level.cpp b/Engine/Source/Runtime/Engine/Private/Level.cpp
--- a/Engine/Source/Runtime/Engine/Private/Level.cpp
+++ b/Engine/Source/Runtime/Engine/Private/Level.cpp
@@ -4572,6 +4572,4 @@ void ULevel::OnMultipleLSAsPopupDismissed()
 
 #endif // WITH_EDITOR
 
-#undef LOCTEXT_NAMESPACE
-
-UE_ENABLE_OPTIMIZATION
\ No newline at end of file
+#undef LOCTEXT_NAMESPACE
\ No newline at end of file
