-> Copy legacy physics replication into a plugin, and add in capacity to replicate FPhysicsObjects.
-> Remove manual replication of physical parameters in cluster union component so that they will go through new physics replication without fighting
-> Add a function to PhysScene_Chaos so that users can manually update to new phys replication factory after the scene has already been created, without dealing directly with phys replication ptrs or deprecated functions. Pretty gross, but atm I see no other hook

diff --git a/Engine/Source/Runtime/Engine/Classes/PhysicsEngine/ClusterUnionComponent.h b/Engine/Source/Runtime/Engine/Classes/PhysicsEngine/ClusterUnionComponent.h
--- a/Engine/Source/Runtime/Engine/Classes/PhysicsEngine/ClusterUnionComponent.h
+++ b/Engine/Source/Runtime/Engine/Classes/PhysicsEngine/ClusterUnionComponent.h
@@ -77,12 +77,6 @@ struct FClusterUnionReplicatedData
 {
 	GENERATED_BODY()
 
-	UPROPERTY()
-	FVector_NetQuantize100 LinVel;
-
-	UPROPERTY()
-	FVector_NetQuantize100 AngVel;
-
 	UPROPERTY()
 	uint8 ObjectState = 0;
 
@@ -164,8 +158,6 @@ public:
 	UFUNCTION(BlueprintCallable, Category = "Cluster Union")
 	void SetIsAnchored(bool bIsAnchored);
 
-	// The SyncVelocitiesFromPhysics will set replicated state using data from the physics thread. 
-	void SyncVelocitiesFromPhysics(const FVector& LinearVelocity, const FVector& AngularVelocity);
 	// SyncClusterUnionFromProxy will examine the make up of the cluster union (particles, child to parent, etc.) and do whatever is needed on the GT in terms of bookkeeping.
 	void SyncClusterUnionFromProxy();
 
diff --git a/Engine/Source/Runtime/Engine/Private/PhysicsEngine/ClusterUnionComponent.cpp b/Engine/Source/Runtime/Engine/Private/PhysicsEngine/ClusterUnionComponent.cpp
--- a/Engine/Source/Runtime/Engine/Private/PhysicsEngine/ClusterUnionComponent.cpp
+++ b/Engine/Source/Runtime/Engine/Private/PhysicsEngine/ClusterUnionComponent.cpp
@@ -518,29 +518,22 @@ void UClusterUnionComponent::HandleComponentPhysicsStateChangePostAddIntoCluster
 	RemoveComponentFromCluster(ChangedComponent);
 }
 
-void UClusterUnionComponent::SyncVelocitiesFromPhysics(const FVector& LinearVelocity, const FVector& AngularVelocity)
-{
-	if (!IsAuthority())
-	{
-		return;
-	}
-
-	ReplicatedRigidState.LinVel = LinearVelocity;
-	ReplicatedRigidState.AngVel = AngularVelocity;
-	MARK_PROPERTY_DIRTY_FROM_NAME(UClusterUnionComponent, ReplicatedRigidState, this);
-}
-
 void UClusterUnionComponent::SyncClusterUnionFromProxy()
 {
+	//SCOPE_CYCLE_COUNTER(STAT_ClusterUnionComponent_SyncClusterUnionFromProxy);
+
 	// NOTE THAT WE ARE ON THE GAME THREAD HERE.
 	if (!PhysicsProxy)
 	{
 		return;
 	}
 
-	ReplicatedRigidState.bIsAnchored = PhysicsProxy->IsAnchored_External();
-	ReplicatedRigidState.ObjectState = static_cast<uint8>(PhysicsProxy->GetObjectState_External());
-	MARK_PROPERTY_DIRTY_FROM_NAME(UClusterUnionComponent, ReplicatedRigidState, this);
+	if (IsAuthority())
+	{
+		ReplicatedRigidState.bIsAnchored = PhysicsProxy->IsAnchored_External();
+		ReplicatedRigidState.ObjectState = static_cast<uint8>(PhysicsProxy->GetObjectState_External());
+		MARK_PROPERTY_DIRTY_FROM_NAME(UClusterUnionComponent, ReplicatedRigidState, this);
+	}
 	
 	const Chaos::FClusterUnionSyncedData& FullData = PhysicsProxy->GetSyncedData_External();
 
@@ -806,8 +799,6 @@ void UClusterUnionComponent::OnRep_RigidState()
 		return;
 	}
 
-	PhysicsProxy->SetLinearVelocity_External(ReplicatedRigidState.LinVel);
-	PhysicsProxy->SetAngularVelocity_External(ReplicatedRigidState.AngVel);
 	PhysicsProxy->SetIsAnchored_External(ReplicatedRigidState.bIsAnchored);
 	PhysicsProxy->SetObjectState_External(static_cast<Chaos::EObjectStateType>(ReplicatedRigidState.ObjectState));
 }
@@ -820,6 +811,11 @@ void UClusterUnionComponent::GetLifetimeReplicatedProps(TArray<FLifetimeProperty
 	Params.bIsPushBased = true;
 
 	DOREPLIFETIME_WITH_PARAMS_FAST(UClusterUnionComponent, ReplicatedRigidState, Params);
+
+	// Allow the physical parameters of this component to be replicated via
+	// physicsreplication even if component replication is enabled
+	DISABLE_REPLICATED_PROPERTY_FAST(USceneComponent, RelativeLocation);
+	DISABLE_REPLICATED_PROPERTY_FAST(USceneComponent, RelativeRotation);
 }
 
 void UClusterUnionComponent::ForceSetChildToParent(UPrimitiveComponent* InComponent, const TArray<int32>& BoneIds, const TArray<FTransform>& ChildToParent)
diff --git a/Engine/Source/Runtime/Engine/Private/PhysicsEngine/Experimental/PhysScene_Chaos.cpp b/Engine/Source/Runtime/Engine/Private/PhysicsEngine/Experimental/PhysScene_Chaos.cpp
--- a/Engine/Source/Runtime/Engine/Private/PhysicsEngine/Experimental/PhysScene_Chaos.cpp
+++ b/Engine/Source/Runtime/Engine/Private/PhysicsEngine/Experimental/PhysScene_Chaos.cpp
@@ -516,10 +516,7 @@ FPhysScene_Chaos::FPhysScene_Chaos(AActor* InSolverActor
 #endif
 
 	// Create replication manager
-	PhysicsReplication
-		= PhysicsReplicationFactory.IsValid()
-		? PhysicsReplicationFactory->CreatePhysicsReplication(this)
-		: MakeUnique<FPhysicsReplication>(this);
+	CreatePhysicsReplication();
 
 	FPhysicsDelegates::OnPhysSceneInit.Broadcast(this);
 
@@ -773,6 +770,18 @@ IPhysicsReplication* FPhysScene_Chaos::GetPhysicsReplication()
 	return PhysicsReplication.Get();
 }
 
+IPhysicsReplication* FPhysScene_Chaos::CreatePhysicsReplication()
+{
+	// Create replication manager
+	PhysicsReplication
+		= PhysicsReplicationFactory.IsValid()
+		? PhysicsReplicationFactory->CreatePhysicsReplication(this)
+		: MakeUnique<FPhysicsReplication>(this);
+
+	// Return ptr to the new physics rep
+	return PhysicsReplication.Get();
+}
+
 void FPhysScene_Chaos::SetPhysicsReplication(IPhysicsReplication* InPhysicsReplication)
 {
 	PhysicsReplication = TUniquePtr<IPhysicsReplication>(InPhysicsReplication);
diff --git a/Engine/Source/Runtime/Engine/Private/PhysicsEngine/PhysicsReplication.cpp b/Engine/Source/Runtime/Engine/Private/PhysicsEngine/PhysicsReplication.cpp
--- a/Engine/Source/Runtime/Engine/Private/PhysicsEngine/PhysicsReplication.cpp
+++ b/Engine/Source/Runtime/Engine/Private/PhysicsEngine/PhysicsReplication.cpp
@@ -21,55 +21,55 @@ namespace CharacterMovementCVars
 	extern int32 NetShowCorrections;
 	extern float NetCorrectionLifetime;
 
-	static int32 SkipPhysicsReplication = 0;
+	int32 SkipPhysicsReplication = 0;
 	static FAutoConsoleVariableRef CVarSkipPhysicsReplication(TEXT("p.SkipPhysicsReplication"), SkipPhysicsReplication, TEXT(""));
 
-	static float NetPingExtrapolation = -1.0f;
+	float NetPingExtrapolation = -1.0f;
 	static FAutoConsoleVariableRef CVarNetPingExtrapolation(TEXT("p.NetPingExtrapolation"), NetPingExtrapolation, TEXT(""));
 
-	static float NetPingLimit = -1.f;
+	float NetPingLimit = -1.f;
 	static FAutoConsoleVariableRef CVarNetPingLimit(TEXT("p.NetPingLimit"), NetPingLimit, TEXT(""));
 
-	static float ErrorPerLinearDifference = -1.0f;
+	float ErrorPerLinearDifference = -1.0f;
 	static FAutoConsoleVariableRef CVarErrorPerLinearDifference(TEXT("p.ErrorPerLinearDifference"), ErrorPerLinearDifference, TEXT(""));
 
-	static float ErrorPerAngularDifference = -1.0f;
+	float ErrorPerAngularDifference = -1.0f;
 	static FAutoConsoleVariableRef CVarErrorPerAngularDifference(TEXT("p.ErrorPerAngularDifference"), ErrorPerAngularDifference, TEXT(""));
 
-	static float ErrorAccumulationSeconds = -1.0f;
+	float ErrorAccumulationSeconds = -1.0f;
 	static FAutoConsoleVariableRef CVarErrorAccumulation(TEXT("p.ErrorAccumulationSeconds"), ErrorAccumulationSeconds, TEXT(""));
 
-	static float ErrorAccumulationDistanceSq = -1.0f;
+	float ErrorAccumulationDistanceSq = -1.0f;
 	static FAutoConsoleVariableRef CVarErrorAccumulationDistanceSq(TEXT("p.ErrorAccumulationDistanceSq"), ErrorAccumulationDistanceSq, TEXT(""));
 
-	static float ErrorAccumulationSimilarity = -1.f;
+	float ErrorAccumulationSimilarity = -1.f;
 	static FAutoConsoleVariableRef CVarErrorAccumulationSimilarity(TEXT("p.ErrorAccumulationSimilarity"), ErrorAccumulationSimilarity, TEXT(""));
 
-	static float MaxLinearHardSnapDistance = -1.f;
+	float MaxLinearHardSnapDistance = -1.f;
 	static FAutoConsoleVariableRef CVarMaxLinearHardSnapDistance(TEXT("p.MaxLinearHardSnapDistance"), MaxLinearHardSnapDistance, TEXT(""));
 
-	static float MaxRestoredStateError = -1.0f;
+	float MaxRestoredStateError = -1.0f;
 	static FAutoConsoleVariableRef CVarMaxRestoredStateError(TEXT("p.MaxRestoredStateError"), MaxRestoredStateError, TEXT(""));
 
-	static float PositionLerp = -1.0f;
+	float PositionLerp = -1.0f;
 	static FAutoConsoleVariableRef CVarLinSet(TEXT("p.PositionLerp"), PositionLerp, TEXT(""));
 
-	static float LinearVelocityCoefficient = -1.0f;
+	float LinearVelocityCoefficient = -1.0f;
 	static FAutoConsoleVariableRef CVarLinLerp(TEXT("p.LinearVelocityCoefficient"), LinearVelocityCoefficient, TEXT(""));
 
-	static float AngleLerp = -1.0f;
+	float AngleLerp = -1.0f;
 	static FAutoConsoleVariableRef CVarAngSet(TEXT("p.AngleLerp"), AngleLerp, TEXT(""));
 
-	static float AngularVelocityCoefficient = -1.0f;
+	float AngularVelocityCoefficient = -1.0f;
 	static FAutoConsoleVariableRef CVarAngLerp(TEXT("p.AngularVelocityCoefficient"), AngularVelocityCoefficient, TEXT(""));
 
-	static int32 AlwaysHardSnap = 0;
+	int32 AlwaysHardSnap = 0;
 	static FAutoConsoleVariableRef CVarAlwaysHardSnap(TEXT("p.AlwaysHardSnap"), AlwaysHardSnap, TEXT(""));
 
-	static int32 AlwaysResetPhysics = 0;
+	int32 AlwaysResetPhysics = 0;
 	static FAutoConsoleVariableRef CVarAlwaysResetPhysics(TEXT("p.AlwaysResetPhysics"), AlwaysResetPhysics, TEXT(""));
 
-	static int32 ApplyAsyncSleepState = 1;
+	int32 ApplyAsyncSleepState = 1;
 	static FAutoConsoleVariableRef CVarApplyAsyncSleepState(TEXT("p.ApplyAsyncSleepState"), ApplyAsyncSleepState, TEXT(""));
 }
 
@@ -112,16 +112,19 @@ private:
 	}
 };
 
-void ComputeDeltas(const FVector& CurrentPos, const FQuat& CurrentQuat, const FVector& TargetPos, const FQuat& TargetQuat, FVector& OutLinDiff, float& OutLinDiffSize,
-	FVector& OutAngDiffAxis, float& OutAngDiff, float& OutAngDiffSize)
+namespace
 {
-	OutLinDiff = TargetPos - CurrentPos;
-	OutLinDiffSize = OutLinDiff.Size();
-	const FQuat InvCurrentQuat = CurrentQuat.Inverse();
-	const FQuat DeltaQuat = TargetQuat * InvCurrentQuat;
-	DeltaQuat.ToAxisAndAngle(OutAngDiffAxis, OutAngDiff);
-	OutAngDiff = FMath::RadiansToDegrees(FMath::UnwindRadians(OutAngDiff));
-	OutAngDiffSize = FMath::Abs(OutAngDiff);
+	void ComputeDeltas(const FVector& CurrentPos, const FQuat& CurrentQuat, const FVector& TargetPos, const FQuat& TargetQuat, FVector& OutLinDiff, float& OutLinDiffSize,
+		FVector& OutAngDiffAxis, float& OutAngDiff, float& OutAngDiffSize)
+	{
+		OutLinDiff = TargetPos - CurrentPos;
+		OutLinDiffSize = OutLinDiff.Size();
+		const FQuat InvCurrentQuat = CurrentQuat.Inverse();
+		const FQuat DeltaQuat = TargetQuat * InvCurrentQuat;
+		DeltaQuat.ToAxisAndAngle(OutAngDiffAxis, OutAngDiff);
+		OutAngDiff = FMath::RadiansToDegrees(FMath::UnwindRadians(OutAngDiff));
+		OutAngDiffSize = FMath::Abs(OutAngDiff);
+	}
 }
 
 FPhysicsReplication::~FPhysicsReplication()
@@ -843,6 +846,8 @@ void FPhysicsReplication::SetReplicatedTarget(UPrimitiveComponent* Component, FN
 		ensure(!Target->PrevPos.ContainsNaN());
 		ensure(!Target->PrevPosTarget.ContainsNaN());
 		ensure(!Target->TargetState.Position.ContainsNaN());
+
+		OnSetReplicatedTarget(Component, BoneName, ReplicatedTarget, ServerFrame, *Target);
 	}
 }
 
diff --git a/Engine/Source/Runtime/Engine/Public/Physics/Experimental/PhysScene_Chaos.h b/Engine/Source/Runtime/Engine/Public/Physics/Experimental/PhysScene_Chaos.h
--- a/Engine/Source/Runtime/Engine/Public/Physics/Experimental/PhysScene_Chaos.h
+++ b/Engine/Source/Runtime/Engine/Public/Physics/Experimental/PhysScene_Chaos.h
@@ -155,7 +155,9 @@ public:
 
 	IPhysicsReplication* GetPhysicsReplication();
 
-	UE_DEPRECATED(5.3, "Can no longer direclty set physics replication at runtime. For now, specify a PhysicsReplication factory instead. This function will take ownership of the IPhysicsReplication's lifetime.")
+	IPhysicsReplication* CreatePhysicsReplication();
+
+	UE_DEPRECATED(5.3, "If possible, avoid directly setting physics replication at runtime - this function will take ownership of the IPhysicsReplication's lifetime. Instead, specify a PhysicsReplication factory and if you must change the replication class at runtime run RecreatePhysicsReplication after having set the factory.")
 	void SetPhysicsReplication(IPhysicsReplication* InPhysicsReplication);
 
 	virtual void AddReferencedObjects(FReferenceCollector& Collector) override;
diff --git a/Engine/Source/Runtime/Engine/Public/PhysicsReplication.h b/Engine/Source/Runtime/Engine/Public/PhysicsReplication.h
--- a/Engine/Source/Runtime/Engine/Public/PhysicsReplication.h
+++ b/Engine/Source/Runtime/Engine/Public/PhysicsReplication.h
@@ -19,6 +19,27 @@ namespace Chaos
 	struct FSimCallbackInput;
 }
 
+namespace CharacterMovementCVars
+{
+	extern ENGINE_API int32 SkipPhysicsReplication;
+	extern ENGINE_API float NetPingExtrapolation;
+	extern ENGINE_API float NetPingLimit;
+	extern ENGINE_API float ErrorPerLinearDifference;
+	extern ENGINE_API float ErrorPerAngularDifference;
+	extern ENGINE_API float ErrorAccumulationSeconds;
+	extern ENGINE_API float ErrorAccumulationDistanceSq;
+	extern ENGINE_API float ErrorAccumulationSimilarity;
+	extern ENGINE_API float MaxLinearHardSnapDistance;
+	extern ENGINE_API float MaxRestoredStateError;
+	extern ENGINE_API float PositionLerp;
+	extern ENGINE_API float LinearVelocityCoefficient;
+	extern ENGINE_API float AngleLerp;
+	extern ENGINE_API float AngularVelocityCoefficient;
+	extern ENGINE_API int32 AlwaysHardSnap;
+	extern ENGINE_API int32 AlwaysResetPhysics;
+	extern ENGINE_API int32 ApplyAsyncSleepState;
+}
+
 #if !UE_BUILD_SHIPPING
 namespace PhysicsReplicationCVars
 {
@@ -117,6 +138,7 @@ protected:
 	/** Update the physics body state given a set of replicated targets */
 	virtual void OnTick(float DeltaSeconds, TMap<TWeakObjectPtr<UPrimitiveComponent>, FReplicatedPhysicsTarget>& ComponentsToTargets);
 	virtual void OnTargetRestored(TWeakObjectPtr<UPrimitiveComponent> Component, const FReplicatedPhysicsTarget& Target) {}
+	virtual void OnSetReplicatedTarget(UPrimitiveComponent* Component, FName BoneName, const FRigidBodyState& ReplicatedTarget, int32 ServerFrame, FReplicatedPhysicsTarget& Target) {}
 
 	/** Called when a dynamic rigid body receives a physics update */
 	virtual bool ApplyRigidBodyState(float DeltaSeconds, FBodyInstance* BI, FReplicatedPhysicsTarget& PhysicsTarget, const FRigidBodyErrorCorrection& ErrorCorrection, const float PingSecondsOneWay, int32 LocalFrame, int32 NumPredictedFrames);
