Prevent updating scene captures if their detail modes are higher than the current setting

diff --git a/Engine/Source/Runtime/Engine/Classes/Components/SceneCaptureComponent.h b/Engine/Source/Runtime/Engine/Classes/Components/SceneCaptureComponent.h
--- a/Engine/Source/Runtime/Engine/Classes/Components/SceneCaptureComponent.h
+++ b/Engine/Source/Runtime/Engine/Classes/Components/SceneCaptureComponent.h
@@ -218,6 +218,8 @@ protected:
 	void UnregisterDelegates();
 	void ReleaseGarbageReferences();
 
+	bool IsCulledByDetailMode() const;
+
 	virtual void UpdateSceneCaptureContents(FSceneInterface* Scene) {};
 
 	/**
diff --git a/Engine/Source/Runtime/Engine/Classes/Components/SceneCaptureComponent2D.h b/Engine/Source/Runtime/Engine/Classes/Components/SceneCaptureComponent2D.h
--- a/Engine/Source/Runtime/Engine/Classes/Components/SceneCaptureComponent2D.h
+++ b/Engine/Source/Runtime/Engine/Classes/Components/SceneCaptureComponent2D.h
@@ -155,7 +155,10 @@ public:
 	UFUNCTION(BlueprintCallable, Category = "Rendering")
 	void RemoveBlendable(TScriptInterface<IBlendableInterface> InBlendableObject) { PostProcessSettings.RemoveBlendable(InBlendableObject); }
 
-	/** Render the scene to the texture the next time the main view is rendered. */
+	/**
+	 * Render the scene to the texture the next time the main view is rendered.
+	 * If r.SceneCapture.CullByDetailMode is set, nothing will happen if DetailMode is higher than r.DetailMode.
+	 */
 	void CaptureSceneDeferred();
 
 	// For backwards compatibility
@@ -164,6 +167,7 @@ public:
 	/** 
 	 * Render the scene to the texture target immediately.  
 	 * This should not be used if bCaptureEveryFrame is enabled, or the scene capture will render redundantly. 
+	 * If r.SceneCapture.CullByDetailMode is set, nothing will happen if DetailMode is higher than r.DetailMode.
 	 */
 	UFUNCTION(BlueprintCallable,Category = "Rendering|SceneCapture")
 	void CaptureScene();
diff --git a/Engine/Source/Runtime/Engine/Classes/Components/SceneCaptureComponentCube.h b/Engine/Source/Runtime/Engine/Classes/Components/SceneCaptureComponentCube.h
--- a/Engine/Source/Runtime/Engine/Classes/Components/SceneCaptureComponentCube.h
+++ b/Engine/Source/Runtime/Engine/Classes/Components/SceneCaptureComponentCube.h
@@ -41,12 +41,16 @@ public:
 #endif // WITH_EDITOR
 	//~ End UObject Interface
 
-	/** Render the scene to the texture the next time the main view is rendered. */
+	/**
+	 * Render the scene to the texture the next time the main view is rendered.
+	 * If r.SceneCapture.CullByDetailMode is set, nothing will happen if DetailMode is higher than r.DetailMode.
+	 */
 	void CaptureSceneDeferred();
 
 	/** 
 	 * Render the scene to the texture target immediately.  
 	 * This should not be used if bCaptureEveryFrame is enabled, or the scene capture will render redundantly. 
+	 * If r.SceneCapture.CullByDetailMode is set, nothing will happen if DetailMode is higher than r.DetailMode.
 	 */
 	UFUNCTION(BlueprintCallable,Category = "Rendering|SceneCapture")
 	void CaptureScene();
diff --git a/Engine/Source/Runtime/Engine/Private/Components/SceneCaptureComponent.cpp b/Engine/Source/Runtime/Engine/Private/Components/SceneCaptureComponent.cpp
--- a/Engine/Source/Runtime/Engine/Private/Components/SceneCaptureComponent.cpp
+++ b/Engine/Source/Runtime/Engine/Private/Components/SceneCaptureComponent.cpp
@@ -30,6 +30,7 @@
 #include "Logging/MessageLog.h"
 #include "Engine/SCS_Node.h"
 #include "Engine/TextureRenderTarget2D.h"
+#include "UnrealEngine.h"
 
 #include UE_INLINE_GENERATED_CPP_BY_NAME(SceneCaptureComponent)
 
@@ -62,6 +63,12 @@ static TAutoConsoleVariable<int32> CVarSCOrthographicNumYTiles(
 	TEXT("Number of Y tiles to render. Ignored in Perspective mode, works only in Orthographic mode and when r.SceneCapture.OverrideOrthographicTilingValues is on."),
 	ECVF_Scalability);
 
+static TAutoConsoleVariable<bool> CVarSCCullByDetailMode(
+	TEXT("r.SceneCapture.CullByDetailMode"),
+	1,
+	TEXT("Whether to prevent scene capture updates according to the current detail mode"),
+	ECVF_Scalability);
+
 ASceneCapture::ASceneCapture(const FObjectInitializer& ObjectInitializer)
 	: Super(ObjectInitializer)
 {
@@ -527,6 +534,11 @@ void USceneCaptureComponent::OnUnregister()
 	Super::OnUnregister();
 }
 
+bool USceneCaptureComponent::IsCulledByDetailMode() const
+{
+	return CVarSCCullByDetailMode.GetValueOnAnyThread() && DetailMode > GetCachedScalabilityCVars().DetailMode;
+}
+
 // -----------------------------------------------
 
 
@@ -671,7 +683,7 @@ void USceneCaptureComponent2D::GetCameraView(float DeltaTime, FMinimalViewInfo&
 void USceneCaptureComponent2D::CaptureSceneDeferred()
 {
 	UWorld* World = GetWorld();
-	if (World && World->Scene && IsVisible())
+	if (World && World->Scene && IsVisible() && !IsCulledByDetailMode())
 	{
 		// Defer until after updates finish
 		// Needs some CS because of parallel updates.
@@ -683,12 +695,12 @@ void USceneCaptureComponent2D::CaptureSceneDeferred()
 void USceneCaptureComponent2D::CaptureScene()
 {
 	UWorld* World = GetWorld();
-	if (World && World->Scene && IsVisible())
+	if (World && World->Scene && IsVisible() && !IsCulledByDetailMode())
 	{
 		// We must push any deferred render state recreations before causing any rendering to happen, to make sure that deleted resource references are updated
 		World->SendAllEndOfFrameUpdates();
 		UpdateSceneCaptureContents(World->Scene);
-	}	
+	}
 
 	if (bCaptureEveryFrame)
 	{
@@ -1249,7 +1261,7 @@ void USceneCaptureComponentCube::UpdateDrawFrustum()
 void USceneCaptureComponentCube::CaptureSceneDeferred()
 {
 	UWorld* World = GetWorld();
-	if (World && World->Scene && IsVisible())
+	if (World && World->Scene && IsVisible() && !IsCulledByDetailMode())
 	{
 		// Defer until after updates finish
 		// Needs some CS because of parallel updates.
@@ -1261,7 +1273,7 @@ void USceneCaptureComponentCube::CaptureSceneDeferred()
 void USceneCaptureComponentCube::CaptureScene()
 {
 	UWorld* World = GetWorld();
-	if (World && World->Scene && IsVisible())
+	if (World && World->Scene && IsVisible() && !IsCulledByDetailMode())
 	{
 		// We must push any deferred render state recreations before causing any rendering to happen, to make sure that deleted resource references are updated
 		World->SendAllEndOfFrameUpdates();
