modify cluster union IsComponentAdded check to also check for components pending sync from PT

diff --git a/Engine/Source/Runtime/Engine/Classes/PhysicsEngine/ClusterUnionComponent.h b/Engine/Source/Runtime/Engine/Classes/PhysicsEngine/ClusterUnionComponent.h
--- a/Engine/Source/Runtime/Engine/Classes/PhysicsEngine/ClusterUnionComponent.h
+++ b/Engine/Source/Runtime/Engine/Classes/PhysicsEngine/ClusterUnionComponent.h
@@ -162,7 +162,7 @@ public:
 	void SyncClusterUnionFromProxy();
 
 	UFUNCTION()
-	bool IsComponentAdded(UPrimitiveComponent* Component) { return ComponentToPhysicsObjects.Contains(Component); }
+	bool IsComponentAdded(UPrimitiveComponent* Component) { return ComponentToPhysicsObjects.Contains(Component) || PendingComponentSync.Contains(Component); }
 
 	bool HasReceivedTransform() const { return bHasReceivedTransform; }
 
