don't bind multiple times to the OnComponentPhysicsStateChanged for the same component when adding into the cluster union

diff --git a/Engine/Source/Runtime/Engine/Private/PhysicsEngine/ClusterUnionComponent.cpp b/Engine/Source/Runtime/Engine/Private/PhysicsEngine/ClusterUnionComponent.cpp
--- a/Engine/Source/Runtime/Engine/Private/PhysicsEngine/ClusterUnionComponent.cpp
+++ b/Engine/Source/Runtime/Engine/Private/PhysicsEngine/ClusterUnionComponent.cpp
@@ -159,10 +159,10 @@ void UClusterUnionComponent::AddComponentToCluster(UPrimitiveComponent* InCompon
 		}
 	}
 
-	PendingComponentSync.Add(InComponent, PendingData);
-
 	// Need to listen to changes in the component's physics state. If it gets destroyed it should be removed from the cluster union as well.
-	InComponent->OnComponentPhysicsStateChanged.AddDynamic(this, &UClusterUnionComponent::HandleComponentPhysicsStateChangePostAddIntoClusterUnion);
+	InComponent->OnComponentPhysicsStateChanged.AddUniqueDynamic(this, &UClusterUnionComponent::HandleComponentPhysicsStateChangePostAddIntoClusterUnion);
+
+	PendingComponentSync.Add(InComponent, PendingData);
 
 	PhysicsProxy->AddPhysicsObjects_External(Objects);
 	ForceRebuildGTParticleGeometry();
