Remove code which should not have been submitted with 25725663
Fix issue where CB actors where landing in both their grid and the main grid when generating streaming

diff --git a/Engine/Source/Runtime/Engine/Private/WorldPartition/WorldPartitionStreamingGeneration.cpp b/Engine/Source/Runtime/Engine/Private/WorldPartition/WorldPartitionStreamingGeneration.cpp
--- a/Engine/Source/Runtime/Engine/Private/WorldPartition/WorldPartitionStreamingGeneration.cpp
+++ b/Engine/Source/Runtime/Engine/Private/WorldPartition/WorldPartitionStreamingGeneration.cpp
@@ -1220,8 +1220,7 @@ bool UWorldPartition::GenerateContainerStreaming(const UActorDescContainer* InAc
 bool UWorldPartition::GenerateStreaming(const FGenerateStreamingParams& InParams, FGenerateStreamingContext& InContext)
 {
 	check(InParams.ActorDescCollection.IsEmpty());
-	FGenerateStreamingParams Params = FGenerateStreamingParams(InParams);
-	Params.ActorDescCollection.Append(*this);
+	FGenerateStreamingParams Params = FGenerateStreamingParams(InParams).SetActorDescContainer(ActorDescContainer);
 
 	OnPreGenerateStreaming.Broadcast(InContext.PackagesToGenerate);
 
