#UE IWYU fix

diff --git a/Engine/Source/Runtime/Engine/Private/Texture.cpp b/Engine/Source/Runtime/Engine/Private/Texture.cpp
--- a/Engine/Source/Runtime/Engine/Private/Texture.cpp
+++ b/Engine/Source/Runtime/Engine/Private/Texture.cpp
@@ -11,6 +11,7 @@
 #include "Misc/FeedbackContext.h"
 #include "UObject/UObjectIterator.h"
 #include "UObject/ObjectSaveContext.h"
+#include "UObject/Package.h"
 #include "Engine/Texture2D.h"
 #include "Engine/TextureMipDataProviderFactory.h"
 #include "ContentStreaming.h"
