Fixed crash in inertialization nodes when reparenting ABP.

diff --git a/Engine/Source/Runtime/Engine/Private/Animation/AnimNode_DeadBlending.cpp b/Engine/Source/Runtime/Engine/Private/Animation/AnimNode_DeadBlending.cpp
--- a/Engine/Source/Runtime/Engine/Private/Animation/AnimNode_DeadBlending.cpp
+++ b/Engine/Source/Runtime/Engine/Private/Animation/AnimNode_DeadBlending.cpp
@@ -520,9 +520,6 @@ void FAnimNode_DeadBlending::Initialize_AnyThread(const FAnimationInitializeCont
 	FAnimNode_Base::Initialize_AnyThread(Context);
 	Source.Initialize(Context);
 
-	const USkeleton* Skeleton = Context.AnimInstanceProxy->GetSkeleton();
-	check(Skeleton);
-
 	CurveFilter.Empty();
 	CurveFilter.SetFilterMode(UE::Anim::ECurveFilterMode::DisallowFiltered);
 	CurveFilter.AppendNames(FilteredCurves);
@@ -531,7 +528,7 @@ void FAnimNode_DeadBlending::Initialize_AnyThread(const FAnimationInitializeCont
 
 	RequestQueue.Reserve(8);
 
-	const int32 NumSkeletonBones = UE::Anim::DeadBlending::Private::GetNumSkeletonBones(Context.AnimInstanceProxy->GetRequiredBones());
+	const int32 NumSkeletonBones = Context.AnimInstanceProxy->GetSkeleton()->GetReferenceSkeleton().GetNum();
 
 	BoneValid.Empty(NumSkeletonBones);
 	BoneTranslations.Empty(NumSkeletonBones);
diff --git a/Engine/Source/Runtime/Engine/Private/Animation/AnimNode_Inertialization.cpp b/Engine/Source/Runtime/Engine/Private/Animation/AnimNode_Inertialization.cpp
--- a/Engine/Source/Runtime/Engine/Private/Animation/AnimNode_Inertialization.cpp
+++ b/Engine/Source/Runtime/Engine/Private/Animation/AnimNode_Inertialization.cpp
@@ -132,7 +132,7 @@ void FAnimNode_Inertialization::Initialize_AnyThread(const FAnimationInitializeC
 	Source.Initialize(Context);
 
 
-	const int32 NumSkeletonBones = UE::Anim::GetNumSkeletonBones(Context.AnimInstanceProxy->GetRequiredBones());
+	const int32 NumSkeletonBones = Context.AnimInstanceProxy->GetSkeleton()->GetReferenceSkeleton().GetNum();
 
 	PoseSnapshots.Empty(INERTIALIZATION_MAX_POSE_SNAPSHOTS);
 	RequestQueue.Reserve(8);
