Added delegate "InjectInputsExternal" to FNetworkPhysicsCallback. This will be broadcast during InjectInputs_External()

diff --git a/Engine/Source/Runtime/Engine/Private/PhysicsEngine/NetworkPhysicsComponent.cpp b/Engine/Source/Runtime/Engine/Private/PhysicsEngine/NetworkPhysicsComponent.cpp
--- a/Engine/Source/Runtime/Engine/Private/PhysicsEngine/NetworkPhysicsComponent.cpp
+++ b/Engine/Source/Runtime/Engine/Private/PhysicsEngine/NetworkPhysicsComponent.cpp
@@ -260,6 +260,11 @@ void FNetworkPhysicsCallback::UpdateServerPlayer_External(int32 PhysicsStep)
 	}
 }
 
+void FNetworkPhysicsCallback::InjectInputs_External(int32 PhysicsStep, int32 NumSteps)
+{
+	InjectInputsExternal.Broadcast(PhysicsStep, NumSteps);
+}
+
 void FNetworkPhysicsCallback::ProcessInputs_External(int32 PhysicsStep, const TArray<Chaos::FSimCallbackInputAndObject>& SimCallbackInputs)
 {
 	for(const Chaos::FSimCallbackInputAndObject& SimCallbackObject : SimCallbackInputs)
diff --git a/Engine/Source/Runtime/Engine/Public/Physics/NetworkPhysicsComponent.h b/Engine/Source/Runtime/Engine/Public/Physics/NetworkPhysicsComponent.h
--- a/Engine/Source/Runtime/Engine/Public/Physics/NetworkPhysicsComponent.h
+++ b/Engine/Source/Runtime/Engine/Public/Physics/NetworkPhysicsComponent.h
@@ -15,6 +15,7 @@
 
 DECLARE_MULTICAST_DELEGATE_OneParam(FOnPreProcessInputsInternal, const int32);
 DECLARE_MULTICAST_DELEGATE_OneParam(FOnPostProcessInputsInternal, const int32);
+DECLARE_MULTICAST_DELEGATE_TwoParams(FOnInjectInputsExternal, const int32 /* PhysicsStep */, const int32 /* NumSteps */);
 
 /** Templated datas history holding a datas buffer */
 template<typename DatasType>
@@ -169,8 +170,11 @@ struct FNetworkPhysicsCallback : public Chaos::IRewindCallback
 	// Delegate on the internal inputs process 
 	FOnPreProcessInputsInternal PreProcessInputsInternal;
 	FOnPostProcessInputsInternal PostProcessInputsInternal;
+	// Bind to this for additional processing on the GT during InjectInputs_External()
+	FOnInjectInputsExternal InjectInputsExternal;
 
 	// Rewind API
+	virtual void InjectInputs_External(int32 PhysicsStep, int32 NumSteps) override;
 	virtual void ProcessInputs_External(int32 PhysicsStep, const TArray<Chaos::FSimCallbackInputAndObject>& SimCallbackInputs);
 	virtual void ProcessInputs_Internal(int32 PhysicsStep, const TArray<Chaos::FSimCallbackInputAndObject>& SimCallbackInputs) override;
 	virtual void ApplyCallbacks_Internal(int32 PhysicsStep, const TArray<Chaos::ISimCallbackObject*>& SimCallbackObjects) override;
