Added optional blocking streaming requests for SVT streaming manager. Since Niagara fluid sim caching is issuing a render command to convert from SVT, it now also needs to issue a render command to tick the streaming manager so that it gets a chance to tick before the conversion command is executed (which is before the normal point in the frame where SVT streaming is ticked).

#rnx

diff --git a/Engine/Plugins/FX/Niagara/Source/Niagara/Private/NiagaraDataInterfaceRenderTargetVolume.cpp b/Engine/Plugins/FX/Niagara/Source/Niagara/Private/NiagaraDataInterfaceRenderTargetVolume.cpp
--- a/Engine/Plugins/FX/Niagara/Source/Niagara/Private/NiagaraDataInterfaceRenderTargetVolume.cpp
+++ b/Engine/Plugins/FX/Niagara/Source/Niagara/Private/NiagaraDataInterfaceRenderTargetVolume.cpp
@@ -30,6 +30,7 @@
 
 #include "SparseVolumeTexture/SparseVolumeTexture.h"
 #include "SparseVolumeTexture/SparseVolumeTextureData.h"
+#include "SparseVolumeTexture/ISparseVolumeTextureStreamingManager.h"
 
 
 #include UE_INLINE_GENERATED_CPP_BY_NAME(NiagaraDataInterfaceRenderTargetVolume)
@@ -827,7 +828,13 @@ bool UNiagaraDataInterfaceRenderTargetVolume::SimCacheReadFrame(UObject* Storage
 		UAnimatedSparseVolumeTexture* SVT = Cast<UAnimatedSparseVolumeTexture>(StorageObject);
 		
 		const int32 MipLevel = 0;
-		USparseVolumeTextureFrame *SVTFrame = USparseVolumeTextureFrame::GetFrameAndIssueStreamingRequest(SVT, FrameA + Interp, MipLevel);
+		const bool bBlocking = true;
+		USparseVolumeTextureFrame *SVTFrame = USparseVolumeTextureFrame::GetFrameAndIssueStreamingRequest(SVT, FrameA + Interp, MipLevel, bBlocking);
+		
+		// The streaming manager normally ticks in FDeferredShadingSceneRenderer::Render(), but the SVT->DenseTexture conversion compute shader happens in a render command before that.
+		// At execution time of that command, the streamer hasn't had the chance to do any streaming yet, so we force another tick here.
+		// Assuming blocking requests are used, this guarantees that the requested frame is fully streamed in (if there is memory available).
+		UE::SVT::GetStreamingManager().Update_GameThread();
 
 		FIntVector VolumeResolution = SVT->GetVolumeResolution();						
 
diff --git a/Engine/Source/Runtime/Engine/Classes/SparseVolumeTexture/SparseVolumeTexture.h b/Engine/Source/Runtime/Engine/Classes/SparseVolumeTexture/SparseVolumeTexture.h
--- a/Engine/Source/Runtime/Engine/Classes/SparseVolumeTexture/SparseVolumeTexture.h
+++ b/Engine/Source/Runtime/Engine/Classes/SparseVolumeTexture/SparseVolumeTexture.h
@@ -238,7 +238,9 @@ public:
 	// Retrieves a frame from the given SparseVolumeTexture and also issues a streaming request for it. 
 	// FrameIndex is of float type so that the streaming system can use the fractional part to more easily keep track of playback speed and direction (forward/reverse playback).
 	// MipLevel is the lowest mip level that the caller intends to use but does not guarantee that the mip is actually resident.
-	static USparseVolumeTextureFrame* GetFrameAndIssueStreamingRequest(USparseVolumeTexture* SparseVolumeTexture, float FrameIndex, int32 MipLevel);
+	// If bBlocking is true, DDC streaming requests will block on completion, guaranteeing that the requested frame will have been streamed in after the next streaming system update.
+	// If streaming cooked data from disk, the highest priority will be used, but no guarantee is given.
+	static USparseVolumeTextureFrame* GetFrameAndIssueStreamingRequest(USparseVolumeTexture* SparseVolumeTexture, float FrameIndex, int32 MipLevel, bool bBlocking = false);
 
 	bool Initialize(USparseVolumeTexture* InOwner, int32 InFrameIndex, UE::SVT::FTextureData& UncookedFrame);
 	int32 GetFrameIndex() const { return FrameIndex; }
diff --git a/Engine/Source/Runtime/Engine/Private/SparseVolumeTexture/SparseVolumeTexture.cpp b/Engine/Source/Runtime/Engine/Private/SparseVolumeTexture/SparseVolumeTexture.cpp
--- a/Engine/Source/Runtime/Engine/Private/SparseVolumeTexture/SparseVolumeTexture.cpp
+++ b/Engine/Source/Runtime/Engine/Private/SparseVolumeTexture/SparseVolumeTexture.cpp
@@ -182,7 +182,21 @@ void FResources::Serialize(FArchive& Ar, UObject* Owner, bool bCooked)
 
 bool FResources::HasStreamingData() const
 {
-	return MipLevelStreamingInfo.Num() > 1;
+	if (MipLevelStreamingInfo.Num() == 1)
+	{
+		// Root mip level does not stream
+		return false;
+	}
+	else
+	{
+		// It is possible for multiple mip levels to exist but all these levels are empty, so we can't just check the number of mip levels
+		bool bHasStreamingData = false;
+		for (int32 MipLevelIndex = 0; MipLevelIndex < MipLevelStreamingInfo.Num() - 1; ++MipLevelIndex)
+		{
+			bHasStreamingData = bHasStreamingData || (MipLevelStreamingInfo[MipLevelIndex].BulkSize > 0);
+		}
+		return bHasStreamingData;
+	}
 }
 
 #if WITH_EDITORONLY_DATA
@@ -643,11 +657,11 @@ USparseVolumeTextureFrame::USparseVolumeTextureFrame(const FObjectInitializer& O
 {
 }
 
-USparseVolumeTextureFrame* USparseVolumeTextureFrame::GetFrameAndIssueStreamingRequest(USparseVolumeTexture* SparseVolumeTexture, float FrameIndex, int32 MipLevel)
+USparseVolumeTextureFrame* USparseVolumeTextureFrame::GetFrameAndIssueStreamingRequest(USparseVolumeTexture* SparseVolumeTexture, float FrameIndex, int32 MipLevel, bool bBlocking)
 {
 	if (UStreamableSparseVolumeTexture* StreamableSVT = Cast<UStreamableSparseVolumeTexture>(SparseVolumeTexture))
 	{
-		UE::SVT::GetStreamingManager().Request_GameThread(StreamableSVT, FrameIndex, MipLevel);
+		UE::SVT::GetStreamingManager().Request_GameThread(StreamableSVT, FrameIndex, MipLevel, bBlocking);
 		return StreamableSVT->GetFrame(static_cast<int32>(FrameIndex));
 	}
 	return nullptr;
diff --git a/Engine/Source/Runtime/Engine/Private/SparseVolumeTexture/SparseVolumeTextureStreamingManager.cpp b/Engine/Source/Runtime/Engine/Private/SparseVolumeTexture/SparseVolumeTextureStreamingManager.cpp
--- a/Engine/Source/Runtime/Engine/Private/SparseVolumeTexture/SparseVolumeTextureStreamingManager.cpp
+++ b/Engine/Source/Runtime/Engine/Private/SparseVolumeTexture/SparseVolumeTextureStreamingManager.cpp
@@ -37,6 +37,14 @@ static FAutoConsoleVariableRef CVarSVTStreamingPrefetchMipLevelBias(
 	ECVF_RenderThreadSafe
 );
 
+static int32 GSVTStreamingForceBlockingRequests = 0;
+static FAutoConsoleVariableRef CVarSVTStreamingForceBlockingRequests(
+	TEXT("r.SparseVolumeTexture.Streaming.ForceBlockingRequests"),
+	GSVTStreamingForceBlockingRequests,
+	TEXT("If enabled, all SVT streaming requests will block on completion, guaranteeing that requested mip levels are available in the same frame they have been requested in (if there is enough memory available to stream them in)."),
+	ECVF_RenderThreadSafe
+);
+
 static int32 GSVTStreamingAsyncThread = 1;
 static FAutoConsoleVariableRef CVarSVTStreamingAsync(
 	TEXT("r.SparseVolumeTexture.Streaming.AsyncThread"),
@@ -676,6 +684,7 @@ void FStreamingManager::InitRHI()
 
 #if WITH_EDITORONLY_DATA
 	RequestOwner = new FRequestOwner(EPriority::Normal);
+	RequestOwnerBlocking = new FRequestOwner(EPriority::Blocking);
 #endif
 }
 
@@ -689,6 +698,8 @@ void FStreamingManager::ReleaseRHI()
 #if WITH_EDITORONLY_DATA
 	delete RequestOwner;
 	RequestOwner = nullptr;
+	delete RequestOwnerBlocking;
+	RequestOwnerBlocking = nullptr;
 #endif
 
 	delete PageTableUpdater;
@@ -748,20 +759,37 @@ void FStreamingManager::Remove_GameThread(UStreamableSparseVolumeTexture* Sparse
 		});
 }
 
-void FStreamingManager::Request_GameThread(UStreamableSparseVolumeTexture* SparseVolumeTexture, float FrameIndex, int32 MipLevel)
+void FStreamingManager::Request_GameThread(UStreamableSparseVolumeTexture* SparseVolumeTexture, float FrameIndex, int32 MipLevel, bool bBlocking)
 {
 	if (!DoesPlatformSupportSparseVolumeTexture(GMaxRHIShaderPlatform) || !SparseVolumeTexture)
 	{
 		return;
 	}
 	ENQUEUE_RENDER_COMMAND(SVTRequest)(
-		[this, SparseVolumeTexture, FrameIndex, MipLevel](FRHICommandListImmediate& RHICmdList)
+		[this, SparseVolumeTexture, FrameIndex, MipLevel, bBlocking](FRHICommandListImmediate& RHICmdList)
+		{
+			Request(SparseVolumeTexture, FrameIndex, MipLevel, bBlocking);
+		});
+}
+
+void FStreamingManager::Update_GameThread()
+{
+	if (!DoesPlatformSupportSparseVolumeTexture(GMaxRHIShaderPlatform))
+	{
+		return;
+	}
+	ENQUEUE_RENDER_COMMAND(SVTUpdate)(
+		[](FRHICommandListImmediate& RHICmdList)
 		{
-			Request(SparseVolumeTexture, FrameIndex, MipLevel);
+			FRDGBuilder GraphBuilder(RHICmdList);
+			const bool bForceNonAsync = true; // No need to spin up a thread if we immediately wait on it anyways.
+			GStreamingManager.BeginAsyncUpdate(GraphBuilder, bForceNonAsync);
+			GStreamingManager.EndAsyncUpdate(GraphBuilder);
+			GraphBuilder.Execute();
 		});
 }
 
-void FStreamingManager::Request(UStreamableSparseVolumeTexture* SparseVolumeTexture, float FrameIndex, int32 MipLevel)
+void FStreamingManager::Request(UStreamableSparseVolumeTexture* SparseVolumeTexture, float FrameIndex, int32 MipLevel, bool bBlocking)
 {
 	check(IsInRenderingThread());
 	if (!DoesPlatformSupportSparseVolumeTexture(GMaxRHIShaderPlatform) || !SparseVolumeTexture)
@@ -837,6 +865,10 @@ void FStreamingManager::Request(UStreamableSparseVolumeTexture* SparseVolumeText
 			Request.Key.FrameIndex = RequestFrameIndex;
 			Request.Key.MipLevelIndex = FMath::Clamp(MipLevel + RequestMipLevelOffset, 0, SVTInfo->PerFrameInfo[RequestFrameIndex].NumMipLevels);
 			Request.Priority = FMath::Max(0, OffsetMagnitude - FMath::Abs(i));
+			if (bBlocking && (i == 0)) // Only block on the actually requested frame
+			{
+				Request.Priority = FStreamingRequest::BlockingPriority;
+			}
 			AddRequest(Request);
 		}
 
@@ -845,7 +877,7 @@ void FStreamingManager::Request(UStreamableSparseVolumeTexture* SparseVolumeText
 	}
 }
 
-void FStreamingManager::BeginAsyncUpdate(FRDGBuilder& GraphBuilder)
+void FStreamingManager::BeginAsyncUpdate(FRDGBuilder& GraphBuilder, bool bForceNonAsync)
 {
 	check(IsInRenderingThread());
 	check(!AsyncState.bUpdateActive);
@@ -891,9 +923,9 @@ void FStreamingManager::BeginAsyncUpdate(FRDGBuilder& GraphBuilder)
 	}
 
 	AddParentRequests();
-	const int32 MaxSelectedMipLevels = MaxPendingMipLevels - NumPendingMipLevels;
-	SelectHighestPriorityRequestsAndUpdateLRU(MaxSelectedMipLevels);
-	IssueRequests(MaxSelectedMipLevels);
+	const int32 MaxSelectedRequests = MaxPendingMipLevels - NumPendingMipLevels;
+	SelectHighestPriorityRequestsAndUpdateLRU(MaxSelectedRequests);
+	IssueRequests(MaxSelectedRequests);
 	AsyncState.NumReadyMipLevels = DetermineReadyMipLevels();
 
 	// Do a first pass over all the mips to be uploaded to compute the upload buffer size requirements.
@@ -931,8 +963,9 @@ void FStreamingManager::BeginAsyncUpdate(FRDGBuilder& GraphBuilder)
 	Parameters.StreamingManager = this;
 	
 	check(AsyncTaskEvents.IsEmpty());
-	if (GSVTStreamingAsyncThread)
+	if (GSVTStreamingAsyncThread && !bForceNonAsync)
 	{
+		AsyncState.bUpdateIsAsync = true;
 		AsyncTaskEvents.Add(TGraphTask<FStreamingUpdateTask>::CreateTask().ConstructAndDispatchWhenReady(Parameters));
 	}
 	else
@@ -953,7 +986,7 @@ void FStreamingManager::EndAsyncUpdate(FRDGBuilder& GraphBuilder)
 	TRACE_CPUPROFILER_EVENT_SCOPE(SVT::FStreamingManager::EndAsyncUpdate);
 
 	// Wait for async processing to finish
-	if (GSVTStreamingAsyncThread)
+	if (AsyncState.bUpdateIsAsync)
 	{
 		check(!AsyncTaskEvents.IsEmpty());
 		FTaskGraphInterface::Get().WaitUntilTasksComplete(AsyncTaskEvents, ENamedThreads::GetRenderThread_Local());
@@ -1010,6 +1043,7 @@ void FStreamingManager::EndAsyncUpdate(FRDGBuilder& GraphBuilder)
 	NumPendingMipLevels -= AsyncState.NumReadyMipLevels;
 	++NextUpdateIndex;
 	AsyncState.bUpdateActive = false;
+	AsyncState.bUpdateIsAsync = false;
 
 #if DO_CHECK
 	for (const auto& Pair : StreamingInfo)
@@ -1290,7 +1324,7 @@ void FStreamingManager::AddParentRequests()
 		FStreamingInfo* SVTInfo = StreamingInfo.Find(Request.Key.SVT);
 		check(SVTInfo);
 		const int32 NumStreamableMipLevels = SVTInfo->PerFrameInfo[Request.Key.FrameIndex].NumMipLevels - 1;
-		uint32 Priority = Request.Value + 1;
+		uint32 Priority = Request.Value == FStreamingRequest::BlockingPriority ? FStreamingRequest::BlockingPriority : (Request.Value + 1);
 		for (int32 MipLevelIndex = Request.Key.MipLevelIndex + 1; MipLevelIndex < NumStreamableMipLevels; ++MipLevelIndex)
 		{
 			FMipLevelKey ParentKey = Request.Key;
@@ -1306,7 +1340,10 @@ void FStreamingManager::AddParentRequests()
 				ParentRequestsToAdd.Add(FStreamingRequest{ ParentKey, Priority });
 			}
 
-			++Priority;
+			if (Priority != FStreamingRequest::BlockingPriority)
+			{
+				++Priority;
+			}
 		}
 	}
 
@@ -1316,10 +1353,10 @@ void FStreamingManager::AddParentRequests()
 	}
 }
 
-void FStreamingManager::SelectHighestPriorityRequestsAndUpdateLRU(int32 MaxSelectedMipLevels)
+void FStreamingManager::SelectHighestPriorityRequestsAndUpdateLRU(int32 MaxSelectedRequests)
 {
 	PrioritizedRequestsHeap.Reset();
-	SelectedMipLevels.Reset();
+	SelectedRequests.Reset();
 
 	if (!RequestsHashTable.IsEmpty())
 	{
@@ -1375,10 +1412,11 @@ void FStreamingManager::SelectHighestPriorityRequestsAndUpdateLRU(int32 MaxSelec
 			}
 		}
 
-		auto PriorityPredicate = [](const auto& A, const auto& B) { return A.Priority > B.Priority; };
+		// Sort by priority but make sure to load higher mip levels with the same priority first. This can happen when a blocking priority is used.
+		auto PriorityPredicate = [](const auto& A, const auto& B) { return A.Priority != B.Priority ? (A.Priority > B.Priority) : A.Key.MipLevelIndex > B.Key.MipLevelIndex; };
 		PrioritizedRequestsHeap.Heapify(PriorityPredicate);
 
-		while (SelectedMipLevels.Num() < MaxSelectedMipLevels && PrioritizedRequestsHeap.Num() > 0)
+		while (SelectedRequests.Num() < MaxSelectedRequests && PrioritizedRequestsHeap.Num() > 0)
 		{
 			FStreamingRequest SelectedRequest;
 			PrioritizedRequestsHeap.HeapPop(SelectedRequest, PriorityPredicate, false /*bAllowShrinking*/);
@@ -1388,7 +1426,7 @@ void FStreamingManager::SelectHighestPriorityRequestsAndUpdateLRU(int32 MaxSelec
 			{
 				check(SelectedRequest.Key.FrameIndex < SVTInfo->PerFrameInfo.Num());
 				check(SelectedRequest.Key.MipLevelIndex < SVTInfo->PerFrameInfo[SelectedRequest.Key.FrameIndex].NumMipLevels);
-				SelectedMipLevels.Push(SelectedRequest.Key);
+				SelectedRequests.Push(SelectedRequest);
 			}
 		}
 
@@ -1396,11 +1434,11 @@ void FStreamingManager::SelectHighestPriorityRequestsAndUpdateLRU(int32 MaxSelec
 	}
 }
 
-void FStreamingManager::IssueRequests(int32 MaxSelectedMipLevels)
+void FStreamingManager::IssueRequests(int32 MaxSelectedRequests)
 {
 	using namespace UE::DerivedData;
 
-	if (SelectedMipLevels.IsEmpty())
+	if (SelectedRequests.IsEmpty())
 	{
 		return;
 	}
@@ -1409,14 +1447,17 @@ void FStreamingManager::IssueRequests(int32 MaxSelectedMipLevels)
 
 #if WITH_EDITORONLY_DATA
 	TArray<FCacheGetChunkRequest> DDCRequests;
-	DDCRequests.Reserve(MaxSelectedMipLevels);
+	DDCRequests.Reserve(MaxSelectedRequests);
+	TArray<FCacheGetChunkRequest> DDCRequestsBlocking;
+	DDCRequestsBlocking.Reserve(MaxSelectedRequests);
 #endif
 
-	FBulkDataBatchRequest::FBatchBuilder Batch = FBulkDataBatchRequest::NewBatch(SelectedMipLevels.Num());
+	FBulkDataBatchRequest::FBatchBuilder Batch = FBulkDataBatchRequest::NewBatch(SelectedRequests.Num());
 	bool bIssueIOBatch = false;
 
-	for (const FMipLevelKey& SelectedKey : SelectedMipLevels)
+	for (const FStreamingRequest& SelectedRequest : SelectedRequests)
 	{
+		const FMipLevelKey& SelectedKey = SelectedRequest.Key;
 		FStreamingInfo* SVTInfo = StreamingInfo.Find(SelectedKey.SVT);
 		check(SVTInfo);
 		check(SVTInfo->PerFrameInfo.Num() > SelectedKey.FrameIndex && SelectedKey.FrameIndex >= 0);
@@ -1500,6 +1541,7 @@ void FStreamingManager::IssueRequests(int32 MaxSelectedMipLevels)
 		PendingMipLevel.FrameIndex = SelectedKey.FrameIndex;
 		PendingMipLevel.MipLevelIndex = SelectedKey.MipLevelIndex;
 		PendingMipLevel.IssuedInFrame = NextUpdateIndex;
+		PendingMipLevel.bBlocking = GSVTStreamingForceBlockingRequests || (SelectedRequest.Priority == FStreamingRequest::BlockingPriority);
 
 		const FByteBulkData& BulkData = Resources->StreamableMipLevels;
 #if WITH_EDITORONLY_DATA
@@ -1513,7 +1555,15 @@ void FStreamingManager::IssueRequests(int32 MaxSelectedMipLevels)
 		{
 			if (Resources->ResourceFlags & EResourceFlag_StreamingDataInDDC)
 			{
-				DDCRequests.Add(BuildDDCRequest(*Resources, MipLevelStreamingInfo, NextPendingMipLevelIndex));
+				UE::DerivedData::FCacheGetChunkRequest DDCRequest = BuildDDCRequest(*Resources, MipLevelStreamingInfo, NextPendingMipLevelIndex);
+				if (PendingMipLevel.bBlocking)
+				{
+					DDCRequestsBlocking.Add(DDCRequest);
+				}
+				else
+				{
+					DDCRequests.Add(DDCRequest);
+				}
 				PendingMipLevel.State = FPendingMipLevel::EState::DDC_Pending;
 			}
 			else
@@ -1525,7 +1575,8 @@ void FStreamingManager::IssueRequests(int32 MaxSelectedMipLevels)
 #endif
 		{
 			PendingMipLevel.RequestBuffer = FIoBuffer(MipLevelStreamingInfo.BulkSize); // SVT_TODO: Use FIoBuffer::Wrap with preallocated memory
-			Batch.Read(BulkData, MipLevelStreamingInfo.BulkOffset, MipLevelStreamingInfo.BulkSize, AIOP_Low, PendingMipLevel.RequestBuffer, PendingMipLevel.Request);
+			const EAsyncIOPriorityAndFlags Priority = PendingMipLevel.bBlocking ? AIOP_CriticalPath : AIOP_Low;
+			Batch.Read(BulkData, MipLevelStreamingInfo.BulkOffset, MipLevelStreamingInfo.BulkSize, Priority, PendingMipLevel.RequestBuffer, PendingMipLevel.Request);
 			bIssueIOBatch = true;
 
 #if WITH_EDITORONLY_DATA
@@ -1584,9 +1635,14 @@ void FStreamingManager::IssueRequests(int32 MaxSelectedMipLevels)
 #if WITH_EDITORONLY_DATA
 	if (!DDCRequests.IsEmpty())
 	{
-		RequestDDCData(DDCRequests);
+		RequestDDCData(DDCRequests, false /*bBlocking*/);
 		DDCRequests.Empty();
 	}
+	if (!DDCRequestsBlocking.IsEmpty())
+	{
+		RequestDDCData(DDCRequestsBlocking, true /*bBlocking*/);
+		DDCRequestsBlocking.Empty();
+	}
 #endif
 
 	if (bIssueIOBatch)
@@ -1701,7 +1757,8 @@ int32 FStreamingManager::DetermineReadyMipLevels()
 
 			const FMipLevelStreamingInfo& MipLevelStreamingInfo = Resources->MipLevelStreamingInfo[PendingMipLevel.MipLevelIndex];
 			FCacheGetChunkRequest Request = BuildDDCRequest(*Resources, MipLevelStreamingInfo, PendingMipLevelIndex);
-			RequestDDCData(MakeArrayView(&Request, 1));
+			const bool bBlocking = GSVTStreamingForceBlockingRequests || PendingMipLevel.bBlocking;
+			RequestDDCData(MakeArrayView(&Request, 1), bBlocking);
 
 			++PendingMipLevel.RetryCount;
 			break;
@@ -2003,36 +2060,44 @@ UE::DerivedData::FCacheGetChunkRequest FStreamingManager::BuildDDCRequest(const
 	return Request;
 }
 
-void FStreamingManager::RequestDDCData(TConstArrayView<UE::DerivedData::FCacheGetChunkRequest> DDCRequests)
+void FStreamingManager::RequestDDCData(TConstArrayView<UE::DerivedData::FCacheGetChunkRequest> DDCRequests, bool bBlocking)
 {
 	using namespace UE::DerivedData;
 
-	FRequestBarrier Barrier(*RequestOwner);	// This is a critical section on the owner. It does not constrain ordering
-	GetCache().GetChunks(DDCRequests, *RequestOwner,
-		[this](FCacheGetChunkResponse&& Response)
-		{
-			const uint32 PendingMipLevelIndex = (uint32)(Response.UserData >> uint64(32));
-			const uint32 RequestVersion = (uint32)Response.UserData;
-			
-			// In case the request returned after the mip level was already streamed out again we need to abort so that we do not overwrite data in the FPendingMipLevel slot.
-			if (RequestVersion < PendingMipLevels[PendingMipLevelIndex].RequestVersion)
+	{
+		FRequestOwner* RequestOwnerPtr = bBlocking ? RequestOwnerBlocking : RequestOwner;
+		FRequestBarrier Barrier(*RequestOwnerPtr);	// This is a critical section on the owner. It does not constrain ordering
+		GetCache().GetChunks(DDCRequests, *RequestOwnerPtr,
+			[this](FCacheGetChunkResponse&& Response)
 			{
-				return;
-			}
-			
-			FPendingMipLevel& PendingMipLevel = PendingMipLevels[PendingMipLevelIndex];
-			check(PendingMipLevel.SparseVolumeTexture); // A valid PendingMipLevel should have a non-nullptr here
+				const uint32 PendingMipLevelIndex = (uint32)(Response.UserData >> uint64(32));
+				const uint32 RequestVersion = (uint32)Response.UserData;
 
-			if (Response.Status == EStatus::Ok)
-			{
-				PendingMipLevel.SharedBuffer = MoveTemp(Response.RawData);
-				PendingMipLevel.State = FPendingMipLevel::EState::DDC_Ready;
-			}
-			else
-			{
-				PendingMipLevel.State = FPendingMipLevel::EState::DDC_Failed;
-			}
-		});
+				// In case the request returned after the mip level was already streamed out again we need to abort so that we do not overwrite data in the FPendingMipLevel slot.
+				if (RequestVersion < PendingMipLevels[PendingMipLevelIndex].RequestVersion)
+				{
+					return;
+				}
+
+				FPendingMipLevel& PendingMipLevel = PendingMipLevels[PendingMipLevelIndex];
+				check(PendingMipLevel.SparseVolumeTexture); // A valid PendingMipLevel should have a non-nullptr here
+
+				if (Response.Status == EStatus::Ok)
+				{
+					PendingMipLevel.SharedBuffer = MoveTemp(Response.RawData);
+					PendingMipLevel.State = FPendingMipLevel::EState::DDC_Ready;
+				}
+				else
+				{
+					PendingMipLevel.State = FPendingMipLevel::EState::DDC_Failed;
+				}
+			});
+	}
+	
+	if (bBlocking)
+	{
+		RequestOwnerBlocking->Wait();
+	}
 }
 
 #endif // WITH_EDITORONLY_DATA
diff --git a/Engine/Source/Runtime/Engine/Private/SparseVolumeTexture/SparseVolumeTextureStreamingManager.h b/Engine/Source/Runtime/Engine/Private/SparseVolumeTexture/SparseVolumeTextureStreamingManager.h
--- a/Engine/Source/Runtime/Engine/Private/SparseVolumeTexture/SparseVolumeTextureStreamingManager.h
+++ b/Engine/Source/Runtime/Engine/Private/SparseVolumeTexture/SparseVolumeTextureStreamingManager.h
@@ -64,8 +64,9 @@ struct FMipLevelKey
 
 struct FStreamingRequest
 {
+	static constexpr uint32 BlockingPriority = 0xFFFFFFFFu;
 	FMipLevelKey Key;
-	uint32 Priority; // A higher value means a higher priority
+	uint32 Priority; // A higher value means a higher priority, 0xFFFFFF means blocking
 
 	FORCEINLINE bool operator<(const FStreamingRequest& Other) const
 	{
@@ -86,12 +87,11 @@ public:
 	//~ Begin IStreamingManager Interface.
 	virtual void Add_GameThread(UStreamableSparseVolumeTexture* SparseVolumeTexture) override;
 	virtual void Remove_GameThread(UStreamableSparseVolumeTexture* SparseVolumeTexture) override;
-	// Request a frame to be streamed in. FrameIndex is of float type so that the fractional part can be used to better track the playback speed/direction.
-	// This function automatically also requests all higher mip levels and adds prefetch requests for upcoming frames.
-	virtual void Request_GameThread(UStreamableSparseVolumeTexture* SparseVolumeTexture, float FrameIndex, int32 MipLevel) override;
+	virtual void Request_GameThread(UStreamableSparseVolumeTexture* SparseVolumeTexture, float FrameIndex, int32 MipLevel, bool bBlocking) override;
+	virtual void Update_GameThread() override;
 	
-	virtual void Request(UStreamableSparseVolumeTexture* SparseVolumeTexture, float FrameIndex, int32 MipLevel) override;
-	virtual void BeginAsyncUpdate(FRDGBuilder& GraphBuilder) override;
+	virtual void Request(UStreamableSparseVolumeTexture* SparseVolumeTexture, float FrameIndex, int32 MipLevel, bool bBlocking) override;
+	virtual void BeginAsyncUpdate(FRDGBuilder& GraphBuilder, bool bForceNonAsync) override;
 	virtual void EndAsyncUpdate(FRDGBuilder& GraphBuilder) override;
 	//~ End IStreamingManager Interface.
 
@@ -234,6 +234,7 @@ private:
 		FIoBuffer RequestBuffer;
 		FBulkDataBatchReadRequest Request;
 		uint32 IssuedInFrame = 0;
+		bool bBlocking = false;
 
 		void Reset()
 		{
@@ -249,6 +250,7 @@ private:
 			RequestBuffer = {};
 			Request.Reset();
 			IssuedInFrame = 0;
+			bBlocking = false;
 		}
 	};
 
@@ -301,6 +303,7 @@ private:
 	{
 		int32 NumReadyMipLevels = 0;
 		bool bUpdateActive = false;
+		bool bUpdateIsAsync = false;
 	};
 
 	// Stores info extracted from the UStreamableSparseVolumeTexture and its USparseVolumeTextureFrames so that we don't need to access these uobjects on the rendering thread
@@ -320,6 +323,7 @@ private:
 	TArray<FPendingMipLevel> PendingMipLevels;
 #if WITH_EDITORONLY_DATA
 	UE::DerivedData::FRequestOwner* RequestOwner = nullptr;
+	UE::DerivedData::FRequestOwner* RequestOwnerBlocking = nullptr;
 #endif
 
 	class FPageTableUpdater* PageTableUpdater = nullptr;
@@ -336,7 +340,7 @@ private:
 	TSet<FTileDataTexture*> TileDataTexturesToUpdate;
 	TSet<FStreamingInfo*> InvalidatedStreamingInfos;
 	TArray<FStreamingRequest> PrioritizedRequestsHeap;
-	TArray<FMipLevelKey> SelectedMipLevels;
+	TArray<FStreamingRequest> SelectedRequests;
 	TArray<FPageTableClear> PageTableClears;
 	TArray<FUploadTask> UploadTasks; // accessed on the async thread
 	TArray<FPendingMipLevel*> UploadCleanupTasks; // accessed on the async thread
@@ -346,15 +350,15 @@ private:
 	bool AddRequest(const FStreamingRequest& Request); // Returns true when the request is new or it has been updated with a higher priority
 	void AsyncUpdate();
 	void AddParentRequests(); // Add requests for all parent mip levels
-	void SelectHighestPriorityRequestsAndUpdateLRU(int32 MaxSelectedMipLevels);
-	void IssueRequests(int32 MaxSelectedMipLevels);
+	void SelectHighestPriorityRequestsAndUpdateLRU(int32 MaxSelectedRequests);
+	void IssueRequests(int32 MaxSelectedRequests);
 	void StreamOutMipLevel(FStreamingInfo& SVTInfo, FLRUNode* LRUNode);
 	int32 DetermineReadyMipLevels();
 	void InstallReadyMipLevels();
 
 #if WITH_EDITORONLY_DATA
 	UE::DerivedData::FCacheGetChunkRequest BuildDDCRequest(const FResources& Resources, const FMipLevelStreamingInfo& MipLevelStreamingInfo, const uint32 PendingMipLevelIndex);
-	void RequestDDCData(TConstArrayView<UE::DerivedData::FCacheGetChunkRequest> DDCRequests);
+	void RequestDDCData(TConstArrayView<UE::DerivedData::FCacheGetChunkRequest> DDCRequests, bool bBlocking);
 #endif
 };
 
diff --git a/Engine/Source/Runtime/Engine/Public/SparseVolumeTexture/ISparseVolumeTextureStreamingManager.h b/Engine/Source/Runtime/Engine/Public/SparseVolumeTexture/ISparseVolumeTextureStreamingManager.h
--- a/Engine/Source/Runtime/Engine/Public/SparseVolumeTexture/ISparseVolumeTextureStreamingManager.h
+++ b/Engine/Source/Runtime/Engine/Public/SparseVolumeTexture/ISparseVolumeTextureStreamingManager.h
@@ -22,12 +22,15 @@ public:
 	virtual void Remove_GameThread(UStreamableSparseVolumeTexture* SparseVolumeTexture) = 0;
 	// Request a frame to be streamed in. FrameIndex is of float type so that the fractional part can be used to better track the playback speed/direction.
 	// This function automatically also requests all higher mip levels and adds prefetch requests for upcoming frames.
-	virtual void Request_GameThread(UStreamableSparseVolumeTexture* SparseVolumeTexture, float FrameIndex, int32 MipLevel) = 0;
+	virtual void Request_GameThread(UStreamableSparseVolumeTexture* SparseVolumeTexture, float FrameIndex, int32 MipLevel, bool bBlocking) = 0;
+	// Issues a rendering command for updating the streaming manager. This is not normally needed, but may be necessary in cases where blocking requests are required
+	// and the SVT is used in a different rendering command that is executed before the streaming manager would normally update.
+	virtual void Update_GameThread() = 0;
 	//~ End game thread functions.
 	
 	//~ Begin rendering thread functions.
-	virtual void Request(UStreamableSparseVolumeTexture* SparseVolumeTexture, float FrameIndex, int32 MipLevel) = 0;
-	virtual void BeginAsyncUpdate(FRDGBuilder& GraphBuilder) = 0;
+	virtual void Request(UStreamableSparseVolumeTexture* SparseVolumeTexture, float FrameIndex, int32 MipLevel, bool bBlocking) = 0;
+	virtual void BeginAsyncUpdate(FRDGBuilder& GraphBuilder, bool bForceNonAsync = false) = 0;
 	virtual void EndAsyncUpdate(FRDGBuilder& GraphBuilder) = 0;
 	//~ End rendering thread functions.
 
