Reduced memory duplication in FInstancedStaticMeshVertexFactory

diff --git a/Engine/Source/Runtime/Engine/Classes/Engine/InstancedStaticMesh.h b/Engine/Source/Runtime/Engine/Classes/Engine/InstancedStaticMesh.h
--- a/Engine/Source/Runtime/Engine/Classes/Engine/InstancedStaticMesh.h
+++ b/Engine/Source/Runtime/Engine/Classes/Engine/InstancedStaticMesh.h
@@ -278,10 +278,6 @@ public:
 	{
 	}
 
-	struct FDataType : public FInstancedStaticMeshDataType, public FLocalVertexFactory::FDataType
-	{
-	};
-
 	/**
 	 * Should we cache the material's shadertype on this platform with this vertex factory? 
 	 */
@@ -297,15 +293,18 @@ public:
 	 * Get vertex elements used when during PSO precaching materials using this vertex factory type
 	 */
 	static void GetPSOPrecacheVertexFetchElements(EVertexInputStreamType VertexInputStreamType, FVertexDeclarationElementList& Elements);
-	static void GetVertexElements(ERHIFeatureLevel::Type FeatureLevel, EVertexInputStreamType InputStreamType, bool bSupportsManualVertexFetch, FDataType& Data, FVertexDeclarationElementList& Elements);
+	static void GetVertexElements(ERHIFeatureLevel::Type FeatureLevel, EVertexInputStreamType InputStreamType, bool bSupportsManualVertexFetch, FDataType& Data, FInstancedStaticMeshDataType& InstanceData, FVertexDeclarationElementList& Elements);
 
 	/**
 	 * An implementation of the interface used by TSynchronizedResource to update the resource with new data from the game thread.
 	 */
-	void SetData(const FDataType& InData)
+	void SetData(const FDataType& InData, const FInstancedStaticMeshDataType* InInstanceData)
 	{
-		FLocalVertexFactory::Data = InData;
 		Data = InData;
+		if (InInstanceData)
+		{
+			InstanceData = *InInstanceData;
+		}
 		UpdateRHI();
 	}
 
@@ -326,22 +325,22 @@ public:
 
 	inline FRHIShaderResourceView* GetInstanceOriginSRV() const
 	{
-		return Data.InstanceOriginSRV;
+		return InstanceData.InstanceOriginSRV;
 	}
 
 	inline FRHIShaderResourceView* GetInstanceTransformSRV() const
 	{
-		return Data.InstanceTransformSRV;
+		return InstanceData.InstanceTransformSRV;
 	}
 
 	inline FRHIShaderResourceView* GetInstanceLightmapSRV() const
 	{
-		return Data.InstanceLightmapSRV;
+		return InstanceData.InstanceLightmapSRV;
 	}
 
 	inline FRHIShaderResourceView* GetInstanceCustomDataSRV() const
 	{
-		return Data.InstanceCustomDataSRV;
+		return InstanceData.InstanceCustomDataSRV;
 	}
 
 	FRHIUniformBuffer* GetUniformBuffer() const
@@ -349,10 +348,10 @@ public:
 		return UniformBuffer.GetReference();
 	}
 protected:
-	static void GetVertexElements(ERHIFeatureLevel::Type FeatureLevel, EVertexInputStreamType InputStreamType, bool bSupportsManualVertexFetch, FDataType& Data, FVertexDeclarationElementList& Elements, FVertexStreamList& Streams);
+	static void GetVertexElements(ERHIFeatureLevel::Type FeatureLevel, EVertexInputStreamType InputStreamType, bool bSupportsManualVertexFetch, FDataType& Data, FInstancedStaticMeshDataType& InstanceData, FVertexDeclarationElementList& Elements, FVertexStreamList& Streams);
 
 private:
-	FDataType Data;
+	FInstancedStaticMeshDataType InstanceData;
 
 	TUniformBufferRef<FInstancedStaticMeshVertexFactoryUniformShaderParameters> UniformBuffer;
 };
diff --git a/Engine/Source/Runtime/Engine/Private/GeometryCollectionRendering.cpp b/Engine/Source/Runtime/Engine/Private/GeometryCollectionRendering.cpp
--- a/Engine/Source/Runtime/Engine/Private/GeometryCollectionRendering.cpp
+++ b/Engine/Source/Runtime/Engine/Private/GeometryCollectionRendering.cpp
@@ -204,9 +204,9 @@ void FGeometryCollectionVertexFactory::InitRHI()
 		}
 	}
 
-	if (Data.LightMapCoordinateComponent.VertexBuffer)
+	if (Data.LightMapCoordinateIndex >= 0)
 	{
-		Elements.Add(AccessStreamComponent(Data.LightMapCoordinateComponent, 15));
+		Elements.Add(AccessStreamComponent(Data.TextureCoordinates[Data.LightMapCoordinateIndex], 15));
 	}
 	else if (Data.TextureCoordinates.Num())
 	{
diff --git a/Engine/Source/Runtime/Engine/Private/InstancedStaticMesh.cpp b/Engine/Source/Runtime/Engine/Private/InstancedStaticMesh.cpp
--- a/Engine/Source/Runtime/Engine/Private/InstancedStaticMesh.cpp
+++ b/Engine/Source/Runtime/Engine/Private/InstancedStaticMesh.cpp
@@ -59,6 +59,10 @@
 #include "Rendering/StaticLightingSystemInterface.h"
 #endif
 
+#if !UE_BUILD_DEBUG && !WITH_EDITORONLY_DATA
+static_assert(sizeof(FInstancedStaticMeshVertexFactory) <= 672, "FInstancedStaticMeshVertexFactory was optimized to fit within 672 bytes bin of MallocBinned3");
+#endif
+
 IMPLEMENT_TYPE_LAYOUT(FInstancedStaticMeshVertexFactoryShaderParameters);
 
 CSV_DEFINE_CATEGORY_MODULE(ENGINE_API, InstancedStaticMeshComponent, true);
@@ -944,11 +948,13 @@ void FInstancedStaticMeshVertexFactory::GetPSOPrecacheVertexFetchElements(EVerte
 void FInstancedStaticMeshVertexFactory::Copy(const FInstancedStaticMeshVertexFactory& Other)
 {
 	FInstancedStaticMeshVertexFactory* VertexFactory = this;
-	const FDataType* DataCopy = &Other.Data;
+	const FLocalVertexFactory::FDataType* DataCopy = &Other.Data;
+	const FInstancedStaticMeshDataType* InstanceDataCopy = &Other.InstanceData;
 	ENQUEUE_RENDER_COMMAND(FInstancedStaticMeshVertexFactoryCopyData)(
-	[VertexFactory, DataCopy](FRHICommandListImmediate& RHICmdList)
+	[VertexFactory, DataCopy, InstanceDataCopy](FRHICommandListImmediate& RHICmdList)
 	{
 		VertexFactory->Data = *DataCopy;
+		VertexFactory->InstanceData = *InstanceDataCopy;
 	});
 	BeginUpdateResourceRHI(this);
 }
@@ -966,7 +972,7 @@ void FInstancedStaticMeshVertexFactory::InitRHI()
 	const bool bUseManualVertexFetch = GetType()->SupportsManualVertexFetch(ThisFeatureLevel);
 
 	FVertexDeclarationElementList Elements;
-	GetVertexElements(ThisFeatureLevel, EVertexInputStreamType::Default, bUseManualVertexFetch, Data, Elements, Streams);
+	GetVertexElements(ThisFeatureLevel, EVertexInputStreamType::Default, bUseManualVertexFetch, Data, InstanceData, Elements, Streams);
 
 	// on mobile with GPUScene enabled instanced attributes[8-12] are used for a general auto-instancing
 	// so we add them only for desktop or if mobile has GPUScene disabled
@@ -992,7 +998,7 @@ void FInstancedStaticMeshVertexFactory::InitRHI()
 		UniformParameters.VertexFetch_InstanceTransformBuffer = GetInstanceTransformSRV();
 		UniformParameters.VertexFetch_InstanceLightmapBuffer = GetInstanceLightmapSRV();
 		UniformParameters.InstanceCustomDataBuffer = GetInstanceCustomDataSRV();
-		UniformParameters.NumCustomDataFloats = Data.NumCustomDataFloats;
+		UniformParameters.NumCustomDataFloats = InstanceData.NumCustomDataFloats;
 		UniformBuffer = TUniformBufferRef<FInstancedStaticMeshVertexFactoryUniformShaderParameters>::CreateUniformBufferImmediate(UniformParameters, UniformBuffer_MultiFrame, EUniformBufferValidation::None);
 	}
 }
@@ -1002,10 +1008,11 @@ void FInstancedStaticMeshVertexFactory::GetVertexElements(
 	EVertexInputStreamType InputStreamType,
 	bool bSupportsManualVertexFetch,
 	FDataType& Data,
+	FInstancedStaticMeshDataType& InstanceData,
 	FVertexDeclarationElementList& Elements)
 {
 	FVertexStreamList VertexStreams;
-	GetVertexElements(FeatureLevel, InputStreamType, bSupportsManualVertexFetch, Data, Elements, VertexStreams);
+	GetVertexElements(FeatureLevel, InputStreamType, bSupportsManualVertexFetch, Data, InstanceData, Elements, VertexStreams);
 }
 
 void FInstancedStaticMeshVertexFactory::GetVertexElements(
@@ -1013,6 +1020,7 @@ void FInstancedStaticMeshVertexFactory::GetVertexElements(
 	EVertexInputStreamType InputStreamType, 
 	bool bSupportsManualVertexFetch, 
 	FDataType& Data, 
+	FInstancedStaticMeshDataType& InstanceData,
 	FVertexDeclarationElementList& Elements, 
 	FVertexStreamList& Streams)
 {
@@ -1081,9 +1089,9 @@ void FInstancedStaticMeshVertexFactory::GetVertexElements(
 			Elements.Add(AccessStreamComponent(NullComponent, 14, Streams));
 		}
 
-		if (Data.LightMapCoordinateComponent.VertexBuffer)
+		if (Data.LightMapCoordinateIndex >= 0)
 		{
-			Elements.Add(AccessStreamComponent(Data.LightMapCoordinateComponent, 15, Streams));
+			Elements.Add(AccessStreamComponent(Data.TextureCoordinates[Data.LightMapCoordinateIndex], 15, Streams));
 		}
 		else if (Data.TextureCoordinates.Num())
 		{
@@ -1097,23 +1105,23 @@ void FInstancedStaticMeshVertexFactory::GetVertexElements(
 	if (FeatureLevel > ERHIFeatureLevel::ES3_1 || !bMobileUsesGPUScene)
 	{
 		// toss in the instanced location stream
-		check(bCanUseGPUScene || Data.InstanceOriginComponent.VertexBuffer);
-		if (Data.InstanceOriginComponent.VertexBuffer)
+		check(bCanUseGPUScene || InstanceData.InstanceOriginComponent.VertexBuffer);
+		if (InstanceData.InstanceOriginComponent.VertexBuffer)
 		{
-			Elements.Add(AccessStreamComponent(Data.InstanceOriginComponent, 8, Streams));
+			Elements.Add(AccessStreamComponent(InstanceData.InstanceOriginComponent, 8, Streams));
 		}
 
-		check(bCanUseGPUScene || Data.InstanceTransformComponent[0].VertexBuffer);
-		if (Data.InstanceTransformComponent[0].VertexBuffer)
+		check(bCanUseGPUScene || InstanceData.InstanceTransformComponent[0].VertexBuffer);
+		if (InstanceData.InstanceTransformComponent[0].VertexBuffer)
 		{
-			Elements.Add(AccessStreamComponent(Data.InstanceTransformComponent[0], 9, Streams));
-			Elements.Add(AccessStreamComponent(Data.InstanceTransformComponent[1], 10, Streams));
-			Elements.Add(AccessStreamComponent(Data.InstanceTransformComponent[2], 11, Streams));
+			Elements.Add(AccessStreamComponent(InstanceData.InstanceTransformComponent[0], 9, Streams));
+			Elements.Add(AccessStreamComponent(InstanceData.InstanceTransformComponent[1], 10, Streams));
+			Elements.Add(AccessStreamComponent(InstanceData.InstanceTransformComponent[2], 11, Streams));
 		}
 
-		if (Data.InstanceLightmapAndShadowMapUVBiasComponent.VertexBuffer)
+		if (InstanceData.InstanceLightmapAndShadowMapUVBiasComponent.VertexBuffer)
 		{
-			Elements.Add(AccessStreamComponent(Data.InstanceLightmapAndShadowMapUVBiasComponent, 12, Streams));
+			Elements.Add(AccessStreamComponent(InstanceData.InstanceLightmapAndShadowMapUVBiasComponent, 12, Streams));
 		}
 	}
 }
@@ -1179,7 +1187,8 @@ void InitInstancedStaticMeshVertexFactoryComponents(
 	const FInstancedStaticMeshVertexFactory* VertexFactory,
 	int32 LightMapCoordinateIndex,
 	bool bRHISupportsManualVertexFetch,
-	FInstancedStaticMeshVertexFactory::FDataType& OutData)
+	FInstancedStaticMeshVertexFactory::FDataType& OutData,
+	FInstancedStaticMeshDataType& OutInstanceData)
 {
 	VertexBuffers.PositionVertexBuffer.BindPositionVertexBuffer(VertexFactory, OutData);
 	VertexBuffers.StaticMeshVertexBuffer.BindTangentVertexBuffer(VertexFactory, OutData);
@@ -1202,7 +1211,7 @@ void InitInstancedStaticMeshVertexFactoryComponents(
 
 	if (InstanceBuffer)
 	{
-		InstanceBuffer->BindInstanceVertexBuffer(VertexFactory, OutData);
+		InstanceBuffer->BindInstanceVertexBuffer(VertexFactory, OutInstanceData);
 	}
 }
 
@@ -1225,6 +1234,7 @@ void FInstancedStaticMeshRenderData::BindBuffersToVertexFactories()
 	{
 		const FStaticMeshLODResources* RenderData = &LODModels[LODIndex];
 
+		FInstancedStaticMeshDataType InstanceData;
 		FInstancedStaticMeshVertexFactory::FDataType Data;
 		// Assign to the vertex factory for this LOD.
 		FInstancedStaticMeshVertexFactory& VertexFactory = VertexFactories[LODIndex];
@@ -1233,8 +1243,8 @@ void FInstancedStaticMeshRenderData::BindBuffersToVertexFactories()
 		{
 			ColorVertexBuffer = Component->LODData[LODIndex].OverrideVertexColors;
 		}
-		InitInstancedStaticMeshVertexFactoryComponents(RenderData->VertexBuffers, ColorVertexBuffer, InstanceBuffer, &VertexFactory, LightMapCoordinateIndex, bRHISupportsManualVertexFetch, Data);
-		VertexFactory.SetData(Data);
+		InitInstancedStaticMeshVertexFactoryComponents(RenderData->VertexBuffers, ColorVertexBuffer, InstanceBuffer, &VertexFactory, LightMapCoordinateIndex, bRHISupportsManualVertexFetch, Data, InstanceData);
+		VertexFactory.SetData(Data, InstanceBuffer ? &InstanceData : nullptr);
 		VertexFactory.InitResource();
 	}
 }
@@ -5003,14 +5013,15 @@ void UInstancedStaticMeshComponent::CollectPSOPrecacheData(const FPSOPrecachePar
 	
 	auto ISMC_GetElements = [LightMapCoordinateIndex, InstanceBuffer, this](const FStaticMeshLODResources& LODRenderData, int32 LODIndex, bool bSupportsManualVertexFetch, FVertexDeclarationElementList& Elements)
 	{
+		FInstancedStaticMeshDataType InstanceData;
 		FInstancedStaticMeshVertexFactory::FDataType Data;
 		const FColorVertexBuffer* ColorVertexBuffer = LODRenderData.bHasColorVertexData ? &(LODRenderData.VertexBuffers.ColorVertexBuffer) : nullptr;
 		if (LODData.IsValidIndex(LODIndex) && LODData[LODIndex].OverrideVertexColors)
 		{
 			ColorVertexBuffer = LODData[LODIndex].OverrideVertexColors;
 		}
-		InitInstancedStaticMeshVertexFactoryComponents(LODRenderData.VertexBuffers, ColorVertexBuffer, InstanceBuffer, nullptr /*VertexFactory*/, LightMapCoordinateIndex, bSupportsManualVertexFetch, Data);
-		FInstancedStaticMeshVertexFactory::GetVertexElements(GMaxRHIFeatureLevel, EVertexInputStreamType::Default, bSupportsManualVertexFetch, Data, Elements);
+		InitInstancedStaticMeshVertexFactoryComponents(LODRenderData.VertexBuffers, ColorVertexBuffer, InstanceBuffer, nullptr /*VertexFactory*/, LightMapCoordinateIndex, bSupportsManualVertexFetch, Data, InstanceData);
+		FInstancedStaticMeshVertexFactory::GetVertexElements(GMaxRHIFeatureLevel, EVertexInputStreamType::Default, bSupportsManualVertexFetch, Data, InstanceData, Elements);
 	};
 
 	CollectPSOPrecacheDataImpl(VFType, BasePrecachePSOParams, ISMC_GetElements, OutParams);
diff --git a/Engine/Source/Runtime/Engine/Private/LocalVertexFactory.cpp b/Engine/Source/Runtime/Engine/Private/LocalVertexFactory.cpp
--- a/Engine/Source/Runtime/Engine/Private/LocalVertexFactory.cpp
+++ b/Engine/Source/Runtime/Engine/Private/LocalVertexFactory.cpp
@@ -556,9 +556,9 @@ void FLocalVertexFactory::GetVertexElements(
 		FVertexStreamComponent NullComponent(&GNullVertexBuffer, 0, 0, VET_Float4);
 		Elements.Add(AccessStreamComponent(Data.PreSkinPositionComponent.VertexBuffer ? Data.PreSkinPositionComponent : NullComponent, 14, InOutStreams));
 
-		if (Data.LightMapCoordinateComponent.VertexBuffer)
+		if (Data.LightMapCoordinateIndex >= 0)
 		{
-			Elements.Add(AccessStreamComponent(Data.LightMapCoordinateComponent, 15, InOutStreams));
+			Elements.Add(AccessStreamComponent(Data.TextureCoordinates[Data.LightMapCoordinateIndex], 15, InOutStreams));
 		}
 		else if (Data.TextureCoordinates.Num())
 		{
diff --git a/Engine/Source/Runtime/Engine/Private/Rendering/StaticMeshVertexBuffer.cpp b/Engine/Source/Runtime/Engine/Private/Rendering/StaticMeshVertexBuffer.cpp
--- a/Engine/Source/Runtime/Engine/Private/Rendering/StaticMeshVertexBuffer.cpp
+++ b/Engine/Source/Runtime/Engine/Private/Rendering/StaticMeshVertexBuffer.cpp
@@ -642,11 +642,22 @@ void FStaticMeshVertexBuffer::BindLightMapVertexBuffer(const FVertexFactory* Ver
 
 		if (LightMapCoordinateIndex >= 0 && (uint32)LightMapCoordinateIndex < GetNumTexCoords())
 		{
-			Data.LightMapCoordinateComponent = FVertexStreamComponent(
+			FVertexStreamComponent* Stream = nullptr;
+			if (LightMapCoordinateIndex < Data.TextureCoordinates.Num())
+			{
+				Stream = &Data.TextureCoordinates[LightMapCoordinateIndex];
+			}
+			else
+			{
+				check(Data.TextureCoordinates.Num() == LightMapCoordinateIndex);
+				Stream = &Data.TextureCoordinates.AddDefaulted_GetRef();
+			}
+
+			*Stream = FVertexStreamComponent(
 				&TexCoordVertexBuffer,
 				UVSizeInBytes * LightMapCoordinateIndex,
 				UvStride,
-				UVVertexElementType, 
+				UVVertexElementType,
 				EVertexStreamUsage::ManualFetch
 			);
 		}
diff --git a/Engine/Source/Runtime/Engine/Public/Components.h b/Engine/Source/Runtime/Engine/Public/Components.h
--- a/Engine/Source/Runtime/Engine/Public/Components.h
+++ b/Engine/Source/Runtime/Engine/Public/Components.h
@@ -48,9 +48,6 @@ struct FStaticMeshDataType
 	/** The streams to read the texture coordinates from. */
 	TArray<FVertexStreamComponent, TFixedAllocator<MAX_STATIC_TEXCOORDS / 2> > TextureCoordinates;
 
-	/** The stream to read the shadow map texture coordinates from. */
-	FVertexStreamComponent LightMapCoordinateComponent;
-
 	/** The stream to read the vertex color from. */
 	FVertexStreamComponent ColorComponent;
 
@@ -64,10 +61,10 @@ struct FStaticMeshDataType
 	/** A SRV to manually bind and load Colors in the vertex shader. */
 	FRHIShaderResourceView* ColorComponentsSRV = nullptr;
 
-	int LightMapCoordinateIndex = -1;
-	int NumTexCoords = -1;
 	uint32 ColorIndexMask = ~0u;
-	uint32 LODLightmapDataIndex = 0;
+	int8 LightMapCoordinateIndex = -1;
+	uint8 NumTexCoords = 0;
+	uint8 LODLightmapDataIndex = 0;
 };
 
 /** The information used to build a mesh. */
diff --git a/Engine/Source/Runtime/RenderCore/Public/VertexFactory.h b/Engine/Source/Runtime/RenderCore/Public/VertexFactory.h
--- a/Engine/Source/Runtime/RenderCore/Public/VertexFactory.h
+++ b/Engine/Source/Runtime/RenderCore/Public/VertexFactory.h
@@ -748,9 +748,6 @@ protected:
 		
 	/** The vertex streams used to render the factory. */
 	FVertexStreamList Streams;
-
-	/* VF can explicitly set this to false to avoid errors without decls; this is for VFs that fetch from buffers directly (e.g. Niagara) */
-	bool bNeedsDeclaration = true;
 	
 	static constexpr int32 PrimitiveIdStreamStride = 0;
 
@@ -774,6 +771,11 @@ private:
 	int8 PrimitiveIdStreamIndex[(int32)EVertexInputStreamType::Count];
 #endif
 
+protected:
+	/* VF can explicitly set this to false to avoid errors without decls; this is for VFs that fetch from buffers directly (e.g. Niagara) */
+	bool bNeedsDeclaration = true;
+
+private:
 	static int32 TranslatePrimitiveIdStreamIndex(const FStaticFeatureLevel InFeatureLevel, EVertexInputStreamType InputStreamType)
 	{
 	#if WITH_EDITOR
