UTexture::UpdateResource(): combined 3 RT tasks into one to reduce the number of tasks spike during loading, both for perf and memory usage

diff --git a/Engine/Source/Runtime/Engine/Private/Texture.cpp b/Engine/Source/Runtime/Engine/Private/Texture.cpp
--- a/Engine/Source/Runtime/Engine/Private/Texture.cpp
+++ b/Engine/Source/Runtime/Engine/Private/Texture.cpp
@@ -254,6 +254,8 @@ void UTexture::SetResource(FTextureResource* InResource)
 
 void UTexture::ReleaseResource()
 {
+	check(!IsInActualRenderingThread() && !IsInRHIThread());
+
 	if (PrivateResource)
 	{
 		UnlinkStreaming();
@@ -269,9 +271,10 @@ void UTexture::ReleaseResource()
 
 		FTextureResource* ToDelete = PrivateResource;
 		// Free the resource.
-		SetResource(nullptr);
-		ENQUEUE_RENDER_COMMAND(DeleteResource)([ToDelete](FRHICommandListImmediate& RHICmdList)
+		PrivateResource = nullptr;
+		ENQUEUE_RENDER_COMMAND(DeleteResource)([this, ToDelete](FRHICommandListImmediate& RHICmdList)
 		{
+			PrivateResourceRenderThread = nullptr;
 			ToDelete->ReleaseResource();
 			delete ToDelete;
 		});
@@ -347,7 +350,22 @@ void UTexture::UpdateResource()
 			}
 		}
 
-		SetResource(NewResource);
+		ENQUEUE_RENDER_COMMAND(SetTextureResourceRenderThread)
+			(
+				[this, NewResource](FRHICommandListImmediate& RHICmdList)
+				{
+					PrivateResourceRenderThread = NewResource;
+					if (NewResource)
+					{
+						// Init the texture reference, which needs to be set from a render command, since TextureReference.TextureReferenceRHI is gamethread coherent.
+						NewResource->SetTextureReference(TextureReference.TextureReferenceRHI);
+						
+						NewResource->InitResource();
+					}
+				}
+		);
+
+		PrivateResource = NewResource;
 		if (NewResource)
 		{
 			LLM_SCOPE(ELLMTag::Textures);
@@ -362,14 +380,7 @@ void UTexture::UpdateResource()
 				}
 			}
 
-			// Init the texture reference, which needs to be set from a render command, since TextureReference.TextureReferenceRHI is gamethread coherent.
-			ENQUEUE_RENDER_COMMAND(SetTextureReference)([this, NewResource](FRHICommandListImmediate& RHICmdList)
-			{
-				NewResource->SetTextureReference(TextureReference.TextureReferenceRHI);
-			});
-
 			NewResource->SetOwnerName(FName(GetPathName()));
-			BeginInitResource(NewResource);
 			// Now that the resource is ready for streaming, bind it to the streamer.
 			LinkStreaming();
 		}
