Small refactor to centralize the seek logic

diff --git a/Engine/Source/Runtime/BinkAudioDecoder/Module/Private/BinkAudioInfo.cpp b/Engine/Source/Runtime/BinkAudioDecoder/Module/Private/BinkAudioInfo.cpp
--- a/Engine/Source/Runtime/BinkAudioDecoder/Module/Private/BinkAudioInfo.cpp
+++ b/Engine/Source/Runtime/BinkAudioDecoder/Module/Private/BinkAudioInfo.cpp
@@ -113,25 +113,19 @@ void FBinkAudioInfo::PrepareToLoop()
 //-----------------------------------------------------------------------------
 //-----------------------------------------------------------------------------
 void FBinkAudioInfo::SeekToTime(const float SeekTimeSeconds)
-{
+{	
+	// If there's no seek table on the header, fall-back to Super implementation.
+	if (Decoder->SeekTableCount == 0)
+	{
+		Super::SeekToTime(SeekTimeSeconds);
+		return;
+	}
+
 	uint32 SeekTimeSamples = 0;
 	if (SeekTimeSeconds > 0)
 	{
 		SeekTimeSamples = (uint32)(SeekTimeSeconds * SampleRate);
 	}
-	
-	// If there's no seek table on the header, see if we there's an embedded one on the chunk.
-	if (Decoder->SeekTableCount == 0)
-	{
-		int32 ChunkIndexToSeekTo = StreamingSoundWave->GetSoundWaveData()->FindChunkIndexForSeeking(SeekTimeSamples);
-		if (ChunkIndexToSeekTo >= 0)
-		{
-			this->StreamSeekBlockIndex = ChunkIndexToSeekTo;
-			this->StreamSeekBlockOffset = 0;
-			this->StreamSeekToSamples = SeekTimeSamples;
-		}
-		return;
-	}
 
 	uint32 SamplesInFrame = GetMaxFrameSizeSamples();
 	if (SeekTimeSamples > this->TrueSampleCount)
diff --git a/Engine/Source/Runtime/BinkAudioDecoder/Module/Public/BinkAudioInfo.h b/Engine/Source/Runtime/BinkAudioDecoder/Module/Public/BinkAudioInfo.h
--- a/Engine/Source/Runtime/BinkAudioDecoder/Module/Public/BinkAudioInfo.h
+++ b/Engine/Source/Runtime/BinkAudioDecoder/Module/Public/BinkAudioInfo.h
@@ -27,6 +27,7 @@ public:
 	//~ End IStreamedCompressedInfo Interface
 
 protected:
+	using Super = IStreamedCompressedInfo;
 
 	// copied from header during ParseHeader
 	uint32 MaxCompSpaceNeeded;
diff --git a/Engine/Source/Runtime/Engine/Private/AudioDecompress.cpp b/Engine/Source/Runtime/Engine/Private/AudioDecompress.cpp
--- a/Engine/Source/Runtime/Engine/Private/AudioDecompress.cpp
+++ b/Engine/Source/Runtime/Engine/Private/AudioDecompress.cpp
@@ -119,6 +119,32 @@ bool IStreamedCompressedInfo::ReadCompressedData(uint8* Destination, bool bLoopi
 	return bFinished;
 }
 
+void IStreamedCompressedInfo::SeekToTime(const float InSeekTimeSeconds)
+{
+	if (StreamingSoundWave.IsValid())
+	{
+		const FSoundWavePtr WaveData = StreamingSoundWave->GetSoundWaveData();
+		if (WaveData.IsValid())
+		{
+			// Negative time will seek to start.
+			uint32 SeekTimeAudioFrames = 0;
+			if (InSeekTimeSeconds > 0)
+			{
+				SeekTimeAudioFrames = IntCastChecked<uint32>((uint64)(WaveData->GetSampleRate() * InSeekTimeSeconds));
+			}
+
+			// If we have a chunk setup to contain a seek-table it will return a value other than INDEX_NONE here.
+			const int32 ChunkIndexToSeekTo = WaveData->FindChunkIndexForSeeking(SeekTimeAudioFrames);
+			if (ChunkIndexToSeekTo >= 0)
+			{
+				StreamSeekBlockIndex = ChunkIndexToSeekTo;
+				StreamSeekBlockOffset = 0;								// We don't know this until the seek-table is loaded, so we leave it 0 for now.
+				StreamSeekToAudioFrames = SeekTimeAudioFrames;			// Store the time in samples so when we load the chunks table loads we can find the offset.
+			}
+		}
+	}
+}
+
 void IStreamedCompressedInfo::ExpandFile(uint8* DstBuffer, struct FSoundQualityInfo* QualityInfo)
 {
 	check(DstBuffer);
@@ -239,17 +265,17 @@ bool IStreamedCompressedInfo::StreamCompressedData(uint8* Destination, bool bLoo
 			if (ensureMsgf(FStreamedAudioChunkSeekTable::Parse(SrcBufferData, ChunkSize, TableOffset, GetCurrentSeekTable()), 
 				TEXT("Failed to parse seektable in '%s', chunk=%d"), *StreamingSoundWave->GetFName().ToString(), CurrentChunkIndex))
 			{
-				if (StreamSeekToSamples != INDEX_NONE)
+				if (StreamSeekToAudioFrames != INDEX_NONE)
 				{
 					const FStreamedAudioChunk& Chunk = StreamingSoundWave->GetSoundWaveData()->GetChunk(CurrentChunkIndex);
-					int32 SeekOffsetInChunkSpace = StreamSeekToSamples - Chunk.SeekOffsetInAudioFrames;
-					uint32 Offset = GetCurrentSeekTable().FindOffset(StreamSeekToSamples);
+					int32 SeekOffsetInChunkSpace = StreamSeekToAudioFrames - Chunk.SeekOffsetInAudioFrames;
+					uint32 Offset = GetCurrentSeekTable().FindOffset(StreamSeekToAudioFrames);
 					if (Offset != INDEX_NONE)
 					{
 						SrcBufferOffset = Offset + TableOffset;
 						SrcBufferOffset += CurrentChunkIndex == 0 ? AudioDataOffset : 0;
 					}
-					StreamSeekToSamples = INDEX_NONE;
+					StreamSeekToAudioFrames = INDEX_NONE;
 				}
 			}
 		}
diff --git a/Engine/Source/Runtime/Engine/Public/AudioDecompress.h b/Engine/Source/Runtime/Engine/Public/AudioDecompress.h
--- a/Engine/Source/Runtime/Engine/Public/AudioDecompress.h
+++ b/Engine/Source/Runtime/Engine/Public/AudioDecompress.h
@@ -197,7 +197,7 @@ public:
 	//~ Begin ICompressedInfo Interface
 	ENGINE_API virtual bool ReadCompressedInfo(const uint8* InSrcBufferData, uint32 InSrcBufferDataSize, FSoundQualityInfo* QualityInfo) override;
 	ENGINE_API virtual bool ReadCompressedData(uint8* Destination, bool bLooping, uint32 BufferSize) override;
-	ENGINE_API virtual void SeekToTime(const float SeekTime) override {};
+	ENGINE_API virtual void SeekToTime(const float SeekTime) override;;
 	ENGINE_API virtual void ExpandFile(uint8* DstBuffer, struct FSoundQualityInfo* QualityInfo) override;
 	ENGINE_API virtual void EnableHalfRate(bool HalfRate) override {};
 	virtual uint32 GetSourceBufferSize() const override { return SrcBufferDataSize; }
@@ -339,7 +339,7 @@ protected:
 	int32 StreamSeekBlockOffset;
 
 	/** If using the Chunked seek-tables, we request the seek in samples so we can resolve after the chunk table loads. */
-	uint32 StreamSeekToSamples = INDEX_NONE;	
+	uint32 StreamSeekToAudioFrames = INDEX_NONE;	
 };
 
 struct IAudioInfoFactory 
