PostProcessVolume: Cleanup overrides by introducing a macro similar to LERP_PP for variables that don't support blending.

Keep the logic identical to what was there before, just applying SET_PP in those cases where it is equivalent to the previous logic. Note that a few settings don't obey the right naming convention for the macro to work. Also note that the IF_PP macro is slightly different as it also checks if the variable in question is non-zero. This prevents a PPV from overriding a setting to off. This may need to be revisited in the future.

diff --git a/Engine/Source/Runtime/Engine/Private/SceneView.cpp b/Engine/Source/Runtime/Engine/Private/SceneView.cpp
--- a/Engine/Source/Runtime/Engine/Private/SceneView.cpp
+++ b/Engine/Source/Runtime/Engine/Private/SceneView.cpp
@@ -1427,7 +1427,8 @@ bool FSceneView::ProjectWorldToScreen(const FVector& WorldPosition, const FIntRe
 
 
 #define LERP_PP(NAME) if(Src.bOverride_ ## NAME)	Dest . NAME = FMath::Lerp(Dest . NAME, Src . NAME, Weight);
-#define IF_PP(NAME) if(Src.bOverride_ ## NAME && Src . NAME)
+#define SET_PP(NAME)  if(Src.bOverride_ ## NAME)    Dest . NAME = Src . NAME;
+#define IF_PP(NAME)   if(Src.bOverride_ ## NAME && Src . NAME)
 
 // @param Weight 0..1
 void FSceneView::OverridePostProcessSettings(const FPostProcessSettings& Src, float Weight)
@@ -1449,10 +1450,7 @@ void FSceneView::OverridePostProcessSettings(const FPostProcessSettings& Src, fl
 		FFinalPostProcessSettings& Dest = FinalPostProcessSettings;
 
 		// The following code needs to be adjusted when settings in FPostProcessSettings change.
-		if (Src.bOverride_TemperatureType)
-		{
-			Dest.TemperatureType = Src.TemperatureType;
-		}
+		SET_PP(TemperatureType);
 		LERP_PP(WhiteTemp);
 		LERP_PP(WhiteTint);
 
@@ -1603,244 +1601,66 @@ void FSceneView::OverridePostProcessSettings(const FPostProcessSettings& Src, fl
 		LERP_PP(ScreenSpaceReflectionIntensity);
 		LERP_PP(ScreenSpaceReflectionMaxRoughness);
 
-		if (Src.bOverride_RayTracingReflectionsMaxRoughness)
-		{
-			Dest.RayTracingReflectionsMaxRoughness = Src.RayTracingReflectionsMaxRoughness;
-		}
-
-		if (Src.bOverride_RayTracingReflectionsMaxBounces)
-		{
-			Dest.RayTracingReflectionsMaxBounces = Src.RayTracingReflectionsMaxBounces;
-		}
-
-		if (Src.bOverride_RayTracingReflectionsSamplesPerPixel)
-		{
-			Dest.RayTracingReflectionsSamplesPerPixel = Src.RayTracingReflectionsSamplesPerPixel;
-		}
-
-		if (Src.bOverride_RayTracingReflectionsShadows)
-		{
-			Dest.RayTracingReflectionsShadows = Src.RayTracingReflectionsShadows;
-		}
-
-		if (Src.bOverride_RayTracingReflectionsTranslucency)
-		{
-			Dest.RayTracingReflectionsTranslucency = Src.RayTracingReflectionsTranslucency;
-		}
-
-		if (Src.bOverride_TranslucencyType)
-		{
-			Dest.TranslucencyType = Src.TranslucencyType;
-		}
-
-		if (Src.bOverride_RayTracingTranslucencyMaxRoughness)
-		{
-			Dest.RayTracingTranslucencyMaxRoughness = Src.RayTracingTranslucencyMaxRoughness;
-		}
-
-		if (Src.bOverride_RayTracingTranslucencyRefractionRays)
-		{
-			Dest.RayTracingTranslucencyRefractionRays = Src.RayTracingTranslucencyRefractionRays;
-		}
-
-		if (Src.bOverride_RayTracingTranslucencySamplesPerPixel)
-		{
-			Dest.RayTracingTranslucencySamplesPerPixel = Src.RayTracingTranslucencySamplesPerPixel;
-		}
-
-		if (Src.bOverride_RayTracingTranslucencyShadows)
-		{
-			Dest.RayTracingTranslucencyShadows = Src.RayTracingTranslucencyShadows;
-		}
-
-		if (Src.bOverride_RayTracingTranslucencyRefraction)
-		{
-			Dest.RayTracingTranslucencyRefraction = Src.RayTracingTranslucencyRefraction;
-		}
+		SET_PP(RayTracingReflectionsMaxRoughness);
+		SET_PP(RayTracingReflectionsMaxBounces);
+		SET_PP(RayTracingReflectionsSamplesPerPixel);
+		SET_PP(RayTracingReflectionsShadows);
+		SET_PP(RayTracingReflectionsTranslucency);
+		SET_PP(TranslucencyType);
+		SET_PP(RayTracingTranslucencyMaxRoughness);
+		SET_PP(RayTracingTranslucencyRefractionRays);
+		SET_PP(RayTracingTranslucencySamplesPerPixel);
+		SET_PP(RayTracingTranslucencyShadows);
+		SET_PP(RayTracingTranslucencyRefraction);
 
 		if (Src.bOverride_RayTracingGI)
 		{
 			Dest.RayTracingGIType = Src.RayTracingGIType;
 		}
 
-		if (Src.bOverride_DynamicGlobalIlluminationMethod)
-		{
-			Dest.DynamicGlobalIlluminationMethod = Src.DynamicGlobalIlluminationMethod;
-		}
-
-		if (Src.bOverride_LumenSurfaceCacheResolution)
-		{
-			Dest.LumenSurfaceCacheResolution = Src.LumenSurfaceCacheResolution;
-		}
-
-		if (Src.bOverride_LumenSceneLightingQuality)
-		{
-			Dest.LumenSceneLightingQuality = Src.LumenSceneLightingQuality;
-		}
-
-		if (Src.bOverride_LumenSceneDetail)
-		{
-			Dest.LumenSceneDetail = Src.LumenSceneDetail;
-		}
-		
-		if (Src.bOverride_LumenSceneViewDistance)
-		{
-			Dest.LumenSceneViewDistance = Src.LumenSceneViewDistance;
-		}
-
-		if (Src.bOverride_LumenSceneLightingUpdateSpeed)
-		{
-			Dest.LumenSceneLightingUpdateSpeed = Src.LumenSceneLightingUpdateSpeed;
-		}
-
-		if (Src.bOverride_LumenFinalGatherQuality)
-		{
-			Dest.LumenFinalGatherQuality = Src.LumenFinalGatherQuality;
-		}
-
-		if (Src.bOverride_LumenFinalGatherLightingUpdateSpeed)
-		{
-			Dest.LumenFinalGatherLightingUpdateSpeed = Src.LumenFinalGatherLightingUpdateSpeed;
-		}
-
-		if (Src.bOverride_LumenMaxTraceDistance)
-		{
-			Dest.LumenMaxTraceDistance = Src.LumenMaxTraceDistance;
-		}
+		SET_PP(DynamicGlobalIlluminationMethod);
+		SET_PP(LumenSurfaceCacheResolution);
+		SET_PP(LumenSceneLightingQuality);
+		SET_PP(LumenSceneDetail);
+		SET_PP(LumenSceneViewDistance);
+		SET_PP(LumenSceneLightingUpdateSpeed);
+		SET_PP(LumenFinalGatherQuality);
+		SET_PP(LumenFinalGatherLightingUpdateSpeed);
+		SET_PP(LumenMaxTraceDistance);
 
 		LERP_PP(LumenDiffuseColorBoost);
 		LERP_PP(LumenSkylightLeaking);
 		LERP_PP(LumenFullSkylightLeakingDistance);
 
-		if (Src.bOverride_LumenRayLightingMode)
-		{
-			Dest.LumenRayLightingMode = Src.LumenRayLightingMode;
-		}
-
-		if (Src.bOverride_LumenFrontLayerTranslucencyReflections)
-		{
-			Dest.LumenFrontLayerTranslucencyReflections = Src.LumenFrontLayerTranslucencyReflections;
-		}
-
-		if (Src.bOverride_LumenMaxReflectionBounces)
-		{
-			Dest.LumenMaxReflectionBounces = Src.LumenMaxReflectionBounces;
-		}
-
-		if (Src.bOverride_ReflectionMethod)
-		{
-			Dest.ReflectionMethod = Src.ReflectionMethod;
-		}
-
-		if (Src.bOverride_LumenReflectionQuality)
-		{
-			Dest.LumenReflectionQuality = Src.LumenReflectionQuality;
-		}
-
-		if (Src.bOverride_RayTracingGIMaxBounces)
-		{
-			Dest.RayTracingGIMaxBounces = Src.RayTracingGIMaxBounces;
-		}
-
-		if (Src.bOverride_RayTracingGISamplesPerPixel)
-		{
-			Dest.RayTracingGISamplesPerPixel = Src.RayTracingGISamplesPerPixel;
-		}
-
-		if (Src.bOverride_RayTracingAO)
-		{
-			Dest.RayTracingAO = Src.RayTracingAO;
-		}
-
-		if (Src.bOverride_RayTracingAOSamplesPerPixel)
-		{
-			Dest.RayTracingAOSamplesPerPixel = Src.RayTracingAOSamplesPerPixel;
-		}
-
-		if (Src.bOverride_RayTracingAOIntensity)
-		{
-			Dest.RayTracingAOIntensity = Src.RayTracingAOIntensity;
-		}
-
-		if (Src.bOverride_RayTracingAORadius)
-		{
-			Dest.RayTracingAORadius = Src.RayTracingAORadius;
-		}
-
-		if (Src.bOverride_PathTracingMaxBounces)
-		{
-			Dest.PathTracingMaxBounces = Src.PathTracingMaxBounces;
-		}
-
-		if (Src.bOverride_PathTracingSamplesPerPixel)
-		{
-			Dest.PathTracingSamplesPerPixel = Src.PathTracingSamplesPerPixel;
-		}
-
-		if (Src.bOverride_PathTracingMaxPathExposure)
-		{
-			Dest.PathTracingMaxPathExposure = Src.PathTracingMaxPathExposure;
-		}
-
-		if (Src.bOverride_PathTracingEnableReferenceDOF)
-		{
-			Dest.PathTracingEnableReferenceDOF = Src.PathTracingEnableReferenceDOF;
-		}
-
-		if (Src.bOverride_PathTracingEnableReferenceAtmosphere)
-		{
-			Dest.PathTracingEnableReferenceAtmosphere = Src.PathTracingEnableReferenceAtmosphere;
-		}
-
-		if (Src.bOverride_PathTracingEnableDenoiser)
-		{
-			Dest.PathTracingEnableDenoiser = Src.PathTracingEnableDenoiser;
-		}
-
-		if (Src.bOverride_PathTracingIncludeEmissive)
-		{
-			Dest.PathTracingIncludeEmissive = Src.PathTracingIncludeEmissive;
-		}
-
-		if (Src.bOverride_PathTracingIncludeIndirectEmissive)
-		{
-			Dest.PathTracingIncludeIndirectEmissive = Src.PathTracingIncludeIndirectEmissive;
-		}
-
-		if (Src.bOverride_PathTracingIncludeDiffuse)
-		{
-			Dest.PathTracingIncludeDiffuse = Src.PathTracingIncludeDiffuse;
-		}
-
-		if (Src.bOverride_PathTracingIncludeIndirectDiffuse)
-		{
-			Dest.PathTracingIncludeIndirectDiffuse = Src.PathTracingIncludeIndirectDiffuse;
-		}
-
-		if (Src.bOverride_PathTracingIncludeSpecular)
-		{
-			Dest.PathTracingIncludeSpecular = Src.PathTracingIncludeSpecular;
-		}
-
-		if (Src.bOverride_PathTracingIncludeIndirectSpecular)
-		{
-			Dest.PathTracingIncludeIndirectSpecular = Src.PathTracingIncludeIndirectSpecular;
-		}
-
-		if (Src.bOverride_PathTracingIncludeVolume)
-		{
-			Dest.PathTracingIncludeVolume = Src.PathTracingIncludeVolume;
-		}
-
-		if (Src.bOverride_PathTracingIncludeIndirectVolume)
-		{
-			Dest.PathTracingIncludeIndirectVolume = Src.PathTracingIncludeIndirectVolume;
-		}
-
-		if (Src.bOverride_DepthOfFieldBladeCount)
-		{
-			Dest.DepthOfFieldBladeCount = Src.DepthOfFieldBladeCount;
-		}
+		SET_PP(LumenRayLightingMode);
+		SET_PP(LumenFrontLayerTranslucencyReflections);
+		SET_PP(LumenMaxReflectionBounces);
+		SET_PP(ReflectionMethod);
+		SET_PP(LumenReflectionQuality);
+		SET_PP(RayTracingGIMaxBounces);
+		SET_PP(RayTracingGISamplesPerPixel);
+		SET_PP(RayTracingAO);
+		SET_PP(RayTracingAOSamplesPerPixel);
+		SET_PP(RayTracingAOIntensity);
+		SET_PP(RayTracingAORadius);
+
+		// Path Tracing related settings
+		SET_PP(PathTracingMaxBounces);
+		SET_PP(PathTracingSamplesPerPixel);
+		LERP_PP(PathTracingMaxPathExposure);
+		SET_PP(PathTracingEnableReferenceDOF);
+		SET_PP(PathTracingEnableReferenceAtmosphere);
+		SET_PP(PathTracingEnableDenoiser);
+		SET_PP(PathTracingIncludeEmissive);
+		SET_PP(PathTracingIncludeIndirectEmissive);
+		SET_PP(PathTracingIncludeDiffuse);
+		SET_PP(PathTracingIncludeIndirectDiffuse);
+		SET_PP(PathTracingIncludeSpecular);
+		SET_PP(PathTracingIncludeIndirectSpecular);
+		SET_PP(PathTracingIncludeVolume);
+		SET_PP(PathTracingIncludeIndirectVolume);
+
+		SET_PP(DepthOfFieldBladeCount);
 
 		// cubemaps are getting blended additively - in contrast to other properties, maybe we should make that consistent
 		if (Src.AmbientCubemap && Src.bOverride_AmbientCubemapIntensity)
@@ -1870,10 +1690,7 @@ void FSceneView::OverridePostProcessSettings(const FPostProcessSettings& Src, fl
 			Dest.BloomDirtMask = Src.BloomDirtMask;
 		}
 
-		if (Src.bOverride_BloomMethod)
-		{
-			Dest.BloomMethod = Src.BloomMethod;
-		}
+		SET_PP(BloomMethod);
 
 		// actual texture cannot be blended but the intensity can be blended
 		IF_PP(BloomConvolutionTexture)
@@ -1924,25 +1741,13 @@ void FSceneView::OverridePostProcessSettings(const FPostProcessSettings& Src, fl
 			Dest.bMobileHQGaussian = Src.bMobileHQGaussian;
 		}
 
-		if (Src.bOverride_AutoExposureMethod)
-		{
-			Dest.AutoExposureMethod = Src.AutoExposureMethod;
-		}
+		SET_PP(AutoExposureMethod);
 
-		if (Src.bOverride_AmbientOcclusionRadiusInWS)
-		{
-			Dest.AmbientOcclusionRadiusInWS = Src.AmbientOcclusionRadiusInWS;
-		}
+		SET_PP(AmbientOcclusionRadiusInWS);
 
-		if (Src.bOverride_MotionBlurTargetFPS)
-		{
-			Dest.MotionBlurTargetFPS = Src.MotionBlurTargetFPS;
-		}
+		SET_PP(MotionBlurTargetFPS);
 
-		if (Src.bOverride_AutoExposureApplyPhysicalCameraExposure)
-		{
-			Dest.AutoExposureApplyPhysicalCameraExposure = Src.AutoExposureApplyPhysicalCameraExposure;
-		}
+		SET_PP(AutoExposureApplyPhysicalCameraExposure);
 	}
 
 	// Blendable objects
