change the default of cluster union components to force the flush net dormancy call (bClusterUnionCallFlushNetDormancy=true)

diff --git a/Engine/Source/Runtime/Engine/Private/PhysicsEngine/ClusterUnionComponent.cpp b/Engine/Source/Runtime/Engine/Private/PhysicsEngine/ClusterUnionComponent.cpp
--- a/Engine/Source/Runtime/Engine/Private/PhysicsEngine/ClusterUnionComponent.cpp
+++ b/Engine/Source/Runtime/Engine/Private/PhysicsEngine/ClusterUnionComponent.cpp
@@ -24,7 +24,7 @@ namespace
 	bool bUseClusterUnionAccelerationStructure = false;
 	FAutoConsoleVariableRef CVarUseClusterUnionAccelerationStructure(TEXT("ClusterUnion.UseAccelerationStructure"), bUseClusterUnionAccelerationStructure, TEXT("Whether component level sweeps and overlaps against cluster unions should use an acceleration structure instead."));
 
-	bool bClusterUnionCallFlushNetDormancy = false;
+	bool bClusterUnionCallFlushNetDormancy = true;
 	FAutoConsoleVariableRef CVarClusterUnionCallFlushNetDormanc(TEXT("clusterunion.callflushnetdormancy"), bClusterUnionCallFlushNetDormancy, TEXT("Whether or not to call FlushNetDormancy"));
 
 	// TODO: Should this be exposed in Chaos instead?
