switch cluster union sync bodies to check the proxy to component map instead of using the proxy's user data as a way to primitive component - should help prevent the existence of destroyed particles in the SQ

diff --git a/Engine/Source/Runtime/Engine/Private/PhysicsEngine/Experimental/PhysScene_Chaos.cpp b/Engine/Source/Runtime/Engine/Private/PhysicsEngine/Experimental/PhysScene_Chaos.cpp
--- a/Engine/Source/Runtime/Engine/Private/PhysicsEngine/Experimental/PhysScene_Chaos.cpp
+++ b/Engine/Source/Runtime/Engine/Private/PhysicsEngine/Experimental/PhysScene_Chaos.cpp
@@ -2169,7 +2169,12 @@ void FPhysScene_Chaos::OnSyncBodies(Chaos::FPhysicsSolverBase* Solver)
 
 			// Cluster unions should always have their owner be the cluster union component.
 			FClusterUnionPhysicsProxy::FExternalParticle* DirtyParticle = Proxy->GetParticle_External();
-			if (UClusterUnionComponent* ParentComponent = Cast<UClusterUnionComponent>(FChaosUserData::Get<UPrimitiveComponent>(DirtyParticle->UserData())))
+
+			// Use the scene's GetOwningComponent rather than UserData since it's guaranteed to be safe and correspond to the GT state of the component.
+			// i.e. if we destroy the component's physics state and then call sync bodies (not entirely impossible), we don't try to add the cluster union
+			// back into the acceleration structure (since destroying the physics state already took it out). In that case, we could potentially end up with
+			// a deleted cluster union particle in the SQ.
+			if (UClusterUnionComponent* ParentComponent = Outer->GetOwningComponent<UClusterUnionComponent>(Proxy))
 			{
 				const FRigidTransform3 NewTransform(DirtyParticle->X(), DirtyParticle->R());
 
