Fix non-unity compile issues

diff --git a/Engine/Source/Runtime/Engine/Public/PrimitiveSceneShaderData.h b/Engine/Source/Runtime/Engine/Public/PrimitiveSceneShaderData.h
--- a/Engine/Source/Runtime/Engine/Public/PrimitiveSceneShaderData.h
+++ b/Engine/Source/Runtime/Engine/Public/PrimitiveSceneShaderData.h
@@ -5,6 +5,8 @@
 #include "CoreMinimal.h"
 #include "SceneDefinitions.h"
 #include "Containers/StaticArray.h"
+#include "InstanceUniformShaderParameters.h"
+#include "LightmapUniformShaderParameters.h"
 
 #if !UE_ENABLE_INCLUDE_ORDER_DEPRECATED_IN_5_3
 #include "PrimitiveUniformShaderParameters.h"
