Added includes to improve compile times

Before: 1018.94 seconds
After: 963.55 seconds

diff --git a/Engine/Source/Runtime/Engine/Public/EngineSharedPCH.h b/Engine/Source/Runtime/Engine/Public/EngineSharedPCH.h
--- a/Engine/Source/Runtime/Engine/Public/EngineSharedPCH.h
+++ b/Engine/Source/Runtime/Engine/Public/EngineSharedPCH.h
@@ -239,4 +239,7 @@
 #include "Vehicles/TireType.h"
 #include "VisualLogger/VisualLogger.h"
 #include "VisualLogger/VisualLoggerTypes.h"
-#include "WorldCollision.h"
\ No newline at end of file
+#include "WorldCollision.h"
+
+// Chaos
+#include "Chaos/ParticleHandle.h"
\ No newline at end of file
