963.55 seconds
920.17 seconds

Added more includes to improve compile times

Before: 963.55 seconds
After: 920.17 seconds

diff --git a/Engine/Source/Runtime/Engine/Public/EngineSharedPCH.h b/Engine/Source/Runtime/Engine/Public/EngineSharedPCH.h
--- a/Engine/Source/Runtime/Engine/Public/EngineSharedPCH.h
+++ b/Engine/Source/Runtime/Engine/Public/EngineSharedPCH.h
@@ -242,4 +242,5 @@
 #include "WorldCollision.h"
 
 // Chaos
-#include "Chaos/ParticleHandle.h"
\ No newline at end of file
+#include "Chaos/ParticleHandle.h"
+#include "Chaos/Framework/PhysicsProxy.h"
