Networked Physics: Extend FPhysicsReplication to work with FPhysicsObject when using new replication modes.

FPhysicsReplication now have two flows splitting out from the SetReplicatedTarget() fucntion by populating two different collections.
The split is only done if Network Physics Prediction is enabled in project settings.

One flow is an untouched legacy flow, except that it has been condensed down to only support the Default physics replication.
Another flow is based on the new FPhysicsObject and is purely running in the async callbacks and it supports multiple different replication modes.
(Replication modes are work in progress and will be refactored and improved in following changes.)

Projects that rely on the current default replication are still able to turn on Network Physics Prediction and test the new replication modes while having the legacy replication working in parallel.
To enabled the Default replication in the new flow use CVar: np2.EnableDefaultReplication.
This will break any custom overrides / extensions that has been done locally to the default replication though.

This CL also makes it possible to choose replication mode of a simulating static mesh instance placed in the scene.
On the actor I also moved the Set and Get functions for PhysicsReplicationMode to Category=Replication since that's where the blueprint option is to set the value.

I moved ResimFrame from FPhysicsReplication to FRewindData which cleaned up some logic around setting and getting it in various scripts.

diff --git a/Engine/Source/Runtime/Engine/Classes/Engine/StaticMeshActor.h b/Engine/Source/Runtime/Engine/Classes/Engine/StaticMeshActor.h
--- a/Engine/Source/Runtime/Engine/Classes/Engine/StaticMeshActor.h
+++ b/Engine/Source/Runtime/Engine/Classes/Engine/StaticMeshActor.h
@@ -33,6 +33,10 @@ public:
 	UPROPERTY(Category=Actor, EditAnywhere, AdvancedDisplay)
 	bool bStaticMeshReplicateMovement;
 
+	/** Set which replication mode to use for this static mesh instance if it both replicates and simulates physics. */
+	UPROPERTY(Category = Actor, EditAnywhere, AdvancedDisplay)
+	EPhysicsReplicationMode StaticMeshPhysicsReplicationMode;
+
 	UE_DEPRECATED(5.0, "Unused property. The actor will use the DefaultGeometryGatheringMode set in FNavigationOctree (see virtual ENavDataGatheringMode GetGeometryGatheringMode()).")
 	UPROPERTY(EditAnywhere, Category = Navigation, AdvancedDisplay)
 	ENavDataGatheringMode NavigationGeometryGatheringMode;
diff --git a/Engine/Source/Runtime/Engine/Classes/GameFramework/Actor.h b/Engine/Source/Runtime/Engine/Classes/GameFramework/Actor.h
--- a/Engine/Source/Runtime/Engine/Classes/GameFramework/Actor.h
+++ b/Engine/Source/Runtime/Engine/Classes/GameFramework/Actor.h
@@ -769,11 +769,11 @@ private:
 	
 public:
 	/** Set the physics replication mode of this body, via EPhysicsReplicationMode */
-	UFUNCTION(BlueprintCallable, Category = Physics)
+	UFUNCTION(BlueprintCallable, Category = Replication)
 	void SetPhysicsReplicationMode(const EPhysicsReplicationMode ReplicationMode);
 
 	/** Get the physics replication mode of this body, via EPhysicsReplicationMode */
-	UFUNCTION(BlueprintCallable, Category = Physics)
+	UFUNCTION(BlueprintCallable, Category = Replication)
 	EPhysicsReplicationMode GetPhysicsReplicationMode();
 
 	/** Can this body trigger a resimulation when Physics Prediction is enabled */
diff --git a/Engine/Source/Runtime/Engine/Private/PhysicsEngine/AsyncPhysicsInputComponent.cpp b/Engine/Source/Runtime/Engine/Private/PhysicsEngine/AsyncPhysicsInputComponent.cpp
--- a/Engine/Source/Runtime/Engine/Private/PhysicsEngine/AsyncPhysicsInputComponent.cpp
+++ b/Engine/Source/Runtime/Engine/Private/PhysicsEngine/AsyncPhysicsInputComponent.cpp
@@ -59,15 +59,9 @@ struct FAsyncPhysicsInputRewindCallback : public Chaos::IRewindCallback
 	virtual int32 TriggerRewindIfNeeded_Internal(int32 LatestStepCompleted) override
 	{
 		int32 CachedClientFrame = INDEX_NONE;
-		if (World->GetNetMode() == NM_Client && World->GetPhysicsScene())
+		if (World->GetNetMode() == NM_Client && RewindData != nullptr)
 		{
-			if (FPhysScene_Chaos* Scene = static_cast<FPhysScene_Chaos*>(World->GetPhysicsScene()))
-			{
-				if (const IPhysicsReplication* PhysicsReplication  = Scene->GetPhysicsReplication())
-				{
-					CachedClientFrame = Scene->GetPhysicsReplication()->GetResimFrame();
-				}
-			}
+			CachedClientFrame = RewindData->GetResimFrame();
 		}
 
 		return CachedClientFrame;
diff --git a/Engine/Source/Runtime/Engine/Private/PhysicsEngine/Experimental/PhysScene_Chaos.cpp b/Engine/Source/Runtime/Engine/Private/PhysicsEngine/Experimental/PhysScene_Chaos.cpp
--- a/Engine/Source/Runtime/Engine/Private/PhysicsEngine/Experimental/PhysScene_Chaos.cpp
+++ b/Engine/Source/Runtime/Engine/Private/PhysicsEngine/Experimental/PhysScene_Chaos.cpp
@@ -86,12 +86,9 @@ public:
 		{
 			if (!PendingCommands[0].OwningObject.IsStale())
 			{
-				if(UWorld* World = PendingCommands[0].OwningObject->GetWorld())
+				if (Chaos::FRewindData* RewindData = GetSolver()->GetRewindData())
 				{
-					if (FPhysScene* PhysScene = World->GetPhysicsScene())
-					{
-						PhysScene->GetPhysicsReplication()->SetResimFrame(INDEX_NONE);
-					}
+					RewindData->SetResimFrame(INDEX_NONE);
 				}
 			}
 		}
@@ -111,16 +108,13 @@ public:
 			{
 				if (PendingCommands[Idx].PhysicsStep < CurrentFrame && PendingCommands[Idx].OwningObject.IsValid())
 				{
-					if(UWorld* World = PendingCommands[Idx].OwningObject->GetWorld())
+					if (Chaos::FRewindData* RewindData = GetSolver()->GetRewindData())
 					{
-						if (FPhysScene* PhysScene = World->GetPhysicsScene())
-						{
-							int32 ResimFrame = PhysScene->GetPhysicsReplication()->GetResimFrame();
-							ResimFrame = (ResimFrame == INDEX_NONE) ? PendingCommands[Idx].PhysicsStep :
-								FMath::Min(ResimFrame, PendingCommands[Idx].PhysicsStep);
+						int32 ResimFrame = RewindData->GetResimFrame();
+						ResimFrame = (ResimFrame == INDEX_NONE) ? PendingCommands[Idx].PhysicsStep :
+							FMath::Min(ResimFrame, PendingCommands[Idx].PhysicsStep);
 
-							PhysScene->GetPhysicsReplication()->SetResimFrame(ResimFrame);
-						}
+						RewindData->SetResimFrame(ResimFrame);
 					}
 				}
 				else if (!bRemove && PendingCommands[Idx].PhysicsStep == CurrentFrame)
diff --git a/Engine/Source/Runtime/Engine/Private/PhysicsEngine/NetworkPhysicsComponent.cpp b/Engine/Source/Runtime/Engine/Private/PhysicsEngine/NetworkPhysicsComponent.cpp
--- a/Engine/Source/Runtime/Engine/Private/PhysicsEngine/NetworkPhysicsComponent.cpp
+++ b/Engine/Source/Runtime/Engine/Private/PhysicsEngine/NetworkPhysicsComponent.cpp
@@ -115,36 +115,26 @@ int32 FNetworkPhysicsCallback::TriggerRewindIfNeeded_Internal(int32 LatestStepCo
 	UE_LOG(LogTemp, Log, TEXT("CLIENT | PT | TriggerRewindIfNeeded_Internal | Callbacks Frame = %d"), ResimFrame);
 #endif
 
-	if (World->GetNetMode() == NM_Client && World->GetPhysicsScene())
+	if (RewindData)
 	{
-		if (FPhysScene_Chaos* Scene = static_cast<FPhysScene_Chaos*>(World->GetPhysicsScene()))
+		if (World->GetNetMode() == NM_Client)
 		{
-			if (const IPhysicsReplication* PhysicsReplication = Scene->GetPhysicsReplication())
-			{
-				const int32 ReplicationFrame = PhysicsReplication->GetResimFrame();
+			const int32 ReplicationFrame = RewindData->GetResimFrame();
 
 #if DEBUG_NETWORK_PHYSICS
-				UE_LOG(LogTemp, Log, TEXT("CLIENT | PT | TriggerRewindIfNeeded_Internal | Replication Frame = %d"), ReplicationFrame);
+			UE_LOG(LogTemp, Log, TEXT("CLIENT | PT | TriggerRewindIfNeeded_Internal | Replication Frame = %d"), ReplicationFrame);
 #endif
-				ResimFrame = (ResimFrame == INDEX_NONE) ? ReplicationFrame : (ReplicationFrame == INDEX_NONE) ? ResimFrame : FMath::Min(ReplicationFrame, ResimFrame);
-				Scene->GetPhysicsReplication()->SetResimFrame(INDEX_NONE);
-			}
+			ResimFrame = (ResimFrame == INDEX_NONE) ? ReplicationFrame : (ReplicationFrame == INDEX_NONE) ? ResimFrame : FMath::Min(ReplicationFrame, ResimFrame);
+			RewindData->SetResimFrame(INDEX_NONE);
 		}
-	}
 
-	if(ResimFrame != INDEX_NONE)
-	{ 
-		FPhysScene* PhysScene = World->GetPhysicsScene();
-		if (ensureAlways(PhysScene))
+		if (ResimFrame != INDEX_NONE)
 		{
-			if (Chaos::FPhysicsSolver* Solver = PhysScene->GetSolver())
-			{
-				const int32 ValidFrame = Solver->GetRewindData()->FindValidResimFrame(ResimFrame);
+			const int32 ValidFrame = RewindData->FindValidResimFrame(ResimFrame);
 #if DEBUG_NETWORK_PHYSICS
-				UE_LOG(LogTemp, Log, TEXT("CLIENT | PT | TriggerRewindIfNeeded_Internal | Resim Frame = %d | Valid Frame = %d"), ResimFrame, ValidFrame);
+			UE_LOG(LogTemp, Log, TEXT("CLIENT | PT | TriggerRewindIfNeeded_Internal | Resim Frame = %d | Valid Frame = %d"), ResimFrame, ValidFrame);
 #endif
-				ResimFrame = ValidFrame;
-			}
+			ResimFrame = ValidFrame;
 		}
 	}
 	
diff --git a/Engine/Source/Runtime/Engine/Private/PhysicsEngine/PhysicsReplication.cpp b/Engine/Source/Runtime/Engine/Private/PhysicsEngine/PhysicsReplication.cpp
--- a/Engine/Source/Runtime/Engine/Private/PhysicsEngine/PhysicsReplication.cpp
+++ b/Engine/Source/Runtime/Engine/Private/PhysicsEngine/PhysicsReplication.cpp
@@ -10,11 +10,13 @@
 #include "DrawDebugHelpers.h"
 #include "Physics/Experimental/PhysScene_Chaos.h"
 #include "PhysicsEngine/PhysicsSettings.h"
+#include "Chaos/PhysicsObjectInternalInterface.h"
+#include "Chaos/PhysicsObject.h"
 #include "GameFramework/PlayerController.h"
 #include "GameFramework/PlayerState.h"
-#include "Components/SkeletalMeshComponent.h"
 #include "PhysicsProxy/SingleParticlePhysicsProxy.h"
 #include "PBDRigidsSolver.h"
+#include "Chaos/DebugDrawQueue.h"
 
 namespace CharacterMovementCVars
 {
@@ -81,224 +83,256 @@ namespace PhysicsReplicationCVars
 	int32 LogPhysicsReplicationHardSnaps = 0;
 	static FAutoConsoleVariableRef CVarLogPhysicsReplicationHardSnaps(TEXT("p.LogPhysicsReplicationHardSnaps"), LogPhysicsReplicationHardSnaps, TEXT(""));
 #endif
+	static int32 EnableDefaultReplication = 0;
+	static FAutoConsoleVariableRef CVarEnableDefaultReplication(TEXT("np2.EnableDefaultReplication"), EnableDefaultReplication, TEXT("Enable default replication in the networked physics prediction flow."));
+
 }
 
-struct FAsyncPhysicsRepCallbackData : public Chaos::FSimCallbackInput
-{
-	TArray<FAsyncPhysicsDesiredState> Buffer;
-	ErrorCorrectionData ErrorCorrection;
 
-	void Reset()
+FPhysicsReplication::FPhysicsReplication(FPhysScene* InPhysicsScene)
+	: PhysScene(InPhysicsScene)
+{
+	using namespace Chaos;
+	AsyncInput = nullptr;
+	PhysicsReplicationAsync = nullptr;
+	if (auto* Solver = PhysScene->GetSolver())
 	{
-		Buffer.Reset();
+		PhysicsReplicationAsync = Solver->CreateAndRegisterSimCallbackObject_External<FPhysicsReplicationAsync>();
+		PhysicsReplicationAsync->Setup(UPhysicsSettings::Get()->PhysicErrorCorrection);
 	}
-};
+}
 
-class FPhysicsReplicationAsyncCallback final : public Chaos::TSimCallbackObject<FAsyncPhysicsRepCallbackData>
+FPhysicsReplication::~FPhysicsReplication()
 {
-public:
-
-	virtual FName GetFNameForStatId() const override
+	if (PhysicsReplicationAsync)
 	{
-		const static FLazyName StaticName("FPhysicsReplicationAsyncCallback");
-		return StaticName;
+		if (auto* Solver = PhysScene->GetSolver())
+		{
+			Solver->UnregisterAndFreeSimCallbackObject_External(PhysicsReplicationAsync);
+		}
 	}
+}
 
-private:
 
-	virtual void OnPreSimulate_Internal() override
+void FPhysicsReplication::SetReplicatedTarget(UPrimitiveComponent* Component, FName BoneName, const FRigidBodyState& ReplicatedTarget, int32 ServerFrame)
+{
+	// If networked physics prediction is enabled, enforce the new physics replication flow via SetReplicatedTarget() using PhysicsObject instead of BodyInstance from BoneName.
+	AActor* Owner = Component->GetOwner();
+	if (UPhysicsSettings::Get()->PhysicsPrediction.bEnablePhysicsPrediction && Owner &&
+		(PhysicsReplicationCVars::EnableDefaultReplication || Owner->GetPhysicsReplicationMode() != EPhysicsReplicationMode::Default)) // For now, only opt in to the PhysicsObject flow if not using Default replication or if default is allowed via CVar.
 	{
-		FPhysicsReplication::ApplyAsyncDesiredState(GetDeltaTime_Internal(), GetConsumerInput_Internal());
+		const ENetRole OwnerRole = Owner->GetLocalRole();
+		const bool bIsSimulated = OwnerRole == ROLE_SimulatedProxy;
+		const bool bIsReplicatedAutonomous = OwnerRole == ROLE_AutonomousProxy && Component->bReplicatePhysicsToAutonomousProxy;
+		if (bIsSimulated || bIsReplicatedAutonomous)
+		{
+			Chaos::FPhysicsObjectHandle PhysicsObject = Component->GetPhysicsObjectByName(BoneName);
+			SetReplicatedTarget(PhysicsObject, ReplicatedTarget, ServerFrame, Owner->GetPhysicsReplicationMode());
+			return;
+		}
 	}
-};
 
-namespace
-{
-	void ComputeDeltas(const FVector& CurrentPos, const FQuat& CurrentQuat, const FVector& TargetPos, const FQuat& TargetQuat, FVector& OutLinDiff, float& OutLinDiffSize,
-		FVector& OutAngDiffAxis, float& OutAngDiff, float& OutAngDiffSize)
+	if (UWorld* OwningWorld = GetOwningWorld())
 	{
-		OutLinDiff = TargetPos - CurrentPos;
-		OutLinDiffSize = OutLinDiff.Size();
-		const FQuat InvCurrentQuat = CurrentQuat.Inverse();
-		const FQuat DeltaQuat = TargetQuat * InvCurrentQuat;
-		DeltaQuat.ToAxisAndAngle(OutAngDiffAxis, OutAngDiff);
-		OutAngDiff = FMath::RadiansToDegrees(FMath::UnwindRadians(OutAngDiff));
-		OutAngDiffSize = FMath::Abs(OutAngDiff);
+		//TODO: there's a faster way to compare this
+		TWeakObjectPtr<UPrimitiveComponent> TargetKey(Component);
+		FReplicatedPhysicsTarget* Target = ComponentToTargets_DEPRECATED.Find(TargetKey);
+		if (!Target)
+		{
+			// First time we add a target, set it's previous and correction
+			// positions to the target position to avoid math with uninitialized
+			// memory.
+			Target = &ComponentToTargets_DEPRECATED.Add(TargetKey);
+			Target->PrevPos = ReplicatedTarget.Position;
+			Target->PrevPosTarget = ReplicatedTarget.Position;
+		}
+
+		Target->ServerFrame = ServerFrame;
+		Target->TargetState = ReplicatedTarget;
+		Target->BoneName = BoneName;
+		Target->ArrivedTimeSeconds = OwningWorld->GetTimeSeconds();
+
+		ensure(!Target->PrevPos.ContainsNaN());
+		ensure(!Target->PrevPosTarget.ContainsNaN());
+		ensure(!Target->TargetState.Position.ContainsNaN());
+
+		OnSetReplicatedTarget(Component, BoneName, ReplicatedTarget, ServerFrame, *Target);
 	}
 }
 
-FPhysicsReplication::~FPhysicsReplication()
+void FPhysicsReplication::SetReplicatedTarget(Chaos::FPhysicsObject* PhysicsObject, const FRigidBodyState& ReplicatedTarget, int32 ServerFrame, EPhysicsReplicationMode ReplicationMode)
 {
-	if (AsyncCallback)
+	UWorld* OwningWorld = GetOwningWorld();
+	if (OwningWorld == nullptr)
 	{
-		if (auto* Solver = PhysScene->GetSolver())
-		{
-			Solver->UnregisterAndFreeSimCallbackObject_External(AsyncCallback);
-		}
+		return;
 	}
+
+	// TODO, Check if owning actor is ROLE_SimulatedProxy or ROLE_AutonomousProxy ?
+
+	FReplicatedPhysicsTarget Target;
+	Target.PhysicsObject = PhysicsObject;
+	Target.ReplicationMode = ReplicationMode;
+	Target.ServerFrame = ServerFrame;
+	Target.TargetState = ReplicatedTarget;
+	Target.ArrivedTimeSeconds = OwningWorld->GetTimeSeconds();
+
+	ensure(!Target.TargetState.Position.ContainsNaN());
+
+	ReplicatedTargetsQueue.Add(Target);
 }
 
-bool FPhysicsReplication::ApplyRigidBodyState(float DeltaSeconds, FBodyInstance* BI, FReplicatedPhysicsTarget& PhysicsTarget, const FRigidBodyErrorCorrection& ErrorCorrection, const float InPingSecondsOneWay, int32 LocalFrame, int32 NumPredictedFrames)
+void FPhysicsReplication::RemoveReplicatedTarget(UPrimitiveComponent* Component)
+{
+	ComponentToTargets_DEPRECATED.Remove(Component);
+}
+
+
+void FPhysicsReplication::Tick(float DeltaSeconds)
+{
+	OnTick(DeltaSeconds, ComponentToTargets_DEPRECATED);
+}
+
+void FPhysicsReplication::OnTick(float DeltaSeconds, TMap<TWeakObjectPtr<UPrimitiveComponent>, FReplicatedPhysicsTarget>& ComponentsToTargets)
 {
+	using namespace Chaos;
+
 	if (ShouldSkipPhysicsReplication())
 	{
-		return false;
+		return;
 	}
 
-	if (!BI->IsInstanceSimulatingPhysics())
+	int32 LocalFrameOffset = 0; // LocalFrame = ServerFrame + LocalFrameOffset;
+
+	if (FPhysicsSolverBase::IsNetworkPhysicsPredictionEnabled())
 	{
-		return false;
+		if (UWorld* World = GetOwningWorld())
+		{
+			if (World->GetNetMode() == NM_Client)
+			{
+				if (APlayerController* PlayerController = World->GetFirstPlayerController())
+				{
+					LocalFrameOffset = PlayerController->GetServerToLocalAsyncPhysicsTickOffset();
+				}
+			}
+		}
 	}
 
-	// LocalFrame the local frame number we should use to access past state in the rewind data
-	// NumPredictedFrames is how many frames (steps) ahead "now" is for the client compared to the latest data we've received from the server (use this to determine accurately how far ahead this object should extrapolate from its "physics target"
+	const FRigidBodyErrorCorrection& PhysicErrorCorrection = UPhysicsSettings::Get()->PhysicErrorCorrection;
+	if (PhysicsReplicationAsync)
+	{
+		PrepareAsyncData_External(PhysicErrorCorrection);
+	}
 
-	Chaos::FPhysicsSolver* Solver = PhysScene->GetSolver();
-	Chaos::FRewindData* RewindData = Solver->GetRewindData();
-	const FPhysicsActorHandle& Proxy = BI->GetPhysicsActorHandle();
-	EPhysicsReplicationMode PhysicsReplicationMode = BI->GetPhysicsReplicationMode();
+	// Get the ping between this PC & the server
+	const float LocalPing = GetLocalPing();
 
-	if (PhysicsReplicationMode != EPhysicsReplicationMode::Default && Proxy && Chaos::FPhysicsSolverBase::IsNetworkPhysicsPredictionEnabled())
+	// BodyInstance replication flow, deprecated
+	for (auto Itr = ComponentsToTargets.CreateIterator(); Itr; ++Itr)
 	{
-		if (RewindData && (LocalFrame <= RewindData->CurrentFrame()) && (LocalFrame >= RewindData->GetEarliestFrame_Internal()))
+		bool bRemoveItr = false;
+		if (UPrimitiveComponent* PrimComp = Itr.Key().Get())
 		{
-			static constexpr Chaos::FFrameAndPhase::EParticleHistoryPhase RewindPhase = Chaos::FFrameAndPhase::EParticleHistoryPhase::PostPushData;
-
-			if (PhysicsReplicationMode == EPhysicsReplicationMode::Resimulation && Chaos::FPhysicsSolverBase::IsPhysicsResimulationEnabled())
+			if (FBodyInstance* BI = PrimComp->GetBodyInstance(Itr.Value().BoneName))
 			{
-				if (LocalFrame < RewindData->CurrentFrame() && !RewindData->IsResim())
+				FReplicatedPhysicsTarget& PhysicsTarget = Itr.Value();
+				FRigidBodyState& UpdatedState = PhysicsTarget.TargetState;
+				bool bUpdated = false;
+				if (AActor* OwningActor = PrimComp->GetOwner())
 				{
-					if (UWorld* World = GetOwningWorld())
+					const ENetRole OwnerRole = OwningActor->GetLocalRole();
+					const bool bIsSimulated = OwnerRole == ROLE_SimulatedProxy;
+					const bool bIsReplicatedAutonomous = OwnerRole == ROLE_AutonomousProxy && PrimComp->bReplicatePhysicsToAutonomousProxy;
+					if (bIsSimulated || bIsReplicatedAutonomous)
 					{
-						if (World->GetNetMode() == NM_Client)
+						// Get the ping of this thing's owner. If nobody owns it,
+						// then it's server authoritative.
+						const float OwnerPing = GetOwnerPing(OwningActor, PhysicsTarget);
+
+						// Get the total ping - this approximates the time since the update was
+						// actually generated on the machine that is doing the authoritative sim.
+						// NOTE: We divide by 2 to approximate 1-way ping from 2-way ping.
+						const float PingSecondsOneWay = (LocalPing + OwnerPing) * 0.5f * 0.001f;
+
+						if (UpdatedState.Flags & ERigidBodyFlags::NeedsUpdate)
 						{
-							if (APlayerController* PlayerController = World->GetFirstPlayerController())
+							const int32 LocalFrame = PhysicsTarget.ServerFrame - LocalFrameOffset;
+							const bool bRestoredState = ApplyRigidBodyState(DeltaSeconds, BI, PhysicsTarget, PhysicErrorCorrection, PingSecondsOneWay, LocalFrame, 0);
+
+							// Need to update the component to match new position.
+							if (PhysicsReplicationCVars::SkipSkeletalRepOptimization == 0 || Cast<USkeletalMeshComponent>(PrimComp) == nullptr)	//simulated skeletal mesh does its own polling of physics results so we don't need to call this as it'll happen at the end of the physics sim
+							{
+								PrimComp->SyncComponentToRBPhysics();
+							}
+							if (bRestoredState)
 							{
-								FAsyncPhysicsTimestamp TimeStamp;
-								TimeStamp.LocalFrame = RewindData->CurrentFrame();
-
-								#if !(UE_BUILD_SHIPPING || UE_BUILD_TEST)
-								if (Chaos::FPhysicsSolverBase::CanDebugNetworkPhysicsPrediction())
-								{
-									UE_LOG(LogTemp, Log, TEXT("Appy Rigid body state at local frame %d with offset = %d"), LocalFrame, PlayerController->GetLocalToServerAsyncPhysicsTickOffset());
-								}
-								#endif
-
-								const float ResimErrorThreshold = Chaos::FPhysicsSolverBase::ResimulationErrorThreshold();
-
-								PlayerController->ExecuteAsyncPhysicsCommand(TimeStamp, BI->OwnerComponent.Get(), [this, Proxy, PhysicsTarget, RewindData, LocalFrame, ResimErrorThreshold]
-								{
-									if (RewindData->IsResim())
-										return;
-
-									auto PastState = RewindData->GetPastStateAtFrame(*Proxy->GetHandle_LowLevel(), LocalFrame, RewindPhase);
-
-									const float ErrorPosition = (PastState.X() - PhysicsTarget.TargetState.Position).Size();
-									const bool ShouldTriggerResim = ErrorPosition >= ResimErrorThreshold;
-									float ColorLerp = ShouldTriggerResim ? 1.0f : 0.0f;
-
-									#if !(UE_BUILD_SHIPPING || UE_BUILD_TEST)
-									if (Chaos::FPhysicsSolverBase::CanDebugNetworkPhysicsPrediction())
-									{
-										UE_LOG(LogTemp, Log, TEXT("Particle Position Error = %f | Should Trigger Resim = %s | Server Frame = %d | Client Frame = %d"), ErrorPosition, (ShouldTriggerResim ? TEXT("True") : TEXT("False")), PhysicsTarget.ServerFrame, LocalFrame);
-										UE_LOG(LogTemp, Log, TEXT("Particle Target Position = %s | Current Position = %s"), *PhysicsTarget.TargetState.Position.ToString(), *PastState.X().ToString());
-										UE_LOG(LogTemp, Log, TEXT("Particle Target Velocity = %s | Current Velocity = %s"), *PhysicsTarget.TargetState.LinVel.ToString(), *PastState.V().ToString());
-										UE_LOG(LogTemp, Log, TEXT("Particle Target Quaternion = %s | Current Quaternion = %s"), *PhysicsTarget.TargetState.Quaternion.ToString(), *PastState.R().ToString());
-										UE_LOG(LogTemp, Log, TEXT("Particle Target Omega = %s | Current Omega= %s"), *PhysicsTarget.TargetState.AngVel.ToString(), *PastState.W().ToString());
-
-										if (UWorld* OwningWorld = GetOwningWorld())
-										{
-											static constexpr float BoxSize = 5.0f;
-											const FColor DebugColor = FLinearColor::LerpUsingHSV(FLinearColor::Green, FLinearColor::Red, ColorLerp).ToFColor(false);
-
-											DrawDebugBox(OwningWorld, PhysicsTarget.TargetState.Position, FVector(BoxSize, BoxSize, BoxSize), PhysicsTarget.TargetState.Quaternion, FColor::Orange, true, CharacterMovementCVars::NetCorrectionLifetime, 0, 1.0f);
-											DrawDebugBox(OwningWorld, PastState.X(), FVector(6, 6, 6), PastState.R(), DebugColor, true, CharacterMovementCVars::NetCorrectionLifetime, 0, 1.0f);
-
-											DrawDebugDirectionalArrow(OwningWorld, PastState.X(), PhysicsTarget.TargetState.Position, 5.0f, FColor::Green, true, CharacterMovementCVars::NetCorrectionLifetime, 0, 0.5f);
-										}
-									}
-									#endif
-
-									const bool bShouldSleep = (PhysicsTarget.TargetState.Flags & ERigidBodyFlags::Sleeping) != 0;
-
-									RewindData->SetTargetStateAtFrame(*Proxy->GetHandle_LowLevel(), LocalFrame, RewindPhase,
-										PhysicsTarget.TargetState.Position, PhysicsTarget.TargetState.Quaternion, 
-										PhysicsTarget.TargetState.LinVel, PhysicsTarget.TargetState.AngVel, bShouldSleep);
-
-									if (ShouldTriggerResim)
-									{
-										if (Chaos::FPBDRigidsSolver* RigidSolver = Proxy->GetSolver<Chaos::FPBDRigidsSolver>())
-										{
-											RigidSolver->GetEvolution()->GetIslandManager().SetParticleResimFrame(Proxy->GetHandle_LowLevel(), LocalFrame);
-										}
-										ResimFrame = (ResimFrame == INDEX_NONE) ? LocalFrame : FMath::Min(ResimFrame, LocalFrame);
-									}
-								}, false);
+								bRemoveItr = true;
 							}
 						}
 					}
 				}
 			}
-			else // PhysicsReplicationMode == EPhysicsReplicationMode::PredictiveInterpolation
-			{
-				const float LerpFactor = Chaos::FPhysicsSolverBase::NetworkPhysicsInterpolationLerp();
-
-				FRigidBodyState CurrentState;
-				BI->GetRigidBodyState(CurrentState);
-
-				const float ExtrapolatedDeltaSeconds = (RewindData->CurrentFrame() - LocalFrame) * Solver->GetLastDt();
-
-				FVector LinDiff;
-				float LinDiffSize;
-				FVector AngDiffAxis;
-				float AngDiff;
-				float AngDiffSize;
-
-				ComputeDeltas(CurrentState.Position, CurrentState.Quaternion, PhysicsTarget.TargetState.Position, PhysicsTarget.TargetState.Quaternion, LinDiff, LinDiffSize, AngDiffAxis, AngDiff, AngDiffSize);
-
-				const FVector ExtrapolatedLinVel = FVector(PhysicsTarget.TargetState.LinVel + LinDiff / ExtrapolatedDeltaSeconds);
-				const FVector ExtrapolatedAngVel = FVector(PhysicsTarget.TargetState.AngVel + AngDiffAxis * AngDiff / ExtrapolatedDeltaSeconds);
-
-				const FVector TargetLinVel = FMath::Lerp(CurrentState.LinVel, ExtrapolatedLinVel, LerpFactor);
-				const FVector TargetAngVel = FMath::Lerp(CurrentState.AngVel, ExtrapolatedAngVel, LerpFactor);
-
-				FAsyncPhysicsDesiredState AsyncDesiredState;
-				AsyncDesiredState.WorldTM = FTransform(CurrentState.Quaternion, CurrentState.Position);
-				AsyncDesiredState.LinearVelocity = TargetLinVel;
-				AsyncDesiredState.AngularVelocity = TargetAngVel;
-				AsyncDesiredState.Proxy = static_cast<Chaos::FSingleParticlePhysicsProxy*>(Proxy);
-				AsyncDesiredState.ErrorCorrection = { 0.0f, 0.0f, 0.0f, 0.0f };
-				AsyncDesiredState.bShouldSleep = (PhysicsTarget.TargetState.Flags & ERigidBodyFlags::Sleeping) != 0;
-
-				CurAsyncData->Buffer.Add(AsyncDesiredState);
-				return false;
-			}
 		}
-		else if(RewindData && LocalFrame > 0)
+
+		if (bRemoveItr)
 		{
-			UE_LOG(LogPhysics, Warning, TEXT("FPhysicsReplication::ApplyRigidBodyState target frame (%d) out of rewind data bounds (%d,%d)"), LocalFrame,
-				RewindData->GetEarliestFrame_Internal(), RewindData->CurrentFrame());
+			OnTargetRestored(Itr.Key().Get(), Itr.Value());
+			Itr.RemoveCurrent();
 		}
-		return true;
 	}
-	else
+
+	// PhysicsObject replication flow
+	for (FReplicatedPhysicsTarget& PhysicsTarget : ReplicatedTargetsQueue)
 	{
-		// Call into the old ApplyRigidBodyState function for now,
-		// "new leash mode" should replace this.
-		// Note that old ApplyRigidBodyState is overridden in other projects, so consider backwards compat path
-		const float PingSecondsOneWay = RewindData == nullptr ? InPingSecondsOneWay : (Solver->GetLastDt() * (Solver->GetCurrentFrame() - LocalFrame)) * 0.5f;
-		return ApplyRigidBodyState(DeltaSeconds, BI, PhysicsTarget, ErrorCorrection, PingSecondsOneWay, nullptr);
+		if (PhysicsTarget.TargetState.Flags & ERigidBodyFlags::NeedsUpdate)
+		{
+			const float PingSecondsOneWay = LocalPing * 0.5f * 0.001f;
+
+			// Queue up the target state for async replication
+			FPhysicsRepAsyncInputData AsyncInputData;
+			AsyncInputData.TargetState = PhysicsTarget.TargetState;
+			AsyncInputData.Proxy = nullptr;
+			AsyncInputData.PhysicsObject = PhysicsTarget.PhysicsObject;
+			AsyncInputData.RepMode = PhysicsTarget.ReplicationMode;
+			AsyncInputData.ServerFrame = PhysicsTarget.ServerFrame;
+			AsyncInputData.FrameOffset = LocalFrameOffset;
+			AsyncInputData.LatencyOneWay = PingSecondsOneWay;
+
+			AsyncInput->InputData.Add(AsyncInputData);
+		}
 	}
+	ReplicatedTargetsQueue.Reset();
+
+	AsyncInput = nullptr;
 }
 
-bool FPhysicsReplication::ApplyRigidBodyState(float DeltaSeconds, FBodyInstance* BI, FReplicatedPhysicsTarget& PhysicsTarget, const FRigidBodyErrorCorrection& ErrorCorrection,
-	const float PingSecondsOneWay, bool* bDidHardSnap)
+namespace
 {
-	if (ShouldSkipPhysicsReplication())
+
+	// Helper to return the deltas between current and target Position and Rotation
+	void ComputeDeltas(const FVector& CurrentPos, const FQuat& CurrentQuat, const FVector& TargetPos, const FQuat& TargetQuat, FVector& OutLinDiff, float& OutLinDiffSize,
+		FVector& OutAngDiffAxis, float& OutAngDiff, float& OutAngDiffSize)
 	{
-		return false;
+		OutLinDiff = TargetPos - CurrentPos;
+		OutLinDiffSize = OutLinDiff.Size();
+		const FQuat InvCurrentQuat = CurrentQuat.Inverse();
+		const FQuat DeltaQuat = TargetQuat * InvCurrentQuat;
+		DeltaQuat.ToAxisAndAngle(OutAngDiffAxis, OutAngDiff);
+		OutAngDiff = FMath::RadiansToDegrees(FMath::UnwindRadians(OutAngDiff));
+		OutAngDiffSize = FMath::Abs(OutAngDiff);
 	}
+}
+
+bool FPhysicsReplication::ApplyRigidBodyState(float DeltaSeconds, FBodyInstance* BI, FReplicatedPhysicsTarget& PhysicsTarget, const FRigidBodyErrorCorrection& ErrorCorrection, const float InPingSecondsOneWay, int32 LocalFrame, int32 NumPredictedFrames)
+{
+	// Call into the old ApplyRigidBodyState function for now,
+	// Note that old ApplyRigidBodyState is overridden in other projects, so consider backwards compatible path
+	return ApplyRigidBodyState(DeltaSeconds, BI, PhysicsTarget, ErrorCorrection, InPingSecondsOneWay, nullptr);
+}
 
+bool FPhysicsReplication::ApplyRigidBodyState(float DeltaSeconds, FBodyInstance* BI, FReplicatedPhysicsTarget& PhysicsTarget, const FRigidBodyErrorCorrection& ErrorCorrection,
+	const float PingSecondsOneWay, bool* bDidHardSnap)
+{
 	if (!BI->IsInstanceSimulatingPhysics())
 	{
 		return false;
@@ -456,9 +490,9 @@ bool FPhysicsReplication::ApplyRigidBodyState(float DeltaSeconds, FBodyInstance*
 		{
 			PhysicsTarget.AccumulatedErrorSeconds = FMath::Max(PhysicsTarget.AccumulatedErrorSeconds - DeltaSeconds, 0.0f);
 		}
-		
+
 		// Hard snap if error accumulation or linear error is big enough, and clear the error accumulator.
-		const bool bHardSnap = 
+		const bool bHardSnap =
 			LinDiffSize > MaxLinearHardSnapDistance ||
 			PhysicsTarget.AccumulatedErrorSeconds > ErrorAccumulationSeconds ||
 			CharacterMovementCVars::AlwaysHardSnap;
@@ -500,14 +534,14 @@ bool FPhysicsReplication::ApplyRigidBodyState(float DeltaSeconds, FBodyInstance*
 		else
 		{
 			// Small enough error to interpolate
-			if (AsyncCallback == nullptr)	//sync case
+			if (PhysicsReplicationAsync == nullptr)	//sync case
 			{
 				const FVector NewLinVel = FVector(NewState.LinVel) + (LinDiff * LinearVelocityCoefficient * DeltaSeconds);
 				const FVector NewAngVel = FVector(NewState.AngVel) + (AngDiffAxis * AngDiff * AngularVelocityCoefficient * DeltaSeconds);
 
 				const FVector NewPos = FMath::Lerp(FVector(CurrentState.Position), FVector(TargetPos), PositionLerp);
 				const FQuat NewAng = FQuat::Slerp(CurrentState.Quaternion, TargetQuat, AngleLerp);
-				
+
 				BI->SetBodyTransform(FTransform(NewAng, NewPos), ETeleportType::ResetPhysics);
 				BI->SetLinearVelocity(NewLinVel, false);
 				BI->SetAngularVelocityInRadians(FMath::DegreesToRadians(NewAngVel), false);
@@ -515,18 +549,17 @@ bool FPhysicsReplication::ApplyRigidBodyState(float DeltaSeconds, FBodyInstance*
 			else
 			{
 				//If async is used, enqueue for callback
-				FAsyncPhysicsDesiredState AsyncDesiredState;
-				AsyncDesiredState.WorldTM = IdealWorldTM;
-				AsyncDesiredState.LinearVelocity = NewState.LinVel;
-				AsyncDesiredState.AngularVelocity = NewState.AngVel;
-				AsyncDesiredState.Proxy = static_cast<Chaos::FSingleParticlePhysicsProxy*>(BI->GetPhysicsActorHandle());
-				AsyncDesiredState.ErrorCorrection = { ErrorCorrection.LinearVelocityCoefficient, ErrorCorrection.AngularVelocityCoefficient, ErrorCorrection.PositionLerp, ErrorCorrection.AngleLerp };
-				AsyncDesiredState.bShouldSleep = bShouldSleep;
-
-				CurAsyncData->Buffer.Add(AsyncDesiredState);
+				FPhysicsRepAsyncInputData AsyncInputData;
+				AsyncInputData.TargetState = NewState;
+				AsyncInputData.TargetState.Position = IdealWorldTM.GetLocation();
+				AsyncInputData.TargetState.Quaternion = IdealWorldTM.GetRotation();
+				AsyncInputData.Proxy = static_cast<Chaos::FSingleParticlePhysicsProxy*>(BI->GetPhysicsActorHandle());
+				AsyncInputData.PhysicsObject = nullptr;
+				AsyncInputData.ErrorCorrection = { ErrorCorrection.LinearVelocityCoefficient, ErrorCorrection.AngularVelocityCoefficient, ErrorCorrection.PositionLerp, ErrorCorrection.AngleLerp };
+				AsyncInputData.LatencyOneWay = PingSeconds;
+
+				AsyncInput->InputData.Add(AsyncInputData);
 			}
-
-
 		}
 
 		// Should we show the async part?
@@ -551,7 +584,7 @@ bool FPhysicsReplication::ApplyRigidBodyState(float DeltaSeconds, FBodyInstance*
 	if (bShouldSleep)
 	{
 		// In the async case, we apply sleep state in ApplyAsyncDesiredState
-		if (AsyncCallback == nullptr)
+		if (PhysicsReplicationAsync == nullptr)
 		{
 			BI->PutInstanceToSleep();
 		}
@@ -563,295 +596,689 @@ bool FPhysicsReplication::ApplyRigidBodyState(float DeltaSeconds, FBodyInstance*
 	return bRestoredState;
 }
 
-UWorld* FPhysicsReplication::GetOwningWorld()
-{
-	return PhysScene ? PhysScene->GetOwningWorld() : nullptr;
-}
 
-const UWorld* FPhysicsReplication::GetOwningWorld() const
+void FPhysicsReplication::PrepareAsyncData_External(const FRigidBodyErrorCorrection& ErrorCorrection)
 {
-	return PhysScene ? PhysScene->GetOwningWorld() : nullptr;
+	//todo move this logic into a common function?
+	const float PositionLerp = CharacterMovementCVars::PositionLerp >= 0.0f ? CharacterMovementCVars::PositionLerp : ErrorCorrection.PositionLerp;
+	const float LinearVelocityCoefficient = CharacterMovementCVars::LinearVelocityCoefficient >= 0.0f ? CharacterMovementCVars::LinearVelocityCoefficient : ErrorCorrection.LinearVelocityCoefficient;
+	const float AngleLerp = CharacterMovementCVars::AngleLerp >= 0.0f ? CharacterMovementCVars::AngleLerp : ErrorCorrection.AngleLerp;
+	const float AngularVelocityCoefficient = CharacterMovementCVars::AngularVelocityCoefficient >= 0.0f ? CharacterMovementCVars::AngularVelocityCoefficient : ErrorCorrection.AngularVelocityCoefficient;
+
+	AsyncInput = PhysicsReplicationAsync->GetProducerInputData_External();
+	AsyncInput->ErrorCorrection.PositionLerp = PositionLerp;
+	AsyncInput->ErrorCorrection.AngleLerp = AngleLerp;
+	AsyncInput->ErrorCorrection.LinearVelocityCoefficient = LinearVelocityCoefficient;
+	AsyncInput->ErrorCorrection.AngularVelocityCoefficient = AngularVelocityCoefficient;
 }
 
-float FPhysicsReplication::GetLocalPing() const
+#pragma region AsyncPhysicsReplication
+void FPhysicsReplicationAsync::OnPreSimulate_Internal()
 {
-	if (const UWorld* World = GetOwningWorld())
+	if (const FPhysicsReplicationAsyncInput* AsyncInput = GetConsumerInput_Internal())
 	{
-		if (APlayerController* PlayerController = World->GetFirstPlayerController())
+
+		Chaos::FPBDRigidsSolver* RigidsSolver = static_cast<Chaos::FPBDRigidsSolver*>(GetSolver());
+		check(RigidsSolver);
+
+		// Early out if this is a resim frame
+		Chaos::FRewindData* RewindData = RigidsSolver->GetRewindData();
+		if (RewindData && RewindData->IsResim())
 		{
-			if (APlayerState* PlayerState = PlayerController->PlayerState)
-			{
-				return PlayerState->ExactPing;
-			}
+			return;
 		}
+
+		// Update async targets with target input
+		for (const FPhysicsRepAsyncInputData& Input : AsyncInput->InputData)
+		{
+			UpdateRewindDataTarget(Input);
+			UpdateAsyncTarget(Input);
+		}
+
+		ApplyTargetStatesAsync(GetDeltaTime_Internal(), AsyncInput->ErrorCorrection, AsyncInput->InputData);
 	}
-	return 0.0f;
 }
 
-float FPhysicsReplication::GetOwnerPing(const AActor* const Owner, const FReplicatedPhysicsTarget& Target) const
+void FPhysicsReplicationAsync::UpdateRewindDataTarget(const FPhysicsRepAsyncInputData& Input)
 {
-	//
-	// NOTE: At the moment, we have no real way to objectively access the ping of the
-	// authoritative simulation owner to the server, which is what this function
-	// claims to return.
-	//
-	// In order to actually use ping to extrapolate replication, we need to access
-	// it with something along the lines of the disabled code below.
-	//
-#if false
-	if (UPlayer* OwningPlayer = OwningActor->GetNetOwningPlayer())
+	if (Input.PhysicsObject == nullptr)
 	{
-		if (UWorld* World = GetOwningWorld())
-		{
-			if (APlayerController* PlayerController = OwningPlayer->GetPlayerController(World))
-			{
-				if (APlayerState* PlayerState = PlayerController->PlayerState)
-				{
-					return PlayerState->ExactPing;
-				}
-			}
-		}
+		return;
 	}
-#endif
 
-	return 0.0f;
+	Chaos::FPBDRigidsSolver* RigidsSolver = static_cast<Chaos::FPBDRigidsSolver*>(GetSolver());
+	if (RigidsSolver == nullptr)
+	{
+		return;
+	}
+
+	Chaos::FRewindData* RewindData = RigidsSolver->GetRewindData();
+	if (RewindData == nullptr)
+	{
+		return;
+	}
+
+	Chaos::FReadPhysicsObjectInterface_Internal Interface = Chaos::FPhysicsObjectInternalInterface::GetRead();
+	Chaos::FPBDRigidParticleHandle* Handle = Interface.GetRigidParticle(Input.PhysicsObject);
+
+	if (Handle != nullptr)
+	{
+		// Cache all target states inside RewindData
+		const int32 LocalFrame = Input.ServerFrame - Input.FrameOffset;
+		RewindData->SetTargetStateAtFrame(*Handle, LocalFrame, Chaos::FFrameAndPhase::EParticleHistoryPhase::PostPushData,
+			Input.TargetState.Position, Input.TargetState.Quaternion,
+			Input.TargetState.LinVel, Input.TargetState.AngVel, (Input.TargetState.Flags & ERigidBodyFlags::Sleeping));
+	}
 }
 
-void FPhysicsReplication::OnTick(float DeltaSeconds, TMap<TWeakObjectPtr<UPrimitiveComponent>, FReplicatedPhysicsTarget>& ComponentsToTargets)
+void FPhysicsReplicationAsync::UpdateAsyncTarget(const FPhysicsRepAsyncInputData& Input)
 {
-	using namespace Chaos;
+	if (Input.PhysicsObject == nullptr)
+	{
+		return;
+	}
+
+	FReplicatedPhysicsTargetAsync* Target = ObjectToTarget.Find(Input.PhysicsObject);
+	if (Target == nullptr)
+	{
+		// First time we add a target, set it's previous and correction
+		// positions to the target position to avoid math with uninitialized
+		// memory.
+		Target = &ObjectToTarget.Add(Input.PhysicsObject);
+		Target->PrevPos = Input.TargetState.Position;
+		Target->PrevPosTarget = Input.TargetState.Position;
+	}
+
+	if (Input.ServerFrame > Target->ServerFrame)
+	{
+		Target->PhysicsObject = Input.PhysicsObject;
+		Target->ServerFrame = Input.ServerFrame;
+		Target->TargetState = Input.TargetState;
+		Target->RepMode = Input.RepMode;
+		Target->FrameOffset = Input.FrameOffset;
+	}
+
+	/** Cache the latest ping time */
+	LatencyOneWay = Input.LatencyOneWay;
+}
 
-	int32 LocalFrameOffset = 0;		// LocalFrame = ServerFrame + LocalFrameOffset;
-	int32 NumPredictedFrames = 0;	// How many frames "ahead" of the server we are predicting. That is, how many frames are in flight between us and the server.
+void FPhysicsReplicationAsync::ApplyTargetStatesAsync(const float DeltaSeconds, const FPhysicsRepErrorCorrectionData& ErrorCorrection, const TArray<FPhysicsRepAsyncInputData>& InputData)
+{
+	using namespace Chaos;
 
-	if (UWorld* World = GetOwningWorld())
+	// Deprecated, BodyInstance flow
+	for (const FPhysicsRepAsyncInputData& Input : InputData)
 	{
-		if (World->GetNetMode() == NM_Client)
+		if (Input.Proxy != nullptr)
 		{
-			if (APlayerController* PlayerController = World->GetFirstPlayerController())
-			{
-				Chaos::FPhysicsSolver* Solver = PhysScene->GetSolver();
-				check(Solver);
+			Chaos::FSingleParticlePhysicsProxy* Proxy = Input.Proxy;
+			Chaos::FRigidBodyHandle_Internal* Handle = Proxy->GetPhysicsThreadAPI();
 
-				if (FPhysicsSolverBase::IsNetworkPhysicsPredictionEnabled())
-				{
-					LocalFrameOffset = PlayerController->GetServerToLocalAsyncPhysicsTickOffset();
-				}
-				//TODO: as we send physics updates down we need to record latest seen
-				//NumPredictedFrames = Solver->GetCurrentFrame() - ClientFrameInfo.LastProcessedInputFrame;
-			}
+			const FPhysicsRepErrorCorrectionData& UsedErrorCorrection = Input.ErrorCorrection.IsSet() ? Input.ErrorCorrection.GetValue() : ErrorCorrection;
+			DefaultReplication_DEPRECATED(Handle, Input, DeltaSeconds, UsedErrorCorrection);
 		}
 	}
 
-	const FRigidBodyErrorCorrection& PhysicErrorCorrection = UPhysicsSettings::Get()->PhysicErrorCorrection;
-	if (AsyncCallback)
-	{
-		PrepareAsyncData_External(PhysicErrorCorrection);
-	}
-	
-	// Get the ping between this PC & the server
-	const float LocalPing = GetLocalPing();
-	for (auto Itr = ComponentsToTargets.CreateIterator(); Itr; ++Itr)
+	// PhysicsObject flow
+	for (auto Itr = ObjectToTarget.CreateIterator(); Itr; ++Itr)
 	{
-		bool bRemoveItr = false;
-		if (UPrimitiveComponent* PrimComp = Itr.Key().Get())
+		bool bRemoveItr = true; // Remove current cached replication target unless replication logic tells us to store it for next tick
+
+		FReplicatedPhysicsTargetAsync& Target = Itr.Value();
+
+		if (Target.PhysicsObject != nullptr)
 		{
-			if (FBodyInstance* BI = PrimComp->GetBodyInstance(Itr.Value().BoneName))
+			Chaos::FWritePhysicsObjectInterface_Internal Interface = Chaos::FPhysicsObjectInternalInterface::GetWrite();
+			FPBDRigidParticleHandle* Handle = Interface.GetRigidParticle(Target.PhysicsObject);
+
+			if (Handle != nullptr)
 			{
-				FReplicatedPhysicsTarget& PhysicsTarget = Itr.Value();
-				FRigidBodyState& UpdatedState = PhysicsTarget.TargetState;
-				bool bUpdated = false;
-				if (AActor* OwningActor = PrimComp->GetOwner())
+				// TODO, Remove the resim option from project settings, we only need the physics prediction one now
+				EPhysicsReplicationMode RepMode = Target.RepMode;
+				if (!Chaos::FPBDRigidsSolver::IsPhysicsResimulationEnabled() && RepMode == EPhysicsReplicationMode::Resimulation)
 				{
-					const ENetRole OwnerRole = OwningActor->GetLocalRole();
-					const bool bIsSimulated = OwnerRole == ROLE_SimulatedProxy;
-					const bool bIsReplicatedAutonomous = OwnerRole == ROLE_AutonomousProxy && PrimComp->bReplicatePhysicsToAutonomousProxy;
-					if (bIsSimulated || bIsReplicatedAutonomous)
-					{
-						// Get the ping of this thing's owner. If nobody owns it,
-						// then it's server authoritative.
-						const float OwnerPing = GetOwnerPing(OwningActor, PhysicsTarget);
-
-						// Get the total ping - this approximates the time since the update was
-						// actually generated on the machine that is doing the authoritative sim.
-						// NOTE: We divide by 2 to approximate 1-way ping from 2-way ping.
-						const float PingSecondsOneWay = (LocalPing + OwnerPing) * 0.5f * 0.001f;
+					RepMode = EPhysicsReplicationMode::Default;
+				}
 
-						if (UpdatedState.Flags & ERigidBodyFlags::NeedsUpdate)
-						{
-							const int32 LocalFrame = PhysicsTarget.ServerFrame - LocalFrameOffset;
-							const bool bRestoredState = ApplyRigidBodyState(DeltaSeconds, BI, PhysicsTarget, PhysicErrorCorrection, PingSecondsOneWay, LocalFrame, NumPredictedFrames);
-							
-							// Need to update the component to match new position.
-							if (PhysicsReplicationCVars::SkipSkeletalRepOptimization == 0 || Cast<USkeletalMeshComponent>(PrimComp) == nullptr)	//simulated skeletal mesh does its own polling of physics results so we don't need to call this as it'll happen at the end of the physics sim
-							{
-								PrimComp->SyncComponentToRBPhysics();
-							}
-							if (bRestoredState)
-							{
-								bRemoveItr = true;
-							}
-						}
-					}
+				switch (RepMode)
+				{
+				case EPhysicsReplicationMode::Default:
+				{
+					bRemoveItr = DefaultReplication(Handle, Target, DeltaSeconds);
+					break;
+				}
+				case EPhysicsReplicationMode::PredictiveInterpolation:
+				{
+					bRemoveItr = PredictiveInterpolation(Handle, Target, DeltaSeconds);
+					break;
+				}
+				case EPhysicsReplicationMode::Resimulation:
+				{
+					bRemoveItr = ResimulationReplication(Handle, Target, DeltaSeconds);
+					break;
+				}
 				}
 			}
 		}
 
 		if (bRemoveItr)
 		{
-			OnTargetRestored(Itr.Key().Get(), Itr.Value());
 			Itr.RemoveCurrent();
 		}
 	}
-
-	CurAsyncData = nullptr;
 }
 
-bool FPhysicsReplication::ShouldSkipPhysicsReplication()
+//** Async function for legacy replication flow that goes partially through GT to then finishes in PT in this function. */
+void FPhysicsReplicationAsync::DefaultReplication_DEPRECATED(Chaos::FRigidBodyHandle_Internal* Handle, const FPhysicsRepAsyncInputData& State, const float DeltaSeconds, const FPhysicsRepErrorCorrectionData& ErrorCorrection)
 {
-	return (CharacterMovementCVars::SkipPhysicsReplication != 0);
+	if (Handle && Handle->CanTreatAsRigid())
+	{
+		const float LinearVelocityCoefficient = ErrorCorrection.LinearVelocityCoefficient;
+		const float AngularVelocityCoefficient = ErrorCorrection.AngularVelocityCoefficient;
+		const float PositionLerp = ErrorCorrection.PositionLerp;
+		const float AngleLerp = ErrorCorrection.AngleLerp;
+
+		const FVector TargetPos = State.TargetState.Position;
+		const FQuat TargetQuat = State.TargetState.Quaternion;
+
+		// Get Current state
+		FRigidBodyState CurrentState;
+		CurrentState.Position = Handle->X();
+		CurrentState.Quaternion = Handle->R();
+		CurrentState.AngVel = Handle->W();
+		CurrentState.LinVel = Handle->V();
+
+		FVector LinDiff;
+		float LinDiffSize;
+		FVector AngDiffAxis;
+		float AngDiff;
+		float AngDiffSize;
+		ComputeDeltas(CurrentState.Position, CurrentState.Quaternion, TargetPos, TargetQuat, LinDiff, LinDiffSize, AngDiffAxis, AngDiff, AngDiffSize);
+
+		const FVector NewLinVel = FVector(State.TargetState.LinVel) + (LinDiff * LinearVelocityCoefficient * DeltaSeconds);
+		const FVector NewAngVel = FVector(State.TargetState.AngVel) + (AngDiffAxis * AngDiff * AngularVelocityCoefficient * DeltaSeconds);
+
+		const FVector NewPos = FMath::Lerp(FVector(CurrentState.Position), TargetPos, PositionLerp);
+		const FQuat NewAng = FQuat::Slerp(CurrentState.Quaternion, TargetQuat, AngleLerp);
+
+		Handle->SetX(NewPos);
+		Handle->SetR(NewAng);
+		Handle->SetV(NewLinVel);
+		Handle->SetW(FMath::DegreesToRadians(NewAngVel));
+
+		if (State.TargetState.Flags & ERigidBodyFlags::Sleeping)
+		{
+			// don't allow kinematic to sleeping transition
+			if (Handle->ObjectState() != Chaos::EObjectStateType::Kinematic)
+			{
+				Chaos::FPBDRigidsSolver* RigidsSolver = static_cast<Chaos::FPBDRigidsSolver*>(GetSolver());
+				if (RigidsSolver)
+				{
+					RigidsSolver->GetEvolution()->SetParticleObjectState(Handle->GetProxy()->GetHandle_LowLevel()->CastToRigidParticle(), Chaos::EObjectStateType::Sleeping);	//todo: move object state into physics thread api
+				}
+			}
+		}
+	}
 }
 
-void FPhysicsReplication::Tick(float DeltaSeconds)
-{
-	OnTick(DeltaSeconds, ComponentToTargets);
-}
 
-FPhysicsReplication::FPhysicsReplication(FPhysScene* InPhysicsScene)
-	: PhysScene(InPhysicsScene)
+/** Default replication, run in simulation tick */
+bool FPhysicsReplicationAsync::DefaultReplication(Chaos::FPBDRigidParticleHandle* Handle, FReplicatedPhysicsTargetAsync& Target, const float DeltaSeconds)
 {
-	using namespace Chaos;
-	CurAsyncData = nullptr;
-	AsyncCallback = nullptr;
-	if (auto* Solver = PhysScene->GetSolver())
+	Chaos::FPBDRigidsSolver* RigidsSolver = static_cast<Chaos::FPBDRigidsSolver*>(GetSolver());
+	if (RigidsSolver == nullptr)
 	{
-		AsyncCallback = Solver->CreateAndRegisterSimCallbackObject_External<FPhysicsReplicationAsyncCallback>();
+		return true;
 	}
+
+	//
+	// NOTES:
+	//
+	// The operation of this method has changed since 4.18.
+	//
+	// When a new remote physics state is received, this method will
+	// be called on tick until the local state is within an adequate
+	// tolerance of the new state.
+	//
+	// The received state is extrapolated based on ping, by some
+	// adjustable amount.
+	//
+	// A correction velocity is added new state's velocity, and assigned
+	// to the body. The correction velocity scales with the positional
+	// difference, so without the interference of external forces, this
+	// will result in an exponentially decaying correction.
+	//
+	// Generally it is not needed and will interrupt smoothness of
+	// the replication, but stronger corrections can be obtained by
+	// adjusting position lerping.
+	//
+	// If progress is not being made towards equilibrium, due to some
+	// divergence in physics states between the owning and local sims,
+	// an error value is accumulated, representing the amount of time
+	// spent in an unresolvable state.
+	//
+	// Once the error value has exceeded some threshold (0.5 seconds
+	// by default), a hard snap to the target physics state is applied.
+	//
+
+
+	bool bRestoredState = true;
+	const FRigidBodyState NewState = Target.TargetState;
+	const float NewQuatSizeSqr = NewState.Quaternion.SizeSquared();
+
+
+	const FString ObjectName
+#if CHAOS_DEBUG_NAME
+		= Handle->DebugName() ? *Handle->DebugName() : FString(TEXT(""));
+#else
+		= FString(TEXT(""));
+#endif
+
+	// failure cases
+	if (Handle == nullptr)
+	{
+		UE_LOG(LogPhysics, Warning, TEXT("Trying to replicate rigid state for non-rigid particle. (%s)"), *ObjectName);
+		return bRestoredState;
+	}
+	else if (NewQuatSizeSqr < UE_KINDA_SMALL_NUMBER)
+	{
+		UE_LOG(LogPhysics, Warning, TEXT("Invalid zero quaternion set for body. (%s)"), *ObjectName);
+		return bRestoredState;
+	}
+	else if (FMath::Abs(NewQuatSizeSqr - 1.f) > UE_KINDA_SMALL_NUMBER)
+	{
+		UE_LOG(LogPhysics, Warning, TEXT("Quaternion (%f %f %f %f) with non-unit magnitude detected. (%s)"),
+			NewState.Quaternion.X, NewState.Quaternion.Y, NewState.Quaternion.Z, NewState.Quaternion.W, *ObjectName);
+		return bRestoredState;
+	}
+	// Grab configuration variables from engine config or from CVars if overriding is turned on.
+	const float NetPingExtrapolation = CharacterMovementCVars::NetPingExtrapolation >= 0.0f ? CharacterMovementCVars::NetPingExtrapolation : ErrorCorrectionDefault.PingExtrapolation;
+	const float NetPingLimit = CharacterMovementCVars::NetPingLimit > 0.0f ? CharacterMovementCVars::NetPingLimit : ErrorCorrectionDefault.PingLimit;
+	const float ErrorPerLinearDiff = CharacterMovementCVars::ErrorPerLinearDifference >= 0.0f ? CharacterMovementCVars::ErrorPerLinearDifference : ErrorCorrectionDefault.ErrorPerLinearDifference;
+	const float ErrorPerAngularDiff = CharacterMovementCVars::ErrorPerAngularDifference >= 0.0f ? CharacterMovementCVars::ErrorPerAngularDifference : ErrorCorrectionDefault.ErrorPerAngularDifference;
+	const float MaxRestoredStateError = CharacterMovementCVars::MaxRestoredStateError >= 0.0f ? CharacterMovementCVars::MaxRestoredStateError : ErrorCorrectionDefault.MaxRestoredStateError;
+	const float ErrorAccumulationSeconds = CharacterMovementCVars::ErrorAccumulationSeconds >= 0.0f ? CharacterMovementCVars::ErrorAccumulationSeconds : ErrorCorrectionDefault.ErrorAccumulationSeconds;
+	const float ErrorAccumulationDistanceSq = CharacterMovementCVars::ErrorAccumulationDistanceSq >= 0.0f ? CharacterMovementCVars::ErrorAccumulationDistanceSq : ErrorCorrectionDefault.ErrorAccumulationDistanceSq;
+	const float ErrorAccumulationSimilarity = CharacterMovementCVars::ErrorAccumulationSimilarity >= 0.0f ? CharacterMovementCVars::ErrorAccumulationSimilarity : ErrorCorrectionDefault.ErrorAccumulationSimilarity;
+	const float PositionLerp = CharacterMovementCVars::PositionLerp >= 0.0f ? CharacterMovementCVars::PositionLerp : ErrorCorrectionDefault.PositionLerp;
+	const float LinearVelocityCoefficient = CharacterMovementCVars::LinearVelocityCoefficient >= 0.0f ? CharacterMovementCVars::LinearVelocityCoefficient : ErrorCorrectionDefault.LinearVelocityCoefficient;
+	const float AngleLerp = CharacterMovementCVars::AngleLerp >= 0.0f ? CharacterMovementCVars::AngleLerp : ErrorCorrectionDefault.AngleLerp;
+	const float AngularVelocityCoefficient = CharacterMovementCVars::AngularVelocityCoefficient >= 0.0f ? CharacterMovementCVars::AngularVelocityCoefficient : ErrorCorrectionDefault.AngularVelocityCoefficient;
+	const float MaxLinearHardSnapDistance = CharacterMovementCVars::MaxLinearHardSnapDistance >= 0.f ? CharacterMovementCVars::MaxLinearHardSnapDistance : ErrorCorrectionDefault.MaxLinearHardSnapDistance;
+
+
+	// Get Current state
+	FRigidBodyState CurrentState;
+	CurrentState.Position = Handle->X();
+	CurrentState.Quaternion = Handle->R();
+	CurrentState.AngVel = Handle->W();
+	CurrentState.LinVel = Handle->V();
+
+
+	// Starting from the last known authoritative position, and
+	// extrapolate an approximation using the last known velocity
+	// and ping.
+	const float PingSeconds = FMath::Clamp(LatencyOneWay, 0.f, NetPingLimit);
+	const float ExtrapolationDeltaSeconds = PingSeconds * NetPingExtrapolation;
+	const FVector ExtrapolationDeltaPos = NewState.LinVel * ExtrapolationDeltaSeconds;
+	const FVector_NetQuantize100 TargetPos = NewState.Position + ExtrapolationDeltaPos;
+	float NewStateAngVel;
+	FVector NewStateAngVelAxis;
+	NewState.AngVel.FVector::ToDirectionAndLength(NewStateAngVelAxis, NewStateAngVel);
+	NewStateAngVel = FMath::DegreesToRadians(NewStateAngVel);
+	const FQuat ExtrapolationDeltaQuaternion = FQuat(NewStateAngVelAxis, NewStateAngVel * ExtrapolationDeltaSeconds);
+	FQuat TargetQuat = ExtrapolationDeltaQuaternion * NewState.Quaternion;
+
+
+	FVector LinDiff;
+	float LinDiffSize;
+	FVector AngDiffAxis;
+	float AngDiff;
+	float AngDiffSize;
+	ComputeDeltas(CurrentState.Position, CurrentState.Quaternion, TargetPos, TargetQuat, LinDiff, LinDiffSize, AngDiffAxis, AngDiff, AngDiffSize);
+
+	/////// ACCUMULATE ERROR IF NOT APPROACHING SOLUTION ///////
+
+	// Store sleeping state
+	const bool bShouldSleep = (NewState.Flags & ERigidBodyFlags::Sleeping) != 0;
+	const bool bWasAwake = !Handle->Sleeping();
+	const bool bAutoWake = false;
+
+	const float Error = (LinDiffSize * ErrorPerLinearDiff) + (AngDiffSize * ErrorPerAngularDiff);
+
+	bRestoredState = Error < MaxRestoredStateError;
+	if (bRestoredState)
+	{
+		Target.AccumulatedErrorSeconds = 0.0f;
+	}
+	else
+	{
+		//
+		// The heuristic for error accumulation here is:
+
+		// 1. Did the physics tick from the previous step fail to
+		//    move the body towards a resolved position?
+		// 2. Was the linear error in the same direction as the
+		//    previous frame?
+		// 3. Is the linear error large enough to accumulate error?
+		//
+		// If these conditions are met, then "error" time will accumulate.
+		// Once error has accumulated for a certain number of seconds,
+		// a hard-snap to the target will be performed.
+		//
+		// TODO: Rotation while moving linearly can still mess up this
+		// heuristic. We need to account for it.
+		//
+
+		// Project the change in position from the previous tick onto the
+		// linear error from the previous tick. This value roughly represents
+		// how much correction was performed over the previous physics tick.
+		const float PrevProgress = FVector::DotProduct(
+			FVector(CurrentState.Position) - Target.PrevPos,
+			(Target.PrevPosTarget - Target.PrevPos).GetSafeNormal());
+
+		// Project the current linear error onto the linear error from the
+		// previous tick. This value roughly represents how little the direction
+		// of the linear error state has changed, and how big the error is.
+		const float PrevSimilarity = FVector::DotProduct(
+			TargetPos - FVector(CurrentState.Position),
+			Target.PrevPosTarget - Target.PrevPos);
+
+		// If the conditions from the heuristic outlined above are met, accumulate
+		// error. Otherwise, reduce it.
+		if (PrevProgress < ErrorAccumulationDistanceSq &&
+			PrevSimilarity > ErrorAccumulationSimilarity)
+		{
+			Target.AccumulatedErrorSeconds += DeltaSeconds;
+		}
+		else
+		{
+			Target.AccumulatedErrorSeconds = FMath::Max(Target.AccumulatedErrorSeconds - DeltaSeconds, 0.0f);
+		}
+
+		// Hard snap if error accumulation or linear error is big enough, and clear the error accumulator.
+		const bool bHardSnap =
+			LinDiffSize > MaxLinearHardSnapDistance ||
+			Target.AccumulatedErrorSeconds > ErrorAccumulationSeconds ||
+			CharacterMovementCVars::AlwaysHardSnap;
+
+		if (bHardSnap)
+		{
+#if !UE_BUILD_SHIPPING
+			if (PhysicsReplicationCVars::LogPhysicsReplicationHardSnaps)
+			{
+				UE_LOG(LogTemp, Warning, TEXT("Simulated HARD SNAP - \nCurrent Pos - %s, Target Pos - %s\n CurrentState.LinVel - %s, New Lin Vel - %s\nTarget Extrapolation Delta - %s, Is Asleep - %d, Prev Progress - %f, Prev Similarity - %f"),
+					*CurrentState.Position.ToString(), *TargetPos.ToString(), *CurrentState.LinVel.ToString(), *NewState.LinVel.ToString(),
+					*ExtrapolationDeltaPos.ToString(), Handle->Sleeping(), PrevProgress, PrevSimilarity);
+
+				if (LinDiffSize > MaxLinearHardSnapDistance)
+				{
+					UE_LOG(LogTemp, Warning, TEXT("Hard snap due to linear difference error"));
+				}
+				else
+				{
+					UE_LOG(LogTemp, Warning, TEXT("Hard snap due to accumulated error"))
+				}
+			}
+#endif
+			// Too much error so just snap state here and be done with it
+			Target.AccumulatedErrorSeconds = 0.0f;
+			bRestoredState = true;
+			Handle->SetX(TargetPos);
+			Handle->SetR(TargetQuat);
+			Handle->SetV(NewState.LinVel);
+			Handle->SetW(FMath::DegreesToRadians(NewState.AngVel));
+		}
+		else
+		{
+			const FVector NewLinVel = FVector(Target.TargetState.LinVel) + (LinDiff * LinearVelocityCoefficient * DeltaSeconds);
+			const FVector NewAngVel = FVector(Target.TargetState.AngVel) + (AngDiffAxis * AngDiff * AngularVelocityCoefficient * DeltaSeconds);
+
+			const FVector NewPos = FMath::Lerp(FVector(CurrentState.Position), TargetPos, PositionLerp);
+			const FQuat NewAng = FQuat::Slerp(CurrentState.Quaternion, TargetQuat, AngleLerp);
+
+			Handle->SetX(NewPos);
+			Handle->SetR(NewAng);
+			Handle->SetV(NewLinVel);
+			Handle->SetW(FMath::DegreesToRadians(NewAngVel));
+		}
+	}
+
+	const Chaos::EObjectStateType ObjectState = Handle->ObjectState();
+	if (ObjectState != Chaos::EObjectStateType::Dynamic)
+	{
+		RigidsSolver->GetEvolution()->SetParticleObjectState(Handle, Chaos::EObjectStateType::Dynamic);
+	}
+
+	if ((Target.TargetState.Flags & ERigidBodyFlags::Sleeping) != 0)
+	{
+		// don't allow kinematic to sleeping transition
+		if (Handle->ObjectState() != Chaos::EObjectStateType::Kinematic)
+		{
+			RigidsSolver->GetEvolution()->SetParticleObjectState(Handle, Chaos::EObjectStateType::Sleeping);
+		}
+	}
+
+	Target.PrevPosTarget = TargetPos;
+	Target.PrevPos = FVector(CurrentState.Position);
+
+	return bRestoredState;
 }
 
-void FPhysicsReplication::PrepareAsyncData_External(const FRigidBodyErrorCorrection& ErrorCorrection)
+/** Interpolating towards replicated states from the server while predicting local physics */
+bool FPhysicsReplicationAsync::PredictiveInterpolation(Chaos::FPBDRigidParticleHandle* Handle, FReplicatedPhysicsTargetAsync& Target, const float DeltaSeconds)
 {
-	//todo move this logic into a common function?
-	const float PositionLerp = CharacterMovementCVars::PositionLerp >= 0.0f ? CharacterMovementCVars::PositionLerp : ErrorCorrection.PositionLerp;
-	const float LinearVelocityCoefficient = CharacterMovementCVars::LinearVelocityCoefficient >= 0.0f ? CharacterMovementCVars::LinearVelocityCoefficient : ErrorCorrection.LinearVelocityCoefficient;
-	const float AngleLerp = CharacterMovementCVars::AngleLerp >= 0.0f ? CharacterMovementCVars::AngleLerp : ErrorCorrection.AngleLerp;
-	const float AngularVelocityCoefficient = CharacterMovementCVars::AngularVelocityCoefficient >= 0.0f ? CharacterMovementCVars::AngularVelocityCoefficient : ErrorCorrection.AngularVelocityCoefficient;
+	Chaos::FPBDRigidsSolver* RigidsSolver = static_cast<Chaos::FPBDRigidsSolver*>(GetSolver());
+	if (RigidsSolver == nullptr)
+	{
+		return true;
+	}
+
+	const int32 LocalFrame = Target.ServerFrame - Target.FrameOffset;
+	const int32 NumPredictedFrames = RigidsSolver->GetCurrentFrame() - LocalFrame;
+	const float PredictedTime = DeltaSeconds * NumPredictedFrames;
+
+	const float PosCorrectionTime = 0.2f; // TODO.
+	const float InterpolationTime = PredictedTime; // TODO
+
+	// CurrentState
+	FRigidBodyState CurrentState;
+	CurrentState.Position = Handle->X();
+	CurrentState.Quaternion = Handle->R();
+	CurrentState.LinVel = Handle->V();
+	CurrentState.AngVel = Handle->W();
+
+	// NewState
+	const FVector TargetPos = Target.TargetState.Position;
+	const FQuat TargetRot = Target.TargetState.Quaternion;
+	const FVector TargetLinVel = Target.TargetState.LinVel;
+	const FVector TargetAngVel = Target.TargetState.AngVel;
+
+	// Get PosDiff
+	FVector PosDiff = TargetPos - CurrentState.Position;
+
+	// Convert PosDiff to PosDiffVelocity with a PosCorrectionTime value(0.2f)
+	FVector PosDiffVelocity = PosDiff / PosCorrectionTime;
+
+	// Get LinVelDiff by adding inverted CurrentState.LinVel to TargetLinVel
+	FVector LinVelDiff = -CurrentState.LinVel + TargetLinVel;
+
+	// Add PosDiffVelocity to LinVelDiff to get BlendedTargetVelocity
+	FVector BlendedTargetVelocity = LinVelDiff + PosDiffVelocity;
+
+	// Multiply BlendedTargetVelocity with(deltaTime / interpolationTime), clamp to 1 and add to CurrentState.LinVel to get BlendedTargetVelocityInterpolated
+	const float BlendStepAmount = FMath::Clamp(DeltaSeconds / InterpolationTime, 0.f, 1.f);
+	FVector RepLinVel = (BlendedTargetVelocity * BlendStepAmount) + CurrentState.LinVel;
 
-	CurAsyncData = AsyncCallback->GetProducerInputData_External();
-	CurAsyncData->ErrorCorrection.PositionLerp = PositionLerp;
-	CurAsyncData->ErrorCorrection.AngleLerp = AngleLerp;
-	CurAsyncData->ErrorCorrection.LinearVelocityCoefficient = LinearVelocityCoefficient;
-	CurAsyncData->ErrorCorrection.AngularVelocityCoefficient = AngularVelocityCoefficient;
+	// Apply velocity on particle
+	Handle->SetV(RepLinVel);
+	Handle->SetR(TargetRot);
+
+
+	// TODO, Angular Velocity
+
+
+	// TODO, Reconciliation
+
+
+	// Sleep
+	if ((Target.TargetState.Flags & ERigidBodyFlags::Sleeping) != 0)
+	{
+		// Don't allow kinematic to sleeping transition
+		if (!Handle->IsKinematic())
+		{
+			RigidsSolver->GetEvolution()->SetParticleObjectState(Handle, Chaos::EObjectStateType::Sleeping);
+		}
+	}
+
+
+	// Extrapolate target state
+	Target.TargetState.Position = Target.TargetState.Position + Target.TargetState.LinVel * DeltaSeconds;
+	// TODO, extrapolate rotation via angular velocity
+
+	return false;
 }
 
-void FPhysicsReplication::ApplyAsyncDesiredState(const float DeltaSeconds, const FAsyncPhysicsRepCallbackData* AsyncData)
+/** Compare states and trigger resimulation if needed */
+bool FPhysicsReplicationAsync::ResimulationReplication(Chaos::FPBDRigidParticleHandle* Handle, FReplicatedPhysicsTargetAsync& Target, const float DeltaSeconds)
 {
-	using namespace Chaos;
-	if (AsyncData)
+	Chaos::FPBDRigidsSolver* RigidsSolver = static_cast<Chaos::FPBDRigidsSolver*>(GetSolver());
+	if (RigidsSolver == nullptr)
+	{
+		return true;
+	}
+
+	Chaos::FRewindData* RewindData = RigidsSolver->GetRewindData();
+	if (RewindData == nullptr)
+	{
+		return true;
+	}
+
+	const int32 LocalFrame = Target.ServerFrame - Target.FrameOffset;
+
+	if (LocalFrame <= RewindData->CurrentFrame() && LocalFrame >= RewindData->GetEarliestFrame_Internal())
 	{
-		for (const FAsyncPhysicsDesiredState& State : AsyncData->Buffer)
+		static constexpr Chaos::FFrameAndPhase::EParticleHistoryPhase RewindPhase = Chaos::FFrameAndPhase::EParticleHistoryPhase::PostPushData;
+
+		FAsyncPhysicsTimestamp TimeStamp;
+		TimeStamp.LocalFrame = RewindData->CurrentFrame();
+
+		const float ResimErrorThreshold = Chaos::FPhysicsSolverBase::ResimulationErrorThreshold();
+
+		auto PastState = RewindData->GetPastStateAtFrame(*Handle, LocalFrame, RewindPhase);
+
+		const FVector ErrorOffset = (PastState.X() - Target.TargetState.Position);
+		const float ErrorDistance = ErrorOffset.Size();
+		const bool ShouldTriggerResim = ErrorDistance >= ResimErrorThreshold;
+		float ColorLerp = ShouldTriggerResim ? 1.0f : 0.0f;
+
+#if !(UE_BUILD_SHIPPING || UE_BUILD_TEST)
+
+		if (Chaos::FPhysicsSolverBase::CanDebugNetworkPhysicsPrediction())
 		{
-			float LinearVelocityCoefficient = AsyncData->ErrorCorrection.LinearVelocityCoefficient;
-			float AngularVelocityCoefficient = AsyncData->ErrorCorrection.AngularVelocityCoefficient;
-			float PositionLerp = AsyncData->ErrorCorrection.PositionLerp;
-			float AngleLerp = AsyncData->ErrorCorrection.AngleLerp;
-			if (State.ErrorCorrection.IsSet())
-			{
-				ErrorCorrectionData ECData = State.ErrorCorrection.GetValue();
-				LinearVelocityCoefficient = ECData.LinearVelocityCoefficient;
-				AngularVelocityCoefficient = ECData.AngularVelocityCoefficient;
-				PositionLerp = ECData.PositionLerp;
-				AngleLerp = ECData.AngleLerp;
-			}
-			//Proxy should exist because we are using latest and any pending deletes would have been enqueued after
-			Chaos::FSingleParticlePhysicsProxy* Proxy = State.Proxy;
-			auto* Handle = Proxy->GetPhysicsThreadAPI();
+			UE_LOG(LogTemp, Log, TEXT("Apply Rigid body state at local frame %d with offset = %d"), LocalFrame, Target.FrameOffset);
+			UE_LOG(LogTemp, Log, TEXT("Particle Position Error = %f | Should Trigger Resim = %s | Server Frame = %d | Client Frame = %d"), ErrorDistance, (ShouldTriggerResim ? TEXT("True") : TEXT("False")), Target.ServerFrame, LocalFrame);
+			UE_LOG(LogTemp, Log, TEXT("Particle Target Position = %s | Current Position = %s"), *Target.TargetState.Position.ToString(), *PastState.X().ToString());
+			UE_LOG(LogTemp, Log, TEXT("Particle Target Velocity = %s | Current Velocity = %s"), *Target.TargetState.LinVel.ToString(), *PastState.V().ToString());
+			UE_LOG(LogTemp, Log, TEXT("Particle Target Quaternion = %s | Current Quaternion = %s"), *Target.TargetState.Quaternion.ToString(), *PastState.R().ToString());
+			UE_LOG(LogTemp, Log, TEXT("Particle Target Omega = %s | Current Omega= %s"), *Target.TargetState.AngVel.ToString(), *PastState.W().ToString());
 
-			if (Handle && Handle->CanTreatAsRigid())
-			{
-				const FVector TargetPos = State.WorldTM.GetLocation();
-				const FQuat TargetQuat = State.WorldTM.GetRotation();
-
-				// Get Current state
-				FRigidBodyState CurrentState;
-				CurrentState.Position = Handle->X();
-				CurrentState.Quaternion = Handle->R();
-				CurrentState.AngVel = Handle->W();
-				CurrentState.LinVel = Handle->V();
-
-				FVector LinDiff;
-				float LinDiffSize;
-				FVector AngDiffAxis;
-				float AngDiff;
-				float AngDiffSize;
-				ComputeDeltas(CurrentState.Position, CurrentState.Quaternion, TargetPos, TargetQuat, LinDiff, LinDiffSize, AngDiffAxis, AngDiff, AngDiffSize);
-
-				const FVector NewLinVel = FVector(State.LinearVelocity) + (LinDiff * LinearVelocityCoefficient * DeltaSeconds);
-				const FVector NewAngVel = FVector(State.AngularVelocity) + (AngDiffAxis * AngDiff * AngularVelocityCoefficient * DeltaSeconds);
-
-				const FVector NewPos = FMath::Lerp(FVector(CurrentState.Position), TargetPos, PositionLerp);
-				const FQuat NewAng = FQuat::Slerp(CurrentState.Quaternion, TargetQuat, AngleLerp);
+			{ // DrawDebug
+				static constexpr float BoxSize = 5.0f;
+				const FColor DebugColor = FLinearColor::LerpUsingHSV(FLinearColor::Green, FLinearColor::Red, ColorLerp).ToFColor(false);
 
-				Handle->SetX(NewPos);
-				Handle->SetR(NewAng);
-				Handle->SetV(NewLinVel);
-				Handle->SetW(FMath::DegreesToRadians(NewAngVel));
+				Chaos::FDebugDrawQueue::GetInstance().DrawDebugBox(Target.TargetState.Position, FVector(BoxSize, BoxSize, BoxSize), Target.TargetState.Quaternion, FColor::Orange, true, CharacterMovementCVars::NetCorrectionLifetime, 0, 1.0f);
+				Chaos::FDebugDrawQueue::GetInstance().DrawDebugBox(PastState.X(), FVector(6, 6, 6), PastState.R(), DebugColor, true, CharacterMovementCVars::NetCorrectionLifetime, 0, 1.0f);
 
-				if (State.bShouldSleep)
-				{
-					// don't allow kinematic to sleeping transition
-					if (Handle->ObjectState() != EObjectStateType::Kinematic)
-					{
-						auto* Solver = Proxy->GetSolver<FPBDRigidsSolver>();
-						Solver->GetEvolution()->SetParticleObjectState(Proxy->GetHandle_LowLevel()->CastToRigidParticle(), EObjectStateType::Sleeping);	//todo: move object state into physics thread api
-					}
-				}
+				Chaos::FDebugDrawQueue::GetInstance().DrawDebugDirectionalArrow(PastState.X(), Target.TargetState.Position, 5.0f, FColor::Green, true, CharacterMovementCVars::NetCorrectionLifetime, 0, 0.5f);
 			}
 		}
+#endif
+
+		if (ShouldTriggerResim)
+		{
+			RigidsSolver->GetEvolution()->GetIslandManager().SetParticleResimFrame(Handle, LocalFrame);
+
+			int32 ResimFrame = RewindData->GetResimFrame();
+			ResimFrame = (ResimFrame == INDEX_NONE) ? LocalFrame : FMath::Min(ResimFrame, LocalFrame);
+			RewindData->SetResimFrame(ResimFrame);
+		}
+	}
+	else if (LocalFrame > 0)
+	{
+		UE_LOG(LogPhysics, Warning, TEXT("FPhysicsReplication::ApplyRigidBodyState target frame (%d) out of rewind data bounds (%d,%d)"), LocalFrame,
+			RewindData->GetEarliestFrame_Internal(), RewindData->CurrentFrame());
 	}
+
+	return true;
+}
+
+FName FPhysicsReplicationAsync::GetFNameForStatId() const
+{
+	const static FLazyName StaticName("FPhysicsReplicationAsyncCallback");
+	return StaticName;
 }
+#pragma endregion // AsyncPhysicsReplication
+
 
-int32 FPhysicsReplication::GetResimFrame() const
+
+
+bool FPhysicsReplication::ShouldSkipPhysicsReplication()
 {
-	return ResimFrame;
+	return (CharacterMovementCVars::SkipPhysicsReplication != 0);
 }
 
-void FPhysicsReplication::SetResimFrame(const int32 InResimFrame)
+UWorld* FPhysicsReplication::GetOwningWorld()
 {
-	ResimFrame = InResimFrame;
+	return PhysScene ? PhysScene->GetOwningWorld() : nullptr;
 }
 
-void FPhysicsReplication::SetReplicatedTarget(UPrimitiveComponent* Component, FName BoneName, const FRigidBodyState& ReplicatedTarget, int32 ServerFrame)
+const UWorld* FPhysicsReplication::GetOwningWorld() const
 {
-	if (UWorld* OwningWorld = GetOwningWorld())
+	return PhysScene ? PhysScene->GetOwningWorld() : nullptr;
+}
+
+float FPhysicsReplication::GetLocalPing() const
+{
+	if (const UWorld* World = GetOwningWorld())
 	{
-		//TODO: there's a faster way to compare this
-		TWeakObjectPtr<UPrimitiveComponent> TargetKey(Component);
-		FReplicatedPhysicsTarget* Target = ComponentToTargets.Find(TargetKey);
-		if (!Target)
+		if (APlayerController* PlayerController = World->GetFirstPlayerController())
 		{
-			// First time we add a target, set it's previous and correction
-			// positions to the target position to avoid math with uninitialized
-			// memory.
-			Target = &ComponentToTargets.Add(TargetKey);
-			Target->PrevPos = ReplicatedTarget.Position;
-			Target->PrevPosTarget = ReplicatedTarget.Position;
+			if (APlayerState* PlayerState = PlayerController->PlayerState)
+			{
+				return PlayerState->ExactPing;
+			}
 		}
 
-		Target->ServerFrame = ServerFrame;
-		Target->TargetState = ReplicatedTarget;
-		Target->BoneName = BoneName;
-		Target->ArrivedTimeSeconds = OwningWorld->GetTimeSeconds();
-		
-		ensure(!Target->PrevPos.ContainsNaN());
-		ensure(!Target->PrevPosTarget.ContainsNaN());
-		ensure(!Target->TargetState.Position.ContainsNaN());
-
-		OnSetReplicatedTarget(Component, BoneName, ReplicatedTarget, ServerFrame, *Target);
 	}
+	return 0.0f;
 }
 
-void FPhysicsReplication::RemoveReplicatedTarget(UPrimitiveComponent* Component)
+float FPhysicsReplication::GetOwnerPing(const AActor* const Owner, const FReplicatedPhysicsTarget& Target) const
 {
-	ComponentToTargets.Remove(Component);
-}
+	//
+	// NOTE: At the moment, we have no real way to objectively access the ping of the
+	// authoritative simulation owner to the server, which is what this function
+	// claims to return.
+	//
+	// In order to actually use ping to extrapolate replication, we need to access
+	// it with something along the lines of the disabled code below.
+	//
+#if false
+	if (UPlayer* OwningPlayer = OwningActor->GetNetOwningPlayer())
+	{
+		if (UWorld* World = GetOwningWorld())
+		{
+			if (APlayerController* PlayerController = OwningPlayer->GetPlayerController(World))
+			{
+				if (APlayerState* PlayerState = PlayerController->PlayerState)
+				{
+					return PlayerState->ExactPing;
+				}
+			}
+		}
+	}
+#endif
+
+	return 0.0f;
+}
\ No newline at end of file
diff --git a/Engine/Source/Runtime/Engine/Private/StaticMeshActor.cpp b/Engine/Source/Runtime/Engine/Private/StaticMeshActor.cpp
--- a/Engine/Source/Runtime/Engine/Private/StaticMeshActor.cpp
+++ b/Engine/Source/Runtime/Engine/Private/StaticMeshActor.cpp
@@ -47,7 +47,6 @@ void AStaticMeshActor::BeginPlay()
 	// we need to do this here. 
 	//
 	// This is a short term fix until we find a better play for SetReplicates to be called in AActor.
-
 	if (GetLocalRole() == ROLE_Authority && bStaticMeshReplicateMovement)
 	{
 		bReplicates = false;
@@ -55,6 +54,11 @@ void AStaticMeshActor::BeginPlay()
 		SetReplicates(true);
 	}	
 
+	if (StaticMeshComponent && StaticMeshComponent->BodyInstance.bSimulatePhysics)
+	{
+		SetPhysicsReplicationMode(StaticMeshPhysicsReplicationMode);
+	}
+
 	Super::BeginPlay();
 }
 
diff --git a/Engine/Source/Runtime/Engine/Public/PhysicsReplication.h b/Engine/Source/Runtime/Engine/Public/PhysicsReplication.h
--- a/Engine/Source/Runtime/Engine/Public/PhysicsReplication.h
+++ b/Engine/Source/Runtime/Engine/Public/PhysicsReplication.h
@@ -13,11 +13,9 @@
 #include "Physics/PhysicsInterfaceDeclares.h"
 #include "PhysicsProxy/SingleParticlePhysicsProxyFwd.h"
 #include "Chaos/Particles.h"
-
-namespace Chaos
-{
-	struct FSimCallbackInput;
-}
+#include "Chaos/PhysicsObject.h"
+#include "Chaos/SimCallbackObject.h"
+#include "Physics/PhysicsInterfaceUtils.h"
 
 namespace CharacterMovementCVars
 {
@@ -49,12 +47,120 @@ namespace PhysicsReplicationCVars
 
 class FPhysScene_PhysX;
 
+
+// -------- Async Flow ------->
+
+struct FPhysicsRepErrorCorrectionData
+{
+	float LinearVelocityCoefficient;
+	float AngularVelocityCoefficient;
+	float PositionLerp;
+	float AngleLerp;
+};
+
+/** Final computed desired state passed into the physics sim */
+struct FPhysicsRepAsyncInputData
+{
+	FRigidBodyState TargetState;
+	Chaos::FSingleParticlePhysicsProxy* Proxy; // Used for legacy (BodyInstance) flow
+	Chaos::FPhysicsObject* PhysicsObject;
+	TOptional<FPhysicsRepErrorCorrectionData> ErrorCorrection;
+	EPhysicsReplicationMode RepMode;
+	int32 ServerFrame;
+	int32 FrameOffset;
+	float LatencyOneWay;
+
+	FPhysicsRepAsyncInputData()
+		: Proxy(nullptr)
+		, PhysicsObject(nullptr)
+		, RepMode(EPhysicsReplicationMode::Default)
+	{};
+};
+
+struct FPhysicsReplicationAsyncInput : public Chaos::FSimCallbackInput
+{
+	FPhysicsRepErrorCorrectionData ErrorCorrection;
+	TArray<FPhysicsRepAsyncInputData> InputData;
+	void Reset()
+	{
+		InputData.Reset();
+	}
+};
+
+
+struct FReplicatedPhysicsTargetAsync
+{
+	FReplicatedPhysicsTargetAsync()
+		: AccumulatedErrorSeconds(0.0f)
+		, ServerFrame(0)
+		, PhysicsObject(nullptr)
+	{ }
+
+	/** The target state replicated by server */
+	FRigidBodyState TargetState;
+
+	/** Physics sync error accumulation */
+	float AccumulatedErrorSeconds;
+
+	/** Correction values from previous update */
+	FVector PrevPosTarget;
+	FVector PrevPos;
+
+	/** ServerFrame this target was replicated on (must be converted to local frame prior to client-side use) */
+	int32 ServerFrame;
+
+	/** Index of physics object on component */
+	Chaos::FPhysicsObject* PhysicsObject;
+
+	/** The frame offset between local client and server */
+	int32 FrameOffset;
+
+	/** The replication mode this PhysicsObject should use */
+	EPhysicsReplicationMode RepMode;
+};
+
+class FPhysicsReplicationAsync : public Chaos::TSimCallbackObject<
+	FPhysicsReplicationAsyncInput,
+	Chaos::FSimCallbackNoOutput,
+	Chaos::ESimCallbackOptions::Presimulate>
+{
+	virtual FName GetFNameForStatId() const override;
+	virtual void OnPreSimulate_Internal() override;
+	virtual void ApplyTargetStatesAsync(const float DeltaSeconds, const FPhysicsRepErrorCorrectionData& ErrorCorrection, const TArray<FPhysicsRepAsyncInputData>& TargetStates);
+
+	// Replication functions
+	virtual void DefaultReplication_DEPRECATED(Chaos::FRigidBodyHandle_Internal* Handle, const FPhysicsRepAsyncInputData& State, const float DeltaSeconds, const FPhysicsRepErrorCorrectionData& ErrorCorrection);
+	virtual bool DefaultReplication(Chaos::FPBDRigidParticleHandle* Handle, FReplicatedPhysicsTargetAsync& Target, const float DeltaSeconds);
+	virtual bool PredictiveInterpolation(Chaos::FPBDRigidParticleHandle* Handle, FReplicatedPhysicsTargetAsync& Target, const float DeltaSeconds);
+	virtual bool ResimulationReplication(Chaos::FPBDRigidParticleHandle* Handle, FReplicatedPhysicsTargetAsync& Target, const float DeltaSeconds);
+
+private:
+	float LatencyOneWay;
+	FRigidBodyErrorCorrection ErrorCorrectionDefault;
+	TMap<Chaos::FPhysicsObject*, FReplicatedPhysicsTargetAsync> ObjectToTarget;
+
+private:
+	void UpdateAsyncTarget(const FPhysicsRepAsyncInputData& Input);
+	void UpdateRewindDataTarget(const FPhysicsRepAsyncInputData& Input);
+
+public:
+	void Setup(FRigidBodyErrorCorrection ErrorCorrection)
+	{
+		ErrorCorrectionDefault = ErrorCorrection;
+	}
+};
+
+// <-------- Async Flow --------
+
+
+
 struct FReplicatedPhysicsTarget
 {
-	FReplicatedPhysicsTarget() :
-		ArrivedTimeSeconds(0.0f),
-		AccumulatedErrorSeconds(0.0f),
-		ServerFrame(0)
+	FReplicatedPhysicsTarget()
+		: ArrivedTimeSeconds(0.0f)
+		, AccumulatedErrorSeconds(0.0f)
+		, ServerFrame(0)
+		, PhysicsObject(nullptr)
 	{ }
 
 	/** The target state replicated by server */
@@ -76,37 +182,17 @@ struct FReplicatedPhysicsTarget
 	/** ServerFrame this target was replicated on (must be converted to local frame prior to client-side use) */
 	int32 ServerFrame;
 
+	/** Index of physics object on component */
+	Chaos::FPhysicsObject* PhysicsObject;
+
+	/** The replication mode the target should be used with */
+	EPhysicsReplicationMode ReplicationMode;
+
 #if !UE_BUILD_SHIPPING
 	FDebugFloatHistory ErrorHistory;
 #endif
 };
 
-struct ErrorCorrectionData
-{
-	float LinearVelocityCoefficient;
-	float AngularVelocityCoefficient;
-	float PositionLerp;
-	float AngleLerp;
-};
-
-/** Final computed desired state passed into the physics sim */
-struct FAsyncPhysicsDesiredState
-{
-	FTransform WorldTM;
-	FVector LinearVelocity;
-	FVector AngularVelocity;
-	Chaos::FSingleParticlePhysicsProxy* Proxy;
-	TOptional<ErrorCorrectionData> ErrorCorrection;
-	bool bShouldSleep;
-};
-
-struct FBodyInstance;
-struct FRigidBodyErrorCorrection;
-class UWorld;
-class UPrimitiveComponent;
-class FPhysicsReplicationAsyncCallback;
-struct FAsyncPhysicsRepCallbackData;
-
 class ENGINE_API FPhysicsReplication : public IPhysicsReplication
 {
 public:
@@ -123,16 +209,13 @@ public:
 	UE_DEPRECATED(5.1, "SetReplicatedTarget now takes the ServerFrame.  Please update calls and overloads.")
 	virtual void SetReplicatedTarget(UPrimitiveComponent* Component, FName BoneName, const FRigidBodyState& ReplicatedTarget) { SetReplicatedTarget(Component, BoneName, ReplicatedTarget, 0); }
 	virtual void SetReplicatedTarget(UPrimitiveComponent* Component, FName BoneName, const FRigidBodyState& ReplicatedTarget, int32 ServerFrame) override;
+private:
+	void SetReplicatedTarget(Chaos::FPhysicsObject* PhysicsObject, const FRigidBodyState& ReplicatedTarget, int32 ServerFrame, EPhysicsReplicationMode ReplicationMode = EPhysicsReplicationMode::Default);
 
+public:
 	/** Remove the replicated target*/
 	virtual void RemoveReplicatedTarget(UPrimitiveComponent* Component) override;
 
-	/** Get the resim frame (min server frame from the targets) */
-	virtual int32 GetResimFrame() const override;
-
-	/** Set the resim frame for replication */
-	virtual void SetResimFrame(const int32 InResimFrame) override;
-
 protected:
 
 	/** Update the physics body state given a set of replicated targets */
@@ -155,18 +238,13 @@ private:
 	/** Get the ping from  */
 	float GetOwnerPing(const AActor* const Owner, const FReplicatedPhysicsTarget& Target) const;
 
-	static void ApplyAsyncDesiredState(float DeltaSeconds, const FAsyncPhysicsRepCallbackData* Input);
-
 private:
-	TMap<TWeakObjectPtr<UPrimitiveComponent>, FReplicatedPhysicsTarget> ComponentToTargets;
+	TMap<TWeakObjectPtr<UPrimitiveComponent>, FReplicatedPhysicsTarget> ComponentToTargets_DEPRECATED; // This collection is keeping the legacy flow working until fully deprecated in a future release
+	TArray<FReplicatedPhysicsTarget> ReplicatedTargetsQueue;
 	FPhysScene* PhysScene;
 
-	FPhysicsReplicationAsyncCallback* AsyncCallback;
-	
-	void PrepareAsyncData_External(const FRigidBodyErrorCorrection& ErrorCorrection);	//prepare async data for writing. Call on external thread (i.e. game thread)
-	FAsyncPhysicsRepCallbackData* CurAsyncData;	//async data being written into before we push into callback
-	friend FPhysicsReplicationAsyncCallback;
-
-	int32 ResimFrame = INDEX_NONE;
+	FPhysicsReplicationAsync* PhysicsReplicationAsync;
+	FPhysicsReplicationAsyncInput* AsyncInput;	//async data being written into before we push into callback
 
+	void PrepareAsyncData_External(const FRigidBodyErrorCorrection& ErrorCorrection);	//prepare async data for writing. Call on external thread (i.e. game thread)
 };
\ No newline at end of file
diff --git a/Engine/Source/Runtime/Engine/Public/PhysicsReplicationInterface.h b/Engine/Source/Runtime/Engine/Public/PhysicsReplicationInterface.h
--- a/Engine/Source/Runtime/Engine/Public/PhysicsReplicationInterface.h
+++ b/Engine/Source/Runtime/Engine/Public/PhysicsReplicationInterface.h
@@ -17,8 +17,4 @@ public:
 	virtual void SetReplicatedTarget(UPrimitiveComponent* Component, FName BoneName, const FRigidBodyState& ReplicatedTarget, int32 ServerFrame) = 0;
 
 	virtual void RemoveReplicatedTarget(UPrimitiveComponent* Component) = 0;
-
-	virtual void SetResimFrame(const int32 InResimFrame) { }
-
-	virtual int32 GetResimFrame() const { return INDEX_NONE; }
 };
diff --git a/Engine/Source/Runtime/Experimental/Chaos/Private/RewindData.cpp b/Engine/Source/Runtime/Experimental/Chaos/Private/RewindData.cpp
--- a/Engine/Source/Runtime/Experimental/Chaos/Private/RewindData.cpp
+++ b/Engine/Source/Runtime/Experimental/Chaos/Private/RewindData.cpp
@@ -827,17 +827,17 @@ void FRewindData::ExtendHistoryWithFrame(const int32 Frame)
 	FramesSaved = FMath::Max(CurFrame - Frame+1, FramesSaved);
 }
 
-int32 FRewindData::FindValidResimFrame(const int32 ResimFrame)
+int32 FRewindData::FindValidResimFrame(const int32 RequestedFrame)
 {
 	int32 ValidFrame = INDEX_NONE;
-	if(ResimFrame > 0)
+	if(RequestedFrame > 0)
 	{
 		EnsureIsInPhysicsThreadContext();
 
 		// First frame of the history datas
 		const int32 EarliestFrame = FMath::Max(GetEarliestFrame_Internal(), 0);
 	
-		for(ValidFrame = ResimFrame; ValidFrame >= EarliestFrame; --ValidFrame)
+		for(ValidFrame = RequestedFrame; ValidFrame >= EarliestFrame; --ValidFrame)
 		{
 			bool bHasTargetHistory = true;
 			for (FDirtyParticleInfo& DirtyParticleInfo : DirtyParticles)
diff --git a/Engine/Source/Runtime/Experimental/Chaos/Public/RewindData.h b/Engine/Source/Runtime/Experimental/Chaos/Public/RewindData.h
--- a/Engine/Source/Runtime/Experimental/Chaos/Public/RewindData.h
+++ b/Engine/Source/Runtime/Experimental/Chaos/Public/RewindData.h
@@ -1427,7 +1427,11 @@ public:
 	const FPBDRigidsSolver* GetSolver() const { return Solver; }
 
 	/** Find the first previous valid frame having received physics target from the server */
-	int32 CHAOS_API FindValidResimFrame(const int32 ResimFrame);
+	int32 CHAOS_API FindValidResimFrame(const int32 RequestedFrame);
+
+	/** Get and set the frame we resimulate from */
+	int32 GetResimFrame() { return ResimFrame; }
+	void SetResimFrame(int32 Frame) { ResimFrame = Frame; }
 
 private:
 
@@ -1605,6 +1609,7 @@ private:
 	int32 DataIdxOffset;
 	bool bNeedsSave;	//Indicates that some data is pointing at head and requires saving before a rewind
 	bool bResimOptimization;
+	int32 ResimFrame = INDEX_NONE;
 
 	template <typename TObj>
 	bool IsResimAndInSync(const TObj& Handle) const { return IsResim() && Handle.SyncState() == ESyncState::InSync; }
