Fixed ensure in LevelStreaming when LoadStreamingLevel and UnloadStreamingLevel are called one after the other on the same level

diff --git a/Engine/Source/Runtime/Engine/Private/LevelStreaming.cpp b/Engine/Source/Runtime/Engine/Private/LevelStreaming.cpp
--- a/Engine/Source/Runtime/Engine/Private/LevelStreaming.cpp
+++ b/Engine/Source/Runtime/Engine/Private/LevelStreaming.cpp
@@ -1048,6 +1048,14 @@ void ULevelStreaming::UpdateStreamingState(bool& bOutUpdateAgain, bool& bOutRede
 				bOutRedetermineTarget = true;
 				break;
 
+			case ELevelStreamingTargetState::Unloaded:
+				// This is to support the case where a request to load is followed by another to unload the same level
+				// We set bOutRedetermineTarget to true so that FStreamingLevelPrivateAccessor::UpdateTargetState 
+				// gets called by UWorld::UpdateLevelStreaming which will return false so that the streaming level
+				// gets removed from StreamingLevelsToConsider.
+				bOutRedetermineTarget = true;
+				break;
+
 			default:
 				ensure(false);
 		}
