Add a "Force Reset" flag to device properties that gets set when PIE ends. Whenever we activate a device property, keep track of its class and what device it was applied to. Then, when PIE stops, spawn a new instance of those properties and reset the input device.

This fixes an small bug where the light color or trigger effect could persist between PIE sessions.

diff --git a/Engine/Source/Runtime/Engine/Classes/GameFramework/InputDeviceProperties.h b/Engine/Source/Runtime/Engine/Classes/GameFramework/InputDeviceProperties.h
--- a/Engine/Source/Runtime/Engine/Classes/GameFramework/InputDeviceProperties.h
+++ b/Engine/Source/Runtime/Engine/Classes/GameFramework/InputDeviceProperties.h
@@ -67,7 +67,7 @@ protected:
 	* @param PlatformUser		The platform user that should receive this device property change
 	*/
 	UFUNCTION(BlueprintNativeEvent, Category = "InputDevice")
-	void ResetDeviceProperty(const FPlatformUserId PlatformUser, const FInputDeviceId DeviceId);
+	void ResetDeviceProperty(const FPlatformUserId PlatformUser, const FInputDeviceId DeviceId, bool bForceReset = false);
 
 	/**
 	* Native C++ implementation of ResetDeviceProperty
@@ -75,7 +75,7 @@ protected:
 	* 
 	* @see ResetDeviceProperty
 	*/
-	virtual void ResetDeviceProperty_Implementation(const FPlatformUserId PlatformUser, const FInputDeviceId DeviceId);
+	virtual void ResetDeviceProperty_Implementation(const FPlatformUserId PlatformUser, const FInputDeviceId DeviceId, bool bForceReset = false);
 
 	/**
 	* Apply the device property from GetInternalDeviceProperty to the given platform user. 
@@ -88,7 +88,7 @@ protected:
 	
 	/** Gets a pointer to the current input device property that the IInputInterface can use. */
 	virtual FInputDeviceProperty* GetInternalDeviceProperty() { return nullptr; };
-
+	
 public:
 
 	/**
@@ -192,7 +192,7 @@ class UColorInputDeviceProperty : public UInputDeviceProperty
 public:
 
 	virtual void EvaluateDeviceProperty_Implementation(const FPlatformUserId PlatformUser, const FInputDeviceId DeviceId, const float DeltaTime, const float Duration) override;
-	virtual void ResetDeviceProperty_Implementation(const FPlatformUserId PlatformUser, const FInputDeviceId DeviceId) override;
+	virtual void ResetDeviceProperty_Implementation(const FPlatformUserId PlatformUser, const FInputDeviceId DeviceId, bool bForceReset = false) override;
 	virtual FInputDeviceProperty* GetInternalDeviceProperty() override;
 
 	/** Default color data that will be used by default. Device Specific overrides will be used when the current input device matches */
@@ -246,7 +246,7 @@ class UColorInputDeviceCurveProperty : public UInputDeviceProperty
 public:
 
 	virtual void EvaluateDeviceProperty_Implementation(const FPlatformUserId PlatformUser, const FInputDeviceId DeviceId, const float DeltaTime, const float Duration) override;
-	virtual void ResetDeviceProperty_Implementation(const FPlatformUserId PlatformUser, const FInputDeviceId DeviceId) override;
+	virtual void ResetDeviceProperty_Implementation(const FPlatformUserId PlatformUser, const FInputDeviceId DeviceId, bool bForceReset = false) override;
 	virtual FInputDeviceProperty* GetInternalDeviceProperty() override;
 	virtual float RecalculateDuration() override;
 
@@ -292,7 +292,7 @@ class ENGINE_API UInputDeviceTriggerEffect : public UInputDeviceProperty
 public:	
 
 	virtual FInputDeviceProperty* GetInternalDeviceProperty() override;
-	virtual void ResetDeviceProperty_Implementation(const FPlatformUserId PlatformUser, const FInputDeviceId DeviceId) override;
+	virtual void ResetDeviceProperty_Implementation(const FPlatformUserId PlatformUser, const FInputDeviceId DeviceId, bool bForceReset = false) override;
 	virtual void ApplyDeviceProperty(const FPlatformUserId UserId, const FInputDeviceId DeviceId) override;
 
 	UPROPERTY(EditAnywhere, BlueprintReadWrite, Category = "Triggers")
diff --git a/Engine/Source/Runtime/Engine/Classes/GameFramework/InputDeviceSubsystem.h b/Engine/Source/Runtime/Engine/Classes/GameFramework/InputDeviceSubsystem.h
--- a/Engine/Source/Runtime/Engine/Classes/GameFramework/InputDeviceSubsystem.h
+++ b/Engine/Source/Runtime/Engine/Classes/GameFramework/InputDeviceSubsystem.h
@@ -216,7 +216,25 @@ protected:
 	void OnPIEStopped(bool bSimulating);
 
 	bool bIsPIEPlaying = false;
-#endif
+
+#endif	// WITH_EDITOR
+
+#if WITH_EDITORONLY_DATA
+
+	struct FResetPIEData
+	{
+	 	TObjectPtr<UClass> ClassToReset;
+	 	FInputDeviceId DeviceId;
+	 	FPlatformUserId UserId;
+	};
+	
+	/** 
+	 * A set of device properties that should be reset when PIE ends.
+	 * This ensures that each device is reset for the editor
+	 **/
+	TArray<FResetPIEData> PropertyClassesRequiringReset;
+
+#endif	// WITH_EDITORONLY_DATA
 	
 	/**
 	* Set of currently active input device properties that will be evaluated on tick
diff --git a/Engine/Source/Runtime/Engine/Private/GameFramework/InputDeviceProperties.cpp b/Engine/Source/Runtime/Engine/Private/GameFramework/InputDeviceProperties.cpp
--- a/Engine/Source/Runtime/Engine/Private/GameFramework/InputDeviceProperties.cpp
+++ b/Engine/Source/Runtime/Engine/Private/GameFramework/InputDeviceProperties.cpp
@@ -62,7 +62,7 @@ void UInputDeviceProperty::EvaluateDeviceProperty_Implementation(const FPlatform
 
 }
 
-void UInputDeviceProperty::ResetDeviceProperty_Implementation(const FPlatformUserId PlatformUser, const FInputDeviceId DeviceId)
+void UInputDeviceProperty::ResetDeviceProperty_Implementation(const FPlatformUserId PlatformUser, const FInputDeviceId DeviceId, bool bForceReset /*= false*/)
 {
 
 }
@@ -86,7 +86,7 @@ void UColorInputDeviceProperty::EvaluateDeviceProperty_Implementation(const FPla
 	}
 }
 
-void UColorInputDeviceProperty::ResetDeviceProperty_Implementation(const FPlatformUserId PlatformUser, const FInputDeviceId DeviceId)
+void UColorInputDeviceProperty::ResetDeviceProperty_Implementation(const FPlatformUserId PlatformUser, const FInputDeviceId DeviceId, bool bForceReset /*= false*/)
 {
 	bool bReset = ColorData.bResetAfterCompletion;
 	if (const FDeviceColorData* Data = GetDeviceSpecificData<FDeviceColorData>(PlatformUser, DeviceId, DeviceOverrideData))
@@ -94,6 +94,8 @@ void UColorInputDeviceProperty::ResetDeviceProperty_Implementation(const FPlatfo
 		bReset = Data->bResetAfterCompletion;
 	}
 
+	bReset |= bForceReset;
+
 	if (bReset)
 	{
 		// Disabling the light will reset the color
@@ -136,7 +138,7 @@ void UColorInputDeviceCurveProperty::EvaluateDeviceProperty_Implementation(const
 	}
 }
 
-void UColorInputDeviceCurveProperty::ResetDeviceProperty_Implementation(const FPlatformUserId PlatformUser, const FInputDeviceId DeviceId)
+void UColorInputDeviceCurveProperty::ResetDeviceProperty_Implementation(const FPlatformUserId PlatformUser, const FInputDeviceId DeviceId, bool bForceReset /*= false*/)
 {
 	bool bReset = ColorData.bResetAfterCompletion;
 	if (const FDeviceColorCurveData* Data = GetDeviceSpecificData<FDeviceColorCurveData>(PlatformUser, DeviceId, DeviceOverrideData))
@@ -144,6 +146,8 @@ void UColorInputDeviceCurveProperty::ResetDeviceProperty_Implementation(const FP
 		bReset = Data->bResetAfterCompletion;
 	}
 
+	bReset |= bForceReset;
+	
 	if (bReset)
 	{
 		// Disabling the light will reset the color
@@ -195,9 +199,9 @@ FInputDeviceProperty* UInputDeviceTriggerEffect::GetInternalDeviceProperty()
 	return &ResetProperty;
 }
 
-void UInputDeviceTriggerEffect::ResetDeviceProperty_Implementation(const FPlatformUserId PlatformUser, const FInputDeviceId DeviceId)
+void UInputDeviceTriggerEffect::ResetDeviceProperty_Implementation(const FPlatformUserId PlatformUser, const FInputDeviceId DeviceId, bool bForceReset /*= false*/)
 {
-	if (BaseTriggerData.bResetUponCompletion && BaseTriggerData.AffectedTriggers != EInputDeviceTriggerMask::None)
+	if (bForceReset || (BaseTriggerData.bResetUponCompletion && BaseTriggerData.AffectedTriggers != EInputDeviceTriggerMask::None))
 	{
 		// Pass in our reset property
 		ResetProperty.AffectedTriggers = BaseTriggerData.AffectedTriggers;
diff --git a/Engine/Source/Runtime/Engine/Private/GameFramework/InputDeviceSubsystem.cpp b/Engine/Source/Runtime/Engine/Private/GameFramework/InputDeviceSubsystem.cpp
--- a/Engine/Source/Runtime/Engine/Private/GameFramework/InputDeviceSubsystem.cpp
+++ b/Engine/Source/Runtime/Engine/Private/GameFramework/InputDeviceSubsystem.cpp
@@ -390,6 +390,28 @@ FInputDevicePropertyHandle UInputDeviceSubsystem::ActivateDeviceProperty(UInputD
 
 		// Keep track of this new active property. Placing it in this set will have it updated on this subsystem's Tick
 		ActiveProperties.Add(ActiveProp);
+
+#if WITH_EDITORONLY_DATA
+		// If we are currently PIE'ing, then track this property if it's class isn't already known. This will
+		// ensure that we reset device properties of unique types
+		if (bIsPIEPlaying)
+		{
+			const bool bAlreadyHasResetData = 
+				 PropertyClassesRequiringReset.ContainsByPredicate([&ActiveProp](const FResetPIEData& Data)
+				 {
+				 	return Data.DeviceId == ActiveProp.DeviceId && Data.UserId == ActiveProp.PlatformUser && Data.ClassToReset == ActiveProp.Property->GetClass();
+				 });
+			
+			if (!bAlreadyHasResetData)
+			{
+			 	FResetPIEData ResetData = {};
+             	ResetData.ClassToReset = ActiveProp.Property->GetClass();
+             	ResetData.DeviceId = ActiveProp.DeviceId;
+             	ResetData.UserId = ActiveProp.PlatformUser;
+			 	PropertyClassesRequiringReset.Emplace(ResetData);	
+			}
+		}
+#endif	// WITH_EDITOR
 	}
 
 	return OutHandle;
@@ -509,6 +531,9 @@ void UInputDeviceSubsystem::OnPrePIEStarted(bool bSimulating)
 {
 	// Remove all active properties, just in case someone was previewing something in the editor that are still going
 	RemoveAllDeviceProperties();
+#if WITH_EDITORONLY_DATA
+	PropertyClassesRequiringReset.Reset();
+#endif	// WITH_EDITORONLY_DATA
 	bIsPIEPlaying = true;
 }
 
@@ -526,6 +551,18 @@ void UInputDeviceSubsystem::OnPIEStopped(bool bSimulating)
 {
 	// Remove all active properties when PIE stops
 	RemoveAllDeviceProperties();
+
+#if WITH_EDITORONLY_DATA
+	// Reset any Device Properties that have been added so that their input devices are in a clean state when PIE ends
+	for (const FResetPIEData& ResetData : PropertyClassesRequiringReset)
+	{
+		UInputDeviceProperty* ToReset = NewObject<UInputDeviceProperty>(/* Outer = */ GetTransientPackage(), /* Class */ ResetData.ClassToReset);
+		ToReset->ResetDeviceProperty(ResetData.UserId, ResetData.DeviceId, /* bForceReset= */ true);
+	}
+	
+	PropertyClassesRequiringReset.Reset();
+#endif	// WITH_EDITORONLY_DATA
+	
 	bIsPIEPlaying = false;
 }
-#endif
+#endif	// WITH_EDITOR
