SVT::FTextureRenderResources now uses texture references and does not own the underlying texture resources anymore. This will allow updating the page table and tile data textures without having to invalidate any materials/uniform buffers. The streaming manager now no longer accesses the RHI resource pointers of a USVTFrame when removing the SVT from the streaming manager, which should fix a crash that sometimes happend during garbage collection.
Also added some clarifying comments to SVT::IStreamingManager.

#rnx

diff --git a/Engine/Plugins/FX/Niagara/Source/Niagara/Private/NiagaraDataInterfaceRenderTargetVolume.cpp b/Engine/Plugins/FX/Niagara/Source/Niagara/Private/NiagaraDataInterfaceRenderTargetVolume.cpp
--- a/Engine/Plugins/FX/Niagara/Source/Niagara/Private/NiagaraDataInterfaceRenderTargetVolume.cpp
+++ b/Engine/Plugins/FX/Niagara/Source/Niagara/Private/NiagaraDataInterfaceRenderTargetVolume.cpp
@@ -903,9 +903,9 @@ bool UNiagaraDataInterfaceRenderTargetVolume::SimCacheReadFrame(UObject* Storage
 							);
 						}						
 
-						FRHITexture* PageTableTexture = RT_SVTRenderResources->GetPageTableTextureRHI();
-						FRHITexture* TextureA = RT_SVTRenderResources->GetPhysicalTileDataATextureRHI();
-						FRHIShaderResourceView* StreamingInfoBufferSRV = RT_SVTRenderResources->GetStreamingInfoBufferSRVRHI();
+						FRHITexture* PageTableTexture = RT_SVTRenderResources->GetPageTableTexture();
+						FRHITexture* TextureA = RT_SVTRenderResources->GetPhysicalTileDataATexture();
+						FRHIShaderResourceView* StreamingInfoBufferSRV = RT_SVTRenderResources->GetStreamingInfoBufferSRV();
 
 						PassParameters->TileDataTextureSampler = TStaticSamplerState<SF_Point, AM_Clamp, AM_Clamp, AM_Clamp>::GetRHI();
 						PassParameters->SparseVolumeTexturePageTable = PageTableTexture ? PageTableTexture : GBlackUintVolumeTexture->TextureRHI.GetReference();
diff --git a/Engine/Plugins/FX/Niagara/Source/Niagara/Private/NiagaraDataInterfaceSparseVolumeTexture.cpp b/Engine/Plugins/FX/Niagara/Source/Niagara/Private/NiagaraDataInterfaceSparseVolumeTexture.cpp
--- a/Engine/Plugins/FX/Niagara/Source/Niagara/Private/NiagaraDataInterfaceSparseVolumeTexture.cpp
+++ b/Engine/Plugins/FX/Niagara/Source/Niagara/Private/NiagaraDataInterfaceSparseVolumeTexture.cpp
@@ -293,15 +293,15 @@ void UNiagaraDataInterfaceSparseVolumeTexture::SetShaderParameters(const FNiagar
 	
 	if (RTInstanceData && RTInstanceData->RenderResources)
 	{
-		FTextureRHIRef PageTableTexture = RTInstanceData->RenderResources->GetPageTableTextureRHI();
-		FTextureRHIRef PhysicalTileDataATexture = RTInstanceData->RenderResources->GetPhysicalTileDataATextureRHI();
-		FTextureRHIRef PhysicalTileDataBTexture = RTInstanceData->RenderResources->GetPhysicalTileDataBTextureRHI();
-		FShaderResourceViewRHIRef StreamingInfoBufferSRV = RTInstanceData->RenderResources->GetStreamingInfoBufferSRVRHI();
-
-		Parameters->PageTableTexture = PageTableTexture ? PageTableTexture.GetReference() : Parameters->PageTableTexture;
-		Parameters->PhysicalTileDataATexture = PhysicalTileDataATexture ? PhysicalTileDataATexture.GetReference() : Parameters->PhysicalTileDataATexture;
-		Parameters->PhysicalTileDataBTexture = PhysicalTileDataBTexture ? PhysicalTileDataBTexture.GetReference() : Parameters->PhysicalTileDataBTexture;
-		Parameters->StreamingInfoBuffer = StreamingInfoBufferSRV ? StreamingInfoBufferSRV.GetReference() : Parameters->StreamingInfoBuffer;
+		FRHITexture* PageTableTexture = RTInstanceData->RenderResources->GetPageTableTexture();
+		FRHITexture* PhysicalTileDataATexture = RTInstanceData->RenderResources->GetPhysicalTileDataATexture();
+		FRHITexture* PhysicalTileDataBTexture = RTInstanceData->RenderResources->GetPhysicalTileDataBTexture();
+		FRHIShaderResourceView* StreamingInfoBufferSRV = RTInstanceData->RenderResources->GetStreamingInfoBufferSRV();
+
+		Parameters->PageTableTexture = PageTableTexture ? PageTableTexture : Parameters->PageTableTexture;
+		Parameters->PhysicalTileDataATexture = PhysicalTileDataATexture ? PhysicalTileDataATexture : Parameters->PhysicalTileDataATexture;
+		Parameters->PhysicalTileDataBTexture = PhysicalTileDataBTexture ? PhysicalTileDataBTexture : Parameters->PhysicalTileDataBTexture;
+		Parameters->StreamingInfoBuffer = StreamingInfoBufferSRV ? StreamingInfoBufferSRV : Parameters->StreamingInfoBuffer;
 		RTInstanceData->RenderResources->GetPackedUniforms(Parameters->PackedUniforms0, Parameters->PackedUniforms1);
 		Parameters->TextureSize = RTInstanceData->TextureSize;
 		Parameters->MipLevels = RTInstanceData->MipLevels;
diff --git a/Engine/Source/Runtime/Engine/Classes/SparseVolumeTexture/SparseVolumeTexture.h b/Engine/Source/Runtime/Engine/Classes/SparseVolumeTexture/SparseVolumeTexture.h
--- a/Engine/Source/Runtime/Engine/Classes/SparseVolumeTexture/SparseVolumeTexture.h
+++ b/Engine/Source/Runtime/Engine/Classes/SparseVolumeTexture/SparseVolumeTexture.h
@@ -129,21 +129,21 @@ class ENGINE_API FTextureRenderResources : public ::FRenderResource
 {
 	friend class FStreamingManager;
 public:
-	const FHeader& GetHeader() const								{ check(IsInParallelRenderingThread()); return Header; }
-	FIntVector3 GetTileDataTextureResolution() const				{ check(IsInParallelRenderingThread()); return TileDataTextureResolution; }
-	int32 GetFrameIndex() const										{ check(IsInParallelRenderingThread()); return FrameIndex; }
-	int32 GetNumLogicalMipLevels() const							{ check(IsInParallelRenderingThread()); return NumLogicalMipLevels; }
-	FTextureRHIRef GetPageTableTextureRHI() const					{ check(IsInParallelRenderingThread()); return PageTableTextureRHI; }
-	FTextureRHIRef GetPhysicalTileDataATextureRHI() const			{ check(IsInParallelRenderingThread()); return PhysicalTileDataATextureRHI; }
-	FTextureRHIRef GetPhysicalTileDataBTextureRHI() const			{ check(IsInParallelRenderingThread()); return PhysicalTileDataBTextureRHI; }
-	FShaderResourceViewRHIRef GetStreamingInfoBufferSRVRHI() const	{ check(IsInParallelRenderingThread()); return StreamingInfoBufferSRVRHI; }
+	const FHeader& GetHeader() const							{ check(IsInParallelRenderingThread()); return Header; }
+	FIntVector3 GetTileDataTextureResolution() const			{ check(IsInParallelRenderingThread()); return TileDataTextureResolution; }
+	int32 GetFrameIndex() const									{ check(IsInParallelRenderingThread()); return FrameIndex; }
+	int32 GetNumLogicalMipLevels() const						{ check(IsInParallelRenderingThread()); return NumLogicalMipLevels; }
+	FRHITextureReference* GetPageTableTexture() const			{ check(IsInParallelRenderingThread()); return PageTableTextureReferenceRHI.GetReference(); }
+	FRHITextureReference* GetPhysicalTileDataATexture() const	{ check(IsInParallelRenderingThread()); return PhysicalTileDataATextureReferenceRHI.GetReference(); }
+	FRHITextureReference* GetPhysicalTileDataBTexture() const	{ check(IsInParallelRenderingThread()); return PhysicalTileDataBTextureReferenceRHI.GetReference(); }
+	FRHIShaderResourceView* GetStreamingInfoBufferSRV() const	{ check(IsInParallelRenderingThread()); return StreamingInfoBufferSRVRHI.GetReference(); }
 	void GetPackedUniforms(FUintVector4& OutPacked0, FUintVector4& OutPacked1) const;
 	// Updates the GlobalVolumeResolution member in a thread-safe way.
 	void SetGlobalVolumeResolution_GameThread(const FIntVector3& GlobalVolumeResolution);
 
 	//~ Begin FRenderResource Interface.
-	virtual void InitRHI() override { /* Managed by FStreamingManager */ }
-	virtual void ReleaseRHI() override { /* Managed by FStreamingManager */ }
+	virtual void InitRHI() override;
+	virtual void ReleaseRHI() override;
 	//~ End FRenderResource Interface.
 
 private:
@@ -152,9 +152,9 @@ private:
 	FIntVector3 TileDataTextureResolution = FIntVector3::ZeroValue;
 	int32 FrameIndex = INDEX_NONE;
 	int32 NumLogicalMipLevels = 0; // Might not all be resident in GPU memory
-	FTextureRHIRef PageTableTextureRHI;
-	FTextureRHIRef PhysicalTileDataATextureRHI;
-	FTextureRHIRef PhysicalTileDataBTextureRHI;
+	FTextureReferenceRHIRef PageTableTextureReferenceRHI;
+	FTextureReferenceRHIRef PhysicalTileDataATextureReferenceRHI;
+	FTextureReferenceRHIRef PhysicalTileDataBTextureReferenceRHI;
 	FShaderResourceViewRHIRef StreamingInfoBufferSRVRHI;
 };
 
diff --git a/Engine/Source/Runtime/Engine/Private/Materials/MaterialUniformExpressions.cpp b/Engine/Source/Runtime/Engine/Private/Materials/MaterialUniformExpressions.cpp
--- a/Engine/Source/Runtime/Engine/Private/Materials/MaterialUniformExpressions.cpp
+++ b/Engine/Source/Runtime/Engine/Private/Materials/MaterialUniformExpressions.cpp
@@ -1277,9 +1277,10 @@ void FUniformExpressionSet::FillUniformBuffer(const FMaterialRenderContext& Mate
 			check(BufferCursor <= TempBuffer + TempBufferSize);
 
 			check(GBlackVolumeTexture->TextureRHI);
-			*ResourceTableTexturePageTablePtr = GBlackVolumeTexture->TextureRHI;
+			*ResourceTableTexturePageTablePtr = GBlackUintVolumeTexture->TextureRHI;
 			*ResourceTableTexturePhysicalAPtr = GBlackVolumeTexture->TextureRHI;
 			*ResourceTableTexturePhysicalBPtr = GBlackVolumeTexture->TextureRHI;
+			*ResourceTableStreamingInfoBufferPtr = GEmptyVertexBufferWithUAV->ShaderResourceViewRHI;
 			check(GBlackVolumeTexture->SamplerStateRHI);
 			*ResourceTablePhysicalSamplerPtr = GBlackVolumeTexture->SamplerStateRHI;
 			
@@ -1290,19 +1291,16 @@ void FUniformExpressionSet::FillUniformBuffer(const FMaterialRenderContext& Mate
 				const UE::SVT::FTextureRenderResources* RenderResources = SVTexture->GetTextureRenderResources();
 				if (RenderResources)
 				{
-					// It's possible for RenderResources to be valid, but PageTableTexture still null, so we need to have a fallback
-					FTextureRHIRef PageTableTexture = RenderResources->GetPageTableTextureRHI();
-					*ResourceTableTexturePageTablePtr = PageTableTexture ? PageTableTexture : GBlackUintVolumeTexture->TextureRHI;
-
-					FTextureRHIRef TileDataATexture = RenderResources->GetPhysicalTileDataATextureRHI();
-					*ResourceTableTexturePhysicalAPtr = TileDataATexture ? TileDataATexture : GBlackVolumeTexture->TextureRHI;
-
-					FTextureRHIRef TileDataBTexture = RenderResources->GetPhysicalTileDataBTextureRHI();
-					*ResourceTableTexturePhysicalBPtr = TileDataBTexture ? TileDataBTexture : GBlackVolumeTexture->TextureRHI;
-
-					FShaderResourceViewRHIRef StreamingInfoBufferSRV = RenderResources->GetStreamingInfoBufferSRVRHI();
-					*ResourceTableStreamingInfoBufferPtr = StreamingInfoBufferSRV ? StreamingInfoBufferSRV : GEmptyVertexBufferWithUAV->ShaderResourceViewRHI;
-
+					FRHITexture* PageTableTexture = RenderResources->GetPageTableTexture();
+					FRHITexture* TileDataATexture = RenderResources->GetPhysicalTileDataATexture();
+					FRHITexture* TileDataBTexture = RenderResources->GetPhysicalTileDataBTexture();
+					FRHIShaderResourceView* StreamingInfoBufferSRV = RenderResources->GetStreamingInfoBufferSRV();
+
+					// It's possible for RenderResources to be valid, but PageTableTexture still null, so we need to have a fallback (a black uint volume texture)
+					*ResourceTableTexturePageTablePtr = PageTableTexture ? PageTableTexture : *ResourceTableTexturePageTablePtr;
+					*ResourceTableTexturePhysicalAPtr = TileDataATexture ? TileDataATexture : *ResourceTableTexturePhysicalAPtr;
+					*ResourceTableTexturePhysicalBPtr = TileDataBTexture ? TileDataBTexture : *ResourceTableTexturePhysicalBPtr;
+					*ResourceTableStreamingInfoBufferPtr = StreamingInfoBufferSRV ? StreamingInfoBufferSRV : *ResourceTableStreamingInfoBufferPtr;
 					*ResourceTablePhysicalSamplerPtr = TStaticSamplerState<SF_Bilinear, AM_Clamp, AM_Clamp, AM_Clamp>::GetRHI();
 				}
 			}
diff --git a/Engine/Source/Runtime/Engine/Private/SparseVolumeTexture/SparseVolumeTexture.cpp b/Engine/Source/Runtime/Engine/Private/SparseVolumeTexture/SparseVolumeTexture.cpp
--- a/Engine/Source/Runtime/Engine/Private/SparseVolumeTexture/SparseVolumeTexture.cpp
+++ b/Engine/Source/Runtime/Engine/Private/SparseVolumeTexture/SparseVolumeTexture.cpp
@@ -13,6 +13,7 @@
 #include "EngineUtils.h"
 #include "Shader/ShaderTypes.h"
 #include "RenderingThread.h"
+#include "GlobalRenderResources.h"
 #include "UObject/Package.h"
 #include "SparseVolumeTexture/SparseVolumeTextureData.h"
 #include "SparseVolumeTexture/SparseVolumeTextureUtility.h"
@@ -581,6 +582,21 @@ void FTextureRenderResources::SetGlobalVolumeResolution_GameThread(const FIntVec
 		});
 }
 
+void FTextureRenderResources::InitRHI()
+{
+	PageTableTextureReferenceRHI = RHICreateTextureReference(GBlackUintVolumeTexture->TextureRHI);
+	PhysicalTileDataATextureReferenceRHI = RHICreateTextureReference(GBlackVolumeTexture->TextureRHI);
+	PhysicalTileDataBTextureReferenceRHI = RHICreateTextureReference(GBlackVolumeTexture->TextureRHI);
+}
+
+void FTextureRenderResources::ReleaseRHI()
+{
+	PageTableTextureReferenceRHI.SafeRelease();
+	PhysicalTileDataATextureReferenceRHI.SafeRelease();
+	PhysicalTileDataBTextureReferenceRHI.SafeRelease();
+	StreamingInfoBufferSRVRHI.SafeRelease();
+}
+
 }
 }
 
diff --git a/Engine/Source/Runtime/Engine/Private/SparseVolumeTexture/SparseVolumeTextureStreamingManager.cpp b/Engine/Source/Runtime/Engine/Private/SparseVolumeTexture/SparseVolumeTextureStreamingManager.cpp
--- a/Engine/Source/Runtime/Engine/Private/SparseVolumeTexture/SparseVolumeTextureStreamingManager.cpp
+++ b/Engine/Source/Runtime/Engine/Private/SparseVolumeTexture/SparseVolumeTextureStreamingManager.cpp
@@ -679,12 +679,12 @@ void FStreamingManager::InitRHI()
 
 	MaxPendingMipLevels = GSVTStreamingMaxPendingMipLevels;
 	PendingMipLevels.SetNum(MaxPendingMipLevels);
-	PageTableUpdater = new FPageTableUpdater();
-	StreamingInfoBufferUpdater = new FStreamingInfoBufferUpdater();
+	PageTableUpdater = MakeUnique<FPageTableUpdater>();
+	StreamingInfoBufferUpdater = MakeUnique<FStreamingInfoBufferUpdater>();
 
 #if WITH_EDITORONLY_DATA
-	RequestOwner = new FRequestOwner(EPriority::Normal);
-	RequestOwnerBlocking = new FRequestOwner(EPriority::Blocking);
+	RequestOwner = MakeUnique<FRequestOwner>(EPriority::Normal);
+	RequestOwnerBlocking = MakeUnique<FRequestOwner>(EPriority::Blocking);
 #endif
 }
 
@@ -694,18 +694,6 @@ void FStreamingManager::ReleaseRHI()
 	{
 		return;
 	}
-
-#if WITH_EDITORONLY_DATA
-	delete RequestOwner;
-	RequestOwner = nullptr;
-	delete RequestOwnerBlocking;
-	RequestOwnerBlocking = nullptr;
-#endif
-
-	delete PageTableUpdater;
-	PageTableUpdater = nullptr;
-	delete StreamingInfoBufferUpdater;
-	StreamingInfoBufferUpdater = nullptr;
 }
 
 void FStreamingManager::Add_GameThread(UStreamableSparseVolumeTexture* SparseVolumeTexture)
@@ -877,7 +865,7 @@ void FStreamingManager::Request(UStreamableSparseVolumeTexture* SparseVolumeText
 	}
 }
 
-void FStreamingManager::BeginAsyncUpdate(FRDGBuilder& GraphBuilder, bool bForceNonAsync)
+void FStreamingManager::BeginAsyncUpdate(FRDGBuilder& GraphBuilder, bool bBlocking)
 {
 	check(IsInRenderingThread());
 	check(!AsyncState.bUpdateActive);
@@ -946,7 +934,7 @@ void FStreamingManager::BeginAsyncUpdate(FRDGBuilder& GraphBuilder, bool bForceN
 
 			const FResources* Resources = SVTInfo->PerFrameInfo[PendingMipLevel.FrameIndex].Resources;
 			SVTInfo->TileDataTexture->NumTilesToUpload += Resources->MipLevelStreamingInfo[PendingMipLevel.MipLevelIndex].NumPhysicalTiles;
-			TileDataTexturesToUpdate.Add(SVTInfo->TileDataTexture);
+			TileDataTexturesToUpdate.Add(SVTInfo->TileDataTexture.Get());
 			NumPageTableUpdatesTotal += Resources->MipLevelStreamingInfo[PendingMipLevel.MipLevelIndex].PageTableSize / (2 * sizeof(uint32));
 		}
 
@@ -963,7 +951,7 @@ void FStreamingManager::BeginAsyncUpdate(FRDGBuilder& GraphBuilder, bool bForceN
 	Parameters.StreamingManager = this;
 	
 	check(AsyncTaskEvents.IsEmpty());
-	if (GSVTStreamingAsyncThread && !bForceNonAsync)
+	if (GSVTStreamingAsyncThread && !bBlocking)
 	{
 		AsyncState.bUpdateIsAsync = true;
 		AsyncTaskEvents.Add(TGraphTask<FStreamingUpdateTask>::CreateTask().ConstructAndDispatchWhenReady(Parameters));
@@ -1144,7 +1132,7 @@ void FStreamingManager::AddInternal(FRDGBuilder& GraphBuilder, FNewSparseVolumeT
 		const FIntVector3 TileDataVolumeResolution = ComputeTileDataVolumeResolution(NumPhysicalTilesCapacity);
 		const FIntVector3 TileDataVolumeResolutionInTiles = TileDataVolumeResolution / SPARSE_VOLUME_TILE_RES_PADDED;
 
-		SVTInfo.TileDataTexture = new FTileDataTexture(TileDataVolumeResolutionInTiles, SVTInfo.FormatA, SVTInfo.FormatB);
+		SVTInfo.TileDataTexture = MakeUnique<FTileDataTexture>(TileDataVolumeResolutionInTiles, SVTInfo.FormatA, SVTInfo.FormatB);
 		SVTInfo.TileDataTexture->InitResource();
 
 		// Create streaming info buffer
@@ -1190,15 +1178,6 @@ void FStreamingManager::AddInternal(FRDGBuilder& GraphBuilder, FNewSparseVolumeT
 			FrameInfo.LowestRequestedMipLevel = NumMipLevels - 1;
 			FrameInfo.LowestResidentMipLevel = NumMipLevels - 1;
 
-			// Initialize TextureRenderResources
-			FrameInfo.TextureRenderResources->PhysicalTileDataATextureRHI = SVTInfo.TileDataTexture->TileDataTextureARHIRef;
-			FrameInfo.TextureRenderResources->PhysicalTileDataBTextureRHI = SVTInfo.TileDataTexture->TileDataTextureBRHIRef;
-			FrameInfo.TextureRenderResources->StreamingInfoBufferSRVRHI = SVTInfo.StreamingInfoBufferSRVRHIRef;
-			FrameInfo.TextureRenderResources->Header = Resources->Header;
-			FrameInfo.TextureRenderResources->TileDataTextureResolution = SVTInfo.TileDataTexture->ResolutionInTiles * SPARSE_VOLUME_TILE_RES_PADDED;
-			FrameInfo.TextureRenderResources->FrameIndex = FrameIdx;
-			FrameInfo.TextureRenderResources->NumLogicalMipLevels = NumMipLevels;
-
 			// Create page table
 			{
 				// SVT_TODO: Currently we keep all mips of the page table resident. It would be better to stream in/out page table mips.
@@ -1212,9 +1191,19 @@ void FStreamingManager::AddInternal(FRDGBuilder& GraphBuilder, FNewSparseVolumeT
 					.SetFlags(ETextureCreateFlags::ShaderResource | ETextureCreateFlags::UAV)
 					.SetNumMips((uint8)NumResidentMipLevels);
 
-				FrameInfo.TextureRenderResources->PageTableTextureRHI = RHICreateTexture(Desc);
+				FrameInfo.PageTableTextureRHIRef = RHICreateTexture(Desc);
 			}
 
+			// Initialize TextureRenderResources
+			RHIUpdateTextureReference(FrameInfo.TextureRenderResources->PageTableTextureReferenceRHI, FrameInfo.PageTableTextureRHIRef);
+			RHIUpdateTextureReference(FrameInfo.TextureRenderResources->PhysicalTileDataATextureReferenceRHI, SVTInfo.TileDataTexture->TileDataTextureARHIRef);
+			RHIUpdateTextureReference(FrameInfo.TextureRenderResources->PhysicalTileDataBTextureReferenceRHI, SVTInfo.TileDataTexture->TileDataTextureBRHIRef);
+			FrameInfo.TextureRenderResources->StreamingInfoBufferSRVRHI = SVTInfo.StreamingInfoBufferSRVRHIRef;
+			FrameInfo.TextureRenderResources->Header = Resources->Header;
+			FrameInfo.TextureRenderResources->TileDataTextureResolution = SVTInfo.TileDataTexture->ResolutionInTiles * SPARSE_VOLUME_TILE_RES_PADDED;
+			FrameInfo.TextureRenderResources->FrameIndex = FrameIdx;
+			FrameInfo.TextureRenderResources->NumLogicalMipLevels = NumMipLevels;
+
 			// Upload root mip data and update page tables
 			const FMipLevelStreamingInfo* RootStreamingInfo = !Resources->MipLevelStreamingInfo.IsEmpty() ? &Resources->MipLevelStreamingInfo.Last() : nullptr;
 			if (!Resources->RootData.IsEmpty() && RootStreamingInfo)
@@ -1245,7 +1234,7 @@ void FStreamingManager::AddInternal(FRDGBuilder& GraphBuilder, FNewSparseVolumeT
 
 				// Update highest mip (1x1x1) in page table
 				const FUpdateTextureRegion3D UpdateRegion(0, 0, 0, 0, 0, 0, 1, 1, 1);
-				RHIUpdateTexture3D(FrameInfo.TextureRenderResources->PageTableTextureRHI, FrameInfo.TextureRenderResources->PageTableTextureRHI->GetDesc().NumMips - 1, UpdateRegion, sizeof(uint32), sizeof(uint32), (uint8*)&TileCoord);
+				RHIUpdateTexture3D(FrameInfo.PageTableTextureRHIRef, FrameInfo.PageTableTextureRHIRef->GetDesc().NumMips - 1, UpdateRegion, sizeof(uint32), sizeof(uint32), (uint8*)&TileCoord);
 			}
 		}
 
@@ -1275,17 +1264,14 @@ void FStreamingManager::RemoveInternal(UStreamableSparseVolumeTexture* SparseVol
 	FStreamingInfo* SVTInfo = StreamingInfo.Find(SparseVolumeTexture);
 	if (SVTInfo)
 	{
-		// The RHI resources in FTextureRenderResources are managed by FStreamingManager, so release them now.
 		for (FFrameInfo& FrameInfo : SVTInfo->PerFrameInfo)
 		{
-			FrameInfo.TextureRenderResources->PageTableTextureRHI.SafeRelease();
-			FrameInfo.TextureRenderResources->PhysicalTileDataATextureRHI.SafeRelease();
-			FrameInfo.TextureRenderResources->PhysicalTileDataBTextureRHI.SafeRelease();
+			FrameInfo.PageTableTextureRHIRef.SafeRelease();
 		}
 		if (SVTInfo->TileDataTexture)
 		{
 			SVTInfo->TileDataTexture->ReleaseResource();
-			delete SVTInfo->TileDataTexture;
+			SVTInfo->TileDataTexture.Reset();
 		}
 
 		StreamingInfo.Remove(SparseVolumeTexture);
@@ -1466,7 +1452,7 @@ void FStreamingManager::IssueRequests(int32 MaxSelectedRequests)
 		check((SelectedKey.MipLevelIndex + 1) < Resources->MipLevelStreamingInfo.Num()); // The lowest/last mip level is always resident and does not stream.
 		const FMipLevelStreamingInfo& MipLevelStreamingInfo = Resources->MipLevelStreamingInfo[SelectedKey.MipLevelIndex];
 
-		FTileDataTexture* TileDataTexture = SVTInfo->TileDataTexture;
+		TUniquePtr<FTileDataTexture>& TileDataTexture = SVTInfo->TileDataTexture;
 		check(TileDataTexture);
 
 		// Ensure that enough tiles are available in the tile texture
@@ -1694,7 +1680,7 @@ void FStreamingManager::StreamOutMipLevel(FStreamingInfo& SVTInfo, FLRUNode* LRU
 
 	if (bNeedsPageTableClear)
 	{
-		PageTableClears.Push({ FrameInfo.TextureRenderResources->PageTableTextureRHI, MipLevelIndex });
+		PageTableClears.Push({ FrameInfo.PageTableTextureRHIRef, MipLevelIndex });
 	}
 
 	// Free allocated tiles
@@ -1905,7 +1891,7 @@ void FStreamingManager::InstallReadyMipLevels()
 
 		uint8* DstPageCoords = nullptr;
 		uint8* DstPageEntries = nullptr;
-		PageTableUpdater->Add_GetRef(FrameInfo.TextureRenderResources->PageTableTextureRHI, PendingMipLevel.MipLevelIndex, NumPageTableUpdates, DstPageCoords, DstPageEntries);
+		PageTableUpdater->Add_GetRef(FrameInfo.PageTableTextureRHIRef, PendingMipLevel.MipLevelIndex, NumPageTableUpdates, DstPageCoords, DstPageEntries);
 
 		// Tile data
 		{
@@ -2065,7 +2051,7 @@ void FStreamingManager::RequestDDCData(TConstArrayView<UE::DerivedData::FCacheGe
 	using namespace UE::DerivedData;
 
 	{
-		FRequestOwner* RequestOwnerPtr = bBlocking ? RequestOwnerBlocking : RequestOwner;
+		FRequestOwner* RequestOwnerPtr = bBlocking ? RequestOwnerBlocking.Get() : RequestOwner.Get();
 		FRequestBarrier Barrier(*RequestOwnerPtr);	// This is a critical section on the owner. It does not constrain ordering
 		GetCache().GetChunks(DDCRequests, *RequestOwnerPtr,
 			[this](FCacheGetChunkResponse&& Response)
diff --git a/Engine/Source/Runtime/Engine/Private/SparseVolumeTexture/SparseVolumeTextureStreamingManager.h b/Engine/Source/Runtime/Engine/Private/SparseVolumeTexture/SparseVolumeTextureStreamingManager.h
--- a/Engine/Source/Runtime/Engine/Private/SparseVolumeTexture/SparseVolumeTextureStreamingManager.h
+++ b/Engine/Source/Runtime/Engine/Private/SparseVolumeTexture/SparseVolumeTextureStreamingManager.h
@@ -91,7 +91,7 @@ public:
 	virtual void Update_GameThread() override;
 	
 	virtual void Request(UStreamableSparseVolumeTexture* SparseVolumeTexture, float FrameIndex, int32 MipLevel, bool bBlocking) override;
-	virtual void BeginAsyncUpdate(FRDGBuilder& GraphBuilder, bool bForceNonAsync) override;
+	virtual void BeginAsyncUpdate(FRDGBuilder& GraphBuilder, bool bBlocking) override;
 	virtual void EndAsyncUpdate(FRDGBuilder& GraphBuilder) override;
 	//~ End IStreamingManager Interface.
 
@@ -178,6 +178,7 @@ private:
 		int32 LowestRequestedMipLevel; // Lowest mip level that should be resident. Can be lower than LowestResidentMipLevel when streaming in new mips. Stream-out is instant, so it should never be higher.
 		int32 LowestResidentMipLevel; // Actually resident on the GPU
 		TArray<TArray<uint32>> TileAllocations; // TileAllocations[MipLevel][PhysicalTileIndex]
+		FTextureRHIRef PageTableTextureRHIRef;
 	};
 
 	// Used to keep track of the least-recently-used order of mip levels. This is done per mip level and not per frame, so that lower resolution mip levels are kept in memory for longer.
@@ -205,7 +206,7 @@ private:
 		TArray<TIntrusiveDoubleLinkedList<FLRUNode>> PerMipLRULists; // PerMipLRULists[MipLevel]
 		TArray<FStreamingWindow> StreamingWindows;
 
-		FTileDataTexture* TileDataTexture;
+		TUniquePtr<FTileDataTexture> TileDataTexture;
 		TRefCountPtr<FRDGPooledBuffer> StreamingInfoBuffer; // One uint32 per frame storing the LowestResidentMipLevel
 		FShaderResourceViewRHIRef StreamingInfoBufferSRVRHIRef;
 		TBitArray<> DirtyStreamingInfoData; // One bit per frame, potentially marking the streaming info data as dirty/in need of an update
@@ -322,12 +323,12 @@ private:
 	TMap<FMipLevelKey, uint32> RequestsHashTable;
 	TArray<FPendingMipLevel> PendingMipLevels;
 #if WITH_EDITORONLY_DATA
-	UE::DerivedData::FRequestOwner* RequestOwner = nullptr;
-	UE::DerivedData::FRequestOwner* RequestOwnerBlocking = nullptr;
+	TUniquePtr<UE::DerivedData::FRequestOwner> RequestOwner = nullptr;
+	TUniquePtr<UE::DerivedData::FRequestOwner> RequestOwnerBlocking = nullptr;
 #endif
 
-	class FPageTableUpdater* PageTableUpdater = nullptr;
-	class FStreamingInfoBufferUpdater* StreamingInfoBufferUpdater = nullptr;
+	TUniquePtr<class FPageTableUpdater> PageTableUpdater = nullptr;
+	TUniquePtr<class FStreamingInfoBufferUpdater> StreamingInfoBufferUpdater = nullptr;
 	FGraphEventArray AsyncTaskEvents;
 	FAsyncState AsyncState;
 	int32 MaxPendingMipLevels = 0;
diff --git a/Engine/Source/Runtime/Engine/Public/SparseVolumeTexture/ISparseVolumeTextureStreamingManager.h b/Engine/Source/Runtime/Engine/Public/SparseVolumeTexture/ISparseVolumeTextureStreamingManager.h
--- a/Engine/Source/Runtime/Engine/Public/SparseVolumeTexture/ISparseVolumeTextureStreamingManager.h
+++ b/Engine/Source/Runtime/Engine/Public/SparseVolumeTexture/ISparseVolumeTextureStreamingManager.h
@@ -23,14 +23,17 @@ public:
 	// Request a frame to be streamed in. FrameIndex is of float type so that the fractional part can be used to better track the playback speed/direction.
 	// This function automatically also requests all higher mip levels and adds prefetch requests for upcoming frames.
 	virtual void Request_GameThread(UStreamableSparseVolumeTexture* SparseVolumeTexture, float FrameIndex, int32 MipLevel, bool bBlocking) = 0;
-	// Issues a rendering command for updating the streaming manager. This is not normally needed, but may be necessary in cases where blocking requests are required
-	// and the SVT is used in a different rendering command that is executed before the streaming manager would normally update.
+	// Issues a rendering command for updating the streaming manager. This is not normally needed, but is necessary in cases where blocking requests are required
+	// and the SVT is used in a rendering command (which is executed before the streaming manager would normally update). A call to this function is not needed when the SVT
+	// is used in a regular pass called from FDeferredShadingSceneRenderer::Render().
 	virtual void Update_GameThread() = 0;
 	//~ End game thread functions.
 	
 	//~ Begin rendering thread functions.
 	virtual void Request(UStreamableSparseVolumeTexture* SparseVolumeTexture, float FrameIndex, int32 MipLevel, bool bBlocking) = 0;
-	virtual void BeginAsyncUpdate(FRDGBuilder& GraphBuilder, bool bForceNonAsync = false) = 0;
+	// Begins updating the streaming manager. If r.SparseVolumeTexture.Streaming.AsyncThread is 1 and bBlocking is false, most of the updating work is done in another thread.
+	virtual void BeginAsyncUpdate(FRDGBuilder& GraphBuilder, bool bBlocking = false) = 0;
+	// Waits for the job started in BeginAsyncUpdate() to complete, issues GPU work and does some cleanup.
 	virtual void EndAsyncUpdate(FRDGBuilder& GraphBuilder) = 0;
 	//~ End rendering thread functions.
 
diff --git a/Engine/Source/Runtime/Renderer/Private/SparseVolumeTexture/SparseVolumeTextureViewerRendering.cpp b/Engine/Source/Runtime/Renderer/Private/SparseVolumeTexture/SparseVolumeTextureViewerRendering.cpp
--- a/Engine/Source/Runtime/Renderer/Private/SparseVolumeTexture/SparseVolumeTextureViewerRendering.cpp
+++ b/Engine/Source/Runtime/Renderer/Private/SparseVolumeTexture/SparseVolumeTextureViewerRendering.cpp
@@ -152,28 +152,17 @@ void AddSparseVolumeTextureViewerRenderPass(FRDGBuilder& GraphBuilder, FSceneRen
 			const UE::SVT::FTextureRenderResources* RenderResources = SVTSceneProxy->TextureRenderResources;
 			if (RenderResources)
 			{
-				SVTSceneProxy->TextureRenderResources->GetPackedUniforms(PsPassParameters->PackedSVTUniforms0, PsPassParameters->PackedSVTUniforms1);
+				FRHITexture* PageTableTexture = RenderResources->GetPageTableTexture();
+				FRHITexture* TextureA = RenderResources->GetPhysicalTileDataATexture();
+				FRHITexture* TextureB = RenderResources->GetPhysicalTileDataBTexture();
+				FRHIShaderResourceView* StreamingInfoBufferSRV = RenderResources->GetStreamingInfoBufferSRV();
 
-				FRHITexture* PageTableTexture = RenderResources->GetPageTableTextureRHI();
-				if (PageTableTexture)
-				{
-					PsPassParameters->SparseVolumeTexturePageTable = PageTableTexture;
-				}
-				FRHITexture* TextureA = RenderResources->GetPhysicalTileDataATextureRHI();
-				if (TextureA)
-				{
-					PsPassParameters->SparseVolumeTextureA = TextureA;
-				}
-				FRHITexture* TextureB = RenderResources->GetPhysicalTileDataBTextureRHI();
-				if (TextureB)
-				{
-					PsPassParameters->SparseVolumeTextureB = TextureB;
-				}
-				FRHIShaderResourceView* StreamingInfoBufferSRV = RenderResources->GetStreamingInfoBufferSRVRHI();
-				if (StreamingInfoBufferSRV)
-				{
-					PsPassParameters->StreamingInfoBuffer = StreamingInfoBufferSRV;
-				}
+				PsPassParameters->SparseVolumeTexturePageTable = PageTableTexture ? PageTableTexture : PsPassParameters->SparseVolumeTexturePageTable;
+				PsPassParameters->SparseVolumeTextureA = TextureA ? TextureA : PsPassParameters->SparseVolumeTextureA;
+				PsPassParameters->SparseVolumeTextureB = TextureB ? TextureB : PsPassParameters->SparseVolumeTextureB;
+				PsPassParameters->StreamingInfoBuffer = StreamingInfoBufferSRV ? StreamingInfoBufferSRV : PsPassParameters->StreamingInfoBuffer;
+
+				SVTSceneProxy->TextureRenderResources->GetPackedUniforms(PsPassParameters->PackedSVTUniforms0, PsPassParameters->PackedSVTUniforms1);
 			}
 
 			ClearUnusedGraphResources(PixelShader, PsPassParameters);
