Fix missing ConditionalPostLoad on DataModelInterface in UAnimSequenceBase

#rnx

diff --git a/Engine/Source/Runtime/Engine/Private/Animation/AnimSequenceBase.cpp b/Engine/Source/Runtime/Engine/Private/Animation/AnimSequenceBase.cpp
--- a/Engine/Source/Runtime/Engine/Private/Animation/AnimSequenceBase.cpp
+++ b/Engine/Source/Runtime/Engine/Private/Animation/AnimSequenceBase.cpp
@@ -136,6 +136,13 @@ void UAnimSequenceBase::PostLoad()
 		    }
 
 		    ValidateModel();
+			if (UObject* DataModelObject = DataModelInterface.GetObject())
+			{
+				if (DataModelObject->HasAnyFlags(RF_NeedPostLoad | RF_NeedPostLoadSubobjects))
+				{
+					DataModelObject->ConditionalPostLoad();
+				}
+			}
 		    GetController();
 		    BindToModelModificationEvent();
 
