GPU Skin Cache update optimizations
*Calling FRWBufferTracker::Advance was followed by FRWBufferTracker::Find, which basically searched for exactly the thing that was written by Advance.
*Advance can just return what Find would have provided.
*Push buffers into TArray, rather than TSet (also not using TArray::AddUnique).
*When making state transitions, can check the state tracker to discard redundant ones.
*Added a marker to avoid pushing duplicates in GetBufferUAVs (tried using TSet in there, but it adds some cost).
*FGPUSkinCache::DoDispatch improvements (batched version).
*Brought DispatchUpdateSkinning inline into DoDispatch. Allows for better batching, reducing duplicate work.
*Make a sorted index for BatchDispatches, sorted by the shader that would be used for positions.
*Set shader only when it changes, instead of for every dispatch. Shaders are pulled into a flat array, instead of building TSet of them.
*FGPUSkinCache::DoDispatch improvements (non-batched version).
*Don't need to build TMap of shaders, just pick the one we're looking for.
*Set shader once, instead of separately for Previous and Current positions.
*Don't need to construct FString for [RT] tag (only used for debug stuff anyway).

diff --git a/Engine/Source/Runtime/Engine/Private/GPUSkinCache.cpp b/Engine/Source/Runtime/Engine/Private/GPUSkinCache.cpp
--- a/Engine/Source/Runtime/Engine/Private/GPUSkinCache.cpp
+++ b/Engine/Source/Runtime/Engine/Private/GPUSkinCache.cpp
@@ -799,13 +799,12 @@ public:
 	}
 };
 
-#define SKIN_CACHE_SHADER(_WEIGHT16_, _INDEX16_, _INFLUENCE_TYPE_, _SKIN_TYPE_) \
-	IMPLEMENT_SHADER_TYPE(template<>, TGPUSkinCacheCS<_WEIGHT16_ * 32 + _INDEX16_ * 16 + _INFLUENCE_TYPE_ * 4 + _SKIN_TYPE_>, TEXT("/Engine/Private/GpuSkinCacheComputeShader.usf"), TEXT("SkinCacheUpdateBatchCS"), SF_Compute) 
+#define SKIN_CACHE_SHADER_IDX(_WEIGHT16_, _INDEX16_, _INFLUENCE_TYPE_, _SKIN_TYPE_) (_WEIGHT16_ * 32 + _INDEX16_ * 16 + _INFLUENCE_TYPE_ * 4 + _SKIN_TYPE_)
 
 #define SKIN_CACHE_SHADER_ALL_SKIN_TYPES(_WEIGHT16_, _INDEX16_, _INFLUENCE_TYPE_) \
-	SKIN_CACHE_SHADER(_WEIGHT16_, _INDEX16_, _INFLUENCE_TYPE_, 0); \
-	SKIN_CACHE_SHADER(_WEIGHT16_, _INDEX16_, _INFLUENCE_TYPE_, 1); \
-	SKIN_CACHE_SHADER(_WEIGHT16_, _INDEX16_, _INFLUENCE_TYPE_, 2)
+	SKIN_CACHE_SHADER(SKIN_CACHE_SHADER_IDX(_WEIGHT16_, _INDEX16_, _INFLUENCE_TYPE_, 0)); \
+	SKIN_CACHE_SHADER(SKIN_CACHE_SHADER_IDX(_WEIGHT16_, _INDEX16_, _INFLUENCE_TYPE_, 1)); \
+	SKIN_CACHE_SHADER(SKIN_CACHE_SHADER_IDX(_WEIGHT16_, _INDEX16_, _INFLUENCE_TYPE_, 2))
 
 // NOTE: Bone influence type 2 (multiple) does not require a 16-bit index or weight permutation.
 #define SKIN_CACHE_SHADER_ALL() \
@@ -819,8 +818,9 @@ public:
 	SKIN_CACHE_SHADER_ALL_SKIN_TYPES(1, 1, 1); \
 	SKIN_CACHE_SHADER_ALL_SKIN_TYPES(0, 0, 2);
 
+#define SKIN_CACHE_SHADER(_SHADER_IDX_) \
+	IMPLEMENT_SHADER_TYPE(template<>, TGPUSkinCacheCS<_SHADER_IDX_>, TEXT("/Engine/Private/GpuSkinCacheComputeShader.usf"), TEXT("SkinCacheUpdateBatchCS"), SF_Compute) 
 SKIN_CACHE_SHADER_ALL()
-
 #undef SKIN_CACHE_SHADER
 
 FGPUSkinCache::FGPUSkinCache(ERHIFeatureLevel::Type InFeatureLevel, bool bInRequiresMemoryLimit, UWorld* InWorld)
@@ -854,7 +854,7 @@ void FGPUSkinCache::Cleanup()
 	ensure(Allocations.Num() == 0);
 }
 
-void FGPUSkinCache::TransitionAllToReadable(FRHICommandList& RHICmdList, const TSet<FSkinCacheRWBuffer*>& BuffersToTransitionToRead)
+void FGPUSkinCache::TransitionAllToReadable(FRHICommandList& RHICmdList, const TArray<FSkinCacheRWBuffer*>& BuffersToTransitionToRead)
 {
 	TRACE_CPUPROFILER_EVENT_SCOPE(FGPUSkinCache::TransitionAllToReadable);
 
@@ -862,9 +862,14 @@ void FGPUSkinCache::TransitionAllToReadable(FRHICommandList& RHICmdList, const T
 	{
 		TArray<FRHITransitionInfo, SceneRenderingAllocator> UAVs;
 		UAVs.Reserve(BuffersToTransitionToRead.Num());
-		for (TSet<FSkinCacheRWBuffer*>::TConstIterator SetIt(BuffersToTransitionToRead); SetIt; ++SetIt)
+		for (TArray<FSkinCacheRWBuffer*>::TConstIterator SetIt(BuffersToTransitionToRead); SetIt; ++SetIt)
 		{
-			UAVs.Add((*SetIt)->UpdateAccessState(ERHIAccess::VertexOrIndexBuffer | ERHIAccess::SRVMask));
+			FSkinCacheRWBuffer* Buffer = *SetIt;
+			constexpr ERHIAccess ToState = ERHIAccess::VertexOrIndexBuffer | ERHIAccess::SRVMask;
+			if (Buffer->AccessState != ToState)
+			{
+				UAVs.Add(Buffer->UpdateAccessState(ToState));
+			}
 		}
 		RHICmdList.Transition(UAVs);
 	}
@@ -1260,7 +1265,10 @@ void FGPUSkinCache::MakeBufferTransitions(FRHICommandListImmediate& RHICmdList,
 		UAVs.Reserve(Buffers.Num());
 		for (FSkinCacheRWBuffer* Buffer : Buffers)
 		{
-			UAVs.Add(Buffer->UpdateAccessState(ToState));
+			if (Buffer->AccessState != ToState)
+			{
+				UAVs.Add(Buffer->UpdateAccessState(ToState));
+			}
 		}
 		RHICmdList.Transition(MakeArrayView(UAVs.GetData(), UAVs.Num()));
 	}
@@ -1269,9 +1277,28 @@ void FGPUSkinCache::MakeBufferTransitions(FRHICommandListImmediate& RHICmdList,
 void FGPUSkinCache::GetBufferUAVs(const TArray<FSkinCacheRWBuffer*>& InBuffers, TArray<FRHIUnorderedAccessView*>& OutUAVs)
 {
 	OutUAVs.Reset(InBuffers.Num());
+
+	// It looks like BeginUAVOverlap wants to get in unique buffers, without any overlaps. Previously that worked out
+	// because InBuffers was being filled with AddUnique, but that can add quite a bit of cost when there are many
+	// skinned meshes to process here.
+	//
+	// Using a TSet instead of TArray with AddUnique improves things, but it still came up slower than just pushing
+	// everything and dealing with duplicates separately.
+	//
+	// So here I've added a member to FSkinCacheRWBuffer that we can just use to check if it's already been visited. On
+	// 64-bit, it's not increasing the struct size (since there's a trailing ERHIAccess/uint32 already, it gets padded
+	// up to an 8 byte multiple).
 	for (const FSkinCacheRWBuffer* Buffer : InBuffers)
 	{
-		OutUAVs.Add(Buffer->Buffer.UAV);
+		Buffer->UniqueOpToken = 0;
+	}
+	for (const FSkinCacheRWBuffer* Buffer : InBuffers)
+	{
+		if (Buffer->UniqueOpToken == 0)
+		{
+			Buffer->UniqueOpToken = 1;
+			OutUAVs.Add(Buffer->Buffer.UAV);
+		}
 	}
 }
 
@@ -1289,25 +1316,171 @@ void FGPUSkinCache::DoDispatch(FRHICommandListImmediate& RHICmdList)
 	TArray<FSkinCacheRWBuffer*> BuffersToTransitionForSkinning;
 	BuffersToTransitionForSkinning.Reserve(BatchCount * 2);
 	{
+		TRACE_CPUPROFILER_EVENT_SCOPE(GPUSkinCache_PrepareUpdateSkinning);
+
 		for (int32 i = 0; i < BatchCount; ++i)
 		{
 			FDispatchEntry& DispatchItem = BatchDispatches[i];
 			PrepareUpdateSkinning(DispatchItem.SkinCacheEntry, DispatchItem.Section, DispatchItem.RevisionNumber, &BuffersToTransitionForSkinning);
 		}
-		MakeBufferTransitions(RHICmdList, BuffersToTransitionForSkinning, ERHIAccess::UAVCompute);
+
+		{
+			TRACE_CPUPROFILER_EVENT_SCOPE(MakeBufferTransitions);
+			MakeBufferTransitions(RHICmdList, BuffersToTransitionForSkinning, ERHIAccess::UAVCompute);
+		}
 	}
 
-	TSet<FSkinCacheRWBuffer*> BuffersToTransitionToRead;
+	TArray<FSkinCacheRWBuffer*> BuffersToTransitionToRead;
 
 	TArray<FRHIUnorderedAccessView*> SkinningBuffersToOverlap;
 	GetBufferUAVs(BuffersToTransitionForSkinning, SkinningBuffersToOverlap);
 	RHICmdList.BeginUAVOverlap(SkinningBuffersToOverlap);
 	{
+		TRACE_CPUPROFILER_EVENT_SCOPE(GPUSkinCache_UpdateSkinningBatches);
 		SCOPED_DRAW_EVENT(RHICmdList, GPUSkinCache_UpdateSkinningBatches);
-		for (int32 i = 0; i < BatchCount; ++i)
+
+		auto* GlobalShaderMap = GetGlobalShaderMap(GetFeatureLevel());
+
+		TArray<FSortedDispatchEntry> SortedDispatches;
+		SortedDispatches.Reserve(BatchCount);
 		{
-			FDispatchEntry& DispatchItem = BatchDispatches[i];
-			DispatchUpdateSkinning(RHICmdList, DispatchItem.SkinCacheEntry, DispatchItem.Section, DispatchItem.RevisionNumber, BuffersToTransitionToRead);
+			TRACE_CPUPROFILER_EVENT_SCOPE(BuildSortedDispatchList);
+
+			for (int32 i = 0; i < BatchCount; ++i)
+			{
+				FDispatchEntry& DispatchItem = BatchDispatches[i];
+				FGPUSkinCacheEntry* Entry = DispatchItem.SkinCacheEntry;
+				int32 Section = DispatchItem.Section;
+				FGPUSkinCacheEntry::FSectionDispatchData& DispatchData = Entry->DispatchData[Section];
+
+				if ((DispatchData.DispatchFlags & ((uint32)EGPUSkinCacheDispatchFlags::DispatchPrevPosition | (uint32)EGPUSkinCacheDispatchFlags::DispatchPosition)) != 0)
+				{
+					// For 'unlimited' bone indexes, we pass in the index and weight sizes via a shader parameter and so we
+					// can re-use the same shader permutation as for 8-bit indexes.
+					bool bUse16BitBoneIndex = Entry->bUse16BitBoneIndex;
+					bool bUse16BitBoneWeight = Entry->bUse16BitBoneWeight;
+					if (Entry->BoneInfluenceType == 2)
+					{
+						bUse16BitBoneIndex = bUse16BitBoneWeight = false;
+					}
+
+					FSortedDispatchEntry SortedEntry;
+
+					SortedEntry.ShaderIndex = SKIN_CACHE_SHADER_IDX(
+						static_cast<int32>(bUse16BitBoneWeight),
+						static_cast<int32>(bUse16BitBoneIndex),
+						Entry->BoneInfluenceType,
+						DispatchData.SkinType);
+					SortedEntry.BatchIndex = i;
+
+					SortedDispatches.Add(SortedEntry);
+				}
+			}
+
+			Algo::Sort(SortedDispatches,
+				[](const FSortedDispatchEntry& A, const FSortedDispatchEntry& B)
+				{
+					if (A.ShaderIndex != B.ShaderIndex)
+					{
+						return A.ShaderIndex < B.ShaderIndex;
+					}
+			return A.BatchIndex < B.BatchIndex;
+				});
+		}
+
+		int32 LastShaderIndex = -1;
+		TShaderRef<FBaseGPUSkinCacheCS> Shader;
+
+		TShaderRef<FBaseGPUSkinCacheCS> AllShaders[64];
+#define SKIN_CACHE_SHADER(_SHADER_IDX_) static_assert(_SHADER_IDX_ < 64); AllShaders[_SHADER_IDX_] = TShaderMapRef<TGPUSkinCacheCS<_SHADER_IDX_>>(GlobalShaderMap);
+		SKIN_CACHE_SHADER_ALL();
+#undef SKIN_CACHE_SHADER
+
+		int32 SortedCount = SortedDispatches.Num();
+		for (int32 i = 0; i < SortedCount; ++i)
+		{
+			const FSortedDispatchEntry& SortedEntry = SortedDispatches[i];
+			if (SortedEntry.ShaderIndex != LastShaderIndex)
+			{
+				TRACE_CPUPROFILER_EVENT_SCOPE(ChangeShader);
+
+				if (Shader.IsValid())
+				{
+					UnsetShaderParametersLegacyCS(RHICmdList, Shader);
+				}
+
+				LastShaderIndex = SortedEntry.ShaderIndex;
+				Shader = AllShaders[SortedEntry.ShaderIndex];
+
+				check(Shader.IsValid());
+
+				SetComputePipelineState(RHICmdList, Shader.GetComputeShader());
+			}
+
+			// This is pulled from FGPUSkinCache::DispatchUpdateSkinning() below, but inlined so we can set the
+			// shader only when it changes. Not sure if it's worth pulling out the common bits into a shared function.
+
+			FDispatchEntry& DispatchEntry = BatchDispatches[SortedEntry.BatchIndex];
+			FGPUSkinCacheEntry* Entry = DispatchEntry.SkinCacheEntry;
+			FGPUSkinCacheEntry::FSectionDispatchData& DispatchData = Entry->DispatchData[DispatchEntry.Section];
+			FGPUBaseSkinVertexFactory::FShaderDataType& ShaderData = DispatchData.SourceVertexFactory->GetShaderData();
+			const TCHAR* RayTracingTag = (Entry->Mode == EGPUSkinCacheEntryMode::RayTracing ? TEXT("[RT]") : TEXT(""));
+
+			SCOPED_DRAW_EVENTF(RHICmdList, SkinCacheDispatch,
+				TEXT("%sSkinning%d%d%d%d Mesh=%s LOD=%d Chunk=%d InStreamStart=%d OutStart=%d Vert=%d Morph=%d/%d"),
+				RayTracingTag, (int32)Entry->bUse16BitBoneIndex, (int32)Entry->bUse16BitBoneWeight, (int32)Entry->BoneInfluenceType, DispatchData.SkinType, *GetSkeletalMeshObjectName(Entry->GPUSkin), Entry->LOD,
+				DispatchData.SectionIndex, DispatchData.InputStreamStart, DispatchData.OutputStreamStart, DispatchData.NumVertices, Entry->MorphBuffer != 0, DispatchData.MorphBufferOffset);
+
+			uint32 VertexCountAlign64 = FMath::DivideAndRoundUp(DispatchData.NumVertices, (uint32)64);
+
+			if ((DispatchData.DispatchFlags & (uint32)EGPUSkinCacheDispatchFlags::DispatchPrevPosition) != 0)
+			{
+				const FVertexBufferAndSRV& PrevBoneBuffer = ShaderData.GetBoneBufferForReading(true);
+
+				SetShaderParametersLegacyCS(
+					RHICmdList,
+					Shader,
+					PrevBoneBuffer,
+					Entry,
+					DispatchData,
+					DispatchData.GetPreviousPositionRWBuffer()->Buffer.UAV,
+					DispatchData.GetActiveTangentRWBuffer() ? DispatchData.GetActiveTangentRWBuffer()->Buffer.UAV.GetReference() : nullptr
+				);
+
+				INC_DWORD_STAT_BY(STAT_GPUSkinCache_TotalNumVertices, VertexCountAlign64 * 64);
+				RHICmdList.DispatchComputeShader(VertexCountAlign64, 1, 1);
+
+				IncrementDispatchCounter(RHICmdList);
+				BuffersToTransitionToRead.Add(DispatchData.GetPreviousPositionRWBuffer());
+			}
+
+			if ((DispatchData.DispatchFlags & (uint32)EGPUSkinCacheDispatchFlags::DispatchPosition) != 0)
+			{
+				const FVertexBufferAndSRV& BoneBuffer = ShaderData.GetBoneBufferForReading(false);
+
+				SetShaderParametersLegacyCS(
+					RHICmdList,
+					Shader,
+					BoneBuffer,
+					Entry,
+					DispatchData,
+					DispatchData.GetPositionRWBuffer()->Buffer.UAV,
+					DispatchData.GetActiveTangentRWBuffer() ? DispatchData.GetActiveTangentRWBuffer()->Buffer.UAV.GetReference() : nullptr
+				);
+
+				INC_DWORD_STAT_BY(STAT_GPUSkinCache_TotalNumVertices, VertexCountAlign64 * 64);
+				RHICmdList.DispatchComputeShader(VertexCountAlign64, 1, 1);
+
+				IncrementDispatchCounter(RHICmdList);
+				BuffersToTransitionToRead.Add(DispatchData.GetPositionRWBuffer());
+			}
+
+			check(DispatchData.PreviousPositionBuffer != DispatchData.PositionBuffer);
+		}
+
+		if (Shader.IsValid())
+		{
+			UnsetShaderParametersLegacyCS(RHICmdList, Shader);
 		}
 	}
 	RHICmdList.EndUAVOverlap(SkinningBuffersToOverlap);
@@ -1321,11 +1494,11 @@ void FGPUSkinCache::DoDispatch(FRHICommandListImmediate& RHICmdList)
 		FGPUSkinCacheEntry::FSectionDispatchData& DispatchData = DispatchItem.SkinCacheEntry->DispatchData[DispatchItem.Section];
 		if (DispatchData.IndexBuffer)
 		{
-			BuffersToSRVForRecomputeTangents.AddUnique(DispatchData.GetPositionRWBuffer());
-			BuffersToSRVForRecomputeTangents.AddUnique(DispatchData.GetActiveTangentRWBuffer());
+			BuffersToSRVForRecomputeTangents.Add(DispatchData.GetPositionRWBuffer());
+			BuffersToSRVForRecomputeTangents.Add(DispatchData.GetActiveTangentRWBuffer());
 			if (GRecomputeTangentsParallelDispatch)
 			{
-				IntermediateAccumulatedTangentBuffers.AddUnique(DispatchData.GetIntermediateAccumulatedTangentBuffer());
+				IntermediateAccumulatedTangentBuffers.Add(DispatchData.GetIntermediateAccumulatedTangentBuffer());
 			}
 			BuffersToTransitionToRead.Add(DispatchData.GetPositionRWBuffer());
 		}	
@@ -1360,7 +1533,8 @@ void FGPUSkinCache::DoDispatch(FRHICommandListImmediate& RHICmdList)
 			{
 				FDispatchEntry& DispatchItem = BatchDispatches[i];
 				FGPUSkinCacheEntry::FSectionDispatchData& DispatchData = DispatchItem.SkinCacheEntry->DispatchData[DispatchItem.Section];
-				TangentBuffers.AddUnique(DispatchData.GetTangentRWBuffer());
+				TangentBuffers.Add(DispatchData.GetTangentRWBuffer());
+				BuffersToTransitionToRead.Add(DispatchData.GetTangentRWBuffer());
 			}
 			MakeBufferTransitions(RHICmdList, TangentBuffers, ERHIAccess::UAVCompute);
 			MakeBufferTransitions(RHICmdList, IntermediateAccumulatedTangentBuffers, ERHIAccess::UAVCompute);
@@ -1381,13 +1555,20 @@ void FGPUSkinCache::DoDispatch(FRHICommandListImmediate& RHICmdList)
 	}
 	RHICmdList.EndUAVOverlap(IntermediateAccumulatedTangentBuffersToOverlap);
 
-	for (int32 i = 0; i < BatchCount; ++i)
 	{
-		FDispatchEntry& DispatchItem = BatchDispatches[i];
-		DispatchItem.SkinCacheEntry->UpdateVertexFactoryDeclaration(DispatchItem.Section);
+		TRACE_CPUPROFILER_EVENT_SCOPE(UpdateVertexFactoryDeclarations);
+
+		for (int32 i = 0; i < BatchCount; ++i)
+		{
+			FDispatchEntry& DispatchItem = BatchDispatches[i];
+			DispatchItem.SkinCacheEntry->UpdateVertexFactoryDeclaration(DispatchItem.Section);
+		}
 	}
 
-	TransitionAllToReadable(RHICmdList, BuffersToTransitionToRead);
+	{
+		TRACE_CPUPROFILER_EVENT_SCOPE(TransitionAllToReadable);
+		TransitionAllToReadable(RHICmdList, BuffersToTransitionToRead);
+	}
 }
 
 void FGPUSkinCache::DoDispatch(FRHICommandListImmediate& RHICmdList, FGPUSkinCacheEntry* SkinCacheEntry, int32 Section, int32 RevisionNumber)
@@ -1399,7 +1580,7 @@ void FGPUSkinCache::DoDispatch(FRHICommandListImmediate& RHICmdList, FGPUSkinCac
 
 	INC_DWORD_STAT(STAT_GPUSkinCache_TotalNumChunks);
 
-	TSet<FSkinCacheRWBuffer*> BuffersToTransitionToRead;
+	TArray<FSkinCacheRWBuffer*> BuffersToTransitionToRead;
 
 	TArray<FSkinCacheRWBuffer*> BuffersToTransitionForSkinning;
 	PrepareUpdateSkinning(SkinCacheEntry, Section, RevisionNumber, &BuffersToTransitionForSkinning);
@@ -1787,15 +1968,14 @@ void FGPUSkinCache::PrepareUpdateSkinning(FGPUSkinCacheEntry* Entry, int32 Secti
 		PositionBuffer = DispatchData.PositionTracker.Find(BoneBuffer, Revision);
 		if (!PositionBuffer)
 		{
-			DispatchData.PositionTracker.Advance(BoneBuffer, Revision, PrevBoneBuffer, PrevRevision);
-			PositionBuffer = DispatchData.PositionTracker.Find(BoneBuffer, Revision);
+			PositionBuffer = DispatchData.PositionTracker.Advance(BoneBuffer, Revision, PrevBoneBuffer, PrevRevision);
 			check(PositionBuffer);
 
 			DispatchData.DispatchFlags |= UpdateFlag;
 
 			if (OverlappedUAVs)
 			{
-				(*OverlappedUAVs).AddUnique(PositionBuffer);
+				(*OverlappedUAVs).Add(PositionBuffer);
 			}
 		}
 	};
@@ -1824,34 +2004,24 @@ void FGPUSkinCache::PrepareUpdateSkinning(FGPUSkinCacheEntry* Entry, int32 Secti
 
 	if (OverlappedUAVs && DispatchData.DispatchFlags != 0 && DispatchData.GetActiveTangentRWBuffer())
 	{
-		 (*OverlappedUAVs).AddUnique(DispatchData.GetActiveTangentRWBuffer());
+		(*OverlappedUAVs).Add(DispatchData.GetActiveTangentRWBuffer());
 	}
 
 	check(DispatchData.PreviousPositionBuffer != DispatchData.PositionBuffer);
 }
 
-void FGPUSkinCache::DispatchUpdateSkinning(FRHICommandListImmediate& RHICmdList, FGPUSkinCacheEntry* Entry, int32 Section, uint32 RevisionNumber, TSet<FSkinCacheRWBuffer*>& BuffersToTransitionToRead)
+void FGPUSkinCache::DispatchUpdateSkinning(FRHICommandListImmediate& RHICmdList, FGPUSkinCacheEntry* Entry, int32 Section, uint32 RevisionNumber, TArray<FSkinCacheRWBuffer*>& BuffersToTransitionToRead)
 {
 	FGPUSkinCacheEntry::FSectionDispatchData& DispatchData = Entry->DispatchData[Section];
 	FGPUBaseSkinVertexFactory::FShaderDataType& ShaderData = DispatchData.SourceVertexFactory->GetShaderData();
-	const FString RayTracingTag = (Entry->Mode == EGPUSkinCacheEntryMode::RayTracing ? TEXT("[RT]") : TEXT(""));
-	
+	const TCHAR* RayTracingTag = (Entry->Mode == EGPUSkinCacheEntryMode::RayTracing ? TEXT("[RT]") : TEXT(""));
+
 	SCOPED_DRAW_EVENTF(RHICmdList, SkinCacheDispatch,
 		TEXT("%sSkinning%d%d%d%d Mesh=%s LOD=%d Chunk=%d InStreamStart=%d OutStart=%d Vert=%d Morph=%d/%d"),
-		*RayTracingTag, (int32)Entry->bUse16BitBoneIndex, (int32)Entry->bUse16BitBoneWeight, (int32)Entry->BoneInfluenceType, DispatchData.SkinType, *GetSkeletalMeshObjectName(Entry->GPUSkin), Entry->LOD,
+		RayTracingTag, (int32)Entry->bUse16BitBoneIndex, (int32)Entry->bUse16BitBoneWeight, (int32)Entry->BoneInfluenceType, DispatchData.SkinType, *GetSkeletalMeshObjectName(Entry->GPUSkin), Entry->LOD,
 		DispatchData.SectionIndex, DispatchData.InputStreamStart, DispatchData.OutputStreamStart, DispatchData.NumVertices, Entry->MorphBuffer != 0, DispatchData.MorphBufferOffset);
 	auto* GlobalShaderMap = GetGlobalShaderMap(GetFeatureLevel());
 
-	TMap<int32, TShaderRef<FBaseGPUSkinCacheCS>> SkinCacheCSMap;
-
-#define SKIN_CACHE_SHADER(_WEIGHT16_, _INDEX16_, _INFLUENCE_TYPE_, _SKIN_TYPE_) \
-	{ \
-		constexpr int32 Index = _WEIGHT16_ * 32 + _INDEX16_ * 16 + _INFLUENCE_TYPE_ * 4 + _SKIN_TYPE_; \
-		SkinCacheCSMap.Add(Index, TShaderMapRef<TGPUSkinCacheCS<Index>>(GlobalShaderMap)); \
-	}
-
-	SKIN_CACHE_SHADER_ALL()
-
 	// For 'unlimited' bone indexes, we pass in the index and weight sizes via a shader parameter and so we
 	// can re-use the same shader permutation as for 8-bit indexes.
 	bool bUse16BitBoneIndex = Entry->bUse16BitBoneIndex;
@@ -1861,69 +2031,71 @@ void FGPUSkinCache::DispatchUpdateSkinning(FRHICommandListImmediate& RHICmdList,
 		bUse16BitBoneIndex = bUse16BitBoneWeight = false;
 	}
 
+	int32 ShaderIndex = SKIN_CACHE_SHADER_IDX(
+		static_cast<int32>(bUse16BitBoneWeight),
+		static_cast<int32>(bUse16BitBoneIndex),
+		Entry->BoneInfluenceType,
+		DispatchData.SkinType);
+
+	TShaderRef<FBaseGPUSkinCacheCS> Shader;
+	switch (ShaderIndex)
+	{
+#define SKIN_CACHE_SHADER(_SHADER_IDX_) case _SHADER_IDX_: Shader = TShaderMapRef<TGPUSkinCacheCS<_SHADER_IDX_>>(GlobalShaderMap); break;
+		SKIN_CACHE_SHADER_ALL()
 #undef SKIN_CACHE_SHADER
-	int32 ShaderIndex =
-		DispatchData.SkinType |
-		(Entry->BoneInfluenceType * 4) |
-		(static_cast<int32>(bUse16BitBoneIndex) * 16) |
-		(static_cast<int32>(bUse16BitBoneWeight) * 32) ;
+	}
 
-	check(SkinCacheCSMap.Contains(ShaderIndex));
-	
-	TShaderRef<FBaseGPUSkinCacheCS> Shader = SkinCacheCSMap[ShaderIndex];
 	check(Shader.IsValid());
 
-	const FVertexBufferAndSRV& BoneBuffer = ShaderData.GetBoneBufferForReading(false);
-	const FVertexBufferAndSRV& PrevBoneBuffer = ShaderData.GetBoneBufferForReading(true);
-
-	uint32 CurrentRevision = ShaderData.GetRevisionNumber(false);
-	uint32 PreviousRevision = ShaderData.GetRevisionNumber(true);
-
-	if ((DispatchData.DispatchFlags & (uint32)EGPUSkinCacheDispatchFlags::DispatchPrevPosition) != 0)
+	if ((DispatchData.DispatchFlags & ((uint32)EGPUSkinCacheDispatchFlags::DispatchPrevPosition | (uint32)EGPUSkinCacheDispatchFlags::DispatchPosition)) != 0)
 	{
 		SetComputePipelineState(RHICmdList, Shader.GetComputeShader());
+		uint32 VertexCountAlign64 = FMath::DivideAndRoundUp(DispatchData.NumVertices, (uint32)64);
 
-		SetShaderParametersLegacyCS(
-			RHICmdList,
-			Shader,
-			PrevBoneBuffer,
-			Entry,
-			DispatchData,
-			DispatchData.GetPreviousPositionRWBuffer()->Buffer.UAV,
-			DispatchData.GetActiveTangentRWBuffer() ? DispatchData.GetActiveTangentRWBuffer()->Buffer.UAV : nullptr
-		);
+		if ((DispatchData.DispatchFlags & (uint32)EGPUSkinCacheDispatchFlags::DispatchPrevPosition) != 0)
+		{
+			const FVertexBufferAndSRV& PrevBoneBuffer = ShaderData.GetBoneBufferForReading(true);
 
-		uint32 VertexCountAlign64 = FMath::DivideAndRoundUp(DispatchData.NumVertices, (uint32)64);
-		INC_DWORD_STAT_BY(STAT_GPUSkinCache_TotalNumVertices, VertexCountAlign64 * 64);
-		RHICmdList.DispatchComputeShader(VertexCountAlign64, 1, 1);
-		UnsetShaderParametersLegacyCS(RHICmdList, Shader);
+			SetShaderParametersLegacyCS(
+				RHICmdList,
+				Shader,
+				PrevBoneBuffer,
+				Entry,
+				DispatchData,
+				DispatchData.GetPreviousPositionRWBuffer()->Buffer.UAV,
+				DispatchData.GetActiveTangentRWBuffer() ? DispatchData.GetActiveTangentRWBuffer()->Buffer.UAV : nullptr
+			);
 
-		IncrementDispatchCounter(RHICmdList);
-		BuffersToTransitionToRead.Add(DispatchData.GetPreviousPositionRWBuffer());
-	}
+			INC_DWORD_STAT_BY(STAT_GPUSkinCache_TotalNumVertices, VertexCountAlign64 * 64);
+			RHICmdList.DispatchComputeShader(VertexCountAlign64, 1, 1);
 
-	if ((DispatchData.DispatchFlags & (uint32)EGPUSkinCacheDispatchFlags::DispatchPosition) != 0)
-	{
-		SetComputePipelineState(RHICmdList, Shader.GetComputeShader());
+			IncrementDispatchCounter(RHICmdList);
+			BuffersToTransitionToRead.Add(DispatchData.GetPreviousPositionRWBuffer());
+		}
 
-		SetShaderParametersLegacyCS(
-			RHICmdList,
-			Shader,
-			BoneBuffer,
-			Entry,
-			DispatchData,
-			DispatchData.GetPositionRWBuffer()->Buffer.UAV,
-			DispatchData.GetActiveTangentRWBuffer() ? DispatchData.GetActiveTangentRWBuffer()->Buffer.UAV : nullptr
-		);
+		if ((DispatchData.DispatchFlags & (uint32)EGPUSkinCacheDispatchFlags::DispatchPosition) != 0)
+		{
+			const FVertexBufferAndSRV& BoneBuffer = ShaderData.GetBoneBufferForReading(false);
 
-		uint32 VertexCountAlign64 = FMath::DivideAndRoundUp(DispatchData.NumVertices, (uint32)64);
-		INC_DWORD_STAT_BY(STAT_GPUSkinCache_TotalNumVertices, VertexCountAlign64 * 64);
-		RHICmdList.DispatchComputeShader(VertexCountAlign64, 1, 1);
+			SetShaderParametersLegacyCS(
+				RHICmdList,
+				Shader,
+				BoneBuffer,
+				Entry,
+				DispatchData,
+				DispatchData.GetPositionRWBuffer()->Buffer.UAV,
+				DispatchData.GetActiveTangentRWBuffer() ? DispatchData.GetActiveTangentRWBuffer()->Buffer.UAV : nullptr
+			);
 
-		UnsetShaderParametersLegacyCS(RHICmdList, Shader);
+			INC_DWORD_STAT_BY(STAT_GPUSkinCache_TotalNumVertices, VertexCountAlign64 * 64);
+			RHICmdList.DispatchComputeShader(VertexCountAlign64, 1, 1);
 
-		IncrementDispatchCounter(RHICmdList);
-		BuffersToTransitionToRead.Add(DispatchData.GetPositionRWBuffer());
+
+			IncrementDispatchCounter(RHICmdList);
+			BuffersToTransitionToRead.Add(DispatchData.GetPositionRWBuffer());
+		}
+
+		UnsetShaderParametersLegacyCS(RHICmdList, Shader);
 	}
 
 	BuffersToTransitionToRead.Add(DispatchData.GetTangentRWBuffer());
diff --git a/Engine/Source/Runtime/Engine/Public/GPUSkinCache.h b/Engine/Source/Runtime/Engine/Public/GPUSkinCache.h
--- a/Engine/Source/Runtime/Engine/Public/GPUSkinCache.h
+++ b/Engine/Source/Runtime/Engine/Public/GPUSkinCache.h
@@ -122,6 +122,12 @@ public:
 		uint32 Section = 0;	
 	};
 
+	struct FSortedDispatchEntry
+	{
+		int32 ShaderIndex;
+		int32 BatchIndex;
+	};
+
 	FGPUSkinCache() = delete;
 	ENGINE_API FGPUSkinCache(ERHIFeatureLevel::Type InFeatureLevel, bool bInRequiresMemoryLimit, UWorld* InWorld);
 	ENGINE_API ~FGPUSkinCache();
@@ -174,6 +180,8 @@ public:
 	{
 		FRWBuffer	Buffer;
 		ERHIAccess	AccessState = ERHIAccess::Unknown;	// Keep track of current access state
+		// See note in FGPUSkinCache::GetBufferUAVs()
+		mutable uint32	UniqueOpToken = 0;
 
 		void Release()
 		{
@@ -362,8 +370,9 @@ public:
 			return Allocation ? Allocation->GetIntermediateAccumulatedTangentBuffer() : nullptr;
 		}
 
-		void Advance(const FVertexBufferAndSRV& BoneBuffer1, uint32 Revision1, const FVertexBufferAndSRV& BoneBuffer2, uint32 Revision2)
+		FSkinCacheRWBuffer* Advance(const FVertexBufferAndSRV& BoneBuffer1, uint32 Revision1, const FVertexBufferAndSRV& BoneBuffer2, uint32 Revision2)
 		{
+			FSkinCacheRWBuffer* Result = nullptr;
 			const FVertexBufferAndSRV* InBoneBuffers[2] = { &BoneBuffer1 , &BoneBuffer2 };
 			uint32 InRevisions[2] = { Revision1 , Revision2 };
 
@@ -374,6 +383,10 @@ public:
 				{
 					if (Revisions[Index] == InRevisions[i] && BoneBuffers[Index] == InBoneBuffers[i])
 					{
+						if (i == 0)
+						{
+							Result = &Allocation->PositionBuffers[Index];
+						}
 						Needed = true;
 					}
 				}
@@ -382,9 +395,11 @@ public:
 				{
 					Revisions[Index] = Revision1;
 					BoneBuffers[Index] = &BoneBuffer1;
+					Result = &Allocation->PositionBuffers[Index];
 					break;
 				}
 			}
+			return Result;
 		}
 
 	private:
@@ -436,11 +451,11 @@ protected:
 		FGPUSkinCacheEntry* Entry, 
 		int32 Section, 
 		uint32 RevisionNumber,
-		TSet<FSkinCacheRWBuffer*>& BuffersToTransitionToRead
+		TArray<FSkinCacheRWBuffer*>& BuffersToTransitionToRead
 		);
 
 	void Cleanup();
-	static void TransitionAllToReadable(FRHICommandList& RHICmdList, const TSet<FSkinCacheRWBuffer*>& BuffersToTransitionToRead);
+	static void TransitionAllToReadable(FRHICommandList& RHICmdList, const TArray<FSkinCacheRWBuffer*>& BuffersToTransitionToRead);
 	static void ReleaseSkinCacheEntry(FGPUSkinCacheEntry* SkinCacheEntry);
 	void InvalidateAllEntries();
 	uint64 UsedMemoryInBytes;
