Netdriver Fix dll build

diff --git a/Engine/Source/Runtime/Engine/Classes/Engine/NetDriver.h b/Engine/Source/Runtime/Engine/Classes/Engine/NetDriver.h
--- a/Engine/Source/Runtime/Engine/Classes/Engine/NetDriver.h
+++ b/Engine/Source/Runtime/Engine/Classes/Engine/NetDriver.h
@@ -799,7 +799,7 @@ public:
 	}
 
 	/** Override the configured server tick rate. Value is in ticks per second. */
-	void SetNetServerMaxTickRate(int32 InServerMaxTickRate);
+	ENGINE_API void SetNetServerMaxTickRate(int32 InServerMaxTickRate);
 
 	/** 
 	* Delegate triggered when SetNetServerMaxTickRate is called and causes a change to the current max tick rate.
