- Fix detailed particle perf stats
- When ECVF_RenderThreadSafe was added these were broken

#rnx

diff --git a/Engine/Source/Runtime/Engine/Private/Particles/ParticlePerfStats.cpp b/Engine/Source/Runtime/Engine/Private/Particles/ParticlePerfStats.cpp
--- a/Engine/Source/Runtime/Engine/Private/Particles/ParticlePerfStats.cpp
+++ b/Engine/Source/Runtime/Engine/Private/Particles/ParticlePerfStats.cpp
@@ -1105,9 +1105,20 @@ bool FParticlePerfStatsListener_TimedTest::Tick()
 
 CSV_DEFINE_CATEGORY_MODULE(ENGINE_API, Particles, false);
 
-static bool bDetailedCSVStats = false;
+void OnDetailedCSVStatsEnabledChanged(IConsoleVariable* Variable);
+
+static FAutoConsoleVariable CVarWriteDetailedCSVStats(
+	TEXT("fx.DetailedCSVStats"),
+	false,
+	TEXT("If true, we write detailed partilce stats to the CSV profiler."),
+	FConsoleVariableDelegate::CreateStatic(&OnDetailedCSVStatsEnabledChanged),
+	ECVF_Default | ECVF_RenderThreadSafe
+);
+
 void OnDetailedCSVStatsEnabledChanged(IConsoleVariable* Variable)
 {
+	const bool bDetailedCSVStats = CVarWriteDetailedCSVStats->GetBool();
+
 	FParticlePerfStats::SetCSVStatsEnabled(bDetailedCSVStats);
 	
 	FCsvProfiler* CSVProfiler = FCsvProfiler::Get();
@@ -1128,17 +1139,11 @@ void OnDetailedCSVStatsEnabledChanged(IConsoleVariable* Variable)
 		FParticlePerfStatsListener_CSVProfiler::OnCSVEnd();
 	}
 }
-static FAutoConsoleVariableRef CVarWriteDetailedCSVStats(
-	TEXT("fx.DetailedCSVStats"),
-	bDetailedCSVStats,
-	TEXT("If true, we write detailed partilce stats to the CSV profiler. \n"),
-	FConsoleVariableDelegate::CreateStatic(&OnDetailedCSVStatsEnabledChanged),
-	ECVF_Default | ECVF_RenderThreadSafe);
-
 
 FParticlePerfStatsListenerPtr FParticlePerfStatsListener_CSVProfiler::CSVListener;
 void FParticlePerfStatsListener_CSVProfiler::OnCSVStart()
 {
+	const bool bDetailedCSVStats = CVarWriteDetailedCSVStats->GetBool();
 	if (bDetailedCSVStats && CSVListener.IsValid() == false)
 	{
 		CSVListener = MakeShared<FParticlePerfStatsListener_CSVProfiler, ESPMode::ThreadSafe>();
