Expose Gravity settings on the immediate physics actor

diff --git a/Engine/Source/Runtime/Engine/Private/PhysicsEngine/ImmediatePhysics/ImmediatePhysicsChaos/ImmediatePhysicsActorHandle_Chaos.cpp b/Engine/Source/Runtime/Engine/Private/PhysicsEngine/ImmediatePhysics/ImmediatePhysicsChaos/ImmediatePhysicsActorHandle_Chaos.cpp
--- a/Engine/Source/Runtime/Engine/Private/PhysicsEngine/ImmediatePhysics/ImmediatePhysicsChaos/ImmediatePhysicsActorHandle_Chaos.cpp
+++ b/Engine/Source/Runtime/Engine/Private/PhysicsEngine/ImmediatePhysics/ImmediatePhysicsChaos/ImmediatePhysicsActorHandle_Chaos.cpp
@@ -455,6 +455,25 @@ namespace ImmediatePhysics_Chaos
 		return ParticleHandle->CastToRigidParticle() != nullptr && ParticleHandle->ObjectState() == Chaos::EObjectStateType::Dynamic;
 	}
 
+	bool FActorHandle::IsGravityEnabled() const
+	{
+		using namespace Chaos;
+		if (FPBDRigidParticleHandle* Rigid = Handle()->CastToRigidParticle())
+		{
+			return IsSimulated() && Rigid->GravityEnabled();
+		}
+		return false;
+	}
+
+	void FActorHandle::SetGravityEnabled(bool bEnable)
+	{
+		using namespace Chaos;
+		if (FPBDRigidParticleHandle* Rigid = Handle()->CastToRigidParticle())
+		{
+			Rigid->SetGravityEnabled(bEnable);
+		}
+	}
+
 	FTransform FActorHandle::GetWorldTransform() const
 	{
 		using namespace Chaos;
diff --git a/Engine/Source/Runtime/Engine/Public/Physics/ImmediatePhysics/ImmediatePhysicsChaos/ImmediatePhysicsActorHandle_Chaos.h b/Engine/Source/Runtime/Engine/Public/Physics/ImmediatePhysics/ImmediatePhysicsChaos/ImmediatePhysicsActorHandle_Chaos.h
--- a/Engine/Source/Runtime/Engine/Public/Physics/ImmediatePhysics/ImmediatePhysicsChaos/ImmediatePhysicsActorHandle_Chaos.h
+++ b/Engine/Source/Runtime/Engine/Public/Physics/ImmediatePhysics/ImmediatePhysicsChaos/ImmediatePhysicsActorHandle_Chaos.h
@@ -48,9 +48,15 @@ namespace ImmediatePhysics_Chaos
 		/** Does this actor have a kinematic target (next kinematic transform to be applied) */
 		bool HasKinematicTarget() const;
 
+		/** Returns true if the world gravity applies to this actor */
+		bool IsGravityEnabled() const;
+
 		/** Whether the body is static */
 		bool IsStatic() const;
 
+		/** Sets whether world gravity should apply to this actor, assuming it is dynamic */
+		void SetGravityEnabled(bool bEnable);
+
 		/** Whether the body is simulating */
 		bool IsSimulated() const;
 
