Add EffectiveParallelization for shader compilation to our cook analytics.

- Stat will show up as `Shaders_EffectiveParallelization`
- GatherAnalytics should be grabbing the CompileStatsLock since it is reading the stats
- Remove `const` from GatherAnalytics since it now needs to grab the lock and call GetTimeShaderCompilationWasActive

diff --git a/Engine/Source/Runtime/Engine/Private/ShaderCompiler/ShaderCompiler.cpp b/Engine/Source/Runtime/Engine/Private/ShaderCompiler/ShaderCompiler.cpp
--- a/Engine/Source/Runtime/Engine/Private/ShaderCompiler/ShaderCompiler.cpp
+++ b/Engine/Source/Runtime/Engine/Private/ShaderCompiler/ShaderCompiler.cpp
@@ -3560,8 +3560,12 @@ void FShaderCompilerStats::WriteStatSummary()
 	}
 }
 
-void FShaderCompilerStats::GatherAnalytics(const FString& BaseName, TArray<FAnalyticsEventAttribute>& Attributes) const
+void FShaderCompilerStats::GatherAnalytics(const FString& BaseName, TArray<FAnalyticsEventAttribute>& Attributes)
 {
+	const double TotalTimeAtLeastOneJobWasInFlight = GetTimeShaderCompilationWasActive();
+
+	FScopeLock Lock(&CompileStatsLock);
+
 	{
 		FString AttrName = BaseName + TEXT("ShadersCompiled");
 		Attributes.Emplace(MoveTemp(AttrName), JobsCompleted);
@@ -3586,6 +3590,12 @@ void FShaderCompilerStats::GatherAnalytics(const FString& BaseName, TArray<FAnal
 			FString AttrName = BaseName + TEXT("TotalThreadPreprocessTime");
 			Attributes.Emplace(MoveTemp(AttrName), TotalThreadPreprocessTimeForAllShaders);
 		}
+
+		{
+            const double EffectiveParallelization = TotalTimeAtLeastOneJobWasInFlight > 0.0 ? TotalThreadTimeForAllShaders / TotalTimeAtLeastOneJobWasInFlight : 0.0;
+			FString AttrName = BaseName + TEXT("EffectiveParallelization");
+			Attributes.Emplace(MoveTemp(AttrName), EffectiveParallelization);
+		}
 	}
 }
 
diff --git a/Engine/Source/Runtime/Engine/Public/ShaderCompiler.h b/Engine/Source/Runtime/Engine/Public/ShaderCompiler.h
--- a/Engine/Source/Runtime/Engine/Public/ShaderCompiler.h
+++ b/Engine/Source/Runtime/Engine/Public/ShaderCompiler.h
@@ -514,7 +514,7 @@ public:
 	/** Informs statistics about a new job batch, so we can tally up batches. */
 	void RegisterJobBatch(int32 NumJobs, EExecutionType ExecType);
 
-	void GatherAnalytics(const FString& BaseName, TArray<FAnalyticsEventAttribute>& Attributes) const;
+	void GatherAnalytics(const FString& BaseName, TArray<FAnalyticsEventAttribute>& Attributes);
 
 private:
 	FCriticalSection CompileStatsLock;
