Relax detachment check to not ensure when a replicated component is detached from a non-replicated component during replays

diff --git a/Engine/Source/Runtime/Engine/Private/Components/SceneComponent.cpp b/Engine/Source/Runtime/Engine/Private/Components/SceneComponent.cpp
--- a/Engine/Source/Runtime/Engine/Private/Components/SceneComponent.cpp
+++ b/Engine/Source/Runtime/Engine/Private/Components/SceneComponent.cpp
@@ -2276,8 +2276,8 @@ void USceneComponent::DetachFromComponent(const FDetachmentTransformRules& Detac
 			PrimComp->UnWeldFromParent();
 		}
 		
-		// Due to replication order the ensure below is only valid on server OR if not both parent and child are replicated
-		if ((Owner && Owner->GetLocalRole() == ROLE_Authority) || !(GetIsReplicated() && GetAttachParent()->GetIsReplicated()))
+		// Due to replication order the ensure below is only valid on server OR if neither the parent nor child are replicated
+		if ((Owner && Owner->GetLocalRole() == ROLE_Authority) || !(GetIsReplicated() || GetAttachParent()->GetIsReplicated()))
 		{
 			// Make sure parent points to us if we're registered
 			ensureMsgf(!bRegistered || GetAttachParent()->GetAttachChildren().Contains(this), TEXT("Attempt to detach SceneComponent '%s' owned by '%s' from AttachParent '%s' while not attached."), *GetName(), (Owner ? *Owner->GetName() : TEXT("Unowned")), *GetAttachParent()->GetName());
