need to modify if an overlap against a cluster union has an overlap if the overlap with the cluster union got ignored because of ignored actors/components/etc.

diff --git a/Engine/Source/Runtime/Engine/Private/Collision/SceneQuery.cpp b/Engine/Source/Runtime/Engine/Private/Collision/SceneQuery.cpp
--- a/Engine/Source/Runtime/Engine/Private/Collision/SceneQuery.cpp
+++ b/Engine/Source/Runtime/Engine/Private/Collision/SceneQuery.cpp
@@ -946,6 +946,7 @@ bool GeomOverlapMultiImp(const UWorld* World, const FPhysicsGeometry& Geom, cons
 							}
 
 							OutOverlaps.Append(AllNewOverlaps);
+							bHaveBlockingHit &= !OutOverlaps.IsEmpty();
 						}
 					}
 				}
