Optimization to material expressions concerning loop detection.

Removed individual calls to UMaterialExpression::ContainsInputLoop() and instead moved graph loop search in HLSLTranslator::Translate(). Since this can be done once for all inputs, the visited node set can be shared, making the loop seach algorithm O(n) (visit each node exactly once) rather than O(n^2) (perform the loop check for each node independently).

This optimization, combined with other optimizations in 25735041, make translating an expensive material (I used MI_WeaponWrap_ROYGBIVLlama for profiling) from 5.6 seconds to 0.9 s.

diff --git a/Engine/Source/Runtime/Engine/Classes/Materials/MaterialExpression.h b/Engine/Source/Runtime/Engine/Classes/Materials/MaterialExpression.h
--- a/Engine/Source/Runtime/Engine/Classes/Materials/MaterialExpression.h
+++ b/Engine/Source/Runtime/Engine/Classes/Materials/MaterialExpression.h
@@ -574,6 +574,9 @@ class ENGINE_API UMaterialExpression : public UObject
 	/** Checks whether any inputs to this expression create a loop */
 	bool ContainsInputLoop(const bool bStopOnFunctionCall = true);
 
+	/** This overload accepts the set of visited expressions to avoid visiting them again when checking at once whether multiple expressions contain a loop. */
+	bool ContainsInputLoop(TSet<UMaterialExpression*>& VisitedExpressions, const bool bStopOnFunctionCall = true);
+
 protected:
 	/** Caches the list of expression inputs this expression has. */
 	TArray<FExpressionInput*> CachedInputs;
@@ -589,10 +592,10 @@ private:
 	 * Checks whether any inputs to this expression create a loop by recursively
 	 * calling itself and keeping a list of inputs as expression keys.
 	 *
-	 * @param ExpressionStack    List of expression keys that have been checked already in the current stack
+	 * @param ExpressionStack    Stack of expressions that have been checked already
 	 * @param VisitedExpressions List of all expression keys that have been visited
 	 */
-	bool ContainsInputLoopInternal(const FContainsInputLoopInternalExpressionStack& ExpressionStack, TSet<class FMaterialExpressionKey>& VisitedExpressions, const bool bStopOnFunctionCall);
+	bool ContainsInputLoopInternal(const FContainsInputLoopInternalExpressionStack& ExpressionStack, TSet<UMaterialExpression*>& VisitedExpressions, const bool bStopOnFunctionCall);
 
 #endif // WITH_EDITOR
 };
diff --git a/Engine/Source/Runtime/Engine/Private/Materials/HLSLMaterialTranslator.cpp b/Engine/Source/Runtime/Engine/Private/Materials/HLSLMaterialTranslator.cpp
--- a/Engine/Source/Runtime/Engine/Private/Materials/HLSLMaterialTranslator.cpp
+++ b/Engine/Source/Runtime/Engine/Private/Materials/HLSLMaterialTranslator.cpp
@@ -835,6 +835,20 @@ bool FHLSLMaterialTranslator::Translate()
 		bEnableExecutionFlow = Material->IsUsingControlFlow();
 		bCompileForComputeShader = Material->IsLightFunction();
 		
+		// Verify for the absence of loops.
+		if (UMaterial* UMaterial = Material->GetMaterialInterface()->GetMaterial())
+		{
+			TSet<UMaterialExpression*> VisitedExpressions;
+			for (UMaterialExpression* Expression : UMaterial->GetExpressions())
+			{
+				if (Expression->ContainsInputLoop(VisitedExpressions))
+				{
+					AppendExpressionError(Expression, TEXT("Expression is part of a cycle. Please make sure the material graph is acyclic."));
+					return false;
+				}
+			}
+		}
+
 		//
 		// Process the strata tree representing the material topology.
 		//
@@ -4179,7 +4193,7 @@ int32 FHLSLMaterialTranslator::CallExpression(FMaterialExpressionKey ExpressionK
 {
 	// For any translated result not relying on material attributes, we can discard the attribute ID from the key
 	// to allow result sharing. In cases where we detect an expression loop we must err on the side of caution
-	if (ExpressionKey.Expression && !ExpressionKey.Expression->ContainsInputLoop() && !ExpressionKey.Expression->IsResultMaterialAttributes(ExpressionKey.OutputIndex))
+	if (ExpressionKey.Expression && !ExpressionKey.Expression->IsResultMaterialAttributes(ExpressionKey.OutputIndex))
 	{
 		ExpressionKey.MaterialAttributeID = FGuid(0, 0, 0, 0);
 	}
diff --git a/Engine/Source/Runtime/Engine/Private/Materials/MaterialExpressions.cpp b/Engine/Source/Runtime/Engine/Private/Materials/MaterialExpressions.cpp
--- a/Engine/Source/Runtime/Engine/Private/Materials/MaterialExpressions.cpp
+++ b/Engine/Source/Runtime/Engine/Private/Materials/MaterialExpressions.cpp
@@ -2187,20 +2187,20 @@ bool UMaterialExpression::HasConnectedOutputs() const
 
 struct UMaterialExpression::FContainsInputLoopInternalExpressionStack
 {
-	const FMaterialExpressionKey* Key;
+	const UMaterialExpression* Expression;
 	const FContainsInputLoopInternalExpressionStack* Previous;
 
-	FContainsInputLoopInternalExpressionStack(const FMaterialExpressionKey* Key, const FContainsInputLoopInternalExpressionStack* Previous)
-		: Key{ Key }
+	FContainsInputLoopInternalExpressionStack(const UMaterialExpression* Expression, const FContainsInputLoopInternalExpressionStack* Previous)
+		: Expression{ Expression }
 		, Previous{ Previous }
 	{}
 
-	bool Contains(const FMaterialExpressionKey& OtherKey) const
+	bool Contains(const UMaterialExpression* OtherExpression) const
 	{
 		const FContainsInputLoopInternalExpressionStack* Node = this;
-		while (Node->Key)
+		while (Node->Expression)
 		{
-			if (*Node->Key == OtherKey)
+			if (Node->Expression == OtherExpression)
 			{
 				return true;
 			}
@@ -2213,11 +2213,21 @@ struct UMaterialExpression::FContainsInputLoopInternalExpressionStack
 bool UMaterialExpression::ContainsInputLoop(const bool bStopOnFunctionCall /*= true*/)
 {
 	FContainsInputLoopInternalExpressionStack ExpressionStack{ nullptr, nullptr };
-	TSet<FMaterialExpressionKey> VisitedExpressions;
+	TSet<UMaterialExpression*> VisitedExpressions;
+	return ContainsInputLoopInternal(ExpressionStack, VisitedExpressions, bStopOnFunctionCall);
+}
+
+bool UMaterialExpression::ContainsInputLoop(TSet<UMaterialExpression*>& VisitedExpressions, const bool bStopOnFunctionCall)
+{
+	if (VisitedExpressions.Contains(this))
+	{
+		return false;
+	}
+	FContainsInputLoopInternalExpressionStack ExpressionStack{ nullptr, nullptr };
 	return ContainsInputLoopInternal(ExpressionStack, VisitedExpressions, bStopOnFunctionCall);
 }
 
-bool UMaterialExpression::ContainsInputLoopInternal(const FContainsInputLoopInternalExpressionStack& ExpressionStack, TSet<FMaterialExpressionKey>& VisitedExpressions, const bool bStopOnFunctionCall)
+bool UMaterialExpression::ContainsInputLoopInternal(const FContainsInputLoopInternalExpressionStack& ExpressionStack, TSet<UMaterialExpression*>& VisitedExpressions, const bool bStopOnFunctionCall)
 {
 	for (FExpressionInput* Input : GetInputsView())
 	{
@@ -2237,22 +2247,24 @@ bool UMaterialExpression::ContainsInputLoopInternal(const FContainsInputLoopInte
 		}
 
 		// A cycle is detected if one of this node's inputs leads back to a node we're coming from.
-		FMaterialExpressionKey InputExpressionKey(InputExpression, Input->OutputIndex);
-		if (ExpressionStack.Contains(InputExpressionKey))
+		if (ExpressionStack.Contains(InputExpression))
 		{
 			return true;
 		}
 
-		// Prevent recurring visits to expressions we've already checked
-		if (!VisitedExpressions.Contains(InputExpressionKey))
+		// Add this expression to the visited set. If it was already there, we do not need to explore it again.
+		bool bAlreadyVisited = false;
+		VisitedExpressions.Add(InputExpression, &bAlreadyVisited);
+		if (bAlreadyVisited)
 		{
-			VisitedExpressions.Add(InputExpressionKey);
+			continue;
+		}
 
-			FContainsInputLoopInternalExpressionStack ExpressionStackWithThisInput{ &InputExpressionKey, &ExpressionStack };
-			if (InputExpression->ContainsInputLoopInternal(ExpressionStackWithThisInput, VisitedExpressions, bStopOnFunctionCall))
-			{
-				return true;
-			}
+		// Push this expression onto the stack and carry on crawling through this expression.
+		FContainsInputLoopInternalExpressionStack ExpressionStackWithThisInput{ InputExpression, &ExpressionStack };
+		if (InputExpression->ContainsInputLoopInternal(ExpressionStackWithThisInput, VisitedExpressions, bStopOnFunctionCall))
+		{
+			return true;
 		}
 	}
 
@@ -3541,7 +3553,7 @@ bool UMaterialExpressionRuntimeVirtualTextureReplace::IsResultMaterialAttributes
 {
 	for (FExpressionInput* ExpressionInput : GetInputsView())
 	{
-		if (ExpressionInput->GetTracedInput().Expression && !ExpressionInput->Expression->ContainsInputLoop() && ExpressionInput->Expression->IsResultMaterialAttributes(ExpressionInput->OutputIndex))
+		if (ExpressionInput->GetTracedInput().Expression && ExpressionInput->Expression->IsResultMaterialAttributes(ExpressionInput->OutputIndex))
 		{
 			return true;
 		}
@@ -3601,7 +3613,7 @@ bool UMaterialExpressionVirtualTextureFeatureSwitch::IsResultMaterialAttributes(
 {
 	for (FExpressionInput* ExpressionInput : GetInputsView())
 	{
-		if (ExpressionInput->GetTracedInput().Expression && !ExpressionInput->Expression->ContainsInputLoop() && ExpressionInput->Expression->IsResultMaterialAttributes(ExpressionInput->OutputIndex))
+		if (ExpressionInput->GetTracedInput().Expression && ExpressionInput->Expression->IsResultMaterialAttributes(ExpressionInput->OutputIndex))
 		{
 			return true;
 		}
@@ -8958,8 +8970,8 @@ bool UMaterialExpressionStaticSwitchParameter::IsResultMaterialAttributes(int32
 	// This one is a little tricky. Since we are treating a dangling reroute as an empty expression, this
 	// should early out, whereas IsResultMaterialAttributes on a reroute node will return false as the 
 	// reroute node's input is dangling and therefore its type is unknown.
-	if ((A.GetTracedInput().Expression && !A.Expression->ContainsInputLoop() && A.Expression->IsResultMaterialAttributes(A.OutputIndex)) ||
-		(B.GetTracedInput().Expression && !B.Expression->ContainsInputLoop() && B.Expression->IsResultMaterialAttributes(B.OutputIndex)))
+	if ((A.GetTracedInput().Expression && A.Expression->IsResultMaterialAttributes(A.OutputIndex)) ||
+		(B.GetTracedInput().Expression && B.Expression->IsResultMaterialAttributes(B.OutputIndex)))
 	{
 		return true;
 	}
@@ -9196,8 +9208,8 @@ bool UMaterialExpressionStaticSwitch::IsResultMaterialAttributes(int32 OutputInd
 	// should early out, whereas IsResultMaterialAttributes on a reroute node will return false as the 
 	// reroute node's input is dangling and therefore its type is unknown.
 	check(OutputIndex == 0);
-	if ((A.GetTracedInput().Expression && !A.Expression->ContainsInputLoop() && A.Expression->IsResultMaterialAttributes(A.OutputIndex)) ||
-		(B.GetTracedInput().Expression && !B.Expression->ContainsInputLoop() && B.Expression->IsResultMaterialAttributes(B.OutputIndex)))
+	if ((A.GetTracedInput().Expression && A.Expression->IsResultMaterialAttributes(A.OutputIndex)) ||
+		(B.GetTracedInput().Expression && B.Expression->IsResultMaterialAttributes(B.OutputIndex)))
 	{
 		return true;
 	}
@@ -9342,8 +9354,8 @@ bool UMaterialExpressionPreviousFrameSwitch::IsResultMaterialAttributes(int32 Ou
 {
 	// If there is a loop anywhere in this expression's inputs then we can't risk checking them
 	check(OutputIndex == 0);
-	if ((CurrentFrame.Expression && !CurrentFrame.Expression->ContainsInputLoop() && CurrentFrame.Expression->IsResultMaterialAttributes(CurrentFrame.OutputIndex)) ||
-		(PreviousFrame.Expression && !PreviousFrame.Expression->ContainsInputLoop() && PreviousFrame.Expression->IsResultMaterialAttributes(PreviousFrame.OutputIndex)))
+	if ((CurrentFrame.Expression && CurrentFrame.Expression->IsResultMaterialAttributes(CurrentFrame.OutputIndex)) ||
+		(PreviousFrame.Expression && PreviousFrame.Expression->IsResultMaterialAttributes(PreviousFrame.OutputIndex)))
 	{
 		return true;
 	}
@@ -9470,7 +9482,7 @@ bool UMaterialExpressionQualitySwitch::IsResultMaterialAttributes(int32 OutputIn
 	for (FExpressionInput* ExpressionInput : ExpressionInputs)
 	{
 		// If there is a loop anywhere in this expression's inputs then we can't risk checking them
-		if (ExpressionInput->Expression && !ExpressionInput->Expression->ContainsInputLoop() && ExpressionInput->Expression->IsResultMaterialAttributes(ExpressionInput->OutputIndex))
+		if (ExpressionInput->Expression && ExpressionInput->Expression->IsResultMaterialAttributes(ExpressionInput->OutputIndex))
 		{
 			return true;
 		}
@@ -9595,7 +9607,7 @@ bool UMaterialExpressionFeatureLevelSwitch::IsResultMaterialAttributes(int32 Out
 	for (FExpressionInput* ExpressionInput : ExpressionInputs)
 	{
 		// If there is a loop anywhere in this expression's inputs then we can't risk checking them
-		if (ExpressionInput->GetTracedInput().Expression && !ExpressionInput->Expression->ContainsInputLoop() && ExpressionInput->Expression->IsResultMaterialAttributes(ExpressionInput->OutputIndex))
+		if (ExpressionInput->GetTracedInput().Expression && ExpressionInput->Expression->IsResultMaterialAttributes(ExpressionInput->OutputIndex))
 		{
 			return true;
 		}
@@ -9784,7 +9796,7 @@ bool UMaterialExpressionDataDrivenShaderPlatformInfoSwitch::IsResultMaterialAttr
 	{
 		// If there is a loop anywhere in this expression's inputs then we can't risk checking them
 		TObjectPtr<class UMaterialExpression> Expression = ExpressionInput->Expression;
-		if (Expression && !Expression->ContainsInputLoop() && ExpressionInput->Expression->IsResultMaterialAttributes(ExpressionInput->OutputIndex))
+		if (Expression && ExpressionInput->Expression->IsResultMaterialAttributes(ExpressionInput->OutputIndex))
 		{
 			return true;
 		}
@@ -9897,7 +9909,7 @@ bool UMaterialExpressionShadingPathSwitch::IsResultMaterialAttributes(int32 Outp
 	for (FExpressionInput* ExpressionInput : GetInputsView())
 	{
 		// If there is a loop anywhere in this expression's inputs then we can't risk checking them
-		if (ExpressionInput->Expression && !ExpressionInput->Expression->ContainsInputLoop() && ExpressionInput->Expression->IsResultMaterialAttributes(ExpressionInput->OutputIndex))
+		if (ExpressionInput->Expression && ExpressionInput->Expression->IsResultMaterialAttributes(ExpressionInput->OutputIndex))
 		{
 			return true;
 		}
@@ -11337,13 +11349,13 @@ uint32 UMaterialExpressionIf::GetInputType(int32 InputIndex)
 	// First two inputs are always float
 	if (InputIndex == 0 || InputIndex == 1)
 	{
-		if ((A.GetTracedInput().Expression && !A.Expression->ContainsInputLoop() && A.Expression->IsResultMaterialAttributes(A.OutputIndex)) ||
-			(B.GetTracedInput().Expression && !B.Expression->ContainsInputLoop() && B.Expression->IsResultMaterialAttributes(B.OutputIndex)))
+		if ((A.GetTracedInput().Expression && A.Expression->IsResultMaterialAttributes(A.OutputIndex)) ||
+			(B.GetTracedInput().Expression && B.Expression->IsResultMaterialAttributes(B.OutputIndex)))
 		{
 			return MCT_MaterialAttributes;
 		}
-		else if ((A.GetTracedInput().Expression && !A.Expression->ContainsInputLoop() && A.Expression->GetOutputType(0) == MCT_ShadingModel) &&
-			(B.GetTracedInput().Expression && !B.Expression->ContainsInputLoop() && B.Expression->GetOutputType(0) == MCT_ShadingModel))
+		else if ((A.GetTracedInput().Expression && A.Expression->GetOutputType(0) == MCT_ShadingModel) &&
+			(B.GetTracedInput().Expression && B.Expression->GetOutputType(0) == MCT_ShadingModel))
 		{
 			return MCT_ShadingModel;
 		}
@@ -11358,9 +11370,9 @@ uint32 UMaterialExpressionIf::GetInputType(int32 InputIndex)
 
 bool UMaterialExpressionIf::IsResultMaterialAttributes(int32 OutputIndex)
 {
-	if ((AGreaterThanB.GetTracedInput().Expression && !AGreaterThanB.Expression->ContainsInputLoop() && AGreaterThanB.Expression->IsResultMaterialAttributes(AGreaterThanB.OutputIndex))
-		&& (!AEqualsB.GetTracedInput().Expression || (!AEqualsB.Expression->ContainsInputLoop() && AEqualsB.Expression->IsResultMaterialAttributes(AEqualsB.OutputIndex)))
-		&& (ALessThanB.GetTracedInput().Expression && !ALessThanB.Expression->ContainsInputLoop() && ALessThanB.Expression->IsResultMaterialAttributes(ALessThanB.OutputIndex))
+	if ((AGreaterThanB.GetTracedInput().Expression && AGreaterThanB.Expression->IsResultMaterialAttributes(AGreaterThanB.OutputIndex))
+		&& (!AEqualsB.GetTracedInput().Expression || AEqualsB.Expression->IsResultMaterialAttributes(AEqualsB.OutputIndex))
+		&& (ALessThanB.GetTracedInput().Expression && ALessThanB.Expression->IsResultMaterialAttributes(ALessThanB.OutputIndex))
 		)
 	{
 		return true;
@@ -12933,7 +12945,7 @@ bool UMaterialExpressionDistanceFieldsRenderingSwitch::IsResultMaterialAttribute
 {
 	for (FExpressionInput* ExpressionInput : GetInputsView())
 	{
-		if (ExpressionInput->GetTracedInput().Expression && !ExpressionInput->Expression->ContainsInputLoop() && ExpressionInput->Expression->IsResultMaterialAttributes(ExpressionInput->OutputIndex))
+		if (ExpressionInput->GetTracedInput().Expression && ExpressionInput->Expression->IsResultMaterialAttributes(ExpressionInput->OutputIndex))
 		{
 			return true;
 		}
@@ -16881,7 +16893,7 @@ void UMaterialExpressionFunctionInput::GatherStrataMaterialInfo(FStrataMaterialI
 	{
 		FExpressionInput EffectivePreviewDuringCompileTracedInput = EffectivePreviewDuringCompile.GetTracedInput();
 		int32 ExpressionResult = INDEX_NONE;
-		if (EffectivePreviewDuringCompileTracedInput.GetTracedInput().Expression && !EffectivePreviewDuringCompileTracedInput.Expression->ContainsInputLoop())
+		if (EffectivePreviewDuringCompileTracedInput.GetTracedInput().Expression)
 		{
 			EffectivePreviewDuringCompileTracedInput.Expression->GatherStrataMaterialInfo(StrataMaterialInfo, EffectivePreviewDuringCompileTracedInput.OutputIndex);
 		}
@@ -16894,7 +16906,7 @@ FStrataOperator* UMaterialExpressionFunctionInput::StrataGenerateMaterialTopolog
 	{
 		FExpressionInput EffectivePreviewDuringCompileTracedInput = EffectivePreviewDuringCompile.GetTracedInput();
 		int32 ExpressionResult = INDEX_NONE;
-		if (EffectivePreviewDuringCompileTracedInput.GetTracedInput().Expression && !EffectivePreviewDuringCompileTracedInput.Expression->ContainsInputLoop())
+		if (EffectivePreviewDuringCompileTracedInput.GetTracedInput().Expression)
 		{
 			return EffectivePreviewDuringCompileTracedInput.Expression->StrataGenerateMaterialTopologyTree(Compiler, Parent, EffectivePreviewDuringCompileTracedInput.OutputIndex);
 		}
@@ -17125,7 +17137,7 @@ void UMaterialExpressionFunctionOutput::ValidateName()
 bool UMaterialExpressionFunctionOutput::IsResultMaterialAttributes(int32 OutputIndex)
 {
 	// If there is a loop anywhere in this expression's inputs then we can't risk checking them
-	if( A.GetTracedInput().Expression && !A.Expression->ContainsInputLoop() )
+	if( A.GetTracedInput().Expression )
 	{
 		return A.Expression->IsResultMaterialAttributes(A.OutputIndex);
 	}
@@ -17138,7 +17150,7 @@ bool UMaterialExpressionFunctionOutput::IsResultMaterialAttributes(int32 OutputI
 bool UMaterialExpressionFunctionOutput::IsResultStrataMaterial(int32 OutputIndex)
 {
 	// If there is a loop anywhere in this expression's inputs then we can't risk checking them
-	if( A.GetTracedInput().Expression && !A.Expression->ContainsInputLoop() )
+	if( A.GetTracedInput().Expression )
 	{
 		return A.Expression->IsResultStrataMaterial(A.OutputIndex);
 	}
@@ -17151,7 +17163,7 @@ bool UMaterialExpressionFunctionOutput::IsResultStrataMaterial(int32 OutputIndex
 void UMaterialExpressionFunctionOutput::GatherStrataMaterialInfo(FStrataMaterialInfo& StrataMaterialInfo, int32 OutputIndex)
 {
 	// If there is a loop anywhere in this expression's inputs then we can't risk checking them
-	if( A.GetTracedInput().Expression && !A.Expression->ContainsInputLoop() )
+	if( A.GetTracedInput().Expression )
 	{
 		A.Expression->GatherStrataMaterialInfo(StrataMaterialInfo, A.OutputIndex);
 	}
@@ -17159,7 +17171,7 @@ void UMaterialExpressionFunctionOutput::GatherStrataMaterialInfo(FStrataMaterial
 
 FStrataOperator* UMaterialExpressionFunctionOutput::StrataGenerateMaterialTopologyTree(class FMaterialCompiler* Compiler, class UMaterialExpression* Parent, int32 OutputIndex)
 {
-	if (A.GetTracedInput().Expression && !A.Expression->ContainsInputLoop())
+	if (A.GetTracedInput().Expression)
 	{
 		return A.Expression->StrataGenerateMaterialTopologyTree(Compiler, Parent, A.OutputIndex);
 	}
@@ -17621,7 +17633,7 @@ bool UMaterialExpressionMaterialProxyReplace::IsResultMaterialAttributes(int32 O
 {
 	for (FExpressionInput* ExpressionInput : GetInputsView())
 	{
-		if (ExpressionInput->GetTracedInput().Expression && !ExpressionInput->Expression->ContainsInputLoop() && ExpressionInput->Expression->IsResultMaterialAttributes(ExpressionInput->OutputIndex))
+		if (ExpressionInput->GetTracedInput().Expression && ExpressionInput->Expression->IsResultMaterialAttributes(ExpressionInput->OutputIndex))
 		{
 			return true;
 		}
diff --git a/Engine/Source/Runtime/Landscape/Private/Materials/MaterialExpressionLandscapeLayerBlend.cpp b/Engine/Source/Runtime/Landscape/Private/Materials/MaterialExpressionLandscapeLayerBlend.cpp
--- a/Engine/Source/Runtime/Landscape/Private/Materials/MaterialExpressionLandscapeLayerBlend.cpp
+++ b/Engine/Source/Runtime/Landscape/Private/Materials/MaterialExpressionLandscapeLayerBlend.cpp
@@ -134,11 +134,6 @@ uint32 UMaterialExpressionLandscapeLayerBlend::GetInputType(int32 InputIndex)
 
 bool UMaterialExpressionLandscapeLayerBlend::IsResultMaterialAttributes(int32 OutputIndex)
 {
-	if (ContainsInputLoop())
-	{
-		// If there is a loop anywhere in this expression's inputs then we can't risk checking them
-		return false;
-	}
 	for (int32 LayerIdx = 0; LayerIdx < Layers.Num(); LayerIdx++)
 	{
 		if (Layers[LayerIdx].LayerInput.Expression && Layers[LayerIdx].LayerInput.Expression->IsResultMaterialAttributes(Layers[LayerIdx].LayerInput.OutputIndex))
diff --git a/Engine/Source/Runtime/Landscape/Private/Materials/MaterialExpressionLandscapeLayerSwitch.cpp b/Engine/Source/Runtime/Landscape/Private/Materials/MaterialExpressionLandscapeLayerSwitch.cpp
--- a/Engine/Source/Runtime/Landscape/Private/Materials/MaterialExpressionLandscapeLayerSwitch.cpp
+++ b/Engine/Source/Runtime/Landscape/Private/Materials/MaterialExpressionLandscapeLayerSwitch.cpp
@@ -42,11 +42,6 @@ UMaterialExpressionLandscapeLayerSwitch::UMaterialExpressionLandscapeLayerSwitch
 #if WITH_EDITOR
 bool UMaterialExpressionLandscapeLayerSwitch::IsResultMaterialAttributes(int32 OutputIndex)
 {
-	if (ContainsInputLoop())
-	{
-		// If there is a loop anywhere in this expression's inputs then we can't risk checking them
-		return false;
-	}
 	bool bLayerUsedIsMaterialAttributes = LayerUsed.Expression != nullptr && LayerUsed.Expression->IsResultMaterialAttributes(LayerUsed.OutputIndex);
 	bool bLayerNotUsedIsMaterialAttributes = LayerNotUsed.Expression != nullptr && LayerNotUsed.Expression->IsResultMaterialAttributes(LayerNotUsed.OutputIndex);
 	return bLayerUsedIsMaterialAttributes || bLayerNotUsedIsMaterialAttributes;
diff --git a/Engine/Source/Runtime/Landscape/Private/Materials/MaterialExpressionLandscapeLayerWeight.cpp b/Engine/Source/Runtime/Landscape/Private/Materials/MaterialExpressionLandscapeLayerWeight.cpp
--- a/Engine/Source/Runtime/Landscape/Private/Materials/MaterialExpressionLandscapeLayerWeight.cpp
+++ b/Engine/Source/Runtime/Landscape/Private/Materials/MaterialExpressionLandscapeLayerWeight.cpp
@@ -51,11 +51,6 @@ void UMaterialExpressionLandscapeLayerWeight::PostLoad()
 #if WITH_EDITOR
 bool UMaterialExpressionLandscapeLayerWeight::IsResultMaterialAttributes(int32 OutputIndex)
 {
-	if (ContainsInputLoop())
-	{
-		// If there is a loop anywhere in this expression's inputs then we can't risk checking them
-		return false;
-	}
 	bool bLayerIsMaterialAttributes = Layer.Expression != nullptr && Layer.Expression->IsResultMaterialAttributes(Layer.OutputIndex);
 	bool bBaseIsMaterialAttributes = Base.Expression != nullptr && Base.Expression->IsResultMaterialAttributes(Base.OutputIndex);
 	return bLayerIsMaterialAttributes || bBaseIsMaterialAttributes;
