Add new methods to add/remove compatible skeletons taking a TSoftObjectPtr<USkeleton> as a parameter

#rnx

diff --git a/Engine/Source/Runtime/Engine/Classes/Animation/Skeleton.h b/Engine/Source/Runtime/Engine/Classes/Animation/Skeleton.h
--- a/Engine/Source/Runtime/Engine/Classes/Animation/Skeleton.h
+++ b/Engine/Source/Runtime/Engine/Classes/Animation/Skeleton.h
@@ -774,7 +774,12 @@ public:
 public:
 	UFUNCTION(BlueprintCallable, Category=Skeleton)
 	ENGINE_API void AddCompatibleSkeleton(const USkeleton* SourceSkeleton);
+
+	UFUNCTION(BlueprintCallable, Category = Skeleton, DisplayName = "AddCompatibleSkeleton")
+	ENGINE_API void AddCompatibleSkeletonSoft(const TSoftObjectPtr<USkeleton>& SourceSkeleton);
+
 	ENGINE_API void RemoveCompatibleSkeleton(const USkeleton* SourceSkeleton);
+	ENGINE_API void RemoveCompatibleSkeleton(const TSoftObjectPtr<USkeleton>& SourceSkeleton);
 
 	/** 
 	 * Indexing naming convention
diff --git a/Engine/Source/Runtime/Engine/Private/Animation/Skeleton.cpp b/Engine/Source/Runtime/Engine/Private/Animation/Skeleton.cpp
--- a/Engine/Source/Runtime/Engine/Private/Animation/Skeleton.cpp
+++ b/Engine/Source/Runtime/Engine/Private/Animation/Skeleton.cpp
@@ -295,11 +295,21 @@ void USkeleton::AddCompatibleSkeleton(const USkeleton* SourceSkeleton)
 	CompatibleSkeletons.AddUnique(SourceSkeleton);
 }
 
+void USkeleton::AddCompatibleSkeletonSoft(const TSoftObjectPtr<USkeleton>& SourceSkeleton)
+{
+	CompatibleSkeletons.AddUnique(SourceSkeleton);
+}
+
 void USkeleton::RemoveCompatibleSkeleton(const USkeleton* SourceSkeleton)
 {
 	CompatibleSkeletons.Remove(SourceSkeleton);
 }
 
+void USkeleton::RemoveCompatibleSkeleton(const TSoftObjectPtr<USkeleton>& SourceSkeleton)
+{
+	CompatibleSkeletons.Remove(SourceSkeleton);
+}
+
 bool USkeleton::IsCompatibleSkeletonByAssetData(const FAssetData& AssetData, const TCHAR* InTag) const
 {
 #if WITH_EDITORONLY_DATA
