[MLDeformer] Some improvements in cleanup in the external morph targets, to also clear out the weight data.

diff --git a/Engine/Source/Runtime/Engine/Private/Components/SkinnedMeshComponent.cpp b/Engine/Source/Runtime/Engine/Private/Components/SkinnedMeshComponent.cpp
--- a/Engine/Source/Runtime/Engine/Private/Components/SkinnedMeshComponent.cpp
+++ b/Engine/Source/Runtime/Engine/Private/Components/SkinnedMeshComponent.cpp
@@ -743,6 +743,11 @@ void USkinnedMeshComponent::RemoveExternalMorphSet(int32 LOD, int32 ID)
 	{
 		ExternalMorphSets[LOD].Remove(ID);
 	}
+
+	if (ExternalMorphWeightData.IsValidIndex(LOD))
+	{
+		ExternalMorphWeightData[LOD].MorphSets.Remove(ID);
+	}
 }
 
 bool USkinnedMeshComponent::HasExternalMorphSet(int32 LOD, int32 ID) const
@@ -762,8 +767,15 @@ void USkinnedMeshComponent::ClearExternalMorphSets(int32 LOD)
 	{
 		ResizeExternalMorphTargetSets();
 	}
+
 	if (ExternalMorphSets.IsValidIndex(LOD))
 	{
+		// Remove all weight sets as well.
+		if (ExternalMorphWeightData.IsValidIndex(LOD))
+		{
+			ExternalMorphWeightData[LOD].MorphSets.Empty();
+		}
+
 		ExternalMorphSets[LOD].Empty();
 	}
 }
@@ -817,7 +829,6 @@ void USkinnedMeshComponent::RefreshExternalMorphTargetWeights(bool bZeroOldWeigh
 			check(MorphSetWeights != nullptr);
 
 			MorphSetWeights->Name = Item.Value->Name;
-
 			if (bZeroOldWeights)
 			{
 				MorphSetWeights->Weights.Reset();
