World Partition Runtime Hash Set (Experimental):
- Added a new runtime hash, WorldPartitionRuntimeHashSet, that is a container of RuntimePartition objects.
- Each runtime partition object is responsible to partition actors into their own partition, without having to deal with data layers, HLODs, spatial queries, etc.
- The runtime spatial hash set will direct actors to their corresponding runtime partition using the RuntimeGrid property for now.
- Partitions can also use the actors RuntimeGrid property as a dot separated string to do things internally.
- Added RuntimePartitionLevelStreaming that will extract the corresponding streaming levels to generate through the actors RuntimeGrid property; this can be used to generate several streaming levels similar to UE4, without having to rely on moving actors from levels to levels, but just by simply changing their RuntimeGrid property and going into PIE.
- Decoupled spatial queries from the hash implementation, so partitions will only output a list of cells and the runtime hash will handle spatial queries of all partitions.
- Added two basic spatial queries interfaces: list and RTree, and added a unit test to validate new spatial query interfaces through the list reference implementation.

Unavailable by default to users, can be enabled using the wp.Editor.WorldPartitionRuntimeHashSet.Enable console command.

diff --git a/Engine/Source/Runtime/Engine/Private/Tests/WorldPartition/WorldPartitionStaticSpatialIndexTests.cpp b/Engine/Source/Runtime/Engine/Private/Tests/WorldPartition/WorldPartitionStaticSpatialIndexTests.cpp
new file mode 100644
--- /dev/null
+++ b/Engine/Source/Runtime/Engine/Private/Tests/WorldPartition/WorldPartitionStaticSpatialIndexTests.cpp
@@ -0,0 +1,77 @@
+// Copyright Epic Games, Inc. All Rights Reserved.
+
+#include "Misc/AutomationTest.h"
+
+#if WITH_EDITOR
+#include "Engine/World.h"
+#include "WorldPartition/RuntimeHashSet/StaticSpatialIndex.h"
+#endif
+
+#if WITH_DEV_AUTOMATION_TESTS
+
+#define TEST_NAME_ROOT "System.Engine.WorldPartition"
+
+namespace WorldPartitionTests
+{
+	IMPLEMENT_SIMPLE_AUTOMATION_TEST(FWorldPartitionStaticSpatialIndexTest, TEST_NAME_ROOT ".StaticSpatialIndex", EAutomationTestFlags::EditorContext | EAutomationTestFlags::EngineFilter)
+
+#if WITH_EDITOR
+	template <class Class>
+	void PerformTests(FWorldPartitionStaticSpatialIndexTest* Test, const TCHAR* Name, const TArray<TPair<FBox, int32>>& Elements, const TArray<FSphere>& Tests, TArray<int32>& Results)
+	{
+		Class SpatialIndex;
+		SpatialIndex.Init(Elements);
+
+		Results.Reserve(Tests.Num());
+
+		const double StartTime = FPlatformTime::Seconds();
+
+		for (int32 ListNumTests = 0; ListNumTests < Tests.Num(); ListNumTests++)
+		{
+			SpatialIndex.ForEachIntersectingElement(Tests[ListNumTests], [&Results](const int32& Value) { Results.Add(Value); });
+		}
+
+		const double RunTime = FPlatformTime::Seconds() - StartTime;
+		Test->AddInfo(FString::Printf(TEXT("%s: %d tests in %s (%.2f/s)"), Name, Tests.Num(), *FPlatformTime::PrettyTime(RunTime), Tests.Num() / RunTime), 0);
+	}
+#endif
+
+	bool FWorldPartitionStaticSpatialIndexTest::RunTest(const FString& Parameters)
+	{
+#if WITH_EDITOR
+		const int32 NumBoxes = 100000;
+		const int32 NumTests = 10000;
+
+		TArray<TPair<FBox, int32>> Elements;
+		Elements.Reserve(NumBoxes);
+		for (int32 i=0; i<NumBoxes; i++)
+		{
+			const FSphere Sphere(FMath::VRand() * 10000000, FMath::RandRange(1, 100000));
+			Elements.Emplace(FBoxSphereBounds(Sphere).GetBox(), i);
+		}
+
+		TArray<FSphere> Tests;
+		Tests.Reserve(NumTests);
+		for (int32 i=0; i<NumTests; i++)
+		{
+			const FSphere Sphere(FMath::VRand() * 10000000, FMath::RandRange(1, 100000));
+			Tests.Add(Sphere);
+		}
+
+		TArray<int32> ListResults;
+		PerformTests<TStaticSpatialIndexList<int32>>(this, TEXT("TStaticSpatialIndexList"), Elements, Tests, ListResults);
+		ListResults.Sort();
+
+		TArray<int32> FRTreeResults;
+		PerformTests<TStaticSpatialIndexRTree<int32>>(this, TEXT("TStaticSpatialIndexRTree"), Elements, Tests, FRTreeResults);
+		FRTreeResults.Sort();
+
+		TestTrue(TEXT("TStaticSpatialIndexRTree"), ListResults == FRTreeResults);
+#endif
+		return true;
+	}
+}
+
+#undef TEST_NAME_ROOT
+
+#endif
\ No newline at end of file
diff --git a/Engine/Source/Runtime/Engine/Private/WorldPartition/RuntimeHashSet/RuntimePartition.cpp b/Engine/Source/Runtime/Engine/Private/WorldPartition/RuntimeHashSet/RuntimePartition.cpp
new file mode 100644
--- /dev/null
+++ b/Engine/Source/Runtime/Engine/Private/WorldPartition/RuntimeHashSet/RuntimePartition.cpp
@@ -0,0 +1,11 @@
+// Copyright Epic Games, Inc. All Rights Reserved.
+
+#include "WorldPartition/RuntimeHashSet/RuntimePartition.h"
+
+#if WITH_EDITOR
+bool URuntimePartition::PopulateCellActorInstances(const TArray<const IStreamingGenerationContext::FActorSetInstance*>& ActorSetInstances, bool bIsMainWorldPartition, bool bIsCellAlwaysLoaded, TArray<IStreamingGenerationContext::FActorInstance>& OutCellActorInstances)
+{
+	UWorldPartitionRuntimeHash* RuntimeHash = GetTypedOuter<UWorldPartitionRuntimeHash>();
+	return RuntimeHash->PopulateCellActorInstances(ActorSetInstances, bIsMainWorldPartition, bIsCellAlwaysLoaded, OutCellActorInstances);
+}
+#endif
\ No newline at end of file
diff --git a/Engine/Source/Runtime/Engine/Private/WorldPartition/RuntimeHashSet/RuntimePartitionLevelStreaming.cpp b/Engine/Source/Runtime/Engine/Private/WorldPartition/RuntimeHashSet/RuntimePartitionLevelStreaming.cpp
new file mode 100644
--- /dev/null
+++ b/Engine/Source/Runtime/Engine/Private/WorldPartition/RuntimeHashSet/RuntimePartitionLevelStreaming.cpp
@@ -0,0 +1,88 @@
+// Copyright Epic Games, Inc. All Rights Reserved.
+
+#include "WorldPartition/RuntimeHashSet/RuntimePartitionLevelStreaming.h"
+#include "WorldPartition/WorldPartitionStreamingGenerationContext.h"
+
+#if WITH_EDITOR
+bool URuntimePartitionLevelStreaming::SupportsHLODs() const
+{
+	return false;
+}
+
+bool URuntimePartitionLevelStreaming::IsValidGrid(FName GridName) const
+{
+	const TArray<FName> GridNameList = UWorldPartitionRuntimeHashSet::ParseGridName(GridName);
+	return GridNameList.Num() && (GridNameList.Num() <= 2);
+}
+
+bool URuntimePartitionLevelStreaming::GenerateStreaming(const TArray<const IStreamingGenerationContext::FActorSetInstance*>& ActorSetInstances, TArray<FCellDesc>& OutRuntimeCellDescs)
+{
+	UWorldPartition* WorldPartition = GetTypedOuter<UWorldPartition>();
+	UWorld* World = WorldPartition->GetWorld();
+	UWorld* OuterWorld = GetTypedOuter<UWorld>();
+	const bool bIsMainWorldPartition = (World == OuterWorld);
+
+	TArray<IStreamingGenerationContext::FActorInstance> CellActorInstances;
+	if (PopulateCellActorInstances(ActorSetInstances, bIsMainWorldPartition, false, CellActorInstances))
+	{
+		TMap<FName, TPair<FName, TArray<IStreamingGenerationContext::FActorInstance>>> SubLevelsActorInstances;
+		for (const IStreamingGenerationContext::FActorInstance& ActorInstance : CellActorInstances)
+		{
+			TArray<FName> ActorSetGridNameList;
+					
+			if (!ActorInstance.ActorSetInstance->RuntimeGrid.IsNone())
+			{
+				ActorSetGridNameList = UWorldPartitionRuntimeHashSet::ParseGridName(ActorInstance.ActorSetInstance->RuntimeGrid);
+			}
+			else
+			{
+				ActorSetGridNameList.Add(NAME_Default);
+			}
+
+			if (bOneLevelPerActorContainer && !ActorInstance.GetContainerID().IsMainContainer())
+			{
+				ActorSetGridNameList.Add(*ActorInstance.GetContainerID().ToString());
+			}
+
+			TStringBuilder<512> StringBuilder;
+			for (FName GridName : ActorSetGridNameList)
+			{
+				StringBuilder += GridName.ToString();
+				StringBuilder += TEXT("_");
+			}
+			StringBuilder.RemoveSuffix(1);
+
+			FName SubLevelName = *StringBuilder;
+
+			TPair<FName, TArray<IStreamingGenerationContext::FActorInstance>>& Pair = SubLevelsActorInstances.FindOrAdd(SubLevelName);
+			Pair.Key = SubLevelName;
+			Pair.Value.Add(ActorInstance);
+		}
+
+		for (auto& [SubLevelName, SubLevelActorSetInstances] : SubLevelsActorInstances)
+		{
+			FCellDesc& CellDesc = OutRuntimeCellDescs.Emplace_GetRef();
+
+			CellDesc.Name = SubLevelActorSetInstances.Key;
+			CellDesc.bIsSpatiallyLoaded = true;
+			CellDesc.ContentBundleID = SubLevelActorSetInstances.Value[0].ActorSetInstance->ContentBundleID;
+			CellDesc.bBlockOnSlowStreaming = bBlockOnSlowStreaming;
+			CellDesc.bClientOnlyVisible = bClientOnlyVisible;
+			CellDesc.Priority = Priority;
+			CellDesc.ActorInstances = SubLevelActorSetInstances.Value;
+
+			for (const IStreamingGenerationContext::FActorInstance& ActorInstance : CellDesc.ActorInstances)
+			{
+				const FWorldPartitionActorDescView& ActorDescView = ActorInstance.GetActorDescView();
+				const FBox RuntimeBounds = ActorDescView.GetRuntimeBounds();
+				if (RuntimeBounds.IsValid)
+				{
+					CellDesc.Bounds += RuntimeBounds.TransformBy(ActorInstance.GetTransform());
+				}
+			}
+		}
+	}
+
+	return true;
+}
+#endif
\ No newline at end of file
diff --git a/Engine/Source/Runtime/Engine/Private/WorldPartition/RuntimeHashSet/RuntimePartitionPersistent.cpp b/Engine/Source/Runtime/Engine/Private/WorldPartition/RuntimeHashSet/RuntimePartitionPersistent.cpp
new file mode 100644
--- /dev/null
+++ b/Engine/Source/Runtime/Engine/Private/WorldPartition/RuntimeHashSet/RuntimePartitionPersistent.cpp
@@ -0,0 +1,30 @@
+// Copyright Epic Games, Inc. All Rights Reserved.
+
+#include "WorldPartition/RuntimeHashSet/RuntimePartitionPersistent.h"
+#include "WorldPartition/WorldPartitionStreamingGenerationContext.h"
+
+#if WITH_EDITOR
+bool URuntimePartitionPersistent::GenerateStreaming(const TArray<const IStreamingGenerationContext::FActorSetInstance*>& ActorSetInstances, TArray<FCellDesc>& OutRuntimeCellDescs)
+{
+	UWorldPartition* WorldPartition = GetTypedOuter<UWorldPartition>();
+	UWorld* World = WorldPartition->GetWorld();
+	UWorld* OuterWorld = GetTypedOuter<UWorld>();
+	const bool bIsMainWorldPartition = (World == OuterWorld);
+
+	TArray<IStreamingGenerationContext::FActorInstance> CellActorInstances;
+	if (PopulateCellActorInstances(ActorSetInstances, bIsMainWorldPartition, true, CellActorInstances))
+	{
+		FCellDesc& CellDesc = OutRuntimeCellDescs.Emplace_GetRef();
+
+		CellDesc.Name = NAME_PersistentLevel;
+		CellDesc.bIsSpatiallyLoaded = false;
+		CellDesc.ContentBundleID = CellActorInstances[0].ActorSetInstance->ContentBundleID;
+		CellDesc.bBlockOnSlowStreaming = bBlockOnSlowStreaming;
+		CellDesc.bClientOnlyVisible = bClientOnlyVisible;
+		CellDesc.Priority = Priority;
+		CellDesc.ActorInstances = CellActorInstances;
+	}
+
+	return true;
+}
+#endif
\ No newline at end of file
diff --git a/Engine/Source/Runtime/Engine/Private/WorldPartition/RuntimeHashSet/StaticSpatialIndex.cpp b/Engine/Source/Runtime/Engine/Private/WorldPartition/RuntimeHashSet/StaticSpatialIndex.cpp
new file mode 100644
--- /dev/null
+++ b/Engine/Source/Runtime/Engine/Private/WorldPartition/RuntimeHashSet/StaticSpatialIndex.cpp
@@ -0,0 +1,183 @@
+// Copyright Epic Games, Inc. All Rights Reserved.
+
+#include "WorldPartition/RuntimeHashSet/StaticSpatialIndex.h"
+#include "Algo/Transform.h"
+
+void FStaticSpatialIndex::FListImpl::Init(const TArray<TPair<FBox, uint32>>& InElements)
+{
+	Elements.Reserve(InElements.Num());
+	Algo::Transform(InElements, Elements, [](const TPair<FBox, uint32>& Element) { return Element.Value; });
+}
+
+void FStaticSpatialIndex::FListImpl::ForEachElement(TFunctionRef<void(uint32 InValueIndex)> Func) const
+{
+	for (uint32 ValueIndex : Elements)
+	{
+		Func(ValueIndex);
+	}
+}
+
+void FStaticSpatialIndex::FListImpl::ForEachIntersectingElement(const FBox& InBox, TFunctionRef<void(uint32 InValueIndex)> Func) const
+{
+	for (uint32 ValueIndex : Elements)
+	{
+		const FBox& Box = DataInterface.GetBox(ValueIndex);
+
+		if (Box.Intersect(InBox))
+		{
+			Func(ValueIndex);
+		}
+	}
+}
+
+void FStaticSpatialIndex::FListImpl::ForEachIntersectingElement(const FSphere& InSphere, TFunctionRef<void(uint32 InValueIndex)> Func) const
+{
+	for (uint32 ValueIndex : Elements)
+	{
+		const FBox& Box = DataInterface.GetBox(ValueIndex);
+
+		if (FMath::SphereAABBIntersection(InSphere, Box))
+		{
+			Func(ValueIndex);
+		}
+	}
+}
+
+void FStaticSpatialIndex::FRTreeImpl::Init(const TArray<TPair<FBox, uint32>>& InElements)
+{
+	if (InElements.Num())
+	{
+		const int32 MaxNumElementsPerNode = 16;
+		const int32 MaxNumElementsPerLeaf = 64;
+
+		TArray<TPair<FBox, uint32>> SortedElements = InElements;
+		SortedElements.Sort([](const TPair<FBox, uint32>& A, const TPair<FBox, uint32>& B) { return A.Key.Min.X < B.Key.Min.X; });
+	
+		// build leaves
+		FNode* CurrentNode = nullptr;
+		TArray<FNode> Nodes;
+
+		for (const TPair<FBox, uint32>& Element : SortedElements)
+		{
+			if (!CurrentNode || (CurrentNode->Content.Get<FNode::FLeafType>().Num() >= MaxNumElementsPerLeaf))
+			{
+				CurrentNode = &Nodes.Emplace_GetRef();
+				CurrentNode->Content.Emplace<FNode::FLeafType>();
+			}
+
+			CurrentNode->Box += Element.Key;
+			CurrentNode->Content.Get<FNode::FLeafType>().Add(Element.Value);
+		}
+
+		// build nodes
+		while (Nodes.Num() > 1)
+		{
+			CurrentNode = nullptr;
+			TArray<FNode> TopNodes;
+
+			Nodes.Sort([](const FNode& A, const FNode& B) { return A.Box.Min.X < B.Box.Min.X; });
+
+			for (FNode& Node : Nodes)
+			{
+				if (!CurrentNode || (CurrentNode->Content.Get<FNode::FNodeType>().Num() >= MaxNumElementsPerNode))
+				{
+					CurrentNode = &TopNodes.Emplace_GetRef();
+					CurrentNode->Content.Emplace<FNode::FNodeType>();
+				}
+
+				CurrentNode->Box += Node.Box;
+				CurrentNode->Content.Get<FNode::FNodeType>().Add(MoveTemp(Node));
+			}
+
+			Nodes = MoveTemp(TopNodes);
+		}
+
+		check(Nodes.Num() == 1);
+		RootNode = MoveTemp(Nodes[0]);
+	}
+}
+
+void FStaticSpatialIndex::FRTreeImpl::ForEachElement(TFunctionRef<void(uint32 InValueIndex)> Func) const
+{
+	ForEachElementRecursive(&RootNode, Func);
+}
+
+void FStaticSpatialIndex::FRTreeImpl::ForEachElementRecursive(const FNode* Node, TFunctionRef<void(uint32 InValueIndex)> Func) const
+{
+	if (Node->Content.IsType<FNode::FNodeType>())
+	{
+		for (auto& ChildNode : Node->Content.Get<FNode::FNodeType>())
+		{
+			ForEachElementRecursive(&ChildNode, Func);
+		}
+	}
+	else
+	{
+		for (uint32 ValueIndex : Node->Content.Get<FNode::FLeafType>())
+		{
+			Func(ValueIndex);
+		}
+	}
+}
+
+void FStaticSpatialIndex::FRTreeImpl::ForEachIntersectingElement(const FBox& InBox, TFunctionRef<void(uint32 InValueIndex)> Func) const
+{
+	ForEachIntersectingElementRecursive(&RootNode, InBox, Func);
+}
+
+void FStaticSpatialIndex::FRTreeImpl::ForEachIntersectingElementRecursive(const FNode* Node, const FBox& InBox, TFunctionRef<void(uint32 InValueIndex)> Func) const
+{
+	if (Node->Content.IsType<FNode::FNodeType>())
+	{
+		for (auto& ChildNode : Node->Content.Get<FNode::FNodeType>())
+		{
+			if (ChildNode.Box.Intersect(InBox))
+			{
+				ForEachIntersectingElementRecursive(&ChildNode, InBox, Func);
+			}
+		}
+	}
+	else
+	{
+		for (uint32 ValueIndex : Node->Content.Get<FNode::FLeafType>())
+		{
+			const FBox& Box = DataInterface.GetBox(ValueIndex);
+
+			if (Box.Intersect(InBox))
+			{
+				Func(ValueIndex);
+			}
+		}
+	}
+}
+
+void FStaticSpatialIndex::FRTreeImpl::ForEachIntersectingElement(const FSphere& InSphere, TFunctionRef<void(uint32 InValueIndex)> Func) const
+{
+	ForEachIntersectingElementRecursive(&RootNode, InSphere, Func);
+}
+
+void FStaticSpatialIndex::FRTreeImpl::ForEachIntersectingElementRecursive(const FNode* Node, const FSphere& InSphere, TFunctionRef<void(uint32 InValueIndex)> Func) const
+{
+	if (Node->Content.IsType<FNode::FNodeType>())
+	{
+		for (auto& ChildNode : Node->Content.Get<FNode::FNodeType>())
+		{
+			if (FMath::SphereAABBIntersection(InSphere, ChildNode.Box))
+			{
+				ForEachIntersectingElementRecursive(&ChildNode, InSphere, Func);
+			}
+		}
+	}
+	else
+	{
+		for (uint32 ValueIndex : Node->Content.Get<FNode::FLeafType>())
+		{
+			const FBox& Box = DataInterface.GetBox(ValueIndex);
+
+			if (FMath::SphereAABBIntersection(InSphere, Box))
+			{
+				Func(ValueIndex);
+			}
+		}
+	}
+}
\ No newline at end of file
diff --git a/Engine/Source/Runtime/Engine/Private/WorldPartition/RuntimeHashSet/WorldPartitionRuntimeHashSet.cpp b/Engine/Source/Runtime/Engine/Private/WorldPartition/RuntimeHashSet/WorldPartitionRuntimeHashSet.cpp
new file mode 100644
--- /dev/null
+++ b/Engine/Source/Runtime/Engine/Private/WorldPartition/RuntimeHashSet/WorldPartitionRuntimeHashSet.cpp
@@ -0,0 +1,717 @@
+// Copyright Epic Games, Inc. All Rights Reserved.
+
+#include "WorldPartition/RuntimeHashSet/WorldPartitionRuntimeHashSet.h"
+#include "WorldPartition/RuntimeHashSet/RuntimePartition.h"
+#include "WorldPartition/RuntimeHashSet/RuntimePartitionPersistent.h"
+#include "WorldPartition/ContentBundle/ContentBundleDescriptor.h"
+#include "WorldPartition/WorldPartitionStreamingGenerationContext.h"
+#include "WorldPartition/WorldPartitionRuntimeSpatialHash.h"
+#include "WorldPartition/WorldPartitionStreamingPolicy.h"
+#include "WorldPartition/WorldPartitionStreamingSource.h"
+#include "WorldPartition/DataLayer/DataLayersID.h"
+#include "WorldPartition/HLOD/HLODLayer.h"
+#include "Algo/ForEach.h"
+#include "Algo/Transform.h"
+#include "Misc/ArchiveMD5.h"
+
+/**
+ *	TODOs:
+ *		- Deprecate WorldPartitionStreamingSourceComponent.TargetGrids and make it more generic to handle all sorts of filtering.
+ *		- Deprecate Actor.RuntimeGrid and introduce TargetPartition + HLOD setup index.
+ *		- Refactor Source.ForEachShape to extract target grid and maybe loading range.
+ */
+
+FAutoConsoleCommand WorldPartitionRuntimeHashSetEnable(
+	TEXT("wp.Editor.WorldPartitionRuntimeHashSet.Enable"),
+	TEXT("Enable experimental runtime hash set class."),
+	FConsoleCommandWithArgsDelegate::CreateLambda([](const TArray<FString>& Args)
+	{
+		if (UClass* WorldPartitionRuntimeHashSetClass = FindObject<UClass>(nullptr, TEXT("/Script/Engine.WorldPartitionRuntimeHashSet")))
+		{
+			WorldPartitionRuntimeHashSetClass->ClassFlags &= ~CLASS_HideDropDown;
+		}
+	})
+);
+
+UWorldPartitionRuntimeHashSet::UWorldPartitionRuntimeHashSet(const FObjectInitializer& ObjectInitializer)
+	: Super(ObjectInitializer)
+{}
+
+void UWorldPartitionRuntimeHashSet::Serialize(FArchive& Ar)
+{
+	Super::Serialize(Ar);
+
+#if WITH_EDITOR
+	if (Ar.GetPortFlags() & PPF_DuplicateForPIE)
+	{
+		bool bSpatialIndexValid = SpatialIndex.IsValid();
+		Ar << bSpatialIndexValid;
+
+		if (bSpatialIndexValid)
+		{
+			if (Ar.IsLoading())
+			{
+				check(!SpatialIndex);
+				SpatialIndex = MakeUnique<FStaticSpatialIndexType>();
+			}
+
+			SpatialIndex->Serialize(Ar);
+		}
+	}
+#endif
+}
+
+void UWorldPartitionRuntimeHashSet::AddReferencedObjects(UObject* InThis, FReferenceCollector& Collector)
+{
+#if WITH_EDITOR
+	UWorldPartitionRuntimeHashSet* This = CastChecked<UWorldPartitionRuntimeHashSet>(InThis);
+	if (This->SpatialIndex)
+	{
+		This->SpatialIndex->AddReferencedObjects(Collector);
+	}
+	Collector.AddStableReferenceSet(&This->InjectedExternalStreamingObjects);
+#endif
+
+	Super::AddReferencedObjects(InThis, Collector);
+}
+
+#if WITH_EDITOR
+void UWorldPartitionRuntimeHashSet::SetDefaultValues()
+{
+}
+
+bool UWorldPartitionRuntimeHashSet::SupportsHLODs() const
+{
+	for (const FRuntimePartitionDesc& RuntimePartitionDesc : RuntimePartitions)
+	{
+		if (RuntimePartitionDesc.Partitions.Num())
+		{
+			if (RuntimePartitionDesc.Partitions[0]->SupportsHLODs())
+			{
+				return true;
+			}
+		}
+	}
+
+	return false;
+}
+
+bool UWorldPartitionRuntimeHashSet::GenerateStreaming(UWorldPartitionStreamingPolicy* StreamingPolicy, const IStreamingGenerationContext* StreamingGenerationContext, TArray<FString>* OutPackagesToGenerate)
+{
+	verify(Super::GenerateStreaming(StreamingPolicy, StreamingGenerationContext, OutPackagesToGenerate));
+
+	UWorldPartition* WorldPartition = GetOuterUWorldPartition();
+	UWorld* World = WorldPartition->GetWorld();
+	UWorld* OuterWorld = GetTypedOuter<UWorld>();
+	const bool bIsMainWorldPartition = (World == OuterWorld);
+
+	check(!PersistentPartitionDesc.Class);
+	PersistentPartitionDesc.Class = URuntimePartitionPersistent::StaticClass();
+	PersistentPartitionDesc.Name = NAME_PersistentLevel;
+	PersistentPartitionDesc.Partitions.Add(NewObject<URuntimePartition>(this, URuntimePartitionPersistent::StaticClass(), NAME_None));
+	PersistentPartitionDesc.Partitions[0]->Name = TEXT("MainPartition");
+
+	//
+	// Split actor sets into their corresponding runtime partition implementation
+	//
+	TMap<FName, const FRuntimePartitionDesc*> NameToRuntimePartitionDescMap;
+
+	NameToRuntimePartitionDescMap.Add(NAME_None, &RuntimePartitions[0]);				// Actors with RuntimeGrid=None will be assigned to the default partition
+	NameToRuntimePartitionDescMap.Add(NAME_PersistentLevel, &PersistentPartitionDesc);	// Non-spatially loaded actors will be assigned to the persistent partition
+
+	for (const FRuntimePartitionDesc& RuntimePartitionDesc : RuntimePartitions)
+	{
+		NameToRuntimePartitionDescMap.Add(RuntimePartitionDesc.Name, &RuntimePartitionDesc);
+	}
+
+	TMap<URuntimePartition*, TArray<const IStreamingGenerationContext::FActorSetInstance*>> RuntimePartitionsToActorSetMap;
+	StreamingGenerationContext->ForEachActorSetInstance([this, &NameToRuntimePartitionDescMap, &RuntimePartitionsToActorSetMap](const IStreamingGenerationContext::FActorSetInstance& ActorSetInstance)
+	{
+		const TArray<FName> ActorSetRuntimeGrid = ActorSetInstance.bIsSpatiallyLoaded ? ParseGridName(ActorSetInstance.RuntimeGrid) : TArray<FName>({ NAME_PersistentLevel });
+
+		if (const FRuntimePartitionDesc** RuntimePartitionDesc = NameToRuntimePartitionDescMap.Find(ActorSetRuntimeGrid[0]))
+		{
+			const uint32 PartitionLayerIndex = 0;
+			RuntimePartitionsToActorSetMap.FindOrAdd((*RuntimePartitionDesc)->Partitions[PartitionLayerIndex]).Add(&ActorSetInstance);
+		}
+	});
+
+	// Generate per data layer instance of partition cells
+	struct FCellDescInstance : public URuntimePartition::FCellDesc
+	{
+		FBox Bounds;
+		TArray<const UDataLayerInstance*> DataLayers;
+	};
+
+	TArray<FCellDescInstance> RuntimeCellDescsDataLayers;
+	{
+		TArray<URuntimePartition::FCellDesc> RuntimeCellDescs;
+		for (auto [RuntimePartition, ActorSetInstances] : RuntimePartitionsToActorSetMap)
+		{
+			if (!RuntimePartition->GenerateStreaming(ActorSetInstances, RuntimeCellDescs))
+			{
+				return false;
+			}
+		}
+
+		// Split cell descs into data layers	
+		for (const URuntimePartition::FCellDesc& RuntimeCellDesc : RuntimeCellDescs)
+		{
+			TMap<FDataLayersID, FCellDescInstance> RuntimeCellDescsDataLayersSet;
+
+			for (const IStreamingGenerationContext::FActorInstance& ActorInstance : RuntimeCellDesc.ActorInstances)
+			{
+				const FDataLayersID DataLayersID(ActorInstance.ActorSetInstance->DataLayers);
+
+				FCellDescInstance* DataLayerCellDesc = RuntimeCellDescsDataLayersSet.Find(DataLayersID);
+			
+				if (!DataLayerCellDesc)
+				{
+					DataLayerCellDesc = &RuntimeCellDescsDataLayersSet.Emplace(DataLayersID);
+
+					DataLayerCellDesc->Name = RuntimeCellDesc.Name;
+					DataLayerCellDesc->bIsSpatiallyLoaded = RuntimeCellDesc.bIsSpatiallyLoaded;
+					DataLayerCellDesc->ContentBundleID = RuntimeCellDesc.ContentBundleID;
+					DataLayerCellDesc->bBlockOnSlowStreaming = RuntimeCellDesc.bBlockOnSlowStreaming;
+					DataLayerCellDesc->bClientOnlyVisible = RuntimeCellDesc.bClientOnlyVisible;
+					DataLayerCellDesc->Priority = RuntimeCellDesc.Priority;
+					DataLayerCellDesc->Bounds = RuntimeCellDesc.Bounds;
+					DataLayerCellDesc->DataLayers = ActorInstance.ActorSetInstance->DataLayers;
+				}
+
+				DataLayerCellDesc->ActorInstances.Add(ActorInstance);
+			}
+
+			Algo::Transform(RuntimeCellDescsDataLayersSet, RuntimeCellDescsDataLayers, [](const auto& Value) { return Value.Value; });
+		}
+	}
+
+	// Generate runtime cells
+	auto CreateRuntimeCellFromCellDesc = [this](const URuntimePartition::FCellDesc& CellDesc, const TArray<const UDataLayerInstance*>& DataLayers, TSubclassOf<UWorldPartitionRuntimeCell> CellClass, TSubclassOf<UWorldPartitionRuntimeCellData> CellDataClass)
+	{
+		FString CellObjectName;
+		FGuid CellGuid;
+		{
+			UWorld* OuterWorld = GetTypedOuter<UWorld>();
+			check(OuterWorld);
+
+			FString WorldName = FPackageName::GetShortName(OuterWorld->GetPackage());
+
+			CellObjectName = FString::Printf(TEXT("%s_%s"), *WorldName, *CellDesc.Name.ToString());
+
+			const FDataLayersID DataLayersID(DataLayers);
+			if (DataLayersID.GetHash())
+			{
+				CellObjectName += FString::Printf(TEXT("_d%X"), DataLayersID.GetHash());
+			}
+
+			if (CellDesc.ContentBundleID.IsValid())
+			{
+				CellObjectName += FString::Printf(TEXT("_c%s"), *UContentBundleDescriptor::GetContentBundleCompactString(CellDesc.ContentBundleID));
+			}
+
+			if (!IsRunningCookCommandlet() && OuterWorld->IsGameWorld())
+			{
+				FString SourceWorldPath;
+				FString InstancedWorldPath;
+				if (OuterWorld->GetSoftObjectPathMapping(SourceWorldPath, InstancedWorldPath))
+				{
+					const FTopLevelAssetPath SourceAssetPath(SourceWorldPath);
+					WorldName = FPackageName::GetShortName(SourceAssetPath.GetPackageName());
+						
+					InstancedWorldPath = UWorld::RemovePIEPrefix(InstancedWorldPath);
+
+					const FString SourcePackageName = SourceAssetPath.GetPackageName().ToString();
+					const FTopLevelAssetPath InstanceAssetPath(InstancedWorldPath);
+					const FString InstancePackageName = InstanceAssetPath.GetPackageName().ToString();
+
+					if (int32 Index = InstancePackageName.Find(SourcePackageName); Index != INDEX_NONE)
+					{
+						CellObjectName += FString::Printf(TEXT("_i%s"), *InstancePackageName.Mid(Index + SourcePackageName.Len()));
+					}
+				}
+			}
+
+			FArchiveMD5 ArMD5;
+			ArMD5 << CellObjectName;
+			CellGuid = ArMD5.GetGuidFromHash();
+			check(CellGuid.IsValid());
+		}
+
+		UWorldPartitionRuntimeCell* RuntimeCell = Super::CreateRuntimeCell(CellClass, CellDataClass, CellObjectName, TEXT(""), this);
+
+		RuntimeCell->SetIsAlwaysLoaded(!CellDesc.bIsSpatiallyLoaded);
+		RuntimeCell->SetDataLayers(DataLayers);
+		RuntimeCell->SetContentBundleUID(CellDesc.ContentBundleID);
+		RuntimeCell->SetPriority(CellDesc.Priority);
+		RuntimeCell->SetClientOnlyVisible(CellDesc.bClientOnlyVisible);
+		RuntimeCell->SetBlockOnSlowLoading(CellDesc.bBlockOnSlowStreaming);
+		RuntimeCell->SetIsHLOD(false);
+		RuntimeCell->SetGuid(CellGuid);
+		RuntimeCell->RuntimeCellData->DebugName = CellObjectName;
+
+		return RuntimeCell;
+	};
+
+	TArray<UWorldPartitionRuntimeCell*> RuntimeCells;	
+	for (const FCellDescInstance& CellDescInstance : RuntimeCellDescsDataLayers)
+	{
+		UWorldPartitionRuntimeCell* RuntimeCell = RuntimeCells.Emplace_GetRef(CreateRuntimeCellFromCellDesc(CellDescInstance, CellDescInstance.DataLayers, StreamingPolicy->GetRuntimeCellClass(), UWorldPartitionRuntimeCellData::StaticClass()));
+		PopulateRuntimeCell(RuntimeCell, CellDescInstance.ActorInstances, nullptr);
+
+		// Override the cell bounds if the runtime partition provided one
+		if (CellDescInstance.Bounds.IsValid)
+		{
+			RuntimeCell->RuntimeCellData->ContentBounds = CellDescInstance.Bounds;
+		}
+	}
+
+	if (OutPackagesToGenerate)
+	{
+		for (UWorldPartitionRuntimeCell* RuntimeCell : RuntimeCells)
+		{
+			// Always loaded cell actors are transfered to World's Persistent Level (see UWorldPartitionRuntimeSpatialHash::PopulateGeneratorPackageForCook)
+			if (RuntimeCell->GetActorCount() && !RuntimeCell->IsAlwaysLoaded())
+			{
+				const FString PackageRelativePath = RuntimeCell->GetPackageNameToCreate();
+				check(!PackageRelativePath.IsEmpty());
+
+				OutPackagesToGenerate->Add(PackageRelativePath);
+
+				// Map relative package to StreamingCell for PopulateGeneratedPackageForCook/PopulateGeneratorPackageForCook/GetCellForPackage
+				PackagesToGenerateForCook.Add(PackageRelativePath, RuntimeCell);
+			}
+		}
+	}
+
+	// Init spatial index
+	check(!SpatialIndex);
+	SpatialIndex = MakeUnique<FStaticSpatialIndexType>();
+
+	TSet<FGuid> RuntimeCellsGuids;
+	RuntimeCellsGuids.Reserve(RuntimeCells.Num());
+
+	TArray<TPair<FBox, UWorldPartitionRuntimeCell*>> Elements;
+	Algo::ForEach(RuntimeCells, [this, &Elements, &RuntimeCellsGuids](UWorldPartitionRuntimeCell* RuntimeCell)
+	{
+		bool bCellWasAlreadyInSet;
+		RuntimeCellsGuids.Add(RuntimeCell->GetGuid(), &bCellWasAlreadyInSet);
+		check(!bCellWasAlreadyInSet);
+
+		if (RuntimeCell->IsAlwaysLoaded())
+		{
+			NonSpatiallyLoadedRuntimeCells.Add(RuntimeCell);
+		}
+		else
+		{
+			Elements.Add(TPair<FBox, UWorldPartitionRuntimeCell*>(RuntimeCell->GetContentBounds(), RuntimeCell));
+		}
+	});
+
+	SpatialIndex->Init(Elements);
+
+	return true;
+}
+
+void UWorldPartitionRuntimeHashSet::FlushStreaming()
+{
+	Super::FlushStreaming();
+	
+	check(PersistentPartitionDesc.Class);
+	PersistentPartitionDesc.Class = nullptr;
+	PersistentPartitionDesc.Name = NAME_None;
+	PersistentPartitionDesc.Partitions.Empty();
+
+	NonSpatiallyLoadedRuntimeCells.Empty();	
+	SpatialIndex.Reset();
+}
+
+bool UWorldPartitionRuntimeHashSet::IsValidGrid(FName GridName) const
+{
+	// The None grid name will always map to the first runtime partition in the list
+	if (GridName.IsNone())
+	{
+		return true;
+	}
+
+	// Parse the potentially dot separated grid name to identiy the associated runtime partition
+	const TArray<FName> GridNameList = ParseGridName(GridName);
+	for (const FRuntimePartitionDesc& RuntimePartitionDesc : RuntimePartitions)
+	{
+		if (RuntimePartitionDesc.Name == GridNameList[0])
+		{
+			if (RuntimePartitionDesc.Partitions.Num())
+			{
+				return RuntimePartitionDesc.Partitions[0]->IsValidGrid(GridName);
+			}
+		}
+	}
+
+	return false;
+}
+
+TArray<UWorldPartitionRuntimeCell*> UWorldPartitionRuntimeHashSet::GetAlwaysLoadedCells() const
+{
+	return NonSpatiallyLoadedRuntimeCells;
+}
+
+void UWorldPartitionRuntimeHashSet::DumpStateLog(FHierarchicalLogArchive& Ar) const
+{
+	Super::DumpStateLog(Ar);
+
+	Ar.Printf(TEXT("----------------------------------------------------------------------------------------------------------------"));
+	Ar.Printf(TEXT("%s - Runtime Hash Set"), *GetWorld()->GetName());
+	Ar.Printf(TEXT("----------------------------------------------------------------------------------------------------------------"));
+
+	TArray<const UWorldPartitionRuntimeCell*> StreamingCells;
+	ForEachStreamingCells([&StreamingCells](const UWorldPartitionRuntimeCell* StreamingCell) { StreamingCells.Add(StreamingCell); return true; });
+				
+	StreamingCells.Sort([this](const UWorldPartitionRuntimeCell& A, const UWorldPartitionRuntimeCell& B) { return A.GetFName().LexicalLess(B.GetFName()); });
+
+	for (const UWorldPartitionRuntimeCell* StreamingCell : StreamingCells)
+	{
+		FHierarchicalLogArchive::FIndentScope CellIndentScope = Ar.PrintfIndent(TEXT("Content of Cell %s (%s)"), *StreamingCell->GetDebugName(), *StreamingCell->GetName());
+		StreamingCell->DumpStateLog(Ar);
+	}
+
+	Ar.Printf(TEXT(""));
+}
+
+TArray<FName> UWorldPartitionRuntimeHashSet::ParseGridName(FName GridName)
+{
+	TArray<FString> GridNameList;
+	const FString GridNameStr = GridName.ToString();
+	if (GridNameStr.ParseIntoArray(GridNameList, TEXT(".")))
+	{
+		TArray<FName> Result;
+		Algo::Transform(GridNameList, Result, [](const FString& GridName) { return *GridName; });
+		return MoveTemp(Result);
+	}
+	return { GridName };
+}
+
+URuntimeHashExternalStreamingObjectBase* UWorldPartitionRuntimeHashSet::StoreToExternalStreamingObject(UObject* StreamingObjectOuter, FName StreamingObjectName)
+{
+	URuntimeHashSetExternalStreamingObject* StreamingObject = CreateExternalStreamingObject<URuntimeHashSetExternalStreamingObject>(StreamingObjectOuter, StreamingObjectName);
+	StreamingObject->NonSpatiallyLoadedRuntimeCells = NonSpatiallyLoadedRuntimeCells;
+	SpatialIndex->ForEachElement([StreamingObject](UWorldPartitionRuntimeCell* RuntimeCell) { StreamingObject->SpatiallyLoadedRuntimeCells.Add(RuntimeCell); });
+	return StreamingObject;
+}
+#endif
+
+bool UWorldPartitionRuntimeHashSet::InjectExternalStreamingObject(URuntimeHashExternalStreamingObjectBase* ExternalStreamingObject)
+{
+	URuntimeHashSetExternalStreamingObject* HashSetExternalStreamingObject = CastChecked<URuntimeHashSetExternalStreamingObject>(ExternalStreamingObject);
+
+	bool bWasAlreadyInSet;
+	InjectedExternalStreamingObjects.Add(HashSetExternalStreamingObject, &bWasAlreadyInSet);
+	check(!bWasAlreadyInSet);
+
+	check(!HashSetExternalStreamingObject->SpatialIndex);
+	HashSetExternalStreamingObject->SpatialIndex = MakeUnique<FStaticSpatialIndexType>();
+
+	TArray<TPair<FBox, UWorldPartitionRuntimeCell*>> Elements;
+	Algo::Transform(HashSetExternalStreamingObject->SpatiallyLoadedRuntimeCells, Elements, 
+		[](UWorldPartitionRuntimeCell* RuntimeCell) { return TPair<FBox, UWorldPartitionRuntimeCell*>(RuntimeCell->GetContentBounds(), RuntimeCell); }
+	);	
+	HashSetExternalStreamingObject->SpatialIndex->Init(Elements);
+
+	return true;
+}
+
+bool UWorldPartitionRuntimeHashSet::RemoveExternalStreamingObject(URuntimeHashExternalStreamingObjectBase* ExternalStreamingObject)
+{
+	URuntimeHashSetExternalStreamingObject* HashSetExternalStreamingObject = CastChecked<URuntimeHashSetExternalStreamingObject>(ExternalStreamingObject);
+
+	check(HashSetExternalStreamingObject->SpatialIndex);
+	HashSetExternalStreamingObject->SpatialIndex.Reset();
+
+	verify(InjectedExternalStreamingObjects.Remove(HashSetExternalStreamingObject));
+
+	return true;
+}
+
+// Streaming interface
+void UWorldPartitionRuntimeHashSet::ForEachStreamingCells(TFunctionRef<bool(const UWorldPartitionRuntimeCell*)> Func) const
+{
+	auto ForEachStreamingCells = [this, &Func](FStaticSpatialIndexType* InSpatialIndex)
+	{
+		if (InSpatialIndex)
+		{
+			InSpatialIndex->ForEachElement([this, &Func](const UWorldPartitionRuntimeCell* RuntimeCell)
+			{
+				if (IsCellRelevantFor(RuntimeCell->GetClientOnlyVisible()))
+				{
+					Func(RuntimeCell);
+				}
+			});
+		}
+	};
+
+	auto ForEachNonStreamingCells = [this, &Func](TArray<TObjectPtr<UWorldPartitionRuntimeCell>> InNonSpatiallyLoadedRuntimeCells)
+	{
+		for (UWorldPartitionRuntimeCell* Cell : InNonSpatiallyLoadedRuntimeCells)
+		{
+			if (IsCellRelevantFor(Cell->GetClientOnlyVisible()))
+			{
+				Func(Cell);
+			}
+		}
+	};
+
+	ForEachStreamingCells(SpatialIndex.Get());
+	ForEachNonStreamingCells(NonSpatiallyLoadedRuntimeCells);
+
+	for (URuntimeHashSetExternalStreamingObject* InjectedExternalStreamingObject : InjectedExternalStreamingObjects)
+	{
+		ForEachStreamingCells(InjectedExternalStreamingObject->SpatialIndex.Get());
+		ForEachNonStreamingCells(InjectedExternalStreamingObject->NonSpatiallyLoadedRuntimeCells);
+	}
+}
+
+void UWorldPartitionRuntimeHashSet::ForEachStreamingCellsQuery(const FWorldPartitionStreamingQuerySource& QuerySource, TFunctionRef<bool(const UWorldPartitionRuntimeCell*)> Func, FWorldPartitionQueryCache* QueryCache) const
+{
+	auto ShouldAddCell = [this](const UWorldPartitionRuntimeCell* Cell, const FWorldPartitionStreamingQuerySource& QuerySource)
+	{
+		if (IsCellRelevantFor(Cell->GetClientOnlyVisible()))
+		{
+			if (Cell->HasDataLayers())
+			{
+				if (Cell->GetDataLayers().FindByPredicate([&](const FName& DataLayerName) { return QuerySource.DataLayers.Contains(DataLayerName); }))
+				{
+					return true;
+				}
+			}
+			else if (!QuerySource.bDataLayersOnly)
+			{
+				return true;
+			}
+		}
+
+		return false;
+	};
+
+	auto ForEachStreamingCells = [&ShouldAddCell, &QuerySource, &Func](FStaticSpatialIndexType* InSpatialIndex)
+	{
+		if (InSpatialIndex)
+		{
+			InSpatialIndex->ForEachElement([&ShouldAddCell, &QuerySource, &Func](UWorldPartitionRuntimeCell* RuntimeCell)
+			{
+				if (ShouldAddCell(RuntimeCell, QuerySource))
+				{
+					Func(RuntimeCell);
+				}
+			});
+		}
+	};
+
+	auto ForEachNonStreamingCells = [&ShouldAddCell, &QuerySource, &Func](TArray<TObjectPtr<UWorldPartitionRuntimeCell>> InNonSpatiallyLoadedRuntimeCells)
+	{
+		for (UWorldPartitionRuntimeCell* Cell : InNonSpatiallyLoadedRuntimeCells)
+		{
+			if (ShouldAddCell(Cell, QuerySource))
+			{
+				Func(Cell);
+			}
+		}
+	};
+
+	ForEachStreamingCells(SpatialIndex.Get());
+	ForEachNonStreamingCells(NonSpatiallyLoadedRuntimeCells);
+
+	for (URuntimeHashSetExternalStreamingObject* InjectedExternalStreamingObject : InjectedExternalStreamingObjects)
+	{
+		ForEachStreamingCells(InjectedExternalStreamingObject->SpatialIndex.Get());
+		ForEachNonStreamingCells(InjectedExternalStreamingObject->NonSpatiallyLoadedRuntimeCells);
+	}
+}
+
+void UWorldPartitionRuntimeHashSet::ForEachStreamingCellsSources(const TArray<FWorldPartitionStreamingSource>& Sources, TFunctionRef<bool(const UWorldPartitionRuntimeCell*, EStreamingSourceTargetState)> Func) const
+{
+	UWorldPartitionRuntimeHash::FStreamingSourceCells ActivateStreamingSourceCells;
+	UWorldPartitionRuntimeHash::FStreamingSourceCells LoadStreamingSourceCells;
+
+	auto ForEachStreamingCells = [this, &Sources, &Func, &ActivateStreamingSourceCells, &LoadStreamingSourceCells](FStaticSpatialIndexType* InSpatialIndex)
+	{
+		if (InSpatialIndex)
+		{
+			for (const FWorldPartitionStreamingSource& Source : Sources)
+			{
+				// @todo_jfd
+				const FName GridName;
+				const FSoftObjectPath HLODLayer;
+
+				Source.ForEachShape(/*LoadingRange*/25600, GridName, HLODLayer, false, [this, &Source, InSpatialIndex, &Func, &ActivateStreamingSourceCells, &LoadStreamingSourceCells](const FSphericalSector& Shape)
+				{
+					const FSphere ShapeSphere(Shape.GetCenter(), Shape.GetRadius());
+
+					InSpatialIndex->ForEachIntersectingElement(ShapeSphere, [this, &Source, &Shape, &ActivateStreamingSourceCells, &LoadStreamingSourceCells](UWorldPartitionRuntimeCell* Cell)
+					{
+						if (IsCellRelevantFor(Cell->GetClientOnlyVisible()))
+						{
+							if (!Cell->HasDataLayers() || Cell->HasAnyDataLayerInEffectiveRuntimeState(EDataLayerRuntimeState::Activated))
+							{
+								if (Source.TargetState == EStreamingSourceTargetState::Loaded)
+								{
+									LoadStreamingSourceCells.AddCell(Cell, Source, Shape);
+								}
+								else
+								{
+									ActivateStreamingSourceCells.AddCell(Cell, Source, Shape);
+								}
+							}
+							else if (Cell->HasAnyDataLayerInEffectiveRuntimeState(EDataLayerRuntimeState::Loaded))
+							{
+								LoadStreamingSourceCells.AddCell(Cell, Source, Shape);
+							}
+						}
+					});
+				});
+			}
+		}
+	};
+
+	auto ForEachNonStreamingCells = [this, &ActivateStreamingSourceCells, &LoadStreamingSourceCells](TArray<TObjectPtr<UWorldPartitionRuntimeCell>> InNonSpatiallyLoadedRuntimeCells)
+	{
+		for (UWorldPartitionRuntimeCell* Cell : InNonSpatiallyLoadedRuntimeCells)
+		{
+			if (IsCellRelevantFor(Cell->GetClientOnlyVisible()))
+			{
+				if (!Cell->HasDataLayers() || Cell->HasAnyDataLayerInEffectiveRuntimeState(EDataLayerRuntimeState::Activated))
+				{
+					ActivateStreamingSourceCells.GetCells().Add(Cell);
+				}
+				else if (Cell->HasAnyDataLayerInEffectiveRuntimeState(EDataLayerRuntimeState::Loaded))
+				{
+					LoadStreamingSourceCells.GetCells().Add(Cell);
+				}
+			}
+		}
+	};
+
+	ForEachStreamingCells(SpatialIndex.Get());
+	ForEachNonStreamingCells(NonSpatiallyLoadedRuntimeCells);
+
+	for (URuntimeHashSetExternalStreamingObject* InjectedExternalStreamingObject : InjectedExternalStreamingObjects)
+	{
+		ForEachStreamingCells(InjectedExternalStreamingObject->SpatialIndex.Get());
+		ForEachNonStreamingCells(InjectedExternalStreamingObject->NonSpatiallyLoadedRuntimeCells);
+	}
+
+	auto ExecuteFuncOnCells = [Func](const TSet<const UWorldPartitionRuntimeCell*>& Cells, EStreamingSourceTargetState TargetState)
+	{
+		for (const UWorldPartitionRuntimeCell* Cell : Cells)
+		{
+			Func(Cell, TargetState);
+		}
+	};
+
+	ExecuteFuncOnCells(ActivateStreamingSourceCells.GetCells(), EStreamingSourceTargetState::Activated);
+	ExecuteFuncOnCells(LoadStreamingSourceCells.GetCells(), EStreamingSourceTargetState::Loaded);
+}
+
+#if WITH_EDITOR
+void UWorldPartitionRuntimeHashSet::PostEditChangeChainProperty(FPropertyChangedChainEvent& PropertyChangedEvent)
+{
+	Super::PostEditChangeChainProperty(PropertyChangedEvent);
+
+	static FName NAME_RuntimePartitions(TEXT("RuntimePartitions"));
+	static FName NAME_HLODSetups_Key(TEXT("HLODSetups_Key"));
+
+	FName PropertyName = PropertyChangedEvent.Property ? PropertyChangedEvent.Property->GetFName() : NAME_None;
+
+	if (PropertyName == GET_MEMBER_NAME_CHECKED(FRuntimePartitionDesc, Class))
+	{
+		int32 RuntimePartitionIndex = PropertyChangedEvent.GetArrayIndex(NAME_RuntimePartitions.ToString());
+		check(RuntimePartitions.IsValidIndex(RuntimePartitionIndex));
+
+		FRuntimePartitionDesc& RuntimePartitionDesc = RuntimePartitions[RuntimePartitionIndex];
+		
+		RuntimePartitionDesc.Partitions.Empty();
+
+		if (RuntimePartitionDesc.Class)
+		{
+			RuntimePartitionDesc.Name = RuntimePartitionDesc.Class->GetFName();
+			RuntimePartitionDesc.Partitions.Add(NewObject<URuntimePartition>(this, RuntimePartitionDesc.Class, NAME_None));
+			RuntimePartitionDesc.Partitions[0]->Name = TEXT("MainPartition");
+		}
+	}
+	else if (PropertyName == GET_MEMBER_NAME_CHECKED(FRuntimePartitionDesc, Name))
+	{
+		int32 RuntimePartitionIndex = PropertyChangedEvent.GetArrayIndex(NAME_RuntimePartitions.ToString());
+		check(RuntimePartitions.IsValidIndex(RuntimePartitionIndex));
+
+		FRuntimePartitionDesc& RuntimePartitionDesc = RuntimePartitions[RuntimePartitionIndex];
+
+		if (RuntimePartitionDesc.Name == NAME_PersistentLevel)
+		{
+			RuntimePartitionDesc.Name = RuntimePartitionDesc.Class->GetFName();
+		}
+		else
+		{
+			for (int32 CurRuntimePartitionIndex = 0; CurRuntimePartitionIndex < RuntimePartitions.Num(); CurRuntimePartitionIndex++)
+			{
+				if (CurRuntimePartitionIndex != RuntimePartitionIndex)
+				{
+					if (RuntimePartitionDesc.Name == RuntimePartitions[CurRuntimePartitionIndex].Name)
+					{
+						RuntimePartitionDesc.Name = RuntimePartitionDesc.Class->GetFName();
+						break;
+					}
+				}
+			}
+		}
+	}
+	else if (PropertyName == NAME_HLODSetups_Key)
+	{
+		int32 RuntimePartitionIndex = PropertyChangedEvent.GetArrayIndex(NAME_RuntimePartitions.ToString());
+		check(RuntimePartitions.IsValidIndex(RuntimePartitionIndex));
+
+		FRuntimePartitionDesc& RuntimePartitionDesc = RuntimePartitions[RuntimePartitionIndex];
+
+		int32 HLODSetupsIndex = PropertyChangedEvent.GetArrayIndex(NAME_HLODSetups_Key.ToString());
+		const FSetElementId HLODSetupElemId = FSetElementId::FromInteger(HLODSetupsIndex);
+		check(RuntimePartitionDesc.HLODSetups.IsValidId(HLODSetupElemId));
+
+		const UHLODLayer* HLODLayer = RuntimePartitionDesc.HLODSetups.Get(HLODSetupElemId).Key;
+		FRuntimePartitionHLODSetup& HLODSetup = RuntimePartitionDesc.HLODSetups.Get(HLODSetupElemId).Value;
+		
+		HLODSetup.Partitions.Empty();
+
+		while (HLODLayer)
+		{
+			URuntimePartition* HLODPartitionLayer = NewObject<URuntimePartition>(this, RuntimePartitionDesc.Class, NAME_None);
+			HLODPartitionLayer->Name = HLODLayer->GetFName();
+			HLODPartitionLayer->bIsHLODSetup = true;
+
+			HLODSetup.Partitions.Add(HLODPartitionLayer);
+
+			HLODLayer = HLODLayer->GetParentLayer().LoadSynchronous();
+		}
+
+		SetupHLODPartitionLayers(RuntimePartitionIndex);
+	}
+}
+
+void UWorldPartitionRuntimeHashSet::SetupHLODPartitionLayers(int32 RuntimePartitionIndex)
+{
+	FRuntimePartitionDesc& RuntimePartitionDesc = RuntimePartitions[RuntimePartitionIndex];
+
+	if (RuntimePartitionDesc.Partitions.Num())
+	{
+		RuntimePartitionDesc.Partitions.SetNum(1);
+
+		for (auto [CurHLODLayer, CurHLODSetup] : RuntimePartitionDesc.HLODSetups)
+		{
+			for (const URuntimePartition* HLODPartitionTemplate : CurHLODSetup.Partitions)
+			{
+				URuntimePartition* HLODPartitionLayer = DuplicateObject<URuntimePartition>(HLODPartitionTemplate, this);
+				HLODPartitionLayer->Name = CurHLODLayer->GetFName();
+				HLODPartitionLayer->bIsHLODSetup = false;
+
+				RuntimePartitionDesc.Partitions.Add(HLODPartitionLayer);
+			}
+		}
+	}
+}
+#endif
\ No newline at end of file
diff --git a/Engine/Source/Runtime/Engine/Private/WorldPartition/WorldPartitionRuntimeCell.cpp b/Engine/Source/Runtime/Engine/Private/WorldPartition/WorldPartitionRuntimeCell.cpp
--- a/Engine/Source/Runtime/Engine/Private/WorldPartition/WorldPartitionRuntimeCell.cpp
+++ b/Engine/Source/Runtime/Engine/Private/WorldPartition/WorldPartitionRuntimeCell.cpp
@@ -87,7 +87,7 @@ void UWorldPartitionRuntimeCell::SetDataLayers(const TArray<const UDataLayerInst
 	DataLayers.Sort([](const FName& A, const FName& B) { return A.ToString() < B.ToString(); });
 }
 
-void UWorldPartitionRuntimeCell::DumpStateLog(FHierarchicalLogArchive& Ar)
+void UWorldPartitionRuntimeCell::DumpStateLog(FHierarchicalLogArchive& Ar) const
 {
 	Ar.Printf(TEXT("Actor Count: %d"), GetActorCount());
 }
diff --git a/Engine/Source/Runtime/Engine/Private/WorldPartition/WorldPartitionRuntimeHash.cpp b/Engine/Source/Runtime/Engine/Private/WorldPartition/WorldPartitionRuntimeHash.cpp
--- a/Engine/Source/Runtime/Engine/Private/WorldPartition/WorldPartitionRuntimeHash.cpp
+++ b/Engine/Source/Runtime/Engine/Private/WorldPartition/WorldPartitionRuntimeHash.cpp
@@ -159,7 +159,7 @@ bool UWorldPartitionRuntimeHash::ConditionalRegisterAlwaysLoadedActorsForPIE(con
 	return false;
 }
 
-bool UWorldPartitionRuntimeHash::PopulateCellActorInstances(const TArray<const IStreamingGenerationContext::FActorSetInstance*> ActorSetInstances, bool bIsMainWorldPartition, bool bIsCellAlwaysLoaded, TArray<IStreamingGenerationContext::FActorInstance>& OutCellActorInstances)
+bool UWorldPartitionRuntimeHash::PopulateCellActorInstances(const TArray<const IStreamingGenerationContext::FActorSetInstance*>& ActorSetInstances, bool bIsMainWorldPartition, bool bIsCellAlwaysLoaded, TArray<IStreamingGenerationContext::FActorInstance>& OutCellActorInstances)
 {
 	for (const IStreamingGenerationContext::FActorSetInstance* ActorSetInstance : ActorSetInstances)
 	{
@@ -320,7 +320,7 @@ bool UWorldPartitionRuntimeHash::PopulateGeneratorPackageForCook(const TArray<FW
 	return true;
 }
 
-void UWorldPartitionRuntimeHash::DumpStateLog(FHierarchicalLogArchive& Ar)
+void UWorldPartitionRuntimeHash::DumpStateLog(FHierarchicalLogArchive& Ar) const
 {
 	if (!GIsAutomationTesting)
 	{
diff --git a/Engine/Source/Runtime/Engine/Private/WorldPartition/WorldPartitionRuntimeLevelStreamingCell.cpp b/Engine/Source/Runtime/Engine/Private/WorldPartition/WorldPartitionRuntimeLevelStreamingCell.cpp
--- a/Engine/Source/Runtime/Engine/Private/WorldPartition/WorldPartitionRuntimeLevelStreamingCell.cpp
+++ b/Engine/Source/Runtime/Engine/Private/WorldPartition/WorldPartitionRuntimeLevelStreamingCell.cpp
@@ -371,7 +371,7 @@ FString UWorldPartitionRuntimeLevelStreamingCell::GetPackageNameToCreate() const
 	return UWorldPartitionLevelStreamingPolicy::GetCellPackagePath(GetFName(), GetOuterWorld());
 }
 
-void UWorldPartitionRuntimeLevelStreamingCell::DumpStateLog(FHierarchicalLogArchive& Ar)
+void UWorldPartitionRuntimeLevelStreamingCell::DumpStateLog(FHierarchicalLogArchive& Ar) const
 {
 	Super::DumpStateLog(Ar);
 
diff --git a/Engine/Source/Runtime/Engine/Private/WorldPartition/WorldPartitionRuntimeSpatialHash.cpp b/Engine/Source/Runtime/Engine/Private/WorldPartition/WorldPartitionRuntimeSpatialHash.cpp
--- a/Engine/Source/Runtime/Engine/Private/WorldPartition/WorldPartitionRuntimeSpatialHash.cpp
+++ b/Engine/Source/Runtime/Engine/Private/WorldPartition/WorldPartitionRuntimeSpatialHash.cpp
@@ -1216,11 +1216,11 @@ bool UWorldPartitionRuntimeSpatialHash::GenerateStreaming(UWorldPartitionStreami
 	return true;
 }
 
-void UWorldPartitionRuntimeSpatialHash::DumpStateLog(FHierarchicalLogArchive& Ar)
+void UWorldPartitionRuntimeSpatialHash::DumpStateLog(FHierarchicalLogArchive& Ar) const
 {
 	Super::DumpStateLog(Ar);
 
-	ForEachStreamingGrid([&Ar, this](FSpatialHashStreamingGrid& StreamingGrid)
+	ForEachStreamingGrid([&Ar, this](const FSpatialHashStreamingGrid& StreamingGrid)
 	{
 		Ar.Printf(TEXT("----------------------------------------------------------------------------------------------------------------"));
 		Ar.Printf(TEXT("%s - Runtime Spatial Hash - Streaming Grid - %s"), *GetWorld()->GetName(), *StreamingGrid.GridName.ToString());
@@ -1250,14 +1250,14 @@ void UWorldPartitionRuntimeSpatialHash::DumpStateLog(FHierarchicalLogArchive& Ar
 			
 		{
 			int32 Level = 0;
-			for (FSpatialHashStreamingGridLevel& GridLevel : StreamingGrid.GridLevels)
+			for (const FSpatialHashStreamingGridLevel& GridLevel : StreamingGrid.GridLevels)
 			{
 				int32 LevelCellCount = 0;
 				int32 LevelActorCount = 0;
-				for (FSpatialHashStreamingGridLayerCell& LayerCell : GridLevel.LayerCells)
+				for (const FSpatialHashStreamingGridLayerCell& LayerCell : GridLevel.LayerCells)
 				{
 					LevelCellCount += LayerCell.GridCells.Num();
-					for (TObjectPtr<UWorldPartitionRuntimeCell>& Cell : LayerCell.GridCells)
+					for (const TObjectPtr<UWorldPartitionRuntimeCell>& Cell : LayerCell.GridCells)
 					{
 						LevelActorCount += Cell->GetActorCount();
 					}
@@ -1284,7 +1284,7 @@ void UWorldPartitionRuntimeSpatialHash::DumpStateLog(FHierarchicalLogArchive& Ar
 			{
 				FHierarchicalLogArchive::FIndentScope LevelIndentScope = Ar.PrintfIndent(TEXT("Content of Grid Level %d"), Level++);
 
-				TArray<UWorldPartitionRuntimeCell*> Cells;
+				TArray<TObjectPtr<UWorldPartitionRuntimeCell>> Cells;
 				for (const FSpatialHashStreamingGridLayerCell& LayerCell : GridLevel.LayerCells)
 				{
 					for (const TObjectPtr<UWorldPartitionRuntimeCell>& Cell : LayerCell.GridCells)
@@ -1295,7 +1295,7 @@ void UWorldPartitionRuntimeSpatialHash::DumpStateLog(FHierarchicalLogArchive& Ar
 
 				Cells.Sort([this](const UWorldPartitionRuntimeCell& A, const UWorldPartitionRuntimeCell& B) { return A.GetFName().LexicalLess(B.GetFName()); });
 
-				for (UWorldPartitionRuntimeCell* Cell : Cells)
+				for (const UWorldPartitionRuntimeCell* Cell : Cells)
 				{
 					FHierarchicalLogArchive::FIndentScope CellIndentScope = Ar.PrintfIndent(TEXT("Content of Cell %s (%s)"), *Cell->GetDebugName(), *Cell->GetName());
 					Cell->DumpStateLog(Ar);
diff --git a/Engine/Source/Runtime/Engine/Private/WorldPartition/WorldPartitionStreamingGeneration.cpp b/Engine/Source/Runtime/Engine/Private/WorldPartition/WorldPartitionStreamingGeneration.cpp
--- a/Engine/Source/Runtime/Engine/Private/WorldPartition/WorldPartitionStreamingGeneration.cpp
+++ b/Engine/Source/Runtime/Engine/Private/WorldPartition/WorldPartitionStreamingGeneration.cpp
@@ -1321,6 +1321,8 @@ void UWorldPartition::FlushStreaming()
 URuntimeHashExternalStreamingObjectBase* UWorldPartition::FlushStreamingToExternalStreamingObject(const FString& ExternalStreamingObjectName)
 {
 	URuntimeHashExternalStreamingObjectBase* ExternalStreamingObject = RuntimeHash->StoreToExternalStreamingObject(this, *ExternalStreamingObjectName);
+	check(ExternalStreamingObject);
+
 	StreamingPolicy->StoreToExternalStreamingObject(*ExternalStreamingObject);
 
 	FlushStreaming();
diff --git a/Engine/Source/Runtime/Engine/Public/WorldPartition/RuntimeHashSet/RuntimePartition.h b/Engine/Source/Runtime/Engine/Public/WorldPartition/RuntimeHashSet/RuntimePartition.h
new file mode 100644
--- /dev/null
+++ b/Engine/Source/Runtime/Engine/Public/WorldPartition/RuntimeHashSet/RuntimePartition.h
@@ -0,0 +1,55 @@
+// Copyright Epic Games, Inc. All Rights Reserved.
+#pragma once
+
+#include "CoreMinimal.h"
+#include "WorldPartition/RuntimeHashSet/WorldPartitionRuntimeHashSet.h"
+#include "WorldPartition/WorldPartitionStreamingGenerationContext.h"
+#include "RuntimePartition.generated.h"
+
+UCLASS(Abstract, Within=WorldPartitionRuntimeHashSet, CollapseCategories)
+class URuntimePartition : public UObject
+{
+	GENERATED_BODY()
+
+public:
+#if WITH_EDITOR
+	struct FCellDesc
+	{
+		FName Name;
+		FBox Bounds = FBox(ForceInit);
+		bool bIsSpatiallyLoaded;
+		bool bBlockOnSlowStreaming;
+		bool bClientOnlyVisible;
+		FGuid ContentBundleID;
+		int32 Priority;
+
+		TArray<IStreamingGenerationContext::FActorInstance> ActorInstances;
+	};
+
+	virtual bool SupportsHLODs() const PURE_VIRTUAL(URuntimePartition::SupportsHLODs, return false;);
+	virtual bool IsValidGrid(FName GridName) const PURE_VIRTUAL(URuntimePartition::IsValidGrid, return false;);
+	virtual bool GenerateStreaming(const TArray<const IStreamingGenerationContext::FActorSetInstance*>& ActorSetInstances, TArray<FCellDesc>& OutRuntimeCellDescs) PURE_VIRTUAL(URuntimePartition::GenerateStreaming, return false;);
+#endif
+
+#if WITH_EDITORONLY_DATA
+	UPROPERTY(VisibleAnywhere, Category = RuntimeSettings, meta = (EditCondition = "!bIsHLODSetup", EditConditionHides, HideEditConditionToggle))
+	FName Name;
+
+	UPROPERTY(EditAnywhere, Category = RuntimeSettings, meta = (EditCondition = "!bIsHLODSetup", EditConditionHides, HideEditConditionToggle))
+	bool bBlockOnSlowStreaming;
+
+	UPROPERTY(EditAnywhere, Category = RuntimeSettings, meta = (EditCondition = "!bIsHLODSetup", EditConditionHides, HideEditConditionToggle))
+	bool bClientOnlyVisible;
+
+	UPROPERTY(EditAnywhere, Category = RuntimeSettings, meta = (EditCondition = "!bIsHLODSetup", EditConditionHides, HideEditConditionToggle))
+	int32 Priority;
+
+	UPROPERTY()
+	bool bIsHLODSetup;
+#endif
+
+protected:
+#if WITH_EDITOR
+	bool PopulateCellActorInstances(const TArray<const IStreamingGenerationContext::FActorSetInstance*>& ActorSetInstances, bool bIsMainWorldPartition, bool bIsCellAlwaysLoaded, TArray<IStreamingGenerationContext::FActorInstance>& OutCellActorInstances);
+#endif
+};
\ No newline at end of file
diff --git a/Engine/Source/Runtime/Engine/Public/WorldPartition/RuntimeHashSet/RuntimePartitionLevelStreaming.h b/Engine/Source/Runtime/Engine/Public/WorldPartition/RuntimeHashSet/RuntimePartitionLevelStreaming.h
new file mode 100644
--- /dev/null
+++ b/Engine/Source/Runtime/Engine/Public/WorldPartition/RuntimeHashSet/RuntimePartitionLevelStreaming.h
@@ -0,0 +1,25 @@
+// Copyright Epic Games, Inc. All Rights Reserved.
+#pragma once
+
+#include "CoreMinimal.h"
+#include "WorldPartition/RuntimeHashSet/RuntimePartition.h"
+#include "RuntimePartitionLevelStreaming.generated.h"
+
+UCLASS()
+class URuntimePartitionLevelStreaming : public URuntimePartition
+{
+	GENERATED_BODY()
+
+public:
+#if WITH_EDITOR
+	virtual bool SupportsHLODs() const override;
+	virtual bool IsValidGrid(FName GridName) const override;
+	virtual bool GenerateStreaming(const TArray<const IStreamingGenerationContext::FActorSetInstance*>& ActorSetInstances, TArray<FCellDesc>& OutRuntimeCellDescs) override;
+#endif
+
+#if WITH_EDITORONLY_DATA
+	/** Splits actors into different sublevels based on their actor containers (level instance) */
+	UPROPERTY(EditAnywhere, Category = RuntimeSettings)
+	bool bOneLevelPerActorContainer;
+#endif
+};
\ No newline at end of file
diff --git a/Engine/Source/Runtime/Engine/Public/WorldPartition/RuntimeHashSet/RuntimePartitionPersistent.h b/Engine/Source/Runtime/Engine/Public/WorldPartition/RuntimeHashSet/RuntimePartitionPersistent.h
new file mode 100644
--- /dev/null
+++ b/Engine/Source/Runtime/Engine/Public/WorldPartition/RuntimeHashSet/RuntimePartitionPersistent.h
@@ -0,0 +1,19 @@
+// Copyright Epic Games, Inc. All Rights Reserved.
+#pragma once
+
+#include "CoreMinimal.h"
+#include "WorldPartition/RuntimeHashSet/RuntimePartition.h"
+#include "RuntimePartitionPersistent.generated.h"
+
+UCLASS(HideDropdown)
+class URuntimePartitionPersistent : public URuntimePartition
+{
+	GENERATED_BODY()
+
+public:
+#if WITH_EDITOR
+	virtual bool SupportsHLODs() const override { return false; }
+	virtual bool IsValidGrid(FName GridName) const override { return true; }
+	virtual bool GenerateStreaming(const TArray<const IStreamingGenerationContext::FActorSetInstance*>& ActorSetInstances, TArray<FCellDesc>& OutRuntimeCellDescs) override;
+#endif
+};
diff --git a/Engine/Source/Runtime/Engine/Public/WorldPartition/RuntimeHashSet/StaticSpatialIndex.h b/Engine/Source/Runtime/Engine/Public/WorldPartition/RuntimeHashSet/StaticSpatialIndex.h
new file mode 100644
--- /dev/null
+++ b/Engine/Source/Runtime/Engine/Public/WorldPartition/RuntimeHashSet/StaticSpatialIndex.h
@@ -0,0 +1,161 @@
+// Copyright Epic Games, Inc. All Rights Reserved.
+#pragma once
+
+#include "CoreMinimal.h"
+#include "Algo/Transform.h"
+#include "Misc/TVariant.h"
+#include "UObject/ObjectPtr.h"
+#include "UObject/UObjectGlobals.h"
+
+struct IStaticSpatialIndexDataInterface
+{
+	virtual ~IStaticSpatialIndexDataInterface() {}
+	virtual const FBox& GetBox(uint32 InIndex) const =0;
+};
+
+template <typename ValueType, class SpatialIndexType>
+class TStaticSpatialIndex : public IStaticSpatialIndexDataInterface
+{
+public:
+	TStaticSpatialIndex()
+		: SpatialIndex(*this)
+	{}
+
+	// Creation
+	void Init(const TArray<TPair<FBox, ValueType>>& InElements)
+	{
+		Elements = InElements;
+
+		InitSpatialIndex();
+	}
+
+	void ForEachElement(TFunctionRef<void(const ValueType& InValue)> Func) const
+	{
+		SpatialIndex.ForEachElement([this, &Func](uint32 ValueIndex)
+		{
+			Func(Elements[ValueIndex].Value);
+		});
+	}
+
+	void ForEachIntersectingElement(const FBox& InBox, TFunctionRef<void(const ValueType& InValue)> Func) const
+	{
+		SpatialIndex.ForEachIntersectingElement(InBox, [this, &Func](uint32 ValueIndex)
+		{
+			Func(Elements[ValueIndex].Value);
+		});
+	}
+
+	void ForEachIntersectingElement(const FSphere& InSphere, TFunctionRef<void(const ValueType& InValue)> Func) const
+	{
+		SpatialIndex.ForEachIntersectingElement(InSphere, [this, &Func](uint32 ValueIndex)
+		{
+			Func(Elements[ValueIndex].Value);
+		});
+	}
+
+	void AddReferencedObjects(FReferenceCollector& Collector)
+	{
+		if constexpr (TIsPointerOrObjectPtrToBaseOf<ValueType, UObject>::Value)
+		{
+			for (TPair<FBox, ValueType>& Element : Elements)
+			{
+				Collector.AddReferencedObject(Element.Value);
+			}
+		}
+	}
+
+	void Serialize(FArchive& Ar)
+	{
+		Ar << Elements;
+
+		if (Ar.IsLoading())
+		{
+			InitSpatialIndex();
+		}
+	}
+
+	// IStaticSpatialIndexDataInterface interface
+	virtual const FBox& GetBox(uint32 InIndex) const override
+	{
+		return Elements[InIndex].Key;
+	}
+
+private:
+	void InitSpatialIndex()
+	{		
+		TArray<TPair<FBox, uint32>> IndexData;
+		
+		int32 ElemIndex = 0;
+		IndexData.Reserve(Elements.Num());
+		Algo::Transform(Elements, IndexData, [&ElemIndex](const TPair<FBox, ValueType>& Element) { return TPair<FBox, uint32>(Element.Key, ElemIndex++); });
+
+		SpatialIndex.Init(IndexData);
+	}
+
+	TArray<TPair<FBox, ValueType>> Elements;
+	SpatialIndexType SpatialIndex;
+};
+
+namespace FStaticSpatialIndex
+{
+	class FImpl
+	{
+	public:
+		FImpl(const IStaticSpatialIndexDataInterface& InDataInterface)
+			: DataInterface(InDataInterface)
+		{}
+
+	protected:
+		const IStaticSpatialIndexDataInterface& DataInterface;
+	};
+
+	class FListImpl : public FImpl
+	{
+	public:
+		FListImpl(const IStaticSpatialIndexDataInterface& InDataInterface)
+			: FImpl(InDataInterface)
+		{}
+
+		void Init(const TArray<TPair<FBox, uint32>>& InElements);
+
+		void ForEachElement(TFunctionRef<void(uint32 InValueIndex)> Func) const;
+		void ForEachIntersectingElement(const FBox& InBox, TFunctionRef<void(uint32 InValueIndex)> Func) const;
+		void ForEachIntersectingElement(const FSphere& InSphere, TFunctionRef<void(uint32 InValueIndex)> Func) const;
+
+	private:
+		TArray<uint32> Elements;
+	};
+
+	class FRTreeImpl : public FImpl
+	{
+	public:
+		FRTreeImpl(const IStaticSpatialIndexDataInterface& InDataInterface)
+			: FImpl(InDataInterface)
+		{}
+
+		void Init(const TArray<TPair<FBox, uint32>>& InElements);
+
+		void ForEachElement(TFunctionRef<void(uint32 InValueIndex)> Func) const;
+		void ForEachIntersectingElement(const FBox& InBox, TFunctionRef<void(uint32 InValueIndex)> Func) const;
+		void ForEachIntersectingElement(const FSphere& InSphere, TFunctionRef<void(uint32 InValueIndex)> Func) const;
+
+	private:
+		struct FNode
+		{
+			using FNodeType = TArray<FNode>;
+			using FLeafType = TArray<uint32>;
+
+			FBox Box = FBox(ForceInit);
+			TVariant<FNodeType, FLeafType> Content;
+		};
+
+		void ForEachElementRecursive(const FNode* Node, TFunctionRef<void(uint32 InValueIndex)> Func) const;
+		void ForEachIntersectingElementRecursive(const FNode* Node, const FBox& InBox, TFunctionRef<void(uint32 InValueIndex)> Func) const;
+		void ForEachIntersectingElementRecursive(const FNode* Node, const FSphere& InSphere, TFunctionRef<void(uint32 InValueIndex)> Func) const;
+
+		FNode RootNode;
+	};
+}
+
+template <class ValueType> class TStaticSpatialIndexList : public TStaticSpatialIndex<ValueType, FStaticSpatialIndex::FListImpl> {};
+template <class ValueType> class TStaticSpatialIndexRTree : public TStaticSpatialIndex<ValueType, FStaticSpatialIndex::FRTreeImpl> {};
\ No newline at end of file
diff --git a/Engine/Source/Runtime/Engine/Public/WorldPartition/RuntimeHashSet/WorldPartitionRuntimeHashSet.h b/Engine/Source/Runtime/Engine/Public/WorldPartition/RuntimeHashSet/WorldPartitionRuntimeHashSet.h
new file mode 100644
--- /dev/null
+++ b/Engine/Source/Runtime/Engine/Public/WorldPartition/RuntimeHashSet/WorldPartitionRuntimeHashSet.h
@@ -0,0 +1,130 @@
+// Copyright Epic Games, Inc. All Rights Reserved.
+#pragma once
+
+#include "CoreMinimal.h"
+#include "WorldPartition/WorldPartitionRuntimeHash.h"
+#include "WorldPartition/RuntimeHashSet/StaticSpatialIndex.h"
+#include "WorldPartitionRuntimeHashSet.generated.h"
+
+class UHLODLayer;
+class URuntimePartition;
+class URuntimePartitionPersistent;
+struct FPropertyChangedChainEvent;
+
+using FStaticSpatialIndexType = TStaticSpatialIndexRTree<UWorldPartitionRuntimeCell*>;
+
+/** Holds settings for an HLOD layer (and its parents) for a particular partition class. */
+USTRUCT()
+struct FRuntimePartitionHLODSetup
+{
+	GENERATED_USTRUCT_BODY()
+
+	UPROPERTY(VisibleAnywhere, Category = RuntimeSettings, EditFixedSize, Instanced, meta = (NoResetToDefault, EditFixedOrder, TitleProperty="Name"))
+	TArray<TObjectPtr<URuntimePartition>> Partitions;
+};
+
+/** Holds settings for a runtime partition instance. */
+USTRUCT()
+struct FRuntimePartitionDesc
+{
+	GENERATED_USTRUCT_BODY()
+
+#if WITH_EDITORONLY_DATA
+	/** Partition class */
+	UPROPERTY(EditAnywhere, Category = RuntimeSettings)
+	TSubclassOf<URuntimePartition> Class;
+
+	/** Name for this partition, used to map actors to it through the Actor.RuntimeGrid property  */
+	UPROPERTY(EditAnywhere, Category = RuntimeSettings, meta = (EditCondition = "Class != nullptr", HideEditConditionToggle))
+	FName Name;
+
+	/** Array of partition objects, first index is the main one and then one per HLOD layer */
+	UPROPERTY(VisibleAnywhere, Category = RuntimeSettings, EditFixedSize, Instanced, meta = (EditCondition = "Class != nullptr", HideEditConditionToggle, NoResetToDefault, EditFixedOrder, TitleProperty="Name"))
+	TArray<TObjectPtr<URuntimePartition>> Partitions;
+
+	/** HLOD setups used by this partition */
+	UPROPERTY(EditAnywhere, Category = RuntimeSettings, meta = (EditCondition = "Class != nullptr", HideEditConditionToggle, ForceInlineRow, DisplayThumbnail=false))
+	TMap<TObjectPtr<const UHLODLayer>, FRuntimePartitionHLODSetup> HLODSetups;
+#endif
+};
+
+UCLASS()
+class URuntimeHashSetExternalStreamingObject : public URuntimeHashExternalStreamingObjectBase
+{
+	GENERATED_BODY()
+
+public:
+	UPROPERTY()
+	TArray<TObjectPtr<UWorldPartitionRuntimeCell>> NonSpatiallyLoadedRuntimeCells;
+
+	UPROPERTY()
+	TArray<TObjectPtr<UWorldPartitionRuntimeCell>> SpatiallyLoadedRuntimeCells;
+
+	// Transient
+	TUniquePtr<FStaticSpatialIndexType> SpatialIndex;
+};
+
+UCLASS(HideDropdown)
+class ENGINE_API UWorldPartitionRuntimeHashSet : public UWorldPartitionRuntimeHash
+{
+	GENERATED_UCLASS_BODY()
+
+	//~ Begin UObject Interface
+	virtual void Serialize(FArchive& Ar) override;
+	static void AddReferencedObjects(UObject* InThis, FReferenceCollector& Collector);
+#if WITH_EDITOR
+	virtual void PostEditChangeChainProperty(FPropertyChangedChainEvent& PropertyChangedEvent) override;
+#endif
+	//~ End UObject Interface
+
+public:
+#if WITH_EDITOR
+	// Streaming generation interface
+	virtual void SetDefaultValues() override;
+	virtual bool SupportsHLODs() const override;
+	virtual bool GenerateStreaming(class UWorldPartitionStreamingPolicy* StreamingPolicy, const IStreamingGenerationContext* StreamingGenerationContext, TArray<FString>* OutPackagesToGenerate) override;
+	virtual void FlushStreaming() override;
+	virtual bool IsValidGrid(FName GridName) const;
+	virtual TArray<UWorldPartitionRuntimeCell*> GetAlwaysLoadedCells() const override;
+	virtual void DumpStateLog(FHierarchicalLogArchive& Ar) const override;
+
+	// Helpers
+	static TArray<FName> ParseGridName(FName GridName);
+#endif
+
+	// External streaming object interface
+#if WITH_EDITOR
+	virtual URuntimeHashExternalStreamingObjectBase* StoreToExternalStreamingObject(UObject* StreamingObjectOuter, FName StreamingObjectName) override;
+#endif
+	virtual bool InjectExternalStreamingObject(URuntimeHashExternalStreamingObjectBase* ExternalStreamingObject) override;
+	virtual bool RemoveExternalStreamingObject(URuntimeHashExternalStreamingObjectBase* ExternalStreamingObject) override;
+
+	// Streaming interface
+	virtual void ForEachStreamingCells(TFunctionRef<bool(const UWorldPartitionRuntimeCell*)> Func) const;
+	virtual void ForEachStreamingCellsQuery(const FWorldPartitionStreamingQuerySource& QuerySource, TFunctionRef<bool(const UWorldPartitionRuntimeCell*)> Func, FWorldPartitionQueryCache* QueryCache) const override;
+	virtual void ForEachStreamingCellsSources(const TArray<FWorldPartitionStreamingSource>& Sources, TFunctionRef<bool(const UWorldPartitionRuntimeCell*, EStreamingSourceTargetState)> Func) const override;
+
+private:
+#if WITH_EDITOR
+	/** Update the partition layers to reflect the curent HLOD setups. */
+	void SetupHLODPartitionLayers(int32 RuntimePartitionIndex);
+#endif
+
+public:
+#if WITH_EDITORONLY_DATA
+	/** Persistent partition */
+	UPROPERTY()
+	FRuntimePartitionDesc PersistentPartitionDesc;
+
+	/** Array of runtime partition descriptors */
+	UPROPERTY(EditAnywhere, Category = RuntimeSettings, meta = (TitleProperty="Name"))
+	TArray<FRuntimePartitionDesc> RuntimePartitions;
+#endif
+
+	UPROPERTY()
+	TArray<TObjectPtr<UWorldPartitionRuntimeCell>> NonSpatiallyLoadedRuntimeCells;
+
+	TUniquePtr<FStaticSpatialIndexType> SpatialIndex;
+
+	TSet<URuntimeHashSetExternalStreamingObject*> InjectedExternalStreamingObjects;
+};
\ No newline at end of file
diff --git a/Engine/Source/Runtime/Engine/Public/WorldPartition/WorldPartitionRuntimeCell.h b/Engine/Source/Runtime/Engine/Public/WorldPartition/WorldPartitionRuntimeCell.h
--- a/Engine/Source/Runtime/Engine/Public/WorldPartition/WorldPartitionRuntimeCell.h
+++ b/Engine/Source/Runtime/Engine/Public/WorldPartition/WorldPartitionRuntimeCell.h
@@ -243,7 +243,7 @@ class ENGINE_API UWorldPartitionRuntimeCell : public UObject, public IWorldParti
 	virtual bool PopulateGeneratedPackageForCook(UPackage* InPackage, TArray<UPackage*>& OutModifiedPackages) PURE_VIRTUAL(UWorldPartitionRuntimeCell::PopulateGeneratedPackageForCook, return false;);
 	virtual FString GetPackageNameToCreate() const PURE_VIRTUAL(UWorldPartitionRuntimeCell::GetPackageNameToCreate, return FString(""););
 
-	virtual void DumpStateLog(FHierarchicalLogArchive& Ar);
+	virtual void DumpStateLog(FHierarchicalLogArchive& Ar) const;
 #endif
 
 	void SetIsHLOD(bool bInIsHLOD) { bIsHLOD = bInIsHLOD; }
diff --git a/Engine/Source/Runtime/Engine/Public/WorldPartition/WorldPartitionRuntimeHash.h b/Engine/Source/Runtime/Engine/Public/WorldPartition/WorldPartitionRuntimeHash.h
--- a/Engine/Source/Runtime/Engine/Public/WorldPartition/WorldPartitionRuntimeHash.h
+++ b/Engine/Source/Runtime/Engine/Public/WorldPartition/WorldPartitionRuntimeHash.h
@@ -81,13 +81,11 @@ class ENGINE_API UWorldPartitionRuntimeHash : public UObject
 {
 	GENERATED_UCLASS_BODY()
 
+	friend class URuntimePartition;
+
 #if WITH_EDITOR
 	virtual void SetDefaultValues() {}
 	virtual bool SupportsHLODs() const { return false; }
-	virtual bool PrepareGeneratorPackageForCook(TArray<UPackage*>& OutModifiedPackages);
-	virtual bool PopulateGeneratorPackageForCook(const TArray<FWorldPartitionCookPackage*>& PackagesToCook, TArray<UPackage*>& OutModifiedPackages);
-	virtual bool PopulateGeneratedPackageForCook(const FWorldPartitionCookPackage& PackagesToCook, TArray<UPackage*>& OutModifiedPackages);
-	virtual UWorldPartitionRuntimeCell* GetCellForPackage(const FWorldPartitionCookPackage& PackageToCook) const;
 	virtual TArray<UWorldPartitionRuntimeCell*> GetAlwaysLoadedCells() const;
 	virtual bool GenerateStreaming(class UWorldPartitionStreamingPolicy* StreamingPolicy, const IStreamingGenerationContext* StreamingGenerationContext, TArray<FString>* OutPackagesToGenerate);
 	virtual void FlushStreaming();
@@ -97,7 +95,13 @@ class ENGINE_API UWorldPartitionRuntimeHash : public UObject
 
 	virtual URuntimeHashExternalStreamingObjectBase* StoreToExternalStreamingObject(UObject* StreamingObjectOuter, FName StreamingObjectName) { return nullptr; }
 
-	virtual void DumpStateLog(FHierarchicalLogArchive& Ar);
+	virtual void DumpStateLog(FHierarchicalLogArchive& Ar) const;
+
+	// non-virtual
+	bool PrepareGeneratorPackageForCook(TArray<UPackage*>& OutModifiedPackages);
+	bool PopulateGeneratorPackageForCook(const TArray<FWorldPartitionCookPackage*>& PackagesToCook, TArray<UPackage*>& OutModifiedPackages);
+	bool PopulateGeneratedPackageForCook(const FWorldPartitionCookPackage& PackagesToCook, TArray<UPackage*>& OutModifiedPackages);
+	UWorldPartitionRuntimeCell* GetCellForPackage(const FWorldPartitionCookPackage& PackageToCook) const;
 
 	// PIE/Game methods
 	void OnBeginPlay();
@@ -105,7 +109,7 @@ class ENGINE_API UWorldPartitionRuntimeHash : public UObject
 
 protected:
 	bool ConditionalRegisterAlwaysLoadedActorsForPIE(const IStreamingGenerationContext::FActorSetInstance* ActorSetInstance, bool bIsMainWorldPartition, bool bIsMainContainer, bool bIsCellAlwaysLoaded);
-	bool PopulateCellActorInstances(const TArray<const IStreamingGenerationContext::FActorSetInstance*> ActorSetInstances, bool bIsMainWorldPartition, bool bIsCellAlwaysLoaded, TArray<IStreamingGenerationContext::FActorInstance>& OutCellActorInstances);
+	bool PopulateCellActorInstances(const TArray<const IStreamingGenerationContext::FActorSetInstance*>& ActorSetInstances, bool bIsMainWorldPartition, bool bIsCellAlwaysLoaded, TArray<IStreamingGenerationContext::FActorInstance>& OutCellActorInstances);
 	void PopulateRuntimeCell(UWorldPartitionRuntimeCell* RuntimeCell, const TArray<IStreamingGenerationContext::FActorInstance>& ActorInstances, TArray<FString>* OutPackagesToGenerate);
 #endif
 
diff --git a/Engine/Source/Runtime/Engine/Public/WorldPartition/WorldPartitionRuntimeLevelStreamingCell.h b/Engine/Source/Runtime/Engine/Public/WorldPartition/WorldPartitionRuntimeLevelStreamingCell.h
--- a/Engine/Source/Runtime/Engine/Public/WorldPartition/WorldPartitionRuntimeLevelStreamingCell.h
+++ b/Engine/Source/Runtime/Engine/Public/WorldPartition/WorldPartitionRuntimeLevelStreamingCell.h
@@ -51,7 +51,7 @@ class ENGINE_API UWorldPartitionRuntimeLevelStreamingCell : public UWorldPartiti
 	virtual void AddActorToCell(const FWorldPartitionActorDescView& ActorDescView, const FActorContainerID& InContainerID, const FTransform& InContainerTransform, const UActorDescContainer* InContainer) override;
 	virtual void Fixup() override;
 	virtual int32 GetActorCount() const override;
-	virtual void DumpStateLog(FHierarchicalLogArchive& Ar) override;
+	virtual void DumpStateLog(FHierarchicalLogArchive& Ar) const override;
 	// Cook methods
 	virtual bool PrepareCellForCook(UPackage* InPackage) override;
 	virtual bool PopulateGeneratorPackageForCook(TArray<UPackage*>& OutModifiedPackages) override;
diff --git a/Engine/Source/Runtime/Engine/Public/WorldPartition/WorldPartitionRuntimeSpatialHash.h b/Engine/Source/Runtime/Engine/Public/WorldPartition/WorldPartitionRuntimeSpatialHash.h
--- a/Engine/Source/Runtime/Engine/Public/WorldPartition/WorldPartitionRuntimeSpatialHash.h
+++ b/Engine/Source/Runtime/Engine/Public/WorldPartition/WorldPartitionRuntimeSpatialHash.h
@@ -316,7 +316,7 @@ protected:
 
 #if WITH_EDITOR
 	virtual bool GenerateStreaming(class UWorldPartitionStreamingPolicy* StreamingPolicy, const IStreamingGenerationContext* StreamingGenerationContext, TArray<FString>* OutPackagesToGenerate = nullptr) override;
-	virtual void DumpStateLog(FHierarchicalLogArchive& Ar) override;
+	virtual void DumpStateLog(FHierarchicalLogArchive& Ar) const override;
 #endif
 
 private:
