Workaround to fix crash occuring when generating the streaming grid after adding a new unsaved non-plugin actor into an in-memory level

diff --git a/Engine/Source/Runtime/Engine/Private/WorldPartition/WorldPartitionStreamingGeneration.cpp b/Engine/Source/Runtime/Engine/Private/WorldPartition/WorldPartitionStreamingGeneration.cpp
--- a/Engine/Source/Runtime/Engine/Private/WorldPartition/WorldPartitionStreamingGeneration.cpp
+++ b/Engine/Source/Runtime/Engine/Private/WorldPartition/WorldPartitionStreamingGeneration.cpp
@@ -374,6 +374,13 @@ class FWorldPartitionStreamingGenerator
 			FWorldPartitionActorDesc* ModifiedActorDesc = ModifiedActorsDescList->AddActor(InActor);
 
 			const UActorDescContainer* HandlingContainer = InActorDescCollection.FindHandlingContainer(InActor);
+			if (HandlingContainer == nullptr)
+			{
+				// @todo_ow : UActorDescContainer::IsActorDescHandled(AActor*) does not handle the case where the InActor is an unsaved new non-plugin actor
+				// and the level is in memory. Fix UActorDescContainer::IsActorDescHandled so it handle those actors
+				HandlingContainer = InActorDescCollection.GetMainActorDescContainer();
+			}
+
 			check(HandlingContainer != nullptr);
 
 			// Pretend that this actor descriptor belongs to the original container, even if it's not present. It's essentially a proxy
