Add a define to enable verbose animation node profiling tags. Disabled by default. See ANIMNODE_STATS_VERBOSE in AnimStats.h

diff --git a/Engine/Plugins/Animation/ControlRig/Source/ControlRig/Private/AnimNode_ControlRig.cpp b/Engine/Plugins/Animation/ControlRig/Source/ControlRig/Private/AnimNode_ControlRig.cpp
--- a/Engine/Plugins/Animation/ControlRig/Source/ControlRig/Private/AnimNode_ControlRig.cpp
+++ b/Engine/Plugins/Animation/ControlRig/Source/ControlRig/Private/AnimNode_ControlRig.cpp
@@ -179,6 +179,7 @@ void FAnimNode_ControlRig::CacheBones_AnyThread(const FAnimationCacheBonesContex
 void FAnimNode_ControlRig::Evaluate_AnyThread(FPoseContext & Output)
 {
 	DECLARE_SCOPE_HIERARCHICAL_COUNTER_FUNC()
+	ANIM_MT_SCOPE_CYCLE_COUNTER_VERBOSE(ControlRig, !IsInGameThread());
 
 	// evaluate 
 	FAnimNode_ControlRigBase::Evaluate_AnyThread(Output);
diff --git a/Engine/Plugins/Animation/IKRig/Source/IKRig/Private/AnimNodes/AnimNode_IKRig.cpp b/Engine/Plugins/Animation/IKRig/Source/IKRig/Private/AnimNodes/AnimNode_IKRig.cpp
--- a/Engine/Plugins/Animation/IKRig/Source/IKRig/Private/AnimNodes/AnimNode_IKRig.cpp
+++ b/Engine/Plugins/Animation/IKRig/Source/IKRig/Private/AnimNodes/AnimNode_IKRig.cpp
@@ -24,6 +24,7 @@ FAnimNode_IKRig::FAnimNode_IKRig()
 void FAnimNode_IKRig::Evaluate_AnyThread(FPoseContext& Output) 
 {
 	DECLARE_SCOPE_HIERARCHICAL_COUNTER_FUNC()
+	ANIM_MT_SCOPE_CYCLE_COUNTER_VERBOSE(IKRig, !IsInGameThread());
 
 	if (Source.GetLinkNode() && !bStartFromRefPose)
 	{
diff --git a/Engine/Plugins/Animation/IKRig/Source/IKRigEditor/Private/RetargetEditor/IKRetargetAnimInstance.cpp b/Engine/Plugins/Animation/IKRig/Source/IKRigEditor/Private/RetargetEditor/IKRetargetAnimInstance.cpp
--- a/Engine/Plugins/Animation/IKRig/Source/IKRigEditor/Private/RetargetEditor/IKRetargetAnimInstance.cpp
+++ b/Engine/Plugins/Animation/IKRig/Source/IKRigEditor/Private/RetargetEditor/IKRetargetAnimInstance.cpp
@@ -52,6 +52,7 @@ void FAnimNode_PreviewRetargetPose::CacheBones_AnyThread(const FAnimationCacheBo
 void FAnimNode_PreviewRetargetPose::Evaluate_AnyThread(FPoseContext& Output)
 {
 	DECLARE_SCOPE_HIERARCHICAL_COUNTER_ANIMNODE(Evaluate_AnyThread)
+	ANIM_MT_SCOPE_CYCLE_COUNTER_VERBOSE(PreviewRetargetPose, !IsInGameThread());
 
 	Output.Pose.ResetToRefPose();
 	
diff --git a/Engine/Plugins/Experimental/Animation/CurveExpression/Source/Runtime/Private/AnimNode_RemapCurves.cpp b/Engine/Plugins/Experimental/Animation/CurveExpression/Source/Runtime/Private/AnimNode_RemapCurves.cpp
--- a/Engine/Plugins/Experimental/Animation/CurveExpression/Source/Runtime/Private/AnimNode_RemapCurves.cpp
+++ b/Engine/Plugins/Experimental/Animation/CurveExpression/Source/Runtime/Private/AnimNode_RemapCurves.cpp
@@ -1,4 +1,4 @@
-﻿// Copyright Epic Games, Inc. All Rights Reserved.
+// Copyright Epic Games, Inc. All Rights Reserved.
 
 #include "AnimNode_RemapCurves.h"
 
@@ -7,7 +7,8 @@
 void FAnimNode_RemapCurves::Evaluate_AnyThread(FPoseContext& Output)
 {
 	DECLARE_SCOPE_HIERARCHICAL_COUNTER_ANIMNODE(Evaluate_AnyThread)
-	
+	ANIM_MT_SCOPE_CYCLE_COUNTER_VERBOSE(RemapCurves, !IsInGameThread());
+
 	using namespace CurveExpression::Evaluator;
 
 	FPoseContext SourceData(Output);
diff --git a/Engine/Plugins/Experimental/Animation/CurveExpression/Source/Runtime/Private/AnimNode_RemapCurvesFromMesh.cpp b/Engine/Plugins/Experimental/Animation/CurveExpression/Source/Runtime/Private/AnimNode_RemapCurvesFromMesh.cpp
--- a/Engine/Plugins/Experimental/Animation/CurveExpression/Source/Runtime/Private/AnimNode_RemapCurvesFromMesh.cpp
+++ b/Engine/Plugins/Experimental/Animation/CurveExpression/Source/Runtime/Private/AnimNode_RemapCurvesFromMesh.cpp
@@ -65,7 +65,8 @@ void FAnimNode_RemapCurvesFromMesh::Evaluate_AnyThread(
 	)
 {
 	DECLARE_SCOPE_HIERARCHICAL_COUNTER_ANIMNODE(Evaluate_AnyThread)
-	
+	ANIM_MT_SCOPE_CYCLE_COUNTER_VERBOSE(RemapCurvesFromMesh, !IsInGameThread());
+
 	using namespace CurveExpression::Evaluator;
 	
 	FPoseContext SourceData(Output);
diff --git a/Engine/Plugins/Experimental/Animation/PoseCorrectives/Source/PoseCorrectives/Private/AnimNode_CorrectPose.cpp b/Engine/Plugins/Experimental/Animation/PoseCorrectives/Source/PoseCorrectives/Private/AnimNode_CorrectPose.cpp
--- a/Engine/Plugins/Experimental/Animation/PoseCorrectives/Source/PoseCorrectives/Private/AnimNode_CorrectPose.cpp
+++ b/Engine/Plugins/Experimental/Animation/PoseCorrectives/Source/PoseCorrectives/Private/AnimNode_CorrectPose.cpp
@@ -95,6 +95,7 @@ void FAnimNode_CorrectPose::UpdateRBFTargetsFromAsset()
 void FAnimNode_CorrectPose::Evaluate_AnyThread(FPoseContext& Output)
 {
 	DECLARE_SCOPE_HIERARCHICAL_COUNTER_ANIMNODE(Evaluate_AnyThread)
+	ANIM_MT_SCOPE_CYCLE_COUNTER_VERBOSE(CorrectPose, !IsInGameThread());
 
 	FPoseContext SourceData(Output);
 	SourcePose.Evaluate(SourceData);
diff --git a/Engine/Plugins/Experimental/Animation/PoseSearch/Source/Runtime/Private/AnimNode_MotionMatching.cpp b/Engine/Plugins/Experimental/Animation/PoseSearch/Source/Runtime/Private/AnimNode_MotionMatching.cpp
--- a/Engine/Plugins/Experimental/Animation/PoseSearch/Source/Runtime/Private/AnimNode_MotionMatching.cpp
+++ b/Engine/Plugins/Experimental/Animation/PoseSearch/Source/Runtime/Private/AnimNode_MotionMatching.cpp
@@ -45,6 +45,7 @@ void FAnimNode_MotionMatching::Initialize_AnyThread(const FAnimationInitializeCo
 void FAnimNode_MotionMatching::Evaluate_AnyThread(FPoseContext& Output)
 {
 	DECLARE_SCOPE_HIERARCHICAL_COUNTER_ANIMNODE(Evaluate_AnyThread);
+	ANIM_MT_SCOPE_CYCLE_COUNTER_VERBOSE(MotionMatching, !IsInGameThread());
 
 	Source.Evaluate(Output);
 
diff --git a/Engine/Plugins/Experimental/Animation/PoseSearch/Source/Runtime/Private/AnimNode_PoseSearchHistoryCollector.cpp b/Engine/Plugins/Experimental/Animation/PoseSearch/Source/Runtime/Private/AnimNode_PoseSearchHistoryCollector.cpp
--- a/Engine/Plugins/Experimental/Animation/PoseSearch/Source/Runtime/Private/AnimNode_PoseSearchHistoryCollector.cpp
+++ b/Engine/Plugins/Experimental/Animation/PoseSearch/Source/Runtime/Private/AnimNode_PoseSearchHistoryCollector.cpp
@@ -92,6 +92,8 @@ void FAnimNode_PoseSearchHistoryCollector::CacheBones_AnyThread(const FAnimation
 void FAnimNode_PoseSearchHistoryCollector::Evaluate_AnyThread(FPoseContext& Output)
 {
 	DECLARE_SCOPE_HIERARCHICAL_COUNTER_ANIMNODE(Evaluate_AnyThread);
+	ANIM_MT_SCOPE_CYCLE_COUNTER_VERBOSE(PoseSearchHistoryCollector, !IsInGameThread());
+
 	check(Output.AnimInstanceProxy);
 
 	Super::Evaluate_AnyThread(Output);
@@ -144,6 +146,8 @@ void FAnimNode_PoseSearchComponentSpaceHistoryCollector::CacheBones_AnyThread(co
 void FAnimNode_PoseSearchComponentSpaceHistoryCollector::EvaluateComponentSpace_AnyThread(FComponentSpacePoseContext& Output)
 {
 	DECLARE_SCOPE_HIERARCHICAL_COUNTER_ANIMNODE(EvaluateComponentSpace_AnyThread);
+	ANIM_MT_SCOPE_CYCLE_COUNTER_VERBOSE(PoseSearchComponentSpaceHistoryCollector, !IsInGameThread());
+
 	check(Output.AnimInstanceProxy);
 
 	Super::EvaluateComponentSpace_AnyThread(Output);
diff --git a/Engine/Plugins/Experimental/ChaosVehiclesPlugin/Source/ChaosVehicles/Private/AnimNode_StageCoachWheelController.cpp b/Engine/Plugins/Experimental/ChaosVehiclesPlugin/Source/ChaosVehicles/Private/AnimNode_StageCoachWheelController.cpp
--- a/Engine/Plugins/Experimental/ChaosVehiclesPlugin/Source/ChaosVehicles/Private/AnimNode_StageCoachWheelController.cpp
+++ b/Engine/Plugins/Experimental/ChaosVehiclesPlugin/Source/ChaosVehicles/Private/AnimNode_StageCoachWheelController.cpp
@@ -49,6 +49,8 @@ void FAnimNode_StageCoachWheelController::EvaluateSkeletalControl_AnyThread(FCom
 {
 	check(OutBoneTransforms.Num() == 0);
 
+	ANIM_MT_SCOPE_CYCLE_COUNTER_VERBOSE(StageCoachWheelController, !IsInGameThread());
+
 	const TArray<FWheelAnimationData>& WheelAnimData = AnimInstanceProxy->GetWheelAnimData();
 
 	const FBoneContainer& BoneContainer = Output.Pose.GetPose().GetBoneContainer();
diff --git a/Engine/Plugins/Experimental/ChaosVehiclesPlugin/Source/ChaosVehicles/Private/AnimNode_WheelController.cpp b/Engine/Plugins/Experimental/ChaosVehiclesPlugin/Source/ChaosVehicles/Private/AnimNode_WheelController.cpp
--- a/Engine/Plugins/Experimental/ChaosVehiclesPlugin/Source/ChaosVehicles/Private/AnimNode_WheelController.cpp
+++ b/Engine/Plugins/Experimental/ChaosVehiclesPlugin/Source/ChaosVehicles/Private/AnimNode_WheelController.cpp
@@ -45,6 +45,8 @@ void FAnimNode_WheelController::EvaluateSkeletalControl_AnyThread(FComponentSpac
 {
 	check(OutBoneTransforms.Num() == 0);
 
+	ANIM_MT_SCOPE_CYCLE_COUNTER_VERBOSE(WheelController, !IsInGameThread());
+
 	const TArray<FWheelAnimationData>& WheelAnimData = AnimInstanceProxy->GetWheelAnimData();
 
 	const FBoneContainer& BoneContainer = Output.Pose.GetPose().GetBoneContainer();
diff --git a/Engine/Source/Runtime/AnimGraphRuntime/Private/AnimNodes/AnimNode_AimOffsetLookAt.cpp b/Engine/Source/Runtime/AnimGraphRuntime/Private/AnimNodes/AnimNode_AimOffsetLookAt.cpp
--- a/Engine/Source/Runtime/AnimGraphRuntime/Private/AnimNodes/AnimNode_AimOffsetLookAt.cpp
+++ b/Engine/Source/Runtime/AnimGraphRuntime/Private/AnimNodes/AnimNode_AimOffsetLookAt.cpp
@@ -87,6 +87,8 @@ void FAnimNode_AimOffsetLookAt::CacheBones_AnyThread(const FAnimationCacheBonesC
 
 void FAnimNode_AimOffsetLookAt::Evaluate_AnyThread(FPoseContext& Context)
 {
+	ANIM_MT_SCOPE_CYCLE_COUNTER_VERBOSE(AimOffsetLookAt, !IsInGameThread());
+
 	// Evaluate base pose
 	BasePose.Evaluate(Context);
 
diff --git a/Engine/Source/Runtime/AnimGraphRuntime/Private/AnimNodes/AnimNode_ApplyAdditive.cpp b/Engine/Source/Runtime/AnimGraphRuntime/Private/AnimNodes/AnimNode_ApplyAdditive.cpp
--- a/Engine/Source/Runtime/AnimGraphRuntime/Private/AnimNodes/AnimNode_ApplyAdditive.cpp
+++ b/Engine/Source/Runtime/AnimGraphRuntime/Private/AnimNodes/AnimNode_ApplyAdditive.cpp
@@ -69,6 +69,8 @@ void FAnimNode_ApplyAdditive::Update_AnyThread(const FAnimationUpdateContext& Co
 void FAnimNode_ApplyAdditive::Evaluate_AnyThread(FPoseContext& Output)
 {
 	DECLARE_SCOPE_HIERARCHICAL_COUNTER_ANIMNODE(Evaluate_AnyThread)
+	ANIM_MT_SCOPE_CYCLE_COUNTER_VERBOSE(ApplyAdditive, !IsInGameThread());
+
 	//@TODO: Could evaluate Base into Output and save a copy
 	if (FAnimWeight::IsRelevant(ActualAlpha))
 	{
diff --git a/Engine/Source/Runtime/AnimGraphRuntime/Private/AnimNodes/AnimNode_BlendBoneByChannel.cpp b/Engine/Source/Runtime/AnimGraphRuntime/Private/AnimNodes/AnimNode_BlendBoneByChannel.cpp
--- a/Engine/Source/Runtime/AnimGraphRuntime/Private/AnimNodes/AnimNode_BlendBoneByChannel.cpp
+++ b/Engine/Source/Runtime/AnimGraphRuntime/Private/AnimNodes/AnimNode_BlendBoneByChannel.cpp
@@ -62,6 +62,8 @@ void FAnimNode_BlendBoneByChannel::Update_AnyThread(const FAnimationUpdateContex
 void FAnimNode_BlendBoneByChannel::Evaluate_AnyThread(FPoseContext& Output)
 {
 	DECLARE_SCOPE_HIERARCHICAL_COUNTER_ANIMNODE(Evaluate_AnyThread)
+	ANIM_MT_SCOPE_CYCLE_COUNTER_VERBOSE(BlendBoneByChannel, !IsInGameThread());
+
 	A.Evaluate(Output);
 
 	if (bBIsRelevant)
diff --git a/Engine/Source/Runtime/AnimGraphRuntime/Private/AnimNodes/AnimNode_BlendSpacePlayer.cpp b/Engine/Source/Runtime/AnimGraphRuntime/Private/AnimNodes/AnimNode_BlendSpacePlayer.cpp
--- a/Engine/Source/Runtime/AnimGraphRuntime/Private/AnimNodes/AnimNode_BlendSpacePlayer.cpp
+++ b/Engine/Source/Runtime/AnimGraphRuntime/Private/AnimNodes/AnimNode_BlendSpacePlayer.cpp
@@ -123,6 +123,7 @@ void FAnimNode_BlendSpacePlayerBase::UpdateInternal(const FAnimationUpdateContex
 void FAnimNode_BlendSpacePlayerBase::Evaluate_AnyThread(FPoseContext& Output)
 {
 	DECLARE_SCOPE_HIERARCHICAL_COUNTER_ANIMNODE(Evaluate_AnyThread) 
+	ANIM_MT_SCOPE_CYCLE_COUNTER_VERBOSE(BlendSpacePlayer, !IsInGameThread());
 
 	UBlendSpace* CurrentBlendSpace = GetBlendSpace();
 	if (CurrentBlendSpace != nullptr && CurrentBlendSpace->GetSkeleton() != nullptr)
diff --git a/Engine/Source/Runtime/AnimGraphRuntime/Private/AnimNodes/AnimNode_CopyPoseFromMesh.cpp b/Engine/Source/Runtime/AnimGraphRuntime/Private/AnimNodes/AnimNode_CopyPoseFromMesh.cpp
--- a/Engine/Source/Runtime/AnimGraphRuntime/Private/AnimNodes/AnimNode_CopyPoseFromMesh.cpp
+++ b/Engine/Source/Runtime/AnimGraphRuntime/Private/AnimNodes/AnimNode_CopyPoseFromMesh.cpp
@@ -183,6 +183,8 @@ void FAnimNode_CopyPoseFromMesh::Update_AnyThread(const FAnimationUpdateContext&
 void FAnimNode_CopyPoseFromMesh::Evaluate_AnyThread(FPoseContext& Output)
 {
 	DECLARE_SCOPE_HIERARCHICAL_COUNTER_ANIMNODE(Evaluate_AnyThread)
+	ANIM_MT_SCOPE_CYCLE_COUNTER_VERBOSE(CopyPoseFromMesh, !IsInGameThread());
+
 	FCompactPose& OutPose = Output.Pose;
 	OutPose.ResetToRefPose();
 	USkeletalMesh* CurrentMesh = CurrentlyUsedMesh.IsValid() ? CurrentlyUsedMesh.Get() : nullptr;
diff --git a/Engine/Source/Runtime/AnimGraphRuntime/Private/AnimNodes/AnimNode_CurveSource.cpp b/Engine/Source/Runtime/AnimGraphRuntime/Private/AnimNodes/AnimNode_CurveSource.cpp
--- a/Engine/Source/Runtime/AnimGraphRuntime/Private/AnimNodes/AnimNode_CurveSource.cpp
+++ b/Engine/Source/Runtime/AnimGraphRuntime/Private/AnimNodes/AnimNode_CurveSource.cpp
@@ -77,6 +77,8 @@ public:
 void FAnimNode_CurveSource::Evaluate_AnyThread(FPoseContext& Output)
 {
 	DECLARE_SCOPE_HIERARCHICAL_COUNTER_ANIMNODE(Evaluate_AnyThread)
+	ANIM_MT_SCOPE_CYCLE_COUNTER_VERBOSE(CurveSource, !IsInGameThread());
+
 	SourcePose.Evaluate(Output);
 
 	if (CurveSource.GetInterface() != nullptr)
diff --git a/Engine/Source/Runtime/AnimGraphRuntime/Private/AnimNodes/AnimNode_MakeDynamicAdditive.cpp b/Engine/Source/Runtime/AnimGraphRuntime/Private/AnimNodes/AnimNode_MakeDynamicAdditive.cpp
--- a/Engine/Source/Runtime/AnimGraphRuntime/Private/AnimNodes/AnimNode_MakeDynamicAdditive.cpp
+++ b/Engine/Source/Runtime/AnimGraphRuntime/Private/AnimNodes/AnimNode_MakeDynamicAdditive.cpp
@@ -42,6 +42,8 @@ void FAnimNode_MakeDynamicAdditive::Update_AnyThread(const FAnimationUpdateConte
 void FAnimNode_MakeDynamicAdditive::Evaluate_AnyThread(FPoseContext& Output)
 {
 	DECLARE_SCOPE_HIERARCHICAL_COUNTER_ANIMNODE(Evaluate_AnyThread)
+	ANIM_MT_SCOPE_CYCLE_COUNTER_VERBOSE(MakeDynamicAdditive, !IsInGameThread());
+
 	FScopedExpectsAdditiveOverride ScopedExpectsAdditiveOverride(Output, false);
 	FPoseContext BaseEvalContext(Output);
 
diff --git a/Engine/Source/Runtime/AnimGraphRuntime/Private/AnimNodes/AnimNode_ModifyCurve.cpp b/Engine/Source/Runtime/AnimGraphRuntime/Private/AnimNodes/AnimNode_ModifyCurve.cpp
--- a/Engine/Source/Runtime/AnimGraphRuntime/Private/AnimNodes/AnimNode_ModifyCurve.cpp
+++ b/Engine/Source/Runtime/AnimGraphRuntime/Private/AnimNodes/AnimNode_ModifyCurve.cpp
@@ -36,6 +36,8 @@ void FAnimNode_ModifyCurve::CacheBones_AnyThread(const FAnimationCacheBonesConte
 void FAnimNode_ModifyCurve::Evaluate_AnyThread(FPoseContext& Output)
 {
 	DECLARE_SCOPE_HIERARCHICAL_COUNTER_ANIMNODE(Evaluate_AnyThread)
+	ANIM_MT_SCOPE_CYCLE_COUNTER_VERBOSE(ModifyCurve, !IsInGameThread());
+
 	FPoseContext SourceData(Output);
 	SourcePose.Evaluate(SourceData);
 
diff --git a/Engine/Source/Runtime/AnimGraphRuntime/Private/AnimNodes/AnimNode_MultiWayBlend.cpp b/Engine/Source/Runtime/AnimGraphRuntime/Private/AnimNodes/AnimNode_MultiWayBlend.cpp
--- a/Engine/Source/Runtime/AnimGraphRuntime/Private/AnimNodes/AnimNode_MultiWayBlend.cpp
+++ b/Engine/Source/Runtime/AnimGraphRuntime/Private/AnimNodes/AnimNode_MultiWayBlend.cpp
@@ -111,6 +111,7 @@ void FAnimNode_MultiWayBlend::Update_AnyThread(const FAnimationUpdateContext& Co
 void FAnimNode_MultiWayBlend::Evaluate_AnyThread(FPoseContext& Output)
 {
 	DECLARE_SCOPE_HIERARCHICAL_COUNTER_ANIMNODE(Evaluate_AnyThread)
+	ANIM_MT_SCOPE_CYCLE_COUNTER_VERBOSE(MultiWayBlend, !IsInGameThread());
 
 	// this function may be reentrant when multiple MultiWayBlend nodes are chained together
 	// these scratch arrays are treated as stacks below
diff --git a/Engine/Source/Runtime/AnimGraphRuntime/Private/AnimNodes/AnimNode_RandomPlayer.cpp b/Engine/Source/Runtime/AnimGraphRuntime/Private/AnimNodes/AnimNode_RandomPlayer.cpp
--- a/Engine/Source/Runtime/AnimGraphRuntime/Private/AnimNodes/AnimNode_RandomPlayer.cpp
+++ b/Engine/Source/Runtime/AnimGraphRuntime/Private/AnimNodes/AnimNode_RandomPlayer.cpp
@@ -265,6 +265,8 @@ void FAnimNode_RandomPlayer::Update_AnyThread(const FAnimationUpdateContext& Con
 void FAnimNode_RandomPlayer::Evaluate_AnyThread(FPoseContext& Output)
 {
 	DECLARE_SCOPE_HIERARCHICAL_COUNTER_ANIMNODE(Evaluate_AnyThread)
+	ANIM_MT_SCOPE_CYCLE_COUNTER_VERBOSE(RandomPlayer, !IsInGameThread());
+
 	if (ValidEntries.Num() == 0)
 	{
 		Output.ResetToRefPose();
diff --git a/Engine/Source/Runtime/AnimGraphRuntime/Private/AnimNodes/AnimNode_RefPose.cpp b/Engine/Source/Runtime/AnimGraphRuntime/Private/AnimNodes/AnimNode_RefPose.cpp
--- a/Engine/Source/Runtime/AnimGraphRuntime/Private/AnimNodes/AnimNode_RefPose.cpp
+++ b/Engine/Source/Runtime/AnimGraphRuntime/Private/AnimNodes/AnimNode_RefPose.cpp
@@ -50,6 +50,8 @@ ERefPoseType FAnimNode_RefPose::GetRefPoseType() const
 void FAnimNode_MeshSpaceRefPose::EvaluateComponentSpace_AnyThread(FComponentSpacePoseContext& Output)
 {
 	DECLARE_SCOPE_HIERARCHICAL_COUNTER_ANIMNODE(EvaluateComponentSpace_AnyThread)
+	ANIM_MT_SCOPE_CYCLE_COUNTER_VERBOSE(MeshSpaceRefPose, !IsInGameThread());
+
 	Output.ResetToRefPose();
 }
 
diff --git a/Engine/Source/Runtime/AnimGraphRuntime/Private/AnimNodes/AnimNode_RotateRootBone.cpp b/Engine/Source/Runtime/AnimGraphRuntime/Private/AnimNodes/AnimNode_RotateRootBone.cpp
--- a/Engine/Source/Runtime/AnimGraphRuntime/Private/AnimNodes/AnimNode_RotateRootBone.cpp
+++ b/Engine/Source/Runtime/AnimGraphRuntime/Private/AnimNodes/AnimNode_RotateRootBone.cpp
@@ -42,6 +42,8 @@ void FAnimNode_RotateRootBone::Update_AnyThread(const FAnimationUpdateContext& C
 void FAnimNode_RotateRootBone::Evaluate_AnyThread(FPoseContext& Output)
 {
 	DECLARE_SCOPE_HIERARCHICAL_COUNTER_ANIMNODE(Evaluate_AnyThread)
+	ANIM_MT_SCOPE_CYCLE_COUNTER_VERBOSE(RotateRootBone, !IsInGameThread());
+
 	// Evaluate the input
 	BasePose.Evaluate(Output);
 
diff --git a/Engine/Source/Runtime/AnimGraphRuntime/Private/AnimNodes/AnimNode_Slot.cpp b/Engine/Source/Runtime/AnimGraphRuntime/Private/AnimNodes/AnimNode_Slot.cpp
--- a/Engine/Source/Runtime/AnimGraphRuntime/Private/AnimNodes/AnimNode_Slot.cpp
+++ b/Engine/Source/Runtime/AnimGraphRuntime/Private/AnimNodes/AnimNode_Slot.cpp
@@ -83,6 +83,8 @@ void FAnimNode_Slot::Update_AnyThread(const FAnimationUpdateContext& Context)
 void FAnimNode_Slot::Evaluate_AnyThread(FPoseContext & Output)
 {
 	DECLARE_SCOPE_HIERARCHICAL_COUNTER_ANIMNODE(Evaluate_AnyThread)
+	ANIM_MT_SCOPE_CYCLE_COUNTER_VERBOSE(Slot, !IsInGameThread());
+
 	// If not playing a montage, just pass through
 	if (WeightData.SlotNodeWeight <= ZERO_ANIMWEIGHT_THRESH)
 	{
diff --git a/Engine/Source/Runtime/AnimGraphRuntime/Private/AnimNodes/AnimNode_TwoWayBlend.cpp b/Engine/Source/Runtime/AnimGraphRuntime/Private/AnimNodes/AnimNode_TwoWayBlend.cpp
--- a/Engine/Source/Runtime/AnimGraphRuntime/Private/AnimNodes/AnimNode_TwoWayBlend.cpp
+++ b/Engine/Source/Runtime/AnimGraphRuntime/Private/AnimNodes/AnimNode_TwoWayBlend.cpp
@@ -109,6 +109,8 @@ void FAnimNode_TwoWayBlend::Update_AnyThread(const FAnimationUpdateContext& Cont
 void FAnimNode_TwoWayBlend::Evaluate_AnyThread(FPoseContext& Output)
 {
 	DECLARE_SCOPE_HIERARCHICAL_COUNTER_ANIMNODE(Evaluate_AnyThread)
+	ANIM_MT_SCOPE_CYCLE_COUNTER_VERBOSE(TwoWayBlend, !IsInGameThread());
+
 	if (bBIsRelevant)
 	{
 		if (bAIsRelevant)
diff --git a/Engine/Source/Runtime/AnimGraphRuntime/Private/BoneControllers/AnimNodeRotationMultiplier.cpp b/Engine/Source/Runtime/AnimGraphRuntime/Private/BoneControllers/AnimNodeRotationMultiplier.cpp
--- a/Engine/Source/Runtime/AnimGraphRuntime/Private/BoneControllers/AnimNodeRotationMultiplier.cpp
+++ b/Engine/Source/Runtime/AnimGraphRuntime/Private/BoneControllers/AnimNodeRotationMultiplier.cpp
@@ -129,6 +129,8 @@ FQuat FAnimNode_RotationMultiplier::MultiplyQuatBasedOnSourceIndex(const FTransf
 void FAnimNode_RotationMultiplier::EvaluateSkeletalControl_AnyThread(FComponentSpacePoseContext& Output, TArray<FBoneTransform>& OutBoneTransforms)
 {
 	DECLARE_SCOPE_HIERARCHICAL_COUNTER_ANIMNODE(EvaluateSkeletalControl_AnyThread)
+	ANIM_MT_SCOPE_CYCLE_COUNTER_VERBOSE(RotationMultiplier, !IsInGameThread());
+
 	check(OutBoneTransforms.Num() == 0);
 
 	if ( Multiplier != 0.f )
diff --git a/Engine/Source/Runtime/AnimGraphRuntime/Private/BoneControllers/AnimNode_ApplyLimits.cpp b/Engine/Source/Runtime/AnimGraphRuntime/Private/BoneControllers/AnimNode_ApplyLimits.cpp
--- a/Engine/Source/Runtime/AnimGraphRuntime/Private/BoneControllers/AnimNode_ApplyLimits.cpp
+++ b/Engine/Source/Runtime/AnimGraphRuntime/Private/BoneControllers/AnimNode_ApplyLimits.cpp
@@ -32,6 +32,8 @@ void FAnimNode_ApplyLimits::GatherDebugData(FNodeDebugData& DebugData)
 void FAnimNode_ApplyLimits::EvaluateSkeletalControl_AnyThread(FComponentSpacePoseContext& Output, TArray<FBoneTransform>& OutBoneTransforms)
 {
 	DECLARE_SCOPE_HIERARCHICAL_COUNTER_ANIMNODE(EvaluateSkeletalControl_AnyThread)
+	ANIM_MT_SCOPE_CYCLE_COUNTER_VERBOSE(ApplyLimits, !IsInGameThread());
+
 	checkSlow(OutBoneTransforms.Num() == 0);
 
 	FPoseContext LocalPose0(Output.AnimInstanceProxy);
diff --git a/Engine/Source/Runtime/AnimGraphRuntime/Private/BoneControllers/AnimNode_BoneDrivenController.cpp b/Engine/Source/Runtime/AnimGraphRuntime/Private/BoneControllers/AnimNode_BoneDrivenController.cpp
--- a/Engine/Source/Runtime/AnimGraphRuntime/Private/BoneControllers/AnimNode_BoneDrivenController.cpp
+++ b/Engine/Source/Runtime/AnimGraphRuntime/Private/BoneControllers/AnimNode_BoneDrivenController.cpp
@@ -61,6 +61,8 @@ void FAnimNode_BoneDrivenController::GatherDebugData(FNodeDebugData& DebugData)
 void FAnimNode_BoneDrivenController::EvaluateSkeletalControl_AnyThread(FComponentSpacePoseContext& Output, TArray<FBoneTransform>& OutBoneTransforms)
 {
 	DECLARE_SCOPE_HIERARCHICAL_COUNTER_ANIMNODE(EvaluateSkeletalControl_AnyThread)
+	ANIM_MT_SCOPE_CYCLE_COUNTER_VERBOSE(BoneDrivenController, !IsInGameThread());
+
 	check(OutBoneTransforms.Num() == 0);
 
 	// Early out if we're not driving from or to anything
diff --git a/Engine/Source/Runtime/AnimGraphRuntime/Private/BoneControllers/AnimNode_CCDIK.cpp b/Engine/Source/Runtime/AnimGraphRuntime/Private/BoneControllers/AnimNode_CCDIK.cpp
--- a/Engine/Source/Runtime/AnimGraphRuntime/Private/BoneControllers/AnimNode_CCDIK.cpp
+++ b/Engine/Source/Runtime/AnimGraphRuntime/Private/BoneControllers/AnimNode_CCDIK.cpp
@@ -49,6 +49,8 @@ FTransform FAnimNode_CCDIK::GetTargetTransform(const FTransform& InComponentTran
 void FAnimNode_CCDIK::EvaluateSkeletalControl_AnyThread(FComponentSpacePoseContext& Output, TArray<FBoneTransform>& OutBoneTransforms)
 {
 	DECLARE_SCOPE_HIERARCHICAL_COUNTER_ANIMNODE(EvaluateSkeletalControl_AnyThread)
+	ANIM_MT_SCOPE_CYCLE_COUNTER_VERBOSE(CCDIK, !IsInGameThread());
+
 	const FBoneContainer& BoneContainer = Output.Pose.GetPose().GetBoneContainer();
 
 	// Update EffectorLocation if it is based off a bone position
diff --git a/Engine/Source/Runtime/AnimGraphRuntime/Private/BoneControllers/AnimNode_Constraint.cpp b/Engine/Source/Runtime/AnimGraphRuntime/Private/BoneControllers/AnimNode_Constraint.cpp
--- a/Engine/Source/Runtime/AnimGraphRuntime/Private/BoneControllers/AnimNode_Constraint.cpp
+++ b/Engine/Source/Runtime/AnimGraphRuntime/Private/BoneControllers/AnimNode_Constraint.cpp
@@ -40,6 +40,8 @@ void FAnimNode_Constraint::GatherDebugData(FNodeDebugData& DebugData)
 void FAnimNode_Constraint::EvaluateSkeletalControl_AnyThread(FComponentSpacePoseContext& Output, TArray<FBoneTransform>& OutBoneTransforms)
 {
 	DECLARE_SCOPE_HIERARCHICAL_COUNTER_ANIMNODE(EvaluateSkeletalControl_AnyThread)
+	ANIM_MT_SCOPE_CYCLE_COUNTER_VERBOSE(Constraint, !IsInGameThread());
+
  	const FBoneContainer& BoneContainer = Output.Pose.GetPose().GetBoneContainer();
 	const int32 ConstraintNum = ConstraintSetup.Num();
 #if WITH_EDITOR
diff --git a/Engine/Source/Runtime/AnimGraphRuntime/Private/BoneControllers/AnimNode_CopyBone.cpp b/Engine/Source/Runtime/AnimGraphRuntime/Private/BoneControllers/AnimNode_CopyBone.cpp
--- a/Engine/Source/Runtime/AnimGraphRuntime/Private/BoneControllers/AnimNode_CopyBone.cpp
+++ b/Engine/Source/Runtime/AnimGraphRuntime/Private/BoneControllers/AnimNode_CopyBone.cpp
@@ -34,6 +34,8 @@ void FAnimNode_CopyBone::GatherDebugData(FNodeDebugData& DebugData)
 void FAnimNode_CopyBone::EvaluateSkeletalControl_AnyThread(FComponentSpacePoseContext& Output, TArray<FBoneTransform>& OutBoneTransforms)
 {
 	DECLARE_SCOPE_HIERARCHICAL_COUNTER_ANIMNODE(EvaluateSkeletalControl_AnyThread)
+	ANIM_MT_SCOPE_CYCLE_COUNTER_VERBOSE(CopyBone, !IsInGameThread());
+
 	check(OutBoneTransforms.Num() == 0);
 
 	// Pass through if we're not doing anything.
diff --git a/Engine/Source/Runtime/AnimGraphRuntime/Private/BoneControllers/AnimNode_CopyBoneDelta.cpp b/Engine/Source/Runtime/AnimGraphRuntime/Private/BoneControllers/AnimNode_CopyBoneDelta.cpp
--- a/Engine/Source/Runtime/AnimGraphRuntime/Private/BoneControllers/AnimNode_CopyBoneDelta.cpp
+++ b/Engine/Source/Runtime/AnimGraphRuntime/Private/BoneControllers/AnimNode_CopyBoneDelta.cpp
@@ -37,6 +37,8 @@ void FAnimNode_CopyBoneDelta::EvaluateSkeletalControl_AnyThread(FComponentSpaceP
 		return;
 	}
 
+	ANIM_MT_SCOPE_CYCLE_COUNTER_VERBOSE(CopyBoneDelta, !IsInGameThread());
+
 	const FBoneContainer& BoneContainer = Output.Pose.GetPose().GetBoneContainer();
 	FCompactPoseBoneIndex SourceBoneIndex = SourceBone.GetCompactPoseIndex(BoneContainer);
 	FCompactPoseBoneIndex TargetBoneIndex = TargetBone.GetCompactPoseIndex(BoneContainer);
diff --git a/Engine/Source/Runtime/AnimGraphRuntime/Private/BoneControllers/AnimNode_Fabrik.cpp b/Engine/Source/Runtime/AnimGraphRuntime/Private/BoneControllers/AnimNode_Fabrik.cpp
--- a/Engine/Source/Runtime/AnimGraphRuntime/Private/BoneControllers/AnimNode_Fabrik.cpp
+++ b/Engine/Source/Runtime/AnimGraphRuntime/Private/BoneControllers/AnimNode_Fabrik.cpp
@@ -52,6 +52,8 @@ FTransform FAnimNode_Fabrik::GetTargetTransform(const FTransform& InComponentTra
 void FAnimNode_Fabrik::EvaluateSkeletalControl_AnyThread(FComponentSpacePoseContext& Output, TArray<FBoneTransform>& OutBoneTransforms)
 {
 	DECLARE_SCOPE_HIERARCHICAL_COUNTER_ANIMNODE(EvaluateSkeletalControl_AnyThread)
+	ANIM_MT_SCOPE_CYCLE_COUNTER_VERBOSE(Fabrik, !IsInGameThread());
+
 	const FBoneContainer& BoneContainer = Output.Pose.GetPose().GetBoneContainer();
 
 	// Update EffectorLocation if it is based off a bone position
diff --git a/Engine/Source/Runtime/AnimGraphRuntime/Private/BoneControllers/AnimNode_HandIKRetargeting.cpp b/Engine/Source/Runtime/AnimGraphRuntime/Private/BoneControllers/AnimNode_HandIKRetargeting.cpp
--- a/Engine/Source/Runtime/AnimGraphRuntime/Private/BoneControllers/AnimNode_HandIKRetargeting.cpp
+++ b/Engine/Source/Runtime/AnimGraphRuntime/Private/BoneControllers/AnimNode_HandIKRetargeting.cpp
@@ -34,6 +34,8 @@ void FAnimNode_HandIKRetargeting::GatherDebugData(FNodeDebugData& DebugData)
 void FAnimNode_HandIKRetargeting::EvaluateSkeletalControl_AnyThread(FComponentSpacePoseContext& Output, TArray<FBoneTransform>& OutBoneTransforms)
 {
 	DECLARE_SCOPE_HIERARCHICAL_COUNTER_ANIMNODE(EvaluateSkeletalControl_AnyThread)
+	ANIM_MT_SCOPE_CYCLE_COUNTER_VERBOSE(HandIKRetargeting, !IsInGameThread());
+
 	checkSlow(OutBoneTransforms.Num() == 0);
 
 	const FBoneContainer& BoneContainer = Output.Pose.GetPose().GetBoneContainer();
diff --git a/Engine/Source/Runtime/AnimGraphRuntime/Private/BoneControllers/AnimNode_LookAt.cpp b/Engine/Source/Runtime/AnimGraphRuntime/Private/BoneControllers/AnimNode_LookAt.cpp
--- a/Engine/Source/Runtime/AnimGraphRuntime/Private/BoneControllers/AnimNode_LookAt.cpp
+++ b/Engine/Source/Runtime/AnimGraphRuntime/Private/BoneControllers/AnimNode_LookAt.cpp
@@ -81,6 +81,8 @@ float FAnimNode_LookAt::AlphaToBlendType(float InAlpha, EInterpolationBlend::Typ
 void FAnimNode_LookAt::EvaluateSkeletalControl_AnyThread(FComponentSpacePoseContext& Output, TArray<FBoneTransform>& OutBoneTransforms)
 {
 	DECLARE_SCOPE_HIERARCHICAL_COUNTER_ANIMNODE(EvaluateSkeletalControl_AnyThread)
+	ANIM_MT_SCOPE_CYCLE_COUNTER_VERBOSE(LookAt, !IsInGameThread());
+
 	check(OutBoneTransforms.Num() == 0);
 
 	const FBoneContainer& BoneContainer = Output.Pose.GetPose().GetBoneContainer();
diff --git a/Engine/Source/Runtime/AnimGraphRuntime/Private/BoneControllers/AnimNode_ModifyBone.cpp b/Engine/Source/Runtime/AnimGraphRuntime/Private/BoneControllers/AnimNode_ModifyBone.cpp
--- a/Engine/Source/Runtime/AnimGraphRuntime/Private/BoneControllers/AnimNode_ModifyBone.cpp
+++ b/Engine/Source/Runtime/AnimGraphRuntime/Private/BoneControllers/AnimNode_ModifyBone.cpp
@@ -37,6 +37,8 @@ void FAnimNode_ModifyBone::GatherDebugData(FNodeDebugData& DebugData)
 void FAnimNode_ModifyBone::EvaluateSkeletalControl_AnyThread(FComponentSpacePoseContext& Output, TArray<FBoneTransform>& OutBoneTransforms)
 {
 	DECLARE_SCOPE_HIERARCHICAL_COUNTER_ANIMNODE(EvaluateSkeletalControl_AnyThread)
+	ANIM_MT_SCOPE_CYCLE_COUNTER_VERBOSE(ModifyBone, !IsInGameThread());
+
 	check(OutBoneTransforms.Num() == 0);
 
 	// the way we apply transform is same as FMatrix or FTransform
diff --git a/Engine/Source/Runtime/AnimGraphRuntime/Private/BoneControllers/AnimNode_ObserveBone.cpp b/Engine/Source/Runtime/AnimGraphRuntime/Private/BoneControllers/AnimNode_ObserveBone.cpp
--- a/Engine/Source/Runtime/AnimGraphRuntime/Private/BoneControllers/AnimNode_ObserveBone.cpp
+++ b/Engine/Source/Runtime/AnimGraphRuntime/Private/BoneControllers/AnimNode_ObserveBone.cpp
@@ -32,6 +32,8 @@ void FAnimNode_ObserveBone::GatherDebugData(FNodeDebugData& DebugData)
 void FAnimNode_ObserveBone::EvaluateSkeletalControl_AnyThread(FComponentSpacePoseContext& Output, TArray<FBoneTransform>& OutBoneTransforms)
 {
 	DECLARE_SCOPE_HIERARCHICAL_COUNTER_ANIMNODE(EvaluateSkeletalControl_AnyThread)
+	ANIM_MT_SCOPE_CYCLE_COUNTER_VERBOSE(ObserveBone, !IsInGameThread());
+
 	const FBoneContainer& BoneContainer = Output.Pose.GetPose().GetBoneContainer();
 
 	const FCompactPoseBoneIndex BoneIndex = BoneToObserve.GetCompactPoseIndex(BoneContainer);
diff --git a/Engine/Source/Runtime/AnimGraphRuntime/Private/BoneControllers/AnimNode_ResetRoot.cpp b/Engine/Source/Runtime/AnimGraphRuntime/Private/BoneControllers/AnimNode_ResetRoot.cpp
--- a/Engine/Source/Runtime/AnimGraphRuntime/Private/BoneControllers/AnimNode_ResetRoot.cpp
+++ b/Engine/Source/Runtime/AnimGraphRuntime/Private/BoneControllers/AnimNode_ResetRoot.cpp
@@ -29,6 +29,8 @@ void FAnimNode_ResetRoot::GatherDebugData(FNodeDebugData& DebugData)
 void FAnimNode_ResetRoot::EvaluateSkeletalControl_AnyThread(FComponentSpacePoseContext& Output, TArray<FBoneTransform>& OutBoneTransforms)
 {
 	DECLARE_SCOPE_HIERARCHICAL_COUNTER_ANIMNODE(EvaluateSkeletalControl_AnyThread)
+	ANIM_MT_SCOPE_CYCLE_COUNTER_VERBOSE(ResetRoot, !IsInGameThread());
+
 	const FBoneContainer& BoneContainer = Output.Pose.GetPose().GetBoneContainer();
 	FCSPose<FCompactPose>& CSPose = Output.Pose;
 
diff --git a/Engine/Source/Runtime/AnimGraphRuntime/Private/BoneControllers/AnimNode_ScaleChainLength.cpp b/Engine/Source/Runtime/AnimGraphRuntime/Private/BoneControllers/AnimNode_ScaleChainLength.cpp
--- a/Engine/Source/Runtime/AnimGraphRuntime/Private/BoneControllers/AnimNode_ScaleChainLength.cpp
+++ b/Engine/Source/Runtime/AnimGraphRuntime/Private/BoneControllers/AnimNode_ScaleChainLength.cpp
@@ -48,6 +48,8 @@ void FAnimNode_ScaleChainLength::CacheBones_AnyThread(const FAnimationCacheBones
 void FAnimNode_ScaleChainLength::Evaluate_AnyThread(FPoseContext& Output)
 {
 	DECLARE_SCOPE_HIERARCHICAL_COUNTER_ANIMNODE(Evaluate_AnyThread)
+	ANIM_MT_SCOPE_CYCLE_COUNTER_VERBOSE(ScaleChainLength, !IsInGameThread());
+
 	// Evaluate incoming pose into our output buffer.
 	InputPose.Evaluate(Output);
 
diff --git a/Engine/Source/Runtime/AnimGraphRuntime/Private/BoneControllers/AnimNode_SkeletalControlBase.cpp b/Engine/Source/Runtime/AnimGraphRuntime/Private/BoneControllers/AnimNode_SkeletalControlBase.cpp
--- a/Engine/Source/Runtime/AnimGraphRuntime/Private/BoneControllers/AnimNode_SkeletalControlBase.cpp
+++ b/Engine/Source/Runtime/AnimGraphRuntime/Private/BoneControllers/AnimNode_SkeletalControlBase.cpp
@@ -119,6 +119,7 @@ void FAnimNode_SkeletalControlBase::EvaluateComponentSpaceInternal(FComponentSpa
 void FAnimNode_SkeletalControlBase::EvaluateComponentSpace_AnyThread(FComponentSpacePoseContext& Output)
 {
 	DECLARE_SCOPE_HIERARCHICAL_COUNTER_ANIMNODE(EvaluateComponentSpace_AnyThread)
+	ANIM_MT_SCOPE_CYCLE_COUNTER_VERBOSE(SkeletalControlBase, !IsInGameThread());
 
 	// Cache the incoming node IDs in a base context
 	FAnimationBaseContext CachedContext(Output);
diff --git a/Engine/Source/Runtime/AnimGraphRuntime/Private/BoneControllers/AnimNode_SplineIK.cpp b/Engine/Source/Runtime/AnimGraphRuntime/Private/BoneControllers/AnimNode_SplineIK.cpp
--- a/Engine/Source/Runtime/AnimGraphRuntime/Private/BoneControllers/AnimNode_SplineIK.cpp
+++ b/Engine/Source/Runtime/AnimGraphRuntime/Private/BoneControllers/AnimNode_SplineIK.cpp
@@ -55,6 +55,8 @@ struct FSplineIKScratchArea : public TThreadSingleton<FSplineIKScratchArea>
 void FAnimNode_SplineIK::EvaluateSkeletalControl_AnyThread(FComponentSpacePoseContext& Output, TArray<FBoneTransform>& OutBoneTransforms)
 {
 	DECLARE_SCOPE_HIERARCHICAL_COUNTER_ANIMNODE(EvaluateSkeletalControl_AnyThread)
+	ANIM_MT_SCOPE_CYCLE_COUNTER_VERBOSE(SplineIK, !IsInGameThread());
+
 	if (CachedBoneReferences.Num() > 0)
 	{
 		const FBoneContainer& BoneContainer = Output.Pose.GetPose().GetBoneContainer();
diff --git a/Engine/Source/Runtime/AnimGraphRuntime/Private/BoneControllers/AnimNode_SpringBone.cpp b/Engine/Source/Runtime/AnimGraphRuntime/Private/BoneControllers/AnimNode_SpringBone.cpp
--- a/Engine/Source/Runtime/AnimGraphRuntime/Private/BoneControllers/AnimNode_SpringBone.cpp
+++ b/Engine/Source/Runtime/AnimGraphRuntime/Private/BoneControllers/AnimNode_SpringBone.cpp
@@ -91,6 +91,8 @@ FORCEINLINE void CopyToVectorByFlags(FVector& DestVec, const FVector& SrcVec, bo
 void FAnimNode_SpringBone::EvaluateSkeletalControl_AnyThread(FComponentSpacePoseContext& Output, TArray<FBoneTransform>& OutBoneTransforms)
 {
 	DECLARE_SCOPE_HIERARCHICAL_COUNTER_ANIMNODE(EvaluateSkeletalControl_AnyThread)
+	ANIM_MT_SCOPE_CYCLE_COUNTER_VERBOSE(SpringBone, !IsInGameThread());
+
 	check(OutBoneTransforms.Num() == 0);
 
 	const bool bNoOffset = !bTranslateX && !bTranslateY && !bTranslateZ;
diff --git a/Engine/Source/Runtime/Engine/Classes/Animation/AnimNode_CustomProperty.h b/Engine/Source/Runtime/Engine/Classes/Animation/AnimNode_CustomProperty.h
--- a/Engine/Source/Runtime/Engine/Classes/Animation/AnimNode_CustomProperty.h
+++ b/Engine/Source/Runtime/Engine/Classes/Animation/AnimNode_CustomProperty.h
@@ -103,7 +103,14 @@ protected:
 	UPROPERTY(Transient)
 	TObjectPtr<UObject> SourceInstance;
 #endif
-	
+
+	// Stats
+#if ANIMNODE_STATS_VERBOSE
+	// Cached StatID for this node
+	TStatId StatID;
+	virtual void InitializeStatID() { StatID = FDynamicStats::CreateStatId<FStatGroup_STATGROUP_Anim>(FString(TEXT("Unknown"))); }
+#endif // ANIMNODE_STATS_VERBOSE
+
 	friend class UAnimGraphNode_CustomProperty;
 };
 
diff --git a/Engine/Source/Runtime/Engine/Classes/Animation/AnimNode_LinkedAnimGraph.h b/Engine/Source/Runtime/Engine/Classes/Animation/AnimNode_LinkedAnimGraph.h
--- a/Engine/Source/Runtime/Engine/Classes/Animation/AnimNode_LinkedAnimGraph.h
+++ b/Engine/Source/Runtime/Engine/Classes/Animation/AnimNode_LinkedAnimGraph.h
@@ -139,6 +139,12 @@ protected:
 	void RequestBlend(const IAnimClassInterface* PriorAnimBPClass, const IAnimClassInterface* NewAnimBPClass);
 
 	friend class UAnimInstance;
+
+	// Stats
+#if ANIMNODE_STATS_VERBOSE
+	virtual void InitializeStatID() override;
+#endif
+
 };
 
 UE_DEPRECATED(4.24, "FAnimNode_SubInstance has been renamed to FAnimNode_LinkedAnimGraph")
diff --git a/Engine/Source/Runtime/Engine/Classes/Animation/AnimNode_LinkedAnimLayer.h b/Engine/Source/Runtime/Engine/Classes/Animation/AnimNode_LinkedAnimLayer.h
--- a/Engine/Source/Runtime/Engine/Classes/Animation/AnimNode_LinkedAnimLayer.h
+++ b/Engine/Source/Runtime/Engine/Classes/Animation/AnimNode_LinkedAnimLayer.h
@@ -74,6 +74,11 @@ protected:
 	// Event fired when the instance we are running has changed
 	FSimpleMulticastDelegate OnInstanceChangedEvent;
 #endif
+
+    // stats
+#if ANIMNODE_STATS_VERBOSE
+	virtual void InitializeStatID() override;
+#endif
 };
 
 UE_DEPRECATED(4.24, "FAnimNode_Layer has been renamed to FAnimNode_LinkedAnimLayer")
diff --git a/Engine/Source/Runtime/Engine/Private/Animation/AnimNodeSpaceConversions.cpp b/Engine/Source/Runtime/Engine/Private/Animation/AnimNodeSpaceConversions.cpp
--- a/Engine/Source/Runtime/Engine/Private/Animation/AnimNodeSpaceConversions.cpp
+++ b/Engine/Source/Runtime/Engine/Private/Animation/AnimNodeSpaceConversions.cpp
@@ -28,6 +28,8 @@ void FAnimNode_ConvertComponentToLocalSpace::Update_AnyThread(const FAnimationUp
 
 void FAnimNode_ConvertComponentToLocalSpace::Evaluate_AnyThread(FPoseContext & Output)
 {
+	ANIM_MT_SCOPE_CYCLE_COUNTER_VERBOSE(ConvertComponentToLocalSpace, !IsInGameThread());
+
 	// Evaluate the child and convert
 	FComponentSpacePoseContext InputCSPose(Output.AnimInstanceProxy);
 
@@ -80,6 +82,8 @@ void FAnimNode_ConvertLocalToComponentSpace::GatherDebugData(FNodeDebugData& Deb
 
 void FAnimNode_ConvertLocalToComponentSpace::EvaluateComponentSpace_AnyThread(FComponentSpacePoseContext & OutputCSPose)
 {
+	ANIM_MT_SCOPE_CYCLE_COUNTER_VERBOSE(ConvertLocalToComponentSpace, !IsInGameThread());
+
 	// Evaluate the child and convert
 	FPoseContext InputPose(OutputCSPose.AnimInstanceProxy);
 
diff --git a/Engine/Source/Runtime/Engine/Private/Animation/AnimNode_ApplyMeshSpaceAdditive.cpp b/Engine/Source/Runtime/Engine/Private/Animation/AnimNode_ApplyMeshSpaceAdditive.cpp
--- a/Engine/Source/Runtime/Engine/Private/Animation/AnimNode_ApplyMeshSpaceAdditive.cpp
+++ b/Engine/Source/Runtime/Engine/Private/Animation/AnimNode_ApplyMeshSpaceAdditive.cpp
@@ -75,6 +75,7 @@ void FAnimNode_ApplyMeshSpaceAdditive::Update_AnyThread(const FAnimationUpdateCo
 void FAnimNode_ApplyMeshSpaceAdditive::Evaluate_AnyThread(FPoseContext& Output)
 {
 	DECLARE_SCOPE_HIERARCHICAL_COUNTER_ANIMNODE(Evaluate_AnyThread)
+	ANIM_MT_SCOPE_CYCLE_COUNTER_VERBOSE(ApplyMeshSpaceAdditive, !IsInGameThread());
 
 	//@TODO: Could evaluate Base into Output and save a copy
 	if (FAnimWeight::IsRelevant(ActualAlpha))
diff --git a/Engine/Source/Runtime/Engine/Private/Animation/AnimNode_CustomProperty.cpp b/Engine/Source/Runtime/Engine/Private/Animation/AnimNode_CustomProperty.cpp
--- a/Engine/Source/Runtime/Engine/Private/Animation/AnimNode_CustomProperty.cpp
+++ b/Engine/Source/Runtime/Engine/Private/Animation/AnimNode_CustomProperty.cpp
@@ -25,6 +25,9 @@ FAnimNode_CustomProperty::~FAnimNode_CustomProperty()
 void FAnimNode_CustomProperty::SetTargetInstance(UObject* InInstance)
 {
 	TargetInstance = InInstance;
+#if ANIMNODE_STATS_VERBOSE
+	InitializeStatID();
+#endif
 }
 
 void FAnimNode_CustomProperty::PropagateInputProperties(const UObject* InSourceInstance)
diff --git a/Engine/Source/Runtime/Engine/Private/Animation/AnimNode_Inertialization.cpp b/Engine/Source/Runtime/Engine/Private/Animation/AnimNode_Inertialization.cpp
--- a/Engine/Source/Runtime/Engine/Private/Animation/AnimNode_Inertialization.cpp
+++ b/Engine/Source/Runtime/Engine/Private/Animation/AnimNode_Inertialization.cpp
@@ -213,6 +213,7 @@ void FAnimNode_Inertialization::Evaluate_AnyThread(FPoseContext& Output)
 {
 	LLM_SCOPE_BYNAME(TEXT("Animation/Inertialization"));
 	DECLARE_SCOPE_HIERARCHICAL_COUNTER_ANIMNODE(Evaluate_AnyThread);
+	ANIM_MT_SCOPE_CYCLE_COUNTER_VERBOSE(Inertialization, !IsInGameThread());
 
 	Source.Evaluate(Output);
 
diff --git a/Engine/Source/Runtime/Engine/Private/Animation/AnimNode_LinkedAnimGraph.cpp b/Engine/Source/Runtime/Engine/Private/Animation/AnimNode_LinkedAnimGraph.cpp
--- a/Engine/Source/Runtime/Engine/Private/Animation/AnimNode_LinkedAnimGraph.cpp
+++ b/Engine/Source/Runtime/Engine/Private/Animation/AnimNode_LinkedAnimGraph.cpp
@@ -159,6 +159,11 @@ void FAnimNode_LinkedAnimGraph::Update_AnyThread(const FAnimationUpdateContext&
 
 void FAnimNode_LinkedAnimGraph::Evaluate_AnyThread(FPoseContext& Output)
 {
+#if	ANIMNODE_STATS_VERBOSE
+	// Record name of linked graph we are updating
+	FScopeCycleCounter LinkedAnimGraphNameCycleCounter(StatID);
+#endif // ANIMNODE_STATS_VERBOSE
+
 	UAnimInstance* InstanceToRun = GetTargetInstance<UAnimInstance>();
 	if(InstanceToRun && LinkedRoot)
 	{
@@ -514,3 +519,17 @@ void FAnimNode_LinkedAnimGraph::HandleObjectsReinstanced_Impl(UObject* InSourceO
 	}
 }
 #endif
+
+#if ANIMNODE_STATS_VERBOSE
+void FAnimNode_LinkedAnimGraph::InitializeStatID()
+{
+	if (GetTargetInstance<UAnimInstance>())
+	{
+		StatID = FDynamicStats::CreateStatId<FStatGroup_STATGROUP_Anim>(GetTargetInstance<UAnimInstance>()->GetClass()->GetName());
+	}
+	else
+	{
+		Super::InitializeStatID();
+	}
+}
+#endif // ANIMNODE_STATS_VERBOSE
diff --git a/Engine/Source/Runtime/Engine/Private/Animation/AnimNode_LinkedAnimLayer.cpp b/Engine/Source/Runtime/Engine/Private/Animation/AnimNode_LinkedAnimLayer.cpp
--- a/Engine/Source/Runtime/Engine/Private/Animation/AnimNode_LinkedAnimLayer.cpp
+++ b/Engine/Source/Runtime/Engine/Private/Animation/AnimNode_LinkedAnimLayer.cpp
@@ -242,3 +242,19 @@ void FAnimNode_LinkedAnimLayer::HandleObjectsReinstanced_Impl(UObject* InSourceO
 }
 
 #endif
+
+
+#if ANIMNODE_STATS_VERBOSE
+void FAnimNode_LinkedAnimLayer::InitializeStatID()
+{
+	if (GetTargetInstance<UAnimInstance>())
+	{
+		FString StatName = GetTargetInstance<UAnimInstance>()->GetClass()->GetName() + " " + GetDynamicLinkFunctionName().ToString();
+		StatID = FDynamicStats::CreateStatId<FStatGroup_STATGROUP_Anim>(StatName);
+	}
+	else
+	{
+		Super::InitializeStatID();
+	}
+}
+#endif // ANIMNODE_STATS_VERBOSE
diff --git a/Engine/Source/Runtime/Engine/Private/Animation/AnimNode_SaveCachedPose.cpp b/Engine/Source/Runtime/Engine/Private/Animation/AnimNode_SaveCachedPose.cpp
--- a/Engine/Source/Runtime/Engine/Private/Animation/AnimNode_SaveCachedPose.cpp
+++ b/Engine/Source/Runtime/Engine/Private/Animation/AnimNode_SaveCachedPose.cpp
@@ -131,6 +131,8 @@ void FAnimNode_SaveCachedPose::Evaluate_AnyThread(FPoseContext& Output)
 	FCachedPoseThreadContext& CachedPoseThreadContext = FCachedPoseThreadContext::Get();
 	const FCachedPoseThreadContext::FCachedPose* CachedPose = nullptr;
 
+	ANIM_MT_SCOPE_CYCLE_COUNTER_VERBOSE(SaveCachedPose, !IsInGameThread());
+
 	// Note that we check here for IsSynchronized_All to deal with cases like Sequencer:
 	// In these cases the counter can stay zeroed between unbound/bound updates and this can cause issues
 	// with using out-of-date cached data (stack-allocated from a previous frame).
diff --git a/Engine/Source/Runtime/Engine/Public/Animation/AnimMTStats.h b/Engine/Source/Runtime/Engine/Public/Animation/AnimMTStats.h
--- a/Engine/Source/Runtime/Engine/Public/Animation/AnimMTStats.h
+++ b/Engine/Source/Runtime/Engine/Public/Animation/AnimMTStats.h
@@ -14,4 +14,55 @@ DO_ANIMSTAT_PROCESSING(EvaluateAnimStateMachine);
 DO_ANIMSTAT_PROCESSING(FillComponentSpaceTransforms);
 DO_ANIMSTAT_PROCESSING(ProxyUpdateAnimation);
 DO_ANIMSTAT_PROCESSING(PoseBlendNodeEvaluate);
-DO_ANIMSTAT_PROCESSING(PoseAssetGetAnimationPose);
\ No newline at end of file
+DO_ANIMSTAT_PROCESSING(PoseAssetGetAnimationPose);
+
+// Verbose nodes
+#if ANIMNODE_STATS_VERBOSE
+DO_ANIMSTAT_PROCESSING(BlendSpacePlayer);
+DO_ANIMSTAT_PROCESSING(Inertialization);
+DO_ANIMSTAT_PROCESSING(ControlRig);
+DO_ANIMSTAT_PROCESSING(IKRig);
+DO_ANIMSTAT_PROCESSING(PreviewRetargetPose);
+DO_ANIMSTAT_PROCESSING(RemapCurves);
+DO_ANIMSTAT_PROCESSING(RemapCurvesFromMesh);
+DO_ANIMSTAT_PROCESSING(CorrectPose);
+DO_ANIMSTAT_PROCESSING(MotionMatching);
+DO_ANIMSTAT_PROCESSING(PoseSearchHistoryCollector);
+DO_ANIMSTAT_PROCESSING(AimOffsetLookAt);
+DO_ANIMSTAT_PROCESSING(BlendBoneByChannel);
+DO_ANIMSTAT_PROCESSING(CopyPoseFromMesh);
+DO_ANIMSTAT_PROCESSING(CurveSource);
+DO_ANIMSTAT_PROCESSING(MakeDynamicAdditive);
+DO_ANIMSTAT_PROCESSING(ModifyCurve);
+DO_ANIMSTAT_PROCESSING(MultiWayBlend);
+DO_ANIMSTAT_PROCESSING(RandomPlayer);
+DO_ANIMSTAT_PROCESSING(RotateRootBone);
+DO_ANIMSTAT_PROCESSING(Slot);
+DO_ANIMSTAT_PROCESSING(TwoWayBlend);
+DO_ANIMSTAT_PROCESSING(ScaleChainLength);
+DO_ANIMSTAT_PROCESSING(ApplyAdditive);
+DO_ANIMSTAT_PROCESSING(ApplyMeshSpaceAdditive);
+DO_ANIMSTAT_PROCESSING(ConvertComponentToLocalSpace);
+DO_ANIMSTAT_PROCESSING(SaveCachedPose);
+DO_ANIMSTAT_PROCESSING(PoseSearchComponentSpaceHistoryCollector);
+DO_ANIMSTAT_PROCESSING(MeshSpaceRefPose);
+DO_ANIMSTAT_PROCESSING(SkeletalControlBase);
+DO_ANIMSTAT_PROCESSING(ConvertLocalToComponentSpace);
+DO_ANIMSTAT_PROCESSING(HandIKRetargeting);
+DO_ANIMSTAT_PROCESSING(StageCoachWheelController);
+DO_ANIMSTAT_PROCESSING(WheelController);
+DO_ANIMSTAT_PROCESSING(ApplyLimits);
+DO_ANIMSTAT_PROCESSING(BoneDrivenController);
+DO_ANIMSTAT_PROCESSING(CCDIK);
+DO_ANIMSTAT_PROCESSING(Constraint);
+DO_ANIMSTAT_PROCESSING(CopyBone);
+DO_ANIMSTAT_PROCESSING(CopyBoneDelta);
+DO_ANIMSTAT_PROCESSING(Fabrik);
+DO_ANIMSTAT_PROCESSING(LookAt);
+DO_ANIMSTAT_PROCESSING(ModifyBone);
+DO_ANIMSTAT_PROCESSING(ObserveBone);
+DO_ANIMSTAT_PROCESSING(ResetRoot);
+DO_ANIMSTAT_PROCESSING(SplineIK);
+DO_ANIMSTAT_PROCESSING(SpringBone);
+DO_ANIMSTAT_PROCESSING(RotationMultiplier);
+#endif
diff --git a/Engine/Source/Runtime/Engine/Public/Animation/AnimStats.h b/Engine/Source/Runtime/Engine/Public/Animation/AnimStats.h
--- a/Engine/Source/Runtime/Engine/Public/Animation/AnimStats.h
+++ b/Engine/Source/Runtime/Engine/Public/Animation/AnimStats.h
@@ -45,6 +45,9 @@ DECLARE_CYCLE_STAT_EXTERN(TEXT("UpdateCurvesPostEvaluation"), STAT_UpdateCurvesP
 DECLARE_CYCLE_STAT_EXTERN(TEXT("TickAssetPlayerInstances"), STAT_TickAssetPlayerInstances, STATGROUP_Anim, );
 DECLARE_CYCLE_STAT_EXTERN(TEXT("TickAssetPlayerInstance"), STAT_TickAssetPlayerInstance, STATGROUP_Anim, );
 
+// Uncomment for verbose anim node state
+#define ANIMNODE_STATS_VERBOSE 0
+
 #define DO_ANIMSTAT_PROCESSING(StatName) DECLARE_CYCLE_STAT_EXTERN(TEXT(#StatName), STAT_ ## StatName, STATGROUP_Anim, ENGINE_API)
 #include "Animation/AnimMTStats.h"
 #undef DO_ANIMSTAT_PROCESSING
@@ -57,6 +60,12 @@ DECLARE_CYCLE_STAT_EXTERN(TEXT("TickAssetPlayerInstance"), STAT_TickAssetPlayerI
 #define ANIM_MT_SCOPE_CYCLE_COUNTER(StatName, bIsMultithreaded) \
 	TStatId CycleCountID_##StatName = (bIsMultithreaded ? GET_STATID(STAT_ ## StatName ## _WorkerThread) : GET_STATID(STAT_ ## StatName)); \
 	FScopeCycleCounter CycleCount_##StatName(CycleCountID_##StatName);
+	#if ANIMNODE_STATS_VERBOSE
+		#define ANIM_MT_SCOPE_CYCLE_COUNTER_VERBOSE(StateName, bIsMultithreaded) ANIM_MT_SCOPE_CYCLE_COUNTER(StateName, bIsMultithreaded)
+	#else
+		#define ANIM_MT_SCOPE_CYCLE_COUNTER_VERBOSE(StateName, bIsMultithreaded)
+	#endif
 #else
 	#define ANIM_MT_SCOPE_CYCLE_COUNTER(StatName, bIsMultithreaded)
+	#define ANIM_MT_SCOPE_CYCLE_COUNTER_VERBOSE(StateName, bIsMultithreaded)
 #endif
