additional protections for invalid components/actors being added into the cluster union

diff --git a/Engine/Source/Runtime/Engine/Private/PhysicsEngine/ClusterUnionComponent.cpp b/Engine/Source/Runtime/Engine/Private/PhysicsEngine/ClusterUnionComponent.cpp
--- a/Engine/Source/Runtime/Engine/Private/PhysicsEngine/ClusterUnionComponent.cpp
+++ b/Engine/Source/Runtime/Engine/Private/PhysicsEngine/ClusterUnionComponent.cpp
@@ -73,7 +73,7 @@ FPhysScene_Chaos* UClusterUnionComponent::GetChaosScene() const
 
 void UClusterUnionComponent::AddComponentToCluster(UPrimitiveComponent* InComponent, const TArray<int32>& BoneIds)
 {
-	if (!InComponent || !PhysicsProxy)
+	if (!InComponent || !PhysicsProxy || !InComponent->GetWorld())
 	{
 		return;
 	}
@@ -590,7 +590,7 @@ void UClusterUnionComponent::SyncClusterUnionFromProxy()
 
 void UClusterUnionComponent::HandleAddOrModifiedClusteredComponent(UPrimitiveComponent* ChangedComponent, const TMap<int32, FTransform>& PerBoneChildToParent)
 {
-	if (!ChangedComponent)
+	if (!ChangedComponent || !ChangedComponent->HasValidPhysicsState() || !ChangedComponent->GetWorld())
 	{
 		return;
 	}
