Fixed crash when reparenting bones

diff --git a/Engine/Plugins/Experimental/MeshModelingToolsetExp/Source/SkeletalMeshModifiers/Private/SkeletonModifier.cpp b/Engine/Plugins/Experimental/MeshModelingToolsetExp/Source/SkeletalMeshModifiers/Private/SkeletonModifier.cpp
--- a/Engine/Plugins/Experimental/MeshModelingToolsetExp/Source/SkeletalMeshModifiers/Private/SkeletonModifier.cpp
+++ b/Engine/Plugins/Experimental/MeshModelingToolsetExp/Source/SkeletalMeshModifiers/Private/SkeletonModifier.cpp
@@ -1042,23 +1042,14 @@ bool USkeletonModifier::ParentBones(const TArray<FName>& InBoneNames, const TArr
 			if (BoneIndex != INDEX_NONE)
 			{
 				const FName NewParentName = GetParentName(Index);
-
-				// store global transforms
-				const FTransform GlobalTransform = TransformComposer->GetGlobalTransform(BoneIndex);
-
-				const int32 ParentIndex = ReferenceSkeleton->FindRawBoneIndex(NewParentName);
-				const FTransform ParentTransform = TransformComposer->GetGlobalTransform(ParentIndex);
 				
 				// change parent
 				const int32 NewIndex = Modifier.SetParent(InBoneNames[Index], NewParentName, bAllowMultipleRoots);
 				if (NewIndex > INDEX_NONE)
 				{
-					// update local transform
-					Modifier.UpdateRefPoseTransform(NewIndex, GlobalTransform.GetRelativeTransform(ParentTransform));
+					// invalidate composer
+					TransformComposer->Invalidate(INDEX_NONE);
 				}
-
-				// invalidate composer
-				TransformComposer->Invalidate(INDEX_NONE);
 			}
 		}
 	}
diff --git a/Engine/Source/Runtime/Engine/Private/Animation/ReferenceSkeleton.cpp b/Engine/Source/Runtime/Engine/Private/Animation/ReferenceSkeleton.cpp
--- a/Engine/Source/Runtime/Engine/Private/Animation/ReferenceSkeleton.cpp
+++ b/Engine/Source/Runtime/Engine/Private/Animation/ReferenceSkeleton.cpp
@@ -210,8 +210,130 @@ void FReferenceSkeleton::Rename(const FName InBoneName, const FName InNewName)
 	RawNameToIndexMap.Add(InNewName, RawBoneIndex);
 }
 
-int32 FReferenceSkeleton::SetParent(const FName InBoneName, const FName InParentName)
+namespace FReferenceSkeletonLocals
 {
+	struct FElement
+	{
+		int32 RawIndex = INDEX_NONE;
+		FElement* Parent = nullptr;
+		TArray<FElement*> Children;
+	};
+	
+	struct FHierarchy
+	{
+		FHierarchy(const FReferenceSkeleton& InRefSkeleton)
+		{
+			const TArray<FMeshBoneInfo>& Infos = InRefSkeleton.GetRawRefBoneInfo();
+			const int32 NumBones = Infos.Num();
+			if (NumBones > 0)
+			{
+				const TArray<FTransform>& Poses = InRefSkeleton.GetRawRefBonePose();
+
+				Elements.Reserve(NumBones);
+				Transforms.Reserve(NumBones);
+				TBitArray<> TransformCached(false, NumBones);
+
+				// build elements
+				for (int32 Index = 0; Index < NumBones; ++Index)
+				{
+					FElement NewElement({ Index, nullptr, {} });
+					Elements.Emplace(MoveTemp(NewElement));
+					Transforms.Add(Poses[Index]);
+				}
+
+				// build hierarchy & global transforms
+				auto GetGlobalTransform = [this, &TransformCached](const FElement& Element, auto&& GetGlobalTransformArg)
+				{
+					if (TransformCached[Element.RawIndex])
+					{
+						return Transforms[Element.RawIndex];
+					}
+
+					if (Element.Parent)
+					{
+						Transforms[Element.RawIndex] *= GetGlobalTransformArg(*Element.Parent, GetGlobalTransformArg);
+					}
+
+					TransformCached[Element.RawIndex] = true;
+					return Transforms[Element.RawIndex];
+				};
+				
+				for (int32 Index = 0; Index < NumBones; ++Index)
+				{
+					SetParent(Index, Infos[Index].ParentIndex);
+					GetGlobalTransform(Elements[Index], GetGlobalTransform);
+				}
+			}
+		}
+
+		// switch parent
+		void SetParent(const int32 InChildIndex, const int32 InParentIndex)
+		{
+			check(Elements.IsValidIndex(InChildIndex));
+			check(InParentIndex >= INDEX_NONE && InParentIndex < Elements.Num());
+
+			FElement& Child = Elements[InChildIndex];
+			if (Child.Parent)
+			{
+				Child.Parent->Children.Remove(&Child);
+				Child.Parent = nullptr;
+			}
+			
+			if (Elements.IsValidIndex(InParentIndex))
+			{
+				FElement& Parent = Elements[InParentIndex];
+				Child.Parent = &Parent;
+				Parent.Children.Add(&Child);
+			}
+		}
+
+		// get the new to old index mapping to update the bone infos and poses
+		void GetNewToOldIndexes(TArray<int32>& OutMapping)
+		{
+			OutMapping.Reset(0); OutMapping.Reserve(Elements.Num());
+			
+			TBitArray<> Visited(false, Elements.Num());
+
+			auto ElementToIndex = [this, &Visited, &OutMapping](const FElement& Element, auto&& ElementToIndexArg) -> void
+			{
+				// add the element
+				if (!Visited[Element.RawIndex])
+				{
+					OutMapping.Add(Element.RawIndex);
+					Visited[Element.RawIndex] = true;
+				}
+
+				// add it's direct children
+				for (const FElement* Child: Element.Children)
+				{
+					if (!Visited[Child->RawIndex])
+					{
+						OutMapping.Add(Child->RawIndex);
+						Visited[Child->RawIndex] = true;
+					}
+				}
+
+				// recurse
+				for (const FElement* Child: Element.Children)
+				{
+					ElementToIndexArg(*Child, ElementToIndexArg);
+				}
+			};
+
+			// parse the full hierarchy
+			for (const FElement& Element: Elements)
+			{
+				ElementToIndex(Element, ElementToIndex);
+			}
+		}
+
+		TArray<FElement> Elements;
+		TArray<FTransform> Transforms;
+	};
+}
+
+int32 FReferenceSkeleton::SetParent(const FName InBoneName, const FName InParentName)
+{	
 	if (InBoneName == InParentName)
 	{
 		return INDEX_NONE;
@@ -231,6 +353,11 @@ int32 FReferenceSkeleton::SetParent(const FName InBoneName, const FName InParent
 		return INDEX_NONE;
 	}
 
+	if (RawRefBoneInfo[BoneIndex].ParentIndex == NewParentIndex)
+	{
+		return INDEX_NONE;
+	}
+
 	auto GetParentIndex = [&](const int32 InBoneIndex) -> int32
 	{
 		return RawRefBoneInfo.IsValidIndex(InBoneIndex) ? RawRefBoneInfo[InBoneIndex].ParentIndex : INDEX_NONE;
@@ -246,129 +373,68 @@ int32 FReferenceSkeleton::SetParent(const FName InBoneName, const FName InParent
 		}
 		ParentParentIndex = GetParentIndex(ParentParentIndex);
 	}
+
+	using namespace FReferenceSkeletonLocals;
 	
-	const int32 NumBones = GetRawBoneNum();
-	auto GetEndOfBranch = [&](const int32 InBoneIndex) -> int32
-	{
-		const int32 StartBranch = InBoneIndex; 
+	// build temp hierarchy
+	FHierarchy Hierarchy(*this);
 
-		int32 EndOfBranch = InBoneIndex+1;
-		int32 ParentIndex = GetParentIndex(EndOfBranch);
+	// switch parent
+	Hierarchy.SetParent(BoneIndex, NewParentIndex);
 
-		// if next child bone's parent is less than or equal to StartParentIndex,
-		// we are leaving the branch so no need to go further
-		while (ParentIndex >= StartBranch && EndOfBranch < NumBones)
-		{
-			EndOfBranch++;
-			ParentIndex = GetParentIndex(EndOfBranch);
-		}
-		return EndOfBranch;
-	};
+	// update infos, poses and name to index mapping
+	TArray<int32> NewToOldIndexes; Hierarchy.GetNewToOldIndexes(NewToOldIndexes);
 
-	// bone branch to move 
-	const int32 EndOfBoneBranch = GetEndOfBranch(BoneIndex);
+	const TArray<FMeshBoneInfo>& Infos = GetRawRefBoneInfo();
+	
+	check(NewToOldIndexes.Num() == Infos.Num());
 
-	// parent branch to move
-	const int32 EndOfNewParentBranch = NewParentIndex > INDEX_NONE ? GetEndOfBranch(NewParentIndex) : EndOfBoneBranch+1;
+	const TArray<FElement>& Elements = Hierarchy.Elements;
+	const TArray<FTransform>& Transforms = Hierarchy.Transforms;
 
-	// reorder bones
-	TArray<FName> Names;
-	auto AddBones = [&](int32 Begin, int32 End)
-	{
-		for (int32 Index = Begin; Index < End; Index++)
-		{
-			Names.Add(RawRefBoneInfo[Index].Name);
-		}		
-	};
+	const int32 NumElements = Elements.Num();
+	TArray<FMeshBoneInfo> NewRawRefBoneInfo; NewRawRefBoneInfo.Reserve(NumElements);
+	TArray<FTransform> NewRawRefBonePose; NewRawRefBonePose.Reserve(NumElements);
+	TMap<FName, int32> NewNameToIndexMap; NewNameToIndexMap.Reserve(NumElements);
 
-	if (NewParentIndex == INDEX_NONE)
+	// recreate infos
+	for (int32 Index = 0; Index < NumElements; ++Index)
 	{
-		// add all bones before BoneIndex
-		AddBones(0, BoneIndex);
-		// add bones after branch
-		AddBones(EndOfBoneBranch, NumBones);
-		// add bone branch at the end
-		AddBones(BoneIndex, EndOfBoneBranch);
-	}
-	else if (NewParentIndex > BoneIndex)
-	{
-		// add all bones before BoneIndex
-		AddBones(0, BoneIndex);
-		// add parent branch up
-		AddBones(NewParentIndex, EndOfNewParentBranch);
-		// add bone branch
-		AddBones(BoneIndex, EndOfBoneBranch);
-		// add bones between bone branch and parent 
-		AddBones(EndOfBoneBranch, NewParentIndex);
-		// add the end
-		AddBones(EndOfNewParentBranch, NumBones);
-	}
-	else
-	{
-		// add all bones before the end of the new parent branch
-		if (EndOfNewParentBranch > BoneIndex)
+		const int32 NewIndex = NewToOldIndexes[Index];
+		const FElement& Element = Elements[NewIndex];
+		
+		// add new info
+		const int32 OldIndex = Element.RawIndex;
+		FMeshBoneInfo& NewBoneInfo = NewRawRefBoneInfo.Add_GetRef(Infos[OldIndex]);
+		
+		// update parent
+		int32 NewParentIdx = INDEX_NONE;
+		if (const FElement* NewParent = Element.Parent)
 		{
-			for (int32 Index = 0; Index < EndOfNewParentBranch; Index++)
+			NewParentIdx = NewToOldIndexes.IndexOfByPredicate([NewParent](int32 InIndex)
 			{
-				if (Index < BoneIndex || Index >= EndOfBoneBranch)
-				{
-					Names.Add(RawRefBoneInfo[Index].Name);
-				}
-			}
+				return NewParent->RawIndex == InIndex;
+			});
 		}
-		else
+		NewBoneInfo.ParentIndex = NewParentIdx; 
+
+		// update pose
+		FTransform& Pose = NewRawRefBonePose.Add_GetRef(Transforms[NewIndex]);
+		if (NewParentIdx != INDEX_NONE)
 		{
-			AddBones(0, EndOfNewParentBranch);
+			Pose = Transforms[NewIndex].GetRelativeTransform(Transforms[Element.Parent->RawIndex]);
 		}
-		// add bone branch
-		AddBones(BoneIndex, EndOfBoneBranch);
-		// add remaining bones
-		AddBones(EndOfNewParentBranch, BoneIndex);
-		AddBones(EndOfBoneBranch, NumBones);
-	}
-
-	check(Names.Num() == NumBones);
-	
-	TMap<int32, int32> OldToNewIndexes; OldToNewIndexes.Reserve(Names.Num());
-	TMap<FName, int32> NewNameToIndexMap; NewNameToIndexMap.Reserve(Names.Num());
-	TArray<FMeshBoneInfo> NewRawRefBoneInfo; NewRawRefBoneInfo.Reserve(Names.Num());
-	TArray<FTransform> NewRawRefBonePose; NewRawRefBonePose.Reserve(Names.Num());
 
-	// store new data
-	for (int32 Index = 0; Index < Names.Num(); ++Index)
-	{
-		const int32 OldIndex = FindRawBoneIndex(Names[Index]);
-		OldToNewIndexes.Add(OldIndex, Index);
-		
-		NewNameToIndexMap.Add(Names[Index], Index);
-		NewRawRefBoneInfo.Add(RawRefBoneInfo[OldIndex]);
-		NewRawRefBonePose.Add(RawRefBonePose[OldIndex]);
+		// update name to index map
+		NewNameToIndexMap.Add(NewBoneInfo.Name, Index);
 	}
 
-	// new bone index
-	const int32 NewBoneIndex = *NewNameToIndexMap.Find(InBoneName);
-	
-	// update parents
-	for (int32 Index = 0; Index < Names.Num(); ++Index)
-	{
-		FMeshBoneInfo& NewBoneInfo = NewRawRefBoneInfo[Index];
-		if (Index == NewBoneIndex)
-		{
-			NewBoneInfo.ParentIndex = NewParentIndex > INDEX_NONE ? *NewNameToIndexMap.Find(InParentName) : INDEX_NONE;
-		}
-		else if (NewBoneInfo.ParentIndex > INDEX_NONE)
-		{
-			NewBoneInfo.ParentIndex = *OldToNewIndexes.Find(NewBoneInfo.ParentIndex);
-		}
-	}
-	
 	// swap data
 	RawRefBonePose = MoveTemp(NewRawRefBonePose);
 	RawRefBoneInfo = MoveTemp(NewRawRefBoneInfo);
 	RawNameToIndexMap = MoveTemp(NewNameToIndexMap);
-	
-	// return new bone index
-	return NewBoneIndex;
+
+	return *RawNameToIndexMap.Find(InBoneName);
 }
 
 int32 FReferenceSkeleton::GetRawSourceBoneIndex(const USkeleton* Skeleton, const FName& SourceBoneName) const
