Fix for erroneously triggered check() in SVT streaming manager.

#rnx

diff --git a/Engine/Source/Runtime/Engine/Private/SparseVolumeTexture/SparseVolumeTextureStreamingManager.cpp b/Engine/Source/Runtime/Engine/Private/SparseVolumeTexture/SparseVolumeTextureStreamingManager.cpp
--- a/Engine/Source/Runtime/Engine/Private/SparseVolumeTexture/SparseVolumeTextureStreamingManager.cpp
+++ b/Engine/Source/Runtime/Engine/Private/SparseVolumeTexture/SparseVolumeTextureStreamingManager.cpp
@@ -1507,8 +1507,8 @@ void FStreamingManager::IssueRequests(int32 MaxSelectedRequests)
 			// Couldn't free enough tiles, so skip this mip level
 			if ((NumAvailableTiles + NumNewlyAvailableTiles) < NumRequiredTiles)
 			{
-				UE_LOG(LogSparseVolumeTextureStreamingManager, Warning, TEXT("IssueRequests() Frame %i Mip %i: Not enough tiles available (%i) to fit mip level (%i) even after freeing"),
-					SelectedKey.FrameIndex, SelectedKey.MipLevelIndex, (NumAvailableTiles + NumNewlyAvailableTiles), NumRequiredTiles);
+				UE_LOG(LogSparseVolumeTextureStreamingManager, Warning, TEXT("IssueRequests() SVT %p Frame %i Mip %i: Not enough tiles available (%i) to fit mip level (%i) even after freeing"),
+					SelectedKey.SVT, SelectedKey.FrameIndex, SelectedKey.MipLevelIndex, (NumAvailableTiles + NumNewlyAvailableTiles), NumRequiredTiles);
 				continue;
 			}
 		}
@@ -1516,7 +1516,7 @@ void FStreamingManager::IssueRequests(int32 MaxSelectedRequests)
 #if DO_CHECK
 		for (auto& Pending : PendingMipLevels)
 		{
-			check(Pending.FrameIndex != SelectedKey.FrameIndex || Pending.MipLevelIndex != SelectedKey.MipLevelIndex);
+			check(Pending.SparseVolumeTexture != SelectedKey.SVT || Pending.FrameIndex != SelectedKey.FrameIndex || Pending.MipLevelIndex != SelectedKey.MipLevelIndex);
 		}
 #endif
 
@@ -1765,7 +1765,7 @@ int32 FStreamingManager::DetermineReadyMipLevels()
 				{
 					// Retry if IO request failed for some reason
 					const FMipLevelStreamingInfo& MipLevelStreamingInfo = Resources->MipLevelStreamingInfo[PendingMipLevel.MipLevelIndex];
-					UE_LOG(LogSparseVolumeTextureStreamingManager, Warning, TEXT("SVT IO request failed for '%p' (frame %i, mip %i, offset %i, size %i). Retrying..."),
+					UE_LOG(LogSparseVolumeTextureStreamingManager, Warning, TEXT("SVT IO request failed for %p (frame %i, mip %i, offset %i, size %i). Retrying..."),
 						PendingMipLevel.SparseVolumeTexture, PendingMipLevel.FrameIndex, PendingMipLevel.MipLevelIndex, MipLevelStreamingInfo.BulkOffset, MipLevelStreamingInfo.BulkSize);
 					
 					FBulkDataBatchRequest::FBatchBuilder Batch = FBulkDataBatchRequest::NewBatch(1);
