Remove ensure as UUserWidget::AddToPlayerScreen must support to add more then one widget at the same layer.

diff --git a/Engine/Source/Runtime/Engine/Private/Slate/SGameLayerManager.cpp b/Engine/Source/Runtime/Engine/Private/Slate/SGameLayerManager.cpp
--- a/Engine/Source/Runtime/Engine/Private/Slate/SGameLayerManager.cpp
+++ b/Engine/Source/Runtime/Engine/Private/Slate/SGameLayerManager.cpp
@@ -280,8 +280,6 @@ void SGameLayerManager::AddWidgetForPlayer(ULocalPlayer* Player, TSharedRef<SWid
 {
 	TSharedPtr<FPlayerLayer> PlayerLayer = FindOrCreatePlayerLayer(Player);
 
-	ensure(!PlayerLayer->Widget->HasSlotWithZOrder(ZOrder));
-
 	PlayerLayer->Widget->AddSlot(ZOrder)
 	[
 		ViewportContent
