Fix a crash caused by the copy constructor using the Duplicate method.
The SubTrees array was not being initialized and Duplicate did not set a value if it was not set in the source quad tree.

diff --git a/Engine/Source/Runtime/Engine/Public/GenericQuadTree.h b/Engine/Source/Runtime/Engine/Public/GenericQuadTree.h
--- a/Engine/Source/Runtime/Engine/Public/GenericQuadTree.h
+++ b/Engine/Source/Runtime/Engine/Public/GenericQuadTree.h
@@ -97,7 +97,7 @@ private:
 	TArray<FNode> Nodes;
 
 	/** The sub-trees of this tree */
-	TreeType* SubTrees[4];
+	TreeType* SubTrees[4] = {nullptr}
 
 	/** AABB of the tree */
 	FBox2D TreeBox;
@@ -440,4 +440,4 @@ void TQuadTree<ElementType, NodeCapacity>::Empty()
 
 	Nodes.Empty();
 	bInternal = false;
-}
\ No newline at end of file
+}
