Switched AInfo's ObjectInitializer constructor over to supporting a default ObjectInitializer parameter.
Because AInfo doesn't override the default constructor, the previous setup resulted in some AInfo-specific initialization not getting run in some cases (e.g. overriding bIsSpatiallyLoaded to False).

#tests Verified fix in a Minimal Build with these shelved changes incorporated.

diff --git a/Engine/Source/Runtime/Engine/Classes/GameFramework/Info.h b/Engine/Source/Runtime/Engine/Classes/GameFramework/Info.h
--- a/Engine/Source/Runtime/Engine/Classes/GameFramework/Info.h
+++ b/Engine/Source/Runtime/Engine/Classes/GameFramework/Info.h
@@ -19,9 +19,11 @@
 UCLASS(abstract, hidecategories=(Input, Movement, Collision, Rendering, HLOD, WorldPartition, DataLayers, Transformation), showcategories=("Input|MouseInput", "Input|TouchInput"), MinimalAPI, NotBlueprintable)
 class AInfo : public AActor
 {
-	GENERATED_UCLASS_BODY()
+	GENERATED_BODY()
 
 public:
+	ENGINE_API AInfo(const FObjectInitializer& ObjectInitializer = FObjectInitializer::Get());
+
 #if WITH_EDITOR
 	ENGINE_API virtual void PostLoad() override;
 #endif
diff --git a/Engine/Source/Runtime/Engine/Private/Info.cpp b/Engine/Source/Runtime/Engine/Private/Info.cpp
--- a/Engine/Source/Runtime/Engine/Private/Info.cpp
+++ b/Engine/Source/Runtime/Engine/Private/Info.cpp
@@ -8,7 +8,7 @@
 
 #include UE_INLINE_GENERATED_CPP_BY_NAME(Info)
 
-AInfo::AInfo(const FObjectInitializer& ObjectInitializer)
+AInfo::AInfo(const FObjectInitializer& ObjectInitializer /*= FObjectInitializer::Get()*/)
 	: Super(ObjectInitializer)
 {
 #if WITH_EDITORONLY_DATA
