World Partition: properly handle OnEndPlay and OnCancelPIE events as the latter can be called alone, but during the normal flow they are both called.

diff --git a/Engine/Source/Runtime/Engine/Private/WorldPartition/WorldPartition.cpp b/Engine/Source/Runtime/Engine/Private/WorldPartition/WorldPartition.cpp
--- a/Engine/Source/Runtime/Engine/Private/WorldPartition/WorldPartition.cpp
+++ b/Engine/Source/Runtime/Engine/Private/WorldPartition/WorldPartition.cpp
@@ -375,22 +375,19 @@ bool UWorldPartition::IsValidPackageName(const FString& InPackageName)
 
 void UWorldPartition::OnPreBeginPIE(bool bStartSimulate)
 {
-	check(!bIsPIE);
-	bIsPIE = true;
-
 	OnBeginPlay();
 }
 
 void UWorldPartition::OnPrePIEEnded(bool bWasSimulatingInEditor)
 {
-	// In some specific cases, bIsPIE can be already false. An example is when the world is loaded in UEditorEngine::StartPlayInEditorSession after the PreBeginPIE event
-	// was fired, because of OverrideMapURL. In this case, streaming generation happens in UWorldPartition::Initialize because we detect it as a world travel, and everything
-	// works fine.
-	bIsPIE = false;
+	OnEndPlay();
 }
 
 void UWorldPartition::OnBeginPlay()
 {
+	check(!bIsPIE);
+	bIsPIE = true;
+
 	FGenerateStreamingParams Params;
 
 	TArray<FString> OutGeneratedStreamingPackageNames;
@@ -413,16 +410,19 @@ void UWorldPartition::OnBeginPlay()
 
 void UWorldPartition::OnCancelPIE()
 {
-	// No check here since CancelPIE can be called after PrePIEEnded
-	bIsPIE = false;
 	// Call OnEndPlay here since EndPlayMapDelegate is not called when cancelling PIE
 	OnEndPlay();
 }
 
 void UWorldPartition::OnEndPlay()
 {
-	FlushStreaming();
-	RuntimeHash->OnEndPlay();
+	// No check here since CancelPIE can be called after PrePIEEnded
+	if (bIsPIE)
+	{
+		FlushStreaming();
+		RuntimeHash->OnEndPlay();
+		bIsPIE = false;
+	}
 }
 
 bool UWorldPartition::CanEditChange(const FProperty* InProperty) const
@@ -681,18 +681,6 @@ void UWorldPartition::Initialize(UWorld* InWorld, const FTransform& InTransform)
 	{
 		if (bIsGame || bIsPIEWorldTravel || bIsDedicatedServer)
 		{
-			if (bIsPIEWorldTravel)
-			{
-				check(!bIsPIE);
-				bIsPIE = true;
-			}
-
-			if (StreamingPolicy)
-			{
-				UE_LOG(LogWorldPartition, Warning, TEXT("StreamingPolicy was set when initializing the world partition object"));
-				StreamingPolicy = nullptr;
-			}
-
 			OnBeginPlay();
 		}
 
