Add camera relative origin options for object, actor and particle nodes

diff --git a/Engine/Plugins/Enterprise/GLTFExporter/Source/GLTFExporter/Private/Converters/GLTFProxyMaterialCompiler.h b/Engine/Plugins/Enterprise/GLTFExporter/Source/GLTFExporter/Private/Converters/GLTFProxyMaterialCompiler.h
--- a/Engine/Plugins/Enterprise/GLTFExporter/Source/GLTFExporter/Private/Converters/GLTFProxyMaterialCompiler.h
+++ b/Engine/Plugins/Enterprise/GLTFExporter/Source/GLTFExporter/Private/Converters/GLTFProxyMaterialCompiler.h
@@ -38,16 +38,16 @@ public:
 
 	using FProxyMaterialCompiler::FProxyMaterialCompiler;
 
-	virtual int32 ActorWorldPosition() override
+	virtual int32 ActorWorldPosition(EPositionOrigin OriginType) override
 	{
 		bUsesActorPosition = true;
-		return Compiler->ActorWorldPosition();
+		return Compiler->ActorWorldPosition(OriginType);
 	}
 
-	virtual int32 ObjectWorldPosition() override
+	virtual int32 ObjectWorldPosition(EPositionOrigin OriginType) override
 	{
 		bUsesObjectPosition = true;
-		return Compiler->ObjectWorldPosition();
+		return Compiler->ObjectWorldPosition(OriginType);
 	}
 
 	virtual int32 ObjectOrientation() override
diff --git a/Engine/Shaders/Private/MaterialTemplate.ush b/Engine/Shaders/Private/MaterialTemplate.ush
--- a/Engine/Shaders/Private/MaterialTemplate.ush
+++ b/Engine/Shaders/Private/MaterialTemplate.ush
@@ -1137,6 +1137,16 @@ MaterialFloat2 GetDefaultSceneTextureUV(FMaterialPixelParameters Parameters, con
 		return GetPrimitiveData(Parameters).ActorWorldPosition;
 	#endif
 	}
+
+	float3 GetActorTranslatedWorldPosition(FMaterialVertexParameters Parameters)
+	{
+		return LWCToFloat(LWCAdd(GetActorWorldPosition(Parameters), ResolvedView.PreViewTranslation));
+	}
+
+	float3 GetActorTranslatedWorldPosition(FMaterialPixelParameters Parameters)
+	{
+		return LWCToFloat(LWCAdd(GetActorWorldPosition(Parameters), ResolvedView.PreViewTranslation));
+	}
 	
 	float3 GetObjectOrientation(FMaterialVertexParameters Parameters)
 	{
@@ -1377,6 +1387,16 @@ FLWCVector3 GetObjectWorldPosition(FMaterialVertexParameters Parameters)
 #endif
 }
 
+float3 GetObjectTranslatedWorldPosition(FMaterialPixelParameters Parameters)
+{
+	return LWCToFloat(LWCAdd(GetObjectWorldPosition(Parameters), ResolvedView.PreViewTranslation));
+}
+
+float3 GetObjectTranslatedWorldPosition(FMaterialVertexParameters Parameters)
+{
+	return LWCToFloat(LWCAdd(GetObjectWorldPosition(Parameters), ResolvedView.PreViewTranslation));
+}
+
 /** Get the per-instance random value when instancing */
 float GetPerInstanceRandom(FMaterialVertexParameters Parameters)
 {
diff --git a/Engine/Shaders/Public/ShaderVersion.ush b/Engine/Shaders/Public/ShaderVersion.ush
--- a/Engine/Shaders/Public/ShaderVersion.ush
+++ b/Engine/Shaders/Public/ShaderVersion.ush
@@ -2,5 +2,5 @@
 // Each time the console command is executed it generates a new GUID. As this file is included
 // in Platform.ush (which should be included in any shader) it allows to invalidate the shader DDC.
 // 
-// GUID = F487A842-1225-4FD4-A1AF-17743B6937E0
+// GUID = 7CC2CDB6-7689-4413-A8FD-E89BB19E2EDD
 
diff --git a/Engine/Source/Developer/MaterialUtilities/Private/MaterialUtilities.cpp b/Engine/Source/Developer/MaterialUtilities/Private/MaterialUtilities.cpp
--- a/Engine/Source/Developer/MaterialUtilities/Private/MaterialUtilities.cpp
+++ b/Engine/Source/Developer/MaterialUtilities/Private/MaterialUtilities.cpp
@@ -259,9 +259,9 @@ struct FExportMaterialCompiler : public FProxyMaterialCompiler
 #endif
 	}
 
-	virtual int32 ObjectWorldPosition() override
+	virtual int32 ObjectWorldPosition(EPositionOrigin OriginType) override
 	{
-		return Compiler->ObjectWorldPosition();
+		return Compiler->ObjectWorldPosition(OriginType);
 	}
 
 	virtual int32 DistanceCullFade() override
@@ -269,9 +269,9 @@ struct FExportMaterialCompiler : public FProxyMaterialCompiler
 		return Compiler->Constant(1.0f);
 	}
 
-	virtual int32 ActorWorldPosition() override
+	virtual int32 ActorWorldPosition(EPositionOrigin OriginType) override
 	{
-		return Compiler->ActorWorldPosition();
+		return Compiler->ActorWorldPosition(OriginType);
 	}
 
 	virtual int32 ParticleRelativeTime() override
diff --git a/Engine/Source/Editor/UnrealEd/Private/Lightmass/LightmassRender.cpp b/Engine/Source/Editor/UnrealEd/Private/Lightmass/LightmassRender.cpp
--- a/Engine/Source/Editor/UnrealEd/Private/Lightmass/LightmassRender.cpp
+++ b/Engine/Source/Editor/UnrealEd/Private/Lightmass/LightmassRender.cpp
@@ -120,7 +120,7 @@ struct FLightmassMaterialCompiler : public FProxyMaterialCompiler
 		return Compiler->Constant3(0.0f,0.0f,0.0f);
 	}
 
-	virtual int32 ObjectWorldPosition() override
+	virtual int32 ObjectWorldPosition(EPositionOrigin OriginType) override
 	{
 		//UE_LOG(LogLightmassRender, Log, TEXT("Lightmass material compiler has encountered ObjectWorldPosition... Forcing constant (0.0f,0.0f,0.0f)."));
 		return Compiler->Constant3(0.0f,0.0f,0.0f);
@@ -148,7 +148,7 @@ struct FLightmassMaterialCompiler : public FProxyMaterialCompiler
 		return Compiler->Constant(1.0f);
 	}
 
-	virtual int32 ActorWorldPosition() override
+	virtual int32 ActorWorldPosition(EPositionOrigin OriginType) override
 	{
 		return Compiler->Constant3(0.0f,0.0f,0.0f);
 	}
diff --git a/Engine/Source/Runtime/Engine/Classes/Materials/MaterialExpression.h b/Engine/Source/Runtime/Engine/Classes/Materials/MaterialExpression.h
--- a/Engine/Source/Runtime/Engine/Classes/Materials/MaterialExpression.h
+++ b/Engine/Source/Runtime/Engine/Classes/Materials/MaterialExpression.h
@@ -600,4 +600,13 @@ private:
 #endif // WITH_EDITOR
 };
 
+/** Specifies what reference point should be used */
+UENUM()
+enum class EPositionOrigin : uint8
+{
+	/** Absolute world position, i.e. relative to (0,0,0) */
+	Absolute UMETA(DisplayName="Absolute World Position"),
 
+	/** Camera relative world position */
+	CameraRelative UMETA(DisplayName="Camera Relative World Position")
+};
diff --git a/Engine/Source/Runtime/Engine/Classes/Materials/MaterialExpressionActorPositionWS.h b/Engine/Source/Runtime/Engine/Classes/Materials/MaterialExpressionActorPositionWS.h
--- a/Engine/Source/Runtime/Engine/Classes/Materials/MaterialExpressionActorPositionWS.h
+++ b/Engine/Source/Runtime/Engine/Classes/Materials/MaterialExpressionActorPositionWS.h
@@ -13,6 +13,8 @@ class UMaterialExpressionActorPositionWS : public UMaterialExpression
 {
 	GENERATED_UCLASS_BODY()
 
+	UPROPERTY(EditAnywhere, Category=UMaterialExpressionActorPositionWS, meta=(DisplayName = "Origin", ShowAsInputPin = "Advanced"))
+	EPositionOrigin OriginType;
 
 	//~ Begin UMaterialExpression Interface
 #if WITH_EDITOR
diff --git a/Engine/Source/Runtime/Engine/Classes/Materials/MaterialExpressionObjectPositionWS.h b/Engine/Source/Runtime/Engine/Classes/Materials/MaterialExpressionObjectPositionWS.h
--- a/Engine/Source/Runtime/Engine/Classes/Materials/MaterialExpressionObjectPositionWS.h
+++ b/Engine/Source/Runtime/Engine/Classes/Materials/MaterialExpressionObjectPositionWS.h
@@ -13,6 +13,8 @@ class UMaterialExpressionObjectPositionWS : public UMaterialExpression
 {
 	GENERATED_UCLASS_BODY()
 
+	UPROPERTY(EditAnywhere, Category=UMaterialExpressionObjectPositionWS, meta=(DisplayName = "Origin", ShowAsInputPin = "Advanced"))
+	EPositionOrigin OriginType;
 
 	//~ Begin UMaterialExpression Interface
 #if WITH_EDITOR
diff --git a/Engine/Source/Runtime/Engine/Classes/Materials/MaterialExpressionParticlePositionWS.h b/Engine/Source/Runtime/Engine/Classes/Materials/MaterialExpressionParticlePositionWS.h
--- a/Engine/Source/Runtime/Engine/Classes/Materials/MaterialExpressionParticlePositionWS.h
+++ b/Engine/Source/Runtime/Engine/Classes/Materials/MaterialExpressionParticlePositionWS.h
@@ -13,6 +13,8 @@ class UMaterialExpressionParticlePositionWS : public UMaterialExpression
 {
 	GENERATED_UCLASS_BODY()
 
+	UPROPERTY(EditAnywhere, Category=UMaterialExpressionParticlePositionWS, meta=(DisplayName = "Origin", ShowAsInputPin = "Advanced"))
+	EPositionOrigin OriginType;
 
 	//~ Begin UMaterialExpression Interface
 #if WITH_EDITOR
diff --git a/Engine/Source/Runtime/Engine/Private/Materials/HLSLMaterialTranslator.cpp b/Engine/Source/Runtime/Engine/Private/Materials/HLSLMaterialTranslator.cpp
--- a/Engine/Source/Runtime/Engine/Private/Materials/HLSLMaterialTranslator.cpp
+++ b/Engine/Source/Runtime/Engine/Private/Materials/HLSLMaterialTranslator.cpp
@@ -5665,16 +5665,23 @@ int32 FHLSLMaterialTranslator::ParticleColor()
 	return AddInlinedCodeChunkZeroDeriv(MCT_Float4,TEXT("Parameters.Particle.Color"));
 }
 
-int32 FHLSLMaterialTranslator::ParticlePosition()
+int32 FHLSLMaterialTranslator::ParticlePosition(EPositionOrigin OriginType)
 {
 	if (ShaderFrequency != SF_Vertex && ShaderFrequency != SF_Pixel && ShaderFrequency != SF_Compute)
 	{
 		return NonVertexOrPixelShaderExpressionError();
 	}
 	bNeedsParticlePosition = true;
-	AddLWCFuncUsage(ELWCFunctionKind::Subtract);
-	const int32 Result = AddInlinedCodeChunkZeroDeriv(MCT_LWCVector3,TEXT("LWCSubtract(Parameters.Particle.TranslatedWorldPositionAndSize.xyz, ResolvedView.PreViewTranslation)"));
-	return CastToNonLWCIfDisabled(Result);
+	if (OriginType == EPositionOrigin::CameraRelative)
+	{
+		return AddInlinedCodeChunkZeroDeriv(MCT_Float3, TEXT("Parameters.Particle.TranslatedWorldPositionAndSize.xyz"));
+	}
+	else
+	{
+		AddLWCFuncUsage(ELWCFunctionKind::Subtract);
+		const int32 Result = AddInlinedCodeChunkZeroDeriv(MCT_LWCVector3,TEXT("LWCSubtract(Parameters.Particle.TranslatedWorldPositionAndSize.xyz, ResolvedView.PreViewTranslation)"));
+		return CastToNonLWCIfDisabled(Result);
+	}
 }
 
 int32 FHLSLMaterialTranslator::ParticleRadius()
@@ -5855,10 +5862,17 @@ int32 FHLSLMaterialTranslator::WorldPosition(EWorldPositionIncludedOffsets World
 	return CastToNonLWCIfDisabled(Result);
 }
 
-int32 FHLSLMaterialTranslator::ObjectWorldPosition()
+int32 FHLSLMaterialTranslator::ObjectWorldPosition(EPositionOrigin OriginType)
 {
-	const int32 Result = AddInlinedCodeChunkZeroDeriv(MCT_LWCVector3,TEXT("GetObjectWorldPosition(Parameters)"));
-	return CastToNonLWCIfDisabled(Result);
+	if (OriginType == EPositionOrigin::CameraRelative)
+	{
+		return AddInlinedCodeChunkZeroDeriv(MCT_Float3,TEXT("GetObjectTranslatedWorldPosition(Parameters)"));
+	}
+	else
+	{
+		const int32 Result = AddInlinedCodeChunkZeroDeriv(MCT_LWCVector3,TEXT("GetObjectWorldPosition(Parameters)"));
+		return CastToNonLWCIfDisabled(Result);
+	}
 }
 
 int32 FHLSLMaterialTranslator::ObjectRadius()
@@ -5916,25 +5930,45 @@ int32 FHLSLMaterialTranslator::DistanceCullFade()
 	return AddInlinedCodeChunk(MCT_Float,TEXT("GetDistanceCullFade()"));		
 }
 
-int32 FHLSLMaterialTranslator::ActorWorldPosition()
+int32 FHLSLMaterialTranslator::ActorWorldPosition(EPositionOrigin OriginType)
 {
-	int32 Result = INDEX_NONE;
-	if (bCompilingPreviousFrame && ShaderFrequency == SF_Vertex)
+	if (OriginType == EPositionOrigin::CameraRelative)
 	{
-		// Decal VS doesn't have material code so FMaterialVertexParameters
-		// and primitve uniform buffer are guaranteed to exist if ActorPosition
-		// material node is used in VS
-		AddLWCFuncUsage(ELWCFunctionKind::MultiplyVectorMatrix, 2);
-		Result = AddInlinedCodeChunkZeroDeriv(
-			MCT_LWCVector3,
-			TEXT("LWCMultiply(LWCMultiply(GetActorWorldPosition(Parameters), GetWorldToInstance(Parameters)), Parameters.PrevFrameLocalToWorld)"));
+		if (bCompilingPreviousFrame && ShaderFrequency == SF_Vertex)
+		{
+			// Decal VS doesn't have material code so FMaterialVertexParameters
+			// and primitve uniform buffer are guaranteed to exist if ActorPosition
+			// material node is used in VS
+			AddLWCFuncUsage(ELWCFunctionKind::MultiplyVectorMatrix, 2);
+			return AddInlinedCodeChunkZeroDeriv(
+				MCT_Float3,
+				TEXT("LWCToFloat(LWCSubtract(LWCMultiply(LWCMultiply(GetActorWorldPosition(Parameters), GetWorldToInstance(Parameters)), Parameters.PrevFrameLocalToWorld), ResolvedView.PreViewTranslation))"));
+		}
+		else
+		{
+			return AddInlinedCodeChunkZeroDeriv(MCT_Float3, TEXT("GetActorTranslatedWorldPosition(Parameters)"));
+		}
 	}
 	else
 	{
-		Result = AddInlinedCodeChunkZeroDeriv(MCT_LWCVector3, TEXT("GetActorWorldPosition(Parameters)"));
-	}
+		int32 Result = INDEX_NONE;
+		if (bCompilingPreviousFrame && ShaderFrequency == SF_Vertex)
+		{
+			// Decal VS doesn't have material code so FMaterialVertexParameters
+			// and primitive uniform buffer are guaranteed to exist if ActorPosition
+			// material node is used in VS
+			AddLWCFuncUsage(ELWCFunctionKind::MultiplyVectorMatrix, 2);
+			Result = AddInlinedCodeChunkZeroDeriv(
+				MCT_LWCVector3,
+				TEXT("LWCMultiply(LWCMultiply(GetActorWorldPosition(Parameters), GetWorldToInstance(Parameters)), Parameters.PrevFrameLocalToWorld)"));
+		}
+		else
+		{
+			Result = AddInlinedCodeChunkZeroDeriv(MCT_LWCVector3, TEXT("GetActorWorldPosition(Parameters)"));
+		}
 
-	return CastToNonLWCIfDisabled(Result);
+		return CastToNonLWCIfDisabled(Result);
+	}
 }
 
 int32 FHLSLMaterialTranslator::DynamicBranch(int32 Condition, int32 A, int32 B)
diff --git a/Engine/Source/Runtime/Engine/Private/Materials/HLSLMaterialTranslator.h b/Engine/Source/Runtime/Engine/Private/Materials/HLSLMaterialTranslator.h
--- a/Engine/Source/Runtime/Engine/Private/Materials/HLSLMaterialTranslator.h
+++ b/Engine/Source/Runtime/Engine/Private/Materials/HLSLMaterialTranslator.h
@@ -888,7 +888,7 @@ protected:
 	virtual int32 ParticleSubUV(int32 TextureIndex, EMaterialSamplerType SamplerType, int32 MipValue0Index, int32 MipValue1Index, ETextureMipValueMode MipValueMode, bool bBlend) override;
 	virtual int32 ParticleSubUVProperty(int32 PropertyIndex) override;
 	virtual int32 ParticleColor() override;
-	virtual int32 ParticlePosition() override;
+	virtual int32 ParticlePosition(EPositionOrigin OriginType) override;
 	virtual int32 ParticleRadius() override;
 
 	virtual int32 SphericalParticleOpacity(int32 Density) override;
@@ -903,7 +903,7 @@ protected:
 
 	virtual int32 WorldPosition(EWorldPositionIncludedOffsets WorldPositionIncludedOffsets) override;
 
-	virtual int32 ObjectWorldPosition() override;
+	virtual int32 ObjectWorldPosition(EPositionOrigin OriginType) override;
 	virtual int32 ObjectRadius() override;
 	virtual int32 ObjectBounds() override;
 
@@ -912,7 +912,7 @@ protected:
 
 	virtual int32 DistanceCullFade() override;
 
-	virtual int32 ActorWorldPosition() override;
+	virtual int32 ActorWorldPosition(EPositionOrigin OriginType) override;
 
 	virtual int32 DynamicBranch(int32 Condition, int32 A, int32 B) override;
 	virtual int32 If(int32 A, int32 B, int32 AGreaterThanB, int32 AEqualsB, int32 ALessThanB, int32 ThresholdArg) override;
diff --git a/Engine/Source/Runtime/Engine/Private/Materials/MaterialExpressions.cpp b/Engine/Source/Runtime/Engine/Private/Materials/MaterialExpressions.cpp
--- a/Engine/Source/Runtime/Engine/Private/Materials/MaterialExpressions.cpp
+++ b/Engine/Source/Runtime/Engine/Private/Materials/MaterialExpressions.cpp
@@ -10145,12 +10145,31 @@ UMaterialExpressionParticlePositionWS::UMaterialExpressionParticlePositionWS(con
 #if WITH_EDITOR
 int32 UMaterialExpressionParticlePositionWS::Compile(class FMaterialCompiler* Compiler, int32 OutputIndex)
 {
-	return Compiler->ParticlePosition();
+	return Compiler->ParticlePosition(OriginType);
 }
 
 void UMaterialExpressionParticlePositionWS::GetCaption(TArray<FString>& OutCaptions) const
 {
-	OutCaptions.Add(TEXT("Particle Position"));
+	switch (OriginType)
+	{
+		case EPositionOrigin::Absolute:
+		{
+			OutCaptions.Add(TEXT("Particle Position (Absolute)"));
+			break;
+		}
+
+		case EPositionOrigin::CameraRelative:
+		{
+			OutCaptions.Add(TEXT("Particle Position (Camera Relative)"));
+			break;
+		}
+
+		default:
+		{
+			UE_LOG(LogMaterial, Fatal, TEXT("Unknown position origin type"));
+			break;
+		}
+	}
 }
 #endif // WITH_EDITOR
 
@@ -12625,12 +12644,31 @@ UMaterialExpressionObjectPositionWS::UMaterialExpressionObjectPositionWS(const F
 #if WITH_EDITOR
 int32 UMaterialExpressionObjectPositionWS::Compile(class FMaterialCompiler* Compiler, int32 OutputIndex)
 {
-	return Compiler->ObjectWorldPosition();
+	return Compiler->ObjectWorldPosition(OriginType);
 }
 
 void UMaterialExpressionObjectPositionWS::GetCaption(TArray<FString>& OutCaptions) const
 {
-	OutCaptions.Add(TEXT("Object Position"));
+	switch (OriginType)
+	{
+		case EPositionOrigin::Absolute:
+		{
+			OutCaptions.Add(TEXT("Object Position (Absolute)"));
+			break;
+		}
+
+		case EPositionOrigin::CameraRelative:
+		{
+			OutCaptions.Add(TEXT("Object Position (Camera Relative)"));
+			break;
+		}
+
+		default:
+		{
+			UE_LOG(LogMaterial, Fatal, TEXT("Unknown position origin type"));
+			break;
+		}
+	}
 }
 #endif // WITH_EDITOR
 
@@ -12995,12 +13033,31 @@ int32 UMaterialExpressionActorPositionWS::Compile(class FMaterialCompiler* Compi
 		return CompilerError(Compiler, TEXT("Expression only available in the Surface and Deferred Decal material domains."));
 	}
 
-	return Compiler->ActorWorldPosition();
+	return Compiler->ActorWorldPosition(OriginType);
 }
 
 void UMaterialExpressionActorPositionWS::GetCaption(TArray<FString>& OutCaptions) const
 {
-	OutCaptions.Add(TEXT("Actor Position"));
+	switch (OriginType)
+	{
+		case EPositionOrigin::Absolute:
+		{
+			OutCaptions.Add(TEXT("Actor Position (Absolute)"));
+			break;
+		}
+
+		case EPositionOrigin::CameraRelative:
+		{
+			OutCaptions.Add(TEXT("Actor Position (Camera Relative)"));
+			break;
+		}
+
+		default:
+		{
+			UE_LOG(LogMaterial, Fatal, TEXT("Unknown position origin type"));
+			break;
+		}
+	}
 }
 #endif // WITH_EDITOR
 
diff --git a/Engine/Source/Runtime/Engine/Public/MaterialCompiler.h b/Engine/Source/Runtime/Engine/Public/MaterialCompiler.h
--- a/Engine/Source/Runtime/Engine/Public/MaterialCompiler.h
+++ b/Engine/Source/Runtime/Engine/Public/MaterialCompiler.h
@@ -232,18 +232,21 @@ public:
 	virtual int32 GetViewportUV() = 0;
 	virtual int32 GetPixelPosition() = 0;
 	virtual int32 WorldPosition(EWorldPositionIncludedOffsets WorldPositionIncludedOffsets) = 0;
-	virtual int32 ObjectWorldPosition() = 0;
+	virtual int32 ObjectWorldPosition(EPositionOrigin OriginType) = 0;
+	int32 ObjectWorldPosition() { return ObjectWorldPosition(EPositionOrigin::Absolute); }
 	virtual int32 ObjectRadius() = 0;
 	virtual int32 ObjectBounds() = 0;
 	virtual int32 ObjectLocalBounds(int32 OutputIndex) = 0;
 	virtual int32 PreSkinnedLocalBounds(int32 OutputIndex) = 0;
 	virtual int32 DistanceCullFade() = 0;
-	virtual int32 ActorWorldPosition() = 0;
+	virtual int32 ActorWorldPosition(EPositionOrigin OriginType) = 0;
+	int32 ActorWorldPosition() { return ActorWorldPosition(EPositionOrigin::Absolute); }
 	virtual int32 ParticleMacroUV() = 0;
 	virtual int32 ParticleSubUV(int32 TextureIndex, EMaterialSamplerType SamplerType, int32 MipValue0Index, int32 MipValue1Index, ETextureMipValueMode MipValueMode, bool bBlend) = 0;
 	virtual int32 ParticleSubUVProperty(int32 PropertyIndex) = 0;
 	virtual int32 ParticleColor() = 0;
-	virtual int32 ParticlePosition() = 0;
+	virtual int32 ParticlePosition(EPositionOrigin OriginType) = 0;
+	int32 ParticlePosition() { return ParticlePosition(EPositionOrigin::Absolute); }
 	virtual int32 ParticleRadius() = 0;
 	virtual int32 SphericalParticleOpacity(int32 Density) = 0;
 	virtual int32 ParticleRelativeTime() = 0;
@@ -773,13 +776,13 @@ public:
 	virtual int32 GetViewportUV() override { return Compiler->GetViewportUV(); }
 	virtual int32 GetPixelPosition() override { return Compiler->GetPixelPosition(); }
 	virtual int32 WorldPosition(EWorldPositionIncludedOffsets WorldPositionIncludedOffsets) override { return Compiler->WorldPosition(WorldPositionIncludedOffsets); }
-	virtual int32 ObjectWorldPosition() override { return Compiler->ObjectWorldPosition(); }
+	virtual int32 ObjectWorldPosition(EPositionOrigin OriginType) override { return Compiler->ObjectWorldPosition(OriginType); }
 	virtual int32 ObjectRadius() override { return Compiler->ObjectRadius(); }
 	virtual int32 ObjectBounds() override { return Compiler->ObjectBounds(); }
 	virtual int32 ObjectLocalBounds(int32 OutputIndex) override { return Compiler->ObjectLocalBounds(OutputIndex); }
 	virtual int32 PreSkinnedLocalBounds(int32 OutputIndex) override { return Compiler->PreSkinnedLocalBounds(OutputIndex); }
 	virtual int32 DistanceCullFade() override { return Compiler->DistanceCullFade(); }
-	virtual int32 ActorWorldPosition() override { return Compiler->ActorWorldPosition(); }
+	virtual int32 ActorWorldPosition(EPositionOrigin OriginType) override { return Compiler->ActorWorldPosition(OriginType); }
 	virtual int32 ParticleMacroUV() override { return Compiler->ParticleMacroUV(); }
 	virtual int32 ParticleSubUV(int32 TextureIndex, EMaterialSamplerType SamplerType, int32 MipValue0Index, int32 MipValue1Index, ETextureMipValueMode MipValueMode, bool bBlend) override 
 	{
@@ -787,7 +790,7 @@ public:
 	}
 	virtual int32 ParticleSubUVProperty(int32 PropertyIndex) override { return Compiler->ParticleSubUVProperty(PropertyIndex); }
 	virtual int32 ParticleColor() override { return Compiler->ParticleColor(); }
-	virtual int32 ParticlePosition() override { return Compiler->ParticlePosition(); }
+	virtual int32 ParticlePosition(EPositionOrigin OriginType) override { return Compiler->ParticlePosition(OriginType); }
 	virtual int32 ParticleRadius() override { return Compiler->ParticleRadius(); }
 	virtual int32 SphericalParticleOpacity(int32 Density) override { return Compiler->SphericalParticleOpacity(Density); }
 
