World Partition - HLOD: Turn EHLODBatchingPolicy::MeshSection warning into a display log, added more details

diff --git a/Engine/Source/Runtime/Engine/Private/WorldPartition/HLOD/HLODBuilder.cpp b/Engine/Source/Runtime/Engine/Private/WorldPartition/HLOD/HLODBuilder.cpp
--- a/Engine/Source/Runtime/Engine/Private/WorldPartition/HLOD/HLODBuilder.cpp
+++ b/Engine/Source/Runtime/Engine/Private/WorldPartition/HLOD/HLODBuilder.cpp
@@ -240,9 +240,16 @@ static bool ShouldBatchComponent(UActorComponent* ActorComponent)
 			bShouldBatch = true;
 			break;
 		case EHLODBatchingPolicy::MeshSection:
+		{
 			bShouldBatch = true;
-			UE_LOG(LogHLODBuilder, Warning, TEXT("EHLODBatchingPolicy::MeshSection is not yet supported by the HLOD builder, falling back to EHLODBatchingPolicy::Instancing for component %s (from actor %s)."), *ActorComponent->GetName(), *ActorComponent->GetOwner()->GetActorLabel());
+			FString LogDetails = FString::Printf(TEXT("%s %s (from actor %s)"), *PrimitiveComponent->GetClass()->GetName(), *ActorComponent->GetName(), *ActorComponent->GetOwner()->GetActorLabel());
+			if (UStaticMeshComponent* SMComponent = Cast<UStaticMeshComponent>(PrimitiveComponent))
+			{
+				LogDetails += FString::Printf(TEXT(" using static mesh %s"), SMComponent->GetStaticMesh() ? *SMComponent->GetStaticMesh()->GetName() : TEXT("<null>"));
+			}
+			UE_LOG(LogHLODBuilder, Display, TEXT("EHLODBatchingPolicy::MeshSection is not yet supported by the HLOD builder, falling back to EHLODBatchingPolicy::Instancing for %s."), *LogDetails);
 			break;
+		}
 		default:
 			checkNoEntry();
 		}
