Fix duplicated rendering LWI instances.

diff --git a/Engine/Source/Runtime/Engine/Classes/GameFramework/LightWeightInstanceStaticMeshManager.h b/Engine/Source/Runtime/Engine/Classes/GameFramework/LightWeightInstanceStaticMeshManager.h
--- a/Engine/Source/Runtime/Engine/Classes/GameFramework/LightWeightInstanceStaticMeshManager.h
+++ b/Engine/Source/Runtime/Engine/Classes/GameFramework/LightWeightInstanceStaticMeshManager.h
@@ -31,7 +31,6 @@ class ENGINE_API ALightWeightInstanceStaticMeshManager : public ALightWeightInst
 protected:
 	virtual void AddNewInstanceAt(FLWIData* InitData, int32 Index) override;
 	virtual void RemoveInstance(int32 Index) override;
-	virtual void UpdateDataAtIndex(FLWIData* InData, int32 Index) override;
 
 	void AddInstanceToRendering(int32 DataIndex);
 	void RemoveInstanceFromRendering(int32 DataIndex);
diff --git a/Engine/Source/Runtime/Engine/Private/LightWeightInstanceStaticMeshManager.cpp b/Engine/Source/Runtime/Engine/Private/LightWeightInstanceStaticMeshManager.cpp
--- a/Engine/Source/Runtime/Engine/Private/LightWeightInstanceStaticMeshManager.cpp
+++ b/Engine/Source/Runtime/Engine/Private/LightWeightInstanceStaticMeshManager.cpp
@@ -78,7 +78,6 @@ int32 ALightWeightInstanceStaticMeshManager::ConvertLightWeightIndexToCollisionI
 void ALightWeightInstanceStaticMeshManager::AddNewInstanceAt(FLWIData* InitData, int32 Index)
 {
 	Super::AddNewInstanceAt(InitData, Index);
-
 	AddInstanceToRendering(Index);
 }
 
@@ -92,18 +91,13 @@ void ALightWeightInstanceStaticMeshManager::RemoveInstance(const int32 Index)
 	Super::RemoveInstance(Index);
 }
 
-void ALightWeightInstanceStaticMeshManager::UpdateDataAtIndex(FLWIData* InData, int32 Index)
+void ALightWeightInstanceStaticMeshManager::AddInstanceToRendering(int32 DataIndex)
 {
-	Super::UpdateDataAtIndex(InData, Index);
-
-	if (!RenderingIndicesToDataIndices.Contains(Index))
+	if (!ensureMsgf(RenderingIndicesToDataIndices.Contains(DataIndex) == false, TEXT("LWI rendering instance added more than once. Index: %d"), DataIndex))
 	{
-		AddInstanceToRendering(Index);
+		return;
 	}
-}
 
-void ALightWeightInstanceStaticMeshManager::AddInstanceToRendering(int32 DataIndex)
-{
 	//cancel any pending deletes
 	const bool bPendingDeleteCancelled = DataIndicesToBeDeleted.RemoveSingle(DataIndex) > 0;
 
