Commented animation syncing code. Additional fixes for sync marker code.

- Added lots of comments accross syncing code to make it easier to understand.
- Added warning when attempting to collect sync markers in a UAnimMontage with an invalid slot track index.
- Organized FAnimSync::TickAssetPlayerInstances() to be more readable since thats where a good chunk of syncing is happening.
- Made FMarkerSyncData::GetMarkerSyncPositionFromMarkerIndicies() get fixes done to UAnimSequence::GetMarkerSyncPositionFromMarkerIndicies(). This a double source of truth issue that will be addressed by making UAnimSequence use the FMarkerSyncData like the UAnimMontage is doing already.
- Fixed bCanUseMarkerSync always being true even when animation asset did not have any sync markers. Causing the asset to not use length-based syncing.
- Fixed bCanUseMarkerSync being truefor followers when all valid markers were filtered out.
- Ensured MarkerIndexSpecialValues enum is used to avoid having magic numbers all over the place. Also fixed type on one of its enum values.
- Removed WIP comment from SyncSlotIndex, and added proper comment for it.
- Removed old TODO.

diff --git a/Engine/Source/Runtime/Engine/Classes/Animation/AnimMontage.h b/Engine/Source/Runtime/Engine/Classes/Animation/AnimMontage.h
--- a/Engine/Source/Runtime/Engine/Classes/Animation/AnimMontage.h
+++ b/Engine/Source/Runtime/Engine/Classes/Animation/AnimMontage.h
@@ -662,7 +662,7 @@ class UAnimMontage : public UAnimCompositeBase
 	UPROPERTY(EditAnywhere, Category = SyncGroup)
 	FName SyncGroup;
 
-	/** wip: until we have UI working */
+	/** Index of the slot track used for collecting sync markers */
 	UPROPERTY(EditAnywhere, Category = SyncGroup)
 	int32 SyncSlotIndex;
 
diff --git a/Engine/Source/Runtime/Engine/Classes/Animation/AnimationAsset.h b/Engine/Source/Runtime/Engine/Classes/Animation/AnimationAsset.h
--- a/Engine/Source/Runtime/Engine/Classes/Animation/AnimationAsset.h
+++ b/Engine/Source/Runtime/Engine/Classes/Animation/AnimationAsset.h
@@ -46,7 +46,7 @@ namespace MarkerIndexSpecialValues
 {
 	enum Type
 	{
-		Unitialized = -2,
+		Uninitialized = -2,
 		AnimationBoundary = -1,
 	};
 };
@@ -56,9 +56,9 @@ struct FMarkerPair
 	int32 MarkerIndex;
 	float TimeToMarker;
 
-	FMarkerPair() : MarkerIndex(MarkerIndexSpecialValues::Unitialized) {}
+	FMarkerPair() : MarkerIndex(MarkerIndexSpecialValues::Uninitialized) {}
 
-	void Reset() { MarkerIndex = MarkerIndexSpecialValues::Unitialized; }
+	void Reset() { MarkerIndex = MarkerIndexSpecialValues::Uninitialized; }
 };
 
 struct FMarkerTickRecord
@@ -69,7 +69,7 @@ struct FMarkerTickRecord
 
 	bool IsValid(bool bLooping) const
 	{
-		int32 Threshold = bLooping ? MarkerIndexSpecialValues::AnimationBoundary : MarkerIndexSpecialValues::Unitialized;
+		int32 Threshold = bLooping ? MarkerIndexSpecialValues::AnimationBoundary : MarkerIndexSpecialValues::Uninitialized;
 		return PreviousMarker.MarkerIndex > Threshold && NextMarker.MarkerIndex > Threshold;
 	}
 
diff --git a/Engine/Source/Runtime/Engine/Private/Animation/AnimMontage.cpp b/Engine/Source/Runtime/Engine/Private/Animation/AnimMontage.cpp
--- a/Engine/Source/Runtime/Engine/Private/Animation/AnimMontage.cpp
+++ b/Engine/Source/Runtime/Engine/Private/Animation/AnimMontage.cpp
@@ -1191,13 +1191,14 @@ void UAnimMontage::TickAssetPlayer(FAnimTickRecord& Instance, struct FAnimNotify
 	// nothing has to happen here
 	// we just have to make sure we set Context data correct
 	//if (ensure (Context.IsLeader()))
-	if ((Context.IsLeader()))
+	if (Context.IsLeader())
 	{
 		check(Instance.DeltaTimeRecord);
 		const float CurrentTime = Instance.Montage.CurrentPosition;
 		const float PreviousTime = Instance.DeltaTimeRecord->GetPrevious();
 		const float MoveDelta = Instance.DeltaTimeRecord->Delta;
 
+		// Update context's data for followers to use.
 		Context.SetLeaderDelta(MoveDelta);
 		Context.SetPreviousAnimationPositionRatio(PreviousTime / GetPlayLength());
 
@@ -1208,10 +1209,10 @@ void UAnimMontage::TickAssetPlayer(FAnimTickRecord& Instance, struct FAnimNotify
 				FMarkerTickRecord* MarkerTickRecord = Instance.MarkerTickRecord;
 				FMarkerTickContext& MarkerTickContext = Context.MarkerTickContext;
 
+				// Store the sync anim position BEFORE the asset has being ticked.
 				if (MarkerTickRecord->IsValid(Instance.bLooping))
 				{
 					MarkerTickContext.SetMarkerSyncStartPosition(GetMarkerSyncPositionFromMarkerIndicies(MarkerTickRecord->PreviousMarker.MarkerIndex, MarkerTickRecord->NextMarker.MarkerIndex, PreviousTime, nullptr));
-
 				}
 				else
 				{
@@ -1222,10 +1223,13 @@ void UAnimMontage::TickAssetPlayer(FAnimTickRecord& Instance, struct FAnimNotify
 					MarkerTickContext.SetMarkerSyncStartPosition(GetMarkerSyncPositionFromMarkerIndicies(PreviousMarker.MarkerIndex, NextMarker.MarkerIndex, PreviousTime, nullptr));
 				}
 
+				// Advance as leader.
 				// @todo this won't work well once we start jumping
 				// only thing is that passed markers won't work in this frame. To do that, I have to figure out how it jumped from where to where, 
 				GetMarkerIndicesForTime(CurrentTime, false, MarkerTickContext.GetValidMarkerNames(), MarkerTickRecord->PreviousMarker, MarkerTickRecord->NextMarker);
-				bRecordNeedsResetting = false; // we have updated it now, no need to reset
+				bRecordNeedsResetting = false; // we have updated it now, no need to reset.
+
+				// Store the sync anim position AFTER the asset has being ticked.
 				MarkerTickContext.SetMarkerSyncEndPosition(GetMarkerSyncPositionFromMarkerIndicies(MarkerTickRecord->PreviousMarker.MarkerIndex, MarkerTickRecord->NextMarker.MarkerIndex, CurrentTime, nullptr));
 
 				MarkerTickContext.MarkersPassedThisTick = *Instance.Montage.MarkersPassedThisTick;
@@ -1240,14 +1244,15 @@ void UAnimMontage::TickAssetPlayer(FAnimTickRecord& Instance, struct FAnimNotify
 				}
 #endif
 
-				UE_LOG(LogAnimMarkerSync, Log, TEXT("Montage Leading SyncGroup: %s(%s) Start [%s], End [%s]"),
-					*GetNameSafe(this), *SyncGroup.ToString(), *MarkerTickContext.GetMarkerSyncStartPosition().ToString(), *MarkerTickContext.GetMarkerSyncEndPosition().ToString());
+				UE_LOG(LogAnimMarkerSync, Log, TEXT("Montage Leading SyncGroup: %s(%s) Start [%s], End [%s]"), *GetNameSafe(this), *SyncGroup.ToString(), *MarkerTickContext.GetMarkerSyncStartPosition().ToString(), *MarkerTickContext.GetMarkerSyncEndPosition().ToString());
 			}
 		}
-
+		
+		// Update context's position for followers to use.
 		Context.SetAnimationPositionRatio(CurrentTime / GetPlayLength());
 	}
 
+	// Reset record if needed.
 	if (bRecordNeedsResetting && Instance.MarkerTickRecord)
 	{
 		Instance.MarkerTickRecord->Reset();
@@ -1258,39 +1263,46 @@ void UAnimMontage::CollectMarkers()
 {
 	MarkerData.AuthoredSyncMarkers.Reset();
 
-	// we want to make sure anim reference actually contains markers
-	if (SyncGroup != NAME_None && SlotAnimTracks.IsValidIndex(SyncSlotIndex))
+	// We want to make sure anim reference actually contains markers
+	if (SyncGroup != NAME_None)
 	{
-		const FAnimTrack& AnimTrack = SlotAnimTracks[SyncSlotIndex].AnimTrack;
-		for (const auto& Seg : AnimTrack.AnimSegments)
+		if (SlotAnimTracks.IsValidIndex(SyncSlotIndex))
 		{
-			const UAnimSequence* Sequence = Cast<UAnimSequence>(Seg.GetAnimReference());
-			if (Sequence && Sequence->AuthoredSyncMarkers.Num() > 0)
+			const FAnimTrack& AnimTrack = SlotAnimTracks[SyncSlotIndex].AnimTrack;
+			for (const auto& Seg : AnimTrack.AnimSegments)
 			{
-				// @todo this won't work well if you have starttime < end time and it does have negative playrate
-				for (const auto& Marker : Sequence->AuthoredSyncMarkers)
+				const UAnimSequence* Sequence = Cast<UAnimSequence>(Seg.GetAnimReference());
+				if (Sequence && Sequence->AuthoredSyncMarkers.Num() > 0)
 				{
-					if (Marker.Time >= Seg.AnimStartTime && Marker.Time <= Seg.AnimEndTime)
+					// @todo this won't work well if you have starttime < end time and it does have negative playrate
+					for (const auto& Marker : Sequence->AuthoredSyncMarkers)
 					{
-						const float TotalSegmentLength = (Seg.AnimEndTime - Seg.AnimStartTime)*Seg.AnimPlayRate;
-						// i don't think we can do negative in this case
-						ensure(TotalSegmentLength >= 0.f);
-
-						// now add to the list
-						for (int32 LoopCount = 0; LoopCount < Seg.LoopingCount; ++LoopCount)
+						if (Marker.Time >= Seg.AnimStartTime && Marker.Time <= Seg.AnimEndTime)
 						{
-							FAnimSyncMarker NewMarker;
+							const float TotalSegmentLength = (Seg.AnimEndTime - Seg.AnimStartTime)*Seg.AnimPlayRate;
+							// i don't think we can do negative in this case
+							ensure(TotalSegmentLength >= 0.f);
+
+							// now add to the list
+							for (int32 LoopCount = 0; LoopCount < Seg.LoopingCount; ++LoopCount)
+							{
+								FAnimSyncMarker NewMarker;
 
-							NewMarker.Time = Seg.StartPos + (Marker.Time - Seg.AnimStartTime)*Seg.AnimPlayRate + TotalSegmentLength*LoopCount;
-							NewMarker.MarkerName = Marker.MarkerName;
-							MarkerData.AuthoredSyncMarkers.Add(NewMarker);
+								NewMarker.Time = Seg.StartPos + (Marker.Time - Seg.AnimStartTime)*Seg.AnimPlayRate + TotalSegmentLength*LoopCount;
+								NewMarker.MarkerName = Marker.MarkerName;
+								MarkerData.AuthoredSyncMarkers.Add(NewMarker);
+							}
 						}
 					}
 				}
 			}
-		}
 
-		MarkerData.CollectUniqueNames();
+			MarkerData.CollectUniqueNames();
+		}
+		else
+		{
+			UE_LOG(LogAnimMontage, Warning, TEXT("Montage's sync slot track index is invalid. Make sure to use a valid slot track index, otherwise this Montage will use old sync markers or not use marker-based syncing at all."))
+		}
 	}
 }
 
diff --git a/Engine/Source/Runtime/Engine/Private/Animation/AnimNode_AssetPlayerBase.cpp b/Engine/Source/Runtime/Engine/Private/Animation/AnimNode_AssetPlayerBase.cpp
--- a/Engine/Source/Runtime/Engine/Private/Animation/AnimNode_AssetPlayerBase.cpp
+++ b/Engine/Source/Runtime/Engine/Private/Animation/AnimNode_AssetPlayerBase.cpp
@@ -59,15 +59,18 @@ void FAnimNode_AssetPlayerBase::CreateTickRecordForNode(const FAnimationUpdateCo
 			MethodToUse = EAnimSyncMethod::DoNotSync;
 		}
 	}
-	
+
+	// Setup parameters for synchronizer
 	const UE::Anim::FAnimSyncParams SyncParams(GroupNameToUse, SyncGroupRole, MethodToUse);
+
+	// Active asset player's tick record
 	FAnimTickRecord TickRecord(Sequence, bLooping, PlayRate, bIsEvaluator, FinalBlendWeight, /*inout*/ InternalTimeAccumulator, MarkerTickRecord);
 	TickRecord.GatherContextData(Context);
-
 	TickRecord.RootMotionWeightModifier = Context.GetRootMotionWeightModifier();
 	TickRecord.DeltaTimeRecord = &DeltaTimeRecord;
 	TickRecord.bRequestedInertialization = bRequestedInertialization;
-	
+
+	// Add asset player to synchronizer
 	SyncScope.AddTickRecord(TickRecord, SyncParams, UE::Anim::FAnimSyncDebugInfo(Context));
 
 	TRACE_ANIM_TICK_RECORD(Context, TickRecord);
diff --git a/Engine/Source/Runtime/Engine/Private/Animation/AnimSequence.cpp b/Engine/Source/Runtime/Engine/Private/Animation/AnimSequence.cpp
--- a/Engine/Source/Runtime/Engine/Private/Animation/AnimSequence.cpp
+++ b/Engine/Source/Runtime/Engine/Private/Animation/AnimSequence.cpp
@@ -3558,55 +3558,80 @@ void UAnimSequence::AdvanceMarkerPhaseAsLeader(bool bLooping, float MoveDelta, c
 	// Hard to reproduce issue triggering this, ensure & clamp for now
 	ensureMsgf(CurrentTime >= 0.f && CurrentTime <= GetPlayLength(), TEXT("Current time inside of AdvanceMarkerPhaseAsLeader is out of range %.3f of 0.0 to %.3f\n    Sequence: %s"), CurrentTime, GetPlayLength(), *GetFullName());
 
+	// Ensure our time is within the boundaries of the anim sequence.
 	CurrentTime = FMath::Clamp(CurrentTime, 0.f, GetPlayLength());
 
 	if (bPlayingForwards)
 	{
+		// Repeat until there is no more move delta to handle.
 		while (true)
 		{
-			if (NextMarker.MarkerIndex == -1)
+			// Our next marker is the end boundary.
+			if (NextMarker.MarkerIndex == MarkerIndexSpecialValues::AnimationBoundary)
 			{
-				float PrevCurrentTime = CurrentTime;
+				const float PrevCurrentTime = CurrentTime;
+
+				// Ensure we dont overshoot when advancing our time.
 				CurrentTime = FMath::Min(CurrentTime + CurrentMoveDelta, GetPlayLength());
+
+				// Compute the distances left to reach the next and previous marker from the current time position.
 				NextMarker.TimeToMarker = GetPlayLength() - CurrentTime;
-				PrevMarker.TimeToMarker -= CurrentTime - PrevCurrentTime; //Add how far we moved to distance from previous marker
+				PrevMarker.TimeToMarker -= CurrentTime - PrevCurrentTime;
 				break;
 			}
+
+			// Good, we have a valid next marker.
 			const FAnimSyncMarker& NextSyncMarker = AuthoredSyncMarkers[NextMarker.MarkerIndex];
 			checkSlow(ValidMarkerNames.Contains(NextSyncMarker.MarkerName));
 
+			// We are going to end up past our next marker.
 			if (CurrentMoveDelta > NextMarker.TimeToMarker)
 			{
+				// Move time to match that of the next marker, and update the move delta to reflect the change.
 				CurrentTime = NextSyncMarker.Time;
 				CurrentMoveDelta -= NextMarker.TimeToMarker;
 
-				PrevMarker.MarkerIndex = NextMarker.MarkerIndex;
+				// Make our new previous marker be the marker we just passed.
+				PrevMarker.MarkerIndex = NextMarker.MarkerIndex; 
 				PrevMarker.TimeToMarker = -CurrentMoveDelta;
 
-				int32 PassedMarker = MarkersPassed.Add(FPassedMarker());
+				// Record that we just passed a marker.
+				const int32 PassedMarker = MarkersPassed.Add(FPassedMarker());
 				MarkersPassed[PassedMarker].PassedMarkerName = NextSyncMarker.MarkerName;
 				MarkersPassed[PassedMarker].DeltaTimeWhenPassed = CurrentMoveDelta;
-
-				float MarkerTimeOffset = 0.f;
-				do
+				
+				// Compute our new next marker.
 				{
-					++NextMarker.MarkerIndex;
-					if (NextMarker.MarkerIndex >= AuthoredSyncMarkers.Num())
+					float MarkerTimeOffset = 0.f;
+					
+					do
 					{
-						if (!bLooping)
+						++NextMarker.MarkerIndex;
+
+						// No more markers up ahead.
+						if (NextMarker.MarkerIndex >= AuthoredSyncMarkers.Num())
 						{
-							NextMarker.MarkerIndex = -1;
-							break;
+							// Stop at anim end boundary.
+							if (!bLooping)
+							{
+								NextMarker.MarkerIndex = MarkerIndexSpecialValues::AnimationBoundary;
+								break;
+							}
+
+							// Make our next marker be the first marker found in the sequence. 
+							NextMarker.MarkerIndex = 0;
+							MarkerTimeOffset = GetPlayLength();
 						}
-						NextMarker.MarkerIndex = 0;
-						MarkerTimeOffset = GetPlayLength();
+					} while (!ValidMarkerNames.Contains(AuthoredSyncMarkers[NextMarker.MarkerIndex].MarkerName));
+
+					// Update time left to reach the new next marker
+					if (NextMarker.MarkerIndex != MarkerIndexSpecialValues::AnimationBoundary)
+					{
+						NextMarker.TimeToMarker = MarkerTimeOffset + AuthoredSyncMarkers[NextMarker.MarkerIndex].Time - CurrentTime;
 					}
-				} while (!ValidMarkerNames.Contains(AuthoredSyncMarkers[NextMarker.MarkerIndex].MarkerName));
-				if (NextMarker.MarkerIndex != -1)
-				{
-					NextMarker.TimeToMarker = MarkerTimeOffset + AuthoredSyncMarkers[NextMarker.MarkerIndex].Time - CurrentTime;
 				}
 			}
+			// We will not go past our next marker, we can advance comfortably.
 			else
 			{
 				CurrentTime = FMath::Fmod(CurrentTime + CurrentMoveDelta, GetPlayLength());
@@ -3614,6 +3639,7 @@ void UAnimSequence::AdvanceMarkerPhaseAsLeader(bool bLooping, float MoveDelta, c
 				{
 					CurrentTime += GetPlayLength();
 				}
+
 				NextMarker.TimeToMarker -= CurrentMoveDelta;
 				PrevMarker.TimeToMarker -= CurrentMoveDelta;
 				break;
@@ -3622,51 +3648,77 @@ void UAnimSequence::AdvanceMarkerPhaseAsLeader(bool bLooping, float MoveDelta, c
 	}
 	else
 	{
+		// Playing backwards.
+		
+		// Repeat until there is no more move delta to handle.
 		while (true)
 		{
-			if (PrevMarker.MarkerIndex == -1)
+			// Our previous marker is the start boundary.
+			if (PrevMarker.MarkerIndex == MarkerIndexSpecialValues::AnimationBoundary)
 			{
-				float PrevCurrentTime = CurrentTime;
+				const float PrevCurrentTime = CurrentTime;
+
+				// Ensure we dont undershoot when advancing our time.
 				CurrentTime = FMath::Max(CurrentTime + CurrentMoveDelta, 0.f);
+
+				// Compute the distances left to reach the next and previous marker from the current time position.
 				PrevMarker.TimeToMarker = CurrentTime;
-				NextMarker.TimeToMarker -= CurrentTime - PrevCurrentTime; //Add how far we moved to distance from previous marker
+				NextMarker.TimeToMarker -= CurrentTime - PrevCurrentTime;
 				break;
 			}
+
+			// Good, we have a valid previous marker.
 			const FAnimSyncMarker& PrevSyncMarker = AuthoredSyncMarkers[PrevMarker.MarkerIndex];
 			checkSlow(ValidMarkerNames.Contains(PrevSyncMarker.MarkerName));
-			
+
+			// We are going to end up past our previous marker.
 			if (CurrentMoveDelta < PrevMarker.TimeToMarker)
 			{
+				// Move time to match that of the previous marker, and update the move delta to reflect the change.
 				CurrentTime = PrevSyncMarker.Time;
 				CurrentMoveDelta -= PrevMarker.TimeToMarker;
 
+				// Make our new next marker be the marker we just passed.
 				NextMarker.MarkerIndex = PrevMarker.MarkerIndex;
 				NextMarker.TimeToMarker = -CurrentMoveDelta;
 
-				int32 PassedMarker = MarkersPassed.Add(FPassedMarker());
+				// Record that we just passed a marker.
+				const int32 PassedMarker = MarkersPassed.Add(FPassedMarker());
 				MarkersPassed[PassedMarker].PassedMarkerName = PrevSyncMarker.MarkerName;
 				MarkersPassed[PassedMarker].DeltaTimeWhenPassed = CurrentMoveDelta;
 
-				float MarkerTimeOffset = 0.f;
-				do
+				// Compute our new previous marker.
 				{
-					--PrevMarker.MarkerIndex;
-					if (PrevMarker.MarkerIndex < 0)
+					float MarkerTimeOffset = 0.f;
+					
+					do
 					{
-						if (!bLooping)
+						--PrevMarker.MarkerIndex;
+
+						// No more markers behind.
+						if (PrevMarker.MarkerIndex < 0)
 						{
-							PrevMarker.MarkerIndex = -1;
-							break;
+							// Stop at the anim start boundary.
+							if (!bLooping)
+							{
+								PrevMarker.MarkerIndex = MarkerIndexSpecialValues::AnimationBoundary;
+								break;
+							}
+							
+							// Make our previous marker be the last marker found in the sequence.
+							PrevMarker.MarkerIndex = AuthoredSyncMarkers.Num() - 1;
+							MarkerTimeOffset -= GetPlayLength();
 						}
-						PrevMarker.MarkerIndex = AuthoredSyncMarkers.Num() - 1;
-						MarkerTimeOffset -= GetPlayLength();
+					} while (!ValidMarkerNames.Contains(AuthoredSyncMarkers[PrevMarker.MarkerIndex].MarkerName));
+
+					// Update time left to reach marker.
+					if (PrevMarker.MarkerIndex != MarkerIndexSpecialValues::AnimationBoundary)
+					{
+						PrevMarker.TimeToMarker = MarkerTimeOffset + AuthoredSyncMarkers[PrevMarker.MarkerIndex].Time - CurrentTime;
 					}
-				} while (!ValidMarkerNames.Contains(AuthoredSyncMarkers[PrevMarker.MarkerIndex].MarkerName));
-				if (PrevMarker.MarkerIndex != -1)
-				{
-					PrevMarker.TimeToMarker = MarkerTimeOffset + AuthoredSyncMarkers[PrevMarker.MarkerIndex].Time - CurrentTime;
 				}
 			}
+			// We will not go past our previous marker, we can advance comfortably.
 			else
 			{
 				CurrentTime = FMath::Fmod(CurrentTime + CurrentMoveDelta, GetPlayLength());
@@ -3674,6 +3726,7 @@ void UAnimSequence::AdvanceMarkerPhaseAsLeader(bool bLooping, float MoveDelta, c
 				{
 					CurrentTime += GetPlayLength();
 				}
+
 				PrevMarker.TimeToMarker -= CurrentMoveDelta;
 				NextMarker.TimeToMarker -= CurrentMoveDelta;
 				break;
@@ -3701,6 +3754,7 @@ void AdvanceMarkerForwards(int32& Marker, FName MarkerToFind, bool bLooping, con
 {
 	int32 MaxIterations = AuthoredSyncMarkers.Num();
 
+	// Get next available marker.
 	while ((MarkerOrMirroredName(AuthoredSyncMarkers[Marker].MarkerName, MirrorTable) != MarkerToFind) && (--MaxIterations >= 0))
 	{
 		++Marker;
@@ -3711,6 +3765,7 @@ void AdvanceMarkerForwards(int32& Marker, FName MarkerToFind, bool bLooping, con
 		Marker %= AuthoredSyncMarkers.Num();
 	}
 
+	// In any invalid case, default to -1 aka an animation boundary.
 	if (!AuthoredSyncMarkers.IsValidIndex(Marker) || (MarkerOrMirroredName(AuthoredSyncMarkers[Marker].MarkerName, MirrorTable) != MarkerToFind))
 	{
 		Marker = MarkerIndexSpecialValues::AnimationBoundary;
@@ -3745,7 +3800,7 @@ void AdvanceMarkerBackwards(int32& Marker, FName MarkerToFind, bool bLooping, co
 
 bool MarkerMatchesPosition(const UAnimSequence* Sequence, int32 MarkerIndex, FName CorrectMarker, const UMirrorDataTable* MirrorTable)
 {
-	checkf(MarkerIndex != MarkerIndexSpecialValues::Unitialized, TEXT("Uninitialized marker supplied to MarkerMatchesPosition. Anim: %s Expecting marker %s (Added to help debug Jira OR-9675)"), *Sequence->GetName(), *CorrectMarker.ToString());
+	checkf(MarkerIndex != MarkerIndexSpecialValues::Uninitialized, TEXT("Uninitialized marker supplied to MarkerMatchesPosition. Anim: %s Expecting marker %s (Added to help debug Jira OR-9675)"), *Sequence->GetName(), *CorrectMarker.ToString());
 	return MarkerIndex == MarkerIndexSpecialValues::AnimationBoundary || CorrectMarker == MarkerOrMirroredName(Sequence->AuthoredSyncMarkers[MarkerIndex].MarkerName, MirrorTable);
 }
 
@@ -3753,16 +3808,24 @@ void UAnimSequence::ValidateCurrentPosition(const FMarkerSyncAnimPosition& Posit
 {
 	if (bPlayingForwards)
 	{
+		// Ensure previous marker matches the desired previous marker given a name.
 		if (!MarkerMatchesPosition(this, PreviousMarker.MarkerIndex, Position.PreviousMarkerName, MirrorTable))
 		{
 			AdvanceMarkerForwards(PreviousMarker.MarkerIndex, Position.PreviousMarkerName, bLooping, AuthoredSyncMarkers, MirrorTable);
-			NextMarker.MarkerIndex = (PreviousMarker.MarkerIndex + 1);
-			if(NextMarker.MarkerIndex >= AuthoredSyncMarkers.Num())
+
+			// Ensure next marker comes after the recently updated previous marker.
 			{
-				NextMarker.MarkerIndex = bLooping ? NextMarker.MarkerIndex % AuthoredSyncMarkers.Num() : MarkerIndexSpecialValues::AnimationBoundary;
+				NextMarker.MarkerIndex = (PreviousMarker.MarkerIndex + 1);
+
+				// If needed, loop back or stop at end boundary.
+				if (NextMarker.MarkerIndex >= AuthoredSyncMarkers.Num())
+				{
+					NextMarker.MarkerIndex = bLooping ? NextMarker.MarkerIndex % AuthoredSyncMarkers.Num() : MarkerIndexSpecialValues::AnimationBoundary;
+				}
 			}
 		}
 
+		// Ensure next marker matches the desired next marker given a name.
 		if (!MarkerMatchesPosition(this, NextMarker.MarkerIndex, Position.NextMarkerName, MirrorTable))
 		{
 			AdvanceMarkerForwards(NextMarker.MarkerIndex, Position.NextMarkerName, bLooping, AuthoredSyncMarkers, MirrorTable);
@@ -3770,11 +3833,13 @@ void UAnimSequence::ValidateCurrentPosition(const FMarkerSyncAnimPosition& Posit
 	}
 	else
 	{
-		const int32 MarkerRange = AuthoredSyncMarkers.Num();
+		// Ensure next marker matches the desired next marker given a name.
 		if (!MarkerMatchesPosition(this, NextMarker.MarkerIndex, Position.NextMarkerName, MirrorTable))
 		{
 			AdvanceMarkerBackwards(NextMarker.MarkerIndex, Position.NextMarkerName, bLooping, AuthoredSyncMarkers, MirrorTable);
-			if(NextMarker.MarkerIndex == MarkerIndexSpecialValues::AnimationBoundary || (NextMarker.MarkerIndex == 0 && bLooping))
+
+			// Ensure previous marker comes before the recently updated next marker.
+			if (NextMarker.MarkerIndex == MarkerIndexSpecialValues::AnimationBoundary || (NextMarker.MarkerIndex == 0 && bLooping))
 			{
 				PreviousMarker.MarkerIndex = AuthoredSyncMarkers.Num() - 1;
 			}
@@ -3783,6 +3848,8 @@ void UAnimSequence::ValidateCurrentPosition(const FMarkerSyncAnimPosition& Posit
 				PreviousMarker.MarkerIndex = NextMarker.MarkerIndex - 1;
 			}
 		}
+		
+		// Ensure previous marker matches the desired previous marker given a name.
 		if (!MarkerMatchesPosition(this, PreviousMarker.MarkerIndex, Position.PreviousMarkerName, MirrorTable))
 		{
 			AdvanceMarkerBackwards(PreviousMarker.MarkerIndex, Position.PreviousMarkerName, bLooping, AuthoredSyncMarkers, MirrorTable);
@@ -4155,105 +4222,139 @@ void UAnimSequence::AdvanceMarkerPhaseAsFollower(const FMarkerTickContext& Conte
 {
 	const bool bPlayingForwards = DeltaRemaining > 0.f;
 
+	// Ensures the sequence's markers match the sync start position.
 	ValidateCurrentPosition(Context.GetMarkerSyncStartPosition(), bPlayingForwards, bLooping, CurrentTime, PreviousMarker, NextMarker, MirrorTable);
+
 	if (bPlayingForwards)
 	{
 		int32 PassedMarkersIndex = 0;
+
+		// Advance all next markers to follow markers passed by leader and update previous markers accordingly.
 		do
 		{
-			if (NextMarker.MarkerIndex == -1)
+			// They are no more markers ahead.
+			if (NextMarker.MarkerIndex == MarkerIndexSpecialValues::AnimationBoundary)
 			{
-				check(!bLooping || Context.GetMarkerSyncEndPosition().NextMarkerName == NAME_None); // shouldnt have an end of anim marker if looping
+				check(!bLooping || Context.GetMarkerSyncEndPosition().NextMarkerName == NAME_None); // You shouldn't have an end of anim marker if looping
 				CurrentTime = FMath::Min(CurrentTime + DeltaRemaining, GetPlayLength());
 				break;
 			}
+			// Find markers passed by group leader.
 			else if (PassedMarkersIndex < Context.MarkersPassedThisTick.Num())
 			{
-				PreviousMarker.MarkerIndex = NextMarker.MarkerIndex;
-				checkSlow(NextMarker.MarkerIndex != -1);
-				const FPassedMarker& PassedMarker = Context.MarkersPassedThisTick[PassedMarkersIndex];
-				AdvanceMarkerForwards(NextMarker.MarkerIndex, PassedMarker.PassedMarkerName, bLooping, AuthoredSyncMarkers, MirrorTable);
-				if (NextMarker.MarkerIndex == -1)
+				PreviousMarker.MarkerIndex = NextMarker.MarkerIndex; 
+
+				checkSlow(NextMarker.MarkerIndex != MarkerIndexSpecialValues::AnimationBoundary);
+				
+				// Advance our next marker to match marker passed by leader.
 				{
-					DeltaRemaining = PassedMarker.DeltaTimeWhenPassed;
+					const FPassedMarker& MarkerPassedByLeader = Context.MarkersPassedThisTick[PassedMarkersIndex];
+					
+					AdvanceMarkerForwards(NextMarker.MarkerIndex, MarkerPassedByLeader.PassedMarkerName, bLooping, AuthoredSyncMarkers, MirrorTable);
+
+					// Ensure that any left over delta is handled in last iteration.
+					if (NextMarker.MarkerIndex == MarkerIndexSpecialValues::AnimationBoundary)
+					{
+						DeltaRemaining = MarkerPassedByLeader.DeltaTimeWhenPassed;
+					}
 				}
+				
 				++PassedMarkersIndex;
 			}
 		} while (PassedMarkersIndex < Context.MarkersPassedThisTick.Num());
 
-		const FMarkerSyncAnimPosition& End = Context.GetMarkerSyncEndPosition();
-		
-		if (End.NextMarkerName == NAME_None)
+		// Get sync position after group leader was ticked.
+		const FMarkerSyncAnimPosition& LeaderEndPosition = Context.GetMarkerSyncEndPosition();
+
+		// Ensure next marker is a boundary, if the group leader's next marker was one.
+		if (LeaderEndPosition.NextMarkerName == NAME_None)
 		{
-			NextMarker.MarkerIndex = -1;
+			NextMarker.MarkerIndex = MarkerIndexSpecialValues::AnimationBoundary;
 		}
-
-		if (NextMarker.MarkerIndex != -1 && Context.MarkersPassedThisTick.Num() > 0)
+		
+		// Ensure next marker matches leader's next marker after tick.
+		if (NextMarker.MarkerIndex != MarkerIndexSpecialValues::AnimationBoundary && Context.MarkersPassedThisTick.Num() > 0)
 		{
-			AdvanceMarkerForwards(NextMarker.MarkerIndex, End.NextMarkerName, bLooping, AuthoredSyncMarkers, MirrorTable);
+			AdvanceMarkerForwards(NextMarker.MarkerIndex, LeaderEndPosition.NextMarkerName, bLooping, AuthoredSyncMarkers, MirrorTable);
 		}
 
-		//Validation
-		if (NextMarker.MarkerIndex != -1)
+		// Validation
+		if (NextMarker.MarkerIndex != MarkerIndexSpecialValues::AnimationBoundary)
 		{
-			check(MarkerOrMirroredName(AuthoredSyncMarkers[NextMarker.MarkerIndex].MarkerName, MirrorTable) == End.NextMarkerName);
+			check(MarkerOrMirroredName(AuthoredSyncMarkers[NextMarker.MarkerIndex].MarkerName, MirrorTable) == LeaderEndPosition.NextMarkerName);
 		}
 
 		// End Validation
 		// Only reset position if we found valid markers. Otherwise stay where we are to not pop.
 		if ((PreviousMarker.MarkerIndex != MarkerIndexSpecialValues::AnimationBoundary) && (NextMarker.MarkerIndex != MarkerIndexSpecialValues::AnimationBoundary))
 		{
-			CurrentTime = GetCurrentTimeFromMarkers(PreviousMarker, NextMarker, End.PositionBetweenMarkers);
+			CurrentTime = GetCurrentTimeFromMarkers(PreviousMarker, NextMarker, LeaderEndPosition.PositionBetweenMarkers);
 		}
 	}
 	else
 	{
 		int32 PassedMarkersIndex = 0;
+
+		// Advance all previous markers to follow markers passed by leader and update next markers accordingly.
 		do
 		{
-			if (PreviousMarker.MarkerIndex == -1)
+			// They are no more markers ahead.
+			if (PreviousMarker.MarkerIndex == MarkerIndexSpecialValues::AnimationBoundary)
 			{
-				check(!bLooping || Context.GetMarkerSyncEndPosition().PreviousMarkerName == NAME_None); // shouldn't have an end of anim marker if looping
+				check(!bLooping || Context.GetMarkerSyncEndPosition().PreviousMarkerName == NAME_None); // You shouldn't have an end of anim marker if looping.
 				CurrentTime = FMath::Max(CurrentTime + DeltaRemaining, 0.f);
 				break;
 			}
+			// Find markers passed by group leader.
 			else if (PassedMarkersIndex < Context.MarkersPassedThisTick.Num())
 			{
 				NextMarker.MarkerIndex = PreviousMarker.MarkerIndex;
-				checkSlow(PreviousMarker.MarkerIndex != -1);
-				const FPassedMarker& PassedMarker = Context.MarkersPassedThisTick[PassedMarkersIndex];
-				AdvanceMarkerBackwards(PreviousMarker.MarkerIndex, PassedMarker.PassedMarkerName, bLooping, AuthoredSyncMarkers, MirrorTable);
-				if (PreviousMarker.MarkerIndex == -1)
+				
+				checkSlow(PreviousMarker.MarkerIndex != MarkerIndexSpecialValues::AnimationBoundary);
+
+				// Advance previous marker matches marker passed by leader
 				{
-					DeltaRemaining = PassedMarker.DeltaTimeWhenPassed;
+					const FPassedMarker& MarkerPassedByLeader = Context.MarkersPassedThisTick[PassedMarkersIndex];
+
+					AdvanceMarkerBackwards(PreviousMarker.MarkerIndex, MarkerPassedByLeader.PassedMarkerName, bLooping, AuthoredSyncMarkers, MirrorTable);
+
+					// Ensure that any left over delta is handled in last iteration.
+					if (PreviousMarker.MarkerIndex == MarkerIndexSpecialValues::AnimationBoundary)
+					{
+						DeltaRemaining = MarkerPassedByLeader.DeltaTimeWhenPassed;
+					}
+					
+					++PassedMarkersIndex;
 				}
-				++PassedMarkersIndex;
 			}
 		} while (PassedMarkersIndex < Context.MarkersPassedThisTick.Num());
 
-		const FMarkerSyncAnimPosition& End = Context.GetMarkerSyncEndPosition();
+		// Get sync position after group leader was ticked.
+		const FMarkerSyncAnimPosition& LeaderEndPosition = Context.GetMarkerSyncEndPosition();
 
-		if (PreviousMarker.MarkerIndex != -1 && Context.MarkersPassedThisTick.Num() > 0)
+		// Ensure previous marker is a boundary, if the group leader's was one.
+		if (LeaderEndPosition.PreviousMarkerName == NAME_None)
 		{
-			AdvanceMarkerBackwards(PreviousMarker.MarkerIndex, End.PreviousMarkerName, bLooping, AuthoredSyncMarkers, MirrorTable);
+			PreviousMarker.MarkerIndex = MarkerIndexSpecialValues::AnimationBoundary;
 		}
-
-		if (End.PreviousMarkerName == NAME_None)
+		
+		// Ensure previous marker match leader's previous marker after tick
+		if (PreviousMarker.MarkerIndex != MarkerIndexSpecialValues::AnimationBoundary && Context.MarkersPassedThisTick.Num() > 0)
 		{
-			PreviousMarker.MarkerIndex = -1;
+			AdvanceMarkerBackwards(PreviousMarker.MarkerIndex, LeaderEndPosition.PreviousMarkerName, bLooping, AuthoredSyncMarkers, MirrorTable);
 		}
 
-		//Validation
-		if (PreviousMarker.MarkerIndex != -1)
+		// Validation
+		if (PreviousMarker.MarkerIndex != MarkerIndexSpecialValues::AnimationBoundary)
 		{
-			check(AuthoredSyncMarkers[PreviousMarker.MarkerIndex].MarkerName == End.PreviousMarkerName);
+			check(AuthoredSyncMarkers[PreviousMarker.MarkerIndex].MarkerName == LeaderEndPosition.PreviousMarkerName);
 		}
 
 		// End Validation
 		// Only reset position if we found valid markers. Otherwise stay where we are to not pop.
 		if ((PreviousMarker.MarkerIndex != MarkerIndexSpecialValues::AnimationBoundary) && (NextMarker.MarkerIndex != MarkerIndexSpecialValues::AnimationBoundary))
 		{
-			CurrentTime = GetCurrentTimeFromMarkers(PreviousMarker, NextMarker, End.PositionBetweenMarkers);
+			CurrentTime = GetCurrentTimeFromMarkers(PreviousMarker, NextMarker, LeaderEndPosition.PositionBetweenMarkers);
 		}
 	}
 }
@@ -4263,9 +4364,9 @@ void UAnimSequence::GetMarkerIndicesForTime(float CurrentTime, bool bLooping, co
 	const int LoopModStart = bLooping ? -1 : 0;
 	const int LoopModEnd = bLooping ? 2 : 1;
 
-	OutPrevMarker.MarkerIndex = -1;
+	OutPrevMarker.MarkerIndex = MarkerIndexSpecialValues::AnimationBoundary;
 	OutPrevMarker.TimeToMarker = -CurrentTime;
-	OutNextMarker.MarkerIndex = -1;
+	OutNextMarker.MarkerIndex = MarkerIndexSpecialValues::AnimationBoundary;
 	OutNextMarker.TimeToMarker = GetPlayLength() - CurrentTime;
 
 	for (int32 LoopMod = LoopModStart; LoopMod < LoopModEnd; ++LoopMod)
@@ -4290,7 +4391,7 @@ void UAnimSequence::GetMarkerIndicesForTime(float CurrentTime, bool bLooping, co
 				}
 			}
 		}
-		if (OutNextMarker.MarkerIndex != -1)
+		if (OutNextMarker.MarkerIndex != MarkerIndexSpecialValues::AnimationBoundary)
 		{
 			break; // Done
 		}
@@ -4301,8 +4402,9 @@ FMarkerSyncAnimPosition UAnimSequence::GetMarkerSyncPositionFromMarkerIndicies(i
 {
 	FMarkerSyncAnimPosition SyncPosition;
 	float PrevTime, NextTime;
-	
-	if (PrevMarker != -1 && ensureAlwaysMsgf(AuthoredSyncMarkers.IsValidIndex(PrevMarker),
+
+	// Get previous marker's time and name.
+	if (PrevMarker != MarkerIndexSpecialValues::AnimationBoundary && ensureAlwaysMsgf(AuthoredSyncMarkers.IsValidIndex(PrevMarker),
 		TEXT("%s - MarkerCount: %d, PrevMarker : %d, NextMarker: %d, CurrentTime : %0.2f"), *GetFullName(), AuthoredSyncMarkers.Num(), PrevMarker, NextMarker, CurrentTime))
 	{
 		PrevTime = AuthoredSyncMarkers[PrevMarker].Time;
@@ -4313,7 +4415,8 @@ FMarkerSyncAnimPosition UAnimSequence::GetMarkerSyncPositionFromMarkerIndicies(i
 		PrevTime = 0.f;
 	}
 
-	if (NextMarker != -1 && ensureAlwaysMsgf(AuthoredSyncMarkers.IsValidIndex(NextMarker),
+	// Get next marker's time and name.
+	if (NextMarker != MarkerIndexSpecialValues::AnimationBoundary && ensureAlwaysMsgf(AuthoredSyncMarkers.IsValidIndex(NextMarker),
 		TEXT("%s - MarkerCount: %d, PrevMarker : %d, NextMarker: %d, CurrentTime : %0.2f"), *GetFullName(), AuthoredSyncMarkers.Num(), PrevMarker, NextMarker, CurrentTime))
 	{
 		NextTime = AuthoredSyncMarkers[NextMarker].Time;
@@ -4342,28 +4445,36 @@ FMarkerSyncAnimPosition UAnimSequence::GetMarkerSyncPositionFromMarkerIndicies(i
 
 	check(NextTime > PrevTime);
 
+	// Store the encoded current time position as a ratio between markers
 	SyncPosition.PositionBetweenMarkers = (CurrentTime - PrevTime) / (NextTime - PrevTime);
 	return SyncPosition;
 }
 
 float UAnimSequence::GetCurrentTimeFromMarkers(FMarkerPair& PrevMarker, FMarkerPair& NextMarker, float PositionBetweenMarkers) const
 {
-	float PrevTime = (PrevMarker.MarkerIndex != -1) ? AuthoredSyncMarkers[PrevMarker.MarkerIndex].Time : 0.f;
-	float NextTime = (NextMarker.MarkerIndex != -1) ? AuthoredSyncMarkers[NextMarker.MarkerIndex].Time : GetPlayLength();
+	// Query marker times, or start and end boundary times, respectively.
+	float PrevTime = (PrevMarker.MarkerIndex != MarkerIndexSpecialValues::AnimationBoundary) ? AuthoredSyncMarkers[PrevMarker.MarkerIndex].Time : 0.f;
+	float NextTime = (NextMarker.MarkerIndex != MarkerIndexSpecialValues::AnimationBoundary) ? AuthoredSyncMarkers[NextMarker.MarkerIndex].Time : GetPlayLength();
 
+	// Account for looping
 	if (PrevTime >= NextTime)
 	{
-		PrevTime -= GetPlayLength(); //Account for looping
+		PrevTime -= GetPlayLength(); 
 	}
+
+	// Compute current time given start and end marker times.
 	float CurrentTime = PrevTime + PositionBetweenMarkers * (NextTime - PrevTime);
 
+	// Compute time to reach each marker.
 	PrevMarker.TimeToMarker = PrevTime - CurrentTime;
 	NextMarker.TimeToMarker = NextTime - CurrentTime;
 
+	// Account for looping while playing backwards.
 	if (CurrentTime < 0.f)
 	{
 		CurrentTime += GetPlayLength();
 	}
+	
 	CurrentTime = FMath::Clamp<float>(CurrentTime, 0, GetPlayLength());
 
 	return CurrentTime;
@@ -4377,9 +4488,9 @@ void UAnimSequence::GetMarkerIndicesForPosition(
 	{
 		return MarkerOrMirroredName(SyncMarker.MarkerName, MirrorTable);
 	};
-
+	
 	// If we're not looping, assume we're playing a transition and we need to stay where we are.
-	// Also do this if we have no usable SyncPosition
+	// Also do this if we have no usable SyncPosition.
 	if (!bLooping || (SyncPosition.PreviousMarkerName == NAME_None && SyncPosition.NextMarkerName == NAME_None))
 	{
 		OutPrevMarker.MarkerIndex = INDEX_NONE;
@@ -4390,11 +4501,13 @@ void UAnimSequence::GetMarkerIndicesForPosition(
 			const FAnimSyncMarker& SyncMarker = AuthoredSyncMarkers[Idx];
 			const float MarkerTime = SyncMarker.Time;
 
+			// Match the position's previous marker name, and store its index.
 			if (OutCurrentTime > MarkerTime && GetMarkerName(SyncMarker) == SyncPosition.PreviousMarkerName)
 			{
 				OutPrevMarker.MarkerIndex = Idx;
 				OutPrevMarker.TimeToMarker = MarkerTime - OutCurrentTime;
 			}
+			// Match the position's next marker name, and store its index. By this point we should have found the previous marker index so we can stop searching.
 			else if (OutCurrentTime < MarkerTime && GetMarkerName(SyncMarker) == SyncPosition.NextMarkerName)
 			{
 				OutNextMarker.MarkerIndex = Idx;
@@ -4407,11 +4520,16 @@ void UAnimSequence::GetMarkerIndicesForPosition(
 		return;
 	}
 
+	// Handle case where the position's previous marker is the start boundary.
 	if (SyncPosition.PreviousMarkerName == NAME_None)
 	{
-		OutPrevMarker.MarkerIndex = -1;
+		// Make output prev marker index be the start boundary.
+		OutPrevMarker.MarkerIndex = MarkerIndexSpecialValues::AnimationBoundary;
+
+		// Our position's next marker should never be the end boundary, otherwise we dont have any sync markers at all.
 		check(SyncPosition.NextMarkerName != NAME_None);
 
+		// Find next marker index.
 		for (int32 Idx = 0; Idx < AuthoredSyncMarkers.Num(); ++Idx)
 		{
 			const FAnimSyncMarker& Marker = AuthoredSyncMarkers[Idx];
@@ -4422,15 +4540,21 @@ void UAnimSequence::GetMarkerIndicesForPosition(
 				return;
 			}
 		}
+		
 		// Should have found a marker above!
 		checkf(false, TEXT("Next Marker not found in GetMarkerIndicesForPosition. Anim: %s Expecting marker %s (Added to help debug Jira OR-9675)"), *GetName(), *SyncPosition.NextMarkerName.ToString());
 	}
 
+	// Handle case where the position's next marker is the end boundary.
 	if (SyncPosition.NextMarkerName == NAME_None)
 	{
-		OutNextMarker.MarkerIndex = -1;
+		// Make output next marker index be the end boundary.
+		OutNextMarker.MarkerIndex = MarkerIndexSpecialValues::AnimationBoundary;
+
+		// Our position's previous marker should never be the start boundary, otherwise we dont have any sync markers at all.
 		check(SyncPosition.PreviousMarkerName != NAME_None);
 
+		// Find previous marker index.
 		for (int32 Idx = AuthoredSyncMarkers.Num() - 1; Idx >= 0; --Idx)
 		{
 			const FAnimSyncMarker& Marker = AuthoredSyncMarkers[Idx];
@@ -4441,6 +4565,7 @@ void UAnimSequence::GetMarkerIndicesForPosition(
 				return;
 			}
 		}
+		
 		// Should have found a marker above!
 		checkf(false, TEXT("Previous Marker not found in GetMarkerIndicesForPosition. Anim: %s Expecting marker %s (Added to help debug Jira OR-9675)"), *GetName(), *SyncPosition.PreviousMarkerName.ToString());
 	}
@@ -4448,31 +4573,43 @@ void UAnimSequence::GetMarkerIndicesForPosition(
 	float DiffToCurrentTime = FLT_MAX;
 	const float CurrentInputTime  = OutCurrentTime;
 
+	// Handle case for looping and sync position not being on either boundary.
 	for (int32 PrevMarkerIdx = 0; PrevMarkerIdx < AuthoredSyncMarkers.Num(); ++PrevMarkerIdx)
 	{
 		const FAnimSyncMarker& PrevMarker = AuthoredSyncMarkers[PrevMarkerIdx];
+
+		// We have matched the position's previous marker name.
 		if (GetMarkerName(PrevMarker) == SyncPosition.PreviousMarkerName)
 		{
 			const int32 EndMarkerSearchStart = PrevMarkerIdx + 1;
-
 			const int32 EndCount = bLooping ? AuthoredSyncMarkers.Num() + EndMarkerSearchStart : AuthoredSyncMarkers.Num();
+			
 			for (int32 NextMarkerCount = EndMarkerSearchStart; NextMarkerCount < EndCount; ++NextMarkerCount)
 			{
 				const int32 NextMarkerIdx = NextMarkerCount % AuthoredSyncMarkers.Num();
 
+				// We have matched the position's next marker name.
 				if (GetMarkerName(AuthoredSyncMarkers[NextMarkerIdx]) == SyncPosition.NextMarkerName)
 				{
 					float NextMarkerTime = AuthoredSyncMarkers[NextMarkerIdx].Time;
+
+					// Handle case where we need to loop to get to be able to get to the next marker.
 					if (NextMarkerTime < PrevMarker.Time)
 					{
 						NextMarkerTime += GetPlayLength();
 					}
+
+					// Get current time based of sync position.
 					float ThisCurrentTime = PrevMarker.Time + SyncPosition.PositionBetweenMarkers * (NextMarkerTime - PrevMarker.Time);
+
+					// Adjust for case where we need to loop to be able to get to the next marker.
 					if (ThisCurrentTime > GetPlayLength())
 					{
 						ThisCurrentTime -= GetPlayLength();
 					}
-					float ThisDiff = FMath::Abs(ThisCurrentTime - CurrentInputTime);
+
+					// Find marker indices closest to input time position.
+					const float ThisDiff = FMath::Abs(ThisCurrentTime - CurrentInputTime);
 					if (ThisDiff < DiffToCurrentTime)
 					{
 						DiffToCurrentTime = ThisDiff;
@@ -4481,14 +4618,14 @@ void UAnimSequence::GetMarkerIndicesForPosition(
 						OutCurrentTime = GetCurrentTimeFromMarkers(OutPrevMarker, OutNextMarker, SyncPosition.PositionBetweenMarkers);
 					}
 
-					// this marker test is done, move onto next one
+					// This marker test is done, move onto next one.
 					break;
 				}
 			}
 
 			// If we get here and we haven't found a match and we are not looping then there 
 			// is no point running the rest of the loop set up something as relevant as we can and carry on
-			if (OutPrevMarker.MarkerIndex == MarkerIndexSpecialValues::Unitialized)
+			if (OutPrevMarker.MarkerIndex == MarkerIndexSpecialValues::Uninitialized)
 			{
 				//Find nearest previous marker that is earlier than our current time
 				DiffToCurrentTime = OutCurrentTime - PrevMarker.Time;
@@ -4507,9 +4644,10 @@ void UAnimSequence::GetMarkerIndicesForPosition(
 			}
 		}
 	}
+	
 	// Should have found a markers above!
-	checkf(OutPrevMarker.MarkerIndex != MarkerIndexSpecialValues::Unitialized, TEXT("Prev Marker not found in GetMarkerIndicesForPosition. Anim: %s Expecting marker %s (Added to help debug Jira OR-9675)"), *GetName(), *SyncPosition.PreviousMarkerName.ToString());
-	checkf(OutNextMarker.MarkerIndex != MarkerIndexSpecialValues::Unitialized, TEXT("Next Marker not found in GetMarkerIndicesForPosition. Anim: %s Expecting marker %s (Added to help debug Jira OR-9675)"), *GetName(), *SyncPosition.NextMarkerName.ToString());
+	checkf(OutPrevMarker.MarkerIndex != MarkerIndexSpecialValues::Uninitialized, TEXT("Prev Marker not found in GetMarkerIndicesForPosition. Anim: %s Expecting marker %s (Added to help debug Jira OR-9675)"), *GetName(), *SyncPosition.PreviousMarkerName.ToString());
+	checkf(OutNextMarker.MarkerIndex != MarkerIndexSpecialValues::Uninitialized, TEXT("Next Marker not found in GetMarkerIndicesForPosition. Anim: %s Expecting marker %s (Added to help debug Jira OR-9675)"), *GetName(), *SyncPosition.NextMarkerName.ToString());
 }
 
 float UAnimSequence::GetFirstMatchingPosFromMarkerSyncPos(const FMarkerSyncAnimPosition& InMarkerSyncGroupPosition) const
diff --git a/Engine/Source/Runtime/Engine/Private/Animation/AnimSequenceBase.cpp b/Engine/Source/Runtime/Engine/Private/Animation/AnimSequenceBase.cpp
--- a/Engine/Source/Runtime/Engine/Private/Animation/AnimSequenceBase.cpp
+++ b/Engine/Source/Runtime/Engine/Private/Animation/AnimSequenceBase.cpp
@@ -544,7 +544,9 @@ void UAnimSequenceBase::GetAnimNotifiesFromDeltaPositions(const float& PreviousP
 
 void UAnimSequenceBase::TickAssetPlayer(FAnimTickRecord& Instance, struct FAnimNotifyQueue& NotifyQueue, FAnimAssetTickContext& Context) const
 {
+	// Local copy of the asset player's time that is about to be updated.
 	float CurrentTime = *(Instance.TimeAccumulator);
+	
 	if (Context.ShouldResyncToSyncGroup() && !Instance.bIsEvaluator)
 	{
 		// Synchronize the asset player time to the other sync group members when (re)joining the group
@@ -553,16 +555,17 @@ void UAnimSequenceBase::TickAssetPlayer(FAnimTickRecord& Instance, struct FAnimN
 
 	float PreviousTime = CurrentTime;
 	float DeltaTime = 0.f;
+	const float PlayRate = Instance.PlayRateMultiplier * RateScale;
 
-	const float PlayRate = Instance.PlayRateMultiplier * this->RateScale;
-
-	if( Context.IsLeader() )
+	if (Context.IsLeader())
 	{
 		DeltaTime = PlayRate * Context.GetDeltaTime();
 
+		// Prepare context before ticking.
 		Context.SetLeaderDelta(DeltaTime);
 		Context.SetPreviousAnimationPositionRatio(PreviousTime / GetPlayLength());
 
+		// Tick as leader using marked based syncing if possible, other fallback to normal ticking.
 		if (DeltaTime != 0.f)
 		{
 			if (Instance.bCanUseMarkerSync && Context.CanUseMarkerPosition())
@@ -577,11 +580,12 @@ void UAnimSequenceBase::TickAssetPlayer(FAnimTickRecord& Instance, struct FAnimN
 			}
 		}
 
+		// Update context's data after ticking.
 		Context.SetAnimationPositionRatio(CurrentTime / GetPlayLength());
 	}
 	else
 	{
-		// Follow the leader
+		// Follow the leader using marker based syncing if possible, otherwise fallback to length based syncing.
 		if (Instance.bCanUseMarkerSync)
 		{
 			if (Context.CanUseMarkerPosition() && Context.MarkerTickContext.IsMarkerSyncStartValid())
@@ -591,6 +595,7 @@ void UAnimSequenceBase::TickAssetPlayer(FAnimTickRecord& Instance, struct FAnimN
 			else
 			{
 				DeltaTime = PlayRate * Context.GetDeltaTime();
+				
 				// If leader is not valid, advance time as normal, do not jump position and pop.
 				FAnimationRuntime::AdvanceTime(Instance.bLooping, DeltaTime, CurrentTime, GetPlayLength());
 				UE_LOG(LogAnimMarkerSync, Log, TEXT("Follower (%s) (normal advance) - PreviousTime (%0.2f), CurrentTime (%0.2f), MoveDelta (%0.2f), Looping (%d) "), *GetName(), PreviousTime, CurrentTime, DeltaTime, Instance.bLooping ? 1 : 0);
@@ -598,44 +603,49 @@ void UAnimSequenceBase::TickAssetPlayer(FAnimTickRecord& Instance, struct FAnimN
 		}
 		else
 		{
+			// Match group leader's last calculated anim position since we're not using marker based sync.
 			PreviousTime = Context.GetPreviousAnimationPositionRatio() * GetPlayLength();
 			CurrentTime = Context.GetAnimationPositionRatio() * GetPlayLength();
+			
 			UE_LOG(LogAnimMarkerSync, Log, TEXT("Follower (%s) (normalized position advance) - PreviousTime (%0.2f), CurrentTime (%0.2f), MoveDelta (%0.2f), Looping (%d) "), *GetName(), PreviousTime, CurrentTime, DeltaTime, Instance.bLooping ? 1 : 0);
 		}
-
-		//@TODO: NOTIFIES: Calculate AdvanceType based on what the new delta time is
-
-		if( CurrentTime != PreviousTime )
+		
+		if (CurrentTime != PreviousTime)
 		{
 			// Figure out delta time 
 			DeltaTime = CurrentTime - PreviousTime;
-			// if we went against play rate, then loop around.
-			if( (DeltaTime * PlayRate) < 0.f )
+			
+			// If we went against play rate, then loop around.
+			if ((DeltaTime * PlayRate) < 0.f)
 			{
 				DeltaTime += FMath::Sign<float>(PlayRate) * GetPlayLength();
 			}
 		}
 	}
 
-	// Assign the instance's TimeAccumulator after all side effects on CurrentTime have been applied
+	// Update the instance's TimeAccumulator after all side effects on the local copy of the asset player's time have been applied.
 	*(Instance.TimeAccumulator) = CurrentTime;
-
+	
 	// Capture the final adjusted delta time and previous frame time as an asset player record
 	check(Instance.DeltaTimeRecord);
 	Instance.DeltaTimeRecord->Set(PreviousTime, DeltaTime);
 
+	// Allow asset to react right after its asset player has being ticked.
 	HandleAssetPlayerTickedInternal(Context, PreviousTime, DeltaTime, Instance, NotifyQueue);
 }
 
 void UAnimSequenceBase::TickByMarkerAsFollower(FMarkerTickRecord &Instance, FMarkerTickContext &MarkerContext, float& CurrentTime, float& OutPreviousTime, const float MoveDelta, const bool bLooping, const UMirrorDataTable* MirrorTable) const
 {
+	// Re-compute marker indices since the asset's tick record is invalid. Get previous and next markers.
 	if (!Instance.IsValid(bLooping))
 	{
 		GetMarkerIndicesForPosition(MarkerContext.GetMarkerSyncStartPosition(), bLooping, Instance.PreviousMarker, Instance.NextMarker, CurrentTime, MirrorTable);
 	}
 
+	// Keep track the asset's previous time, since its current time is about to be modified.
 	OutPreviousTime = CurrentTime;
 
+	// Tick and update as follower.
 	AdvanceMarkerPhaseAsFollower(MarkerContext, MoveDelta, bLooping, CurrentTime, Instance.PreviousMarker, Instance.NextMarker, MirrorTable);
 
 	UE_LOG(LogAnimMarkerSync, Log, TEXT("Follower (%s) (TickByMarker) PreviousTime(%0.2f) CurrentTime(%0.2f) MoveDelta(%0.2f) Looping(%d) %s"), *GetName(), OutPreviousTime, CurrentTime, MoveDelta, bLooping ? 1 : 0, *MarkerContext.ToString());
@@ -643,6 +653,7 @@ void UAnimSequenceBase::TickByMarkerAsFollower(FMarkerTickRecord &Instance, FMar
 
 void UAnimSequenceBase::TickByMarkerAsLeader(FMarkerTickRecord& Instance, FMarkerTickContext& MarkerContext, float& CurrentTime, float& OutPreviousTime, const float MoveDelta, const bool bLooping, const UMirrorDataTable* MirrorTable) const
 {
+	// Re-compute marker indices since the asset's tick record is invalid. Get previous and next markers.
 	if (!Instance.IsValid(bLooping))
 	{
 		if (MarkerContext.IsMarkerSyncStartValid())
@@ -653,15 +664,18 @@ void UAnimSequenceBase::TickByMarkerAsLeader(FMarkerTickRecord& Instance, FMarke
 		{
 			GetMarkerIndicesForTime(CurrentTime, bLooping, MarkerContext.GetValidMarkerNames(), Instance.PreviousMarker, Instance.NextMarker);
 		}
-		
 	}
 
+	// Store the sync anim position BEFORE the asset has being ticked.
 	MarkerContext.SetMarkerSyncStartPosition(GetMarkerSyncPositionFromMarkerIndicies(Instance.PreviousMarker.MarkerIndex, Instance.NextMarker.MarkerIndex, CurrentTime, MirrorTable));
 
+	// Keep track the asset's previous time, since its current time is about to be modified.
 	OutPreviousTime = CurrentTime;
 
+	// Advance as leader.
 	AdvanceMarkerPhaseAsLeader(bLooping, MoveDelta, MarkerContext.GetValidMarkerNames(), CurrentTime, Instance.PreviousMarker, Instance.NextMarker, MarkerContext.MarkersPassedThisTick, MirrorTable);
 
+	// Store the sync anim position AFTER the asset has being ticked.
 	MarkerContext.SetMarkerSyncEndPosition(GetMarkerSyncPositionFromMarkerIndicies(Instance.PreviousMarker.MarkerIndex, Instance.NextMarker.MarkerIndex, CurrentTime, MirrorTable));
 
 	UE_LOG(LogAnimMarkerSync, Log, TEXT("Leader (%s) (TickByMarker) PreviousTime(%0.2f) CurrentTime(%0.2f) MoveDelta(%0.2f) Looping(%d) %s"), *GetName(), OutPreviousTime, CurrentTime, MoveDelta, bLooping ? 1 : 0, *MarkerContext.ToString());
diff --git a/Engine/Source/Runtime/Engine/Private/Animation/AnimSync.cpp b/Engine/Source/Runtime/Engine/Private/Animation/AnimSync.cpp
--- a/Engine/Source/Runtime/Engine/Private/Animation/AnimSync.cpp
+++ b/Engine/Source/Runtime/Engine/Private/Animation/AnimSync.cpp
@@ -48,10 +48,16 @@ void FAnimSync::AddTickRecord(const FAnimTickRecord& InTickRecord, const FAnimSy
 {
 	if (InSyncParams.GroupName != NAME_None)
 	{
+		// Get specified sync group for the animation instance. 
 		FSyncGroupMap& SyncGroupMap = SyncGroupMaps[GetSyncGroupWriteIndex()];
 		FAnimGroupInstance& SyncGroupInstance = SyncGroupMap.FindOrAdd(InSyncParams.GroupName);
+
+		// Add animation instance's tick record to the sync group. 
 		SyncGroupInstance.ActivePlayers.Add(InTickRecord);
 		SyncGroupInstance.ActivePlayers.Top().MirrorDataTable = MirrorDataTable;
+
+		// Set leader score for the tick record we just added, and ensure there is only one montage per group.
+		// CanBeLeader or TransitionLeader Score (BlendWeight) < Always Leader Score (2.0) < Montage Leader Score (3.0).
 		SyncGroupInstance.TestTickRecordForLeadership(InSyncParams.Role);
 	}
 	else
@@ -72,6 +78,7 @@ void FAnimSync::TickAssetPlayerInstances(FAnimInstanceProxy& InProxy, float InDe
 
 	SCOPE_CYCLE_COUNTER(STAT_TickAssetPlayerInstances);
 
+	// Helper function to accumulate the root motion of a tick record.
 	auto AccumulateRootMotion = [&InProxy](FAnimTickRecord& TickRecord, FAnimAssetTickContext& TickContext)
 	{
 		if (TickRecord.MirrorDataTable)
@@ -95,190 +102,251 @@ void FAnimSync::TickAssetPlayerInstances(FAnimInstanceProxy& InProxy, float InDe
 
 	const ERootMotionMode::Type RootMotionMode = InProxy.GetRootMotionMode();
 
-	// Handle all players inside sync groups
-	FSyncGroupMap& SyncGroupMap = SyncGroupMaps[GetSyncGroupWriteIndex()];
-	const FSyncGroupMap& PreviousSyncGroupMap = SyncGroupMaps[GetSyncGroupReadIndex()];
-	TArray<FAnimTickRecord>& UngroupedActivePlayers = UngroupedActivePlayerArrays[GetSyncGroupWriteIndex()];
-	const TArray<FAnimTickRecord>& PreviousUngroupedActivePlayers = UngroupedActivePlayerArrays[GetSyncGroupReadIndex()];
-
-	for (auto& SyncGroupPair : SyncGroupMap)
+	// Tick grouped animation instances.
 	{
-		FAnimGroupInstance& SyncGroup = SyncGroupPair.Value;
-	
-		if (SyncGroup.ActivePlayers.Num() > 0)
+		FSyncGroupMap& SyncGroupMap = SyncGroupMaps[GetSyncGroupWriteIndex()];
+		const FSyncGroupMap& PreviousSyncGroupMap = SyncGroupMaps[GetSyncGroupReadIndex()];
+		
+		// Handle animation players inside group instances.
+		for (TTuple<FName, FAnimGroupInstance>& SyncGroupPair : SyncGroupMap)
 		{
-			const FAnimGroupInstance* PreviousGroup = PreviousSyncGroupMap.Find(SyncGroupPair.Key);
-			SyncGroup.Prepare(PreviousGroup);
+			FAnimGroupInstance& SyncGroup = SyncGroupPair.Value;
 
-			UE_LOG(LogAnimMarkerSync, Log, TEXT("Ticking Group [%s] GroupLeader [%d]"), *SyncGroupPair.Key.ToString(), SyncGroup.GroupLeaderIndex);
+			if (SyncGroup.ActivePlayers.Num() > 0)
+			{
+				const FAnimGroupInstance* PreviousGroup = PreviousSyncGroupMap.Find(SyncGroupPair.Key);
 
-			const bool bOnlyOneAnimationInGroup = SyncGroup.ActivePlayers.Num() == 1;
+				// Prepare sync group.
+				// This is where our asset players in the group are sorted based on a leader score, and tick record are reset if needed.
+				// Additionally, any markers that are not common to all animations in group are removed.
+				SyncGroup.Prepare(PreviousGroup);
 
-			// Tick the group leader
-			FAnimAssetTickContext TickContext(InDeltaSeconds, RootMotionMode, bOnlyOneAnimationInGroup, SyncGroup.ValidMarkers);
-			if (PreviousGroup)
-			{
-				// Initialize the anim position ratio from the previous frame's final anim position ratio in case we're not using marker based sync
-				TickContext.SetPreviousAnimationPositionRatio(PreviousGroup->AnimLengthRatio);
-				TickContext.SetAnimationPositionRatio(PreviousGroup->AnimLengthRatio);
-
-				const FMarkerSyncAnimPosition& EndPosition = PreviousGroup->MarkerTickContext.GetMarkerSyncEndPosition();
-				if ( EndPosition.IsValid() &&
-						(EndPosition.PreviousMarkerName == NAME_None || SyncGroup.ValidMarkers.Contains(EndPosition.PreviousMarkerName)) &&
-						(EndPosition.NextMarkerName == NAME_None || SyncGroup.ValidMarkers.Contains(EndPosition.NextMarkerName)))
+				UE_LOG(LogAnimMarkerSync, Log, TEXT("Ticking Group [%s] GroupLeader [%d]"), *SyncGroupPair.Key.ToString(), SyncGroup.GroupLeaderIndex);
+
+				// Determine if we have a single animation context.
+				const bool bOnlyOneAnimationInGroup = SyncGroup.ActivePlayers.Num() == 1;
+
+				// Animation context for group. (Modified by leader, read by followers)
+				FAnimAssetTickContext TickContext(InDeltaSeconds, RootMotionMode, bOnlyOneAnimationInGroup, SyncGroup.ValidMarkers);
+
+				// Initialize group context using previous update's group information.
+				if (PreviousGroup)
 				{
-					TickContext.MarkerTickContext.SetMarkerSyncStartPosition(EndPosition);
+					// Initialize the anim position ratio from the previous frame's final anim position ratio in case we're not using marker based sync.
+					TickContext.SetPreviousAnimationPositionRatio(PreviousGroup->AnimLengthRatio);
+					TickContext.SetAnimationPositionRatio(PreviousGroup->AnimLengthRatio);
+					
+					// Continue from where previous group marker-based sync left off, use its end sync position as our new start sync position. 
+					const FMarkerSyncAnimPosition& EndPosition = PreviousGroup->MarkerTickContext.GetMarkerSyncEndPosition();
+					if (EndPosition.IsValid() && (EndPosition.PreviousMarkerName == NAME_None || SyncGroup.ValidMarkers.Contains(EndPosition.PreviousMarkerName))
+					                          && (EndPosition.NextMarkerName == NAME_None || SyncGroup.ValidMarkers.Contains(EndPosition.NextMarkerName)))
+					{
+						TickContext.MarkerTickContext.SetMarkerSyncStartPosition(EndPosition);
+					}
 				}
-			}
 
 #if DO_CHECK
-			//For debugging UE-54705
-			FName InitialMarkerPrevious = TickContext.MarkerTickContext.GetMarkerSyncStartPosition().PreviousMarkerName;
-			FName InitialMarkerEnd = TickContext.MarkerTickContext.GetMarkerSyncStartPosition().NextMarkerName;
-			const bool bIsLeaderRecordValidPre = SyncGroup.ActivePlayers[0].MarkerTickRecord->IsValid(SyncGroup.ActivePlayers[0].bLooping);
-			FMarkerTickRecord LeaderPreMarkerTickRecord = *SyncGroup.ActivePlayers[0].MarkerTickRecord;
+				//For debugging UE-54705
+				FName InitialMarkerPrevious = TickContext.MarkerTickContext.GetMarkerSyncStartPosition().PreviousMarkerName;
+				FName InitialMarkerEnd = TickContext.MarkerTickContext.GetMarkerSyncStartPosition().NextMarkerName;
+				const bool bIsLeaderRecordValidPre = SyncGroup.ActivePlayers[0].MarkerTickRecord->IsValid(SyncGroup.ActivePlayers[0].bLooping);
+				FMarkerTickRecord LeaderPreMarkerTickRecord = *SyncGroup.ActivePlayers[0].MarkerTickRecord;
 #endif
 
-			// initialize to invalidate first
-			ensureMsgf(SyncGroup.GroupLeaderIndex == INDEX_NONE, TEXT("SyncGroup %s had a non -1 group leader index of %d in asset %s"), *SyncGroupPair.Key.ToString(), SyncGroup.GroupLeaderIndex, *GetNameSafe(InProxy.GetAnimInstanceObject()));
-			int32 GroupLeaderIndex = 0;
-			for (; GroupLeaderIndex < SyncGroup.ActivePlayers.Num(); ++GroupLeaderIndex)
-			{
-				FAnimTickRecord& GroupLeader = SyncGroup.ActivePlayers[GroupLeaderIndex];
-				// if it has leader score
-				SCOPE_CYCLE_COUNTER(STAT_TickAssetPlayerInstance);
-				FScopeCycleCounterUObject Scope(GroupLeader.SourceAsset);
+				// Initialize with sync group leader being invalid first.
+				ensureMsgf(SyncGroup.GroupLeaderIndex == INDEX_NONE, TEXT("SyncGroup %s had a non -1 group leader index of %d in asset %s"), *SyncGroupPair.Key.ToString(), SyncGroup.GroupLeaderIndex, *GetNameSafe(InProxy.GetAnimInstanceObject()));
 				
-				// Inertialization was requested therefore we resync to previous leader, if needed.
-				if (GroupLeader.bRequestedInertialization)
+				// Tick group leader, and try to find a new one in case the current one has an invalid position.
+				int32 GroupLeaderIndex = 0;
+				for (; GroupLeaderIndex < SyncGroup.ActivePlayers.Num(); ++GroupLeaderIndex)
 				{
-					// Ensure we have a previous group leader to sync to otherwise we play normally. 
-					if (PreviousGroup && !PreviousGroup->ActivePlayers.IsEmpty())
+					FAnimTickRecord& GroupLeader = SyncGroup.ActivePlayers[GroupLeaderIndex];
+					// if it has leader score
+					SCOPE_CYCLE_COUNTER(STAT_TickAssetPlayerInstance);
+					FScopeCycleCounterUObject Scope(GroupLeader.SourceAsset);
+
+					// Inertialization was requested therefore we resync to previous leader, if needed.
+                    if (GroupLeader.bRequestedInertialization)
+                    {
+                    	// Ensure we have a previous group leader to sync to otherwise we play normally. 
+                    	if (PreviousGroup && !PreviousGroup->ActivePlayers.IsEmpty())
+                    	{
+                    		// Only need to resync when using length based syncing since we initialized the tick context with the previous group's sync end position
+                    		// and that will take care of marker based syncing.
+                    		if (!TickContext.CanUseMarkerPosition() || !(PreviousGroup->ValidMarkers.Num() > 0))
+                    		{
+                    			const FAnimTickRecord & PreviousGroupLeader = PreviousGroup->ActivePlayers[PreviousGroup->GroupLeaderIndex];
+                    			const bool bShouldResyncToSyncGroup = PreviousGroupLeader.LeaderScore >= GroupLeader.LeaderScore;
+                
+                    			// Sync to previous group leader if we have a lower score than them.
+                    			TickContext.SetResyncToSyncGroup(bShouldResyncToSyncGroup);
+                    		}
+                    	}
+                    }
+
+					// Tick group leader's asset.
+					TickContext.MarkerTickContext.MarkersPassedThisTick.Reset();
+					TickContext.RootMotionMovementParams.Clear();
+					GroupLeader.SourceAsset->TickAssetPlayer(GroupLeader, InProxy.NotifyQueue, TickContext);
+
+					// Accumulate root motion if possible.
+					if (RootMotionMode == ERootMotionMode::RootMotionFromEverything && TickContext.RootMotionMovementParams.bHasRootMotion)
 					{
-						// Only need to resync when using length based syncing since we initialized the tick context with the previous group's sync end position
-						// and that will take care of marker based syncing.
-						if (!TickContext.CanUseMarkerPosition() || !(PreviousGroup->ValidMarkers.Num() > 0))
-						{
-							const FAnimTickRecord & PreviousGroupLeader = PreviousGroup->ActivePlayers[PreviousGroup->GroupLeaderIndex];
-							const bool bShouldResyncToSyncGroup = PreviousGroupLeader.LeaderScore >= GroupLeader.LeaderScore;
-
-							// Sync to previous group leader if we have a lower score than them.
-							TickContext.SetResyncToSyncGroup(bShouldResyncToSyncGroup);
-						}
+						AccumulateRootMotion(GroupLeader, TickContext);
 					}
-				}
-				
-				TickContext.MarkerTickContext.MarkersPassedThisTick.Reset();
-				TickContext.RootMotionMovementParams.Clear();
-				GroupLeader.SourceAsset->TickAssetPlayer(GroupLeader, InProxy.NotifyQueue, TickContext);
 
-				if (RootMotionMode == ERootMotionMode::RootMotionFromEverything && TickContext.RootMotionMovementParams.bHasRootMotion)
-				{
-					AccumulateRootMotion(GroupLeader, TickContext);
-				}
+					// If we're not using marker based sync, we don't care, update and get out.
+					if (TickContext.CanUseMarkerPosition() == false)
+					{
+						SyncGroup.PreviousAnimLengthRatio = TickContext.GetPreviousAnimationPositionRatio();
+						SyncGroup.AnimLengthRatio = TickContext.GetAnimationPositionRatio();
+						SyncGroup.GroupLeaderIndex = GroupLeaderIndex;
+						break;
+					}
+					// Otherwise, the new position should contain the valid position for end, otherwise, we don't know where to sync to.
+					else if (TickContext.MarkerTickContext.IsMarkerSyncEndValid())
+					{
+						// If this leader contains correct position, break
+						SyncGroup.PreviousAnimLengthRatio = TickContext.GetPreviousAnimationPositionRatio();
+						SyncGroup.AnimLengthRatio = TickContext.GetAnimationPositionRatio();
+						SyncGroup.MarkerTickContext = TickContext.MarkerTickContext;
+						SyncGroup.GroupLeaderIndex = GroupLeaderIndex;
+						UE_LOG(LogAnimMarkerSync, Log, TEXT("Previous Sync Group Marker Tick Context :\n%s"), *SyncGroup.MarkerTickContext.ToString());
+						UE_LOG(LogAnimMarkerSync, Log, TEXT("New Sync Group Marker Tick Context :\n%s"), *TickContext.MarkerTickContext.ToString());
+						break;
+					}
+					// We have an invalid sync end position, keep searching for a valid leader.
+					else
+					{
+						SyncGroup.PreviousAnimLengthRatio = TickContext.GetPreviousAnimationPositionRatio();
+						SyncGroup.AnimLengthRatio = TickContext.GetAnimationPositionRatio();
+						SyncGroup.GroupLeaderIndex = GroupLeaderIndex;
+						UE_LOG(LogAnimMarkerSync, Log, TEXT("Invalid position from Leader %d. Trying next leader"), GroupLeaderIndex);
+					}
+				} 
 
-				// if we're not using marker based sync, we don't care, get out
-				if (TickContext.CanUseMarkerPosition() == false)
-				{
-					SyncGroup.PreviousAnimLengthRatio = TickContext.GetPreviousAnimationPositionRatio();
-					SyncGroup.AnimLengthRatio = TickContext.GetAnimationPositionRatio();
-					SyncGroup.GroupLeaderIndex = GroupLeaderIndex;
-					break;
-				}
-				// otherwise, the new position should contain the valid position for end, otherwise, we don't know where to sync to
-				else if (TickContext.MarkerTickContext.IsMarkerSyncEndValid())
-				{
-					// if this leader contains correct position, break
-					SyncGroup.PreviousAnimLengthRatio = TickContext.GetPreviousAnimationPositionRatio();
-					SyncGroup.AnimLengthRatio = TickContext.GetAnimationPositionRatio();
-					SyncGroup.MarkerTickContext = TickContext.MarkerTickContext;
-					SyncGroup.GroupLeaderIndex = GroupLeaderIndex;
-					UE_LOG(LogAnimMarkerSync, Log, TEXT("Previous Sync Group Marker Tick Context :\n%s"), *SyncGroup.MarkerTickContext.ToString());
-					UE_LOG(LogAnimMarkerSync, Log, TEXT("New Sync Group Marker Tick Context :\n%s"), *TickContext.MarkerTickContext.ToString());
-					break;
-				}
-				else
-				{
-					SyncGroup.PreviousAnimLengthRatio = TickContext.GetPreviousAnimationPositionRatio();
-					SyncGroup.AnimLengthRatio = TickContext.GetAnimationPositionRatio();
-					SyncGroup.GroupLeaderIndex = GroupLeaderIndex;
-					UE_LOG(LogAnimMarkerSync, Log, TEXT("Invalid position from Leader %d. Trying next leader"), GroupLeaderIndex);
-				}
-			} 
+				check(SyncGroup.GroupLeaderIndex != INDEX_NONE);
+			
+				// By this point we have found a valid group leader candidate.
 
-			check(SyncGroup.GroupLeaderIndex != INDEX_NONE);
-			// we found leader
-			SyncGroup.Finalize(PreviousGroup);
+				// Finalize sync group.
+				// This is where the followers tick records are reset if the needed. 
+				SyncGroup.Finalize(PreviousGroup);
 
-			if (TickContext.CanUseMarkerPosition())
-			{
-				const FMarkerSyncAnimPosition& MarkerStart = TickContext.MarkerTickContext.GetMarkerSyncStartPosition();
-				FName SyncGroupName = SyncGroupPair.Key;
-				FAnimTickRecord& GroupLeader = SyncGroup.ActivePlayers[SyncGroup.GroupLeaderIndex];
-				FString LeaderAnimName = GroupLeader.SourceAsset->GetName();
+				// Ensure group leader's markers are valid, otherwise invalidate marker-based syncing for followers.
+				if (TickContext.CanUseMarkerPosition())
+				{
+					const FMarkerSyncAnimPosition& MarkerStart = TickContext.MarkerTickContext.GetMarkerSyncStartPosition();
+					FName SyncGroupName = SyncGroupPair.Key;
+					FAnimTickRecord& GroupLeader = SyncGroup.ActivePlayers[SyncGroup.GroupLeaderIndex];
+					FString LeaderAnimName = GroupLeader.SourceAsset->GetName();
 
-				//  Updated logic in search for cause of UE-54705
-				const bool bStartMarkerValid = (MarkerStart.PreviousMarkerName == NAME_None) || SyncGroup.ValidMarkers.Contains(MarkerStart.PreviousMarkerName);
-				const bool bEndMarkerValid = (MarkerStart.NextMarkerName == NAME_None) || SyncGroup.ValidMarkers.Contains(MarkerStart.NextMarkerName);
+					//  Updated logic in search for cause of UE-54705
+					const bool bStartMarkerValid = (MarkerStart.PreviousMarkerName == NAME_None) || SyncGroup.ValidMarkers.Contains(MarkerStart.PreviousMarkerName);
+					const bool bEndMarkerValid = (MarkerStart.NextMarkerName == NAME_None) || SyncGroup.ValidMarkers.Contains(MarkerStart.NextMarkerName);
 
-				if (!bStartMarkerValid)
-				{
-#if DO_CHECK
-					FString ErrorMsg = FString(TEXT("Prev Marker name not valid for sync group.\n"));
-					ErrorMsg += FString::Format(TEXT("\tMarker {0} : SyncGroupName {1} : Leader {2}\n"), { MarkerStart.PreviousMarkerName.ToString(), SyncGroupName.ToString(), LeaderAnimName });
-					ErrorMsg += FString::Format(TEXT("\tInitalPrev {0} : InitialNext {1} : GroupLeaderIndex {2}\n"), { InitialMarkerPrevious.ToString(), InitialMarkerEnd.ToString(), GroupLeaderIndex });
-					ErrorMsg += FString::Format(TEXT("\tLeader (0 index) was originally valid: {0} | Record: {1}\n"), { bIsLeaderRecordValidPre, LeaderPreMarkerTickRecord.ToString() });
-					ErrorMsg += FString::Format(TEXT("\t Valid Markers : {0}\n"), { SyncGroup.ValidMarkers.Num() });
-					for (int32 MarkerIndex = 0; MarkerIndex < SyncGroup.ValidMarkers.Num(); ++MarkerIndex)
+					if (!bStartMarkerValid)
 					{
-						ErrorMsg += FString::Format(TEXT("\t\t{0}) '{1}'\n"), {MarkerIndex, SyncGroup.ValidMarkers[MarkerIndex].ToString()});
+#if DO_CHECK
+						FString ErrorMsg = FString(TEXT("Prev Marker name not valid for sync group.\n"));
+						ErrorMsg += FString::Format(TEXT("\tMarker {0} : SyncGroupName {1} : Leader {2}\n"), { MarkerStart.PreviousMarkerName.ToString(), SyncGroupName.ToString(), LeaderAnimName });
+						ErrorMsg += FString::Format(TEXT("\tInitalPrev {0} : InitialNext {1} : GroupLeaderIndex {2}\n"), { InitialMarkerPrevious.ToString(), InitialMarkerEnd.ToString(), GroupLeaderIndex });
+						ErrorMsg += FString::Format(TEXT("\tLeader (0 index) was originally valid: {0} | Record: {1}\n"), { bIsLeaderRecordValidPre, LeaderPreMarkerTickRecord.ToString() });
+						ErrorMsg += FString::Format(TEXT("\t Valid Markers : {0}\n"), { SyncGroup.ValidMarkers.Num() });
+						for (int32 MarkerIndex = 0; MarkerIndex < SyncGroup.ValidMarkers.Num(); ++MarkerIndex)
+						{
+							ErrorMsg += FString::Format(TEXT("\t\t{0}) '{1}'\n"), {MarkerIndex, SyncGroup.ValidMarkers[MarkerIndex].ToString()});
+						}
+						ensureMsgf(false, TEXT("%s"), *ErrorMsg);
+#endif
+						TickContext.InvalidateMarkerSync();
 					}
-					ensureMsgf(false, TEXT("%s"), *ErrorMsg);
+					else if (!bEndMarkerValid)
+					{
+#if DO_CHECK
+						FString ErrorMsg = FString(TEXT("Next Marker name not valid for sync group.\n"));
+						ErrorMsg += FString::Format(TEXT("\tMarker {0} : SyncGroupName {1} : Leader {2}\n"), { MarkerStart.NextMarkerName.ToString(), SyncGroupName.ToString(), LeaderAnimName });
+						ErrorMsg += FString::Format(TEXT("\tInitalPrev {0} : InitialNext {1} : GroupLeaderIndex {2}\n"), { InitialMarkerPrevious.ToString(), InitialMarkerEnd.ToString(), GroupLeaderIndex });
+						ErrorMsg += FString::Format(TEXT("\tLeader (0 index) was originally valid: {0} | Record: {1}\n"), { bIsLeaderRecordValidPre, LeaderPreMarkerTickRecord.ToString() });
+						ErrorMsg += FString::Format(TEXT("\t Valid Markers : {0}\n"), { SyncGroup.ValidMarkers.Num() });
+						for (int32 MarkerIndex = 0; MarkerIndex < SyncGroup.ValidMarkers.Num(); ++MarkerIndex)
+						{
+							ErrorMsg += FString::Format(TEXT("\t\t{0}) '{1}'\n"), { MarkerIndex, SyncGroup.ValidMarkers[MarkerIndex].ToString() });
+						}
+						ensureMsgf(false, TEXT("%s"), *ErrorMsg);
 #endif
-					TickContext.InvalidateMarkerSync();
+						TickContext.InvalidateMarkerSync();
+					}
 				}
-				else if (!bEndMarkerValid)
+
+				// Update everything else to follow the leader, if there is more followers
+				if (SyncGroup.ActivePlayers.Num() > GroupLeaderIndex + 1)
 				{
-#if DO_CHECK
-					FString ErrorMsg = FString(TEXT("Next Marker name not valid for sync group.\n"));
-					ErrorMsg += FString::Format(TEXT("\tMarker {0} : SyncGroupName {1} : Leader {2}\n"), { MarkerStart.NextMarkerName.ToString(), SyncGroupName.ToString(), LeaderAnimName });
-					ErrorMsg += FString::Format(TEXT("\tInitalPrev {0} : InitialNext {1} : GroupLeaderIndex {2}\n"), { InitialMarkerPrevious.ToString(), InitialMarkerEnd.ToString(), GroupLeaderIndex });
-					ErrorMsg += FString::Format(TEXT("\tLeader (0 index) was originally valid: {0} | Record: {1}\n"), { bIsLeaderRecordValidPre, LeaderPreMarkerTickRecord.ToString() });
-					ErrorMsg += FString::Format(TEXT("\t Valid Markers : {0}\n"), { SyncGroup.ValidMarkers.Num() });
-					for (int32 MarkerIndex = 0; MarkerIndex < SyncGroup.ValidMarkers.Num(); ++MarkerIndex)
+					// if we don't have a good leader, no reason to convert to follower
+                    // tick as leader
+					TickContext.ConvertToFollower();
+
+					for (int32 TickIndex = GroupLeaderIndex + 1; TickIndex < SyncGroup.ActivePlayers.Num(); ++TickIndex)
 					{
-						ErrorMsg += FString::Format(TEXT("\t\t{0}) '{1}'\n"), { MarkerIndex, SyncGroup.ValidMarkers[MarkerIndex].ToString() });
+						// Tick follower's asset player.
+						FAnimTickRecord& AssetPlayer = SyncGroup.ActivePlayers[TickIndex];
+						{
+							SCOPE_CYCLE_COUNTER(STAT_TickAssetPlayerInstance);
+							FScopeCycleCounterUObject Scope(AssetPlayer.SourceAsset);
+							
+							TickContext.SetResyncToSyncGroup(AssetPlayer.bRequestedInertialization);
+							TickContext.RootMotionMovementParams.Clear();
+							
+							AssetPlayer.SourceAsset->TickAssetPlayer(AssetPlayer, InProxy.NotifyQueue, TickContext);
+						}
+					
+						// Accumulate root motion if possible.
+						if (RootMotionMode == ERootMotionMode::RootMotionFromEverything && TickContext.RootMotionMovementParams.bHasRootMotion)
+						{
+							AccumulateRootMotion(AssetPlayer, TickContext);
+						}
 					}
-					ensureMsgf(false, TEXT("%s"), *ErrorMsg);
-#endif
-					TickContext.InvalidateMarkerSync();
 				}
+
+#if ANIM_TRACE_ENABLED
+				for(const FPassedMarker& PassedMarker : TickContext.MarkerTickContext.MarkersPassedThisTick)
+				{
+					TRACE_ANIM_SYNC_MARKER(CastChecked<UAnimInstance>(InProxy.GetAnimInstanceObject()), PassedMarker);
+				}
+#endif
 			}
+		}
+	}
+
+	// Tick ungrouped animation instances.
+	{
+		TArray<FAnimTickRecord>& UngroupedActivePlayers = UngroupedActivePlayerArrays[GetSyncGroupWriteIndex()];
+		const TArray<FAnimTickRecord>& PreviousUngroupedActivePlayers = UngroupedActivePlayerArrays[GetSyncGroupReadIndex()];
+		
+		// Handle ungrouped animation asset players.
+		for (int32 TickIndex = 0; TickIndex < UngroupedActivePlayers.Num(); ++TickIndex)
+		{
+			FAnimTickRecord& AssetPlayerToTick = UngroupedActivePlayers[TickIndex];
+
+			// Get marker names.
+			const TArray<FName>* UniqueNames = AssetPlayerToTick.SourceAsset->GetUniqueMarkerNames();
+			const TArray<FName>& ValidMarkers = UniqueNames ? *UniqueNames : FMarkerTickContext::DefaultMarkerNames;
+
+			// Single animation context used for ticking.
+			const bool bOnlyOneAnimationInGroup = true;
+			FAnimAssetTickContext TickContext(InDeltaSeconds, RootMotionMode, bOnlyOneAnimationInGroup, ValidMarkers);
 
-			// Update everything else to follow the leader, if there is more followers
-			if (SyncGroup.ActivePlayers.Num() > GroupLeaderIndex + 1)
+			// Tick asset player.
 			{
-				// if we don't have a good leader, no reason to convert to follower
-				// tick as leader
-				TickContext.ConvertToFollower();
+				SCOPE_CYCLE_COUNTER(STAT_TickAssetPlayerInstance);
+				FScopeCycleCounterUObject Scope(AssetPlayerToTick.SourceAsset);
+				AssetPlayerToTick.SourceAsset->TickAssetPlayer(AssetPlayerToTick, InProxy.NotifyQueue, TickContext);
+			}
 
-				for (int32 TickIndex = GroupLeaderIndex + 1; TickIndex < SyncGroup.ActivePlayers.Num(); ++TickIndex)
-				{
-					FAnimTickRecord& AssetPlayer = SyncGroup.ActivePlayers[TickIndex];
-					{
-						SCOPE_CYCLE_COUNTER(STAT_TickAssetPlayerInstance);
-						FScopeCycleCounterUObject Scope(AssetPlayer.SourceAsset);
-						
-						TickContext.SetResyncToSyncGroup(AssetPlayer.bRequestedInertialization);
-						TickContext.RootMotionMovementParams.Clear();
-						
-						AssetPlayer.SourceAsset->TickAssetPlayer(AssetPlayer, InProxy.NotifyQueue, TickContext);
-					}
-					if (RootMotionMode == ERootMotionMode::RootMotionFromEverything && TickContext.RootMotionMovementParams.bHasRootMotion)
-					{
-						AccumulateRootMotion(AssetPlayer, TickContext);
-					}
-				}
+			// Accumulate root motion if possible.
+			if (RootMotionMode == ERootMotionMode::RootMotionFromEverything && TickContext.RootMotionMovementParams.bHasRootMotion)
+			{
+				AccumulateRootMotion(AssetPlayerToTick, TickContext);
 			}
 
 #if ANIM_TRACE_ENABLED
@@ -290,34 +358,7 @@ void FAnimSync::TickAssetPlayerInstances(FAnimInstanceProxy& InProxy, float InDe
 		}
 	}
 	
-	// Handle the remaining ungrouped animation players
-	for (int32 TickIndex = 0; TickIndex < UngroupedActivePlayers.Num(); ++TickIndex)
-	{
-		FAnimTickRecord& AssetPlayerToTick = UngroupedActivePlayers[TickIndex];
-		const TArray<FName>* UniqueNames = AssetPlayerToTick.SourceAsset->GetUniqueMarkerNames();
-		const TArray<FName>& ValidMarkers = UniqueNames ? *UniqueNames : FMarkerTickContext::DefaultMarkerNames;
-
-		const bool bOnlyOneAnimationInGroup = true;
-		FAnimAssetTickContext TickContext(InDeltaSeconds, RootMotionMode, bOnlyOneAnimationInGroup, ValidMarkers);
-		{
-			SCOPE_CYCLE_COUNTER(STAT_TickAssetPlayerInstance);
-			FScopeCycleCounterUObject Scope(AssetPlayerToTick.SourceAsset);
-			AssetPlayerToTick.SourceAsset->TickAssetPlayer(AssetPlayerToTick, InProxy.NotifyQueue, TickContext);
-		}
-		if (RootMotionMode == ERootMotionMode::RootMotionFromEverything && TickContext.RootMotionMovementParams.bHasRootMotion)
-		{
-			AccumulateRootMotion(AssetPlayerToTick, TickContext);
-		}
-
-#if ANIM_TRACE_ENABLED
-		for(const FPassedMarker& PassedMarker : TickContext.MarkerTickContext.MarkersPassedThisTick)
-		{
-			TRACE_ANIM_SYNC_MARKER(CastChecked<UAnimInstance>(InProxy.GetAnimInstanceObject()), PassedMarker);
-		}
-#endif
-	}
-
-	// Flip buffers now we have ticked
+	// Flip buffers now that we have ticked.
 	TickSyncGroupWriteIndex();
 }
 
diff --git a/Engine/Source/Runtime/Engine/Private/Animation/AnimTypes.cpp b/Engine/Source/Runtime/Engine/Private/Animation/AnimTypes.cpp
--- a/Engine/Source/Runtime/Engine/Private/Animation/AnimTypes.cpp
+++ b/Engine/Source/Runtime/Engine/Private/Animation/AnimTypes.cpp
@@ -192,10 +192,10 @@ void FMarkerSyncData::GetMarkerIndicesForTime(float CurrentTime, bool bLooping,
 {
 	const int LoopModStart = bLooping ? -1 : 0;
 	const int LoopModEnd = bLooping ? 2 : 1;
-
-	OutPrevMarker.MarkerIndex = -1;
+	
+	OutPrevMarker.MarkerIndex = MarkerIndexSpecialValues::AnimationBoundary;
 	OutPrevMarker.TimeToMarker = -CurrentTime;
-	OutNextMarker.MarkerIndex = -1;
+	OutNextMarker.MarkerIndex = MarkerIndexSpecialValues::AnimationBoundary;
 	OutNextMarker.TimeToMarker = SequenceLength - CurrentTime;
 
 	for (int32 LoopMod = LoopModStart; LoopMod < LoopModEnd; ++LoopMod)
@@ -220,7 +220,9 @@ void FMarkerSyncData::GetMarkerIndicesForTime(float CurrentTime, bool bLooping,
 				}
 			}
 		}
-		if (OutNextMarker.MarkerIndex != -1)
+
+		// Continue looking for an authored next sync marker.
+		if (OutNextMarker.MarkerIndex != MarkerIndexSpecialValues::AnimationBoundary)
 		{
 			break; // Done
 		}
@@ -237,7 +239,8 @@ FMarkerSyncAnimPosition FMarkerSyncData::GetMarkerSyncPositionFromMarkerIndicies
 	FMarkerSyncAnimPosition SyncPosition;
 	float PrevTime, NextTime;
 
-	if (PrevMarker != -1)
+	// Get previous marker's time and name.
+	if (PrevMarker != MarkerIndexSpecialValues::AnimationBoundary && AuthoredSyncMarkers.IsValidIndex(PrevMarker))
 	{
 		PrevTime = AuthoredSyncMarkers[PrevMarker].Time;
 		SyncPosition.PreviousMarkerName = AuthoredSyncMarkers[PrevMarker].MarkerName;
@@ -255,7 +258,8 @@ FMarkerSyncAnimPosition FMarkerSyncData::GetMarkerSyncPositionFromMarkerIndicies
 		PrevTime = 0.f;
 	}
 
-	if (NextMarker != -1)
+	// Get next marker's time and name.
+	if (NextMarker != MarkerIndexSpecialValues::AnimationBoundary && AuthoredSyncMarkers.IsValidIndex(NextMarker))
 	{
 		NextTime = AuthoredSyncMarkers[NextMarker].Time;
 		SyncPosition.NextMarkerName = AuthoredSyncMarkers[NextMarker].MarkerName;
@@ -274,9 +278,16 @@ FMarkerSyncAnimPosition FMarkerSyncData::GetMarkerSyncPositionFromMarkerIndicies
 	}
 
 	// Account for looping
-	PrevTime = (PrevTime > CurrentTime) ? PrevTime - SequenceLength : PrevTime;
-	NextTime = (NextTime < CurrentTime) ? NextTime + SequenceLength : NextTime;
-
+	if (PrevTime > NextTime)
+	{
+		PrevTime = (PrevTime > CurrentTime) ? PrevTime - SequenceLength : PrevTime;
+		NextTime = (NextTime < CurrentTime) ? NextTime + SequenceLength : NextTime;
+	}
+	else if (PrevTime > CurrentTime)
+	{
+		CurrentTime += SequenceLength;
+	}
+	
 	if (PrevTime == NextTime)
 	{
 		PrevTime -= SequenceLength;
@@ -284,6 +295,7 @@ FMarkerSyncAnimPosition FMarkerSyncData::GetMarkerSyncPositionFromMarkerIndicies
 
 	check(NextTime > PrevTime);
 
+	// Store the encoded current time position as a ratio between markers.
 	SyncPosition.PositionBetweenMarkers = (CurrentTime - PrevTime) / (NextTime - PrevTime);
 	return SyncPosition;
 }
diff --git a/Engine/Source/Runtime/Engine/Private/Animation/BlendSpace.cpp b/Engine/Source/Runtime/Engine/Private/Animation/BlendSpace.cpp
--- a/Engine/Source/Runtime/Engine/Private/Animation/BlendSpace.cpp
+++ b/Engine/Source/Runtime/Engine/Private/Animation/BlendSpace.cpp
@@ -404,11 +404,12 @@ void UBlendSpace::TickAssetPlayer(FAnimTickRecord& Instance, struct FAnimNotifyQ
 					*GetName(), *BlendSpacePosition.ToString(), *FilteredBlendInput.ToString(), FilterMultiplier);
 			}
 
+			// We can use marker-based syncing when a valid sample with sync marker data exists.
 			bool bCanDoMarkerSync = (SampleIndexWithMarkers != INDEX_NONE) && (Context.IsSingleAnimationContext() || (Instance.bCanUseMarkerSync && Context.CanUseMarkerPosition()));
-
+			
 			if (bCanDoMarkerSync)
 			{
-				//Copy previous frame marker data to current frame
+				// Copy previous frame marker data to current frame
 				for (const FBlendSampleData& PrevBlendSampleItem : OldSampleDataList)
 				{
 					for (FBlendSampleData& CurrentBlendSampleItem : SampleDataList)
@@ -459,44 +460,59 @@ void UBlendSpace::TickAssetPlayer(FAnimTickRecord& Instance, struct FAnimNotifyQ
 				// advance current time - blend spaces hold normalized time as when dealing with changing anim length it would be possible to go backwards
 				UE_LOG(LogAnimation, Verbose, TEXT("BlendSpace(%s) - FilteredBlendInput(%s) : AnimLength(%0.5f) "), *GetName(), *FilteredBlendInput.ToString(), NewAnimLength);
 
+				// Set context's data before updating time position.
 				Context.SetPreviousAnimationPositionRatio(NormalizedCurrentTime);
 
+				// Get highest weight sample with sync markers. This will become the leader for all other samples to follow.
 				const int32 HighestMarkerSyncWeightIndex = bCanDoMarkerSync ? FBlendSpaceUtilities::GetHighestWeightMarkerSyncSample(SampleDataList, SampleData) : -1;
 
+				// Skip syncing, fallback to normal ticking.
 				if (HighestMarkerSyncWeightIndex == -1)
 				{
 					bCanDoMarkerSync = false;
 				}
 
+				// Tick as leader using marked based syncing.
 				if (bCanDoMarkerSync)
 				{
-					FBlendSampleData& SampleDataItem = SampleDataList[HighestMarkerSyncWeightIndex];
-					const FBlendSample& Sample = SampleData[SampleDataItem.SampleDataIndex];
+					FBlendSampleData& LeaderSampleData = SampleDataList[HighestMarkerSyncWeightIndex];
+					const FBlendSample& LeaderSample = SampleData[LeaderSampleData.SampleDataIndex];
 
-					if (Sample.Animation)
+					if (LeaderSample.Animation)
 					{
 						bool bResetMarkerDataOnFollowers = false;
+
+						// Invalidate sample followers' tick records if instance doesn't have any valid sync marker data. 
 						if (!Instance.MarkerTickRecord->IsValid(Instance.bLooping))
 						{
-							SampleDataItem.MarkerTickRecord.Reset();
+							LeaderSampleData.MarkerTickRecord.Reset();
+							LeaderSampleData.Time = NormalizedCurrentTime * LeaderSample.Animation->GetPlayLength();
 							bResetMarkerDataOnFollowers = true;
-							SampleDataItem.Time = NormalizedCurrentTime * Sample.Animation->GetPlayLength();
 						}
-						else if (!SampleDataItem.MarkerTickRecord.IsValid(Instance.bLooping) && Context.MarkerTickContext.GetMarkerSyncStartPosition().IsValid())
+						// Re-compute marker indices since the leader sample's tick record is invalid. Get previous and next markers.
+						else if (!LeaderSampleData.MarkerTickRecord.IsValid(Instance.bLooping) && Context.MarkerTickContext.GetMarkerSyncStartPosition().IsValid())
 						{
-							Sample.Animation->GetMarkerIndicesForPosition(Context.MarkerTickContext.GetMarkerSyncStartPosition(), true, SampleDataItem.MarkerTickRecord.PreviousMarker, SampleDataItem.MarkerTickRecord.NextMarker, SampleDataItem.Time, Instance.MirrorDataTable);
+							// TODO: Look into the reason for not passing bLooping variable and just forcing the vale to be true. 
+							LeaderSample.Animation->GetMarkerIndicesForPosition(Context.MarkerTickContext.GetMarkerSyncStartPosition(), true, LeaderSampleData.MarkerTickRecord.PreviousMarker, LeaderSampleData.MarkerTickRecord.NextMarker, LeaderSampleData.Time, Instance.MirrorDataTable);
 						}
 
-						const float NewDeltaTime = Context.GetDeltaTime() * Instance.PlayRateMultiplier * Sample.RateScale * Sample.Animation->RateScale;
+						// Only tick samples if leader sample has any delta time to consume.
+						const float NewDeltaTime = Context.GetDeltaTime() * Instance.PlayRateMultiplier * LeaderSample.RateScale * LeaderSample.Animation->RateScale;
 						if (!FMath::IsNearlyZero(NewDeltaTime))
 						{
 							Context.SetLeaderDelta(NewDeltaTime);
-							Sample.Animation->TickByMarkerAsLeader(SampleDataItem.MarkerTickRecord, Context.MarkerTickContext, SampleDataItem.Time, SampleDataItem.PreviousTime, NewDeltaTime, Instance.bLooping, Instance.MirrorDataTable);
+
+							// Tick leader sample
+							LeaderSample.Animation->TickByMarkerAsLeader(LeaderSampleData.MarkerTickRecord, Context.MarkerTickContext, LeaderSampleData.Time, LeaderSampleData.PreviousTime, NewDeltaTime, Instance.bLooping, Instance.MirrorDataTable);
+
 							check(!Instance.bLooping || Context.MarkerTickContext.IsMarkerSyncStartValid());
+							
+							// Tick all the follower samples
 							TickFollowerSamples(SampleDataList, HighestMarkerSyncWeightIndex, Context, bResetMarkerDataOnFollowers, Instance.bLooping, Instance.MirrorDataTable);
 						}
-						NormalizedCurrentTime = SampleDataItem.Time / Sample.Animation->GetPlayLength();
-						*Instance.MarkerTickRecord = SampleDataItem.MarkerTickRecord;
+						
+						NormalizedCurrentTime = LeaderSampleData.Time / LeaderSample.Animation->GetPlayLength();
+						*Instance.MarkerTickRecord = LeaderSampleData.MarkerTickRecord;
 					}
 				}
 				else
@@ -510,15 +526,18 @@ void UBlendSpace::TickAssetPlayer(FAnimTickRecord& Instance, struct FAnimNotifyQ
 						*GetName(), NormalizedPreviousTime, NormalizedCurrentTime, Instance.DeltaTimeRecord->Delta);
 				}
 
+				// Update time position after it has undergone all side effects.
 				Context.SetAnimationPositionRatio(NormalizedCurrentTime);
 			}
 			else
 			{
+				// Skip syncing if leader doesn't have a valid sync start position.
 				if (!Context.MarkerTickContext.IsMarkerSyncStartValid())
 				{
 					bCanDoMarkerSync = false;
 				}
 
+				// Tick as follower using marked-based syncing.
 				if (bCanDoMarkerSync)
 				{
 					const int32 HighestWeightIndex = FBlendSpaceUtilities::GetHighestWeightSample(SampleDataList);
@@ -527,6 +546,7 @@ void UBlendSpace::TickAssetPlayer(FAnimTickRecord& Instance, struct FAnimNotifyQ
 
 					if (Sample.Animation)
 					{
+						// Only tick samples if sync group leader has any delta time to consume.
 						if (Context.GetDeltaTime() != 0.f)
 						{
 							if (!Instance.MarkerTickRecord->IsValid(Instance.bLooping))
@@ -534,23 +554,27 @@ void UBlendSpace::TickAssetPlayer(FAnimTickRecord& Instance, struct FAnimNotifyQ
 								SampleDataItem.Time = NormalizedCurrentTime * Sample.Animation->GetPlayLength();
 							}
 
+							// Tick all samples as followers
 							TickFollowerSamples(SampleDataList, -1, Context, false, Instance.bLooping, Instance.MirrorDataTable);
 						}
+						
 						*Instance.MarkerTickRecord = SampleDataItem.MarkerTickRecord;
 						NormalizedCurrentTime = SampleDataItem.Time / Sample.Animation->GetPlayLength();
 					}
 				}
 				else
 				{
+					// Fallback to length-based syncing. Match sync group leader position.
 					NormalizedPreviousTime = Context.GetPreviousAnimationPositionRatio();
 					NormalizedCurrentTime = Context.GetAnimationPositionRatio();
+					
 					UE_LOG(LogAnimMarkerSync, Log, 
 						TEXT("Follower (%s) (bCanDoMarkerSync == false) - PreviousTime (%0.2f), CurrentTime (%0.2f), MoveDelta (%0.2f) "), 
 						*GetName(), NormalizedPreviousTime, NormalizedCurrentTime, Instance.DeltaTimeRecord->Delta);
 				}
 			}
 
-			// generate notifies and sets time
+			// Generate notifies and sets time.
 			{
 				FAnimNotifyContext NotifyContext(Instance);
 				float ClampedNormalizedPreviousTime = FMath::Clamp<float>(NormalizedPreviousTime, 0.f, 1.f);
@@ -1978,8 +2002,7 @@ void UBlendSpace::UpdateBlendSpacesUsingAnimSequence(UAnimSequenceBase* Sequence
 
 TArray<FName>* UBlendSpace::GetUniqueMarkerNames()
 {
-	if (SampleIndexWithMarkers != INDEX_NONE
-		&& SampleData.Num() > SampleIndexWithMarkers)
+	if (SampleIndexWithMarkers != INDEX_NONE && SampleData.Num() > SampleIndexWithMarkers)
 	{
 		FBlendSample& BlendSample = SampleData[SampleIndexWithMarkers];
 		if (BlendSample.Animation != nullptr)
diff --git a/Engine/Source/Runtime/Engine/Private/AnimationAsset.cpp b/Engine/Source/Runtime/Engine/Private/AnimationAsset.cpp
--- a/Engine/Source/Runtime/Engine/Private/AnimationAsset.cpp
+++ b/Engine/Source/Runtime/Engine/Private/AnimationAsset.cpp
@@ -36,25 +36,27 @@ void FAnimGroupInstance::TestTickRecordForLeadership(EAnimGroupRole::Type Member
 {
 	check(ActivePlayers.Num() > 0);
 
-	// always set leader score if you have potential to be leader
-	// that way if the top leader fails, we'll continue to search next available leader
-	int32 TestIndex = ActivePlayers.Num() - 1;
+	// Always set leader score if you have potential to be leader
+	// that way if the top leader fails, we'll continue to search next available leader.
+	const int32 TestIndex = ActivePlayers.Num() - 1;
 	FAnimTickRecord& Candidate = ActivePlayers[TestIndex];
 
-	if(Candidate.SourceAsset->IsA<UAnimMontage>())
+	// Handle Montage candidate.
+	if (Candidate.SourceAsset->IsA<UAnimMontage>())
 	{
-		// if the candidate has higher weight
+		// Check if the candidate has a higher weight.
 		if (Candidate.EffectiveBlendWeight > MontageLeaderWeight)
 		{
-			// if this is going to be leader, I'll clean ActivePlayers because we don't sync multi montages
+			// If this is going to be leader, clean ActivePlayers because we don't sync multi montages.
 			const int32 LastIndex = TestIndex - 1;
 			if (LastIndex >= 0)
 			{
-				ActivePlayers.RemoveAt(TestIndex - 1, 1);
+				// Removing based on the last index works because any montage's tick records are all added during Montage Update which happens before the Anim Graph is updated.
+				ActivePlayers.RemoveAt(LastIndex, 1);
 			}
 
-			// at this time, it should only have one
-			ensure(ActivePlayers.Num() == 1);
+			// At this time, it should only have one.
+			check(ActivePlayers.Num() == 1);
 
 			// then override
 			// @note : leader weight doesn't applied WITHIN montages
@@ -66,7 +68,7 @@ void FAnimGroupInstance::TestTickRecordForLeadership(EAnimGroupRole::Type Member
 		{
 			if (TestIndex != 0)
 			{
-				// we delete the later ones because we only have one montage for leader. 
+				// We delete the later ones because we only have one montage for leader 
 				// this can happen if there was already active one with higher weight. 
 				ActivePlayers.RemoveAt(TestIndex, 1);
 			}
@@ -74,6 +76,7 @@ void FAnimGroupInstance::TestTickRecordForLeadership(EAnimGroupRole::Type Member
 
 		ensureAlways(ActivePlayers.Num() == 1);
 	}
+	// Handle Sequence or BlendSpace candidate.
 	else
 	{
 		switch (MembershipType)
@@ -97,8 +100,8 @@ void FAnimGroupInstance::TestTickRecordForLeadership(EAnimGroupRole::Type Member
 
 void FAnimGroupInstance::Finalize(const FAnimGroupInstance* PreviousGroup)
 {
-	if (!PreviousGroup || PreviousGroup->GroupLeaderIndex != GroupLeaderIndex
-		|| (PreviousGroup->MontageLeaderWeight > 0.f && MontageLeaderWeight == 0.f/*if montage disappears, we should reset as well*/))
+	// Reset follower records if the previous group is non-existent, the group leader changed, or the montage leader disappears
+	if (!PreviousGroup || PreviousGroup->GroupLeaderIndex != GroupLeaderIndex || (PreviousGroup->MontageLeaderWeight > 0.f && MontageLeaderWeight == 0.f))
 	{
 		UE_LOG(LogAnimMarkerSync, Log, TEXT("Resetting Marker Sync Groups"));
 
@@ -111,24 +114,32 @@ void FAnimGroupInstance::Finalize(const FAnimGroupInstance* PreviousGroup)
 
 void FAnimGroupInstance::Prepare(const FAnimGroupInstance* PreviousGroup)
 {
+	// Sort asset players by leader score.
 	ActivePlayers.Sort();
 
-	TArray<FName>* MarkerNames = ActivePlayers[0].SourceAsset->GetUniqueMarkerNames();
-	if (MarkerNames)
+	// Prepare group, if leader has any markers.
+	const TArray<FName>* MarkerNames = ActivePlayers[0].SourceAsset->GetUniqueMarkerNames();
+	const bool bLeaderHasMarkers = MarkerNames && !MarkerNames->IsEmpty();
+	
+	if (bLeaderHasMarkers)
 	{
-		// Group leader has markers, off to a good start
+		// Get leader's markers.
 		ValidMarkers = *MarkerNames;
+
+		// Enable marker based syncing for leader and instance group.
 		ActivePlayers[0].bCanUseMarkerSync = true;
 		bCanUseMarkerSync = true;
 
-		//filter markers based on what exists in the other animations
-		for ( int32 ActivePlayerIndex = 0; ActivePlayerIndex < ActivePlayers.Num(); ++ActivePlayerIndex )
+		// Prepare asset player candidates.
+		for (int32 ActivePlayerIndex = 0; ActivePlayerIndex < ActivePlayers.Num(); ++ActivePlayerIndex)
 		{
 			FAnimTickRecord& Candidate = ActivePlayers[ActivePlayerIndex];
 
+			// Reset candidate's marker tick record if needed.
 			if (PreviousGroup)
 			{
 				bool bCandidateFound = false;
+				
 				for (const FAnimTickRecord& PrevRecord : PreviousGroup->ActivePlayers)
 				{
 					if (PrevRecord.MarkerTickRecord == Candidate.MarkerTickRecord)
@@ -136,28 +147,37 @@ void FAnimGroupInstance::Prepare(const FAnimGroupInstance* PreviousGroup)
 						// Found previous record for "us"
 						if (PrevRecord.SourceAsset != Candidate.SourceAsset)
 						{
-							Candidate.MarkerTickRecord->Reset(); // Changed animation, clear our cached data
+							Candidate.MarkerTickRecord->Reset(); // Changed animation, clear our cached data.
 						}
 						bCandidateFound = true;
 						break;
 					}
 				}
+				
 				if (!bCandidateFound)
 				{
-					Candidate.MarkerTickRecord->Reset(); // we weren't active last frame, reset
+					Candidate.MarkerTickRecord->Reset(); // We weren't active last frame, invalidate record.
 				}
 			}
 
+			// Filter follower's markers that are not shared in common with the group's candidate leader.
 			if (ActivePlayerIndex != 0 && ValidMarkers.Num() > 0)
 			{
-				TArray<FName>* PlayerMarkerNames = Candidate.SourceAsset->GetUniqueMarkerNames();
-				if ( PlayerMarkerNames ) // Let anims with no markers set use length scaling sync
+				// Let anims with no markers use length scaling sync.
+				const TArray<FName>* PlayerMarkerNames = Candidate.SourceAsset->GetUniqueMarkerNames();
+				const bool bFollowerHasMarkers = PlayerMarkerNames && !PlayerMarkerNames->IsEmpty();
+				
+				if (bFollowerHasMarkers) 
 				{
+					// Make follower use marker based-syncing.
 					Candidate.bCanUseMarkerSync = true;
-					for ( int32 ValidMarkerIndex = ValidMarkers.Num() - 1; ValidMarkerIndex >= 0; --ValidMarkerIndex )
+
+					// Filter.
+					for (int32 ValidMarkerIndex = ValidMarkers.Num() - 1; ValidMarkerIndex >= 0; --ValidMarkerIndex)
 					{
 						FName& MarkerName = ValidMarkers[ValidMarkerIndex];
-						if ( !PlayerMarkerNames->Contains(MarkerName) )
+						
+						if (!PlayerMarkerNames->Contains(MarkerName))
 						{
 							ValidMarkers.RemoveAtSwap(ValidMarkerIndex, 1, false);
 						}
@@ -166,10 +186,13 @@ void FAnimGroupInstance::Prepare(const FAnimGroupInstance* PreviousGroup)
 			}
 		}
 
+		// Ensure group can use maker based syncing.
 		bCanUseMarkerSync = ValidMarkers.Num() > 0;
-
+		
+		// Alphabetical ordering for markers.
 		ValidMarkers.Sort(FNameLexicalLess());
 
+		// Source marker data changed.
 		if (!PreviousGroup || (ValidMarkers != PreviousGroup->ValidMarkers))
 		{
 			for (int32 InternalActivePlayerIndex = 0; InternalActivePlayerIndex < ActivePlayers.Num(); ++InternalActivePlayerIndex)
@@ -178,14 +201,21 @@ void FAnimGroupInstance::Prepare(const FAnimGroupInstance* PreviousGroup)
 			}
 		}
 	}
-	else
+	
+	// Leader has no markers or all them were filtered out, fallback to length based syncing.
+	if (!bLeaderHasMarkers || !bCanUseMarkerSync)
 	{
-		// Leader has no markers, we can't use SyncMarkers.
+		// We can't use sync markers in sync group.
 		bCanUseMarkerSync = false;
+
+		// Invalidate markers.
 		ValidMarkers.Reset();
+
+		// Ensure tick records do not use marker-based syncing.
 		for (FAnimTickRecord& AnimTickRecord : ActivePlayers)
 		{
 			AnimTickRecord.MarkerTickRecord->Reset();
+			AnimTickRecord.bCanUseMarkerSync = false;
 		}
 	}
 }
