Fix missing cached input in StrataLegacyConversionNode.

diff --git a/Engine/Source/Runtime/Engine/Private/Materials/MaterialExpressions.cpp b/Engine/Source/Runtime/Engine/Private/Materials/MaterialExpressions.cpp
--- a/Engine/Source/Runtime/Engine/Private/Materials/MaterialExpressions.cpp
+++ b/Engine/Source/Runtime/Engine/Private/Materials/MaterialExpressions.cpp
@@ -22869,6 +22869,9 @@ UMaterialExpressionStrataLegacyConversion::UMaterialExpressionStrataLegacyConver
 #if WITH_EDITORONLY_DATA
 	MenuCategories.Add(ConstructorStatics.NAME_Strata);
 #endif
+#if WITH_EDITOR
+	CachedInputs.Add(&ShadingModel);
+#endif
 }
 
 #if WITH_EDITOR
