Fix layered blend per bone not working correctly with per-bone curve filtering

Simple typo - just needed to remove a '!'

diff --git a/Engine/Source/Runtime/Engine/Public/Animation/NamedValueArray.h b/Engine/Source/Runtime/Engine/Public/Animation/NamedValueArray.h
--- a/Engine/Source/Runtime/Engine/Public/Animation/NamedValueArray.h
+++ b/Engine/Source/Runtime/Engine/Public/Animation/NamedValueArray.h
@@ -1,4 +1,4 @@
-﻿// Copyright Epic Games, Inc. All Rights Reserved.
+// Copyright Epic Games, Inc. All Rights Reserved.
 
 #pragma once
 
@@ -629,7 +629,7 @@ struct FNamedValueArrayUtils
 			if(Element0.Name == Element1.Name)
 			{
 				// Elements match so check filter flags to see if it should be removed from InOutValueArray0
-				if(!InPredicate(Element0, Element1))
+				if(InPredicate(Element0, Element1))
 				{
 					InOutValueArray0.Elements.RemoveAt(ElementIndex0, 1, false);
 					++ElementIndex1;
