Fix issue where if an FActorInstanceHandle has been cached, and then a spawned actor from an LWI manager is destroyed, if FetchActor is called on the cached handle it can respawn the actor.

diff --git a/Engine/Source/Runtime/Engine/Classes/GameFramework/LightWeightInstanceManager.h b/Engine/Source/Runtime/Engine/Classes/GameFramework/LightWeightInstanceManager.h
--- a/Engine/Source/Runtime/Engine/Classes/GameFramework/LightWeightInstanceManager.h
+++ b/Engine/Source/Runtime/Engine/Classes/GameFramework/LightWeightInstanceManager.h
@@ -253,6 +253,8 @@ protected:
 	UPROPERTY(Transient)
 	TMap<int32, TObjectPtr<AActor>> Actors;
 
+	TArray<int32> DestroyedActorIndices;
+
 	// list of indices that we are no longer using
 	UPROPERTY(Replicated)
 	TArray<int32> FreeIndices;
diff --git a/Engine/Source/Runtime/Engine/Private/LightWeightInstanceManager.cpp b/Engine/Source/Runtime/Engine/Private/LightWeightInstanceManager.cpp
--- a/Engine/Source/Runtime/Engine/Private/LightWeightInstanceManager.cpp
+++ b/Engine/Source/Runtime/Engine/Private/LightWeightInstanceManager.cpp
@@ -114,6 +114,17 @@ AActor* ALightWeightInstanceManager::ConvertInstanceToActor(const FActorInstance
 		return nullptr;
 	}
 
+	//if we have previously spawned an actor for this index, don't spawn again even though we have no current actor entry.
+	//this may occur if an FActorInstanceHandle was cached before a spawned actor was destroyed and fetched after.
+	if (DestroyedActorIndices.Contains(Handle.GetInstanceIndex()))
+	{
+		UE_LOG(LogLightWeightInstance, Verbose,
+			TEXT("ALightWeightInstanceManager::ConvertInstanceToActor - manager [ %s ] had already spawned an actor for index [ %d ] that was destroyed"),
+			*GetActorNameOrLabel(),
+			Handle.GetInstanceIndex());
+		return nullptr;
+	}
+
 	AActor* NewActor = nullptr;
 	// Only spawn actors on the server so they are replicated to the clients. Otherwise we'll end up with multiples.
 	if (HasAuthority())
@@ -158,6 +169,18 @@ void ALightWeightInstanceManager::OnSpawnedActorDestroyed(AActor* DestroyedActor
 
 	Actors.Remove(DestroyedActorInstanceIndex);
 	DestroyedActor->OnDestroyed.RemoveAll(this);
+
+	if (ensure(!DestroyedActorIndices.Contains(DestroyedActorInstanceIndex)))
+	{
+		DestroyedActorIndices.Add(DestroyedActorInstanceIndex);
+	}
+	else
+	{
+		UE_LOG(LogLightWeightInstance, Error,
+			TEXT("ALightWeightInstanceManager::OnSpawnedActorDestroyed - manager [ %s ] DestroyedActorIndices already has an entry for index [ %d ]"),
+			*GetActorNameOrLabel(),
+			DestroyedActorInstanceIndex);
+	}
 }
 
 int32 ALightWeightInstanceManager::FindIndexForActor(const AActor* InActor) const
