Added a property and cvar to control SparseVolumeTexture DDC caching behavior. Also fixed a crash when very quickly removing and adding a SVT from the streaming manager when there were still streaming requests in the system. Fixed cancel button on the "Caching SparseVolumeTexture frames to DDC" dialog not working.

#rnx

diff --git a/Engine/Source/Runtime/Engine/Classes/SparseVolumeTexture/SparseVolumeTexture.h b/Engine/Source/Runtime/Engine/Classes/SparseVolumeTexture/SparseVolumeTexture.h
--- a/Engine/Source/Runtime/Engine/Classes/SparseVolumeTexture/SparseVolumeTexture.h
+++ b/Engine/Source/Runtime/Engine/Classes/SparseVolumeTexture/SparseVolumeTexture.h
@@ -104,8 +104,10 @@ public:
 	bool RebuildBulkDataFromCacheAsync(const UObject* Owner, bool& bFailed);
 	// Builds all the data from SourceData. Is called by Cache().
 	bool Build(USparseVolumeTextureFrame* Owner);
-	// Cache the built data to/from DDC.
-	void Cache(USparseVolumeTextureFrame* Owner);
+	// Cache the built data to/from DDC. If bLocalCachingOnly is true, the read/write queries will only use the local DDC; otherwise the remote DDC will also be used.
+	void Cache(USparseVolumeTextureFrame* Owner, bool bLocalCachingOnly);
+	// Sets empty default data. This is used when caching/building is canceled but some form of valid data is needed.
+	void SetDefault(EPixelFormat FormatA, EPixelFormat FormatB, const FVector4f& FallbackValueA, const FVector4f& FallbackValueB);
 #endif
 
 private:
@@ -250,7 +252,7 @@ public:
 
 #if WITH_EDITORONLY_DATA
 	// Caches the derived data (FResources) of this frame to/from DDC and ensures that FTextureRenderResources exists.
-	void Cache();
+	void Cache(bool bSkipDDCAndSetResourcesToDefault);
 #endif
 
 	//~ Begin UObject Interface.
@@ -321,18 +323,22 @@ public:
 	UPROPERTY(VisibleAnywhere, Category = "Texture")
 	FVector4f FallbackValueB;
 
-	/** The addressing mode to use for the X axis.								*/
+	// The addressing mode to use for the X axis.
 	UPROPERTY(EditAnywhere, BlueprintReadWrite, Category = "Texture", meta = (DisplayName = "X-axis Tiling Method"), AssetRegistrySearchable, AdvancedDisplay)
 	TEnumAsByte<enum TextureAddress> AddressX;
 
-	/** The addressing mode to use for the Y axis.								*/
+	// The addressing mode to use for the Y axis.
 	UPROPERTY(EditAnywhere, BlueprintReadWrite, Category = "Texture", meta = (DisplayName = "Y-axis Tiling Method"), AssetRegistrySearchable, AdvancedDisplay)
 	TEnumAsByte<enum TextureAddress> AddressY;
 
-	/** The addressing mode to use for the Z axis.								*/
+	// The addressing mode to use for the Z axis.
 	UPROPERTY(EditAnywhere, BlueprintReadWrite, Category = "Texture", meta = (DisplayName = "Z-axis Tiling Method"), AssetRegistrySearchable, AdvancedDisplay)
 	TEnumAsByte<enum TextureAddress> AddressZ;
 
+	// If enabled, the SparseVolumeTexture is only going to use the local DDC. For certain assets it might be reasonable to also use the remote DDC, but for larger assets this will mean long up- and download times.
+	UPROPERTY(EditAnywhere, BlueprintReadWrite, Category = "Texture", meta = (DisplayName = "Local DDC Only"), AssetRegistrySearchable, AdvancedDisplay)
+	bool bLocalDDCOnly = true;
+
 	UStreamableSparseVolumeTexture();
 	virtual ~UStreamableSparseVolumeTexture() = default;
 
diff --git a/Engine/Source/Runtime/Engine/Private/SparseVolumeTexture/SparseVolumeTexture.cpp b/Engine/Source/Runtime/Engine/Private/SparseVolumeTexture/SparseVolumeTexture.cpp
--- a/Engine/Source/Runtime/Engine/Private/SparseVolumeTexture/SparseVolumeTexture.cpp
+++ b/Engine/Source/Runtime/Engine/Private/SparseVolumeTexture/SparseVolumeTexture.cpp
@@ -41,6 +41,14 @@
 
 DEFINE_LOG_CATEGORY_STATIC(LogSparseVolumeTexture, Log, All);
 
+static int32 GSVTRemoteDDCBehavior = 0;
+static FAutoConsoleVariableRef CVarSVTRemoteDDCBehavior(
+	TEXT("r.SparseVolumeTexture.RemoteDDCBehavior"),
+	GSVTRemoteDDCBehavior,
+	TEXT("Controls how SVTs use remote DDC. 0: The bLocalDDCOnly property controls per-SVT caching behavior, 1: Force local DDC only usage for all SVTs, 2: Force local + remote DDC usage for all SVTs"),
+	ECVF_Default
+);
+
 ////////////////////////////////////////////////////////////////////////////////////////////////
 
 FArchive& operator<<(FArchive& Ar, UE::SVT::FMipLevelStreamingInfo& MipLevelStreamingInfo)
@@ -366,7 +374,7 @@ bool FResources::Build(USparseVolumeTextureFrame* Owner)
 	return false;
 }
 
-void FResources::Cache(USparseVolumeTextureFrame* Owner)
+void FResources::Cache(USparseVolumeTextureFrame* Owner, bool bLocalCachingOnly)
 {
 	if (Owner->GetPackage()->bIsCookedForEditor)
 	{
@@ -385,12 +393,21 @@ void FResources::Cache(USparseVolumeTextureFrame* Owner)
 	CacheKey.Bucket = FCacheBucket(TEXT("SparseVolumeTexture"));
 	CacheKey.Hash = FIoHash::HashBuffer(MakeMemoryView(FTCHARToUTF8(DerivedDataKey)));
 
+	// Set the cache policy (local only vs local+remote)
+	ECachePolicy DefaultCachePolicy = ECachePolicy::Default;
+	switch (GSVTRemoteDDCBehavior)
+	{
+	case 1:	DefaultCachePolicy = ECachePolicy::Local; break;
+	case 2:	DefaultCachePolicy = ECachePolicy::Default; break;
+	default: DefaultCachePolicy = bLocalCachingOnly ? ECachePolicy::Local : ECachePolicy::Default; break;
+	}
+
 	// Check if the data already exists in DDC
 	FSharedBuffer ResourcesDataBuffer;
 	FIoHash SVTStreamingDataHash;
 	{
-		FCacheRecordPolicyBuilder PolicyBuilder(ECachePolicy::Default | ECachePolicy::KeepAlive);
-		PolicyBuilder.AddValuePolicy(SVTStreamingDataId, ECachePolicy::Default | ECachePolicy::SkipData);
+		FCacheRecordPolicyBuilder PolicyBuilder(DefaultCachePolicy | ECachePolicy::KeepAlive);
+		PolicyBuilder.AddValuePolicy(SVTStreamingDataId, DefaultCachePolicy | ECachePolicy::SkipData);
 
 		FCacheGetRequest Request;
 		Request.Name = Owner->GetPathName();
@@ -458,7 +475,7 @@ void FResources::Cache(USparseVolumeTextureFrame* Owner)
 			RecordBuilder.AddValue(SVTDataId, Value);
 
 			FRequestOwner RequestOwner(UE::DerivedData::EPriority::Blocking);
-			const FCachePutRequest PutRequest = { FSharedString(Owner->GetPathName()), RecordBuilder.Build(), ECachePolicy::Default | ECachePolicy::KeepAlive };
+			const FCachePutRequest PutRequest = { FSharedString(Owner->GetPathName()), RecordBuilder.Build(), DefaultCachePolicy | ECachePolicy::KeepAlive };
 			GetCache().Put(MakeArrayView(&PutRequest, 1), RequestOwner,
 				[&bSavedToDDC](FCachePutResponse&& Response)
 				{
@@ -486,6 +503,19 @@ void FResources::Cache(USparseVolumeTextureFrame* Owner)
 	}
 }
 
+void FResources::SetDefault(EPixelFormat FormatA, EPixelFormat FormatB, const FVector4f& FallbackValueA, const FVector4f& FallbackValueB)
+{
+	Header = FHeader(FIntVector(0, 0, 0), FIntVector(1, 1, 1), FormatA, FormatB, FallbackValueA, FallbackValueB);
+	ResourceFlags = 0;
+	MipLevelStreamingInfo.SetNumZeroed(1);
+	RootData.Reset();
+	StreamableMipLevels.RemoveBulkData();
+	ResourceName.Reset();
+	DDCKeyHash.Reset();
+	DDCRawHash.Reset();
+	DDCRebuildState.store(EDDCRebuildState::Initial);
+}
+
 #endif // WITH_EDITORONLY_DATA
 
 #if WITH_EDITORONLY_DATA
@@ -781,7 +811,7 @@ bool USparseVolumeTextureFrame::IsCachedCookedPlatformDataLoaded(const ITargetPl
 	{
 		UE_LOG(LogSparseVolumeTexture, Log, TEXT("Failed to recover SparseVolumeTexture streaming from DDC for '%s' Frame %i. Rebuilding and retrying."), *Owner->GetPathName(), FrameIndex);
 
-		Resources.Cache(this);
+		Resources.Cache(this, CastChecked<UStreamableSparseVolumeTexture>(Owner)->bLocalDDCOnly);
 		return false;
 	}
 
@@ -805,9 +835,16 @@ void USparseVolumeTextureFrame::ClearAllCachedCookedPlatformData()
 #endif // WITH_EDITOR
 
 #if WITH_EDITORONLY_DATA
-void USparseVolumeTextureFrame::Cache()
+void USparseVolumeTextureFrame::Cache(bool bSkipDDCAndSetResourcesToDefault)
 {
-	Resources.Cache(this);
+	if (bSkipDDCAndSetResourcesToDefault)
+	{
+		Resources.SetDefault(GetFormat(0), GetFormat(1), GetFallbackValue(0), GetFallbackValue(1));
+	}
+	else
+	{
+		Resources.Cache(this, CastChecked<UStreamableSparseVolumeTexture>(Owner)->bLocalDDCOnly);
+	}
 	CreateTextureRenderResources();
 }
 #endif
@@ -1081,20 +1118,36 @@ void UStreamableSparseVolumeTexture::RecacheFrames()
 		return;
 	}
 
-	// SVT_TODO: This shows the user that something is actually happening and the editor did not freeze. The cancel button is currently ignored.
 	FScopedSlowTask RecacheTask(static_cast<float>(Frames.Num() + 2), LOCTEXT("SparseVolumeTextureCacheFrames", "Caching SparseVolumeTexture frames in Derived Data Cache"));
 	RecacheTask.MakeDialog(true);
 
 	UE::SVT::GetStreamingManager().Remove_GameThread(this);
 	RecacheTask.EnterProgressFrame(1.0f);
 
+	bool bCanceled = false;
 	for (USparseVolumeTextureFrame* Frame : Frames)
 	{
-		RecacheTask.EnterProgressFrame(1.0f);
-		Frame->Cache();
+		if (!bCanceled && RecacheTask.ShouldCancel())
+		{
+			bCanceled = true;
+			UE_LOG(LogSparseVolumeTexture, Warning, TEXT("Canceled '%s' SparseVolumeTexture caching. All frames not yet cached will be empty!"), *GetName());
+		}
+		if (!bCanceled)
+		{
+			RecacheTask.EnterProgressFrame(1.0f);
+		}
+		
+		// Even when the user cancels the caching, we still need valid data, so skip any DDC interactions and any actual building and instead just set the FResources
+		// of the frame to (valid) empty/default data.
+		const bool bSkipDDCAndSetResourcesToDefault = bCanceled;
+		Frame->Cache(bSkipDDCAndSetResourcesToDefault);
 	}
 
-	RecacheTask.EnterProgressFrame(1.0f);
+	if (!bCanceled)
+	{
+		RecacheTask.EnterProgressFrame(1.0f);
+	}
+	
 	UE::SVT::GetStreamingManager().Add_GameThread(this);
 }
 #endif
diff --git a/Engine/Source/Runtime/Engine/Private/SparseVolumeTexture/SparseVolumeTextureStreamingManager.cpp b/Engine/Source/Runtime/Engine/Private/SparseVolumeTexture/SparseVolumeTextureStreamingManager.cpp
--- a/Engine/Source/Runtime/Engine/Private/SparseVolumeTexture/SparseVolumeTextureStreamingManager.cpp
+++ b/Engine/Source/Runtime/Engine/Private/SparseVolumeTexture/SparseVolumeTextureStreamingManager.cpp
@@ -1264,6 +1264,30 @@ void FStreamingManager::RemoveInternal(UStreamableSparseVolumeTexture* SparseVol
 	FStreamingInfo* SVTInfo = StreamingInfo.Find(SparseVolumeTexture);
 	if (SVTInfo)
 	{
+		// Remove any requests for this SVT
+		TArray<FMipLevelKey> RequestsToRemove;
+		for (auto& Pair : RequestsHashTable)
+		{
+			if (Pair.Key.SVT == SparseVolumeTexture)
+			{
+				RequestsToRemove.Add(Pair.Key);
+			}
+		}
+		for (const FMipLevelKey& Key : RequestsToRemove)
+		{
+			RequestsHashTable.Remove(Key);
+		}
+
+		// Cancel any pending mip levels
+		for (FPendingMipLevel& PendingMipLevel : PendingMipLevels)
+		{
+			if (PendingMipLevel.SparseVolumeTexture == SparseVolumeTexture)
+			{
+				PendingMipLevel.Reset();
+			}
+		}
+
+		// Release resources
 		for (FFrameInfo& FrameInfo : SVTInfo->PerFrameInfo)
 		{
 			FrameInfo.PageTableTextureRHIRef.SafeRelease();
