World Partition HLODs:
- Convert HLODLayer.ParentLayer to an object reference instead of a soft object pointer.
- Detect circular HLOD layer parent references now that the entire parent chain is loaded.
- Initial per runtime partition HLODs layer setup for RuntimeHashSet: detect HLOD layer setup changes and update internal partition layers.

diff --git a/Engine/Plugins/Editor/WorldPartitionHLODUtilities/Source/Private/WorldPartition/HLOD/Utilities/WorldPartitionHLODUtilities.cpp b/Engine/Plugins/Editor/WorldPartitionHLODUtilities/Source/Private/WorldPartition/HLOD/Utilities/WorldPartitionHLODUtilities.cpp
--- a/Engine/Plugins/Editor/WorldPartitionHLODUtilities/Source/Private/WorldPartition/HLOD/Utilities/WorldPartitionHLODUtilities.cpp
+++ b/Engine/Plugins/Editor/WorldPartitionHLODUtilities/Source/Private/WorldPartition/HLOD/Utilities/WorldPartitionHLODUtilities.cpp
@@ -244,7 +244,7 @@ TArray<AWorldPartitionHLOD*> FWorldPartitionHLODUtilities::CreateHLODActors(FHLO
 		}
 
 		// Parent HLOD layer
-		UHLODLayer* ParentHLODLayer = HLODLayer->GetParentLayer().LoadSynchronous();
+		UHLODLayer* ParentHLODLayer = HLODLayer->GetParentLayer();
 		if (HLODActor->GetHLODLayer() != ParentHLODLayer)
 		{
 			HLODActor->SetHLODLayer(ParentHLODLayer);
diff --git a/Engine/Source/Editor/UnrealEd/Private/Commandlets/WorldPartitionConvertCommandlet.cpp b/Engine/Source/Editor/UnrealEd/Private/Commandlets/WorldPartitionConvertCommandlet.cpp
--- a/Engine/Source/Editor/UnrealEd/Private/Commandlets/WorldPartitionConvertCommandlet.cpp
+++ b/Engine/Source/Editor/UnrealEd/Private/Commandlets/WorldPartitionConvertCommandlet.cpp
@@ -260,7 +260,7 @@ UWorldPartition* UWorldPartitionConvertCommandlet::CreateWorldPartition(AWorldSe
 		while (CurrentHLODLayer)
 		{
 			PackagesToSave.Add(CurrentHLODLayer->GetPackage());
-			CurrentHLODLayer = Cast<UHLODLayer>(CurrentHLODLayer->GetParentLayer().Get());
+			CurrentHLODLayer = CurrentHLODLayer->GetParentLayer();
 		}
 	}
 	
diff --git a/Engine/Source/Editor/UnrealEd/Private/FileHelpers.cpp b/Engine/Source/Editor/UnrealEd/Private/FileHelpers.cpp
--- a/Engine/Source/Editor/UnrealEd/Private/FileHelpers.cpp
+++ b/Engine/Source/Editor/UnrealEd/Private/FileHelpers.cpp
@@ -819,7 +819,7 @@ static bool SaveWorld(UWorld* World,
 							while (CurrentHLODLayer)
 							{
 								PackagesToSave.Add(CurrentHLODLayer->GetPackage());
-								CurrentHLODLayer = Cast<UHLODLayer>(CurrentHLODLayer->GetParentLayer().Get());
+								CurrentHLODLayer = CurrentHLODLayer->GetParentLayer();
 							}
 						}
 					}
diff --git a/Engine/Source/Runtime/Engine/Private/WorldPartition/HLOD/HLODLayer.cpp b/Engine/Source/Runtime/Engine/Private/WorldPartition/HLOD/HLODLayer.cpp
--- a/Engine/Source/Runtime/Engine/Private/WorldPartition/HLOD/HLODLayer.cpp
+++ b/Engine/Source/Runtime/Engine/Private/WorldPartition/HLOD/HLODLayer.cpp
@@ -29,13 +29,9 @@ UHLODLayer::UHLODLayer(const FObjectInitializer& ObjectInitializer)
 	, CellSize(25600)
 	, LoadingRange(51200)
 	, HLODActorClass(AWorldPartitionHLOD::StaticClass())
-{
-}
+{}
 
 #if WITH_EDITOR
-
-
-
 UHLODLayer* UHLODLayer::GetHLODLayer(const AActor* InActor)
 {
 	if (UHLODLayer* HLODLayer = InActor->GetHLODLayer())
@@ -147,7 +143,7 @@ UHLODLayer* UHLODLayer::DuplicateHLODLayersSetup(UHLODLayer* HLODLayer, const FS
 		}
 
 		LastHLODLayer = NewHLODLayer;
-		CurrentHLODLayer = Cast<UHLODLayer>(CurrentHLODLayer->GetParentLayer().LoadSynchronous());
+		CurrentHLODLayer = CurrentHLODLayer->GetParentLayer();
 	}
 
 	return Result;
@@ -199,6 +195,23 @@ void UHLODLayer::PostEditChangeProperty(FPropertyChangedEvent& PropertyChangedEv
 	{
 		HLODBuilderSettings = WPHLODUtilities->CreateHLODBuilderSettings(this);
 	}
+	else if (PropertyName == GET_MEMBER_NAME_CHECKED(UHLODLayer, ParentLayer))
+	{
+		TSet<const UHLODLayer*> VisitedHLODLayers;
+		const UHLODLayer* CurHLODLayer = ParentLayer;
+		while (CurHLODLayer)
+		{
+			bool bHLODLayerWasAlreadyInSet;
+			VisitedHLODLayers.Add(CurHLODLayer, &bHLODLayerWasAlreadyInSet);
+			if (bHLODLayerWasAlreadyInSet)
+			{
+				UE_LOG(LogHLODLayer, Warning, TEXT("Circular HLOD parent chain detedted: HLODLayer=%s ParentLayer=%s"), *GetName(), *ParentLayer->GetName());
+				ParentLayer = nullptr;
+				break;
+			}
+			CurHLODLayer = CurHLODLayer->GetParentLayer();
+		}
+	}
 }
 
 FName UHLODLayer::GetRuntimeGridName(uint32 InLODLevel, int32 InCellSize, double InLoadingRange)
@@ -210,16 +223,4 @@ FName UHLODLayer::GetRuntimeGrid(uint32 InHLODLevel) const
 {
 	return !IsSpatiallyLoaded() ? NAME_None : GetRuntimeGridName(InHLODLevel, CellSize, LoadingRange);
 }
-
-const TSoftObjectPtr<UHLODLayer>& UHLODLayer::GetParentLayer() const
-{
-	static const TSoftObjectPtr<UHLODLayer> NullLayer;
-	return !IsSpatiallyLoaded() ? NullLayer : ParentLayer;
-}
-
-const void UHLODLayer::SetParentLayer(const TSoftObjectPtr<UHLODLayer>& InParentLayer)
-{
-	ParentLayer = InParentLayer;
-}
-
 #endif // WITH_EDITOR
diff --git a/Engine/Source/Runtime/Engine/Private/WorldPartition/RuntimeHashSet/WorldPartitionRuntimeHashSet.cpp b/Engine/Source/Runtime/Engine/Private/WorldPartition/RuntimeHashSet/WorldPartitionRuntimeHashSet.cpp
--- a/Engine/Source/Runtime/Engine/Private/WorldPartition/RuntimeHashSet/WorldPartitionRuntimeHashSet.cpp
+++ b/Engine/Source/Runtime/Engine/Private/WorldPartition/RuntimeHashSet/WorldPartitionRuntimeHashSet.cpp
@@ -33,6 +33,56 @@ FAutoConsoleCommand WorldPartitionRuntimeHashSetEnable(
 	})
 );
 
+#if WITH_EDITOR
+void FRuntimePartitionDesc::UpdateHLODPartitionLayers()
+{
+	if (!Class || !MainLayer || !HLODLayer)
+	{
+		HLODSetups.Empty();
+	}
+	else
+	{
+		TSet<const UHLODLayer*> VisitedHLODLayers;
+
+		const UHLODLayer* CurHLODLayer = HLODLayer;
+		while (CurHLODLayer)
+		{
+			const int32 HLODSetupIndex = VisitedHLODLayers.Num();
+
+			bool bHLODLayerWasAlreadyInSet;
+			VisitedHLODLayers.Add(CurHLODLayer, &bHLODLayerWasAlreadyInSet);
+			if (bHLODLayerWasAlreadyInSet)
+			{
+				break;
+			}
+
+			if (!HLODSetups.IsValidIndex(HLODSetupIndex))
+			{
+				HLODSetups.AddDefaulted();
+			}
+
+			FRuntimePartitionHLODSetup& HLODSetup = HLODSetups[HLODSetupIndex];
+
+			const bool bHLODLayerMatches = HLODSetup.HLODLayer == CurHLODLayer;
+			const UClass* ExpectedHLODPartitionClass = CurHLODLayer->IsSpatiallyLoaded() ? MainLayer->GetClass() : URuntimePartitionPersistent::StaticClass();
+			const bool bHasValidPartitionLayer = HLODSetup.PartitionLayer && (HLODSetup.PartitionLayer->GetClass() == ExpectedHLODPartitionClass);
+			
+			if (!bHLODLayerMatches || !bHasValidPartitionLayer)
+			{
+				HLODSetup.HLODLayer = CurHLODLayer;
+				HLODSetup.PartitionLayer = CurHLODLayer->IsSpatiallyLoaded() ? DuplicateObject<URuntimePartition>(MainLayer, MainLayer->GetOuter()) : NewObject<URuntimePartition>(MainLayer->GetOuter(), ExpectedHLODPartitionClass);
+				HLODSetup.PartitionLayer->Name = CurHLODLayer->GetFName();
+				HLODSetup.PartitionLayer->bIsHLODSetup = true;
+			}
+
+			CurHLODLayer = CurHLODLayer->GetParentLayer();
+		}
+
+		HLODSetups.SetNum(VisitedHLODLayers.Num());
+	}
+}
+#endif
+
 UWorldPartitionRuntimeHashSet::UWorldPartitionRuntimeHashSet(const FObjectInitializer& ObjectInitializer)
 	: Super(ObjectInitializer)
 {}
@@ -61,6 +111,18 @@ void UWorldPartitionRuntimeHashSet::Serialize(FArchive& Ar)
 #endif
 }
 
+void UWorldPartitionRuntimeHashSet::PostLoad()
+{
+	Super::PostLoad();
+
+#if WITH_EDITOR
+	if (!GetTypedOuter<UWorld>()->IsGameWorld())
+	{
+		UpdateHLODPartitionLayers();
+	}
+#endif
+}
+
 void UWorldPartitionRuntimeHashSet::AddReferencedObjects(UObject* InThis, FReferenceCollector& Collector)
 {
 #if WITH_EDITOR
@@ -84,9 +146,9 @@ bool UWorldPartitionRuntimeHashSet::SupportsHLODs() const
 {
 	for (const FRuntimePartitionDesc& RuntimePartitionDesc : RuntimePartitions)
 	{
-		if (RuntimePartitionDesc.Partitions.Num())
+		if (RuntimePartitionDesc.MainLayer)
 		{
-			if (RuntimePartitionDesc.Partitions[0]->SupportsHLODs())
+			if (RuntimePartitionDesc.MainLayer->SupportsHLODs())
 			{
 				return true;
 			}
@@ -108,8 +170,8 @@ bool UWorldPartitionRuntimeHashSet::GenerateStreaming(UWorldPartitionStreamingPo
 	check(!PersistentPartitionDesc.Class);
 	PersistentPartitionDesc.Class = URuntimePartitionPersistent::StaticClass();
 	PersistentPartitionDesc.Name = NAME_PersistentLevel;
-	PersistentPartitionDesc.Partitions.Add(NewObject<URuntimePartition>(this, URuntimePartitionPersistent::StaticClass(), NAME_None));
-	PersistentPartitionDesc.Partitions[0]->Name = TEXT("MainPartition");
+	PersistentPartitionDesc.MainLayer = NewObject<URuntimePartition>(this, URuntimePartitionPersistent::StaticClass(), NAME_None);
+	PersistentPartitionDesc.MainLayer->Name = TEXT("MainPartition");
 
 	//
 	// Split actor sets into their corresponding runtime partition implementation
@@ -131,8 +193,7 @@ bool UWorldPartitionRuntimeHashSet::GenerateStreaming(UWorldPartitionStreamingPo
 
 		if (const FRuntimePartitionDesc** RuntimePartitionDesc = NameToRuntimePartitionDescMap.Find(ActorSetRuntimeGrid[0]))
 		{
-			const uint32 PartitionLayerIndex = 0;
-			RuntimePartitionsToActorSetMap.FindOrAdd((*RuntimePartitionDesc)->Partitions[PartitionLayerIndex]).Add(&ActorSetInstance);
+			RuntimePartitionsToActorSetMap.FindOrAdd((*RuntimePartitionDesc)->MainLayer).Add(&ActorSetInstance);
 		}
 	});
 
@@ -320,7 +381,7 @@ void UWorldPartitionRuntimeHashSet::FlushStreaming()
 	check(PersistentPartitionDesc.Class);
 	PersistentPartitionDesc.Class = nullptr;
 	PersistentPartitionDesc.Name = NAME_None;
-	PersistentPartitionDesc.Partitions.Empty();
+	PersistentPartitionDesc.MainLayer = nullptr;
 
 	NonSpatiallyLoadedRuntimeCells.Empty();	
 	SpatialIndex.Reset();
@@ -340,9 +401,9 @@ bool UWorldPartitionRuntimeHashSet::IsValidGrid(FName GridName) const
 	{
 		if (RuntimePartitionDesc.Name == GridNameList[0])
 		{
-			if (RuntimePartitionDesc.Partitions.Num())
+			if (RuntimePartitionDesc.MainLayer)
 			{
-				return RuntimePartitionDesc.Partitions[0]->IsValidGrid(GridName);
+				return RuntimePartitionDesc.MainLayer->IsValidGrid(GridName);
 			}
 		}
 	}
@@ -617,6 +678,7 @@ void UWorldPartitionRuntimeHashSet::PostEditChangeChainProperty(FPropertyChanged
 
 	static FName NAME_RuntimePartitions(TEXT("RuntimePartitions"));
 	static FName NAME_HLODSetups_Key(TEXT("HLODSetups_Key"));
+	static FName NAME_HLODLayer(TEXT("HLODLayer"));
 
 	FName PropertyName = PropertyChangedEvent.Property ? PropertyChangedEvent.Property->GetFName() : NAME_None;
 
@@ -627,13 +689,15 @@ void UWorldPartitionRuntimeHashSet::PostEditChangeChainProperty(FPropertyChanged
 
 		FRuntimePartitionDesc& RuntimePartitionDesc = RuntimePartitions[RuntimePartitionIndex];
 		
-		RuntimePartitionDesc.Partitions.Empty();
+		RuntimePartitionDesc.MainLayer = nullptr;
 
 		if (RuntimePartitionDesc.Class)
 		{
 			RuntimePartitionDesc.Name = RuntimePartitionDesc.Class->GetFName();
-			RuntimePartitionDesc.Partitions.Add(NewObject<URuntimePartition>(this, RuntimePartitionDesc.Class, NAME_None));
-			RuntimePartitionDesc.Partitions[0]->Name = TEXT("MainPartition");
+			RuntimePartitionDesc.MainLayer = NewObject<URuntimePartition>(this, RuntimePartitionDesc.Class, NAME_None);
+			RuntimePartitionDesc.MainLayer->Name = TEXT("MainPartition");
+
+			RuntimePartitionDesc.UpdateHLODPartitionLayers();
 		}
 	}
 	else if (PropertyName == GET_MEMBER_NAME_CHECKED(FRuntimePartitionDesc, Name))
@@ -662,56 +726,20 @@ void UWorldPartitionRuntimeHashSet::PostEditChangeChainProperty(FPropertyChanged
 			}
 		}
 	}
-	else if (PropertyName == NAME_HLODSetups_Key)
+	else if (PropertyName == NAME_HLODLayer)
 	{
 		int32 RuntimePartitionIndex = PropertyChangedEvent.GetArrayIndex(NAME_RuntimePartitions.ToString());
 		check(RuntimePartitions.IsValidIndex(RuntimePartitionIndex));
 
-		FRuntimePartitionDesc& RuntimePartitionDesc = RuntimePartitions[RuntimePartitionIndex];
-
-		int32 HLODSetupsIndex = PropertyChangedEvent.GetArrayIndex(NAME_HLODSetups_Key.ToString());
-		const FSetElementId HLODSetupElemId = FSetElementId::FromInteger(HLODSetupsIndex);
-		check(RuntimePartitionDesc.HLODSetups.IsValidId(HLODSetupElemId));
-
-		const UHLODLayer* HLODLayer = RuntimePartitionDesc.HLODSetups.Get(HLODSetupElemId).Key;
-		FRuntimePartitionHLODSetup& HLODSetup = RuntimePartitionDesc.HLODSetups.Get(HLODSetupElemId).Value;
-		
-		HLODSetup.Partitions.Empty();
-
-		while (HLODLayer)
-		{
-			URuntimePartition* HLODPartitionLayer = NewObject<URuntimePartition>(this, RuntimePartitionDesc.Class, NAME_None);
-			HLODPartitionLayer->Name = HLODLayer->GetFName();
-			HLODPartitionLayer->bIsHLODSetup = true;
-
-			HLODSetup.Partitions.Add(HLODPartitionLayer);
-
-			HLODLayer = HLODLayer->GetParentLayer().LoadSynchronous();
-		}
-
-		SetupHLODPartitionLayers(RuntimePartitionIndex);
+		RuntimePartitions[RuntimePartitionIndex].UpdateHLODPartitionLayers();
 	}
 }
 
-void UWorldPartitionRuntimeHashSet::SetupHLODPartitionLayers(int32 RuntimePartitionIndex)
+void UWorldPartitionRuntimeHashSet::UpdateHLODPartitionLayers()
 {
-	FRuntimePartitionDesc& RuntimePartitionDesc = RuntimePartitions[RuntimePartitionIndex];
-
-	if (RuntimePartitionDesc.Partitions.Num())
+	for (FRuntimePartitionDesc& RuntimePartitionDesc : RuntimePartitions)
 	{
-		RuntimePartitionDesc.Partitions.SetNum(1);
-
-		for (auto [CurHLODLayer, CurHLODSetup] : RuntimePartitionDesc.HLODSetups)
-		{
-			for (const URuntimePartition* HLODPartitionTemplate : CurHLODSetup.Partitions)
-			{
-				URuntimePartition* HLODPartitionLayer = DuplicateObject<URuntimePartition>(HLODPartitionTemplate, this);
-				HLODPartitionLayer->Name = CurHLODLayer->GetFName();
-				HLODPartitionLayer->bIsHLODSetup = false;
-
-				RuntimePartitionDesc.Partitions.Add(HLODPartitionLayer);
-			}
-		}
+		RuntimePartitionDesc.UpdateHLODPartitionLayers();
 	}
 }
 #endif
\ No newline at end of file
diff --git a/Engine/Source/Runtime/Engine/Private/WorldPartition/RuntimeSpatialHash/RuntimeSpatialHashHLOD.cpp b/Engine/Source/Runtime/Engine/Private/WorldPartition/RuntimeSpatialHash/RuntimeSpatialHashHLOD.cpp
--- a/Engine/Source/Runtime/Engine/Private/WorldPartition/RuntimeSpatialHash/RuntimeSpatialHashHLOD.cpp
+++ b/Engine/Source/Runtime/Engine/Private/WorldPartition/RuntimeSpatialHash/RuntimeSpatialHashHLOD.cpp
@@ -297,7 +297,7 @@ static TMap<UHLODLayer*, int32> GatherHLODLayers(const IStreamingGenerationConte
 							int32& HLODLevel = HLODLayersLevel.FindOrAdd(HLODLayer);
 							HLODLevel = FMath::Max(HLODLevel, CurrentHLODLevel);
 		
-							HLODLayer = HLODLayer->GetParentLayer().LoadSynchronous();
+							HLODLayer = HLODLayer->GetParentLayer();
 							CurrentHLODLevel++;
 						}
 					}
diff --git a/Engine/Source/Runtime/Engine/Public/WorldPartition/HLOD/HLODLayer.h b/Engine/Source/Runtime/Engine/Public/WorldPartition/HLOD/HLODLayer.h
--- a/Engine/Source/Runtime/Engine/Public/WorldPartition/HLOD/HLODLayer.h
+++ b/Engine/Source/Runtime/Engine/Public/WorldPartition/HLOD/HLODLayer.h
@@ -35,7 +35,7 @@ enum class EHLODLayerType : uint8
 	Custom					UMETA(DisplayName = "Custom"),
 };
 
-UCLASS(Blueprintable, Config=Engine, PerObjectConfig)
+UCLASS(Blueprintable)
 class ENGINE_API UHLODLayer : public UObject
 {
 	GENERATED_UCLASS_BODY()
@@ -61,8 +61,8 @@ public:
 	FName GetRuntimeGrid(uint32 InHLODLevel) const;
 	int32 GetCellSize() const { return !bIsSpatiallyLoaded ? 0 : CellSize; }
 	double GetLoadingRange() const { return !bIsSpatiallyLoaded ? WORLD_MAX : LoadingRange; }
-	const TSoftObjectPtr<UHLODLayer>& GetParentLayer() const;
-	const void SetParentLayer(const TSoftObjectPtr<UHLODLayer>& InParentLayer);
+	UHLODLayer* GetParentLayer() const { return !bIsSpatiallyLoaded ? nullptr : ParentLayer; }
+	void SetParentLayer(UHLODLayer* InParentLayer) { ParentLayer = InParentLayer; }
 	bool IsSpatiallyLoaded() const { return bIsSpatiallyLoaded; }
 	void SetIsSpatiallyLoaded(bool bInIsSpatiallyLoaded) { bIsSpatiallyLoaded = bInIsSpatiallyLoaded; }
 
@@ -82,38 +82,38 @@ private:
 
 private:
 	/** Type of HLOD generation to use */
-	UPROPERTY(EditAnywhere, Config, Category=HLOD)
+	UPROPERTY(EditAnywhere, Category=HLOD)
 	EHLODLayerType LayerType;
 
 	/** HLOD Builder class */
-	UPROPERTY(EditAnywhere, Config, Category=HLOD, meta = (DisplayName = "HLOD Builder Class", EditConditionHides, EditCondition = "LayerType == EHLODLayerType::Custom"))
+	UPROPERTY(EditAnywhere, Category=HLOD, meta = (DisplayName = "HLOD Builder Class", EditConditionHides, EditCondition = "LayerType == EHLODLayerType::Custom"))
 	TSubclassOf<UHLODBuilder> HLODBuilderClass;
 
 	UPROPERTY(VisibleAnywhere, Export, NoClear, Category=HLOD, meta = (EditInline, NoResetToDefault))
 	TObjectPtr<UHLODBuilderSettings> HLODBuilderSettings;
 
 	/** Whether HLOD actors generated for this layer will be spatially loaded */
-	UPROPERTY(EditAnywhere, Config, Category=HLOD)
+	UPROPERTY(EditAnywhere, Category=HLOD)
 	uint32 bIsSpatiallyLoaded : 1;
 
 	/** Cell size of the runtime grid created to encompass HLOD actors generated for this HLOD Layer */
-	UPROPERTY(EditAnywhere, Config, Category=HLOD, meta = (EditConditionHides, EditCondition = "bIsSpatiallyLoaded"))
+	UPROPERTY(EditAnywhere, Category=HLOD, meta = (EditConditionHides, EditCondition = "bIsSpatiallyLoaded"))
 	int32 CellSize;
 
 	/** Loading range of the runtime grid created to encompass HLOD actors generated for this HLOD Layer */
-	UPROPERTY(EditAnywhere, Config, Category=HLOD, meta = (EditConditionHides, EditCondition = "bIsSpatiallyLoaded"))
+	UPROPERTY(EditAnywhere, Category=HLOD, meta = (EditConditionHides, EditCondition = "bIsSpatiallyLoaded"))
 	double LoadingRange;
 
 	/** HLOD Layer to assign to the generated HLOD actors */
-	UPROPERTY(EditAnywhere, Config, Category=HLOD, meta = (EditConditionHides, EditCondition = "bIsSpatiallyLoaded"))
-	TSoftObjectPtr<UHLODLayer> ParentLayer;
+	UPROPERTY(EditAnywhere, Category=HLOD, meta = (EditConditionHides, EditCondition = "bIsSpatiallyLoaded"))
+	TObjectPtr<UHLODLayer> ParentLayer;
 
 	/** Specify a custom HLOD Actor class, the default is AWorldPartitionHLOD */
 	UPROPERTY(EditAnywhere, Category = HLOD, AdvancedDisplay, meta = (DisplayName = "HLOD Actor Class"))
 	TSubclassOf<AWorldPartitionHLOD> HLODActorClass;
 
 	/** HLOD Modifier class, to allow changes to the HLOD at runtime */
-	UPROPERTY(EditAnywhere, Config, Category = HLOD, AdvancedDisplay, meta = (DisplayName = "HLOD Modifier Class"))
+	UPROPERTY(EditAnywhere, Category = HLOD, AdvancedDisplay, meta = (DisplayName = "HLOD Modifier Class"))
 	TSubclassOf<UWorldPartitionHLODModifier> HLODModifierClass;
 
 private:
diff --git a/Engine/Source/Runtime/Engine/Public/WorldPartition/RuntimeHashSet/RuntimePartition.h b/Engine/Source/Runtime/Engine/Public/WorldPartition/RuntimeHashSet/RuntimePartition.h
--- a/Engine/Source/Runtime/Engine/Public/WorldPartition/RuntimeHashSet/RuntimePartition.h
+++ b/Engine/Source/Runtime/Engine/Public/WorldPartition/RuntimeHashSet/RuntimePartition.h
@@ -32,16 +32,16 @@ public:
 #endif
 
 #if WITH_EDITORONLY_DATA
-	UPROPERTY(VisibleAnywhere, Category = RuntimeSettings, meta = (EditCondition = "!bIsHLODSetup", EditConditionHides, HideEditConditionToggle))
+	UPROPERTY(VisibleAnywhere, Category = RuntimeSettings, Meta = (EditCondition = "!bIsHLODSetup", EditConditionHides, HideEditConditionToggle))
 	FName Name;
 
-	UPROPERTY(EditAnywhere, Category = RuntimeSettings, meta = (EditCondition = "!bIsHLODSetup", EditConditionHides, HideEditConditionToggle))
+	UPROPERTY(EditAnywhere, Category = RuntimeSettings, Meta = (EditCondition = "!bIsHLODSetup", EditConditionHides, HideEditConditionToggle))
 	bool bBlockOnSlowStreaming;
 
-	UPROPERTY(EditAnywhere, Category = RuntimeSettings, meta = (EditCondition = "!bIsHLODSetup", EditConditionHides, HideEditConditionToggle))
+	UPROPERTY(EditAnywhere, Category = RuntimeSettings, Meta = (EditCondition = "!bIsHLODSetup", EditConditionHides, HideEditConditionToggle))
 	bool bClientOnlyVisible;
 
-	UPROPERTY(EditAnywhere, Category = RuntimeSettings, meta = (EditCondition = "!bIsHLODSetup", EditConditionHides, HideEditConditionToggle))
+	UPROPERTY(EditAnywhere, Category = RuntimeSettings, Meta = (EditCondition = "!bIsHLODSetup", EditConditionHides, HideEditConditionToggle))
 	int32 Priority;
 
 	UPROPERTY()
diff --git a/Engine/Source/Runtime/Engine/Public/WorldPartition/RuntimeHashSet/RuntimePartitionLevelStreaming.h b/Engine/Source/Runtime/Engine/Public/WorldPartition/RuntimeHashSet/RuntimePartitionLevelStreaming.h
--- a/Engine/Source/Runtime/Engine/Public/WorldPartition/RuntimeHashSet/RuntimePartitionLevelStreaming.h
+++ b/Engine/Source/Runtime/Engine/Public/WorldPartition/RuntimeHashSet/RuntimePartitionLevelStreaming.h
@@ -19,7 +19,7 @@ public:
 
 #if WITH_EDITORONLY_DATA
 	/** Splits actors into different sublevels based on their actor containers (level instance) */
-	UPROPERTY(EditAnywhere, Category = RuntimeSettings)
+	UPROPERTY(EditAnywhere, Category = RuntimeSettings, Meta = (EditCondition = "!bIsHLODSetup", EditConditionHides, HideEditConditionToggle))
 	bool bOneLevelPerActorContainer;
 #endif
 };
\ No newline at end of file
diff --git a/Engine/Source/Runtime/Engine/Public/WorldPartition/RuntimeHashSet/WorldPartitionRuntimeHashSet.h b/Engine/Source/Runtime/Engine/Public/WorldPartition/RuntimeHashSet/WorldPartitionRuntimeHashSet.h
--- a/Engine/Source/Runtime/Engine/Public/WorldPartition/RuntimeHashSet/WorldPartitionRuntimeHashSet.h
+++ b/Engine/Source/Runtime/Engine/Public/WorldPartition/RuntimeHashSet/WorldPartitionRuntimeHashSet.h
@@ -13,14 +13,17 @@ struct FPropertyChangedChainEvent;
 
 using FStaticSpatialIndexType = TStaticSpatialIndexRTree<UWorldPartitionRuntimeCell*>;
 
-/** Holds settings for an HLOD layer (and its parents) for a particular partition class. */
+/** Holds settings for an HLOD layer for a particular partition class. */
 USTRUCT()
 struct FRuntimePartitionHLODSetup
 {
 	GENERATED_USTRUCT_BODY()
 
-	UPROPERTY(VisibleAnywhere, Category = RuntimeSettings, EditFixedSize, Instanced, meta = (NoResetToDefault, EditFixedOrder, TitleProperty="Name"))
-	TArray<TObjectPtr<URuntimePartition>> Partitions;
+	UPROPERTY(VisibleAnywhere, Category = RuntimeSettings, Meta = (DisplayThumbnail = false))
+	TObjectPtr<const UHLODLayer> HLODLayer;
+
+	UPROPERTY(EditAnywhere, Category = RuntimeSettings, Instanced)
+	TObjectPtr<URuntimePartition> PartitionLayer;
 };
 
 /** Holds settings for a runtime partition instance. */
@@ -35,16 +38,24 @@ struct FRuntimePartitionDesc
 	TSubclassOf<URuntimePartition> Class;
 
 	/** Name for this partition, used to map actors to it through the Actor.RuntimeGrid property  */
-	UPROPERTY(EditAnywhere, Category = RuntimeSettings, meta = (EditCondition = "Class != nullptr", HideEditConditionToggle))
+	UPROPERTY(EditAnywhere, Category = RuntimeSettings, Meta = (EditCondition = "Class != nullptr", HideEditConditionToggle))
 	FName Name;
 
-	/** Array of partition objects, first index is the main one and then one per HLOD layer */
-	UPROPERTY(VisibleAnywhere, Category = RuntimeSettings, EditFixedSize, Instanced, meta = (EditCondition = "Class != nullptr", HideEditConditionToggle, NoResetToDefault, EditFixedOrder, TitleProperty="Name"))
-	TArray<TObjectPtr<URuntimePartition>> Partitions;
+	/** Main partition object */
+	UPROPERTY(VisibleAnywhere, Category = RuntimeSettings, Instanced, Meta = (EditCondition = "Class != nullptr", HideEditConditionToggle, NoResetToDefault, TitleProperty = "Name"))
+	TObjectPtr<URuntimePartition> MainLayer;
+
+	/** Associated HLOD Layer object */
+	UPROPERTY(EditAnywhere, Category = RuntimeSettings, Meta = (EditCondition = "Class != nullptr", HideEditConditionToggle, NoResetToDefault, ForceInlineRow))
+	TObjectPtr<const UHLODLayer> HLODLayer;
 
-	/** HLOD setups used by this partition */
-	UPROPERTY(EditAnywhere, Category = RuntimeSettings, meta = (EditCondition = "Class != nullptr", HideEditConditionToggle, ForceInlineRow, DisplayThumbnail=false))
-	TMap<TObjectPtr<const UHLODLayer>, FRuntimePartitionHLODSetup> HLODSetups;
+	/** HLOD setups used by this partition, one for each layers in the hierarchy */
+	UPROPERTY(VisibleAnywhere, Category = RuntimeSettings, EditFixedSize, Meta = (EditCondition = "HLODLayer != nullptr", HideEditConditionToggle, ForceInlineRow))
+	TArray<FRuntimePartitionHLODSetup> HLODSetups;
+#endif
+
+#if WITH_EDITOR
+	void UpdateHLODPartitionLayers();
 #endif
 };
 
@@ -71,6 +82,7 @@ class ENGINE_API UWorldPartitionRuntimeHashSet : public UWorldPartitionRuntimeHa
 
 	//~ Begin UObject Interface
 	virtual void Serialize(FArchive& Ar) override;
+	virtual void PostLoad() override;
 	static void AddReferencedObjects(UObject* InThis, FReferenceCollector& Collector);
 #if WITH_EDITOR
 	virtual void PostEditChangeChainProperty(FPropertyChangedChainEvent& PropertyChangedEvent) override;
@@ -107,7 +119,7 @@ public:
 private:
 #if WITH_EDITOR
 	/** Update the partition layers to reflect the curent HLOD setups. */
-	void SetupHLODPartitionLayers(int32 RuntimePartitionIndex);
+	void UpdateHLODPartitionLayers();
 #endif
 
 public:
@@ -117,7 +129,7 @@ public:
 	FRuntimePartitionDesc PersistentPartitionDesc;
 
 	/** Array of runtime partition descriptors */
-	UPROPERTY(EditAnywhere, Category = RuntimeSettings, meta = (TitleProperty="Name"))
+	UPROPERTY(EditAnywhere, Category = RuntimeSettings, Meta = (TitleProperty = "Name"))
 	TArray<FRuntimePartitionDesc> RuntimePartitions;
 #endif
 
