PR #10417: ReferenceCluster generation optimization
Implement simple depth-first search instead of doing naive implementation.
Original PR code claiming it was ~10x faster in their case, adapted to better fit coding standards, etc.
Tested Reference Clustering unit tests + compared output from various internal maps.

diff --git a/Engine/Source/Runtime/Engine/Private/ReferenceCluster.cpp b/Engine/Source/Runtime/Engine/Private/ReferenceCluster.cpp
--- a/Engine/Source/Runtime/Engine/Private/ReferenceCluster.cpp
+++ b/Engine/Source/Runtime/Engine/Private/ReferenceCluster.cpp
@@ -5,56 +5,49 @@
 
 TArray<TArray<FGuid>> GenerateObjectsClusters(const TArray<TPair<FGuid, TArray<FGuid>>>& InObjects)
 {
-	TMap<FGuid, FGuid> ObjectToClusters;
-	TMap<FGuid, TSet<FGuid>> Clusters;
+	TArray<TArray<FGuid>> Result;
+	Result.Reserve(InObjects.Num());
+
+	TMap<FGuid, TSet<FGuid>> Graph;
+	Graph.Reserve(InObjects.Num());
 
-	for (const auto& Object : InObjects)
+	for (const auto& [Object, References] : InObjects)
 	{
-		const FGuid& ObjectGuid = Object.Key;
-		FGuid ClusterGuid = ObjectToClusters.FindRef(ObjectGuid);
+		TSet<FGuid>& ObjectReferences = Graph.FindOrAdd(Object);
+		ObjectReferences.Reserve(References.Num());
 
-		if (!ClusterGuid.IsValid())
+		for (const FGuid& Reference : References)
 		{
-			ClusterGuid = FGuid::NewGuid();
-			TSet<FGuid>& Cluster = Clusters.Add(ClusterGuid);
-			ObjectToClusters.Add(ObjectGuid, ClusterGuid);
-			Cluster.Add(ObjectGuid);
+			ObjectReferences.Add(Reference);
+			Graph.FindOrAdd(Reference).FindOrAdd(Object);
 		}
+	}
+
+	TSet<FGuid> VisitedObjects;
+	VisitedObjects.Reserve(InObjects.Num());
 
-		for (const FGuid& ReferenceGuid : Object.Value)
+	TArray<FGuid> ObjectStack;
+	for (const auto& ObjectNode : Graph)
+	{
+		if (!VisitedObjects.Contains(ObjectNode.Key))
 		{
-			FGuid ReferenceClusterGuid = ObjectToClusters.FindRef(ReferenceGuid);
+			TArray<FGuid>& NewCluster = Result.Emplace_GetRef();
+
+			ObjectStack.Reset();
+			ObjectStack.Add(ObjectNode.Key);
 
-			if (ReferenceClusterGuid.IsValid())
+			while (!ObjectStack.IsEmpty())
 			{
-				if (ReferenceClusterGuid != ClusterGuid)
+				const FGuid Object = ObjectStack.Pop(false);
+				if (!VisitedObjects.Contains(Object))
 				{
-					TSet<FGuid> ReferenceCluster = Clusters.FindAndRemoveChecked(ReferenceClusterGuid);
-					TSet<FGuid>& Cluster = Clusters.FindChecked(ClusterGuid);
-					Cluster.Append(ReferenceCluster);
-					for (const FGuid& OtherReferenceGuid : ReferenceCluster)
-					{
-						ObjectToClusters.FindChecked(OtherReferenceGuid) = ClusterGuid;
-					}
+					VisitedObjects.Add(Object);
+					NewCluster.Add(Object);
+					ObjectStack.Append(Graph[Object].Array());
 				}
 			}
-			else
-			{
-				TSet<FGuid>& Cluster = Clusters.FindChecked(ClusterGuid);
-				Cluster.Add(ReferenceGuid);
-			}
-
-			ObjectToClusters.Add(ReferenceGuid, ClusterGuid);
 		}
 	}
-	
-	TArray<TArray<FGuid>> Result;
-	Result.Reserve(Clusters.Num());
-
-	for (const auto& Cluster : Clusters)
-	{
-		Result.AddDefaulted_GetRef() = Cluster.Value.Array();
-	}
 
 	return MoveTemp(Result);
-}
+}
\ No newline at end of file
diff --git a/Engine/Source/Runtime/Engine/Private/Tests/ReferenceClusterTest.cpp b/Engine/Source/Runtime/Engine/Private/Tests/ReferenceClusterTest.cpp
--- a/Engine/Source/Runtime/Engine/Private/Tests/ReferenceClusterTest.cpp
+++ b/Engine/Source/Runtime/Engine/Private/Tests/ReferenceClusterTest.cpp
@@ -12,6 +12,63 @@ namespace ReferenceClusterTests
 	constexpr const uint32 TestFlags = EAutomationTestFlags::ApplicationContextMask | EAutomationTestFlags::EngineFilter;
 
 	IMPLEMENT_SIMPLE_AUTOMATION_TEST(FReferenceClusterTest, TEST_NAME_ROOT, TestFlags)
+
+	TArray<TArray<FGuid>> GenerateObjectsClustersReference(const TArray<TPair<FGuid, TArray<FGuid>>>& InObjects)
+	{
+		TMap<FGuid, FGuid> ObjectToClusters;
+		TMap<FGuid, TSet<FGuid>> Clusters;
+
+		for (const auto& Object : InObjects)
+		{
+			const FGuid& ObjectGuid = Object.Key;
+			FGuid ClusterGuid = ObjectToClusters.FindRef(ObjectGuid);
+
+			if (!ClusterGuid.IsValid())
+			{
+				ClusterGuid = FGuid::NewGuid();
+				TSet<FGuid>& Cluster = Clusters.Add(ClusterGuid);
+				ObjectToClusters.Add(ObjectGuid, ClusterGuid);
+				Cluster.Add(ObjectGuid);
+			}
+
+			for (const FGuid& ReferenceGuid : Object.Value)
+			{
+				FGuid ReferenceClusterGuid = ObjectToClusters.FindRef(ReferenceGuid);
+
+				if (ReferenceClusterGuid.IsValid())
+				{
+					if (ReferenceClusterGuid != ClusterGuid)
+					{
+						TSet<FGuid> ReferenceCluster = Clusters.FindAndRemoveChecked(ReferenceClusterGuid);
+						TSet<FGuid>& Cluster = Clusters.FindChecked(ClusterGuid);
+						Cluster.Append(ReferenceCluster);
+						for (const FGuid& OtherReferenceGuid : ReferenceCluster)
+						{
+							ObjectToClusters.FindChecked(OtherReferenceGuid) = ClusterGuid;
+						}
+					}
+				}
+				else
+				{
+					TSet<FGuid>& Cluster = Clusters.FindChecked(ClusterGuid);
+					Cluster.Add(ReferenceGuid);
+				}
+
+				ObjectToClusters.Add(ReferenceGuid, ClusterGuid);
+			}
+		}
+	
+		TArray<TArray<FGuid>> Result;
+		Result.Reserve(Clusters.Num());
+
+		for (const auto& Cluster : Clusters)
+		{
+			Result.AddDefaulted_GetRef() = Cluster.Value.Array();
+		}
+
+		return MoveTemp(Result);
+	}
+
 	bool FReferenceClusterTest::RunTest(const FString& Parameters)
 	{
 #if WITH_EDITOR
@@ -107,6 +164,62 @@ namespace ReferenceClusterTests
 			TestTrue(TEXT("Reference Cluster Complex Merging"), Clusters.Num() == 1);
 			TestTrue(TEXT("Reference Cluster Complex Merging"), Clusters[0].Num() == 10);
 		}
+
+		// Random clusters test
+		{
+			const int32 NumObjects = 25000;
+
+			TArray<FGuid> Objects;
+			Objects.Reserve(NumObjects);
+			for (int32 i=0; i<NumObjects; i++)
+			{
+				Objects.Emplace(FGuid::NewGuid());
+			}
+
+			ObjectsWithRefs.Empty(Objects.Num());
+			for (int32 i=0; i<Objects.Num(); i++)
+			{
+				TPair<FGuid, TArray<FGuid>>& Pair = ObjectsWithRefs.AddDefaulted_GetRef();
+				Pair.Key = Objects[i];
+
+				const int32 NumRefs = FMath::RandRange(-10, 10);
+
+				if (NumRefs > 0)
+				{
+					Pair.Value.Reserve(NumRefs);
+					for (int32 j=0; j<NumRefs; j++)
+					{
+						Pair.Value.Add(Objects[FMath::RandHelper(Objects.Num())]);
+					}
+				}
+			}
+
+			const double StartTimeRef = FPlatformTime::Seconds();
+			TArray<TArray<FGuid>> ClustersRef = GenerateObjectsClustersReference(ObjectsWithRefs);
+			const double RunTimeRef = FPlatformTime::Seconds() - StartTimeRef;
+			AddInfo(FString::Printf(TEXT("Clustered %d objects into %d clusters in %s (reference)"), ObjectsWithRefs.Num(), ClustersRef.Num(), *FPlatformTime::PrettyTime(RunTimeRef)));
+
+			const double StartTimeCur = FPlatformTime::Seconds();
+			TArray<TArray<FGuid>> ClustersCur = GenerateObjectsClusters(ObjectsWithRefs);
+			const double RunTimeCur = FPlatformTime::Seconds() - StartTimeCur;
+			AddInfo(FString::Printf(TEXT("Clustered %d objects into %d clusters in %s (current)"), ObjectsWithRefs.Num(), ClustersCur.Num(), *FPlatformTime::PrettyTime(RunTimeCur)));
+
+			auto SortInnerClusters = [](const TArray<FGuid>& A, const TArray<FGuid>& B)
+			{
+				const FGuid Guid0;
+				const FGuid& GuidA = A.Num() ? A[0] : Guid0;
+				const FGuid& GuidB = B.Num() ? B[0] : Guid0;
+				return GuidA < GuidB;
+			};
+
+			for (TArray<FGuid>& Cluster : ClustersRef) { Cluster.Sort(); }
+			ClustersRef.Sort(SortInnerClusters);
+
+			for (TArray<FGuid>& Cluster : ClustersCur) { Cluster.Sort(); }
+			ClustersCur.Sort(SortInnerClusters);
+
+			TestTrue(TEXT("Clustering against reference algorithm"), ClustersRef == ClustersCur);
+		}
 #endif
 		return true;
 	}
@@ -114,4 +227,4 @@ namespace ReferenceClusterTests
 
 #undef TEST_NAME_ROOT
 
-#endif 
+#endif 
\ No newline at end of file
