Fixed crash with blend profiles in DeadBlending node.

diff --git a/Engine/Source/Runtime/Engine/Private/Animation/AnimNode_DeadBlending.cpp b/Engine/Source/Runtime/Engine/Private/Animation/AnimNode_DeadBlending.cpp
--- a/Engine/Source/Runtime/Engine/Private/Animation/AnimNode_DeadBlending.cpp
+++ b/Engine/Source/Runtime/Engine/Private/Animation/AnimNode_DeadBlending.cpp
@@ -673,27 +673,19 @@ void FAnimNode_DeadBlending::Evaluate_AnyThread(FPoseContext& Output)
 
 		InertializationTime = 0.0f;
 		InertializationDuration = BlendTimeMultiplier * RequestQueue[ShortestRequestIdx].Duration;
-
-		UE::Anim::TTypedIndexArray<FSkeletonPoseBoneIndex, float, FAnimStackAllocator> RequestDurationPerBone;
+		InertializationDurationPerBone.Init(InertializationDuration, NumSkeletonBones);
 
 		if (RequestQueue[ShortestRequestIdx].BlendProfile)
 		{
-			RequestQueue[ShortestRequestIdx].BlendProfile->FillSkeletonBoneDurationsArray(RequestDurationPerBone, InertializationDuration);
+			RequestQueue[ShortestRequestIdx].BlendProfile->FillSkeletonBoneDurationsArray(InertializationDurationPerBone, InertializationDuration);
 		}
 		else if (DefaultBlendProfile)
 		{
-			DefaultBlendProfile->FillSkeletonBoneDurationsArray(RequestDurationPerBone, InertializationDuration);
-		}
-		else
-		{
-			RequestDurationPerBone.Init(InertializationDuration, NumSkeletonBones);
+			DefaultBlendProfile->FillSkeletonBoneDurationsArray(InertializationDurationPerBone, InertializationDuration);
 		}
 
-		InertializationDurationPerBone.SetNumUninitialized(NumSkeletonBones);
-		for (int32 BoneIndex = 0; BoneIndex < NumSkeletonBones; BoneIndex++)
-		{
-			InertializationDurationPerBone[BoneIndex] = RequestDurationPerBone[BoneIndex];
-		}
+		// Cache the maximum duration across all bones (so we know when to deactivate the inertialization request)
+		InertializationMaxDuration = FMath::Max(InertializationDuration, *Algo::MaxElement(InertializationDurationPerBone));
 
 		if (RequestQueue[ShortestRequestIdx].bUseBlendMode)
 		{
@@ -701,9 +693,6 @@ void FAnimNode_DeadBlending::Evaluate_AnyThread(FPoseContext& Output)
 			InertializationCustomBlendCurve = RequestQueue[ShortestRequestIdx].CustomBlendCurve;
 		}
 
-		// Cache the maximum duration across all bones (so we know when to deactivate the inertialization request)
-		InertializationMaxDuration = FMath::Max(InertializationDuration, *Algo::MaxElement(InertializationDurationPerBone));
-
 #if ANIM_TRACE_ENABLED
 		InertializationRequestDescription = RequestQueue[ShortestRequestIdx].Description;
 #endif
