Back out CL 24615501 to prevent performance issues with defaulting bIsStatic to true for World Partition streaming levels

diff --git a/Engine/Source/Runtime/Engine/Private/World.cpp b/Engine/Source/Runtime/Engine/Private/World.cpp
--- a/Engine/Source/Runtime/Engine/Private/World.cpp
+++ b/Engine/Source/Runtime/Engine/Private/World.cpp
@@ -8820,24 +8820,24 @@ void UWorld::DuplicateRequestedLevels(const FName MapName)
 {
 	if (GEngine->Experimental_ShouldPreDuplicateMap(MapName))
 	{
-		// Duplicate the persistent level and only dynamic levels, but don't add them to the world.
-		FLevelCollection DuplicateLevels;
-		DuplicateLevels.SetType(ELevelCollectionType::DynamicDuplicatedLevels);
-		DuplicateLevels.SetIsVisible(false);
-		ULevel* const DuplicatePersistentLevel = DuplicateLevelWithPrefix(PersistentLevel, 1);
-		if (!DuplicatePersistentLevel)
-		{
-			UE_LOG(LogWorld, Warning, TEXT("UWorld::DuplicateRequestedLevels: failed to duplicate persistent level %s. No duplicate level collection will be created."),
-				*GetFullNameSafe(PersistentLevel));
-			return;
-		}
-		// Don't tell the server about this level
-		DuplicatePersistentLevel->bClientOnlyVisible = true;
-		DuplicateLevels.SetPersistentLevel(DuplicatePersistentLevel);
-		DuplicateLevels.AddLevel(DuplicatePersistentLevel);
-
 		if (!IsPartitionedWorld())
 		{
+			// Duplicate the persistent level and only dynamic levels, but don't add them to the world.
+			FLevelCollection DuplicateLevels;
+			DuplicateLevels.SetType(ELevelCollectionType::DynamicDuplicatedLevels);
+			DuplicateLevels.SetIsVisible(false);
+			ULevel* const DuplicatePersistentLevel = DuplicateLevelWithPrefix(PersistentLevel, 1);
+			if (!DuplicatePersistentLevel)
+			{
+				UE_LOG(LogWorld, Warning, TEXT("UWorld::DuplicateRequestedLevels: failed to duplicate persistent level %s. No duplicate level collection will be created."),
+					*GetFullNameSafe(PersistentLevel));
+				return;
+			}
+			// Don't tell the server about this level
+			DuplicatePersistentLevel->bClientOnlyVisible = true;
+			DuplicateLevels.SetPersistentLevel(DuplicatePersistentLevel);
+			DuplicateLevels.AddLevel(DuplicatePersistentLevel);
+
 			for (ULevelStreaming* StreamingLevel : StreamingLevels)
 			{
 				if (StreamingLevel && !StreamingLevel->bIsStatic)
@@ -8854,20 +8854,13 @@ void UWorld::DuplicateRequestedLevels(const FName MapName)
 					DuplicateLevels.AddLevel(DuplicatedLevel);
 				}
 			}
+
+			LevelCollections.Add(MoveTemp(DuplicateLevels));
 		}
 		else
 		{
-			UE_LOG(LogWorld, Log, TEXT("UWorld::DuplicateRequestedLevels: Attempting to duplicate levels for partitioned world. The persistent Level will be duplicated, but all streaming cells will be considered static."));
-
-			// If duplicating a partitioned world, remove and destroy the duplicated persistent level's UWorldPartition
-			UWorldPartition* DuplicateWP = DuplicatePersistentLevel->GetWorldPartition();
-			if (DuplicateWP)
-			{
-				DuplicateWP->MarkAsGarbage();
-			}
-		}
-
-		LevelCollections.Add(MoveTemp(DuplicateLevels));
+			UE_LOG(LogWorld, Error, TEXT("UWorld::DuplicateRequestedLevels: Attempted to duplicate streaming levels for partitioned world. This is not a supported operation."));
+		}	
 	}
 }
 
diff --git a/Engine/Source/Runtime/Engine/Private/WorldPartition/WorldPartitionLevelStreamingDynamic.cpp b/Engine/Source/Runtime/Engine/Private/WorldPartition/WorldPartitionLevelStreamingDynamic.cpp
--- a/Engine/Source/Runtime/Engine/Private/WorldPartition/WorldPartitionLevelStreamingDynamic.cpp
+++ b/Engine/Source/Runtime/Engine/Private/WorldPartition/WorldPartitionLevelStreamingDynamic.cpp
@@ -51,10 +51,6 @@ void UWorldPartitionLevelStreamingDynamic::Initialize(const UWorldPartitionRunti
 	bShouldBeAlwaysLoaded = InCell.IsAlwaysLoaded();
 	StreamingPriority = 0;
 
-	// Default WP streaming levels to static, so they're not hidden during instant replay playback
-	// This means only the persistent level will be duplicated/hidden
-	bIsStatic = true;
-
 	UWorld* CellOuterWorld = InCell.GetOuterWorld();
 #if WITH_EDITOR
 	check(ChildPackages.Num() == 0);
