Removed old unused versioning mechanisms from SparseVolumeTexture. Going forward, SVT is using object versioning (FArchive::UsingCustomVersion()) and future changes to the serialization format should try to maintain backwards compatibility.

#rnx

diff --git a/Engine/Source/Runtime/Engine/Classes/SparseVolumeTexture/SparseVolumeTexture.h b/Engine/Source/Runtime/Engine/Classes/SparseVolumeTexture/SparseVolumeTexture.h
--- a/Engine/Source/Runtime/Engine/Classes/SparseVolumeTexture/SparseVolumeTexture.h
+++ b/Engine/Source/Runtime/Engine/Classes/SparseVolumeTexture/SparseVolumeTexture.h
@@ -33,9 +33,6 @@ class FStreamingManager;
 
 struct ENGINE_API FHeader
 {
-	static const uint32 kVersion = 0; // The current data format version for the header.
-	uint32 Version = kVersion; // This version can be used to convert existing header to new version later.
-
 	FIntVector3 VirtualVolumeResolution = FIntVector3(0, 0, 0);
 	FIntVector3 VirtualVolumeAABBMin = FIntVector3(INT32_MAX, INT32_MAX, INT32_MAX);
 	FIntVector3 VirtualVolumeAABBMax = FIntVector3(INT32_MIN, INT32_MIN, INT32_MIN);
@@ -43,11 +40,10 @@ struct ENGINE_API FHeader
 	FIntVector3 PageTableVolumeAABBMin = FIntVector3(INT32_MAX, INT32_MAX, INT32_MAX);
 	FIntVector3 PageTableVolumeAABBMax = FIntVector3(INT32_MIN, INT32_MIN, INT32_MIN);
 	TStaticArray<EPixelFormat, 2> AttributesFormats = TStaticArray<EPixelFormat, 2>(InPlace, PF_Unknown);
-	TStaticArray<FVector4f, 2> FallbackValues = TStaticArray<FVector4f, 2>(InPlace, FVector4f());
+	TStaticArray<FVector4f, 2> FallbackValues = TStaticArray<FVector4f, 2>(InPlace, FVector4f(0.0f, 0.0f, 0.0f, 0.0f));
 
 	FHeader() = default;
 	FHeader(const FIntVector3& AABBMin, const FIntVector3& AABBMax, EPixelFormat FormatA, EPixelFormat FormatB, const FVector4f& FallbackValueA, const FVector4f& FallbackValueB);
-	void Serialize(FArchive& Ar);
 };
 
 // Describes a mip level of a SVT frame in terms of the sizes and offsets of the data in the built bulk data.
@@ -81,10 +77,6 @@ public:
 	TArray<uint8> RootData;
 	// Data for all streamable mip levels
 	FByteBulkData StreamableMipLevels;
-#if WITH_EDITORONLY_DATA
-	// FTextureData from which all the other data can be built with a call to Build()
-	UE::Serialization::FEditorBulkData SourceData;
-#endif
 
 	// These are used for logging and retrieving StreamableMipLevels from DDC in FStreamingManager
 #if WITH_EDITORONLY_DATA
@@ -103,9 +95,9 @@ public:
 	// Fills StreamableMipLevels with data from DDC. Returns true when done.
 	bool RebuildBulkDataFromCacheAsync(const UObject* Owner, bool& bFailed);
 	// Builds all the data from SourceData. Is called by Cache().
-	bool Build(USparseVolumeTextureFrame* Owner);
+	bool Build(USparseVolumeTextureFrame* Owner, UE::Serialization::FEditorBulkData& SourceData);
 	// Cache the built data to/from DDC. If bLocalCachingOnly is true, the read/write queries will only use the local DDC; otherwise the remote DDC will also be used.
-	void Cache(USparseVolumeTextureFrame* Owner, bool bLocalCachingOnly);
+	void Cache(USparseVolumeTextureFrame* Owner, UE::Serialization::FEditorBulkData& SourceData, bool bLocalCachingOnly);
 	// Sets empty default data. This is used when caching/building is canceled but some form of valid data is needed.
 	void SetDefault(EPixelFormat FormatA, EPixelFormat FormatB, const FVector4f& FallbackValueA, const FVector4f& FallbackValueB);
 #endif
@@ -163,6 +155,8 @@ private:
 }
 }
 
+FArchive& operator<<(FArchive& Ar, UE::SVT::FHeader& Header);
+
 enum ESparseVolumeTextureShaderUniform
 {
 	ESparseVolumeTexture_TileSize,
@@ -224,6 +218,7 @@ public:
 #endif // WITH_EDITOR
 };
 
+// Represents a frame in a SparseVolumeTexture sequence and owns the actual data needed for rendering. Is owned by a UStreamableSparseVolumeTexture object.
 UCLASS(ClassGroup = Rendering, BlueprintType)
 class ENGINE_API USparseVolumeTextureFrame : public USparseVolumeTexture
 {
@@ -283,13 +278,25 @@ public:
 	//~ End USparseVolumeTexture Interface.
 
 private:
+
 	UPROPERTY()
 	TObjectPtr<USparseVolumeTexture> Owner;
+
+	UPROPERTY()
 	int32 FrameIndex;
+
+#if WITH_EDITORONLY_DATA
+	// FTextureData from which the FResources data can be built with a call to FResources::Build()
+	UE::Serialization::FEditorBulkData SourceData;
+#endif
+
+	// Derived data used at runtime
 	UE::SVT::FResources Resources;
+	// Runtime render data
 	UE::SVT::FTextureRenderResources* TextureRenderResources;
 };
 
+// Represents a streamable SparseVolumeTexture asset and serves as base class for UStaticSparseVolumeTexture and UAnimatedSparseVolumeTexture. It has an array of USparseVolumeTextureFrame.
 UCLASS(ClassGroup = Rendering, BlueprintType)//, hidecategories = (Object))
 class UStreamableSparseVolumeTexture : public USparseVolumeTexture
 {
@@ -297,14 +304,6 @@ class UStreamableSparseVolumeTexture : public USparseVolumeTexture
 
 public:
 
-	enum class EInitState : uint8
-	{
-		Uninitialized,
-		Pending,
-		Done,
-		Failed,
-	};
-
 	UPROPERTY(VisibleAnywhere, Category = "Texture")
 	FIntVector VolumeResolution;
 
@@ -379,12 +378,27 @@ public:
 
 protected:
 
-	UPROPERTY()
+	UPROPERTY(Export)
 	TArray<TObjectPtr<USparseVolumeTextureFrame>> Frames;
+
 #if WITH_EDITORONLY_DATA
+
+	enum EInitState : uint8
+	{
+		EInitState_Uninitialized,
+		EInitState_Pending,
+		EInitState_Done,
+		EInitState_Failed,
+	};
+	
+	UPROPERTY()
 	FIntVector VolumeBoundsMin;
+	
+	UPROPERTY()
 	FIntVector VolumeBoundsMax;
-	EInitState InitState = EInitState::Uninitialized;
+	
+	UPROPERTY()
+	uint8 InitState = EInitState_Uninitialized;
 
 	// Ensures all frames have derived data (based on the source data and the current settings like TextureAddress modes etc.) cached to DDC and are ready for rendering.
 	// Disconnects this SVT from the streaming manager, calls Cache() on all frames and finally connects to FStreamingManager again.
@@ -392,6 +406,7 @@ protected:
 #endif // WITH_EDITORONLY_DATA
 };
 
+// Represents a streamable SparseVolumeTexture asset with a single frame. Although there is only a single frame, it is still recommended to use USparseVolumeTextureFrame::GetFrameAndIssueStreamingRequest().
 UCLASS(ClassGroup = Rendering, BlueprintType)//, hidecategories = (Object))
 class ENGINE_API UStaticSparseVolumeTexture : public UStreamableSparseVolumeTexture
 {
@@ -402,8 +417,10 @@ public:
 	UStaticSparseVolumeTexture();
 	virtual ~UStaticSparseVolumeTexture() = default;
 
+	//~ Begin UStreamableSparseVolumeTexture Interface.
 	// Override AppendFrame() to ensure that there is never more than a single frame in a static SVT
 	virtual bool AppendFrame(UE::SVT::FTextureData& UncookedFrame) override;
+	//~ End UStreamableSparseVolumeTexture Interface.
 
 	//~ Begin USparseVolumeTexture Interface.
 	int32 GetNumFrames() const override { return 1; }
@@ -412,7 +429,7 @@ public:
 private:
 };
 
-// UAnimatedSparseVolumeTexture inherit from USparseVolumeTexture to be viewed using the first frame by default.
+// Represents a streamable SparseVolumeTexture with one or more frames. Use USparseVolumeTextureFrame::GetFrameAndIssueStreamingRequest() to bind extract a particular frame to be used for rendering.
 UCLASS(ClassGroup = Rendering, BlueprintType)//, hidecategories = (Object))
 class ENGINE_API UAnimatedSparseVolumeTexture : public UStreamableSparseVolumeTexture
 {
@@ -431,6 +448,7 @@ private:
 	int32 PreviewFrameIndex;
 };
 
+// Utility (blueprint) class for controlling SparseVolumeTexture playback.
 UCLASS(ClassGroup = Rendering, BlueprintType)//, hidecategories = (Object))
 class ENGINE_API UAnimatedSparseVolumeTextureController : public UObject
 {
diff --git a/Engine/Source/Runtime/Engine/Private/SparseVolumeTexture/SparseVolumeTexture.cpp b/Engine/Source/Runtime/Engine/Private/SparseVolumeTexture/SparseVolumeTexture.cpp
--- a/Engine/Source/Runtime/Engine/Private/SparseVolumeTexture/SparseVolumeTexture.cpp
+++ b/Engine/Source/Runtime/Engine/Private/SparseVolumeTexture/SparseVolumeTexture.cpp
@@ -35,8 +35,6 @@
 #include "Serialization/EditorBulkDataReader.h"
 #include "Serialization/EditorBulkDataWriter.h"
 
-#include "ContentStreaming.h"
-
 #define LOCTEXT_NAMESPACE "USparseVolumeTexture"
 
 DEFINE_LOG_CATEGORY_STATIC(LogSparseVolumeTexture, Log, All);
@@ -65,6 +63,21 @@ FArchive& operator<<(FArchive& Ar, UE::SVT::FMipLevelStreamingInfo& MipLevelStre
 	return Ar;
 }
 
+FArchive& operator<<(FArchive& Ar, UE::SVT::FHeader& Header)
+{
+	Ar << Header.VirtualVolumeResolution;
+	Ar << Header.VirtualVolumeAABBMin;
+	Ar << Header.VirtualVolumeAABBMax;
+	Ar << Header.PageTableVolumeResolution;
+	Ar << Header.PageTableVolumeAABBMin;
+	Ar << Header.PageTableVolumeAABBMax;
+	UE::SVT::Private::SerializeEnumAs<uint8>(Ar, Header.AttributesFormats[0]);
+	UE::SVT::Private::SerializeEnumAs<uint8>(Ar, Header.AttributesFormats[1]);
+	Ar << Header.FallbackValues[0];
+	Ar << Header.FallbackValues[1];
+	return Ar;
+}
+
 ////////////////////////////////////////////////////////////////////////////////////////////////
 
 namespace UE
@@ -113,30 +126,6 @@ FHeader::FHeader(const FIntVector3& AABBMin, const FIntVector3& AABBMax, EPixelF
 	FallbackValues[1] = FallbackValueB;
 }
 
-void FHeader::Serialize(FArchive& Ar)
-{
-	Ar << Version;
-
-	if (Version == 0)
-	{
-		Ar << VirtualVolumeResolution;
-		Ar << VirtualVolumeAABBMin;
-		Ar << VirtualVolumeAABBMax;
-		Ar << PageTableVolumeResolution;
-		Ar << PageTableVolumeAABBMin;
-		Ar << PageTableVolumeAABBMax;
-		UE::SVT::Private::SerializeEnumAs<uint8>(Ar, AttributesFormats[0]);
-		UE::SVT::Private::SerializeEnumAs<uint8>(Ar, AttributesFormats[1]);
-		Ar << FallbackValues[0];
-		Ar << FallbackValues[1];
-	}
-	else
-	{
-		// FHeader needs to account for new version
-		check(false);
-	}
-}
-
 ////////////////////////////////////////////////////////////////////////////////////////////////
 
 void FResources::Serialize(FArchive& Ar, UObject* Owner, bool bCooked)
@@ -145,7 +134,7 @@ void FResources::Serialize(FArchive& Ar, UObject* Owner, bool bCooked)
 	FStripDataFlags StripFlags(Ar, 0);
 	if (!StripFlags.IsDataStrippedForServer())
 	{
-		Header.Serialize(Ar);
+		Ar << Header;
 
 		uint32 StoredResourceFlags;
 		if (Ar.IsSaving() && bCooked)
@@ -176,13 +165,6 @@ void FResources::Serialize(FArchive& Ar, UObject* Owner, bool bCooked)
 			StreamableMipLevels.Serialize(Ar, Owner, 0);
 		}
 
-#if WITH_EDITORONLY_DATA
-		if (!bCooked)
-		{
-			SourceData.Serialize(Ar, Owner);
-		}
-#endif
-
 #if !WITH_EDITORONLY_DATA
 		check(!HasStreamingData() || StreamableMipLevels.GetBulkDataSize() > 0);
 #endif
@@ -253,14 +235,14 @@ bool FResources::RebuildBulkDataFromCacheAsync(const UObject* Owner, bool& bFail
 	}
 }
 
-bool FResources::Build(USparseVolumeTextureFrame* Owner)
+bool FResources::Build(USparseVolumeTextureFrame* Owner, UE::Serialization::FEditorBulkData& SourceData)
 {
 	// Check if the virtualized bulk data payload is available
 	if (SourceData.HasPayloadData())
 	{
 		UE::Serialization::FEditorBulkDataReader SourceDataReader(SourceData);
 		FTextureData SourceTextureData;
-		SourceTextureData.Serialize(SourceDataReader);
+		SourceDataReader << SourceTextureData;
 
 		FTextureDataAddressingInfo AddressingInfo{};
 		AddressingInfo.VolumeResolution = Owner->GetVolumeResolution();
@@ -374,7 +356,7 @@ bool FResources::Build(USparseVolumeTextureFrame* Owner)
 	return false;
 }
 
-void FResources::Cache(USparseVolumeTextureFrame* Owner, bool bLocalCachingOnly)
+void FResources::Cache(USparseVolumeTextureFrame* Owner, UE::Serialization::FEditorBulkData& SourceData, bool bLocalCachingOnly)
 {
 	if (Owner->GetPackage()->bIsCookedForEditor)
 	{
@@ -446,7 +428,7 @@ void FResources::Cache(USparseVolumeTextureFrame* Owner, bool bLocalCachingOnly)
 	else
 	{
 		// DDC lookup failed! Build the data again.
-		const bool bBuiltSuccessfully = Build(Owner);
+		const bool bBuiltSuccessfully = Build(Owner, SourceData);
 		check(bBuiltSuccessfully);
 
 		FCacheRecordBuilder RecordBuilder(CacheKey);
@@ -719,8 +701,8 @@ bool USparseVolumeTextureFrame::Initialize(USparseVolumeTexture* InOwner, int32
 	Owner = InOwner;
 	FrameIndex = InFrameIndex;
 	{
-		UE::Serialization::FEditorBulkDataWriter SourceDataArchiveWriter(Resources.SourceData);
-		UncookedFrame.Serialize(SourceDataArchiveWriter);
+		UE::Serialization::FEditorBulkDataWriter SourceDataArchiveWriter(SourceData);
+		SourceDataArchiveWriter << UncookedFrame;
 	}
 
 	return true;
@@ -780,11 +762,15 @@ void USparseVolumeTextureFrame::Serialize(FArchive& Ar)
 {
 	Super::Serialize(Ar);
 
-	FStripDataFlags StripFlags(Ar);
-	
+#if WITH_EDITORONLY_DATA
+	if (!Ar.IsCooking())
+	{
+		SourceData.Serialize(Ar, this);
+	}
+#endif
+
 	bool bCooked = Ar.IsCooking();
 	Ar << bCooked;
-	Ar << FrameIndex;
 	Resources.Serialize(Ar, this, bCooked);
 }
 
@@ -811,7 +797,7 @@ bool USparseVolumeTextureFrame::IsCachedCookedPlatformDataLoaded(const ITargetPl
 	{
 		UE_LOG(LogSparseVolumeTexture, Log, TEXT("Failed to recover SparseVolumeTexture streaming from DDC for '%s' Frame %i. Rebuilding and retrying."), *Owner->GetPathName(), FrameIndex);
 
-		Resources.Cache(this, CastChecked<UStreamableSparseVolumeTexture>(Owner)->bLocalDDCOnly);
+		Resources.Cache(this, SourceData, CastChecked<UStreamableSparseVolumeTexture>(Owner)->bLocalDDCOnly);
 		return false;
 	}
 
@@ -843,7 +829,7 @@ void USparseVolumeTextureFrame::Cache(bool bSkipDDCAndSetResourcesToDefault)
 	}
 	else
 	{
-		Resources.Cache(this, CastChecked<UStreamableSparseVolumeTexture>(Owner)->bLocalDDCOnly);
+		Resources.Cache(this, SourceData, CastChecked<UStreamableSparseVolumeTexture>(Owner)->bLocalDDCOnly);
 	}
 	CreateTextureRenderResources();
 }
@@ -859,7 +845,7 @@ UStreamableSparseVolumeTexture::UStreamableSparseVolumeTexture(const FObjectInit
 bool UStreamableSparseVolumeTexture::BeginInitialize(int32 NumExpectedFrames)
 {
 #if WITH_EDITORONLY_DATA
-	if (InitState != EInitState::Uninitialized)
+	if (InitState != EInitState_Uninitialized)
 	{
 		UE_LOG(LogSparseVolumeTexture, Error, TEXT("Tried to call UStreamableSparseVolumeTexture::BeginInitialize() while not in the Uninitialized init state."));
 		return false;
@@ -872,7 +858,7 @@ bool UStreamableSparseVolumeTexture::BeginInitialize(int32 NumExpectedFrames)
 	check(FormatA == PF_Unknown);
 	check(FormatB == PF_Unknown);
 
-	InitState = EInitState::Pending;
+	InitState = EInitState_Pending;
 
 	return true;
 #else
@@ -883,7 +869,7 @@ bool UStreamableSparseVolumeTexture::BeginInitialize(int32 NumExpectedFrames)
 bool UStreamableSparseVolumeTexture::AppendFrame(UE::SVT::FTextureData& UncookedFrame)
 {
 #if WITH_EDITORONLY_DATA
-	if (InitState != EInitState::Pending)
+	if (InitState != EInitState_Pending)
 	{
 		UE_LOG(LogSparseVolumeTexture, Error, TEXT("Tried to call UStreamableSparseVolumeTexture::AppendFrame() while not in the Pending init state."));
 		return false;
@@ -947,7 +933,7 @@ bool UStreamableSparseVolumeTexture::AppendFrame(UE::SVT::FTextureData& Uncooked
 bool UStreamableSparseVolumeTexture::EndInitialize(int32 InNumMipLevels)
 {
 #if WITH_EDITORONLY_DATA
-	if (InitState != EInitState::Pending)
+	if (InitState != EInitState_Pending)
 	{
 		UE_LOG(LogSparseVolumeTexture, Error, TEXT("Tried to call UStreamableSparseVolumeTexture::EndInitialize() while not in the Pending init state."));
 		return false;
@@ -982,7 +968,7 @@ bool UStreamableSparseVolumeTexture::EndInitialize(int32 InNumMipLevels)
 		Frame->PostLoad();
 	}
 
-	InitState = EInitState::Done;
+	InitState = EInitState_Done;
 
 	return true;
 #else
@@ -1047,31 +1033,10 @@ void UStreamableSparseVolumeTexture::Serialize(FArchive& Ar)
 {
 	Super::Serialize(Ar);
 
-	// To ensure that we keep the same binary format between builds with and without editor-only data, we serialize dummy data if the editor-only data is stripped.
-	// SVT_TODO: Is there a cleaner way of doing this?
-	FIntVector* VolumeBoundsMinPtr = nullptr;
-	FIntVector* VolumeBoundsMaxPtr = nullptr;
-	EInitState* InitStatePtr = nullptr;
-
 #if WITH_EDITORONLY_DATA
 	// Check that we are not trying to cook unitialized data!
-	check(!Ar.IsCooking() || InitState == EInitState::Done);
-
-	VolumeBoundsMinPtr = &VolumeBoundsMin;
-	VolumeBoundsMaxPtr = &VolumeBoundsMax;
-	InitStatePtr = &InitState;
-#else
-	FIntVector VolumeBoundsMinDummy{};
-	FIntVector VolumeBoundsMaxDummy{};
-	EInitState InitStateDummy{};
-	VolumeBoundsMinPtr = &VolumeBoundsMinDummy;
-	VolumeBoundsMaxPtr = &VolumeBoundsMaxDummy;
-	InitStatePtr = &InitStateDummy;
+	check(!Ar.IsCooking() || InitState == EInitState_Done);
 #endif
-
-	UE::SVT::Private::SerializeEnumAs<uint8>(Ar, *InitStatePtr);
-	Ar << *VolumeBoundsMinPtr;
-	Ar << *VolumeBoundsMaxPtr;
 }
 
 #if WITH_EDITOR
@@ -1112,7 +1077,7 @@ void UStreamableSparseVolumeTexture::GetResourceSizeEx(FResourceSizeEx& Cumulati
 #if WITH_EDITORONLY_DATA
 void UStreamableSparseVolumeTexture::RecacheFrames()
 {
-	if (InitState != EInitState::Done)
+	if (InitState != EInitState_Done)
 	{
 		UE_LOG(LogSparseVolumeTexture, Warning, TEXT("Tried to cache derived data of an uninitialized SVT: %s"), *GetName());
 		return;
diff --git a/Engine/Source/Runtime/Engine/Private/SparseVolumeTexture/SparseVolumeTextureData.cpp b/Engine/Source/Runtime/Engine/Private/SparseVolumeTexture/SparseVolumeTextureData.cpp
--- a/Engine/Source/Runtime/Engine/Private/SparseVolumeTexture/SparseVolumeTextureData.cpp
+++ b/Engine/Source/Runtime/Engine/Private/SparseVolumeTexture/SparseVolumeTextureData.cpp
@@ -16,41 +16,29 @@ DEFINE_LOG_CATEGORY_STATIC(LogSparseVolumeTextureData, Log, All);
 
 ////////////////////////////////////////////////////////////////////////////////////////////////
 
-namespace UE
-{
-namespace SVT
+FArchive& operator<<(FArchive& Ar, UE::SVT::FTextureData::FMipMap& MipMap)
 {
+	Ar << MipMap.PageTable;
+	Ar << MipMap.PhysicalTileDataA;
+	Ar << MipMap.PhysicalTileDataB;
+	Ar << MipMap.NumPhysicalTiles;
+	return Ar;
+}
 
-void FTextureData::Serialize(FArchive& Ar)
+FArchive& operator<<(FArchive& Ar, UE::SVT::FTextureData& TextureData)
 {
-	Header.Serialize(Ar);
+	Ar << TextureData.Header;
+	Ar << TextureData.FallbackValuesQuantized;
+	Ar << TextureData.MipMaps;
+	return Ar;
+}
 
-	Ar << Version;
+////////////////////////////////////////////////////////////////////////////////////////////////
 
-	if (Version == 0)
-	{
-		int32 NumMipLevels = MipMaps.Num();
-		Ar << NumMipLevels;
-		if (Ar.IsLoading())
-		{
-			MipMaps.Reset(NumMipLevels);
-			MipMaps.SetNum(NumMipLevels);
-		}
-		for (int32 MipLevel = 0; MipLevel < NumMipLevels; ++MipLevel)
-		{
-			FMipMap& MipMap = MipMaps[MipLevel];
-			Ar << MipMap.PageTable;
-			Ar << MipMap.PhysicalTileDataA;
-			Ar << MipMap.PhysicalTileDataB;
-			Ar << MipMap.NumPhysicalTiles;
-		}
-	}
-	else
-	{
-		// FTextureData needs to account for new version
-		check(false);
-	}
-}
+namespace UE
+{
+namespace SVT
+{
 
 bool FTextureData::Create(const ITextureDataProvider& DataProvider)
 {
diff --git a/Engine/Source/Runtime/Engine/Public/SparseVolumeTexture/SparseVolumeTextureData.h b/Engine/Source/Runtime/Engine/Public/SparseVolumeTexture/SparseVolumeTextureData.h
--- a/Engine/Source/Runtime/Engine/Public/SparseVolumeTexture/SparseVolumeTextureData.h
+++ b/Engine/Source/Runtime/Engine/Public/SparseVolumeTexture/SparseVolumeTextureData.h
@@ -15,7 +15,7 @@ struct ENGINE_API FTextureDataCreateInfo
 	FIntVector3 VirtualVolumeAABBMin = FIntVector3(INT32_MAX, INT32_MAX, INT32_MAX);
 	FIntVector3 VirtualVolumeAABBMax = FIntVector3(INT32_MIN, INT32_MIN, INT32_MIN);
 	TStaticArray<EPixelFormat, 2> AttributesFormats = TStaticArray<EPixelFormat, 2>(InPlace, PF_Unknown); // Currently supported: PF_R8, PF_R8G8, PF_R8G8B8A8, PF_R16F, PF_G16R16F, PF_FloatRGBA, PF_R32_FLOAT, PF_G32R32F, PF_A32B32G32R32F
-	TStaticArray<FVector4f, 2> FallbackValues = TStaticArray<FVector4f, 2>(InPlace, FVector4f());
+	TStaticArray<FVector4f, 2> FallbackValues = TStaticArray<FVector4f, 2>(InPlace, FVector4f(0.0f, 0.0f, 0.0f, 0.0f));
 };
 
 class ENGINE_API ITextureDataProvider
@@ -47,15 +47,10 @@ struct ENGINE_API FTextureData
 		int32 NumPhysicalTiles;
 	};
 
-	static const uint32 kVersion = 0; // The current data format version for the raw source data.
-	uint32 Version = kVersion; // This version can be used to convert existing source data to new version later.
-
-	UE::SVT::FHeader Header = {};
-	TStaticArray<FVector4f, 2> FallbackValuesQuantized = TStaticArray<FVector4f, 2>(InPlace, FVector4f());
+	FHeader Header = {};
+	TStaticArray<FVector4f, 2> FallbackValuesQuantized = TStaticArray<FVector4f, 2>(InPlace, FVector4f(0.0f, 0.0f, 0.0f, 0.0f));
 	TArray<FMipMap> MipMaps;
 
-	void Serialize(FArchive& Ar);
-	
 	bool Create(const ITextureDataProvider& DataProvider);
 	bool CreateFromDense(const FTextureDataCreateInfo& CreateInfo, const TArrayView64<uint8>& VoxelDataA, const TArrayView64<uint8>& VoxelDataB);
 	void CreateDefault();
@@ -75,3 +70,5 @@ private:
 
 }
 }
+
+FArchive& operator<<(FArchive& Ar, UE::SVT::FTextureData& TextureData);
