Added detail mode to ISM descriptor

diff --git a/Engine/Source/Runtime/Engine/Private/ISMPartition/ISMComponentDescriptor.cpp b/Engine/Source/Runtime/Engine/Private/ISMPartition/ISMComponentDescriptor.cpp
--- a/Engine/Source/Runtime/Engine/Private/ISMPartition/ISMComponentDescriptor.cpp
+++ b/Engine/Source/Runtime/Engine/Private/ISMPartition/ISMComponentDescriptor.cpp
@@ -99,6 +99,7 @@ void FISMComponentDescriptorBase::InitFrom(const UStaticMeshComponent* Template,
 	bEvaluateWorldPositionOffset = Template->bEvaluateWorldPositionOffset;
 	WorldPositionOffsetDisableDistance = Template->WorldPositionOffsetDisableDistance;
 	ShadowCacheInvalidationBehavior = Template->ShadowCacheInvalidationBehavior;
+	DetailMode = Template->DetailMode;
 	// Determine if this instance must render with reversed culling based on both scale and the component property
 	const bool bIsLocalToWorldDeterminantNegative = Template->GetRenderMatrix().Determinant() < 0;
 	bReverseCulling = Template->bReverseCulling != bIsLocalToWorldDeterminantNegative;
@@ -204,6 +205,7 @@ bool FISMComponentDescriptorBase::operator==(const FISMComponentDescriptorBase&
 	bGenerateOverlapEvents == Other.bGenerateOverlapEvents &&
 	WorldPositionOffsetDisableDistance == Other.WorldPositionOffsetDisableDistance &&
 	ShadowCacheInvalidationBehavior == Other.ShadowCacheInvalidationBehavior &&
+	DetailMode == Other.DetailMode &&
 #if WITH_EDITORONLY_DATA
 	HLODBatchingPolicy == Other.HLODBatchingPolicy &&
 	bIncludeInHLOD == Other.bIncludeInHLOD &&
@@ -313,6 +315,7 @@ void FISMComponentDescriptorBase::InitComponent(UInstancedStaticMeshComponent* I
 	ISMComponent->SetGenerateOverlapEvents(bGenerateOverlapEvents);
 	ISMComponent->WorldPositionOffsetDisableDistance = WorldPositionOffsetDisableDistance;
 	ISMComponent->ShadowCacheInvalidationBehavior = ShadowCacheInvalidationBehavior;
+	ISMComponent->DetailMode = DetailMode;
 	
 #if WITH_EDITORONLY_DATA
 	ISMComponent->HLODBatchingPolicy = HLODBatchingPolicy;
diff --git a/Engine/Source/Runtime/Engine/Public/ISMPartition/ISMComponentDescriptor.h b/Engine/Source/Runtime/Engine/Public/ISMPartition/ISMComponentDescriptor.h
--- a/Engine/Source/Runtime/Engine/Public/ISMPartition/ISMComponentDescriptor.h
+++ b/Engine/Source/Runtime/Engine/Public/ISMPartition/ISMComponentDescriptor.h
@@ -165,6 +165,9 @@ public:
 
 	UPROPERTY(EditAnywhere, Category = "Component Settings")
 	EShadowCacheInvalidationBehavior ShadowCacheInvalidationBehavior;
+
+	UPROPERTY(EditAnywhere, Category = "Component Settings")
+	TEnumAsByte<enum EDetailMode> DetailMode;
 };
 
 USTRUCT()
