AnimMontage: Exposed IsValidAdditiveSlot to blueprint

diff --git a/Engine/Source/Runtime/Engine/Classes/Animation/AnimMontage.h b/Engine/Source/Runtime/Engine/Classes/Animation/AnimMontage.h
--- a/Engine/Source/Runtime/Engine/Classes/Animation/AnimMontage.h
+++ b/Engine/Source/Runtime/Engine/Classes/Animation/AnimMontage.h
@@ -772,6 +772,7 @@ public:
 	/*
 	 * Check if this slot has valid animation data
 	 */
+	UFUNCTION(BlueprintPure, Category = "Montage")
 	bool IsValidAdditiveSlot(const FName& SlotNodeName) const;
 
 	/** Get FCompositeSection with InSectionName */
