FActorDescContainerInstance::UpdateBounds(): skip Main world only actor bounds

#rnx

diff --git a/Engine/Source/Runtime/Engine/Private/WorldPartition/WorldPartitionSubsystem.cpp b/Engine/Source/Runtime/Engine/Private/WorldPartition/WorldPartitionSubsystem.cpp
--- a/Engine/Source/Runtime/Engine/Private/WorldPartition/WorldPartitionSubsystem.cpp
+++ b/Engine/Source/Runtime/Engine/Private/WorldPartition/WorldPartitionSubsystem.cpp
@@ -343,7 +343,8 @@ void UWorldPartitionSubsystem::FActorDescContainerInstanceManager::FActorDescCon
 	Bounds.Init();
 	for (FActorDescList::TIterator<> ActorDescIt(Container); ActorDescIt; ++ActorDescIt)
 	{
-		if (ActorDescIt->GetActorNativeClass()->IsChildOf<ALevelBounds>())
+		// FActorDescContainerInstance are currently only used by FLevelInstanceActorDescs and so we don't consider actors that should exit only in the main world.
+		if (ActorDescIt->IsMainWorldOnly())
 		{
 			continue;
 		}
