Fix iris/plan normal transformation for Substrate eye shader.

diff --git a/Engine/Source/Runtime/Engine/Private/Materials/MaterialExpressions.cpp b/Engine/Source/Runtime/Engine/Private/Materials/MaterialExpressions.cpp
--- a/Engine/Source/Runtime/Engine/Private/Materials/MaterialExpressions.cpp
+++ b/Engine/Source/Runtime/Engine/Private/Materials/MaterialExpressions.cpp
@@ -24428,8 +24428,8 @@ int32 UMaterialExpressionStrataEyeBSDF::Compile(class FMaterialCompiler* Compile
 		CompileWithDefaultFloat1(Compiler, Roughness,	 0.5f),
 		CompileWithDefaultFloat1(Compiler, IrisMask,	 0.0f),
 		CompileWithDefaultFloat1(Compiler, IrisDistance, 0.0f),
-		CompileWithDefaultFloat3(Compiler, IrisNormal,   0.0f, 0.0f, 1.0f),
-		CompileWithDefaultFloat3(Compiler, IrisPlaneNormal, 0.0f, 0.0f, 1.0f),
+		CompileWithDefaultNormalWS(Compiler, IrisNormal),
+		CompileWithDefaultNormalWS(Compiler, IrisPlaneNormal),
 		SSSProfileCodeChunk != INDEX_NONE ? SSSProfileCodeChunk : Compiler->Constant(0.0f),
 		CompileWithDefaultFloat3(Compiler, EmissiveColor,0.0f, 0.0f, 0.0f),
 		CorneaNormalCodeChunk,
