Fix incorrect DistanceFieldApproxAO when user specifies Normal.

diff --git a/Engine/Source/Runtime/Engine/Private/Materials/MaterialExpressions.cpp b/Engine/Source/Runtime/Engine/Private/Materials/MaterialExpressions.cpp
--- a/Engine/Source/Runtime/Engine/Private/Materials/MaterialExpressions.cpp
+++ b/Engine/Source/Runtime/Engine/Private/Materials/MaterialExpressions.cpp
@@ -19320,7 +19320,7 @@ int32 UMaterialExpressionDistanceFieldApproxAO::Compile(class FMaterialCompiler*
 
 	if (Normal.GetTracedInput().Expression)
 	{
-		NormalArg = Position.Compile(Compiler);
+		NormalArg = Normal.Compile(Compiler);
 	}
 	else
 	{
