Add a blueprint callable function to shuffle an array using a 'random stream' for the source of randomness

diff --git a/Engine/Source/Runtime/Engine/Classes/Kismet/KismetArrayLibrary.h b/Engine/Source/Runtime/Engine/Classes/Kismet/KismetArrayLibrary.h
--- a/Engine/Source/Runtime/Engine/Classes/Kismet/KismetArrayLibrary.h
+++ b/Engine/Source/Runtime/Engine/Classes/Kismet/KismetArrayLibrary.h
@@ -46,6 +46,15 @@ class ENGINE_API UKismetArrayLibrary : public UBlueprintFunctionLibrary
 	*/
 	UFUNCTION(BlueprintCallable, CustomThunk, meta=(DisplayName = "Shuffle", CompactNodeTitle = "SHUFFLE", ArrayParm = "TargetArray"), Category="Utilities|Array")
 	static void Array_Shuffle(const TArray<int32>& TargetArray);
+	
+	/** 
+	 * Shuffle (randomize) the elements of an array from a specific stream of random data, useful for achieving determinism
+	 *
+	 * @param	TargetArray		The array to shuffle
+	 * @param	RandomStream	The random stream
+	 */
+	UFUNCTION(BlueprintCallable, CustomThunk, meta=(DisplayName = "Shuffle from Stream", CompactNodeTitle = "SHUFFLE", ArrayParm = "TargetArray"), Category="Utilities|Array")
+	static void Array_ShuffleFromStream(const TArray<int32>& TargetArray, UPARAM(Ref) FRandomStream& RandomStream);
 
 	/** 
 	 * Checks if two arrays are memberwise identical
@@ -261,6 +270,7 @@ class ENGINE_API UKismetArrayLibrary : public UBlueprintFunctionLibrary
 	static int32 GenericArray_Add(void* TargetArray, const FArrayProperty* ArrayProp, const void* NewItem);
 	static int32 GenericArray_AddUnique(void* TargetArray, const FArrayProperty* ArrayProp, const void* NewItem);
 	static void GenericArray_Shuffle(void* TargetArray, const FArrayProperty* ArrayProp);
+	static void GenericArray_ShuffleFromStream(void* TargetArray, const FArrayProperty* ArrayProp, FRandomStream* RandomStream);
 	static bool GenericArray_Identical(void* ArrayA, const FArrayProperty* ArrayAProp, void* ArrayB, const FArrayProperty* ArrayBProperty);
 	static void GenericArray_Append(void* TargetArray, const FArrayProperty* TargetArrayProp, void* SourceArray, const FArrayProperty* SourceArrayProperty);
 	static void GenericArray_Insert(void* TargetArray, const FArrayProperty* ArrayProp, const void* NewItem, int32 Index);
@@ -377,6 +387,29 @@ public:
 		P_NATIVE_END;
 	}
 
+	DECLARE_FUNCTION(execArray_ShuffleFromStream)
+	{
+		Stack.MostRecentProperty = nullptr;
+		Stack.StepCompiledIn<FArrayProperty>(NULL);
+		void* ArrayAddr = Stack.MostRecentPropertyAddress;
+		FArrayProperty* ArrayProperty = CastField<FArrayProperty>(Stack.MostRecentProperty);
+		if (!ArrayProperty)
+		{
+			Stack.bArrayContextFailed = true;
+			return;
+		}
+
+		Stack.MostRecentProperty = nullptr;
+		Stack.StepCompiledIn<FProperty>(nullptr);
+		FRandomStream* RandomStream = (FRandomStream*)Stack.MostRecentPropertyAddress;
+
+		P_FINISH;
+		P_NATIVE_BEGIN;
+		MARK_PROPERTY_DIRTY(Stack.Object, ArrayProperty);
+		GenericArray_ShuffleFromStream(ArrayAddr, ArrayProperty, RandomStream);
+		P_NATIVE_END;
+	}
+
 	DECLARE_FUNCTION(execArray_Identical)
 	{
 		// Retrieve the first array
diff --git a/Engine/Source/Runtime/Engine/Private/KismetArrayLibrary.cpp b/Engine/Source/Runtime/Engine/Private/KismetArrayLibrary.cpp
--- a/Engine/Source/Runtime/Engine/Private/KismetArrayLibrary.cpp
+++ b/Engine/Source/Runtime/Engine/Private/KismetArrayLibrary.cpp
@@ -275,6 +275,23 @@ void UKismetArrayLibrary::GenericArray_Shuffle(void* TargetArray, const FArrayPr
 	}
 }
 
+void UKismetArrayLibrary::GenericArray_ShuffleFromStream(void* TargetArray, const FArrayProperty* ArrayProp, FRandomStream* RandomStream)
+{
+	if (TargetArray && RandomStream)
+	{
+		FScriptArrayHelper ArrayHelper(ArrayProp, TargetArray);
+		int32 LastIndex = ArrayHelper.Num() - 1;
+		for (int32 i = 0; i <= LastIndex; ++i)
+		{
+			const int32 Index = RandomStream->RandRange(i, LastIndex);
+			if (i != Index)
+			{
+				ArrayHelper.SwapValues(i, Index);
+			}
+		}
+	}
+}
+
 void UKismetArrayLibrary::GenericArray_Clear(void* TargetArray, const FArrayProperty* ArrayProp)
 {
 	if( TargetArray )
