- Include custom prim data in UB hash otherwise we don't update the relevant UBs
- When this data is normally update the proxy's main UB is updated not the variants.

#rnx

diff --git a/Engine/Plugins/FX/Niagara/Source/Niagara/Private/NiagaraComponent.cpp b/Engine/Plugins/FX/Niagara/Source/Niagara/Private/NiagaraComponent.cpp
--- a/Engine/Plugins/FX/Niagara/Source/Niagara/Private/NiagaraComponent.cpp
+++ b/Engine/Plugins/FX/Niagara/Source/Niagara/Private/NiagaraComponent.cpp
@@ -354,6 +354,14 @@ TUniformBuffer<FPrimitiveUniformShaderParameters>* FNiagaraSceneProxy::GetCustom
 		KeyHash = HashCombine(KeyHash, UE::Math::GetTypeHash(InstanceBounds.Max));
 	}
 
+	const FCustomPrimitiveData* LocalCustomPrimitiveData = GetCustomPrimitiveData();
+	if (LocalCustomPrimitiveData && LocalCustomPrimitiveData->Data.Num())
+	{
+		const uint32 CustomFloatHash = FCrc::MemCrc_DEPRECATED(LocalCustomPrimitiveData->Data.GetData(), LocalCustomPrimitiveData->Data.Num() * LocalCustomPrimitiveData->Data.GetTypeSize());
+		KeyHash = HashCombine(KeyHash, CustomFloatHash);
+	}
+
+	CustomUniformBuffers.Empty();
 	TUniformBuffer<FPrimitiveUniformShaderParameters>*& CustomUBRef = CustomUniformBuffers.FindOrAdd(KeyHash);
 	if (CustomUBRef == nullptr)
 	{
@@ -364,7 +372,7 @@ TUniformBuffer<FPrimitiveUniformShaderParameters>* FNiagaraSceneProxy::GetCustom
 				.ActorWorldPosition(GetActorPosition())
 				.WorldBounds(GetBounds())
 				.LocalBounds(GetLocalBounds())
-				.CustomPrimitiveData(GetCustomPrimitiveData())
+				.CustomPrimitiveData(LocalCustomPrimitiveData)
 				.LightingChannelMask(GetLightingChannelMask())
 				.LightmapDataIndex(LocalPrimitiveSceneInfo ? LocalPrimitiveSceneInfo->GetLightmapDataOffset() : 0)
 				.LightmapUVIndex(GetLightMapCoordinateIndex())
diff --git a/Engine/Source/Runtime/Engine/Private/Particles/ParticleSystemRender.cpp b/Engine/Source/Runtime/Engine/Private/Particles/ParticleSystemRender.cpp
--- a/Engine/Source/Runtime/Engine/Private/Particles/ParticleSystemRender.cpp
+++ b/Engine/Source/Runtime/Engine/Private/Particles/ParticleSystemRender.cpp
@@ -6834,13 +6834,26 @@ FPrimitiveViewRelevance FParticleSystemSceneProxy::GetViewRelevance(const FScene
 void FParticleSystemSceneProxy::OnTransformChanged()
 {
 	WorldSpacePrimitiveUniformBuffer.ReleaseResource();
+	WorldSpaceUBHash = 0;
 }
 
 void FParticleSystemSceneProxy::UpdateWorldSpacePrimitiveUniformBuffer() const
 {
 	check(IsInRenderingThread());
-	if (!WorldSpacePrimitiveUniformBuffer.IsInitialized())
+
+	// Hash custom floats because we need to invalidate this UB if they don't match otherwise updates to the buffer won't work
+	uint32 NewWorldSpaceUBHash = 0;
+	const FCustomPrimitiveData* LocalCustomPrimitiveData = GetCustomPrimitiveData();
+	if (LocalCustomPrimitiveData && LocalCustomPrimitiveData->Data.Num())
+	{
+		NewWorldSpaceUBHash = FCrc::MemCrc32(LocalCustomPrimitiveData->Data.GetData(), LocalCustomPrimitiveData->Data.Num() * LocalCustomPrimitiveData->Data.GetTypeSize());
+	}
+
+	const bool bNeedsInit = !WorldSpacePrimitiveUniformBuffer.IsInitialized();
+
+	if (bNeedsInit || (WorldSpaceUBHash != NewWorldSpaceUBHash))
 	{
+		WorldSpaceUBHash = NewWorldSpaceUBHash;
 		WorldSpacePrimitiveUniformBuffer.SetContents(
 			FPrimitiveUniformShaderParametersBuilder{}
 			.Defaults()
@@ -6853,8 +6866,13 @@ void FParticleSystemSceneProxy::UpdateWorldSpacePrimitiveUniformBuffer() const
 				.LightingChannelMask(GetLightingChannelMask())
 				.UseSingleSampleShadowFromStationaryLights(UseSingleSampleShadowFromStationaryLights())
 				.UseVolumetricLightmap(GetScene().HasPrecomputedVolumetricLightmap_RenderThread())
+				.CustomPrimitiveData(GetCustomPrimitiveData())
 			.Build()
 		);
+	}
+
+	if ( bNeedsInit)
+	{
 		WorldSpacePrimitiveUniformBuffer.InitResource();
 	}
 }
diff --git a/Engine/Source/Runtime/Engine/Public/ParticleHelper.h b/Engine/Source/Runtime/Engine/Public/ParticleHelper.h
--- a/Engine/Source/Runtime/Engine/Public/ParticleHelper.h
+++ b/Engine/Source/Runtime/Engine/Public/ParticleHelper.h
@@ -2532,6 +2532,7 @@ protected:
 
 	/** The primitive's uniform buffer.  Mutable because it is cached state during GetDynamicMeshElements. */
 	mutable TUniformBuffer<FPrimitiveUniformShaderParameters> WorldSpacePrimitiveUniformBuffer;
+	mutable uint32 WorldSpaceUBHash = 0;
 
 	/** Pool for holding FMeshBatches to reduce allocations. */
 	TIndirectArray<FMeshBatch, TInlineAllocator<4> > MeshBatchPool;
