Anim montage: Improve tooltip for IsValidAdditiveSlot and exposed GetGroupName to blueprint.

diff --git a/Engine/Source/Runtime/Engine/Classes/Animation/AnimMontage.h b/Engine/Source/Runtime/Engine/Classes/Animation/AnimMontage.h
--- a/Engine/Source/Runtime/Engine/Classes/Animation/AnimMontage.h
+++ b/Engine/Source/Runtime/Engine/Classes/Animation/AnimMontage.h
@@ -769,8 +769,9 @@ public:
 	 */
 	ENGINE_API void UpdateLinkableElements(int32 SlotIdx, int32 SegmentIdx);
 #endif
-	/*
-	 * Check if this slot has valid animation data
+	/** Check if this slot has a valid additive animation for the specified slot.
+	 * The slot name should not include the group name.
+	 * i.e. for "DefaultGroup.DefaultSlot", the slot name is "DefaultSlot".
 	 */
 	UFUNCTION(BlueprintPure, Category = "Montage")
 	bool IsValidAdditiveSlot(const FName& SlotNodeName) const;
@@ -850,6 +851,7 @@ public:
 	ENGINE_API FTransform ExtractRootMotionFromTrackRange(float StartTrackPosition, float EndTrackPosition) const;
 
 	/** Get the Montage's Group Name. This is the group from the first slot.  */
+	UFUNCTION(BlueprintPure, Category = "Montage")
 	ENGINE_API FName GetGroupName() const;
 
 	/** true if valid, false otherwise. Will log warning if not valid. */
