FPrimitiveSceneProxy::bSplineMesh was not being initialized.

diff --git a/Engine/Source/Runtime/Engine/Private/PrimitiveSceneProxy.cpp b/Engine/Source/Runtime/Engine/Private/PrimitiveSceneProxy.cpp
--- a/Engine/Source/Runtime/Engine/Private/PrimitiveSceneProxy.cpp
+++ b/Engine/Source/Runtime/Engine/Private/PrimitiveSceneProxy.cpp
@@ -346,6 +346,7 @@ FPrimitiveSceneProxy::FPrimitiveSceneProxy(const UPrimitiveComponent* InComponen
 #endif
 ,	bAllowApproximateOcclusion(InComponent->Mobility != EComponentMobility::Movable)
 ,   bHoldout(InComponent->bHoldout)
+,	bSplineMesh(false)
 ,	bUseAsOccluder(InComponent->bUseAsOccluder)
 ,	bSelectable(InComponent->bSelectable)
 ,	bHasPerInstanceHitProxies(InComponent->bHasPerInstanceHitProxies)
