Networked Physics: Handle FPhysicsObjects in FPhysScene_Chaos::PopulateReplicationCache(), this makes cluster unions cache their replication state and unlocks the refactored replication logic for CUs.

diff --git a/Engine/Source/Runtime/Engine/Private/PhysicsEngine/Experimental/PhysScene_Chaos.cpp b/Engine/Source/Runtime/Engine/Private/PhysicsEngine/Experimental/PhysScene_Chaos.cpp
--- a/Engine/Source/Runtime/Engine/Private/PhysicsEngine/Experimental/PhysScene_Chaos.cpp
+++ b/Engine/Source/Runtime/Engine/Private/PhysicsEngine/Experimental/PhysScene_Chaos.cpp
@@ -15,6 +15,7 @@
 #include "PhysicsEngine/ConstraintInstance.h"
 #include "PhysicsEngine/PhysicsCollisionHandler.h"
 #include "PhysicsEngine/PhysicsObjectExternalInterface.h"
+#include "Chaos/PhysicsObjectInternalInterface.h"
 #include "Physics/Experimental/ChaosEventRelay.h"
 #include "EngineUtils.h"
 
@@ -1218,6 +1219,18 @@ void FPhysScene_Chaos::AddToComponentMaps(UPrimitiveComponent* Component, IPhysi
 
 void FPhysScene_Chaos::PopulateReplicationCache(const int32 PhysicsStep)
 {
+	auto ReplicationCacheHelper = [this](auto& Handle, UPrimitiveComponent* RootComponent)
+	{
+		FObjectKey Key(RootComponent);
+		FRigidBodyState& ReplicationState = ReplicationCache.Map.FindOrAdd(Key);
+
+		ReplicationState.Position = Handle->X();
+		ReplicationState.Quaternion = Handle->R();
+		ReplicationState.LinVel = Handle->V();
+		ReplicationState.AngVel = Handle->W();
+		ReplicationState.Flags = Handle->ObjectState() == Chaos::EObjectStateType::Sleeping ? ERigidBodyFlags::Sleeping : 0;
+	};
+
 	ReplicationCache.ServerFrame = PhysicsStep;
 	for (TActorIterator<AActor> It(GetOwningWorld()); It; ++It)
 	{
@@ -1227,16 +1240,19 @@ void FPhysScene_Chaos::PopulateReplicationCache(const int32 PhysicsStep)
 			UPrimitiveComponent* RootComponent = Cast<UPrimitiveComponent>(Actor->GetRootComponent());
 			if (RootComponent && RootComponent->IsSimulatingPhysics())
 			{
-				if (FBodyInstanceAsyncPhysicsTickHandle Handle = RootComponent->GetBodyInstanceAsyncPhysicsTickHandle())
+				if (FBodyInstanceAsyncPhysicsTickHandle BIHandle = RootComponent->GetBodyInstanceAsyncPhysicsTickHandle())
 				{
-					FObjectKey Key(RootComponent);
-					FRigidBodyState& ReplicationState = ReplicationCache.Map.FindOrAdd(Key);
-
-					ReplicationState.Position = Handle->X();
-					ReplicationState.Quaternion = Handle->R();
-					ReplicationState.LinVel = Handle->V();
-					ReplicationState.AngVel = Handle->W();
-					ReplicationState.Flags = Handle->ObjectState() == Chaos::EObjectStateType::Sleeping ? ERigidBodyFlags::Sleeping : 0;
+					ReplicationCacheHelper(BIHandle, RootComponent);
+				}
+				else if (Chaos::FPhysicsObjectHandle PhysicsObject = RootComponent->GetPhysicsObjectByName(NAME_None))
+				{
+					Chaos::FReadPhysicsObjectInterface_Internal Interface = Chaos::FPhysicsObjectInternalInterface::GetRead();
+					Chaos::FPBDRigidParticleHandle* POHandle = Interface.GetRigidParticle(PhysicsObject);
+
+					if (POHandle != nullptr)
+					{
+						ReplicationCacheHelper(POHandle, RootComponent);
+					}
 				}
 			}
 		}
