Networked Physics: Predictive Interpolation replication, stage 1.

Predictive Interpolation is a replication method that replicates behind the last received state from the server (interpolates), while still allowing for local physics manipulation (predictive).

The replication is done via linear velocity and angular velocity by blending in local physics with the target physics over an interpolation time.

The target state from the server is extrapolated along its velocitites by DeltaTime each simulation tick, which makes the replicated object never reach the target but the replicated object will be at/around the source state of the received target when a new target is recieved.

There are several CVars to adjust and balance the replication via.

diff --git a/Engine/Source/Runtime/Engine/Private/PhysicsEngine/PhysicsReplication.cpp b/Engine/Source/Runtime/Engine/Private/PhysicsEngine/PhysicsReplication.cpp
--- a/Engine/Source/Runtime/Engine/Private/PhysicsEngine/PhysicsReplication.cpp
+++ b/Engine/Source/Runtime/Engine/Private/PhysicsEngine/PhysicsReplication.cpp
@@ -83,9 +83,31 @@ namespace PhysicsReplicationCVars
 	int32 LogPhysicsReplicationHardSnaps = 0;
 	static FAutoConsoleVariableRef CVarLogPhysicsReplicationHardSnaps(TEXT("p.LogPhysicsReplicationHardSnaps"), LogPhysicsReplicationHardSnaps, TEXT(""));
 #endif
+
 	static int32 EnableDefaultReplication = 0;
 	static FAutoConsoleVariableRef CVarEnableDefaultReplication(TEXT("np2.EnableDefaultReplication"), EnableDefaultReplication, TEXT("Enable default replication in the networked physics prediction flow."));
 
+	namespace PredictiveInterpolationCVars
+	{
+		static float PosCorrectionTimeMultiplier = 1.0f;
+		static FAutoConsoleVariableRef CVarPosCorrectionTimeMultiplier(TEXT("np2.PredictiveInterpolation.PosCorrectionTimeMultiplier"), PosCorrectionTimeMultiplier, TEXT("Multiplier to adjust the time to correct positional offset over, which is based on the clients forward predicted time ahead of the server."));
+
+		static float InterpolationTimeMultiplier = 1.1f;
+		static FAutoConsoleVariableRef CVarInterpolationTimeMultiplier(TEXT("np2.PredictiveInterpolation.InterpolationTimeMultiplier"), InterpolationTimeMultiplier, TEXT("Multiplier to adjust the replication interpolation time which is based on the clients forward predicted time ahead of the server."));
+
+		static float ExtrapolationTimeMultiplier = 1.5f;
+		static FAutoConsoleVariableRef CVarExtrapolationTimeMultiplier(TEXT("np2.PredictiveInterpolation.ExtrapolationTimeMultiplier"), ExtrapolationTimeMultiplier, TEXT("Multiplier to adjust the time to extrapolate the target forward over, the time is based on current send-rate."));
+
+		static float MinExpectedDistanceCovered = 0.25f;
+		static FAutoConsoleVariableRef CVarMinExpectedDistanceCovered(TEXT("np2.PredictiveInterpolation.MinExpectedDistanceCovered"), MinExpectedDistanceCovered, TEXT("Value in percentage where 0.25 = 25%. How much of the expected distance based on replication velocity should the object have covered in a simulation tick to Not be considered stuck."));
+
+		static float MaxDistanceToSleepSqr = 8.f;
+		static FAutoConsoleVariableRef CVarMaxDistanceToSleepSqr(TEXT("np2.PredictiveInterpolation.MaxDistanceToSleepSqr"), MaxDistanceToSleepSqr, TEXT("Squared value. Max distance from the source target to allow the object to go to sleep."));
+		
+		static bool PostResimWaitForUpdate = true;
+		static FAutoConsoleVariableRef CVarPostResimWaitForUpdate(TEXT("np2.PredictiveInterpolation.PostResimWaitForUpdate"), PostResimWaitForUpdate, TEXT("After a resimulation, wait for replicated states that correspond to post-resim state before processing replication again."));
+	}
+
 }
 
 
@@ -617,7 +639,6 @@ void FPhysicsReplicationAsync::OnPreSimulate_Internal()
 {
 	if (const FPhysicsReplicationAsyncInput* AsyncInput = GetConsumerInput_Internal())
 	{
-
 		Chaos::FPBDRigidsSolver* RigidsSolver = static_cast<Chaos::FPBDRigidsSolver*>(GetSolver());
 		check(RigidsSolver);
 
@@ -625,6 +646,22 @@ void FPhysicsReplicationAsync::OnPreSimulate_Internal()
 		Chaos::FRewindData* RewindData = RigidsSolver->GetRewindData();
 		if (RewindData && RewindData->IsResim())
 		{
+			// TODO, Handle the transition from post-resim to interpolation better.
+			if (PhysicsReplicationCVars::PredictiveInterpolationCVars::PostResimWaitForUpdate && RewindData->IsFinalResim())
+			{
+				for (auto Itr = ObjectToTarget.CreateIterator(); Itr; ++Itr)
+				{
+					FReplicatedPhysicsTargetAsync& Target = Itr.Value();
+
+					// If final resim frame, mark interpolated targets as waiting for up to date data from the server.
+					if (Target.RepMode == EPhysicsReplicationMode::PredictiveInterpolation)
+					{
+						Target.bWaiting = true;
+						Target.ServerFrame = RigidsSolver->GetCurrentFrame() + Target.FrameOffset;
+					}
+				}
+			}
+
 			return;
 		}
 
@@ -687,15 +724,27 @@ void FPhysicsReplicationAsync::UpdateAsyncTarget(const FPhysicsRepAsyncInputData
 		Target = &ObjectToTarget.Add(Input.PhysicsObject);
 		Target->PrevPos = Input.TargetState.Position;
 		Target->PrevPosTarget = Input.TargetState.Position;
+		Target->PrevRotTarget = Input.TargetState.Quaternion;
+		Target->PrevLinVel = Input.TargetState.LinVel;
 	}
 
 	if (Input.ServerFrame > Target->ServerFrame)
 	{
 		Target->PhysicsObject = Input.PhysicsObject;
+		Target->PrevServerFrame = Target->ServerFrame;
 		Target->ServerFrame = Input.ServerFrame;
 		Target->TargetState = Input.TargetState;
 		Target->RepMode = Input.RepMode;
 		Target->FrameOffset = Input.FrameOffset;
+		Target->TickCount = 0;
+		Target->bWaiting = false;
+
+		if (Input.RepMode == EPhysicsReplicationMode::PredictiveInterpolation)
+		{
+			// Cache the position we received this target at, Predictive Interpolation will alter the target state but use this as the source position for reconciliation.
+			Target->PrevPosTarget = Input.TargetState.Position;
+			Target->PrevRotTarget = Input.TargetState.Quaternion;
+		}
 	}
 
 	/** Cache the latest ping time */
@@ -742,21 +791,17 @@ void FPhysicsReplicationAsync::ApplyTargetStatesAsync(const float DeltaSeconds,
 
 				switch (RepMode)
 				{
-				case EPhysicsReplicationMode::Default:
-				{
-					bRemoveItr = DefaultReplication(Handle, Target, DeltaSeconds);
-					break;
-				}
-				case EPhysicsReplicationMode::PredictiveInterpolation:
-				{
-					bRemoveItr = PredictiveInterpolation(Handle, Target, DeltaSeconds);
-					break;
-				}
-				case EPhysicsReplicationMode::Resimulation:
-				{
-					bRemoveItr = ResimulationReplication(Handle, Target, DeltaSeconds);
-					break;
-				}
+					case EPhysicsReplicationMode::Default:
+						bRemoveItr = DefaultReplication(Handle, Target, DeltaSeconds);
+						break;
+
+					case EPhysicsReplicationMode::PredictiveInterpolation:
+						bRemoveItr = PredictiveInterpolation(Handle, Target, DeltaSeconds);
+						break;
+
+					case EPhysicsReplicationMode::Resimulation:
+						bRemoveItr = ResimulationReplication(Handle, Target, DeltaSeconds);
+						break;
 				}
 			}
 		}
@@ -1044,13 +1089,7 @@ bool FPhysicsReplicationAsync::DefaultReplication(Chaos::FPBDRigidParticleHandle
 		}
 	}
 
-	const Chaos::EObjectStateType ObjectState = Handle->ObjectState();
-	if (ObjectState != Chaos::EObjectStateType::Dynamic)
-	{
-		RigidsSolver->GetEvolution()->SetParticleObjectState(Handle, Chaos::EObjectStateType::Dynamic);
-	}
-
-	if ((Target.TargetState.Flags & ERigidBodyFlags::Sleeping) != 0)
+	if (bShouldSleep)
 	{
 		// don't allow kinematic to sleeping transition
 		if (Handle->ObjectState() != Chaos::EObjectStateType::Kinematic)
@@ -1065,21 +1104,35 @@ bool FPhysicsReplicationAsync::DefaultReplication(Chaos::FPBDRigidParticleHandle
 	return bRestoredState;
 }
 
-/** Interpolating towards replicated states from the server while predicting local physics */
+/** Interpolating towards replicated states from the server while predicting local physics 
+* TODO, detailed description
+*/
 bool FPhysicsReplicationAsync::PredictiveInterpolation(Chaos::FPBDRigidParticleHandle* Handle, FReplicatedPhysicsTargetAsync& Target, const float DeltaSeconds)
 {
+	if (Target.bWaiting)
+	{
+		return false;
+	}
+
 	Chaos::FPBDRigidsSolver* RigidsSolver = static_cast<Chaos::FPBDRigidsSolver*>(GetSolver());
 	if (RigidsSolver == nullptr)
 	{
 		return true;
 	}
 
+	const float ErrorAccumulationSeconds = CharacterMovementCVars::ErrorAccumulationSeconds >= 0.0f ? CharacterMovementCVars::ErrorAccumulationSeconds : ErrorCorrectionDefault.ErrorAccumulationSeconds;
+	const float MaxRestoredStateErrorSqr = CharacterMovementCVars::MaxRestoredStateError >= 0.0f ? 
+		(CharacterMovementCVars::MaxRestoredStateError * CharacterMovementCVars::MaxRestoredStateError) :
+		(ErrorCorrectionDefault.MaxRestoredStateError * ErrorCorrectionDefault.MaxRestoredStateError);
+
+	const bool bShouldSleep = (Target.TargetState.Flags & ERigidBodyFlags::Sleeping) != 0;
 	const int32 LocalFrame = Target.ServerFrame - Target.FrameOffset;
-	const int32 NumPredictedFrames = RigidsSolver->GetCurrentFrame() - LocalFrame;
+	const int32 NumPredictedFrames = RigidsSolver->GetCurrentFrame() - LocalFrame - Target.TickCount;
 	const float PredictedTime = DeltaSeconds * NumPredictedFrames;
+	const float SendRate = (Target.ServerFrame - Target.PrevServerFrame) * DeltaSeconds;
 
-	const float PosCorrectionTime = 0.2f; // TODO.
-	const float InterpolationTime = PredictedTime; // TODO
+	const float PosCorrectionTime = PredictedTime * PhysicsReplicationCVars::PredictiveInterpolationCVars::PosCorrectionTimeMultiplier;
+	const float InterpolationTime = PredictedTime * PhysicsReplicationCVars::PredictiveInterpolationCVars::InterpolationTimeMultiplier;
 
 	// CurrentState
 	FRigidBodyState CurrentState;
@@ -1094,49 +1147,125 @@ bool FPhysicsReplicationAsync::PredictiveInterpolation(Chaos::FPBDRigidParticleH
 	const FVector TargetLinVel = Target.TargetState.LinVel;
 	const FVector TargetAngVel = Target.TargetState.AngVel;
 
-	// Get PosDiff
-	FVector PosDiff = TargetPos - CurrentState.Position;
 
-	// Convert PosDiff to PosDiffVelocity with a PosCorrectionTime value(0.2f)
-	FVector PosDiffVelocity = PosDiff / PosCorrectionTime;
+	/** --- Reconciliation ---
+	* Get the traveled direction and distance from previous frame and compare with replicated linear velocity.
+	* If the object isn't moving enough along the replicated velocity it's considered stuck and needs a hard reconciliation.
+	*/
+	const FVector PrevDiff = CurrentState.Position - Target.PrevPos;
+	const float	ExpectedDistance = (Target.PrevLinVel * DeltaSeconds).Size();
+	const float CoveredDistance = FVector::DotProduct(PrevDiff, Target.PrevLinVel.GetSafeNormal());
+	
+	// If the object is moving less than X% of the expected distance, accumulate error seconds
+	if (CoveredDistance / ExpectedDistance < PhysicsReplicationCVars::PredictiveInterpolationCVars::MinExpectedDistanceCovered)
+	{
+		Target.AccumulatedErrorSeconds += DeltaSeconds;
+	}
+	else
+	{
+		Target.AccumulatedErrorSeconds = FMath::Max(Target.AccumulatedErrorSeconds - DeltaSeconds, 0.0f);
+	}
 
-	// Get LinVelDiff by adding inverted CurrentState.LinVel to TargetLinVel
-	FVector LinVelDiff = -CurrentState.LinVel + TargetLinVel;
+	const bool bHardSnap = Target.AccumulatedErrorSeconds > ErrorAccumulationSeconds;
+	bool bClearTarget = bHardSnap;
+	if (bHardSnap)
+	{
+		// Too much error so just snap state here and be done with it
+		Target.AccumulatedErrorSeconds = 0.0f;
+		Handle->SetX(Target.PrevPosTarget);
+		Handle->SetR(Target.PrevRotTarget);
+		Handle->SetV(Target.TargetState.LinVel);
+		Handle->SetW(Target.TargetState.AngVel);
+	}
+	else // Velocity-based Replication
+	{
+		// --- Velocity Replication ---
+		// Get PosDiff
+		const FVector PosDiff = TargetPos - CurrentState.Position;
 
-	// Add PosDiffVelocity to LinVelDiff to get BlendedTargetVelocity
-	FVector BlendedTargetVelocity = LinVelDiff + PosDiffVelocity;
+		// Convert PosDiff to a velocity
+		const FVector PosDiffVelocity = PosDiff / PosCorrectionTime;
 
-	// Multiply BlendedTargetVelocity with(deltaTime / interpolationTime), clamp to 1 and add to CurrentState.LinVel to get BlendedTargetVelocityInterpolated
-	const float BlendStepAmount = FMath::Clamp(DeltaSeconds / InterpolationTime, 0.f, 1.f);
-	FVector RepLinVel = (BlendedTargetVelocity * BlendStepAmount) + CurrentState.LinVel;
+		// Get LinVelDiff by adding inverted CurrentState.LinVel to TargetLinVel
+		const FVector LinVelDiff = -CurrentState.LinVel + TargetLinVel;
 
-	// Apply velocity on particle
-	Handle->SetV(RepLinVel);
-	Handle->SetR(TargetRot);
+		// Add PosDiffVelocity to LinVelDiff to get BlendedTargetVelocity
+		const FVector BlendedTargetVelocity = LinVelDiff + PosDiffVelocity;
 
+		// Multiply BlendedTargetVelocity with(deltaTime / interpolationTime), clamp to 1 and add to CurrentState.LinVel to get BlendedTargetVelocityInterpolated
+		const float BlendStepAmount = FMath::Clamp(DeltaSeconds / InterpolationTime, 0.f, 1.f);
+		const FVector RepLinVel = CurrentState.LinVel + (BlendedTargetVelocity * BlendStepAmount);
 
-	// TODO, Angular Velocity
 
+		// --- Angular Velocity Replication ---
+		// Extrapolate current rotation along current angular velocity to see where we would end up
+		float CurAngVelSize;
+		FVector CurAngVelAxis;
+		CurrentState.AngVel.FVector::ToDirectionAndLength(CurAngVelAxis, CurAngVelSize);
+		CurAngVelSize = FMath::DegreesToRadians(CurAngVelSize);
+		const FQuat CurRotExtrapDelta = FQuat(CurAngVelAxis, CurAngVelSize * DeltaSeconds);
+		const FQuat CurRotExtrap = CurRotExtrapDelta * CurrentState.Quaternion;
 
-	// TODO, Reconciliation
+		// Slerp from the extrapolated current rotation towards the target rotation
+		// This takes current angular velocity into account
+		const FQuat TargetRotBlended = FQuat::Slerp(CurRotExtrap, TargetRot, BlendStepAmount);
 
+		// Get the rotational offset between the blended rotation target and the current rotation
+		const FQuat TargetRotDelta = TargetRotBlended * CurrentState.Quaternion.Inverse();
 
-	// Sleep
-	if ((Target.TargetState.Flags & ERigidBodyFlags::Sleeping) != 0)
+		// Convert the rotational delta to angular velocity with a magnitude that will make it arrive at the rotation after DeltaTime has passed
+		float WAngle;
+		FVector WAxis;
+		TargetRotDelta.ToAxisAndAngle(WAxis, WAngle);
+
+		const FVector RepAngVel = WAxis * (WAngle / DeltaSeconds);
+
+
+		// Apply velocity
+		Handle->SetV(RepLinVel);
+		Handle->SetW(RepAngVel);
+
+
+		// Cache data for reconciliation
+		Target.PrevPos = FVector(CurrentState.Position);
+		Target.PrevLinVel = FVector(RepLinVel);
+	}
+
+
+	if (bShouldSleep)
 	{
+		// --- Sleep ---
+		// Get the distance from the current position to the source position of our target state
+		const float SourceDistanceSqr = (Target.PrevPosTarget - CurrentState.Position).SizeSquared();
+		
 		// Don't allow kinematic to sleeping transition
-		if (!Handle->IsKinematic())
+		if (SourceDistanceSqr < PhysicsReplicationCVars::PredictiveInterpolationCVars::MaxDistanceToSleepSqr && !Handle->IsKinematic())
 		{
 			RigidsSolver->GetEvolution()->SetParticleObjectState(Handle, Chaos::EObjectStateType::Sleeping);
+			bClearTarget = true;
+		}
+	}
+	else
+	{
+		// --- Target Extrapolation ---
+		if ((Target.TickCount * DeltaSeconds) < SendRate * PhysicsReplicationCVars::PredictiveInterpolationCVars::ExtrapolationTimeMultiplier)
+		{
+			// Extrapolate target position
+			Target.TargetState.Position = Target.TargetState.Position + Target.TargetState.LinVel * DeltaSeconds;
+
+			// Extrapolate target rotation
+			float TargetAngVelSize;
+			FVector TargetAngVelAxis;
+			Target.TargetState.AngVel.FVector::ToDirectionAndLength(TargetAngVelAxis, TargetAngVelSize);
+			TargetAngVelSize = FMath::DegreesToRadians(TargetAngVelSize);
+			const FQuat TargetRotExtrapDelta = FQuat(TargetAngVelAxis, TargetAngVelSize * DeltaSeconds);
+			Target.TargetState.Quaternion = TargetRotExtrapDelta * Target.TargetState.Quaternion;
 		}
 	}
 
+	Target.TickCount++;
 
-	// Extrapolate target state
-	Target.TargetState.Position = Target.TargetState.Position + Target.TargetState.LinVel * DeltaSeconds;
-	// TODO, extrapolate rotation via angular velocity
-
-	return false;
+	return bClearTarget;;
 }
 
 /** Compare states and trigger resimulation if needed */
diff --git a/Engine/Source/Runtime/Engine/Public/PhysicsReplication.h b/Engine/Source/Runtime/Engine/Public/PhysicsReplication.h
--- a/Engine/Source/Runtime/Engine/Public/PhysicsReplication.h
+++ b/Engine/Source/Runtime/Engine/Public/PhysicsReplication.h
@@ -102,21 +102,30 @@ struct FReplicatedPhysicsTargetAsync
 	/** Physics sync error accumulation */
 	float AccumulatedErrorSeconds;
 
-	/** Correction values from previous update */
-	FVector PrevPosTarget;
-	FVector PrevPos;
+	/** The amount of simulation ticks this target has been used for */
+	int32 TickCount;
 
 	/** ServerFrame this target was replicated on (must be converted to local frame prior to client-side use) */
 	int32 ServerFrame;
 
-	/** Index of physics object on component */
-	Chaos::FPhysicsObject* PhysicsObject;
-
 	/** The frame offset between local client and server */
 	int32 FrameOffset;
 
+	/** Index of physics object on component */
+	Chaos::FPhysicsObject* PhysicsObject;
+
 	/** The replication mode this PhysicsObject should use */
 	EPhysicsReplicationMode RepMode;
+
+	/** Correction values from previous update */
+	FVector PrevPosTarget;
+	FQuat PrevRotTarget;
+	FVector PrevPos;
+	FVector PrevLinVel;
+	int32 PrevServerFrame;
+
+	/** If this target is waiting for up-to-date data? */
+	bool bWaiting;
 };
 
 class FPhysicsReplicationAsync : public Chaos::TSimCallbackObject<
