Manually conform sparse class data against the new archetype when reparenting a BP

#rnx

diff --git a/Engine/Source/Editor/Kismet/Private/BlueprintEditor.cpp b/Engine/Source/Editor/Kismet/Private/BlueprintEditor.cpp
--- a/Engine/Source/Editor/Kismet/Private/BlueprintEditor.cpp
+++ b/Engine/Source/Editor/Kismet/Private/BlueprintEditor.cpp
@@ -3482,6 +3482,13 @@ void FBlueprintEditor::ReparentBlueprint_NewParentChosen(UClass* ChosenClass)
 			FBlueprintEditorUtils::RefreshAllNodes(GetBlueprintObj());
 			FBlueprintEditorUtils::MarkBlueprintAsModified(BlueprintObj);
 
+			// Changing the parent may change the sparse data used, so mark any current 
+			// sparse data as requiring a conform post-compile against the new archetype
+			if (UBlueprintGeneratedClass* Class = Cast<UBlueprintGeneratedClass>(BlueprintObj->GeneratedClass))
+			{
+				Class->PrepareToConformSparseClassData(ChosenClass->GetSparseClassDataStruct());
+			}
+
 			Compile();
 
 			// Ensure that the Blueprint is up-to-date (valid SCS etc.) after compiling (new parent class)
diff --git a/Engine/Source/Runtime/Engine/Classes/Animation/AnimBlueprintGeneratedClass.h b/Engine/Source/Runtime/Engine/Classes/Animation/AnimBlueprintGeneratedClass.h
--- a/Engine/Source/Runtime/Engine/Classes/Animation/AnimBlueprintGeneratedClass.h
+++ b/Engine/Source/Runtime/Engine/Classes/Animation/AnimBlueprintGeneratedClass.h
@@ -631,7 +631,7 @@ public:
 	static FName GetMutablesStructName();
 	
 #if WITH_EDITOR
-	virtual void PrepareToConformSparseClassData() override;
+	virtual void PrepareToConformSparseClassData(UScriptStruct* SparseClassDataArchetypeStruct) override;
 	virtual void ConformSparseClassData(UObject* Object) override;
 #endif
 
diff --git a/Engine/Source/Runtime/Engine/Classes/Engine/BlueprintGeneratedClass.h b/Engine/Source/Runtime/Engine/Classes/Engine/BlueprintGeneratedClass.h
--- a/Engine/Source/Runtime/Engine/Classes/Engine/BlueprintGeneratedClass.h
+++ b/Engine/Source/Runtime/Engine/Classes/Engine/BlueprintGeneratedClass.h
@@ -764,7 +764,7 @@ public:
 	 * Called during serialization to allow the class to stash any sparse class data that needs to 
 	 * be conformed against the archetype once the CDO is available.
 	 */
-	virtual void PrepareToConformSparseClassData();
+	virtual void PrepareToConformSparseClassData(UScriptStruct* SparseClassDataArchetypeStruct);
 
 	/**
 	 * Called to conform any pending sparse class data stashed by PrepareToConformSparseClassData.
diff --git a/Engine/Source/Runtime/Engine/Private/Animation/AnimBlueprintGeneratedClass.cpp b/Engine/Source/Runtime/Engine/Private/Animation/AnimBlueprintGeneratedClass.cpp
--- a/Engine/Source/Runtime/Engine/Private/Animation/AnimBlueprintGeneratedClass.cpp
+++ b/Engine/Source/Runtime/Engine/Private/Animation/AnimBlueprintGeneratedClass.cpp
@@ -617,7 +617,7 @@ uint8* UAnimBlueprintGeneratedClass::GetPersistentUberGraphFrame(UObject* Obj, U
 }
 
 #if WITH_EDITOR
-void UAnimBlueprintGeneratedClass::PrepareToConformSparseClassData()
+void UAnimBlueprintGeneratedClass::PrepareToConformSparseClassData(UScriptStruct* SparseClassDataArchetypeStruct)
 {
 	// ABPGC sparse data is generated on compile, so there is nothing to conform
 }
diff --git a/Engine/Source/Runtime/Engine/Private/BlueprintGeneratedClass.cpp b/Engine/Source/Runtime/Engine/Private/BlueprintGeneratedClass.cpp
--- a/Engine/Source/Runtime/Engine/Private/BlueprintGeneratedClass.cpp
+++ b/Engine/Source/Runtime/Engine/Private/BlueprintGeneratedClass.cpp
@@ -600,7 +600,7 @@ void UBlueprintGeneratedClass::SerializeDefaultObject(UObject* Object, FStructur
 #if WITH_EDITOR
 		if (!GetPackage()->HasAnyPackageFlags(PKG_Cooked))
 		{
-			PrepareToConformSparseClassData();
+			PrepareToConformSparseClassData(GetSparseClassDataArchetypeStruct());
 			if(ClassDefaultObject)
 			{
 				// conform immediately, so that child types can correct delta serialize their
@@ -649,14 +649,12 @@ void UBlueprintGeneratedClass::PostLoadDefaultObject(UObject* Object)
 }
 
 #if WITH_EDITOR
-void UBlueprintGeneratedClass::PrepareToConformSparseClassData()
+void UBlueprintGeneratedClass::PrepareToConformSparseClassData(UScriptStruct* SparseClassDataArchetypeStruct)
 {
 	checkf(SparseClassDataPendingConformStruct.IsExplicitlyNull() && SparseClassDataPendingConform == nullptr, TEXT("PrepareToConformSparseClassData was called while data was already pending conform!"));
 
 	if (SparseClassDataStruct)
 	{
-		UScriptStruct* SparseClassDataArchetypeStruct = GetSparseClassDataArchetypeStruct();
-
 		if (SparseClassDataStruct != SparseClassDataArchetypeStruct)
 		{
 			if (SparseClassDataArchetypeStruct)
