LevelInstance: Fix Undo of Deleted Level Instance

#rnx

diff --git a/Engine/Source/Runtime/Engine/Private/LevelInstance/LevelInstanceLevelStreaming.cpp b/Engine/Source/Runtime/Engine/Private/LevelInstance/LevelInstanceLevelStreaming.cpp
--- a/Engine/Source/Runtime/Engine/Private/LevelInstance/LevelInstanceLevelStreaming.cpp
+++ b/Engine/Source/Runtime/Engine/Private/LevelInstance/LevelInstanceLevelStreaming.cpp
@@ -33,6 +33,22 @@ FAutoConsoleVariableRef CVarDisableLevelInstanceEditorPartialLoading(
 	GDisableLevelInstanceEditorPartialLoading,
 	TEXT("Allow disabling partial loading of level instances in the editor."),
 	ECVF_Default);
+
+namespace FLevelInstanceLevelStreamingUtils
+{
+	static void MarkObjectsInPackageAsTransientAndNonTransactional(UPackage* InPackage)
+	{
+		InPackage->ClearFlags(RF_Transactional);
+		InPackage->SetFlags(RF_Transient);
+		ForEachObjectWithPackage(InPackage, [](UObject* Obj)
+			{
+				Obj->SetFlags(RF_Transient);
+				Obj->ClearFlags(RF_Transactional);
+				return true;
+			}, true);
+	}
+}
+
 #endif
 
 ULevelStreamingLevelInstance::ULevelStreamingLevelInstance(const FObjectInitializer& ObjectInitializer)
@@ -158,9 +174,6 @@ void ULevelStreamingLevelInstance::ResetLevelInstanceLoaders()
 
 void ULevelStreamingLevelInstance::PrepareLevelInstanceLoadedActor(AActor& InActor, ILevelInstanceInterface* InLevelInstance, bool bResetLoaders)
 {
-	InActor.ClearFlags(RF_Transactional);
-	InActor.SetFlags(RF_Transient);
-
 	if (InActor.IsPackageExternal())
 	{
 		if (bResetLoaders)
@@ -168,7 +181,7 @@ void ULevelStreamingLevelInstance::PrepareLevelInstanceLoadedActor(AActor& InAct
 			ResetLoaders(InActor.GetExternalPackage());
 		}
 
-		InActor.GetPackage()->SetFlags(RF_Transient);
+		FLevelInstanceLevelStreamingUtils::MarkObjectsInPackageAsTransientAndNonTransactional(InActor.GetExternalPackage());
 	}
 
 	InActor.PushSelectionToProxies();
@@ -265,18 +278,7 @@ ULevelStreamingLevelInstance* ULevelStreamingLevelInstance::LoadInstance(ILevelI
 				Level->OnLoadedActorAddedToLevelEvent.AddUObject(LevelStreaming, &ULevelStreamingLevelInstance::OnLoadedActorAddedToLevel);
 				Level->OnLoadedActorRemovedFromLevelEvent.AddUObject(LevelStreaming, &ULevelStreamingLevelInstance::OnLoadedActorRemovedFromLevel);
 
-				UWorld* OuterWorld = Level->GetTypedOuter<UWorld>();
-				OuterWorld->ClearFlags(RF_Transactional);
-				OuterWorld->SetFlags(RF_Transient);
-			
-				OuterWorld->GetPackage()->ClearFlags(RF_Transactional);
-				OuterWorld->GetPackage()->SetFlags(RF_Transient);
-
-				ForEachObjectWithOuter(OuterWorld, [&](UObject* Obj)
-				{
-					Obj->ClearFlags(RF_Transactional);
-					Obj->SetFlags(RF_Transient);
-				}, true);
+				FLevelInstanceLevelStreamingUtils::MarkObjectsInPackageAsTransientAndNonTransactional(Level->GetPackage());
 
 				for (AActor* LevelActor : Level->Actors)
 				{
@@ -290,7 +292,7 @@ ULevelStreamingLevelInstance* ULevelStreamingLevelInstance::LoadInstance(ILevelI
 				{
 					if (ActorFolder->IsPackageExternal())
 					{
-						ActorFolder->GetPackage()->SetFlags(RF_Transient);
+						FLevelInstanceLevelStreamingUtils::MarkObjectsInPackageAsTransientAndNonTransactional(ActorFolder->GetPackage());
 					}
 					return true;
 				});
@@ -336,6 +338,7 @@ void ULevelStreamingLevelInstance::UnloadInstance(ULevelStreamingLevelInstance*
 			if(Obj->HasAnyFlags(RF_Transactional))
 			{
 				bResetTrans = true;
+				UE_LOG(LogLevelInstance, Warning, TEXT("Found RF_Transactional object '%s' while unloading Level Instance."), *Obj->GetPathName());
 				return false;
 			}
 			return true;
diff --git a/Engine/Source/Runtime/Engine/Public/LevelInstance/LevelInstanceLevelStreaming.h b/Engine/Source/Runtime/Engine/Public/LevelInstance/LevelInstanceLevelStreaming.h
--- a/Engine/Source/Runtime/Engine/Public/LevelInstance/LevelInstanceLevelStreaming.h
+++ b/Engine/Source/Runtime/Engine/Public/LevelInstance/LevelInstanceLevelStreaming.h
@@ -10,6 +10,7 @@
 
 class ILevelInstanceInterface;
 class ALevelInstanceEditorInstanceActor;
+class UPackage;
 
 UCLASS(Transient)
 class ENGINE_API ULevelStreamingLevelInstance : public ULevelStreamingDynamic
