Backout header change submitted by mistake

#rnx

diff --git a/Engine/Source/Runtime/Engine/Public/LevelInstance/LevelInstanceLevelStreaming.h b/Engine/Source/Runtime/Engine/Public/LevelInstance/LevelInstanceLevelStreaming.h
--- a/Engine/Source/Runtime/Engine/Public/LevelInstance/LevelInstanceLevelStreaming.h
+++ b/Engine/Source/Runtime/Engine/Public/LevelInstance/LevelInstanceLevelStreaming.h
@@ -10,7 +10,6 @@
 
 class ILevelInstanceInterface;
 class ALevelInstanceEditorInstanceActor;
-class UPackage;
 
 UCLASS(Transient)
 class ENGINE_API ULevelStreamingLevelInstance : public ULevelStreamingDynamic
