Fixed bug in IWorldPartitionActorLoaderInterface::ILoaderAdapter::AddReferenceToActor not properly recursing through actor references

diff --git a/Engine/Source/Runtime/Engine/Private/WorldPartition/WorldPartitionActorLoaderInterface.cpp b/Engine/Source/Runtime/Engine/Private/WorldPartition/WorldPartitionActorLoaderInterface.cpp
--- a/Engine/Source/Runtime/Engine/Private/WorldPartition/WorldPartitionActorLoaderInterface.cpp
+++ b/Engine/Source/Runtime/Engine/Private/WorldPartition/WorldPartitionActorLoaderInterface.cpp
@@ -288,15 +288,15 @@ void IWorldPartitionActorLoaderInterface::ILoaderAdapter::PostLoadedStateChanged
 
 void IWorldPartitionActorLoaderInterface::ILoaderAdapter::AddReferenceToActor(FWorldPartitionHandle& ActorHandle)
 {
-	TFunction<void(const FWorldPartitionHandle&, FReferenceMap&)> AddReferences = [this, &ActorHandle, &AddReferences](const FWorldPartitionHandle& Handle, FReferenceMap& ReferenceMap)
+	TFunction<void(const FWorldPartitionHandle&, FReferenceMap&)> AddReferences = [this, &AddReferences](const FWorldPartitionHandle& Handle, FReferenceMap& ReferenceMap)
 	{
 		if (!ReferenceMap.Contains(Handle->GetGuid()))
 		{
 			ReferenceMap.Emplace(Handle->GetGuid(), Handle.ToReference());
 			
-			for (const FGuid& ReferencedActorGuid : ActorHandle->GetReferences())
+			for (const FGuid& ReferencedActorGuid : Handle->GetReferences())
 			{
-				FWorldPartitionHandle ReferenceActorHandle(ActorHandle->GetContainer(), ReferencedActorGuid);
+				FWorldPartitionHandle ReferenceActorHandle(Handle->GetContainer(), ReferencedActorGuid);
 
 				if (ReferenceActorHandle.IsValid())
 				{
