Add missing ENGINE_API decorator.

#pf trivial

diff --git a/Engine/Source/Runtime/Engine/Public/PSOPrecache.h b/Engine/Source/Runtime/Engine/Public/PSOPrecache.h
--- a/Engine/Source/Runtime/Engine/Public/PSOPrecache.h
+++ b/Engine/Source/Runtime/Engine/Public/PSOPrecache.h
@@ -383,7 +383,7 @@ namespace PSOCollectorStats
 	 * Track a PSO precache stat's total count, and optionally also track the counts at
 	 * the mesh pass type and vertex factory type granularity.
 	 */
-	class FPrecacheUsageData
+	class ENGINE_API FPrecacheUsageData
 	{
 	public:
 		FPrecacheUsageData(FName StatFName = FName())
@@ -423,7 +423,7 @@ namespace PSOCollectorStats
 	/**
 	 * Collect stats for different cache hit states
 	 */
-	struct FPrecacheStats
+	struct ENGINE_API FPrecacheStats
 	{
 		FPrecacheStats(FName UntrackedStatFName = FName(), FName MissStatFName = FName(), FName HitStatFName = FName(), FName UsedStatFName = FName(), FName TooLateStatFName = FName())
 			: UsageData(UsedStatFName)
