Substrate - fixed UseDiffusion affecting the look of simple volume slabs.

diff --git a/Engine/Shaders/Private/DebugProbes.usf b/Engine/Shaders/Private/DebugProbes.usf
--- a/Engine/Shaders/Private/DebugProbes.usf
+++ b/Engine/Shaders/Private/DebugProbes.usf
@@ -140,8 +140,9 @@ void WriteOutMaterial(
 		/*float GlintValue						*/0.0f,
 		/*float GlintUV							*/0.0f.xx,
 		/*float SpecularProfile					*/0.0f,
-		/*float Thickness						*/0.0f,
+		/*float Thickness						*/1.0f,
 		/*bool bIsThinSurface					*/false,
+		/*bool bIsAtTheBottomOfTopology			*/false,
 		/*uint SharedLocalBasisIndex			*/SharedLocalBasisIndex,
 		/*inout uint SharedLocalBasisTypes		*/StrataPixelHeader.SharedLocalBases.Types);
 
diff --git a/Engine/Shaders/Private/Strata/Strata.ush b/Engine/Shaders/Private/Strata/Strata.ush
--- a/Engine/Shaders/Private/Strata/Strata.ush
+++ b/Engine/Shaders/Private/Strata/Strata.ush
@@ -1331,6 +1331,7 @@ FStrataData GetStrataSlabBSDF(
 	float SpecularProfileId,
 	float ThicknessCm,
 	bool bIsThinSurface,
+	bool bIsAtTheBottomOfTopology,	// Make sure we do not affect MFP when wrap mode is selected. This is also an optimisation, disabling lots of code when we know about topology.
 	uint SharedLocalBasisIndex, inout uint SharedLocalBasisTypes)
 {
 	// Curvature-based roughness for geometric AA
@@ -1344,10 +1345,10 @@ FStrataData GetStrataSlabBSDF(
 	SSSMFPScale = saturate(SSSMFPScale);
 	const bool bSSSProfilePlugged = SSSProfileID > 0.0f;
 	const bool bMFPPlugged = any(SSSMFP > 0);
-	bool bIsThin = bMFPPlugged && bIsThinSurface; // Enable thin only if MFP is plugged, as otherwise there will be no back-tramission
-	bool bHasSSSWrap = bMFPPlugged && !UseSSSDiffusion;
-	bool bHasSSSProfile = bSSSProfilePlugged && SSSMFPScale > 0.f;
-	bool bHasSSS = bMFPPlugged || bHasSSSProfile || bHasSSSWrap;
+	bool bIsThin		= bIsAtTheBottomOfTopology && (bMFPPlugged && bIsThinSurface); // Enable thin only if MFP is plugged, as otherwise there will be no back-tramission
+	bool bHasSSSWrap	= bIsAtTheBottomOfTopology && (bMFPPlugged && !UseSSSDiffusion);
+	bool bHasSSSProfile	= bIsAtTheBottomOfTopology && (bSSSProfilePlugged && SSSMFPScale > 0.f);
+	bool bHasSSS		= bIsAtTheBottomOfTopology && (bMFPPlugged || bHasSSSProfile || bHasSSSWrap);
 
 	FHaziness Haziness = InitialiseHaziness();
 
@@ -3556,6 +3557,7 @@ void FStrataPixelHeader::StrataConvertFromDBuffer(
 		0.f, 				// SpecularProfileId
 		DefaultThicknessCm,
 		false,				// bIsThinSurface,
+		false,				// bIsAtTheBottomOfTopology, this will disable many unneeded features.
 		SharedLocalBasisIndex,
 		SharedLocalBasisTypes);
 #if STRATA_INLINE_SHADING
diff --git a/Engine/Shaders/Private/Strata/StrataLegacyConversion.ush b/Engine/Shaders/Private/Strata/StrataLegacyConversion.ush
--- a/Engine/Shaders/Private/Strata/StrataLegacyConversion.ush
+++ b/Engine/Shaders/Private/Strata/StrataLegacyConversion.ush
@@ -124,6 +124,7 @@ FStrataData CreateLegacySlab(
 		0.0f,												// Specular profile
 		In.Thickness,
 		In.bIsThinSurface,
+		true,												// bIsAtTheBottomOfTopology, legacy always has a single BSDF being at the bottom and top at teh same time.
 		In.SharedLocalBasisIndex, SharedLocalBasisTypes);
 }
 
diff --git a/Engine/Source/Runtime/Engine/Private/Materials/HLSLMaterialTranslator.cpp b/Engine/Source/Runtime/Engine/Private/Materials/HLSLMaterialTranslator.cpp
--- a/Engine/Source/Runtime/Engine/Private/Materials/HLSLMaterialTranslator.cpp
+++ b/Engine/Source/Runtime/Engine/Private/Materials/HLSLMaterialTranslator.cpp
@@ -12110,6 +12110,7 @@ int32 FHLSLMaterialTranslator::StrataSlabBSDF(
 	int32 Thickness,
 	int32 GlintValue, int32 GlintUV,
 	int32 SpecularProfileId,
+	bool bIsAtTheBottomOfTopology,
 	int32 Normal, int32 Tangent, const FString& SharedLocalBasisIndexMacro,
 	FStrataOperator* PromoteToOperator)
 {
@@ -12127,7 +12128,7 @@ int32 FHLSLMaterialTranslator::StrataSlabBSDF(
 		}
 
 		return AddCodeChunk(
-			MCT_Strata, TEXT("Parameters.%s.PromoteParameterBlendedBSDFToOperator(GetStrataSlabBSDF(Parameters.StrataPixelFootprint, %s, %s, %s, %s, %s, %s, %s, %s, %s, %s, %s, %s, %s, %s, %s, %s, %s, %s, %s, %s, %s, %s, %s, Parameters.%s.Types) /* Normal = %s ; Tangent = %s ; Thickness = %s */, %u, %u, %u, %u)"),
+			MCT_Strata, TEXT("Parameters.%s.PromoteParameterBlendedBSDFToOperator(GetStrataSlabBSDF(Parameters.StrataPixelFootprint, %s, %s, %s, %s, %s, %s, %s, %s, %s, %s, %s, %s, %s, %s, %s, %s, %s, %s, %s, %s, %s, %s, %s, %s, Parameters.%s.Types) /* Normal = %s ; Tangent = %s ; Thickness = %s */, %u, %u, %u, %u)"),
 			*GetParametersStrataTreeName(CurrentStrataCompilationContext),
 			*StrataGetCastParameterCode(DiffuseAlbedo,			MCT_Float3),
 			*StrataGetCastParameterCode(F0,						MCT_Float3),
@@ -12151,6 +12152,7 @@ int32 FHLSLMaterialTranslator::StrataSlabBSDF(
 			*StrataGetCastParameterCode(SpecularProfileId,		MCT_Float),
 			*StrataGetCastParameterCode(Thickness,				MCT_Float),
 			bIsThinSurface ? TEXT("true") : TEXT("false"),
+			bIsAtTheBottomOfTopology ? TEXT("true") : TEXT("false"),
 			*SharedLocalBasisIndexMacro,
 			*GetParametersSharedLocalBasesName(CurrentStrataCompilationContext),
 			*NormalCode,
@@ -12164,7 +12166,7 @@ int32 FHLSLMaterialTranslator::StrataSlabBSDF(
 	}
 	
 	return AddCodeChunk(
-		MCT_Strata, TEXT("GetStrataSlabBSDF(Parameters.StrataPixelFootprint, %s, %s, %s, %s, %s, %s, %s, %s, %s, %s, %s, %s, %s, %s, %s, %s, %s, %s, %s, %s, %s, %s, %s, Parameters.%s.Types) /* Normal = %s ; Tangent = %s ; Thickness = %s */"),
+		MCT_Strata, TEXT("GetStrataSlabBSDF(Parameters.StrataPixelFootprint, %s, %s, %s, %s, %s, %s, %s, %s, %s, %s, %s, %s, %s, %s, %s, %s, %s, %s, %s, %s, %s, %s, %s, %s, Parameters.%s.Types) /* Normal = %s ; Tangent = %s ; Thickness = %s */"),
 		*StrataGetCastParameterCode(DiffuseAlbedo,			MCT_Float3),
 		*StrataGetCastParameterCode(F0,						MCT_Float3),
 		*StrataGetCastParameterCode(F90,					MCT_Float3),
@@ -12187,6 +12189,7 @@ int32 FHLSLMaterialTranslator::StrataSlabBSDF(
 		*StrataGetCastParameterCode(SpecularProfileId,		MCT_Float),
 		*StrataGetCastParameterCode(Thickness, MCT_Float),
 		bIsThinSurface ? TEXT("true") : TEXT("false"),
+		bIsAtTheBottomOfTopology ? TEXT("true") : TEXT("false"),
 		*SharedLocalBasisIndexMacro,
 		*GetParametersSharedLocalBasesName(CurrentStrataCompilationContext),
 		*NormalCode,
diff --git a/Engine/Source/Runtime/Engine/Private/Materials/HLSLMaterialTranslator.h b/Engine/Source/Runtime/Engine/Private/Materials/HLSLMaterialTranslator.h
--- a/Engine/Source/Runtime/Engine/Private/Materials/HLSLMaterialTranslator.h
+++ b/Engine/Source/Runtime/Engine/Private/Materials/HLSLMaterialTranslator.h
@@ -1171,6 +1171,7 @@ protected:
 		int32 Thickness,
 		int32 GlintValue, int32 GlintUV,
 		int32 SpecularProfileId,
+		bool bIsAtTheBottomOfTopology,
 		int32 Normal, int32 Tangent, const FString& SharedLocalBasisIndexMacro,
 		FStrataOperator* PromoteToOperator) override;
 	virtual int32 StrataConversionFromLegacy(
diff --git a/Engine/Source/Runtime/Engine/Private/Materials/MaterialExpressions.cpp b/Engine/Source/Runtime/Engine/Private/Materials/MaterialExpressions.cpp
--- a/Engine/Source/Runtime/Engine/Private/Materials/MaterialExpressions.cpp
+++ b/Engine/Source/Runtime/Engine/Private/Materials/MaterialExpressions.cpp
@@ -23300,6 +23300,7 @@ int32 UMaterialExpressionStrataSlabBSDF::Compile(class FMaterialCompiler* Compil
 		CompileWithDefaultFloat1(Compiler, GlintValue, 0.0f),
 		CompileWithDefaultFloat2(Compiler, GlintUV, 0.0f, 0.0f),
 		SpecularProfileCodeChunk != INDEX_NONE ? SpecularProfileCodeChunk : Compiler->Constant(0.0f),
+		StrataOperator.bIsBottom > 0 ? true : false,
 		NormalCodeChunk,
 		TangentCodeChunk,
 		Compiler->GetStrataSharedLocalBasisIndexMacro(NewRegisteredSharedLocalBasis),
@@ -23649,6 +23650,7 @@ int32 UMaterialExpressionStrataSimpleClearCoatBSDF::Compile(class FMaterialCompi
 		Compiler->Constant(0.0f),								// GlintValue
 		Compiler->Constant2(0.0f, 0.0f),						// GlintUV
 		Compiler->Constant(0.0f),								// SpecularProfile
+		false,													// bIsAtTheBottomOfTopology, always false for SimpleClearCoat
 		NormalCodeChunk,
 		NullTangentCodeChunk,
 		Compiler->GetStrataSharedLocalBasisIndexMacro(NewRegisteredSharedLocalBasis),
diff --git a/Engine/Source/Runtime/Engine/Public/MaterialCompiler.h b/Engine/Source/Runtime/Engine/Public/MaterialCompiler.h
--- a/Engine/Source/Runtime/Engine/Public/MaterialCompiler.h
+++ b/Engine/Source/Runtime/Engine/Public/MaterialCompiler.h
@@ -559,6 +559,7 @@ public:
 		int32 Thickness,
 		int32 GlintValue, int32 GlintUV,
 		int32 SpecularProfileId,
+		bool bIsAtTheBottomOfTopology,
 		int32 Normal, int32 Tangent, const FString& SharedLocalBasisIndexMacro,
 		FStrataOperator* PromoteToOperator) = 0;
 	virtual int32 StrataConversionFromLegacy(
@@ -1180,6 +1181,7 @@ public:
 		int32 Thickness,
 		int32 GlintValue, int32 GlintUV,
 		int32 SpecularProfileId,
+		bool bIsAtTheBottomOfTopology,
 		int32 Normal, int32 Tangent, const FString& SharedLocalBasisIndexMacro,
 		FStrataOperator* PromoteToOperator) override
 	{
@@ -1193,6 +1195,7 @@ public:
 			Thickness,
 			GlintValue, GlintUV,
 			SpecularProfileId,
+			bIsAtTheBottomOfTopology,
 			Normal, Tangent, SharedLocalBasisIndexMacro,
 			PromoteToOperator);
 	}
