LevelInstance: Edit now shows a transient filter that can be changed to allow editing of the source Level, this filter isn't modifying the current LevelInstance filter. Filter is restored when exiting edit.

#rnx

diff --git a/Engine/Source/Editor/LevelInstanceEditor/Private/LevelInstanceEditorModule.cpp b/Engine/Source/Editor/LevelInstanceEditor/Private/LevelInstanceEditorModule.cpp
--- a/Engine/Source/Editor/LevelInstanceEditor/Private/LevelInstanceEditorModule.cpp
+++ b/Engine/Source/Editor/LevelInstanceEditor/Private/LevelInstanceEditorModule.cpp
@@ -688,7 +688,8 @@ void FLevelInstanceEditorModule::StartupModule()
 	FPropertyEditorModule& PropertyModule = FModuleManager::LoadModuleChecked<FPropertyEditorModule>("PropertyEditor");
 	PropertyModule.RegisterCustomClassLayout("LevelInstance", FOnGetDetailCustomizationInstance::CreateStatic(&FLevelInstanceActorDetails::MakeInstance));
 	PropertyModule.RegisterCustomClassLayout("LevelInstancePivot", FOnGetDetailCustomizationInstance::CreateStatic(&FLevelInstancePivotDetails::MakeInstance));		
-	PropertyModule.RegisterCustomPropertyTypeLayout("WorldPartitionActorFilter", FOnGetPropertyTypeCustomizationInstance::CreateStatic(&FLevelInstanceFilterPropertyTypeCustomization::MakeInstance), MakeShared<FLevelInstancePropertyTypeIdentifier>());
+	PropertyModule.RegisterCustomPropertyTypeLayout("WorldPartitionActorFilter", FOnGetPropertyTypeCustomizationInstance::CreateStatic(&FLevelInstanceFilterPropertyTypeCustomization::MakeInstance, false), MakeShared<FLevelInstancePropertyTypeIdentifier>(false));
+	PropertyModule.RegisterCustomPropertyTypeLayout("WorldPartitionActorFilter", FOnGetPropertyTypeCustomizationInstance::CreateStatic(&FLevelInstanceFilterPropertyTypeCustomization::MakeInstance, true), MakeShared<FLevelInstancePropertyTypeIdentifier>(true));
 	PropertyModule.NotifyCustomizationModuleChanged();
 
 	// GEditor needs to be set before this module is loaded
diff --git a/Engine/Source/Editor/LevelInstanceEditor/Private/LevelInstanceFilterPropertyTypeCustomization.cpp b/Engine/Source/Editor/LevelInstanceEditor/Private/LevelInstanceFilterPropertyTypeCustomization.cpp
--- a/Engine/Source/Editor/LevelInstanceEditor/Private/LevelInstanceFilterPropertyTypeCustomization.cpp
+++ b/Engine/Source/Editor/LevelInstanceEditor/Private/LevelInstanceFilterPropertyTypeCustomization.cpp
@@ -33,6 +33,12 @@ TSharedPtr<FWorldPartitionActorFilterMode::FFilter> FLevelInstanceFilterProperty
 		{
 			if (ILevelInstanceInterface* LevelInstanceInterface = Cast<ILevelInstanceInterface>(OuterActor))
 			{
+				// Make sure this Customization supports current filter type
+				if (LevelInstanceInterface->IsEditing() != bIsEditFilter)
+				{
+					return nullptr;
+				}
+
 				LevelInstances.Add(LevelInstanceInterface, OuterObjectIndex);
 				ActorLabel = OuterActor->GetActorLabel();
 				
diff --git a/Engine/Source/Editor/LevelInstanceEditor/Private/LevelInstanceFilterPropertyTypeCustomization.h b/Engine/Source/Editor/LevelInstanceEditor/Private/LevelInstanceFilterPropertyTypeCustomization.h
--- a/Engine/Source/Editor/LevelInstanceEditor/Private/LevelInstanceFilterPropertyTypeCustomization.h
+++ b/Engine/Source/Editor/LevelInstanceEditor/Private/LevelInstanceFilterPropertyTypeCustomization.h
@@ -9,28 +9,42 @@
 #include "Templates/SharedPointer.h"
 #include "WorldPartition/Filter/WorldPartitionActorFilterPropertyTypeCustomization.h"
 #include "WorldPartition/Filter/WorldPartitionActorFilterMode.h"
+#include "LevelInstance/LevelInstanceComponent.h"
 
 class ILevelInstanceInterface;
 
 class FLevelInstancePropertyTypeIdentifier : public IPropertyTypeIdentifier
 {
+public:
+	FLevelInstancePropertyTypeIdentifier(bool bInIsEditFilter)
+		: bIsEditFilter(bInIsEditFilter)
+	{}
+
 	virtual bool IsPropertyTypeCustomized(const IPropertyHandle& PropertyHandle) const override
 	{
-		return PropertyHandle.HasMetaData(TEXT("LevelInstanceFilter"));
+		return bIsEditFilter ? PropertyHandle.HasMetaData(TEXT("LevelInstanceEditFilter")) : PropertyHandle.HasMetaData(TEXT("LevelInstanceFilter"));
 	}
+
+private:
+	bool bIsEditFilter;
 };
 
 // Registered (FLevelInstanceEditorModule::StartupModule) Property Customization for properties of type FWorldPartitionActorFilter for Level Instances
 struct FLevelInstanceFilterPropertyTypeCustomization : public FWorldPartitionActorFilterPropertyTypeCustomization
 {
 public:
-	static TSharedRef<IPropertyTypeCustomization> MakeInstance()
+	static TSharedRef<IPropertyTypeCustomization> MakeInstance(bool bInIsEditFilter)
 	{
-		return MakeShareable(new FLevelInstanceFilterPropertyTypeCustomization);
+		return MakeShareable(new FLevelInstanceFilterPropertyTypeCustomization(bInIsEditFilter));
 	}
 private:
+	FLevelInstanceFilterPropertyTypeCustomization(bool bInIsEditFilter)
+		: bIsEditFilter(bInIsEditFilter)
+	{}
+
 	virtual TSharedPtr<FWorldPartitionActorFilterMode::FFilter> CreateModeFilter(TArray<UObject*> OuterObjects) override;
 	virtual void ApplyFilter(TSharedRef<IPropertyHandle> PropertyHandle, const FWorldPartitionActorFilterMode& Mode) override;
 
 	TMap<ILevelInstanceInterface*, int32> LevelInstances;
+	bool bIsEditFilter;
 };
diff --git a/Engine/Source/Editor/WorldPartitionEditor/Private/WorldPartition/Filter/WorldPartitionActorFilterMode.cpp b/Engine/Source/Editor/WorldPartitionEditor/Private/WorldPartition/Filter/WorldPartitionActorFilterMode.cpp
--- a/Engine/Source/Editor/WorldPartitionEditor/Private/WorldPartition/Filter/WorldPartitionActorFilterMode.cpp
+++ b/Engine/Source/Editor/WorldPartitionEditor/Private/WorldPartition/Filter/WorldPartitionActorFilterMode.cpp
@@ -88,7 +88,10 @@ void FWorldPartitionActorFilterMode::FFilter::Override(const FWorldPartitionActo
 
 		for (auto& [ActorGuid, WorldPartitionActorFilter] : InOverride->GetChildFilters())
 		{
-			OverrideInternal(InValue->GetChildFilters().FindChecked(ActorGuid), WorldPartitionActorFilter);
+			if (FWorldPartitionActorFilter*const* ChildFilter = InValue->GetChildFilters().Find(ActorGuid))
+			{
+				OverrideInternal(*ChildFilter, WorldPartitionActorFilter);
+			}
 		}
 	};
 	OverrideInternal(LevelFilter.Get(), Other);
diff --git a/Engine/Source/Runtime/Engine/Private/LevelInstance/LevelInstanceComponent.cpp b/Engine/Source/Runtime/Engine/Private/LevelInstance/LevelInstanceComponent.cpp
--- a/Engine/Source/Runtime/Engine/Private/LevelInstance/LevelInstanceComponent.cpp
+++ b/Engine/Source/Runtime/Engine/Private/LevelInstance/LevelInstanceComponent.cpp
@@ -19,6 +19,9 @@ static const TCHAR* GLevelSpriteAssetName = TEXT("/Engine/EditorResources/LevelI
 
 ULevelInstanceComponent::ULevelInstanceComponent(const FObjectInitializer& ObjectInitializer)
 	: Super(ObjectInitializer)
+#if WITH_EDITORONLY_DATA
+	, bWasDirtyBeforeEditFilterChange(false)
+#endif
 {
 #if WITH_EDITOR
 	bWantsOnUpdateTransform = true;
@@ -69,15 +72,27 @@ void ULevelInstanceComponent::OnRegister()
 
 void ULevelInstanceComponent::SetFilter(const FWorldPartitionActorFilter& InFilter)
 {
-	if (Filter != InFilter)
+	if (GetFilter() != InFilter)
 	{
-		Modify();
-		Filter = InFilter;
+		Modify(!IsEditFilter());
+		SetActiveFilter(InFilter);
 		OnFilterChanged();
 		FWorldPartitionActorFilter::GetOnWorldPartitionActorFilterChanged().Broadcast();
 	}
 }
 
+void ULevelInstanceComponent::SetActiveFilter(const FWorldPartitionActorFilter& InFilter)
+{
+	if (IsEditFilter())
+	{
+		EditFilter = InFilter;
+	}
+	else
+	{
+		Filter = InFilter;
+	}
+}
+
 void ULevelInstanceComponent::OnFilterChanged()
 {
 	if (ILevelInstanceInterface* LevelInstance = Cast<ILevelInstanceInterface>(GetOwner()))
@@ -93,7 +108,7 @@ bool ULevelInstanceComponent::ShouldShowSpriteComponent() const
 
 void ULevelInstanceComponent::PreEditUndo()
 {
-	UndoRedoCachedFilter = Filter;
+	UndoRedoCachedFilter = GetFilter();
 }
 
 void ULevelInstanceComponent::PostEditUndo()
@@ -103,7 +118,7 @@ void ULevelInstanceComponent::PostEditUndo()
 	UpdateComponentToWorld();
 	UpdateEditorInstanceActor();
 
-	if (Filter != UndoRedoCachedFilter)
+	if (GetFilter() != UndoRedoCachedFilter)
 	{
 		OnFilterChanged();
 		FWorldPartitionActorFilter::RequestFilterRefresh(false);
@@ -112,16 +127,34 @@ void ULevelInstanceComponent::PostEditUndo()
 	UndoRedoCachedFilter = FWorldPartitionActorFilter();
 }
 
+void ULevelInstanceComponent::PreEditChange(FProperty* PropertyAboutToChange)
+{
+	if (PropertyAboutToChange && PropertyAboutToChange->GetFName() == GET_MEMBER_NAME_CHECKED(ULevelInstanceComponent, EditFilter))
+	{
+		bWasDirtyBeforeEditFilterChange = GetPackage()->IsDirty();
+	}
+
+	Super::PreEditChange(PropertyAboutToChange);
+}
+
 void ULevelInstanceComponent::PostEditChangeProperty(FPropertyChangedEvent& PropertyChangedEvent)
 {
 	Super::PostEditChangeProperty(PropertyChangedEvent);
 
 	UpdateEditorInstanceActor();
 
-	if (PropertyChangedEvent.Property && PropertyChangedEvent.Property->GetFName() == GET_MEMBER_NAME_CHECKED(ULevelInstanceComponent, Filter))
+	const bool FilterProperty = PropertyChangedEvent.Property && PropertyChangedEvent.Property->GetFName() == GET_MEMBER_NAME_CHECKED(ULevelInstanceComponent, Filter);
+	const bool EditFilterProperty = PropertyChangedEvent.Property && PropertyChangedEvent.Property->GetFName() == GET_MEMBER_NAME_CHECKED(ULevelInstanceComponent, EditFilter);
+
+	if (FilterProperty || EditFilterProperty)
 	{
 		OnFilterChanged();
 	}
+
+	if (EditFilterProperty && !bWasDirtyBeforeEditFilterChange)
+	{
+		GetPackage()->SetDirtyFlag(false);
+	}
 }
 
 void ULevelInstanceComponent::OnUpdateTransform(EUpdateTransformFlags UpdateTransformFlags, ETeleportType Teleport)
@@ -178,9 +211,26 @@ void ULevelInstanceComponent::OnCommit()
 	}
 }
 
+void ULevelInstanceComponent::UpdateEditFilter()
+{
+	// @todo_ow: for child level instances setup the proper filter based on parent hierarchy filter
+	EditFilter = Filter;
+}
+
+bool ULevelInstanceComponent::IsEditFilter() const
+{
+	if (ILevelInstanceInterface* LevelInstance = Cast<ILevelInstanceInterface>(GetOwner()))
+	{
+		return LevelInstance->IsEditing();
+	}
+
+	return false;
+}
+
 const TMap<FActorContainerID, TSet<FGuid>>& ULevelInstanceComponent::GetFilteredActorsPerContainer() const
 {
-	if (CachedFilter != Filter)
+	const FWorldPartitionActorFilter& CurrentFilter = GetFilter();
+	if (CachedFilter != CurrentFilter)
 	{
 		CachedFilteredActorsPerContainer.Reset();
 	}
@@ -197,11 +247,11 @@ const TMap<FActorContainerID, TSet<FGuid>>& ULevelInstanceComponent::GetFiltered
 		UWorldPartitionSubsystem* WorldPartitionSubsystem = GetWorld()->GetSubsystem<UWorldPartitionSubsystem>();
 		check(WorldPartitionSubsystem);
 				
-		FilteredActors = WorldPartitionSubsystem->GetFilteredActorsPerContainer(LevelInstance->GetLevelInstanceID().GetContainerID(), LevelInstance->GetWorldAssetPackage(), Filter, LevelInstance->GetLoadingFilterTypes());
+		FilteredActors = WorldPartitionSubsystem->GetFilteredActorsPerContainer(LevelInstance->GetLevelInstanceID().GetContainerID(), LevelInstance->GetWorldAssetPackage(), CurrentFilter, LevelInstance->GetLoadingFilterTypes());
 	}
 	
 	CachedFilteredActorsPerContainer = MoveTemp(FilteredActors);
-	CachedFilter = Filter;
+	CachedFilter = CurrentFilter;
 	return CachedFilteredActorsPerContainer.GetValue();
 }
 
diff --git a/Engine/Source/Runtime/Engine/Private/LevelInstance/LevelInstanceInterface.cpp b/Engine/Source/Runtime/Engine/Private/LevelInstance/LevelInstanceInterface.cpp
--- a/Engine/Source/Runtime/Engine/Private/LevelInstance/LevelInstanceInterface.cpp
+++ b/Engine/Source/Runtime/Engine/Private/LevelInstance/LevelInstanceInterface.cpp
@@ -157,14 +157,13 @@ void ILevelInstanceInterface::UpdateLevelInstanceFromWorldAsset()
 
 bool ILevelInstanceInterface::IsEditing() const
 {
-	if (HasValidLevelInstanceID())
+	// No need to check LevelInstanceID here since Editing Level Instance is referenced directly:
+	// Level Instance ID can become invalid while Actor is being edited through AActor property change events causing component unregisters
+	if (ULevelInstanceSubsystem* LevelInstanceSubsystem = GetLevelInstanceSubsystem())
 	{
-		if (ULevelInstanceSubsystem* LevelInstanceSubsystem = GetLevelInstanceSubsystem())
-		{
-			return LevelInstanceSubsystem->IsEditingLevelInstance(this);
-		}
+		return LevelInstanceSubsystem->IsEditingLevelInstance(this);
 	}
-
+	
 	return false;
 }
 
diff --git a/Engine/Source/Runtime/Engine/Private/LevelInstance/LevelInstanceSubsystem.cpp b/Engine/Source/Runtime/Engine/Private/LevelInstance/LevelInstanceSubsystem.cpp
--- a/Engine/Source/Runtime/Engine/Private/LevelInstance/LevelInstanceSubsystem.cpp
+++ b/Engine/Source/Runtime/Engine/Private/LevelInstance/LevelInstanceSubsystem.cpp
@@ -19,6 +19,7 @@
 #include "LevelInstance/LevelInstanceEditorInstanceActor.h"
 #include "LevelInstance/LevelInstanceEditorObject.h"
 #include "LevelInstance/LevelInstanceEditorPivotActor.h"
+#include "LevelInstance/LevelInstanceComponent.h"
 #include "WorldPartition/LevelInstance/LevelInstanceActorDesc.h"
 #include "WorldPartition/WorldPartition.h"
 #include "WorldPartition/WorldPartitionSubsystem.h"
@@ -398,7 +399,7 @@ void ULevelInstanceSubsystem::RegisterLoadedLevelStreamingLevelInstanceEditor(UL
 	{
 		check(!LevelInstanceEdit.IsValid());
 		ILevelInstanceInterface* LevelInstance = LevelStreaming->GetLevelInstance();
-		LevelInstanceEdit = MakeUnique<FLevelInstanceEdit>(LevelStreaming, LevelInstance->GetLevelInstanceID());
+		LevelInstanceEdit = MakeUnique<FLevelInstanceEdit>(LevelStreaming, LevelInstance);
 
 		if (ILevelInstanceEditorModule* EditorModule = FModuleManager::GetModulePtr<ILevelInstanceEditorModule>("LevelInstanceEditor"))
 		{
@@ -1272,7 +1273,8 @@ ILevelInstanceInterface* ULevelInstanceSubsystem::CreateLevelInstanceFrom(const
 	};
 
 	FStackLevelInstanceEdit StackLevelInstanceEdit;
-	StackLevelInstanceEdit.LevelInstanceEdit = MakeUnique<FLevelInstanceEdit>(LevelStreaming, NewLevelInstance->GetLevelInstanceID());
+	LevelStreaming->LevelInstanceID = NewLevelInstanceID;
+	StackLevelInstanceEdit.LevelInstanceEdit = MakeUnique<FLevelInstanceEdit>(LevelStreaming, NewLevelInstance);
 	// Force mark it as changed
 	StackLevelInstanceEdit.LevelInstanceEdit->MarkCommittedChanges();
 
@@ -1687,10 +1689,12 @@ bool ULevelInstanceSubsystem::ShouldIgnoreDirtyPackage(UPackage* DirtyPackage, c
 	return bIgnore;
 }
 
-ULevelInstanceSubsystem::FLevelInstanceEdit::FLevelInstanceEdit(ULevelStreamingLevelInstanceEditor* InLevelStreaming, FLevelInstanceID InLevelInstanceID)
-	: LevelStreaming(InLevelStreaming)
+ULevelInstanceSubsystem::FLevelInstanceEdit::FLevelInstanceEdit(ULevelStreamingLevelInstanceEditor* InLevelStreaming, ILevelInstanceInterface* InLevelInstance)
+	: LevelStreaming(InLevelStreaming), LevelInstanceActor(CastChecked<AActor>(InLevelInstance))
 {
-	LevelStreaming->LevelInstanceID = InLevelInstanceID;
+	check(LevelStreaming->LevelInstanceID == InLevelInstance->GetLevelInstanceID());
+	// Update Edit Filter before actors are added to world
+	InLevelInstance->GetLevelInstanceComponent()->UpdateEditFilter();
 	EditorObject = NewObject<ULevelInstanceEditorObject>(GetTransientPackage(), NAME_None, RF_Transactional);
 	EditorObject->EnterEdit(GetEditWorld());
 }
@@ -1711,10 +1715,16 @@ UWorld* ULevelInstanceSubsystem::FLevelInstanceEdit::GetEditWorld() const
 	return nullptr;
 }
 
+ILevelInstanceInterface* ULevelInstanceSubsystem::FLevelInstanceEdit::GetLevelInstance() const
+{
+	return Cast<ILevelInstanceInterface>(LevelInstanceActor);
+}
+
 void ULevelInstanceSubsystem::FLevelInstanceEdit::AddReferencedObjects(FReferenceCollector& Collector)
 {
 	Collector.AddReferencedObject(EditorObject);
 	Collector.AddReferencedObject(LevelStreaming);
+	Collector.AddReferencedObject(LevelInstanceActor);
 }
 
 bool ULevelInstanceSubsystem::FLevelInstanceEdit::CanDiscard(FText* OutReason) const
@@ -1751,14 +1761,9 @@ void ULevelInstanceSubsystem::FLevelInstanceEdit::GetPackagesToSave(TArray<UPack
 	}
 }
 
-FLevelInstanceID ULevelInstanceSubsystem::FLevelInstanceEdit::GetLevelInstanceID() const
-{
-	return LevelStreaming ? LevelStreaming->GetLevelInstanceID() : FLevelInstanceID();
-}
-
 const ULevelInstanceSubsystem::FLevelInstanceEdit* ULevelInstanceSubsystem::GetLevelInstanceEdit(const ILevelInstanceInterface* LevelInstance) const
 {
-	if (LevelInstance && LevelInstanceEdit && LevelInstanceEdit->GetLevelInstanceID() == LevelInstance->GetLevelInstanceID())
+	if (LevelInstance && LevelInstanceEdit && LevelInstanceEdit->GetLevelInstance() == LevelInstance)
 	{
 		return LevelInstanceEdit.Get();
 	}
@@ -1788,7 +1793,7 @@ ILevelInstanceInterface* ULevelInstanceSubsystem::GetEditingLevelInstance() cons
 {
 	if (LevelInstanceEdit)
 	{
-		return GetLevelInstance(LevelInstanceEdit->GetLevelInstanceID());
+		return LevelInstanceEdit->GetLevelInstance();
 	}
 
 	return nullptr;
@@ -1819,7 +1824,7 @@ bool ULevelInstanceSubsystem::CanEditLevelInstance(const ILevelInstanceInterface
 	
 	if (LevelInstanceEdit)
 	{
-		if (LevelInstanceEdit->GetLevelInstanceID() == LevelInstance->GetLevelInstanceID())
+		if (LevelInstanceEdit->GetLevelInstance() == LevelInstance)
 		{
 			if (OutReason)
 			{
@@ -1920,7 +1925,8 @@ bool ULevelInstanceSubsystem::EditLevelInstanceInternal(ILevelInstanceInterface*
 		});
 	}
 
-	GEditor->SelectNone(false, true);
+	// Make sure selection is refreshed (Edit can have impact on details view)
+	GEditor->SelectNone(true, true);
 	
 	// Avoid calling OnEditChild twice  on ancestors when EditLevelInstance calls itself
 	if (!bRecursive)
@@ -1976,7 +1982,7 @@ bool ULevelInstanceSubsystem::EditLevelInstanceInternal(ILevelInstanceInterface*
 	}
 
 	check(LevelInstanceEdit.IsValid());
-	check(LevelInstanceEdit->GetLevelInstanceID() == LevelInstance->GetLevelInstanceID());
+	check(LevelInstanceEdit->GetLevelInstance() == LevelInstance);
 	check(LevelInstanceEdit->LevelStreaming == LevelStreaming);
 		
 	// Try and select something meaningful
@@ -2075,7 +2081,7 @@ bool ULevelInstanceSubsystem::CommitLevelInstance(ILevelInstanceInterface* Level
 bool ULevelInstanceSubsystem::CommitLevelInstanceInternal(TUniquePtr<FLevelInstanceEdit>& InLevelInstanceEdit, bool bDiscardEdits, bool bDiscardOnFailure, TSet<FName>* DirtyPackages)
 {
 	TGuardValue<bool> CommitScope(bIsCommittingLevelInstance, true);
-	ILevelInstanceInterface* LevelInstance = GetLevelInstance(InLevelInstanceEdit->GetLevelInstanceID());
+	ILevelInstanceInterface* LevelInstance = InLevelInstanceEdit->GetLevelInstance();
 	check(InLevelInstanceEdit);
 	UWorld* EditingWorld = InLevelInstanceEdit->GetEditWorld();
 	check(EditingWorld);
@@ -2136,7 +2142,8 @@ bool ULevelInstanceSubsystem::CommitLevelInstanceInternal(TUniquePtr<FLevelInsta
 	FScopedSlowTask SlowTask(0, LOCTEXT("EndEditLevelInstance", "Unloading Level..."), !GetWorld()->IsGameWorld());
 	SlowTask.MakeDialog();
 
-	GEditor->SelectNone(false, true);
+	// Make sure selection is refreshed (Commit can have impact on details view)
+	GEditor->SelectNone(true, true);
 
 	const FString EditPackage = LevelInstance->GetWorldAssetPackage();
 
diff --git a/Engine/Source/Runtime/Engine/Private/PackedLevelActor/PackedLevelActor.cpp b/Engine/Source/Runtime/Engine/Private/PackedLevelActor/PackedLevelActor.cpp
--- a/Engine/Source/Runtime/Engine/Private/PackedLevelActor/PackedLevelActor.cpp
+++ b/Engine/Source/Runtime/Engine/Private/PackedLevelActor/PackedLevelActor.cpp
@@ -91,11 +91,21 @@ void APackedLevelActor::RerunConstructionScripts()
 
 EWorldPartitionActorFilterType APackedLevelActor::GetDetailsFilterTypes() const
 {
+	if (IsEditing())
+	{
+		return Super::GetDetailsFilterTypes();
+	}
+
 	return IsRootBlueprintTemplate() ? EWorldPartitionActorFilterType::All : EWorldPartitionActorFilterType::None;
 }
 
 EWorldPartitionActorFilterType APackedLevelActor::GetLoadingFilterTypes() const
 {
+	if (IsEditing())
+	{
+		return Super::GetLoadingFilterTypes();
+	}
+
 	return ShouldLoadForPacking() ? EWorldPartitionActorFilterType::All : EWorldPartitionActorFilterType::None;
 }
 
diff --git a/Engine/Source/Runtime/Engine/Public/LevelInstance/LevelInstanceComponent.h b/Engine/Source/Runtime/Engine/Public/LevelInstance/LevelInstanceComponent.h
--- a/Engine/Source/Runtime/Engine/Public/LevelInstance/LevelInstanceComponent.h
+++ b/Engine/Source/Runtime/Engine/Public/LevelInstance/LevelInstanceComponent.h
@@ -33,6 +33,7 @@ public:
 	virtual void OnRegister() override;
 	virtual void PreEditUndo() override;
 	virtual void PostEditUndo() override;
+	virtual void PreEditChange(FProperty* PropertyAboutToChange) override;
 	virtual void PostEditChangeProperty(FPropertyChangedEvent& PropertyChangedEvent) override;
 	virtual void OnUpdateTransform(EUpdateTransformFlags UpdateTransformFlags, ETeleportType Teleport) override;
 
@@ -40,12 +41,15 @@ public:
 	void OnEdit();
 	void OnCommit();
 
-	const FWorldPartitionActorFilter& GetFilter() const { return Filter; }
+	const FWorldPartitionActorFilter& GetFilter() const { return IsEditFilter() ? EditFilter : Filter; }
 	void SetFilter(const FWorldPartitionActorFilter& InFilter);
 	const TMap<FActorContainerID, TSet<FGuid>>& GetFilteredActorsPerContainer() const;
+	void UpdateEditFilter();
 private:
 	bool ShouldShowSpriteComponent() const;
 	void OnFilterChanged();
+	void SetActiveFilter(const FWorldPartitionActorFilter& InFilter);
+	bool IsEditFilter() const;
 
 	TWeakObjectPtr<AActor> CachedEditorInstanceActorPtr;
 #endif
@@ -54,9 +58,15 @@ private:
 	UPROPERTY(EditAnywhere, Category = Filter, meta=(LevelInstanceFilter))
 	FWorldPartitionActorFilter Filter;
 
+	UPROPERTY(EditAnywhere, Transient, Category = Filter, meta=(LevelInstanceEditFilter))
+	FWorldPartitionActorFilter EditFilter;
+
 	FWorldPartitionActorFilter UndoRedoCachedFilter;
 
 	mutable FWorldPartitionActorFilter CachedFilter;
 	mutable TOptional<TMap<FActorContainerID, TSet<FGuid>>> CachedFilteredActorsPerContainer;
+
+	// Used to cancel the package getting dirty when editing the EditFilter which is transient
+	bool bWasDirtyBeforeEditFilterChange;
 #endif
 };
\ No newline at end of file
diff --git a/Engine/Source/Runtime/Engine/Public/LevelInstance/LevelInstanceSubsystem.h b/Engine/Source/Runtime/Engine/Public/LevelInstance/LevelInstanceSubsystem.h
--- a/Engine/Source/Runtime/Engine/Public/LevelInstance/LevelInstanceSubsystem.h
+++ b/Engine/Source/Runtime/Engine/Public/LevelInstance/LevelInstanceSubsystem.h
@@ -164,12 +164,13 @@ private:
 	public:
 		TObjectPtr<ULevelStreamingLevelInstanceEditor> LevelStreaming;
 		TObjectPtr<ULevelInstanceEditorObject> EditorObject;
+		TObjectPtr<AActor> LevelInstanceActor;
 
-		FLevelInstanceEdit(ULevelStreamingLevelInstanceEditor* InLevelStreaming, FLevelInstanceID InLevelInstanceID);
+		FLevelInstanceEdit(ULevelStreamingLevelInstanceEditor* InLevelStreaming, ILevelInstanceInterface* InLevelInstance);
 		virtual ~FLevelInstanceEdit();
 
 		UWorld* GetEditWorld() const;
-		FLevelInstanceID GetLevelInstanceID() const;
+		ILevelInstanceInterface* GetLevelInstance() const;
 
 		void AddReferencedObjects(FReferenceCollector& Collector);
 
