Iris
- Fixed spammy warning

diff --git a/Engine/Source/Runtime/Engine/Private/Net/Iris/ReplicationSystem/ActorReplicationBridge.cpp b/Engine/Source/Runtime/Engine/Private/Net/Iris/ReplicationSystem/ActorReplicationBridge.cpp
--- a/Engine/Source/Runtime/Engine/Private/Net/Iris/ReplicationSystem/ActorReplicationBridge.cpp
+++ b/Engine/Source/Runtime/Engine/Private/Net/Iris/ReplicationSystem/ActorReplicationBridge.cpp
@@ -261,7 +261,7 @@ UE::Net::FNetRefHandle UActorReplicationBridge::BeginReplication(AActor* Actor,
 	const uint32 ClampedPollPeriod = FMath::Clamp<uint32>(UnclampedPollPeriod, 1U, static_cast<uint32>(std::numeric_limits<uint8>::max()) + 1U);
 	CreateNetRefHandleParams.PollFramePeriod = (ClampedPollPeriod - 1U) & 255U;
 #if !UE_BUILD_SHIPPING
-	if ((ClampedPollPeriod >= UnclampedPollPeriod))
+	if ((ClampedPollPeriod < UnclampedPollPeriod))
 	{
 		ensureAlwaysMsgf(!UE::Net::Private::bIrisEnableLowNetUpdateFrequencyEnsure, TEXT("Very low NetUpdateFrequency %f for Actor %s. Suggest setting it to %f or higher."), Actor->NetUpdateFrequency, ToCStr(Actor->GetName()), GetMinSupportedNetUpdateFrequency());
 		UE_CLOG(!UE::Net::Private::bIrisEnableLowNetUpdateFrequencyEnsure, LogIrisBridge, Warning, TEXT("Very low NetUpdateFrequency %f for Actor %s. Suggest setting it to %f or higher."), Actor->NetUpdateFrequency, ToCStr(Actor->GetName()), GetMinSupportedNetUpdateFrequency());
