Allow FShaderCompilingManager::SubmitJobs() to be invoked from worker threads by conditinally extending the CompileQueueSection lock

diff --git a/Engine/Source/Runtime/Engine/Private/ShaderCompiler/ShaderCompiler.cpp b/Engine/Source/Runtime/Engine/Private/ShaderCompiler/ShaderCompiler.cpp
--- a/Engine/Source/Runtime/Engine/Private/ShaderCompiler/ShaderCompiler.cpp
+++ b/Engine/Source/Runtime/Engine/Private/ShaderCompiler/ShaderCompiler.cpp
@@ -4492,9 +4492,17 @@ void FShaderCompilingManager::SubmitJobs(TArray<FShaderCommonCompileJobPtr>& New
 			PendingShaderMap->NumPendingJobs.Increment();
 			Job->PendingShaderMap = PendingShaderMap;
 		}
-	}
 
-	AllJobs.SubmitJobs(NewJobs);
+		// in the case of submitting jobs from worker threads we need to be sure that the lock extends to
+		// include AllJobs.SubmitJobs().  This will increase contention for the lock, but this will let us
+		// prototype getting shader translation and preprocessing being done on worker threads.
+		if (IsInGameThread())
+		{
+			Lock.Unlock();
+		}
+
+		AllJobs.SubmitJobs(NewJobs);
+	}
 
 	UpdateNumRemainingAssets();
 }
