Update TActorDescContainerCollection to support TObjectPtr instead of raw pointers
WorldPartition's ActorDescCollection is now using TObjectPtr<UActorDescContainer> instead of raw pointer to be compliant with the upcoming ARO API update

diff --git a/Engine/Source/Runtime/Engine/Private/WorldPartition/WorldPartitionStreamingGeneration.cpp b/Engine/Source/Runtime/Engine/Private/WorldPartition/WorldPartitionStreamingGeneration.cpp
--- a/Engine/Source/Runtime/Engine/Private/WorldPartition/WorldPartitionStreamingGeneration.cpp
+++ b/Engine/Source/Runtime/Engine/Private/WorldPartition/WorldPartitionStreamingGeneration.cpp
@@ -1390,8 +1390,8 @@ void UWorldPartition::CheckForErrors(IStreamingGenerationErrorHandler* ErrorHand
 	
 }
 
-FStreamingGenerationActorDescCollection::FStreamingGenerationActorDescCollection(std::initializer_list<const UActorDescContainer*> ActorDescContainerArray)
-	: TActorDescContainerCollection<const UActorDescContainer*>(ActorDescContainerArray)
+FStreamingGenerationActorDescCollection::FStreamingGenerationActorDescCollection(std::initializer_list<TObjectPtr<const UActorDescContainer>> ActorDescContainerArray)
+	: TActorDescContainerCollection<TObjectPtr<const UActorDescContainer>>(ActorDescContainerArray)
 {
 	SortCollection();
 }
diff --git a/Engine/Source/Runtime/Engine/Public/WorldPartition/ActorDescContainerCollection.h b/Engine/Source/Runtime/Engine/Public/WorldPartition/ActorDescContainerCollection.h
--- a/Engine/Source/Runtime/Engine/Public/WorldPartition/ActorDescContainerCollection.h
+++ b/Engine/Source/Runtime/Engine/Public/WorldPartition/ActorDescContainerCollection.h
@@ -72,14 +72,14 @@ public:
 	void ForEachActorDescContainer(TFunctionRef<void(ActorDescContPtrType)> Func) const;
 
 protected:
-	void RegisterDelegates(ActorDescContPtrType Container);
-	void UnregisterDelegates(ActorDescContPtrType Container);
-
 	virtual void OnCollectionChanged() {};
 
 	TArray<ActorDescContPtrType> ActorDescContainerCollection;
 
 private:
+	void RegisterDelegates(ActorDescContPtrType Container);
+	void UnregisterDelegates(ActorDescContPtrType Container);
+
 	void OnActorDescAdded(FWorldPartitionActorDesc* ActorDesc);
 	void OnActorDescRemoved(FWorldPartitionActorDesc* ActorDesc);
 
@@ -281,17 +281,15 @@ void TActorDescContainerCollection<ActorDescContPtrType>::Empty()
 template<class ActorDescContPtrType>
 void TActorDescContainerCollection<ActorDescContPtrType>::RegisterDelegates(ActorDescContPtrType Container)
 {
-	using ActorDescContNonConstPtrType = typename std::remove_const<typename TRemovePointer<ActorDescContPtrType>::Type>::type*;
-	const_cast<ActorDescContNonConstPtrType>(Container)->OnActorDescAddedEvent.AddRaw(this, &TActorDescContainerCollection<ActorDescContPtrType>::OnActorDescAdded);
-	const_cast<ActorDescContNonConstPtrType>(Container)->OnActorDescRemovedEvent.AddRaw(this, &TActorDescContainerCollection::OnActorDescRemoved);
+	ConstCast(Container)->OnActorDescAddedEvent.AddRaw(this, &TActorDescContainerCollection<ActorDescContPtrType>::OnActorDescAdded);
+	ConstCast(Container)->OnActorDescRemovedEvent.AddRaw(this, &TActorDescContainerCollection::OnActorDescRemoved);
 }
 
 template<class ActorDescContPtrType>
 void TActorDescContainerCollection<ActorDescContPtrType>::UnregisterDelegates(ActorDescContPtrType Container)
 {
-	using ActorDescContNonConstPtrType = typename std::remove_const<typename TRemovePointer<ActorDescContPtrType>::Type>::type*;
-	const_cast<ActorDescContNonConstPtrType>(Container)->OnActorDescAddedEvent.RemoveAll(this);
-	const_cast<ActorDescContNonConstPtrType>(Container)->OnActorDescRemovedEvent.RemoveAll(this);
+	ConstCast(Container)->OnActorDescAddedEvent.RemoveAll(this);
+	ConstCast(Container)->OnActorDescRemovedEvent.RemoveAll(this);
 }
 
 template<class ActorDescContPtrType>
@@ -489,4 +487,4 @@ void TActorDescContainerCollection<ActorDescContPtrType>::OnActorDescRemoved(FWo
 
 #endif 
 
-using FActorDescContainerCollection = TActorDescContainerCollection<UActorDescContainer*>;
\ No newline at end of file
+using FActorDescContainerCollection = TActorDescContainerCollection<TObjectPtr<UActorDescContainer>>;
\ No newline at end of file
diff --git a/Engine/Source/Runtime/Engine/Public/WorldPartition/WorldPartitionStreamingGeneration.h b/Engine/Source/Runtime/Engine/Public/WorldPartition/WorldPartitionStreamingGeneration.h
--- a/Engine/Source/Runtime/Engine/Public/WorldPartition/WorldPartitionStreamingGeneration.h
+++ b/Engine/Source/Runtime/Engine/Public/WorldPartition/WorldPartitionStreamingGeneration.h
@@ -101,11 +101,11 @@ protected:
 	TMultiMap<FName, const FWorldPartitionActorDescView*> ActorDescViewsByClass;
 };
 
-class ENGINE_API FStreamingGenerationActorDescCollection : public TActorDescContainerCollection<const UActorDescContainer*>
+class ENGINE_API FStreamingGenerationActorDescCollection : public TActorDescContainerCollection<TObjectPtr<const UActorDescContainer>>
 {
 public:
 	FStreamingGenerationActorDescCollection() = default;
-	FStreamingGenerationActorDescCollection(std::initializer_list<const UActorDescContainer*> ActorDescContainerArray);
+	FStreamingGenerationActorDescCollection(std::initializer_list<TObjectPtr<const UActorDescContainer>> ActorDescContainerArray);
 
 	UWorld* GetWorld() const;
 
