Fix missing transition

diff --git a/Engine/Source/Runtime/Engine/Private/GPUSkinCache.cpp b/Engine/Source/Runtime/Engine/Private/GPUSkinCache.cpp
--- a/Engine/Source/Runtime/Engine/Private/GPUSkinCache.cpp
+++ b/Engine/Source/Runtime/Engine/Private/GPUSkinCache.cpp
@@ -1475,6 +1475,7 @@ void FGPUSkinCache::DoDispatch(FRHICommandListImmediate& RHICmdList)
 				BuffersToTransitionToRead.Add(DispatchData.GetPositionRWBuffer());
 			}
 
+			BuffersToTransitionToRead.Add(DispatchData.GetTangentRWBuffer());
 			check(DispatchData.PreviousPositionBuffer != DispatchData.PositionBuffer);
 		}
 
