AActor types can now supports differents UDataLayerInstance types via IsDataLayerTypeSupported
Deprecate ActorTypeSupportsDataLayer

diff --git a/Engine/Source/Developer/FunctionalTesting/Classes/FunctionalTest.h b/Engine/Source/Developer/FunctionalTesting/Classes/FunctionalTest.h
--- a/Engine/Source/Developer/FunctionalTesting/Classes/FunctionalTest.h
+++ b/Engine/Source/Developer/FunctionalTesting/Classes/FunctionalTest.h
@@ -722,7 +722,7 @@ public:
 	virtual void EndPlay(const EEndPlayReason::Type EndPlayReason) override;
 #if WITH_EDITOR
 	virtual bool CanChangeIsSpatiallyLoadedFlag() const override { return false; }
-	virtual bool ActorTypeSupportsDataLayer() const override { return false; }
+	virtual bool IsDataLayerTypeSupported(TSubclassOf<UDataLayerInstance> DataLayerType) const override { return false; }
 #endif
 	// AActor interface end
 
diff --git a/Engine/Source/Editor/DataLayerEditor/Private/DataLayer/DataLayerEditorSubsystem.cpp b/Engine/Source/Editor/DataLayerEditor/Private/DataLayer/DataLayerEditorSubsystem.cpp
--- a/Engine/Source/Editor/DataLayerEditor/Private/DataLayer/DataLayerEditorSubsystem.cpp
+++ b/Engine/Source/Editor/DataLayerEditor/Private/DataLayer/DataLayerEditorSubsystem.cpp
@@ -430,7 +430,7 @@ void UDataLayerEditorSubsystem::PostUndoRedo()
 bool UDataLayerEditorSubsystem::IsActorValidForDataLayer(AActor* Actor)
 {
 	UWorld* World = Actor ? Actor->GetWorld() : nullptr;
-	return World && (World->WorldType == EWorldType::Editor) && World->IsPartitionedWorld() && Actor && Actor->SupportsDataLayer() && ((Actor->GetLevel() == Actor->GetWorld()->PersistentLevel) || Actor->GetLevel()->GetWorldDataLayers());
+	return World && (World->WorldType == EWorldType::Editor) && World->IsPartitionedWorld() && Actor && Actor->SupportsDataLayerType(UDataLayerInstance::StaticClass()) && ((Actor->GetLevel() == Actor->GetWorld()->PersistentLevel) || Actor->GetLevel()->GetWorldDataLayers());
 }
 
 void UDataLayerEditorSubsystem::OnWorldPartitionInitialized(UWorldPartition* InWorldPartition)
diff --git a/Engine/Source/Editor/UnrealEd/Private/Commandlets/WorldPartitionConvertCommandlet.cpp b/Engine/Source/Editor/UnrealEd/Private/Commandlets/WorldPartitionConvertCommandlet.cpp
--- a/Engine/Source/Editor/UnrealEd/Private/Commandlets/WorldPartitionConvertCommandlet.cpp
+++ b/Engine/Source/Editor/UnrealEd/Private/Commandlets/WorldPartitionConvertCommandlet.cpp
@@ -1037,7 +1037,7 @@ int32 UWorldPartitionConvertCommandlet::Main(const FString& Params)
 	auto ConvertActorLayersToDataLayers = [this, MainWorldDataLayers, &AssetTools](AActor* Actor)
 	{
 		// Convert Layers into DataLayers with DynamicallyLoaded flag disabled
-		if (Actor->SupportsDataLayer())
+		if (Actor->SupportsDataLayerType(UDataLayerInstance::StaticClass()))
 		{
 			for (FName Layer : Actor->Layers)
 			{
diff --git a/Engine/Source/Runtime/Engine/Classes/Atmosphere/AtmosphericFog.h b/Engine/Source/Runtime/Engine/Classes/Atmosphere/AtmosphericFog.h
--- a/Engine/Source/Runtime/Engine/Classes/Atmosphere/AtmosphericFog.h
+++ b/Engine/Source/Runtime/Engine/Classes/Atmosphere/AtmosphericFog.h
@@ -19,7 +19,7 @@ class UE_DEPRECATED(4.26, "Please use the SkyAtmosphere actor instead.") AAtmosp
 
 private:
 #if WITH_EDITOR
-	virtual bool ActorTypeSupportsDataLayer() const override { return true; }
+	virtual bool IsDataLayerTypeSupported(TSubclassOf<UDataLayerInstance> DataLayerType) const override { return true; }
 #endif
 
 	/** Main fog component */
diff --git a/Engine/Source/Runtime/Engine/Classes/Components/HeterogeneousVolumeComponent.h b/Engine/Source/Runtime/Engine/Classes/Components/HeterogeneousVolumeComponent.h
--- a/Engine/Source/Runtime/Engine/Classes/Components/HeterogeneousVolumeComponent.h
+++ b/Engine/Source/Runtime/Engine/Classes/Components/HeterogeneousVolumeComponent.h
@@ -68,7 +68,7 @@ private:
 	TObjectPtr<class UHeterogeneousVolumeComponent> HeterogeneousVolumeComponent;
 
 #if WITH_EDITOR
-	virtual bool ActorTypeSupportsDataLayer() const override { return true; }
+	virtual bool IsDataLayerTypeSupported(TSubclassOf<UDataLayerInstance> DataLayerType)  const override { return true; }
 #endif
 
 };
\ No newline at end of file
diff --git a/Engine/Source/Runtime/Engine/Classes/Components/SkyAtmosphereComponent.h b/Engine/Source/Runtime/Engine/Classes/Components/SkyAtmosphereComponent.h
--- a/Engine/Source/Runtime/Engine/Classes/Components/SkyAtmosphereComponent.h
+++ b/Engine/Source/Runtime/Engine/Classes/Components/SkyAtmosphereComponent.h
@@ -286,7 +286,7 @@ class ASkyAtmosphere : public AInfo
 
 private:
 #if WITH_EDITOR
-	virtual bool ActorTypeSupportsDataLayer() const override { return true; }
+	virtual bool IsDataLayerTypeSupported(TSubclassOf<UDataLayerInstance> DataLayerType) const override { return true; }
 #endif
 
 	UPROPERTY(BlueprintReadOnly, VisibleAnywhere, Category = Atmosphere, meta = (AllowPrivateAccess = "true"))
diff --git a/Engine/Source/Runtime/Engine/Classes/Components/VolumetricCloudComponent.h b/Engine/Source/Runtime/Engine/Classes/Components/VolumetricCloudComponent.h
--- a/Engine/Source/Runtime/Engine/Classes/Components/VolumetricCloudComponent.h
+++ b/Engine/Source/Runtime/Engine/Classes/Components/VolumetricCloudComponent.h
@@ -228,7 +228,7 @@ private:
 	TObjectPtr<class UVolumetricCloudComponent> VolumetricCloudComponent;
 
 #if WITH_EDITOR
-	virtual bool ActorTypeSupportsDataLayer() const override { return true; }
+	virtual bool IsDataLayerTypeSupported(TSubclassOf<UDataLayerInstance> DataLayerType) const override { return true; }
 #endif
 
 };
diff --git a/Engine/Source/Runtime/Engine/Classes/Engine/ExponentialHeightFog.h b/Engine/Source/Runtime/Engine/Classes/Engine/ExponentialHeightFog.h
--- a/Engine/Source/Runtime/Engine/Classes/Engine/ExponentialHeightFog.h
+++ b/Engine/Source/Runtime/Engine/Classes/Engine/ExponentialHeightFog.h
@@ -20,7 +20,7 @@ class AExponentialHeightFog
 
 private:
 #if WITH_EDITOR
-	virtual bool ActorTypeSupportsDataLayer() const override { return true; }
+	virtual bool IsDataLayerTypeSupported(TSubclassOf<UDataLayerInstance> DataLayerType) const override { return true; }
 #endif
 
 	/** @todo document */
diff --git a/Engine/Source/Runtime/Engine/Classes/Engine/LocalHeightFog.h b/Engine/Source/Runtime/Engine/Classes/Engine/LocalHeightFog.h
--- a/Engine/Source/Runtime/Engine/Classes/Engine/LocalHeightFog.h
+++ b/Engine/Source/Runtime/Engine/Classes/Engine/LocalHeightFog.h
@@ -19,7 +19,7 @@ class ALocalHeightFog : public AInfo
 
 private:
 #if WITH_EDITOR
-	virtual bool ActorTypeSupportsDataLayer() const override { return true; }
+	virtual bool IsDataLayerTypeSupported(TSubclassOf<UDataLayerInstance> DataLayerType) const override { return true; }
 #endif
 
 	/** Object used to visualize the local fog volume */
diff --git a/Engine/Source/Runtime/Engine/Classes/Engine/SkyLight.h b/Engine/Source/Runtime/Engine/Classes/Engine/SkyLight.h
--- a/Engine/Source/Runtime/Engine/Classes/Engine/SkyLight.h
+++ b/Engine/Source/Runtime/Engine/Classes/Engine/SkyLight.h
@@ -16,7 +16,7 @@ class ENGINE_API ASkyLight : public AInfo
 
 private:
 #if WITH_EDITOR
-	virtual bool ActorTypeSupportsDataLayer() const override { return true; }
+	virtual bool IsDataLayerTypeSupported(TSubclassOf<UDataLayerInstance> DataLayerType) const override { return true; }
 #endif
 
 	/** @todo document */
diff --git a/Engine/Source/Runtime/Engine/Classes/Engine/WindDirectionalSource.h b/Engine/Source/Runtime/Engine/Classes/Engine/WindDirectionalSource.h
--- a/Engine/Source/Runtime/Engine/Classes/Engine/WindDirectionalSource.h
+++ b/Engine/Source/Runtime/Engine/Classes/Engine/WindDirectionalSource.h
@@ -14,7 +14,7 @@ class ENGINE_API AWindDirectionalSource : public AInfo
 
 private:
 #if WITH_EDITOR
-	virtual bool ActorTypeSupportsDataLayer() const override { return true; }
+	virtual bool IsDataLayerTypeSupported(TSubclassOf<UDataLayerInstance> DataLayerType) const override { return true; }
 #endif
 
 	UPROPERTY(Category = WindDirectionalSource, VisibleAnywhere, BlueprintReadOnly, meta = (AllowPrivateAccess = "true"))
diff --git a/Engine/Source/Runtime/Engine/Classes/GameFramework/Actor.h b/Engine/Source/Runtime/Engine/Classes/GameFramework/Actor.h
--- a/Engine/Source/Runtime/Engine/Classes/GameFramework/Actor.h
+++ b/Engine/Source/Runtime/Engine/Classes/GameFramework/Actor.h
@@ -1292,14 +1292,14 @@ public:
 	// DataLayers functions.
 #if WITH_EDITOR
 private:
-	virtual bool ActorTypeSupportsDataLayer() const { return true; }
+	virtual bool IsDataLayerTypeSupported(TSubclassOf<UDataLayerInstance> DataLayerType) const { return true; }
 	virtual bool ActorTypeIsMainWorldOnly() const { return false; }
 	TArray<const UDataLayerAsset*> ResolveDataLayerAssets(const TArray<TSoftObjectPtr<UDataLayerAsset>>& InDataLayerAssets) const;
 public:
 	bool AddDataLayer(const UDataLayerInstance* DataLayerInstance);
 	bool RemoveDataLayer(const UDataLayerInstance* DataLayerInstance);
 	TArray<const UDataLayerInstance*> RemoveAllDataLayers();
-	bool SupportsDataLayer() const;
+	bool SupportsDataLayerType(TSubclassOf<UDataLayerInstance> DataLayerType) const;
 	
 	TArray<const UDataLayerInstance*> GetDataLayerInstancesForLevel() const;
 	TArray<FName> GetDataLayerInstanceNames() const;
@@ -1322,6 +1322,12 @@ public:
 
 	PRAGMA_ENABLE_DEPRECATION_WARNINGS
 
+	UE_DEPRECATED(5.3, "Use AActor::SupportsDataLayerType(TSubclassOf<UDataLayerInstance>) instead.")
+	bool SupportsDataLayer() const;
+
+	UE_DEPRECATED(5.3, "Use AActor::IsDataLayerTypeSupported(TSubclassOf<UDataLayerInstance>) instead.")
+	virtual bool ActorTypeSupportsDataLayer() const { return true; }
+
 	UE_DEPRECATED(5.3, "Use UDataLayerInstance::AddActor(AActor*) instead.")
 	bool AddDataLayer(const UDataLayerAsset* DataLayerAsset) { return false; }
 
diff --git a/Engine/Source/Runtime/Engine/Classes/GameFramework/Info.h b/Engine/Source/Runtime/Engine/Classes/GameFramework/Info.h
--- a/Engine/Source/Runtime/Engine/Classes/GameFramework/Info.h
+++ b/Engine/Source/Runtime/Engine/Classes/GameFramework/Info.h
@@ -30,7 +30,7 @@ public:
 
 private:
 #if WITH_EDITOR
-	virtual bool ActorTypeSupportsDataLayer() const override { return false; }
+	ENGINE_API virtual bool IsDataLayerTypeSupported(TSubclassOf<UDataLayerInstance> DataLayerType) const override { return false; }
 #endif
 
 #if WITH_EDITORONLY_DATA
diff --git a/Engine/Source/Runtime/Engine/Private/Actor.cpp b/Engine/Source/Runtime/Engine/Private/Actor.cpp
--- a/Engine/Source/Runtime/Engine/Private/Actor.cpp
+++ b/Engine/Source/Runtime/Engine/Private/Actor.cpp
@@ -6207,7 +6207,7 @@ TArray<const UDataLayerInstance*> AActor::GetDataLayerInstancesInternal(bool bUs
 	}
 
 #if WITH_EDITOR
-	if (SupportsDataLayer())
+	if (SupportsDataLayerType(UDataLayerInstance::StaticClass()))
 	{
 		TArray<const UDataLayerInstance*> DataLayerInstances;
 		if (UDataLayerManager* DataLayerManager = bUseLevelContext ? UDataLayerManager::GetDataLayerManager(this) : UDataLayerManager::GetDataLayerManager(GetWorld()))
diff --git a/Engine/Source/Runtime/Engine/Private/ActorEditor.cpp b/Engine/Source/Runtime/Engine/Private/ActorEditor.cpp
--- a/Engine/Source/Runtime/Engine/Private/ActorEditor.cpp
+++ b/Engine/Source/Runtime/Engine/Private/ActorEditor.cpp
@@ -26,6 +26,8 @@
 #include "ActorTransactionAnnotation.h"
 #include "Elements/Framework/EngineElementsLibrary.h"
 #include "WorldPartition/DataLayer/IDataLayerEditorModule.h"
+#include "WorldPartition/DataLayer/DataLayerInstanceWithAsset.h"
+#include "WorldPartition/DataLayer/DeprecatedDataLayerInstance.h"
 #include "LevelInstance/LevelInstanceSubsystem.h"
 #include "ActorFolder.h"
 #include "WorldPersistentFolders.h"
@@ -115,7 +117,7 @@ bool AActor::CanEditChange(const FProperty* PropertyThatWillChange) const
 		return false;
 	}
 
-	if (bIsDataLayersProperty && (!SupportsDataLayer() || !IsUserManaged()))
+	if (bIsDataLayersProperty && (!SupportsDataLayerType(UDataLayerInstance::StaticClass()) || !IsUserManaged()))
 	{
 		return false;
 	}
@@ -1695,10 +1697,18 @@ void AActor::FixupDataLayers(bool bRevertChangesOnLockedDataLayer /*= false*/)
 		return;
 	}
 
-	if (!SupportsDataLayer())
+	if (!SupportsDataLayerType(UDeprecatedDataLayerInstance::StaticClass()))
 	{
-		DataLayerAssets.Empty();
 		DataLayers.Empty();
+	}
+
+	if (!SupportsDataLayerType(UDataLayerInstanceWithAsset::StaticClass()))
+	{
+		DataLayerAssets.Empty();
+	}
+
+	if (DataLayers.IsEmpty() && DataLayerAssets.IsEmpty())
+	{
 		return;
 	}
 
@@ -1809,14 +1819,14 @@ bool AActor::IsPropertyChangedAffectingDataLayers(FPropertyChangedEvent& Propert
 	return false;
 }
 
-bool AActor::SupportsDataLayer() const
+bool AActor::SupportsDataLayerType(TSubclassOf<UDataLayerInstance> InDataLayerType) const
 {
 	ULevel* Level = GetLevel();
 	const bool bIsLevelNotPartitioned = Level ? !Level->bIsPartitioned : false;
 	return (!bIsLevelNotPartitioned &&
-			ActorTypeSupportsDataLayer() &&
-			!FActorEditorUtils::IsABuilderBrush(this) &&
-			!GetClass()->GetDefaultObject<AActor>()->bHiddenEd);
+		IsDataLayerTypeSupported(InDataLayerType) &&
+		!FActorEditorUtils::IsABuilderBrush(this) &&
+		!GetClass()->GetDefaultObject<AActor>()->bHiddenEd);
 }
 
 bool FAssignActorDataLayer::AddDataLayerAsset(AActor* InActor, const UDataLayerAsset* InDataLayerAsset)
@@ -1851,9 +1861,14 @@ bool FAssignActorDataLayer::RemoveDataLayerAsset(AActor* InActor, const UDataLay
 
 PRAGMA_DISABLE_DEPRECATION_WARNINGS
 
+bool AActor::SupportsDataLayer() const
+{
+	return SupportsDataLayerType(UDataLayerInstance::StaticClass());
+}
+
 bool AActor::AddDataLayer(const FActorDataLayer& ActorDataLayer)
 {
-	if (SupportsDataLayer() && !ContainsDataLayer(ActorDataLayer))
+	if (SupportsDataLayerType(UDataLayerInstance::StaticClass()) && !ContainsDataLayer(ActorDataLayer))
 	{
 		Modify();
 		DataLayers.Add(ActorDataLayer);
@@ -1877,7 +1892,7 @@ bool AActor::RemoveDataLayer(const FActorDataLayer& ActorDataLayer)
 
 bool AActor::AddDataLayer(const UDEPRECATED_DataLayer* DataLayer)
 {
-	if (SupportsDataLayer() && DataLayer)
+	if (SupportsDataLayerType(UDataLayerInstance::StaticClass()) && DataLayer)
 	{
 		return AddDataLayer(FActorDataLayer(DataLayer->GetFName()));
 	}
diff --git a/Engine/Source/Runtime/Engine/Private/LevelInstance/LevelInstanceSubsystem.cpp b/Engine/Source/Runtime/Engine/Private/LevelInstance/LevelInstanceSubsystem.cpp
--- a/Engine/Source/Runtime/Engine/Private/LevelInstance/LevelInstanceSubsystem.cpp
+++ b/Engine/Source/Runtime/Engine/Private/LevelInstance/LevelInstanceSubsystem.cpp
@@ -1401,7 +1401,7 @@ void ULevelInstanceSubsystem::BreakLevelInstance_Impl(ILevelInstanceInterface* L
 					}
 
 					// Apply the same data layer settings to the actors to move out
-					if (Actor->SupportsDataLayer())
+					if (Actor->SupportsDataLayerType(UDataLayerInstance::StaticClass()))
 					{
 						for (const UDataLayerInstance* DataLayerInstance : LevelInstanceDataLayerInstances)
 						{
diff --git a/Engine/Source/Runtime/Engine/Private/LightWeightInstanceSubsystem.cpp b/Engine/Source/Runtime/Engine/Private/LightWeightInstanceSubsystem.cpp
--- a/Engine/Source/Runtime/Engine/Private/LightWeightInstanceSubsystem.cpp
+++ b/Engine/Source/Runtime/Engine/Private/LightWeightInstanceSubsystem.cpp
@@ -4,6 +4,7 @@
 #include "Engine/Engine.h"
 #include "Engine/World.h"
 #include "UObject/UObjectIterator.h"
+#include "WorldPartition/DataLayer/DataLayerInstance.h"
 
 DEFINE_LOG_CATEGORY(LogLightWeightInstance);
 
@@ -146,7 +147,7 @@ ALightWeightInstanceManager* FLightWeightInstanceSubsystem::FindLightWeightInsta
 				if (ManagerGridCoord == GridCoord)
 				{
 #if WITH_EDITOR
-					if (!DataLayer || (InstanceManager->SupportsDataLayer() && InstanceManager->ContainsDataLayer(DataLayer)))
+					if (!DataLayer || (InstanceManager->SupportsDataLayerType(UDataLayerInstance::StaticClass()) && InstanceManager->ContainsDataLayer(DataLayer)))
 #endif // WITH_EDITOR
 					{
 						return InstanceManager;
@@ -198,7 +199,7 @@ ALightWeightInstanceManager* FLightWeightInstanceSubsystem::FindOrAddLightWeight
 		// Add the new manager to the DataLayer
 		if (DataLayer)
 		{
-			ensure(NewInstanceManager->SupportsDataLayer());
+			ensure(NewInstanceManager->SupportsDataLayerType(UDataLayerInstance::StaticClass()));
 			NewInstanceManager->AddDataLayer(DataLayer);
 		}
 #endif // WITH_EDITOR
diff --git a/Engine/Source/Runtime/Engine/Private/WorldPartition/DataLayer/DataLayerInstance.cpp b/Engine/Source/Runtime/Engine/Private/WorldPartition/DataLayer/DataLayerInstance.cpp
--- a/Engine/Source/Runtime/Engine/Private/WorldPartition/DataLayer/DataLayerInstance.cpp
+++ b/Engine/Source/Runtime/Engine/Private/WorldPartition/DataLayer/DataLayerInstance.cpp
@@ -157,7 +157,7 @@ bool UDataLayerInstance::CanUserAddActors() const
 
 bool UDataLayerInstance::CanAddActor(AActor* InActor) const
 {
-	return InActor != nullptr && InActor->SupportsDataLayer() && !InActor->ContainsDataLayer(this);
+	return InActor != nullptr && InActor->SupportsDataLayerType(GetClass()) && !InActor->ContainsDataLayer(this);
 }
 
 bool UDataLayerInstance::AddActor(AActor* InActor) const
diff --git a/Engine/Source/Runtime/Engine/Public/LevelInstance/LevelInstanceEditorPivotActor.h b/Engine/Source/Runtime/Engine/Public/LevelInstance/LevelInstanceEditorPivotActor.h
--- a/Engine/Source/Runtime/Engine/Public/LevelInstance/LevelInstanceEditorPivotActor.h
+++ b/Engine/Source/Runtime/Engine/Public/LevelInstance/LevelInstanceEditorPivotActor.h
@@ -22,6 +22,6 @@ public:
 	virtual void PostEditUndo() override;
 	virtual bool ShouldExport() override { return false; }
 private:
-	virtual bool ActorTypeSupportsDataLayer() const { return false; }
+	virtual bool IsDataLayerTypeSupported(TSubclassOf<UDataLayerInstance> DataLayerType) const override { return false; }
 #endif
 };
\ No newline at end of file
diff --git a/Engine/Source/Runtime/Engine/Public/LocationVolume.h b/Engine/Source/Runtime/Engine/Public/LocationVolume.h
--- a/Engine/Source/Runtime/Engine/Public/LocationVolume.h
+++ b/Engine/Source/Runtime/Engine/Public/LocationVolume.h
@@ -63,7 +63,7 @@ public:
 
 #if WITH_EDITOR
 private:
-	virtual bool ActorTypeSupportsDataLayer() const override { return false; }
+	virtual bool IsDataLayerTypeSupported(TSubclassOf<UDataLayerInstance> DataLayerType) const override { return false; }
 	FLoaderAdapterActor* WorldPartitionActorLoader;
 #endif
 };
\ No newline at end of file
diff --git a/Engine/Source/Runtime/Engine/Public/WorldPartition/DataLayer/WorldDataLayers.h b/Engine/Source/Runtime/Engine/Public/WorldPartition/DataLayer/WorldDataLayers.h
--- a/Engine/Source/Runtime/Engine/Public/WorldPartition/DataLayer/WorldDataLayers.h
+++ b/Engine/Source/Runtime/Engine/Public/WorldPartition/DataLayer/WorldDataLayers.h
@@ -55,6 +55,7 @@ public:
 	virtual bool ShouldImport(FStringView ActorPropString, bool IsMovingLevel) override { return false; }
 	virtual bool IsLockLocation() const { return true; }
 	virtual bool IsUserManaged() const override { return false; }
+	virtual bool IsDataLayerTypeSupported(TSubclassOf<UDataLayerInstance> DataLayerType) const { return false; }
 	virtual TUniquePtr<class FWorldPartitionActorDesc> CreateClassActorDesc() const override;
 
 	static AWorldDataLayers* Create(UWorld* World, FName InWorldDataLayerName = NAME_None);
diff --git a/Engine/Source/Runtime/Engine/Public/WorldPartition/WorldPartitionMiniMap.h b/Engine/Source/Runtime/Engine/Public/WorldPartition/WorldPartitionMiniMap.h
--- a/Engine/Source/Runtime/Engine/Public/WorldPartition/WorldPartitionMiniMap.h
+++ b/Engine/Source/Runtime/Engine/Public/WorldPartition/WorldPartitionMiniMap.h
@@ -18,7 +18,7 @@ class ENGINE_API AWorldPartitionMiniMap : public AInfo
 
 private:
 #if WITH_EDITOR
-	virtual bool ActorTypeSupportsDataLayer() const final { return false; }
+	virtual bool IsDataLayerTypeSupported(TSubclassOf<UDataLayerInstance> DataLayerType) const final { return false; }
 	virtual bool ActorTypeIsMainWorldOnly() const override  { return true; }
 #endif
 
diff --git a/Engine/Source/Runtime/Engine/Public/WorldPartition/WorldPartitionMiniMapVolume.h b/Engine/Source/Runtime/Engine/Public/WorldPartition/WorldPartitionMiniMapVolume.h
--- a/Engine/Source/Runtime/Engine/Public/WorldPartition/WorldPartitionMiniMapVolume.h
+++ b/Engine/Source/Runtime/Engine/Public/WorldPartition/WorldPartitionMiniMapVolume.h
@@ -16,7 +16,7 @@ public:
 	virtual bool IsEditorOnly() const final { return true; }
 
 #if WITH_EDITOR
-	virtual bool ActorTypeSupportsDataLayer() const final { return false; }
+	virtual bool IsDataLayerTypeSupported(TSubclassOf<UDataLayerInstance> DataLayerType) const final { return false; }
 	virtual bool CanChangeIsSpatiallyLoadedFlag() const { return false; }
 #endif
 };
\ No newline at end of file
diff --git a/Engine/Source/Runtime/Engine/Public/WorldPartition/WorldPartitionVolume.h b/Engine/Source/Runtime/Engine/Public/WorldPartition/WorldPartitionVolume.h
--- a/Engine/Source/Runtime/Engine/Public/WorldPartition/WorldPartitionVolume.h
+++ b/Engine/Source/Runtime/Engine/Public/WorldPartition/WorldPartitionVolume.h
@@ -15,7 +15,7 @@ class ENGINE_API ADEPRECATED_WorldPartitionVolume : public AVolume
 
 private:
 #if WITH_EDITOR
-	virtual bool ActorTypeSupportsDataLayer() const override { return false; }
+	virtual bool IsDataLayerTypeSupported(TSubclassOf<UDataLayerInstance> DataLayerType) const override { return false; }
 #endif
 
 public:
diff --git a/Engine/Source/Runtime/Renderer/Private/SparseVolumeTexture/SparseVolumeTextureViewerComponent.h b/Engine/Source/Runtime/Renderer/Private/SparseVolumeTexture/SparseVolumeTextureViewerComponent.h
--- a/Engine/Source/Runtime/Renderer/Private/SparseVolumeTexture/SparseVolumeTextureViewerComponent.h
+++ b/Engine/Source/Runtime/Renderer/Private/SparseVolumeTexture/SparseVolumeTextureViewerComponent.h
@@ -112,7 +112,7 @@ private:
 	TObjectPtr<class USparseVolumeTextureViewerComponent> SparseVolumeTextureViewerComponent;
 
 #if WITH_EDITOR
-	virtual bool ActorTypeSupportsDataLayer() const override { return true; }
+	virtual bool IsDataLayerTypeSupported(TSubclassOf<UDataLayerInstance> DataLayerType) const override { return true; }
 #endif
 
 };
