Bounds check index in case we have a stale view of our initial properties

diff --git a/Engine/Source/Runtime/Engine/Private/BlueprintGeneratedClass.cpp b/Engine/Source/Runtime/Engine/Private/BlueprintGeneratedClass.cpp
--- a/Engine/Source/Runtime/Engine/Private/BlueprintGeneratedClass.cpp
+++ b/Engine/Source/Runtime/Engine/Private/BlueprintGeneratedClass.cpp
@@ -992,6 +992,16 @@ void UBlueprintGeneratedClass::InitArrayPropertyFromCustomList(const FArrayPrope
 			continue;
 		}
 
+		if (!SrcArrayValueHelper.IsValidIndex(ArrayIndex)) // dst bounds were conformed above, so just need to check source
+		{
+			ensureMsgf(false,
+				TEXT("InitArrayPropertyFromCustomList attempted out of bounds access within %s,"
+				"this indicates a template was mutated without calling UpdateCustomPropertyListForPostConstruction"),
+				*ArrayProperty->GetOwner<UStruct>()->GetFullName()
+				);
+			continue;
+		}
+
 		if (!InitPropertyFromSubPropertyList(ArrayProperty->Inner, CustomArrayPropertyListNode->SubPropertyList, DstArrayItemValue, SrcArrayItemValue))
 		{
 			// Unable to init properties from sub custom property list, fall back to the default copy value behavior
