Fix for seeking off the end off a stream

diff --git a/Engine/Source/Runtime/Engine/Private/AudioDecompress.cpp b/Engine/Source/Runtime/Engine/Private/AudioDecompress.cpp
--- a/Engine/Source/Runtime/Engine/Private/AudioDecompress.cpp
+++ b/Engine/Source/Runtime/Engine/Private/AudioDecompress.cpp
@@ -225,7 +225,6 @@ bool IStreamedCompressedInfo::StreamCompressedData(uint8* Destination, bool bLoo
 
 	SCOPE_CYCLE_COUNTER(STAT_AudioStreamedDecompressTime);
 
-
 	UE_LOG(LogAudio, Log, TEXT("Streaming compressed data from SoundWave'%s' - Chunk=%d\tNumChunks=%d\tOffset=%d\tChunkSize=%d\tLooping=%s\tLastPCMOffset=%d\tContainsEOF=%s" ), 
 		*StreamingSoundWave->GetFName().ToString(), 	
 		CurrentChunkIndex, 
@@ -244,41 +243,63 @@ bool IStreamedCompressedInfo::StreamCompressedData(uint8* Destination, bool bLoo
 	if (StreamSeekBlockIndex != INDEX_NONE)
 	{
 		uint32 ChunkSize = 0;
-		SrcBufferData = GetLoadedChunk(StreamingSoundWave, StreamSeekBlockIndex, ChunkSize);
+		const uint8* NewlySeekedChunk = GetLoadedChunk(StreamingSoundWave, StreamSeekBlockIndex, ChunkSize);
 		UE_LOG(LogAudio, Log, TEXT("Seek block request: %d / %d (%s)"), StreamSeekBlockIndex.load(), StreamSeekBlockOffset, SrcBufferData == nullptr ? TEXT("present") : TEXT("missing"));
-		if (SrcBufferData == nullptr)
+		if (NewlySeekedChunk == nullptr)
 		{
 			// After a seek we're likely to need to wait a bit for the chunk to get in to memory.
 			ZeroBuffer(Destination + RawPCMOffset, BufferSize - RawPCMOffset);
 			return false;
-		}
-
-		CurrentChunkIndex = StreamSeekBlockIndex;
-		SrcBufferDataSize = ChunkSize;
-		SrcBufferOffset = StreamSeekBlockOffset;
-		StreamSeekBlockIndex = INDEX_NONE;
+		}	
 
-		// If we are using a chunked seek-table, we need to parse the table first.
-		if (StreamingSoundWave->GetSoundWaveData()->HasChunkSeekTable(CurrentChunkIndex))
+		if (StreamSeekToAudioFrames == INDEX_NONE)
+		{
+			// Non-streaming-seek tables.
+			// Commit the new chunk as the current.
+			SrcBufferData = NewlySeekedChunk;
+			CurrentChunkIndex = StreamSeekBlockIndex;
+			SrcBufferDataSize = ChunkSize;
+			SrcBufferOffset = StreamSeekBlockOffset;
+		}
+		else
 		{
-			uint32 TableOffset = 0;
-			if (ensureMsgf(FStreamedAudioChunkSeekTable::Parse(SrcBufferData, ChunkSize, TableOffset, GetCurrentSeekTable()), 
-				TEXT("Failed to parse seektable in '%s', chunk=%d"), *StreamingSoundWave->GetFName().ToString(), CurrentChunkIndex))
+			// If we are using a chunked seek-table, we need to parse the table first.
+			if (StreamingSoundWave->GetSoundWaveData()->HasChunkSeekTable(CurrentChunkIndex))
 			{
-				if (StreamSeekToAudioFrames != INDEX_NONE)
+				uint32 TableOffset = 0;
+				if (ensureMsgf(FStreamedAudioChunkSeekTable::Parse(NewlySeekedChunk, ChunkSize, TableOffset, GetCurrentSeekTable()),
+					TEXT("Failed to parse seektable in '%s', chunk=%d"), *StreamingSoundWave->GetFName().ToString(), CurrentChunkIndex))
 				{
-					const FStreamedAudioChunk& Chunk = StreamingSoundWave->GetSoundWaveData()->GetChunk(CurrentChunkIndex);
-					int32 SeekOffsetInChunkSpace = StreamSeekToAudioFrames - Chunk.SeekOffsetInAudioFrames;
-					uint32 Offset = GetCurrentSeekTable().FindOffset(StreamSeekToAudioFrames);
-					if (Offset != INDEX_NONE)
+					if (StreamSeekToAudioFrames != INDEX_NONE)
 					{
-						SrcBufferOffset = Offset + TableOffset;
-						SrcBufferOffset += CurrentChunkIndex == 0 ? AudioDataOffset : 0;
+						const FStreamedAudioChunk& Chunk = StreamingSoundWave->GetSoundWaveData()->GetChunk(CurrentChunkIndex);
+						int32 SeekOffsetInChunkSpace = StreamSeekToAudioFrames - Chunk.SeekOffsetInAudioFrames;
+						uint32 Offset = GetCurrentSeekTable().FindOffset(StreamSeekToAudioFrames);
+						if (Offset != INDEX_NONE)
+						{
+							// All looks good, commit this as our current chunk.
+							CurrentChunkIndex = StreamSeekBlockIndex;
+							SrcBufferDataSize = ChunkSize;
+							SrcBufferData = NewlySeekedChunk;
+
+							// Set our offset from the table.
+							SrcBufferOffset = Offset + TableOffset;
+							SrcBufferOffset += CurrentChunkIndex == 0 ? AudioDataOffset : 0;							
+						}
+						else // Seek failed (off the end of the chunk).
+						{
+							const float TimeInSeconds = (float)StreamSeekToAudioFrames / StreamingSoundWave->GetSampleRate();
+							UE_LOG(LogAudio, Log, TEXT("Failed seeking to %2.2f seconds as its off then end of the stream. Wave=%s"), TimeInSeconds, *StreamingSoundWave->GetFName().ToString());
+						}
 					}
-					StreamSeekToAudioFrames = INDEX_NONE;
 				}
 			}
 		}
+
+		// Seeking logic over, turn off seeking state
+		StreamSeekBlockIndex = INDEX_NONE;
+		StreamSeekToAudioFrames = INDEX_NONE;
+		StreamSeekBlockOffset = INDEX_NONE;
 	}
 
 	// If next chunk wasn't loaded when last one finished reading, try to get it again now
