Fix per-instance velocity output by checking if the HasPerInstanceDynamicData() is true (manifests in the Mass cars in CitySample).

diff --git a/Engine/Source/Runtime/Engine/Public/PrimitiveSceneProxy.h b/Engine/Source/Runtime/Engine/Public/PrimitiveSceneProxy.h
--- a/Engine/Source/Runtime/Engine/Public/PrimitiveSceneProxy.h
+++ b/Engine/Source/Runtime/Engine/Public/PrimitiveSceneProxy.h
@@ -773,9 +773,9 @@ public:
 	/** Returns true if this proxy can change transform so that we should cache previous transform for calculating velocity. */
 	inline bool HasDynamicTransform() const { return IsMovable() || bIsBeingMovedByEditor; }
 	/** Returns true if this proxy can write velocity. This is used for setting velocity relevance. */
-	inline bool DrawsVelocity() const { return HasDynamicTransform() || bAlwaysHasVelocity || bHasWorldPositionOffsetVelocity; }
+	inline bool DrawsVelocity() const { return HasDynamicTransform() || bAlwaysHasVelocity || bHasWorldPositionOffsetVelocity || HasPerInstanceDynamicData(); }
 	/** Returns true if this proxy should write velocity even when the transform isn't changing. Usually this is combined with a check for the transform changing. */
-	inline bool AlwaysHasVelocity() const {	return bAlwaysHasVelocity || (bHasWorldPositionOffsetVelocity && EvaluateWorldPositionOffset()); }
+	inline bool AlwaysHasVelocity() const {	return bAlwaysHasVelocity || (bHasWorldPositionOffsetVelocity && EvaluateWorldPositionOffset()) || HasPerInstanceDynamicData(); }
 
 #if WITH_EDITOR
 	inline int32 GetNumUncachedStaticLightingInteractions() { return NumUncachedStaticLightingInteractions; }
