LocalFogVolume - sort priority now workss the same way as for meshes.

diff --git a/Engine/Source/Runtime/Engine/Classes/Components/LocalHeightFogComponent.h b/Engine/Source/Runtime/Engine/Classes/Components/LocalHeightFogComponent.h
--- a/Engine/Source/Runtime/Engine/Classes/Components/LocalHeightFogComponent.h
+++ b/Engine/Source/Runtime/Engine/Classes/Components/LocalHeightFogComponent.h
@@ -55,9 +55,9 @@ class ULocalHeightFogComponent : public USceneComponent
 	UPROPERTY(EditAnywhere, BlueprintReadOnly, interp, Category = "Fog Shading", meta = (ClampMin = 0.0, ClampMax = 1.0))
 	FLinearColor FogEmissive = FLinearColor::Black;
 
-	/** The priority can be used as a way to overide the sorting by distance. A higher value means the volume will be considered further away, i.e. behing the one with a lower priority value. */
-	UPROPERTY(EditAnywhere, BlueprintReadOnly, Category = "Sorting", meta = (UIMin = "0", UIMax = "255", ClampMin = 0, ClampMax = 255))
-	uint8 FogSortPriority = 127;
+	/** The priority can be used as a way to override the sorting by distance. A lower value means the volume will be considered further away, i.e. it will draw behind the one with a higher priority value. */
+	UPROPERTY(EditAnywhere, BlueprintReadOnly, Category = "Sorting", meta = (UIMin = "-127", UIMax = "127", ClampMin = -127, ClampMax = 127))
+	int32 FogSortPriority = 0;
 
 public:
 
diff --git a/Engine/Source/Runtime/Engine/Private/LocalHeightFogSceneProxy.cpp b/Engine/Source/Runtime/Engine/Private/LocalHeightFogSceneProxy.cpp
--- a/Engine/Source/Runtime/Engine/Private/LocalHeightFogSceneProxy.cpp
+++ b/Engine/Source/Runtime/Engine/Private/LocalHeightFogSceneProxy.cpp
@@ -16,7 +16,7 @@ FLocalHeightFogSceneProxy::FLocalHeightFogSceneProxy(const ULocalHeightFogCompon
 	, FogHeightOffset(InComponent->FogHeightOffset)
 	, FogRadialAttenuation(InComponent->FogRadialAttenuation)
 	, FogMode((uint8)InComponent->FogMode)
-	, FogSortPriority(InComponent->FogSortPriority)
+	, FogSortPriority(127 + InComponent->FogSortPriority)
 	, FogPhaseG(InComponent->FogPhaseG)
 	, FogAlbedo(InComponent->FogAlbedo)
 	, FogEmissive(InComponent->FogEmissive)
diff --git a/Engine/Source/Runtime/Renderer/Private/LocalHeightFogRendering.cpp b/Engine/Source/Runtime/Renderer/Private/LocalHeightFogRendering.cpp
--- a/Engine/Source/Runtime/Renderer/Private/LocalHeightFogRendering.cpp
+++ b/Engine/Source/Runtime/Renderer/Private/LocalHeightFogRendering.cpp
@@ -156,7 +156,7 @@ public:
 
 	FORCEINLINE bool operator<(FLocalFogVolumeSortKey B) const
 	{
-		return PackedData > B.PackedData;
+		return PackedData < B.PackedData;
 	}
 };
 
