Fix for ParallelFor() iterating over the wrong number of jobs in SVT::FTextureData::CreateFromDense(). Fixes a potential crash that can happen when writing to invalid memory.

#rnx

diff --git a/Engine/Source/Runtime/Engine/Private/SparseVolumeTexture/SparseVolumeTextureData.cpp b/Engine/Source/Runtime/Engine/Private/SparseVolumeTexture/SparseVolumeTextureData.cpp
--- a/Engine/Source/Runtime/Engine/Private/SparseVolumeTexture/SparseVolumeTextureData.cpp
+++ b/Engine/Source/Runtime/Engine/Private/SparseVolumeTexture/SparseVolumeTextureData.cpp
@@ -313,10 +313,11 @@ bool FTextureData::CreateFromDense(const FTextureDataCreateInfo& CreateInfo, con
 		}
 	};
 
-	const int32 NumJobs = Header.PageTableVolumeResolution.Z * Header.PageTableVolumeResolution.Y;
+	// Process an entire row of pages along the X dimension per job. We iterate over the logical pages and sample the dense source voxel data; we do not iterate over the physical tiles in the destination SVT.
+	const int32 NumPageTableJobs = Header.PageTableVolumeResolution.Z * Header.PageTableVolumeResolution.Y;
 
 	// Create page table
-	ParallelFor(NumJobs, [&](int32 Index) { ProcessBlock(Index, false /*bWriteTileData*/); });
+	ParallelFor(NumPageTableJobs, [&](int32 Index) { ProcessBlock(Index, false /*bWriteTileData*/); });
 	const int32 NumAllocatedPages = (int32)NumAllocatedPagesAtomic.load();
 
 #if SVT_DETERMINISTIC_PAGE_TABLE_GENERATION
@@ -342,7 +343,8 @@ bool FTextureData::CreateFromDense(const FTextureDataCreateInfo& CreateInfo, con
 	// SVT_TODO: Make this smarter. We don't need to write ALL voxels in this stage because we know which voxels will be overwritten anyways.
 	if (FractionalPageOffset != FIntVector3::ZeroValue || (Header.VirtualVolumeAABBMax % SPARSE_VOLUME_TILE_RES) != FIntVector3::ZeroValue)
 	{
-		ParallelFor(NumJobs, [&](int32 TileIndex)
+		// Iterate over physical tiles (which at this point equals the number of allocated pages); NOT logical pages!
+		ParallelFor(NumAllocatedPages, [&](int32 TileIndex)
 			{
 				for (int32 VoxelIndex = 0; VoxelIndex < SVT::NumVoxelsPerPaddedTile; ++VoxelIndex)
 				{
@@ -360,7 +362,7 @@ bool FTextureData::CreateFromDense(const FTextureDataCreateInfo& CreateInfo, con
 	}
 
 	// Write physical tile data
-	ParallelFor(NumJobs, [&](int32 Index) { ProcessBlock(Index, true /*bWriteTileData*/); });
+	ParallelFor(NumPageTableJobs, [&](int32 Index) { ProcessBlock(Index, true /*bWriteTileData*/); });
 
 	return true;
 }
