LocalHeightFogVolume - fixed wrongdepth  ordering.

diff --git a/Engine/Source/Runtime/Engine/Private/LocalHeightFogSceneProxy.cpp b/Engine/Source/Runtime/Engine/Private/LocalHeightFogSceneProxy.cpp
--- a/Engine/Source/Runtime/Engine/Private/LocalHeightFogSceneProxy.cpp
+++ b/Engine/Source/Runtime/Engine/Private/LocalHeightFogSceneProxy.cpp
@@ -16,7 +16,7 @@ FLocalHeightFogSceneProxy::FLocalHeightFogSceneProxy(const ULocalHeightFogCompon
 	, FogHeightOffset(InComponent->FogHeightOffset)
 	, FogRadialAttenuation(InComponent->FogRadialAttenuation)
 	, FogMode((uint8)InComponent->FogMode)
-	, FogSortPriority(127 + InComponent->FogSortPriority)
+	, FogSortPriority(uint8(127 - int8(InComponent->FogSortPriority))) // FogSortPriority on the component is in [-127,127] and needs to be negated to match expected priority behavior.
 	, FogPhaseG(InComponent->FogPhaseG)
 	, FogAlbedo(InComponent->FogAlbedo)
 	, FogEmissive(InComponent->FogEmissive)
diff --git a/Engine/Source/Runtime/Renderer/Private/LocalHeightFogRendering.cpp b/Engine/Source/Runtime/Renderer/Private/LocalHeightFogRendering.cpp
--- a/Engine/Source/Runtime/Renderer/Private/LocalHeightFogRendering.cpp
+++ b/Engine/Source/Runtime/Renderer/Private/LocalHeightFogRendering.cpp
@@ -156,7 +156,7 @@ public:
 
 	FORCEINLINE bool operator<(FLocalFogVolumeSortKey B) const
 	{
-		return PackedData < B.PackedData;
+		return PackedData > B.PackedData;
 	}
 };
 
