Reintroduced (and deprecated) the older version of ContainsInnerLoopInternal() in UMaterialExpression.

diff --git a/Engine/Source/Runtime/Engine/Classes/Materials/MaterialExpression.h b/Engine/Source/Runtime/Engine/Classes/Materials/MaterialExpression.h
--- a/Engine/Source/Runtime/Engine/Classes/Materials/MaterialExpression.h
+++ b/Engine/Source/Runtime/Engine/Classes/Materials/MaterialExpression.h
@@ -27,6 +27,7 @@ class UMaterialExpression;
 class UMaterialExpressionComment;
 class UMaterialExpressionExecBegin;
 class UMaterialExpressionExecEnd;
+class FMaterialExpressionKey;
 
 class FMaterialHLSLGenerator;
 enum class EMaterialNewScopeFlag : uint8;
@@ -349,7 +350,7 @@ class ENGINE_API UMaterialExpression : public UObject
 	/** Get the expression inputs supported by this expression (Note: property inputs NOT included). */
 	virtual TArrayView<FExpressionInput*> GetInputsView();
 	
-	UE_DEPRECATED(5.2, "Use GetInputsView() instead.")
+	UE_DEPRECATED(5.3, "Use GetInputsView() instead.")
 	const TArray<FExpressionInput*> GetInputs() { return TArray<FExpressionInput*>{ GetInputsView() }; }
 
 	virtual FExpressionInput* GetInput(int32 InputIndex);
@@ -597,6 +598,8 @@ private:
 	 */
 	bool ContainsInputLoopInternal(const FContainsInputLoopInternalExpressionStack& ExpressionStack, TSet<UMaterialExpression*>& VisitedExpressions, const bool bStopOnFunctionCall);
 
+	UE_DEPRECATED(5.3, "Use the other, more efficient ContainsInputLoopInternal() implementation.")
+	bool ContainsInputLoopInternal(TArray<FMaterialExpressionKey>& ExpressionStack, TSet<FMaterialExpressionKey>& VisitedExpressions, const bool bStopOnFunctionCall);
 #endif // WITH_EDITOR
 };
 
diff --git a/Engine/Source/Runtime/Engine/Private/Materials/MaterialExpressions.cpp b/Engine/Source/Runtime/Engine/Private/Materials/MaterialExpressions.cpp
--- a/Engine/Source/Runtime/Engine/Private/Materials/MaterialExpressions.cpp
+++ b/Engine/Source/Runtime/Engine/Private/Materials/MaterialExpressions.cpp
@@ -2271,6 +2271,47 @@ bool UMaterialExpression::ContainsInputLoopInternal(const FContainsInputLoopInte
 	return false;
 }
 
+// Deprecated
+bool UMaterialExpression::ContainsInputLoopInternal(TArray<FMaterialExpressionKey>& ExpressionStack, TSet<FMaterialExpressionKey>& VisitedExpressions, const bool bStopOnFunctionCall)
+{
+	for (FExpressionInput* Input : GetInputsView())
+	{
+		if (!Input->Expression)
+		{
+			continue;
+		}
+		
+		// ContainsInputLoop primarily used to detect safe traversal path for IsResultMaterialAttributes.
+		// In those cases we can bail on a function as the inputs are strongly typed
+		UMaterialExpressionMaterialFunctionCall* FunctionCall = Cast<UMaterialExpressionMaterialFunctionCall>(Input->Expression);
+		UMaterialExpressionMaterialAttributeLayers* Layers = Cast<UMaterialExpressionMaterialAttributeLayers>(Input->Expression);
+		if (bStopOnFunctionCall && (FunctionCall || Layers))
+		{
+			continue;
+		}
+
+		FMaterialExpressionKey InputExpressionKey(Input->Expression, Input->OutputIndex);
+		if (ExpressionStack.Contains(InputExpressionKey))
+		{
+			return true;
+		}
+
+		// prevent recurring visits to expressions we've already checked
+		if (!VisitedExpressions.Contains(InputExpressionKey))
+		{
+			VisitedExpressions.Add(InputExpressionKey);
+			ExpressionStack.Add(InputExpressionKey);
+			if (Input->Expression->ContainsInputLoopInternal(ExpressionStack, VisitedExpressions, bStopOnFunctionCall))
+			{
+				return true;
+			}
+			ExpressionStack.Pop();
+		}
+	}
+
+	return false;
+}
+
 bool UMaterialExpression::IsUsingNewHLSLGenerator() const
 {
 	if (Material)
