Fixed a crash in FMirrorSyncScope

diff --git a/Engine/Source/Runtime/Engine/Private/Animation/AnimNodeMessages.cpp b/Engine/Source/Runtime/Engine/Private/Animation/AnimNodeMessages.cpp
--- a/Engine/Source/Runtime/Engine/Private/Animation/AnimNodeMessages.cpp
+++ b/Engine/Source/Runtime/Engine/Private/Animation/AnimNodeMessages.cpp
@@ -116,6 +116,22 @@ void FMessageStack::TopMessage(FName InMessageType, TFunctionRef<void(IGraphMess
 	}
 }
 
+void FMessageStack::TopMessageSharedPtr(FName InMessageType, TFunctionRef<void(TSharedPtr<IGraphMessage>&)> InFunction)
+{
+	if (FMessageStackEntry* StackEntryPtr = MessageStacks.Find(InMessageType))
+	{
+		FMessageStackEntry& StackEntry = *StackEntryPtr;
+		if (StackEntry.Num())
+		{
+			FMessageEntry& MessageEntry = StackEntry.Top();
+			if (MessageEntry.Message->GetTypeName() == InMessageType)
+			{
+				InFunction(MessageEntry.Message);
+			}
+		}
+	}
+}
+
 void FMessageStack::BottomMessage(FName InMessageType, TFunctionRef<void(IGraphMessage&)> InFunction) const
 {
 	if(const FMessageStackEntry* StackEntryPtr = MessageStacks.Find(InMessageType))
diff --git a/Engine/Source/Runtime/Engine/Private/Animation/MirrorSyncScope.cpp b/Engine/Source/Runtime/Engine/Private/Animation/MirrorSyncScope.cpp
--- a/Engine/Source/Runtime/Engine/Private/Animation/MirrorSyncScope.cpp
+++ b/Engine/Source/Runtime/Engine/Private/Animation/MirrorSyncScope.cpp
@@ -39,30 +39,39 @@ FMirrorSyncScope::FMirrorSyncScope(const FAnimationBaseContext& InContext, const
 	{
 		// Syncing is done when animation sequences are evaluated, which means we need to examine the stack of mirror
 		// nodes to determine if we can simply enable and disable sync mirroring.  
-		OuterScopeMirrorDataTable = ParentMirrorScope->MirrorDataTable;
 		MirrorScopeDepth = ParentMirrorScope->GetMirrorScopeDepth() + 1;
+		OuterScopeMirrorDataTable = ParentMirrorScope->MirrorDataTable;
 		if (MirrorScopeDepth % 2 == 0)
 		{
 			bClearMirrorSync = IsMirrorSyncIdentical(ParentMirrorScope->MirrorDataTable, MirrorDataTable);
 		}
 	}
-		 
-	AnimSyncGroupScope = InContext.GetMessage<FAnimSyncGroupScope>();
-	if(AnimSyncGroupScope)
+
+	if (InContext.SharedContext)
+	{
+		TWeakPtr<FAnimSyncGroupScope> SyncScope;
+		InContext.SharedContext->MessageStack.TopMessageWeakPtr<FAnimSyncGroupScope>([&SyncScope](TWeakPtr<FAnimSyncGroupScope>& InMessage)
+			{
+				SyncScope = InMessage;
+			});
+		AnimSyncGroupScope = SyncScope; 
+	}
+
+	if(TSharedPtr<FAnimSyncGroupScope> PinnedSyncGroupScope = AnimSyncGroupScope.Pin())
 	{
 		if (bClearMirrorSync)
 		{
 			MirrorDataTable = nullptr;
 		}
-		AnimSyncGroupScope->SetMirror(MirrorDataTable);
+		PinnedSyncGroupScope.Get()->SetMirror(MirrorDataTable);
 	}
 }
 
 FMirrorSyncScope::~FMirrorSyncScope()
 {
-	if(AnimSyncGroupScope)
+	if (TSharedPtr<FAnimSyncGroupScope> PinnedSyncGroupScope = AnimSyncGroupScope.Pin())
 	{
-		AnimSyncGroupScope->SetMirror(OuterScopeMirrorDataTable);
+		PinnedSyncGroupScope->SetMirror(OuterScopeMirrorDataTable);
 	}
 }
 
diff --git a/Engine/Source/Runtime/Engine/Public/Animation/AnimNodeMessages.h b/Engine/Source/Runtime/Engine/Public/Animation/AnimNodeMessages.h
--- a/Engine/Source/Runtime/Engine/Public/Animation/AnimNodeMessages.h
+++ b/Engine/Source/Runtime/Engine/Public/Animation/AnimNodeMessages.h
@@ -257,6 +257,22 @@ public:
 		});
 	}
 
+	// Call the supplied function on the top-most node subscribed to the specified message in this stack
+	// @param	InFunction		The function to call with the specified message.
+	template<typename TGraphMessageType>
+	void TopMessageWeakPtr(TFunctionRef<void(TWeakPtr<TGraphMessageType>&)> InFunction)
+	{
+		static_assert(TIsDerivedFrom<TGraphMessageType, IGraphMessage>::IsDerived, "Argument TGraphMessageType must derive from IGraphMessage");
+		static_assert(!std::is_same_v<TGraphMessageType, IGraphMessage>, "Argument TGraphMessageType must not be IGraphMessage");
+
+		TopMessageSharedPtr(TGraphMessageType::GetStaticTypeName(), [&InFunction](TSharedPtr<IGraphMessage>& InMessage)
+			{
+				TSharedPtr<TGraphMessageType> PinPtr = StaticCastSharedPtr<TGraphMessageType>(InMessage);
+				TWeakPtr<TGraphMessageType> WeakPtr(PinPtr); 
+				InFunction(WeakPtr);
+			});
+	}
+
 	// @return true if a message of the specified type is present
 	template<typename TGraphMessageType>
 	bool HasMessage() const
@@ -295,6 +311,9 @@ private:
 	// Helper function for the templated function of the same name above
 	void TopMessage(FName InMessageType, TFunctionRef<void(IGraphMessage&)> InFunction) const;
 
+	// Helper function for the templated function of the same name above
+	void TopMessageSharedPtr(FName InMessageType, TFunctionRef<void(TSharedPtr<IGraphMessage>&)> InFunction);
+
 	// Helper function for the templated function of the same name above
 	void BottomMessage(FName InMessageType, TFunctionRef<void(IGraphMessage&)> InFunction) const;
 
diff --git a/Engine/Source/Runtime/Engine/Public/Animation/MirrorSyncScope.h b/Engine/Source/Runtime/Engine/Public/Animation/MirrorSyncScope.h
--- a/Engine/Source/Runtime/Engine/Public/Animation/MirrorSyncScope.h
+++ b/Engine/Source/Runtime/Engine/Public/Animation/MirrorSyncScope.h
@@ -32,9 +32,9 @@ public:
 	virtual TUniquePtr<const IAnimNotifyEventContextDataInterface> MakeUniqueEventContextData() const override;
 private:
 	const UMirrorDataTable* MirrorDataTable = nullptr;
-	int32 MirrorScopeDepth = 1; 
+	int32 MirrorScopeDepth = 1;
 	const UMirrorDataTable* OuterScopeMirrorDataTable = nullptr;
-	FAnimSyncGroupScope* AnimSyncGroupScope = nullptr; 
+	TWeakPtr<FAnimSyncGroupScope> AnimSyncGroupScope;
 };
 
 }}	// namespace UE::Anim
