Lowering the Net Update Frequency for HLOD actors in case we ever need to set them to replicated.

#[review]-25877431 [REDACTED_HANDLE], [REDACTED_HANDLE], [REDACTED_HANDLE], [REDACTED_HANDLE], [REDACTED_HANDLE], [REDACTED_HANDLE], [REDACTED_HANDLE], [REDACTED_HANDLE]

diff --git a/Engine/Source/Runtime/Engine/Private/WorldPartition/HLOD/HLODActor.cpp b/Engine/Source/Runtime/Engine/Private/WorldPartition/HLOD/HLODActor.cpp
--- a/Engine/Source/Runtime/Engine/Private/WorldPartition/HLOD/HLODActor.cpp
+++ b/Engine/Source/Runtime/Engine/Private/WorldPartition/HLOD/HLODActor.cpp
@@ -52,6 +52,7 @@ AWorldPartitionHLOD::AWorldPartitionHLOD(const FObjectInitializer& ObjectInitial
 	bReplicates = true;
 
 	NetDormancy = DORM_DormantAll;
+	NetUpdateFrequency = 1.f;
 
 #if WITH_EDITORONLY_DATA
 	HLODHash = 0;
