Fix editor performance issue caused by navigation building even if there is nowhere to build navigation.
The fix expect that there should be no difference in simulation mode.

diff --git a/Engine/Source/Runtime/Engine/Classes/AI/NavigationSystemBase.h b/Engine/Source/Runtime/Engine/Classes/AI/NavigationSystemBase.h
--- a/Engine/Source/Runtime/Engine/Classes/AI/NavigationSystemBase.h
+++ b/Engine/Source/Runtime/Engine/Classes/AI/NavigationSystemBase.h
@@ -32,7 +32,7 @@ enum class FNavigationSystemRunMode : uint8
 	InvalidMode,
 	GameMode,
 	EditorMode,
-	SimulationMode,
+	SimulationMode UE_DEPRECATED(5.3, "No longer used.") UMETA(Hidden, DisplayName = "DEPRECATED No longer used, to be removed soon"),
 	PIEMode,
 	InferFromWorldMode,
 	EditorWorldPartitionBuildMode
diff --git a/Engine/Source/Runtime/Engine/Private/GameInstance.cpp b/Engine/Source/Runtime/Engine/Private/GameInstance.cpp
--- a/Engine/Source/Runtime/Engine/Private/GameInstance.cpp
+++ b/Engine/Source/Runtime/Engine/Private/GameInstance.cpp
@@ -513,7 +513,7 @@ FGameInstancePIEResult UGameInstance::StartPlayInEditorGameInstance(ULocalPlayer
 		SlowTask.EnterProgressFrame(10, NSLOCTEXT("UnrealEd", "PIEInitializingActors", "Starting PIE (Initializing Actors)..."));
 		PlayWorld->InitializeActorsForPlay(URL);
 		// calling it after InitializeActorsForPlay has been called to have all potential bounding boxed initialized
-		FNavigationSystem::AddNavigationSystemToWorld(*PlayWorld, LocalPlayers.Num() > 0 ? FNavigationSystemRunMode::PIEMode : FNavigationSystemRunMode::SimulationMode);
+		FNavigationSystem::AddNavigationSystemToWorld(*PlayWorld, FNavigationSystemRunMode::PIEMode);
 
 		// @todo, just use WorldContext.GamePlayer[0]?
 		if (LocalPlayer)
diff --git a/Engine/Source/Runtime/NavigationSystem/Private/NavigationSystem.cpp b/Engine/Source/Runtime/NavigationSystem/Private/NavigationSystem.cpp
--- a/Engine/Source/Runtime/NavigationSystem/Private/NavigationSystem.cpp
+++ b/Engine/Source/Runtime/NavigationSystem/Private/NavigationSystem.cpp
@@ -987,9 +987,7 @@ void UNavigationSystemV1::OnWorldInitDone(FNavigationSystemRunMode Mode)
 	// those links were serialized-in or spawned)
 	ProcessCustomLinkPendingRegistration();
 
-	if (IsThereAnywhereToBuildNavigation() == false
-		// Simulation mode is a special case - better not do it in this case
-		&& OperationMode != FNavigationSystemRunMode::SimulationMode)
+	if (IsThereAnywhereToBuildNavigation() == false)
 	{
 		// remove all navigation data instances
 		for (TActorIterator<ANavigationData> It(World); It; ++It)
