Remove User.ID from SessionStats and SessionMachineStats.

#rnx

diff --git a/Engine/Source/Runtime/Engine/Private/EngineAnalytics.cpp b/Engine/Source/Runtime/Engine/Private/EngineAnalytics.cpp
--- a/Engine/Source/Runtime/Engine/Private/EngineAnalytics.cpp
+++ b/Engine/Source/Runtime/Engine/Private/EngineAnalytics.cpp
@@ -238,7 +238,6 @@ void FEngineAnalytics::LowDriveSpaceDetected()
 
 void FEngineAnalytics::AppendMachineStats(TArray<FAnalyticsEventAttribute>& EventAttributes)
 {
-	const FString UserID = FPlatformProcess::UserName(false);
 	const UGeneralProjectSettings& ProjectSettings = *GetDefault<UGeneralProjectSettings>();
 	const FPlatformMemoryStats Stats = FPlatformMemory::GetStats();
 
@@ -257,7 +256,6 @@ void FEngineAnalytics::AppendMachineStats(TArray<FAnalyticsEventAttribute>& Even
 	EventAttributes.Emplace(TEXT("ProjectDescription"), ProjectSettings.Description);
 	EventAttributes.Emplace(TEXT("ProjectVersion"), ProjectSettings.ProjectVersion);
 	EventAttributes.Emplace(TEXT("Application.Commandline"), FCommandLine::Get());
-	EventAttributes.Emplace(TEXT("User.ID"), UserID);
 	EventAttributes.Emplace(TEXT("Build.Configuration"), LexToString(FApp::GetBuildConfiguration()));
 	EventAttributes.Emplace(TEXT("Build.IsInternalBuild"), FEngineBuildSettings::IsInternalBuild());
 	EventAttributes.Emplace(TEXT("Build.IsPerforceBuild"), FEngineBuildSettings::IsPerforceBuild());
