Gracefully handle seeking to a large time value (>32bits)

diff --git a/Engine/Source/Runtime/Engine/Private/AudioDecompress.cpp b/Engine/Source/Runtime/Engine/Private/AudioDecompress.cpp
--- a/Engine/Source/Runtime/Engine/Private/AudioDecompress.cpp
+++ b/Engine/Source/Runtime/Engine/Private/AudioDecompress.cpp
@@ -127,14 +127,19 @@ void IStreamedCompressedInfo::SeekToTime(const float InSeekTimeSeconds)
 		if (WaveData.IsValid())
 		{
 			// Negative time will seek to start.
-			uint32 SeekTimeAudioFrames = 0;
+			int64 SeekTimeAudioFrames = 0;
 			if (InSeekTimeSeconds > 0)
 			{
-				SeekTimeAudioFrames = IntCastChecked<uint32>((uint64)(WaveData->GetSampleRate() * InSeekTimeSeconds));
+				SeekTimeAudioFrames = FMath::FloorToInt64(WaveData->GetSampleRate() * InSeekTimeSeconds);
+				if (!IntFitsIn<uint32>(SeekTimeAudioFrames))
+				{
+					UE_LOG(LogAudio, Warning, TEXT("Seek too large (%.2f seconds), ignoring..."), InSeekTimeSeconds);
+					return;
+				}
 			}
 
 			// If we have a chunk setup to contain a seek-table it will return a value other than INDEX_NONE here.
-			const int32 ChunkIndexToSeekTo = WaveData->FindChunkIndexForSeeking(SeekTimeAudioFrames);
+			const int32 ChunkIndexToSeekTo = WaveData->FindChunkIndexForSeeking(IntCastChecked<uint32>(SeekTimeAudioFrames));
 			if (ChunkIndexToSeekTo >= 0)
 			{
 				StreamSeekBlockIndex = ChunkIndexToSeekTo;
