Substrate - added in material editor complexity on Slab node.

diff --git a/Engine/Source/Runtime/Engine/Classes/Materials/MaterialExpressionStrata.h b/Engine/Source/Runtime/Engine/Classes/Materials/MaterialExpressionStrata.h
--- a/Engine/Source/Runtime/Engine/Classes/Materials/MaterialExpressionStrata.h
+++ b/Engine/Source/Runtime/Engine/Classes/Materials/MaterialExpressionStrata.h
@@ -330,6 +330,7 @@ class UMaterialExpressionStrataSlabBSDF : public UMaterialExpressionStrataBSDF
 	virtual FStrataOperator* StrataGenerateMaterialTopologyTree(class FMaterialCompiler* Compiler, class UMaterialExpression* Parent, int32 OutputIndex) override;
 	virtual FName GetInputName(int32 InputIndex) const override;
 	virtual void GetConnectorToolTip(int32 InputIndex, int32 OutputIndex, TArray<FString>& OutToolTip) override;
+	virtual void GetExpressionToolTip(TArray<FString>& OutToolTip) override;
 
 	bool HasEdgeColor() const;
 	bool HasFuzz() const;
@@ -341,6 +342,15 @@ class UMaterialExpressionStrataSlabBSDF : public UMaterialExpressionStrataBSDF
 	bool HasAnisotropy() const;
 	bool HasGlint() const;
 	bool HasSpecularProfile() const;
+
+	struct FComplexity
+	{
+		bool bStrataMaterialIsComplexSpecial;
+		bool bStrataMaterialIsComplex;
+		bool bStrataMaterialIsSingle;
+		// If all the above are false, complexity is Simple
+	};
+	FComplexity GetComplexity() const;
 #endif
 	//~ End UMaterialExpression Interface
 };
diff --git a/Engine/Source/Runtime/Engine/Private/Materials/MaterialExpressions.cpp b/Engine/Source/Runtime/Engine/Private/Materials/MaterialExpressions.cpp
--- a/Engine/Source/Runtime/Engine/Private/Materials/MaterialExpressions.cpp
+++ b/Engine/Source/Runtime/Engine/Private/Materials/MaterialExpressions.cpp
@@ -23350,9 +23350,53 @@ int32 UMaterialExpressionStrataSlabBSDF::Compile(class FMaterialCompiler* Compil
 	return OutputCodeChunk;
 }
 
+UMaterialExpressionStrataSlabBSDF::FComplexity UMaterialExpressionStrataSlabBSDF::GetComplexity() const
+{
+	UMaterialExpressionStrataSlabBSDF::FComplexity Complexity;
+	Complexity.bStrataMaterialIsComplexSpecial = HasGlint() || HasSpecularProfile();
+	Complexity.bStrataMaterialIsComplex = HasAnisotropy();
+	Complexity.bStrataMaterialIsSingle = HasEdgeColor() || HasFuzz() || HasSecondRoughness() || HasMFPPluggedIn() || HasSSS();
+
+	// MAsk out to only have a single possibility
+	Complexity.bStrataMaterialIsComplex &= !Complexity.bStrataMaterialIsComplexSpecial;
+	Complexity.bStrataMaterialIsSingle &= !Complexity.bStrataMaterialIsComplexSpecial && !Complexity.bStrataMaterialIsComplex;
+
+	return Complexity;
+}
+
+static FString GetSlabComplexityString(UMaterialExpressionStrataSlabBSDF::FComplexity Complexity)
+{
+	FString ComplexityString = TEXT("Simple");
+	if (Complexity.bStrataMaterialIsComplexSpecial)
+	{
+		ComplexityString = TEXT("ComplexSpecial");
+	}
+	else if (Complexity.bStrataMaterialIsComplex)
+	{
+		ComplexityString = TEXT("Complex");
+	}
+	else if (Complexity.bStrataMaterialIsSingle)
+	{
+		ComplexityString = TEXT("Single");
+	}
+	return ComplexityString;
+}
+
 void UMaterialExpressionStrataSlabBSDF::GetCaption(TArray<FString>& OutCaptions) const
 {
-	OutCaptions.Add(TEXT("Substrate Slab BSDF"));
+	// The node complexity is manually maintained to match FStrataCompilationContext::StrataGenerateDerivedMaterialOperatorData and shaders.
+	OutCaptions.Add(TEXT("Substrate Slab BSDF - ") + GetSlabComplexityString(GetComplexity()));
+}
+
+void UMaterialExpressionStrataSlabBSDF::GetExpressionToolTip(TArray<FString>& OutToolTip)
+{
+	OutToolTip.Add(TEXT("Substrate Slab BSDF"));
+	OutToolTip.Add(TEXT("Complexity = ") + GetSlabComplexityString(GetComplexity()));
+	OutToolTip.Add(TEXT("The complexity represents the cost of the shading path (Lighting, Lumen, SSS) the material will follow:"));
+	OutToolTip.Add(TEXT(" - Simple means the Slab only relies on Diffuse, F0 and Roughness. It will follow a fast shading path."));
+	OutToolTip.Add(TEXT(" - Single means the Slab uses more features such as F90, Fuzz, Second Roughness, MFP or SSS. It will follow a more expenssive shading path."));
+	OutToolTip.Add(TEXT(" - Complex means a Slab uses anisotropic lighting, with any of the previous features."));
+	OutToolTip.Add(TEXT(" - ComplexSpecial means the Slab is using more advanced features such as glints or specular LUT. This is the most expenssive shading path."));
 }
 
 uint32 UMaterialExpressionStrataSlabBSDF::GetOutputType(int32 OutputIndex)
