HeterogeneousVolumeComponent: Moving MaterialInstanceDynamic creation to SceneProxy.

diff --git a/Engine/Source/Runtime/Engine/Classes/Components/HeterogeneousVolumeComponent.h b/Engine/Source/Runtime/Engine/Classes/Components/HeterogeneousVolumeComponent.h
--- a/Engine/Source/Runtime/Engine/Classes/Components/HeterogeneousVolumeComponent.h
+++ b/Engine/Source/Runtime/Engine/Classes/Components/HeterogeneousVolumeComponent.h
@@ -30,6 +30,9 @@ class UHeterogeneousVolumeComponent : public UMeshComponent
 
 	UPROPERTY(EditAnywhere, BlueprintReadOnly, Category = Lighting)
 	float LightingDownsampleFactor;
+	
+	UPROPERTY(Transient)
+	TObjectPtr<UMaterialInstanceDynamic> MaterialInstanceDynamic;
 
 	~UHeterogeneousVolumeComponent() {}
 
@@ -38,6 +41,7 @@ public:
 	virtual FBoxSphereBounds CalcBounds(const FTransform& LocalToWorld) const override;
 	virtual void TickComponent(float DeltaTime, ELevelTick TickType, FActorComponentTickFunction* ThisTickFunction) override;
 
+	virtual void PostLoad() override;
 #if WITH_EDITOR
 	virtual void PostEditChangeProperty(FPropertyChangedEvent& PropertyChangedEvent) override;
 #endif // WITH_EDITOR
@@ -48,6 +52,7 @@ public:
 	virtual FPrimitiveSceneProxy* CreateSceneProxy() override;
 	virtual int32 GetNumMaterials() const override { return 1; }
 	virtual bool ShouldRenderSelected() const override { return true; }
+	virtual void GetUsedMaterials(TArray<UMaterialInterface*>& OutMaterials, bool bGetDebugMaterials = false) const override;
 	//~ End UPrimitiveComponent Interface.
 
 private:
diff --git a/Engine/Source/Runtime/Engine/Private/Components/HeterogeneousVolumeComponent.cpp b/Engine/Source/Runtime/Engine/Private/Components/HeterogeneousVolumeComponent.cpp
--- a/Engine/Source/Runtime/Engine/Private/Components/HeterogeneousVolumeComponent.cpp
+++ b/Engine/Source/Runtime/Engine/Private/Components/HeterogeneousVolumeComponent.cpp
@@ -78,6 +78,11 @@ FHeterogeneousVolumeSceneProxy::FHeterogeneousVolumeSceneProxy(UHeterogeneousVol
 	VoxelSize.Z /= InComponent->VolumeResolution.Z;
 	HeterogeneousVolumeData.MinimumVoxelSize = FMath::Max(VoxelSize.GetMin(), 0.001);
 
+	if (InComponent->bAnimate && InComponent->MaterialInstanceDynamic)
+	{
+		MaterialInterface = InComponent->MaterialInstanceDynamic;
+	}
+
 	HeterogeneousVolumeData.LightingDownsampleFactor = InComponent->LightingDownsampleFactor;
 
 	// Initialize vertex buffer data for a quad
@@ -204,6 +209,7 @@ UHeterogeneousVolumeComponent::UHeterogeneousVolumeComponent(const FObjectInitia
 	bTickInEditor = true;
 #endif // WITH_EDITORONLY_DATA
 
+	MaterialInstanceDynamic = nullptr;
 	VolumeResolution = FIntVector(128);
 	Framerate = 24.0f;
 	bAnimate = false;
@@ -226,6 +232,24 @@ FBoxSphereBounds UHeterogeneousVolumeComponent::CalcBounds(const FTransform& Loc
 	return NewBounds.TransformBy(LocalToWorld);
 }
 
+void UHeterogeneousVolumeComponent::PostLoad()
+{
+	Super::PostLoad();
+
+	UMaterialInterface* MaterialInterface = GetMaterial(0);
+	if (MaterialInterface)
+	{
+		const UMaterial* Material = MaterialInterface->GetMaterial();
+		if (Material && Material->MaterialDomain == EMaterialDomain::MD_Volume)
+		{
+			Material->GetRenderProxy();
+			MaterialInterface->CheckMaterialUsage(MATUSAGE_HeterogeneousVolumes);
+		}
+
+		MaterialInstanceDynamic = UMaterialInstanceDynamic::Create(MaterialInterface, nullptr);
+	}
+}
+
 #if WITH_EDITOR
 void UHeterogeneousVolumeComponent::PostEditChangeProperty(FPropertyChangedEvent& PropertyChangedEvent)
 {
@@ -248,17 +272,29 @@ void UHeterogeneousVolumeComponent::PostEditChangeProperty(FPropertyChangedEvent
 				Material->GetRenderProxy();
 				MaterialInterface->CheckMaterialUsage(MATUSAGE_HeterogeneousVolumes);
 			}
+
+			MaterialInstanceDynamic = UMaterialInstanceDynamic::Create(MaterialInterface, nullptr);
 		}
 	}
 }
 #endif // WITH_EDITOR
 
+void UHeterogeneousVolumeComponent::GetUsedMaterials(TArray<UMaterialInterface*>& OutMaterials, bool bGetDebugMaterials) const
+{
+	UMeshComponent::GetUsedMaterials(OutMaterials, bGetDebugMaterials);
+	
+	if (MaterialInstanceDynamic)
+	{
+		OutMaterials.Add(MaterialInstanceDynamic);
+	}
+}
+
 void UHeterogeneousVolumeComponent::TickComponent(float DeltaTime, ELevelTick TickType, FActorComponentTickFunction* ThisTickFunction)
 {
 	Super::TickComponent(DeltaTime, TickType, ThisTickFunction);
 
 	uint32 MaterialIndex = 0;
-	UMaterialInterface* Material = GetMaterial(MaterialIndex);
+	UMaterialInterface* Material = UMeshComponent::GetMaterial(MaterialIndex);
 	if (Material)
 	{
 		// Get all SVT params in the material
@@ -268,38 +304,20 @@ void UHeterogeneousVolumeComponent::TickComponent(float DeltaTime, ELevelTick Ti
 
 		if (!SVTParameterInfo.IsEmpty())
 		{
-			// Create a MID if this isn't one already
-			UMaterialInstanceDynamic* MaterialInstanceDynamic = nullptr;
-			UMaterialInterface* ParentMaterial = nullptr;
-			if (!Material->IsA<UMaterialInstanceDynamic>())
-			{
-				MaterialInstanceDynamic = CreateAndSetMaterialInstanceDynamicFromMaterial(MaterialIndex, Material);
-				ParentMaterial = Material;
-			}
-			else
-			{
-				MaterialInstanceDynamic = CastChecked<UMaterialInstanceDynamic>(Material);
-				ParentMaterial = MaterialInstanceDynamic->Parent;
-			}
-
 			USparseVolumeTexture* DefaultSVT = nullptr;
-			MaterialInstanceDynamic->GetSparseVolumeTextureParameterDefaultValue(SVTParameterInfo[MaterialIndex], DefaultSVT);
-
+			bool bValid = Material->GetSparseVolumeTextureParameterDefaultValue(SVTParameterInfo[MaterialIndex], DefaultSVT);
 			if (DefaultSVT)
 			{
-				if (bAnimate)
+				const int32 FrameCount = DefaultSVT->GetNumFrames();
+				const float FrameIndexF = bAnimate ? FMath::Fmod(Time * Framerate, (float)FrameCount) : 0;
+				const int32 MipLevel = 0;
+				bool bIsBlocking = false;
+				USparseVolumeTextureFrame* SVTFrame = USparseVolumeTextureFrame::GetFrameAndIssueStreamingRequest(DefaultSVT, FrameIndexF, MipLevel, bIsBlocking);
+				VolumeResolution = SVTFrame ? SVTFrame->GetVolumeResolution() : FIntVector(1);
+
+				if (bAnimate && MaterialInstanceDynamic)
 				{
-					const int32 FrameCount = DefaultSVT->GetNumFrames();
-					const float FrameIndexF = FMath::Fmod(Time * Framerate, (float)FrameCount);
-					const int32 MipLevel = 0;
-					USparseVolumeTextureFrame* SVTFrame = USparseVolumeTextureFrame::GetFrameAndIssueStreamingRequest(DefaultSVT, FrameIndexF, MipLevel);
 					MaterialInstanceDynamic->SetSparseVolumeTextureParameterValue(SVTParameterInfo[MaterialIndex].Name, SVTFrame);
-
-					VolumeResolution = SVTFrame ? SVTFrame->GetVolumeResolution() : FIntVector(1);
-				}
-				else
-				{
-					VolumeResolution = DefaultSVT->GetVolumeResolution();
 				}
 			}
 		}
