Unity build compilation fix.

diff --git a/Engine/Source/Runtime/Engine/Private/Components/PrimitiveComponent.cpp b/Engine/Source/Runtime/Engine/Private/Components/PrimitiveComponent.cpp
--- a/Engine/Source/Runtime/Engine/Private/Components/PrimitiveComponent.cpp
+++ b/Engine/Source/Runtime/Engine/Private/Components/PrimitiveComponent.cpp
@@ -44,6 +44,7 @@
 #include "UObject/UE5PrivateFrostyStreamObjectVersion.h"
 #include "UObject/ObjectSaveContext.h"
 #include "Engine/DamageEvents.h"
+#include "MeshUVChannelInfo.h"
 
 #if WITH_EDITOR
 #include "Engine/LODActor.h"
