Allow grid size to be included in Partition Actor label (disabled by default)

#rnx

diff --git a/Engine/Source/Runtime/Engine/Private/ActorPartition/ActorPartitionSubsystem.cpp b/Engine/Source/Runtime/Engine/Private/ActorPartition/ActorPartitionSubsystem.cpp
--- a/Engine/Source/Runtime/Engine/Private/ActorPartition/ActorPartitionSubsystem.cpp
+++ b/Engine/Source/Runtime/Engine/Private/ActorPartition/ActorPartitionSubsystem.cpp
@@ -272,7 +272,12 @@ public:
 
 			// Once actor is created, update its label
 			TStringBuilderWithBuffer<TCHAR, NAME_SIZE> ActorLabelBuilder;
-			ActorLabelBuilder += FString::Printf(TEXT("%s_%d_%d_%d"), *InActorPartitionId.GetClass()->GetName(), InCellCoord.X, InCellCoord.Y, InCellCoord.Z);
+			ActorLabelBuilder += FString::Printf(TEXT("%s"), *InActorPartitionId.GetClass()->GetName());
+			if (FoundActor->ShouldIncludeGridSizeInLabel())
+			{
+				ActorLabelBuilder += FString::Printf(TEXT("_%u"), InGridSize);
+			}
+			ActorLabelBuilder += FString::Printf(TEXT("_%d_%d_%d"), InCellCoord.X, InCellCoord.Y, InCellCoord.Z);
 			if (InActorPartitionId.GetDataLayerEditorContextHash() != FDataLayerEditorContext::EmptyHash)
 			{
 				ActorLabelBuilder += FString::Printf(TEXT("_%X"), InActorPartitionId.GetDataLayerEditorContextHash());
diff --git a/Engine/Source/Runtime/Engine/Public/ActorPartition/PartitionActor.h b/Engine/Source/Runtime/Engine/Public/ActorPartition/PartitionActor.h
--- a/Engine/Source/Runtime/Engine/Public/ActorPartition/PartitionActor.h
+++ b/Engine/Source/Runtime/Engine/Public/ActorPartition/PartitionActor.h
@@ -26,6 +26,7 @@ public:
 	ENGINE_API virtual TUniquePtr<class FWorldPartitionActorDesc> CreateClassActorDesc() const override;
 	ENGINE_API virtual uint32 GetDefaultGridSize(UWorld* InWorld) const PURE_VIRTUAL(APartitionActor, return 0;)
 	ENGINE_API virtual bool ShouldIncludeGridSizeInName(UWorld * InWorld, const FActorPartitionIdentifier& InIdentifier) const;
+	ENGINE_API virtual bool ShouldIncludeGridSizeInLabel() const { return false; }
 	virtual FGuid GetGridGuid() const { return FGuid(); }
 	ENGINE_API virtual bool IsUserManaged() const override;
 #endif
