LevelInstance: Deprecate CanBeUsedByLevelInstance.
- Runtime: Level Instances that are WP will still get embedded in their parent world grid (independent of the value of bEnableStreaming)
- Editor: Level Instances with bEnableStreaming=true will now act as regular World Partitions when loaded on their own. Allowing for big Level Instances to be partially loaded outside of their owning world.
- Add Project Setting for default value to assign to new Level Instances

diff --git a/Engine/Source/Editor/LevelInstanceEditor/Private/LevelInstanceEditorModule.cpp b/Engine/Source/Editor/LevelInstanceEditor/Private/LevelInstanceEditorModule.cpp
--- a/Engine/Source/Editor/LevelInstanceEditor/Private/LevelInstanceEditorModule.cpp
+++ b/Engine/Source/Editor/LevelInstanceEditor/Private/LevelInstanceEditorModule.cpp
@@ -340,6 +340,7 @@ namespace LevelInstanceMenuUtils
 				
 				CreationParams.TemplateWorld = TemplatePackage ? UWorld::FindWorldInPackage(TemplatePackage) : nullptr;
 				CreationParams.LevelInstanceClass = GetDefaultLevelInstanceClass(CreationType);
+				CreationParams.bEnableStreaming = LevelInstanceEditorSettings->bEnableStreaming;
 
 				if (!LevelInstanceSubsystem->CreateLevelInstanceFrom(ActorsToMove, CreationParams))
 				{
@@ -617,7 +618,6 @@ namespace LevelInstanceMenuUtils
 		FWorldPartitionConverter::FParameters Parameters;
 		Parameters.bConvertSubLevels = false;
 		Parameters.bEnableStreaming = false;
-		Parameters.bCanBeUsedByLevelInstance = true;
 		Parameters.bUseActorFolders = true;
 
 		if (FWorldPartitionConverter::Convert(WorldAsset, Parameters))
@@ -658,7 +658,7 @@ namespace LevelInstanceMenuUtils
 	void AddPartitionedStreamingSupportFromMenu(UToolMenu* Menu, FAssetData WorldAsset)
 	{
 		FName WorldAssetName = WorldAsset.PackageName;
-		if (!ULevel::GetIsLevelPartitionedFromPackage(WorldAssetName) || !ULevel::GetPartitionedLevelCanBeUsedByLevelInstanceFromPackage(WorldAssetName))
+		if (!ULevel::GetIsLevelPartitionedFromPackage(WorldAssetName))
 		{
 			FToolMenuSection& Section = CreateLevelSection(Menu);
 			FToolUIAction UIAction;
diff --git a/Engine/Source/Editor/LevelInstanceEditor/Private/LevelInstanceEditorSettings.cpp b/Engine/Source/Editor/LevelInstanceEditor/Private/LevelInstanceEditorSettings.cpp
--- a/Engine/Source/Editor/LevelInstanceEditor/Private/LevelInstanceEditorSettings.cpp
+++ b/Engine/Source/Editor/LevelInstanceEditor/Private/LevelInstanceEditorSettings.cpp
@@ -6,6 +6,7 @@
 ULevelInstanceEditorSettings::ULevelInstanceEditorSettings()
 {
 	LevelInstanceClassName = ALevelInstance::StaticClass()->GetPathName();
+	bEnableStreaming = false;
 }
 
 ULevelInstanceEditorPerProjectUserSettings::ULevelInstanceEditorPerProjectUserSettings()
diff --git a/Engine/Source/Editor/LevelInstanceEditor/Public/LevelInstanceEditorSettings.h b/Engine/Source/Editor/LevelInstanceEditor/Public/LevelInstanceEditorSettings.h
--- a/Engine/Source/Editor/LevelInstanceEditor/Public/LevelInstanceEditorSettings.h
+++ b/Engine/Source/Editor/LevelInstanceEditor/Public/LevelInstanceEditorSettings.h
@@ -10,8 +10,8 @@
 #include "LevelInstance/LevelInstanceTypes.h"
 #include "LevelInstanceEditorSettings.generated.h"
 
-UCLASS(config = Editor)
-class ULevelInstanceEditorSettings : public UObject
+UCLASS(config = Editor, defaultconfig, meta = (DisplayName = "Level Instance"))
+class ULevelInstanceEditorSettings : public UDeveloperSettings
 {
 	GENERATED_BODY()
 
@@ -24,6 +24,9 @@ public:
 
 	UPROPERTY(config)
 	FString LevelInstanceClassName;
+
+	UPROPERTY(config, EditAnywhere, Category="World Partition", meta = (ToolTip="Create World Partition Level Instances with Streaming Enabled/Disabled by default"))
+	bool bEnableStreaming;
 };
 
 UCLASS(config = EditorPerProjectUserSettings, meta = (DisplayName = "Level Instance"))
@@ -45,4 +48,4 @@ public:
 
 	UPROPERTY(config, EditAnywhere, Category = Pivot)
 	ELevelInstancePivotType PivotType;
-};
\ No newline at end of file
+};
diff --git a/Engine/Source/Editor/WorldPartitionEditor/Private/WorldPartition/Customizations/WorldPartitionDetailsCustomization.cpp b/Engine/Source/Editor/WorldPartitionEditor/Private/WorldPartition/Customizations/WorldPartitionDetailsCustomization.cpp
--- a/Engine/Source/Editor/WorldPartitionEditor/Private/WorldPartition/Customizations/WorldPartitionDetailsCustomization.cpp
+++ b/Engine/Source/Editor/WorldPartitionEditor/Private/WorldPartition/Customizations/WorldPartitionDetailsCustomization.cpp
@@ -59,24 +59,6 @@ void FWorldPartitionDetails::CustomizeDetails(IDetailLayoutBuilder& InDetailBuil
 	// Disable world partition button
 	if (GetDefault<UWorldPartitionEditorSettings>()->bAdvancedMode)
 	{
-		// Can be used by level instance checkbox
-		WorldPartitionCategory.AddCustomRow(LOCTEXT("CanBeUsedByLevelInstance", "Can Be Used By Level Instance"), false)
-			.RowTag(TEXT("CanBeUsedByLevelInstance"))
-			.NameContent()
-			[
-				SNew(STextBlock)
-				.Text(LOCTEXT("WorldPartitionCanBeUsedByLevelInstance", "Can Be Used By Level Instance"))
-				.ToolTipText(LOCTEXT("WorldPartitionCanBeUsedByLevelInstance_ToolTip", "Whether Level Instance can reference this partition."))
-				.Font(IDetailLayoutBuilder::GetDetailFont())
-			]
-			.ValueContent()
-			[
-				SNew(SCheckBox)
-				.IsChecked(MakeAttributeLambda([this]() { return WorldPartition.IsValid() && WorldPartition->CanBeUsedByLevelInstance() ? ECheckBoxState::Checked : ECheckBoxState::Unchecked; }))
-				.OnCheckStateChanged(this, &FWorldPartitionDetails::HandleWorldPartitionCanBeUsedByLevelInstanceChanged)
-			]
-			.Visibility(TAttribute<EVisibility>::CreateLambda([this]() { return WorldPartition.IsValid() && !WorldPartition->IsStreamingEnabled() ? EVisibility::Visible : EVisibility::Hidden; }));
-
 		// Default world partition settings
 		WorldPartitionCategory.AddCustomRow(LOCTEXT("DefaultWorldPartitionSettingsRow", "DefaultWorldPartitionSettings"), true)
 			.NameContent()
@@ -164,14 +146,6 @@ void FWorldPartitionDetails::HandleWorldPartitionEnableStreamingChanged(ECheckBo
 {
 	if (InCheckState == ECheckBoxState::Checked)
 	{
-		if (WorldPartition->CanBeUsedByLevelInstance())
-		{
-			if (FMessageDialog::Open(EAppMsgType::YesNo, LOCTEXT("WorldPartitionConfirmEnableStreamingOnLevelReferencedByLevelInstance", "You are about to enable streaming on a level that could be used by level instances. Continuing will invalidate level instances referencing this level. Continue?")) == EAppReturnType::No)
-			{
-				return;
-			}
-		}
-
 		if (!WorldPartition->bStreamingWasEnabled)
 		{
 			if (FMessageDialog::Open(EAppMsgType::YesNo, LOCTEXT("WorldPartitionConfirmEnableStreaming", "You are about to enable streaming for the first time, the world will be setup to stream. Continue?")) == EAppReturnType::No)
@@ -205,11 +179,6 @@ void FWorldPartitionDetails::HandleWorldPartitionEnableStreamingChanged(ECheckBo
 	}
 }
 
-void FWorldPartitionDetails::HandleWorldPartitionCanBeUsedByLevelInstanceChanged(ECheckBoxState InCheckState)
-{
-	WorldPartition->SetCanBeUsedByLevelInstance((InCheckState == ECheckBoxState::Checked));
-}
-
 void FWorldPartitionDetails::HandleWorldPartitionRuntimeHashClassChanged(const UClass* InRuntimeHashClass)
 {
 	check(InRuntimeHashClass);
diff --git a/Engine/Source/Editor/WorldPartitionEditor/Private/WorldPartition/Customizations/WorldPartitionDetailsCustomization.h b/Engine/Source/Editor/WorldPartitionEditor/Private/WorldPartition/Customizations/WorldPartitionDetailsCustomization.h
--- a/Engine/Source/Editor/WorldPartitionEditor/Private/WorldPartition/Customizations/WorldPartitionDetailsCustomization.h
+++ b/Engine/Source/Editor/WorldPartitionEditor/Private/WorldPartition/Customizations/WorldPartitionDetailsCustomization.h
@@ -28,9 +28,6 @@ private:
 	// Callback for changes in the world partition enable streaming.
 	void HandleWorldPartitionEnableStreamingChanged(ECheckBoxState InCheckState);
 
-	// Callback for changes in the world partition can be used by level instance.
-	void HandleWorldPartitionCanBeUsedByLevelInstanceChanged(ECheckBoxState InCheckState);
-
 	// Callback for changes in the world partition runtime hash class.
 	void HandleWorldPartitionRuntimeHashClassChanged(const UClass* InRuntimeHashClass);
 
diff --git a/Engine/Source/Runtime/Engine/Classes/Engine/Level.h b/Engine/Source/Runtime/Engine/Classes/Engine/Level.h
--- a/Engine/Source/Runtime/Engine/Classes/Engine/Level.h
+++ b/Engine/Source/Runtime/Engine/Classes/Engine/Level.h
@@ -716,8 +716,10 @@ public:
 	ENGINE_API static bool GetIsStreamingDisabledFromAsset(const FAssetData& Asset);	
 	ENGINE_API static bool GetIsStreamingDisabledFromPackage(FName LevelPackage);
 
-	ENGINE_API static bool GetPartitionedLevelCanBeUsedByLevelInstanceFromAsset(const FAssetData& Asset);
-	ENGINE_API static bool GetPartitionedLevelCanBeUsedByLevelInstanceFromPackage(FName LevelPackage);
+	UE_DEPRECATED(5.3, "GetPartitionedLevelCanBeUsedByLevelInstanceFromAsset is deprecated.")
+	ENGINE_API static bool GetPartitionedLevelCanBeUsedByLevelInstanceFromAsset(const FAssetData& Asset) { return true; }
+	UE_DEPRECATED(5.3, "GetPartitionedLevelCanBeUsedByLevelInstanceFromPackage is deprecated.")
+	ENGINE_API static bool GetPartitionedLevelCanBeUsedByLevelInstanceFromPackage(FName LevelPackage) { return true; }
 
 	ENGINE_API static FVector GetLevelInstancePivotOffsetFromAsset(const FAssetData& Asset);
 	ENGINE_API static FVector GetLevelInstancePivotOffsetFromPackage(FName LevelPackage);
diff --git a/Engine/Source/Runtime/Engine/Private/Level.cpp b/Engine/Source/Runtime/Engine/Private/Level.cpp
--- a/Engine/Source/Runtime/Engine/Private/Level.cpp
+++ b/Engine/Source/Runtime/Engine/Private/Level.cpp
@@ -2508,26 +2508,6 @@ bool ULevel::GetLevelScriptExternalActorsReferencesFromPackage(FName LevelPackag
 	});
 }
 
-bool ULevel::GetPartitionedLevelCanBeUsedByLevelInstanceFromAsset(const FAssetData& Asset)
-{
-	FString PartitionedLevelCanBeUsedByLevelInstanceStr;
-	static const FName NAME_PartitionedLevelCanBeUsedByLevelInstance(TEXT("PartitionedLevelCanBeUsedByLevelInstance"));
-	if (Asset.GetTagValue(NAME_PartitionedLevelCanBeUsedByLevelInstance, PartitionedLevelCanBeUsedByLevelInstanceStr))
-	{
-		check(PartitionedLevelCanBeUsedByLevelInstanceStr == TEXT("1"));
-		return true;
-	}
-	return false;
-}
-
-bool ULevel::GetPartitionedLevelCanBeUsedByLevelInstanceFromPackage(FName LevelPackage)
-{
-	return LevelAssetRegistryHelper::GetLevelInfoFromAssetRegistry(LevelPackage, [](const FAssetData& Asset)
-	{
-		return GetPartitionedLevelCanBeUsedByLevelInstanceFromAsset(Asset);
-	});
-}
-
 bool ULevel::GetLevelBoundsFromAsset(const FAssetData& Asset, FBox& OutLevelBounds)
 {
 	FString LevelBoundsLocationStr;
@@ -2986,15 +2966,6 @@ void ULevel::OnLevelLoaded()
 	FixupActorFolders();
 #endif
 
-	auto IsValidLevelInstanceWorldPartition = [](UWorldPartition* InWorldPartition)
-	{
-#if WITH_EDITOR
-		return InWorldPartition->CanBeUsedByLevelInstance();
-#else
-		return false;
-#endif
-	};
-
 	// Find associated level streaming for this level
 	const ULevelStreaming* LevelStreaming = FLevelUtils::FindStreamingLevel(this);
 
@@ -3019,13 +2990,12 @@ void ULevel::OnLevelLoaded()
 			//
 			const bool bIsOwningWorldGameWorld = OwningWorld->IsGameWorld();
 			const bool bIsOwningWorldPartitioned = OwningWorld->IsPartitionedWorld();
-			const bool bIsValidLevelInstance = IsValidLevelInstanceWorldPartition(WorldPartition);
 			const bool bIsMainWorldLevel = OwningWorld->PersistentLevel == this;
-			const bool bInitializeForEditor = !bIsOwningWorldGameWorld && bIsValidLevelInstance;
+			const bool bInitializeForEditor = !bIsOwningWorldGameWorld;
 			const bool bInitializeForGame = bIsOwningWorldGameWorld;
 
-			UE_LOG(LogWorldPartition, Log, TEXT("ULevel::OnLevelLoaded(%s)(bIsOwningWorldGameWorld=%d, bIsOwningWorldPartitioned=%d, bIsValidLevelInstance=%d, InitializeForMainWorld=%d, InitializeForEditor=%d, InitializeForGame=%d)"), 
-				*GetTypedOuter<UWorld>()->GetName(), bIsOwningWorldGameWorld ? 1 : 0, bIsOwningWorldPartitioned ? 1 : 0, bIsValidLevelInstance ? 1 : 0, bIsMainWorldLevel ? 1 : 0, bInitializeForEditor ? 1 : 0, bInitializeForGame ? 1 : 0);
+			UE_LOG(LogWorldPartition, Log, TEXT("ULevel::OnLevelLoaded(%s)(bIsOwningWorldGameWorld=%d, bIsOwningWorldPartitioned=%d, InitializeForMainWorld=%d, InitializeForEditor=%d, InitializeForGame=%d)"), 
+				*GetTypedOuter<UWorld>()->GetName(), bIsOwningWorldGameWorld ? 1 : 0, bIsOwningWorldPartitioned ? 1 : 0, bIsMainWorldLevel ? 1 : 0, bInitializeForEditor ? 1 : 0, bInitializeForGame ? 1 : 0);
 
 			if (bIsMainWorldLevel || bInitializeForEditor)
 			{
diff --git a/Engine/Source/Runtime/Engine/Private/LevelInstance/LevelInstanceActor.cpp b/Engine/Source/Runtime/Engine/Private/LevelInstance/LevelInstanceActor.cpp
--- a/Engine/Source/Runtime/Engine/Private/LevelInstance/LevelInstanceActor.cpp
+++ b/Engine/Source/Runtime/Engine/Private/LevelInstance/LevelInstanceActor.cpp
@@ -68,10 +68,6 @@ void ALevelInstance::PostRegisterAllComponents()
 {
 	Super::PostRegisterAllComponents();
 
-#if WITH_EDITOR
-	ResetUnsupportedWorldAsset();
-#endif
-
 	if (GetLocalRole() == ENetRole::ROLE_Authority && GetWorld()->IsGameWorld())
 	{
 #if !WITH_EDITOR
@@ -314,16 +310,6 @@ void ALevelInstance::PushLevelInstanceEditingStateToProxies(bool bInEditingState
 	LevelInstanceActorImpl.PushLevelInstanceEditingStateToProxies(bInEditingState);
 }
 
-void ALevelInstance::ResetUnsupportedWorldAsset()
-{
-	if (!ULevelInstanceSubsystem::CanUsePackage(*WorldAsset.GetLongPackageName()))
-	{
-		UE_LOG(LogLevelInstance, Warning, TEXT("LevelInstance doesn't support partitioned world %s, make sure to flag world partition's 'Can be Used by Level Instance'."), *WorldAsset.GetLongPackageName());
-		WorldAsset.Reset();
-		UpdateCookedAsset();
-	}
-}
-
 void ALevelInstance::UpdateCookedAsset()
 {
 #if WITH_EDITORONLY_DATA
diff --git a/Engine/Source/Runtime/Engine/Private/LevelInstance/LevelInstanceSubsystem.cpp b/Engine/Source/Runtime/Engine/Private/LevelInstance/LevelInstanceSubsystem.cpp
--- a/Engine/Source/Runtime/Engine/Private/LevelInstance/LevelInstanceSubsystem.cpp
+++ b/Engine/Source/Runtime/Engine/Private/LevelInstance/LevelInstanceSubsystem.cpp
@@ -1074,22 +1074,23 @@ ILevelInstanceInterface* ULevelInstanceSubsystem::CreateLevelInstanceFrom(const
 	ULevelStreamingLevelInstanceEditor* LevelStreaming = nullptr;
 	{
 		const bool bIsPartitioned = GetWorld()->IsPartitionedWorld();
+		const bool bEnableStreaming = CreationParams.bEnableStreaming;
 
 		// We want to properly setup the world partition prior to its initialization
-		FDelegateHandle PreWorldInit = FWorldDelegates::OnPreWorldInitialization.AddLambda([&bIsPartitioned](UWorld* World, const UWorld::InitializationValues IVS)
+		FDelegateHandle PreWorldInit = FWorldDelegates::OnPreWorldInitialization.AddLambda([bIsPartitioned, bEnableStreaming](UWorld* World, const UWorld::InitializationValues IVS)
 		{
 			if (bIsPartitioned)
 			{
 				UWorldPartition* WorldPartition = World->GetWorldPartition();
 				if (ensure(WorldPartition))
 				{
-					WorldPartition->SetCanBeUsedByLevelInstance(true);
+					WorldPartition->bEnableStreaming = bEnableStreaming;
 				}
 			}
 		});
 
 		LevelStreaming = StaticCast<ULevelStreamingLevelInstanceEditor*>(EditorLevelUtils::CreateNewStreamingLevelForWorld(
-		*GetWorld(), ULevelStreamingLevelInstanceEditor::StaticClass(), CreationParams.UseExternalActors(), LevelFilename, &ActorsToMove, CreationParams.TemplateWorld, /*bUseSaveAs*/true, bIsPartitioned, [this, bIsPartitioned, &ActorsToMove](ULevel* InLevel)
+		*GetWorld(), ULevelStreamingLevelInstanceEditor::StaticClass(), CreationParams.UseExternalActors(), LevelFilename, &ActorsToMove, CreationParams.TemplateWorld, /*bUseSaveAs*/true, bIsPartitioned, [this, bIsPartitioned, bEnableStreaming, &ActorsToMove](ULevel* InLevel)
 		{
 			if (InLevel->IsUsingExternalActors())
 			{
@@ -1106,8 +1107,7 @@ ILevelInstanceInterface* ULevelInstanceSubsystem::CreateLevelInstanceFrom(const
 				UWorldPartition* WorldPartition = InLevel->GetWorldPartition();
 				// Validations
 				check(WorldPartition);
-				check(WorldPartition->CanBeUsedByLevelInstance());
-				check(!WorldPartition->IsStreamingEnabled());
+				check(WorldPartition->IsStreamingEnabled() == bEnableStreaming);
 				check(InLevel->IsUsingActorFolders());
 
 				// Reset HLOD Layer (no defaults needed for Level Instances)
@@ -1807,21 +1807,6 @@ bool ULevelInstanceSubsystem::CanEditLevelInstance(const ILevelInstanceInterface
 		return false;
 	}
 
-	if (ULevel* LevelInstanceLevel = GetLevelInstanceLevel(LevelInstance))
-	{
-		if (UWorldPartition* WorldPartition = LevelInstanceLevel->GetWorldPartition())
-		{
-			if (!WorldPartition->CanBeUsedByLevelInstance())
-			{
-				if (OutReason)
-				{
-					*OutReason = FText::Format(LOCTEXT("CanEditPartitionedLevelInstance", "LevelInstance doesn't support partitioned world {0}, make sure to flag world partition's 'Can be Used by Level Instance'."), FText::FromString(LevelInstance->GetWorldAssetPackage()));
-				}
-				return false;
-			}
-		}
-	}
-	
 	if (LevelInstanceEdit)
 	{
 		if (LevelInstanceEdit->GetLevelInstance() == LevelInstance)
@@ -2419,11 +2404,6 @@ bool ULevelInstanceSubsystem::CheckForLoop(const ILevelInstanceInterface* LevelI
 	return bValid;
 }
 
-bool ULevelInstanceSubsystem::CanUsePackage(FName InPackageName)
-{
-	return !ULevel::GetIsLevelPartitionedFromPackage(InPackageName) || ULevel::GetPartitionedLevelCanBeUsedByLevelInstanceFromPackage(InPackageName);
-}
-
 bool ULevelInstanceSubsystem::CanUseWorldAsset(const ILevelInstanceInterface* LevelInstance, TSoftObjectPtr<UWorld> WorldAsset, FString* OutReason)
 {
 	// Do not validate when running convert commandlet as package might not exist yet.
@@ -2450,16 +2430,6 @@ bool ULevelInstanceSubsystem::CanUseWorldAsset(const ILevelInstanceInterface* Le
 	TArray<TPair<FText, TSoftObjectPtr<UWorld>>> LoopInfo;
 	const ILevelInstanceInterface* LoopStart = nullptr;
 
-	if (!ULevelInstanceSubsystem::CanUsePackage(*WorldAsset.GetLongPackageName()))
-	{
-		if (OutReason)
-		{
-			*OutReason = FString::Format(TEXT("LevelInstance doesn't support partitioned world {0}, make sure to flag world partition's 'Can be Used by Level Instance'.\n"), { WorldAsset.GetLongPackageName() });
-		}
-
-		return false;
-	}
-
 	if (!CheckForLoop(LevelInstance, WorldAsset, OutReason ? &LoopInfo : nullptr, OutReason ? &LoopStart : nullptr))
 	{
 		if (OutReason)
diff --git a/Engine/Source/Runtime/Engine/Private/WorldPartition/LevelInstance/LevelInstanceActorDesc.cpp b/Engine/Source/Runtime/Engine/Private/WorldPartition/LevelInstance/LevelInstanceActorDesc.cpp
--- a/Engine/Source/Runtime/Engine/Private/WorldPartition/LevelInstance/LevelInstanceActorDesc.cpp
+++ b/Engine/Source/Runtime/Engine/Private/WorldPartition/LevelInstance/LevelInstanceActorDesc.cpp
@@ -171,7 +171,7 @@ bool FLevelInstanceActorDesc::IsContainerInstanceInternal() const
 		return false;
 	}
 
-	return ULevelInstanceSubsystem::CanUsePackage(GetContainerPackage());
+	return true;
 }
 
 bool FLevelInstanceActorDesc::GetContainerInstance(FContainerInstance& OutContainerInstance) const
@@ -211,7 +211,7 @@ void FLevelInstanceActorDesc::CheckForErrors(IStreamingGenerationErrorHandler* E
 	}
 	else if (ULevel::GetIsLevelPartitionedFromPackage(GetContainerPackage()))
 	{
-		if ((DesiredRuntimeBehavior != ELevelInstanceRuntimeBehavior::Partitioned) || !ULevel::GetPartitionedLevelCanBeUsedByLevelInstanceFromPackage(GetContainerPackage()))
+		if (DesiredRuntimeBehavior != ELevelInstanceRuntimeBehavior::Partitioned)
 		{
 			ErrorHandler->OnLevelInstanceInvalidWorldAsset(this, GetContainerPackage(), IStreamingGenerationErrorHandler::ELevelInstanceInvalidReason::WorldAssetImcompatiblePartitioned);
 		}
diff --git a/Engine/Source/Runtime/Engine/Private/WorldPartition/WorldPartition.cpp b/Engine/Source/Runtime/Engine/Private/WorldPartition/WorldPartition.cpp
--- a/Engine/Source/Runtime/Engine/Private/WorldPartition/WorldPartition.cpp
+++ b/Engine/Source/Runtime/Engine/Private/WorldPartition/WorldPartition.cpp
@@ -270,7 +270,6 @@ UWorldPartition::UWorldPartition(const FObjectInitializer& ObjectInitializer)
 	, WorldPartitionEditor(nullptr)
 	, bStreamingWasEnabled(true)
 	, bShouldCheckEnableStreamingWarning(false)
-	, bCanBeUsedByLevelInstance(false)
 	, bForceGarbageCollection(false)
 	, bForceGarbageCollectionPurge(false)
 	, bEnablingStreamingJustified(false)
@@ -1280,27 +1279,6 @@ void UWorldPartition::SetEnableStreaming(bool bInEnableStreaming)
 	}
 }
 
-bool UWorldPartition::CanBeUsedByLevelInstance() const
-{
-	return bCanBeUsedByLevelInstance && !IsStreamingEnabled();
-}
-
-void UWorldPartition::SetCanBeUsedByLevelInstance(bool bInCanBeUsedByLevelInstance)
-{
-	if (bCanBeUsedByLevelInstance != bInCanBeUsedByLevelInstance)
-	{
-		const FScopedTransaction Transaction(LOCTEXT("EditorWorldPartitionCanBeUsedByLevelInstance", "Set WorldPartition CanBeUsedByLevelInstance"));
-
-		SetFlags(RF_Transactional);
-		Modify();
-		bCanBeUsedByLevelInstance = bInCanBeUsedByLevelInstance;
-		if (bCanBeUsedByLevelInstance)
-		{
-			bEnableStreaming = false;
-		}
-	}
-}
-
 void UWorldPartition::OnEnableStreamingChanged()
 {
 	for (FActorDescContainerCollection::TIterator<> ActorDescIterator(this); ActorDescIterator; ++ActorDescIterator)
@@ -1822,12 +1800,6 @@ void UWorldPartition::AppendAssetRegistryTags(TArray<FAssetRegistryTag>& OutTags
 		OutTags.Add(FAssetRegistryTag(NAME_LevelHasStreamingDisabled, TEXT("1"), FAssetRegistryTag::TT_Hidden));
 	}
 
-	if (CanBeUsedByLevelInstance())
-	{
-		static const FName NAME_PartitionedLevelCanBeUsedByLevelInstance(TEXT("PartitionedLevelCanBeUsedByLevelInstance"));
-		OutTags.Add(FAssetRegistryTag(NAME_PartitionedLevelCanBeUsedByLevelInstance, TEXT("1"), FAssetRegistryTag::TT_Hidden));
-	}
-
 	// Append level script references so we can perform changelists validations without loading the world
 	if (const ULevelScriptBlueprint* LevelScriptBlueprint = GetWorld()->PersistentLevel->GetLevelScriptBlueprint(true))
 	{
diff --git a/Engine/Source/Runtime/Engine/Private/WorldPartition/WorldPartitionConverter.cpp b/Engine/Source/Runtime/Engine/Private/WorldPartition/WorldPartitionConverter.cpp
--- a/Engine/Source/Runtime/Engine/Private/WorldPartition/WorldPartitionConverter.cpp
+++ b/Engine/Source/Runtime/Engine/Private/WorldPartition/WorldPartitionConverter.cpp
@@ -48,11 +48,6 @@ bool FWorldPartitionConverter::Convert()
 	UWorldPartition* WorldPartition = World->GetWorldPartition();
 	if (!WorldPartition)
 	{
-		if (Parameters.bCanBeUsedByLevelInstance && Parameters.bEnableStreaming)
-		{
-			return false;
-		}
-
 		World->PersistentLevel->ConvertAllActorsToPackaging(true);
 		World->PersistentLevel->bUseExternalActors = true;
 		
@@ -65,7 +60,6 @@ bool FWorldPartitionConverter::Convert()
 		FLevelActorFoldersHelper::SetUseActorFolders(World->PersistentLevel, Parameters.bUseActorFolders);
 		WorldPartition->bEnableStreaming = Parameters.bEnableStreaming;
 		WorldPartition->bStreamingWasEnabled = Parameters.bEnableStreaming;
-		WorldPartition->SetCanBeUsedByLevelInstance(Parameters.bCanBeUsedByLevelInstance);
 		bCreatedWorldPartition = true;
 	}
 
diff --git a/Engine/Source/Runtime/Engine/Public/LevelInstance/LevelInstanceActor.h b/Engine/Source/Runtime/Engine/Public/LevelInstance/LevelInstanceActor.h
--- a/Engine/Source/Runtime/Engine/Public/LevelInstance/LevelInstanceActor.h
+++ b/Engine/Source/Runtime/Engine/Public/LevelInstance/LevelInstanceActor.h
@@ -106,7 +106,6 @@ public:
 
 private:
 	ENGINE_API void UpdateCookedAsset();
-	ENGINE_API void ResetUnsupportedWorldAsset();
 #endif
 };
 
diff --git a/Engine/Source/Runtime/Engine/Public/LevelInstance/LevelInstanceSubsystem.h b/Engine/Source/Runtime/Engine/Public/LevelInstance/LevelInstanceSubsystem.h
--- a/Engine/Source/Runtime/Engine/Public/LevelInstance/LevelInstanceSubsystem.h
+++ b/Engine/Source/Runtime/Engine/Public/LevelInstance/LevelInstanceSubsystem.h
@@ -119,7 +119,9 @@ public:
 	static ENGINE_API bool CheckForLoop(const ILevelInstanceInterface* LevelInstance, TArray<TPair<FText, TSoftObjectPtr<UWorld>>>* LoopInfo = nullptr, const ILevelInstanceInterface** LoopStart = nullptr);
 	static ENGINE_API bool CheckForLoop(const ILevelInstanceInterface* LevelInstance, TSoftObjectPtr<UWorld> WorldAsset, TArray<TPair<FText, TSoftObjectPtr<UWorld>>>* LoopInfo = nullptr, const ILevelInstanceInterface** LoopStart = nullptr);
 	static ENGINE_API bool CanUseWorldAsset(const ILevelInstanceInterface* LevelInstance, TSoftObjectPtr<UWorld> WorldAsset, FString* OutReason);
-	static ENGINE_API bool CanUsePackage(FName InPackageName);
+	
+	UE_DEPRECATED(5.3, "CanUsePackage is deprecated.")
+	static ENGINE_API bool CanUsePackage(FName InPackageName) { return true;  }
 
 	/** Editor-only event triggered when level instance is committed with changes */
 	DECLARE_EVENT_OneParam(ULevelInstanceSubsystem, FLevelInstanceChanged, FName);
@@ -234,4 +236,4 @@ private:
 
 	FWorldPartitionReference CurrentEditLevelInstanceActor;
 #endif
-};
+};
\ No newline at end of file
diff --git a/Engine/Source/Runtime/Engine/Public/LevelInstance/LevelInstanceTypes.h b/Engine/Source/Runtime/Engine/Public/LevelInstance/LevelInstanceTypes.h
--- a/Engine/Source/Runtime/Engine/Public/LevelInstance/LevelInstanceTypes.h
+++ b/Engine/Source/Runtime/Engine/Public/LevelInstance/LevelInstanceTypes.h
@@ -109,6 +109,9 @@ struct FNewLevelInstanceParams
 	UPROPERTY()
 	TSubclassOf<AActor> LevelInstanceClass;
 
+	UPROPERTY()
+	bool bEnableStreaming = false;
+
 private:
 	UPROPERTY(EditAnywhere, Category = Default, meta = (EditCondition = "!bForceExternalActors", EditConditionHides, HideEditConditionToggle))
 	bool bExternalActors = true;
diff --git a/Engine/Source/Runtime/Engine/Public/WorldPartition/WorldPartition.h b/Engine/Source/Runtime/Engine/Public/WorldPartition/WorldPartition.h
--- a/Engine/Source/Runtime/Engine/Public/WorldPartition/WorldPartition.h
+++ b/Engine/Source/Runtime/Engine/Public/WorldPartition/WorldPartition.h
@@ -123,8 +123,12 @@ public:
 
 public:
 	ENGINE_API void SetEnableStreaming(bool bInEnableStreaming);
-	ENGINE_API bool CanBeUsedByLevelInstance() const;
-	ENGINE_API void SetCanBeUsedByLevelInstance(bool bInCanBeUsedByLevelInstance);
+
+	UE_DEPRECATED(5.3, "CanBeUsedByLevelInstance is deprecated.")
+	ENGINE_API bool CanBeUsedByLevelInstance() const { return true; }
+	UE_DEPRECATED(5.3, "SetCanBeUsedByLevelInstance is deprecated.")
+	ENGINE_API void SetCanBeUsedByLevelInstance(bool bInCanBeUsedByLevelInstance) {}
+
 	ENGINE_API void OnEnableStreamingChanged();
 
 private:
@@ -390,10 +394,6 @@ private:
 
 	/** Used to know if we need to recheck if the user should enable streaming based on world size. */
 	bool bShouldCheckEnableStreamingWarning;
-
-	/** Whether Level Instance can reference this partition. */
-	UPROPERTY()
-	bool bCanBeUsedByLevelInstance;
 #endif
 
 public:
diff --git a/Engine/Source/Runtime/Engine/Public/WorldPartition/WorldPartitionConverter.h b/Engine/Source/Runtime/Engine/Public/WorldPartition/WorldPartitionConverter.h
--- a/Engine/Source/Runtime/Engine/Public/WorldPartition/WorldPartitionConverter.h
+++ b/Engine/Source/Runtime/Engine/Public/WorldPartition/WorldPartitionConverter.h
@@ -22,7 +22,6 @@ public:
 		: bConvertSubLevels(true)
 		, bEnableStreaming(true)
 		, bEnableLoadingInEditor(true)
-		, bCanBeUsedByLevelInstance(false)
 		, bUseActorFolders(true)
 		, EditorHashClass(nullptr)
 		, RuntimeHashClass(nullptr)
@@ -31,7 +30,6 @@ public:
 		bool bConvertSubLevels;
 		bool bEnableStreaming;
 		bool bEnableLoadingInEditor;
-		bool bCanBeUsedByLevelInstance;
 		bool bUseActorFolders;
 		TSubclassOf<UWorldPartitionEditorHash> EditorHashClass;
 		TSubclassOf<UWorldPartitionRuntimeHash> RuntimeHashClass;
