made ApplicationWillTerminateDelegate thread-safe: deprecated FCoreDelegates::ApplicationWillTerminateDelegate (because thread-safe delegates can't be statically initialised) and added FCoreDelegates::GetApplicationWillTerminateDelegate() accessor for a thread-safe delegate

diff --git a/Engine/Plugins/Media/ElectraPlayer/Source/ElectraPlayerRuntime/Private/Runtime/PlayerRuntimeGlobal.cpp b/Engine/Plugins/Media/ElectraPlayer/Source/ElectraPlayerRuntime/Private/Runtime/PlayerRuntimeGlobal.cpp
--- a/Engine/Plugins/Media/ElectraPlayer/Source/ElectraPlayerRuntime/Private/Runtime/PlayerRuntimeGlobal.cpp
+++ b/Engine/Plugins/Media/ElectraPlayer/Source/ElectraPlayerRuntime/Private/Runtime/PlayerRuntimeGlobal.cpp
@@ -83,7 +83,7 @@ namespace Electra
 	{
 		if (!ApplicationTerminatingDelegate.IsValid())
 		{
-			ApplicationTerminatingDelegate = FCoreDelegates::ApplicationWillTerminateDelegate.AddStatic(&HandleApplicationWillTerminate);
+			ApplicationTerminatingDelegate = FCoreDelegates::GetApplicationWillTerminateDelegate().AddStatic(&HandleApplicationWillTerminate);
 		}
 		if (!ApplicationSuspendedDelegate.IsValid())
 		{
@@ -126,7 +126,7 @@ namespace Electra
 		}
 		if (ApplicationTerminatingDelegate.IsValid())
 		{
-			FCoreDelegates::ApplicationWillTerminateDelegate.Remove(ApplicationTerminatingDelegate);
+			FCoreDelegates::GetApplicationWillTerminateDelegate().Remove(ApplicationTerminatingDelegate);
 		}
 	}
 
diff --git a/Engine/Plugins/Runtime/nDisplay/Source/DisplayCluster/Private/Game/EngineClasses/Basics/DisplayClusterGameEngine.cpp b/Engine/Plugins/Runtime/nDisplay/Source/DisplayCluster/Private/Game/EngineClasses/Basics/DisplayClusterGameEngine.cpp
--- a/Engine/Plugins/Runtime/nDisplay/Source/DisplayCluster/Private/Game/EngineClasses/Basics/DisplayClusterGameEngine.cpp
+++ b/Engine/Plugins/Runtime/nDisplay/Source/DisplayCluster/Private/Game/EngineClasses/Basics/DisplayClusterGameEngine.cpp
@@ -64,7 +64,7 @@ void UDisplayClusterGameEngine::Init(class IEngineLoop* InEngineLoop)
 	// This is required to prevent GameThread dead-locking when Alt+F4 is used
 	// to terminate p-node of a 2+ nodes cluster. This gets called before
 	// virtual PreExit which allows to unlock GameThread in advance.
-	FCoreDelegates::ApplicationWillTerminateDelegate.AddLambda([this]()
+	FCoreDelegates::GetApplicationWillTerminateDelegate().AddLambda([this]()
 	{
 		PreExitImpl();
 	});
diff --git a/Engine/Source/Programs/CrashReportClient/Private/CrashReportAnalyticsSessionSummary.cpp b/Engine/Source/Programs/CrashReportClient/Private/CrashReportAnalyticsSessionSummary.cpp
--- a/Engine/Source/Programs/CrashReportClient/Private/CrashReportAnalyticsSessionSummary.cpp
+++ b/Engine/Source/Programs/CrashReportClient/Private/CrashReportAnalyticsSessionSummary.cpp
@@ -295,7 +295,7 @@ void FCrashReportAnalyticsSessionSummary::Initialize(const FString& ProcessGroup
 			PropertyStore = SessionSummaryManager->MakeStore(ReservedFileCapacity);
 			if (PropertyStore)
 			{
-				FCoreDelegates::ApplicationWillTerminateDelegate.AddRaw(this, &FCrashReportAnalyticsSessionSummary::OnApplicationWillTerminate);
+				FCoreDelegates::GetApplicationWillTerminateDelegate().AddRaw(this, &FCrashReportAnalyticsSessionSummary::OnApplicationWillTerminate);
 				FCoreDelegates::OnHandleSystemError.AddRaw(this, &FCrashReportAnalyticsSessionSummary::OnHandleSystemError);
 
 				CrcAnalyticsProperties::EngineVersion.Set(PropertyStore.Get(), FEngineVersion::Current().ToString(EVersionComponent::Changelist));
@@ -393,7 +393,7 @@ void FCrashReportAnalyticsSessionSummary::Shutdown(IAnalyticsProviderET* Analyti
 	LogEvent(FString::Printf(TEXT("CRC/Shutdown:%s:%.1fs"), *FDateTime::UtcNow().ToString(), FPlatformTime::Seconds() - SessionStartTimeSecs));
 
 	// Unregister from the core.
-	FCoreDelegates::ApplicationWillTerminateDelegate.RemoveAll(this);
+	FCoreDelegates::GetApplicationWillTerminateDelegate().RemoveAll(this);
 	FCoreDelegates::OnHandleSystemError.RemoveAll(this);
 	GLog->RemoveOutputDevice(this);
 
diff --git a/Engine/Source/Programs/SwitchboardListenerHelper/Private/SwitchboardListenerHelper.cpp b/Engine/Source/Programs/SwitchboardListenerHelper/Private/SwitchboardListenerHelper.cpp
--- a/Engine/Source/Programs/SwitchboardListenerHelper/Private/SwitchboardListenerHelper.cpp
+++ b/Engine/Source/Programs/SwitchboardListenerHelper/Private/SwitchboardListenerHelper.cpp
@@ -88,7 +88,7 @@ namespace UE::SwitchboardListenerHelper
 	static void MainLoop()
 	{
 		// We need a chance to handle service being terminated and not leave the GPU clocks locked if not desired.
-		FCoreDelegates::ApplicationWillTerminateDelegate.AddStatic(&HandleApplicationWillTerminate);
+		FCoreDelegates::GetApplicationWillTerminateDelegate().AddStatic(&HandleApplicationWillTerminate);
 
 		double LastTime = FPlatformTime::Seconds();
 		const float MinFrameTimeSeconds = 0.1;
diff --git a/Engine/Source/Runtime/ApplicationCore/Private/IOS/IOSAppDelegate.cpp b/Engine/Source/Runtime/ApplicationCore/Private/IOS/IOSAppDelegate.cpp
--- a/Engine/Source/Runtime/ApplicationCore/Private/IOS/IOSAppDelegate.cpp
+++ b/Engine/Source/Runtime/ApplicationCore/Private/IOS/IOSAppDelegate.cpp
@@ -1483,7 +1483,11 @@ extern double GCStartTime;
 	 Save data if appropriate.
 	 See also applicationDidEnterBackground:.
 	 */
+	PRAGMA_DISABLE_DEPRECATION_WARNINGS
 	FCoreDelegates::ApplicationWillTerminateDelegate.Broadcast();
+	PRAGMA_ENABLE_DEPRECATION_WARNINGS
+	
+	FCoreDelegates::GetApplicationWillTerminateDelegate().Broadcast();
     
     // note that we are shutting down
     // TODO: fix the reason why we are hanging when asked to shutdown
diff --git a/Engine/Source/Runtime/Core/Private/Mac/MacPlatformMisc.cpp b/Engine/Source/Runtime/Core/Private/Mac/MacPlatformMisc.cpp
--- a/Engine/Source/Runtime/Core/Private/Mac/MacPlatformMisc.cpp
+++ b/Engine/Source/Runtime/Core/Private/Mac/MacPlatformMisc.cpp
@@ -711,7 +711,11 @@ void FMacPlatformMisc::RequestExit( bool Force, const TCHAR* CallSite)
 	UE_LOG(LogMac, Log,  TEXT("FPlatformMisc::RequestExit(%i, %s)"), Force,
 		CallSite ? CallSite : TEXT("<NoCallSiteInfo>"));
 
+	PRAGMA_DISABLE_DEPRECATION_WARNINGS
 	FCoreDelegates::ApplicationWillTerminateDelegate.Broadcast();
+	PRAGMA_ENABLE_DEPRECATION_WARNINGS
+
+	FCoreDelegates::GetApplicationWillTerminateDelegate().Broadcast();
 
 	notify_cancel(GMacAppInfo.PowerSourceNotification);
 	GMacAppInfo.PowerSourceNotification = 0;
diff --git a/Engine/Source/Runtime/Core/Private/Misc/CoreDelegates.cpp b/Engine/Source/Runtime/Core/Private/Misc/CoreDelegates.cpp
--- a/Engine/Source/Runtime/Core/Private/Misc/CoreDelegates.cpp
+++ b/Engine/Source/Runtime/Core/Private/Misc/CoreDelegates.cpp
@@ -125,7 +125,17 @@ TMulticastDelegate<void()> FCoreDelegates::ApplicationWillDeactivateDelegate;
 TMulticastDelegate<void()> FCoreDelegates::ApplicationHasReactivatedDelegate;
 TMulticastDelegate<void()> FCoreDelegates::ApplicationWillEnterBackgroundDelegate;
 TMulticastDelegate<void()> FCoreDelegates::ApplicationHasEnteredForegroundDelegate;
+
+PRAGMA_DISABLE_DEPRECATION_WARNINGS
 TMulticastDelegate<void()> FCoreDelegates::ApplicationWillTerminateDelegate;
+PRAGMA_ENABLE_DEPRECATION_WARNINGS
+
+TTSMulticastDelegate<void()>& FCoreDelegates::GetApplicationWillTerminateDelegate()
+{
+	static TTSMulticastDelegate<void()> Singleton;
+	return Singleton;
+}
+
 TMulticastDelegate<void()> FCoreDelegates::ApplicationShouldUnloadResourcesDelegate;
 TMulticastDelegate<void(float)> FCoreDelegates::MobileBackgroundTickDelegate;
 
diff --git a/Engine/Source/Runtime/Core/Private/Windows/WindowsPlatformMisc.cpp b/Engine/Source/Runtime/Core/Private/Windows/WindowsPlatformMisc.cpp
--- a/Engine/Source/Runtime/Core/Private/Windows/WindowsPlatformMisc.cpp
+++ b/Engine/Source/Runtime/Core/Private/Windows/WindowsPlatformMisc.cpp
@@ -635,7 +635,13 @@ static BOOL WINAPI ConsoleCtrlHandler(DWORD CtrlType)
 	if (!AppTermDelegateBroadcast)
 	{
 		RequestEngineExit(TEXT("ConsoleCtrl RequestExit"));
+
+		PRAGMA_DISABLE_DEPRECATION_WARNINGS
 		FCoreDelegates::ApplicationWillTerminateDelegate.Broadcast();
+		PRAGMA_ENABLE_DEPRECATION_WARNINGS
+
+		FCoreDelegates::GetApplicationWillTerminateDelegate().Broadcast();
+
 		AppTermDelegateBroadcast = true;
 	}
 
@@ -1020,7 +1026,12 @@ void FWindowsPlatformMisc::RequestExitWithStatus(bool Force, uint8 ReturnCode, c
 #endif
 
 	RequestEngineExit(TEXT("Win RequestExit"));
+
+	PRAGMA_DISABLE_DEPRECATION_WARNINGS
 	FCoreDelegates::ApplicationWillTerminateDelegate.Broadcast();
+	PRAGMA_ENABLE_DEPRECATION_WARNINGS
+
+	FCoreDelegates::GetApplicationWillTerminateDelegate().Broadcast();
 
 	if (Force)
 	{
diff --git a/Engine/Source/Runtime/Core/Public/Misc/CoreDelegates.h b/Engine/Source/Runtime/Core/Public/Misc/CoreDelegates.h
--- a/Engine/Source/Runtime/Core/Public/Misc/CoreDelegates.h
+++ b/Engine/Source/Runtime/Core/Public/Misc/CoreDelegates.h
@@ -300,8 +300,11 @@ public:
 	// This *may* be called when the application is getting terminated by the OS.
 	// There is no guarantee that this will ever be called on a mobile device,
 	// save state when ApplicationWillEnterBackgroundDelegate is called instead.
+	UE_DEPRECATED(5.3, "This delegate is not thread safe while it can be used concurrently. Please use to GetApplicationWillTerminateDelegate() instead.")
 	static TMulticastDelegate<void()> ApplicationWillTerminateDelegate;
 
+	static TTSMulticastDelegate<void()>& GetApplicationWillTerminateDelegate();
+
 	// Some platform have a System UI Overlay that can draw on top of the application.
 	// The game might want to be notified so it can pause, etc...
 	// Parameter (bool) should be true if the system UI is displayed, otherwise false should be passed
diff --git a/Engine/Source/Runtime/Engine/Private/Analytics/EngineAnalyticsSessionSummary.cpp b/Engine/Source/Runtime/Engine/Private/Analytics/EngineAnalyticsSessionSummary.cpp
--- a/Engine/Source/Runtime/Engine/Private/Analytics/EngineAnalyticsSessionSummary.cpp
+++ b/Engine/Source/Runtime/Engine/Private/Analytics/EngineAnalyticsSessionSummary.cpp
@@ -217,7 +217,7 @@ FEngineAnalyticsSessionSummary::FEngineAnalyticsSessionSummary(TSharedPtr<IAnaly
 
 	// Listen to interesting events.
 	FCoreDelegates::OnHandleSystemError.AddRaw(this, &FEngineAnalyticsSessionSummary::OnCrashing); // WARNING: Don't assume this function is only called from game thread.
-	FCoreDelegates::ApplicationWillTerminateDelegate.AddRaw(this, &FEngineAnalyticsSessionSummary::OnTerminate); // WARNING: Don't assume this function is only called from game thread.
+	FCoreDelegates::GetApplicationWillTerminateDelegate().AddRaw(this, &FEngineAnalyticsSessionSummary::OnTerminate); // WARNING: Don't assume this function is only called from game thread.
 	FCoreDelegates::IsVanillaProductChanged.AddRaw(this, &FEngineAnalyticsSessionSummary::OnVanillaStateChanged);
 	FCoreDelegates::OnUserLoginChangedEvent.AddRaw(this, &FEngineAnalyticsSessionSummary::OnUserLoginChanged);
 	FUserActivityTracking::OnActivityChanged.AddRaw(this, &FEngineAnalyticsSessionSummary::OnUserActivity);
@@ -231,7 +231,7 @@ void FEngineAnalyticsSessionSummary::Shutdown()
 		ShutdownInternal();
 
 		FCoreDelegates::OnHandleSystemError.RemoveAll(this);
-		FCoreDelegates::ApplicationWillTerminateDelegate.RemoveAll(this);
+		FCoreDelegates::GetApplicationWillTerminateDelegate().RemoveAll(this);
 		FCoreDelegates::IsVanillaProductChanged.RemoveAll(this);
 		FCoreDelegates::OnUserLoginChangedEvent.RemoveAll(this);
 		FUserActivityTracking::OnActivityChanged.RemoveAll(this);
diff --git a/Engine/Source/Runtime/Engine/Private/Components/ApplicationLifecycleComponent.cpp b/Engine/Source/Runtime/Engine/Private/Components/ApplicationLifecycleComponent.cpp
--- a/Engine/Source/Runtime/Engine/Private/Components/ApplicationLifecycleComponent.cpp
+++ b/Engine/Source/Runtime/Engine/Private/Components/ApplicationLifecycleComponent.cpp
@@ -19,7 +19,7 @@ void UApplicationLifecycleComponent::OnRegister()
 	FCoreDelegates::ApplicationHasReactivatedDelegate.AddUObject(this, &UApplicationLifecycleComponent::ApplicationHasReactivatedDelegate_Handler);
 	FCoreDelegates::ApplicationWillEnterBackgroundDelegate.AddUObject(this, &UApplicationLifecycleComponent::ApplicationWillEnterBackgroundDelegate_Handler);
 	FCoreDelegates::ApplicationHasEnteredForegroundDelegate.AddUObject(this, &UApplicationLifecycleComponent::ApplicationHasEnteredForegroundDelegate_Handler);
-	FCoreDelegates::ApplicationWillTerminateDelegate.AddUObject(this, &UApplicationLifecycleComponent::ApplicationWillTerminateDelegate_Handler);
+	FCoreDelegates::GetApplicationWillTerminateDelegate().AddUObject(this, &UApplicationLifecycleComponent::ApplicationWillTerminateDelegate_Handler);
 	FCoreDelegates::ApplicationShouldUnloadResourcesDelegate.AddUObject(this, &UApplicationLifecycleComponent::ApplicationShouldUnloadResourcesDelegate_Handler);
 	FCoreDelegates::ApplicationReceivedStartupArgumentsDelegate.AddUObject(this, &UApplicationLifecycleComponent::ApplicationReceivedStartupArgumentsDelegate_Handler);
 
@@ -35,7 +35,7 @@ void UApplicationLifecycleComponent::OnUnregister()
  	FCoreDelegates::ApplicationHasReactivatedDelegate.RemoveAll(this);
  	FCoreDelegates::ApplicationWillEnterBackgroundDelegate.RemoveAll(this);
  	FCoreDelegates::ApplicationHasEnteredForegroundDelegate.RemoveAll(this);
- 	FCoreDelegates::ApplicationWillTerminateDelegate.RemoveAll(this);
+ 	FCoreDelegates::GetApplicationWillTerminateDelegate().RemoveAll(this);
  	FCoreDelegates::ApplicationShouldUnloadResourcesDelegate.RemoveAll(this);
  	FCoreDelegates::ApplicationReceivedStartupArgumentsDelegate.RemoveAll(this);
 	FCoreDelegates::OnTemperatureChange.RemoveAll(this);
diff --git a/Engine/Source/Runtime/Engine/Private/Kismet/BlueprintPlatformLibrary.cpp b/Engine/Source/Runtime/Engine/Private/Kismet/BlueprintPlatformLibrary.cpp
--- a/Engine/Source/Runtime/Engine/Private/Kismet/BlueprintPlatformLibrary.cpp
+++ b/Engine/Source/Runtime/Engine/Private/Kismet/BlueprintPlatformLibrary.cpp
@@ -18,7 +18,7 @@ void UPlatformGameInstance::PostInitProperties()
     FCoreDelegates::ApplicationHasReactivatedDelegate.AddUObject(this, &UPlatformGameInstance::ApplicationHasReactivatedDelegate_Handler);
     FCoreDelegates::ApplicationWillEnterBackgroundDelegate.AddUObject(this, &UPlatformGameInstance::ApplicationWillEnterBackgroundDelegate_Handler);
     FCoreDelegates::ApplicationHasEnteredForegroundDelegate.AddUObject(this, &UPlatformGameInstance::ApplicationHasEnteredForegroundDelegate_Handler);
-    FCoreDelegates::ApplicationWillTerminateDelegate.AddUObject(this, &UPlatformGameInstance::ApplicationWillTerminateDelegate_Handler);
+    FCoreDelegates::GetApplicationWillTerminateDelegate().AddUObject(this, &UPlatformGameInstance::ApplicationWillTerminateDelegate_Handler);
 	FCoreDelegates::ApplicationShouldUnloadResourcesDelegate.AddUObject(this, &UPlatformGameInstance::ApplicationShouldUnloadResourcesDelegate_Handler);
 	FCoreDelegates::ApplicationReceivedStartupArgumentsDelegate.AddUObject(this, &UPlatformGameInstance::ApplicationReceivedStartupArgumentsDelegate_Handler);
     FCoreDelegates::ApplicationRegisteredForRemoteNotificationsDelegate.AddUObject(this, &UPlatformGameInstance::ApplicationRegisteredForRemoteNotificationsDelegate_Handler);
@@ -36,7 +36,7 @@ void UPlatformGameInstance::BeginDestroy()
 	FCoreDelegates::ApplicationHasReactivatedDelegate.RemoveAll(this);
 	FCoreDelegates::ApplicationWillEnterBackgroundDelegate.RemoveAll(this);
 	FCoreDelegates::ApplicationHasEnteredForegroundDelegate.RemoveAll(this);
-	FCoreDelegates::ApplicationWillTerminateDelegate.RemoveAll(this);
+	FCoreDelegates::GetApplicationWillTerminateDelegate().RemoveAll(this);
  	FCoreDelegates::ApplicationShouldUnloadResourcesDelegate.RemoveAll(this);
  	FCoreDelegates::ApplicationReceivedStartupArgumentsDelegate.RemoveAll(this);
 	FCoreDelegates::ApplicationRegisteredForRemoteNotificationsDelegate.RemoveAll(this);
diff --git a/Engine/Source/Runtime/Launch/Private/Android/LaunchAndroid.cpp b/Engine/Source/Runtime/Launch/Private/Android/LaunchAndroid.cpp
--- a/Engine/Source/Runtime/Launch/Private/Android/LaunchAndroid.cpp
+++ b/Engine/Source/Runtime/Launch/Private/Android/LaunchAndroid.cpp
@@ -1602,7 +1602,11 @@ static void OnAppCommandCB(struct android_app* app, int32_t cmd)
 		{
 			FGraphEventRef WillTerminateTask = FFunctionGraphTask::CreateAndDispatchWhenReady([]()
 			{
+				PRAGMA_DISABLE_DEPRECATION_WARNINGS
 				FCoreDelegates::ApplicationWillTerminateDelegate.Broadcast();
+				PRAGMA_ENABLE_DEPRECATION_WARNINGS
+
+				FCoreDelegates::GetApplicationWillTerminateDelegate().Broadcast();
 			}, TStatId(), NULL, ENamedThreads::GameThread);
 			FTaskGraphInterface::Get().WaitUntilTaskCompletes(WillTerminateTask);
 			FAndroidMisc::NonReentrantRequestExit();
