WorldPartition: Make sure LastLoadedRegion FBox are valid
- Fixes new Empty Open World saving a invalid LastLoadedRegion
- Spatial loader adapters created with invalid bounds will not load any actors and ensure on creation

#rnx

diff --git a/Engine/Source/Editor/UnrealEd/Private/FileHelpers.cpp b/Engine/Source/Editor/UnrealEd/Private/FileHelpers.cpp
--- a/Engine/Source/Editor/UnrealEd/Private/FileHelpers.cpp
+++ b/Engine/Source/Editor/UnrealEd/Private/FileHelpers.cpp
@@ -958,10 +958,12 @@ static bool SaveWorld(UWorld* World,
 				}
 				else if (bIsTempPackage)
 				{
-					const FBox WorldBounds = RenamedWorldPartition->GetRuntimeWorldBounds();
-					UWorldPartitionEditorLoaderAdapter* EditorLoaderAdapter = RenamedWorldPartition->CreateEditorLoaderAdapter<FLoaderAdapterShape>(World, WorldBounds, TEXT("Loaded Region"));
-					EditorLoaderAdapter->GetLoaderAdapter()->SetUserCreated(true);
-					EditorLoaderAdapter->GetLoaderAdapter()->Load();
+					if (const FBox WorldBounds = RenamedWorldPartition->GetRuntimeWorldBounds(); WorldBounds.IsValid)
+					{
+						UWorldPartitionEditorLoaderAdapter* EditorLoaderAdapter = RenamedWorldPartition->CreateEditorLoaderAdapter<FLoaderAdapterShape>(World, WorldBounds, TEXT("Loaded Region"));
+						EditorLoaderAdapter->GetLoaderAdapter()->SetUserCreated(true);
+						EditorLoaderAdapter->GetLoaderAdapter()->Load();
+					}
 				}
 			}
 		}
diff --git a/Engine/Source/Editor/WorldPartitionEditor/Private/WorldPartition/SWorldPartitionEditorGrid2D.cpp b/Engine/Source/Editor/WorldPartitionEditor/Private/WorldPartition/SWorldPartitionEditorGrid2D.cpp
--- a/Engine/Source/Editor/WorldPartitionEditor/Private/WorldPartition/SWorldPartitionEditorGrid2D.cpp
+++ b/Engine/Source/Editor/WorldPartitionEditor/Private/WorldPartition/SWorldPartitionEditorGrid2D.cpp
@@ -1095,7 +1095,7 @@ void SWorldPartitionEditorGrid2D::Tick(const FGeometry& AllottedGeometry, const
 
 		const IWorldPartitionActorLoaderInterface::ILoaderAdapter* LoaderAdapter = LoaderInterface->GetLoaderAdapter();
 
-		if (!LoaderAdapter->GetBoundingBox()->IsInsideXY(ViewRectWorld))
+		if (LoaderAdapter->GetBoundingBox().IsSet() && !LoaderAdapter->GetBoundingBox()->IsInsideXY(ViewRectWorld))
 		{
 			if (IsBoundsHovered(MouseCursorPosWorld, *LoaderAdapter->GetBoundingBox()))
 			{
diff --git a/Engine/Source/Runtime/Engine/Private/WorldPartition/LoaderAdapter/LoaderAdapterShape.cpp b/Engine/Source/Runtime/Engine/Private/WorldPartition/LoaderAdapter/LoaderAdapterShape.cpp
--- a/Engine/Source/Runtime/Engine/Private/WorldPartition/LoaderAdapter/LoaderAdapterShape.cpp
+++ b/Engine/Source/Runtime/Engine/Private/WorldPartition/LoaderAdapter/LoaderAdapterShape.cpp
@@ -7,11 +7,13 @@ FLoaderAdapterShape::FLoaderAdapterShape(UWorld* InWorld, const FBox& InBounding
 	: ILoaderAdapterSpatial(InWorld)
 	, BoundingBox(InBoundingBox)
 	, Label(InLabel)
-{}
+{
+	ensureMsgf(BoundingBox.IsValid, TEXT("Invalid FLoaderAdapterShape: %s"), *InLabel);
+}
 
 TOptional<FBox> FLoaderAdapterShape::GetBoundingBox() const
 {
-	return BoundingBox;
+	return BoundingBox.IsValid? BoundingBox : TOptional<FBox>();
 }
 
 TOptional<FString> FLoaderAdapterShape::GetLabel() const
diff --git a/Engine/Source/Runtime/Engine/Private/WorldPartition/LoaderAdapter/LoaderAdapterSpatial.cpp b/Engine/Source/Runtime/Engine/Private/WorldPartition/LoaderAdapter/LoaderAdapterSpatial.cpp
--- a/Engine/Source/Runtime/Engine/Private/WorldPartition/LoaderAdapter/LoaderAdapterSpatial.cpp
+++ b/Engine/Source/Runtime/Engine/Private/WorldPartition/LoaderAdapter/LoaderAdapterSpatial.cpp
@@ -15,9 +15,12 @@ ILoaderAdapterSpatial::ILoaderAdapterSpatial(UWorld* InWorld)
 
 void ILoaderAdapterSpatial::ForEachActor(TFunctionRef<void(const FWorldPartitionHandle&)> InOperation) const
 {
-	if (UWorldPartition* WorldPartition = GetWorld()->GetWorldPartition())
+	if (GetBoundingBox().IsSet())
 	{
-		HandleIntersectingContainer(WorldPartition, WorldPartition->GetInstanceTransform(), *GetBoundingBox(), InOperation);
+		if (UWorldPartition* WorldPartition = GetWorld()->GetWorldPartition())
+		{
+			HandleIntersectingContainer(WorldPartition, WorldPartition->GetInstanceTransform(), *GetBoundingBox(), InOperation);
+		}
 	}
 }
 
diff --git a/Engine/Source/Runtime/Engine/Private/WorldPartition/WorldPartition.cpp b/Engine/Source/Runtime/Engine/Private/WorldPartition/WorldPartition.cpp
--- a/Engine/Source/Runtime/Engine/Private/WorldPartition/WorldPartition.cpp
+++ b/Engine/Source/Runtime/Engine/Private/WorldPartition/WorldPartition.cpp
@@ -1727,7 +1727,7 @@ TArray<FBox> UWorldPartition::GetUserLoadedEditorRegions() const
 	{
 		IWorldPartitionActorLoaderInterface::ILoaderAdapter* LoaderAdapter = EditorLoaderAdapter->GetLoaderAdapter();
 		check(LoaderAdapter);
-		if (LoaderAdapter->IsLoaded() && LoaderAdapter->GetUserCreated())
+		if (LoaderAdapter->GetBoundingBox().IsSet() && LoaderAdapter->IsLoaded() && LoaderAdapter->GetUserCreated())
 		{
 			Result.Add(*LoaderAdapter->GetBoundingBox());
 		}
@@ -1973,11 +1973,14 @@ void UWorldPartition::LoadLastLoadedRegions(const TArray<FBox>& EditorLastLoaded
 {
 	for (const FBox& EditorLastLoadedRegion : EditorLastLoadedRegions)
 	{
-		UWorldPartitionEditorLoaderAdapter* EditorLoaderAdapter = CreateEditorLoaderAdapter<FLoaderAdapterShape>(World, EditorLastLoadedRegion, TEXT("Last Loaded Region"));
-		IWorldPartitionActorLoaderInterface::ILoaderAdapter* LoaderAdapter = EditorLoaderAdapter->GetLoaderAdapter();
-		check(LoaderAdapter);
-		LoaderAdapter->SetUserCreated(true);
-		LoaderAdapter->Load();
+		if (EditorLastLoadedRegion.IsValid)
+		{
+			UWorldPartitionEditorLoaderAdapter* EditorLoaderAdapter = CreateEditorLoaderAdapter<FLoaderAdapterShape>(World, EditorLastLoadedRegion, TEXT("Last Loaded Region"));
+			IWorldPartitionActorLoaderInterface::ILoaderAdapter* LoaderAdapter = EditorLoaderAdapter->GetLoaderAdapter();
+			check(LoaderAdapter);
+			LoaderAdapter->SetUserCreated(true);
+			LoaderAdapter->Load();
+		}
 	}
 }
 
diff --git a/Engine/Source/Runtime/Engine/Private/WorldPartition/WorldPartitionEditorPerProjectUserSettings.cpp b/Engine/Source/Runtime/Engine/Private/WorldPartition/WorldPartitionEditorPerProjectUserSettings.cpp
--- a/Engine/Source/Runtime/Engine/Private/WorldPartition/WorldPartitionEditorPerProjectUserSettings.cpp
+++ b/Engine/Source/Runtime/Engine/Private/WorldPartition/WorldPartitionEditorPerProjectUserSettings.cpp
@@ -5,6 +5,7 @@
 #include "GameFramework/WorldSettings.h"
 #include "Misc/PackageName.h"
 #include "UObject/Package.h"
+#include "Algo/Transform.h"
 
 #include UE_INLINE_GENERATED_CPP_BY_NAME(WorldPartitionEditorPerProjectUserSettings)
 
@@ -27,8 +28,8 @@ void UWorldPartitionEditorPerProjectUserSettings::SetEditorLoadedRegions(UWorld*
 	if (ShouldSaveSettings(InWorld))
 	{
 		FWorldPartitionPerWorldSettings& PerWorldSettings = PerWorldEditorSettings.FindOrAdd(TSoftObjectPtr<UWorld>(InWorld));
-		PerWorldSettings.LoadedEditorRegions = InEditorLoadedRegions;
-		
+		PerWorldSettings.LoadedEditorRegions.Empty();
+		Algo::TransformIf(InEditorLoadedRegions, PerWorldSettings.LoadedEditorRegions, [](const FBox& InBox) { return InBox.IsValid; }, [](const FBox& InBox) { return InBox; });		
 		SaveConfig();
 	}
 }
