Subobject Replication:
* Downgrade IsValid asserts into ensures

#rnx

diff --git a/Engine/Source/Runtime/Engine/Private/ActorReplication.cpp b/Engine/Source/Runtime/Engine/Private/ActorReplication.cpp
--- a/Engine/Source/Runtime/Engine/Private/ActorReplication.cpp
+++ b/Engine/Source/Runtime/Engine/Private/ActorReplication.cpp
@@ -614,7 +614,11 @@ void AActor::RemoveReplicatedComponent(UActorComponent* Component)
 
 void AActor::AddReplicatedSubObject(UObject* SubObject, ELifetimeCondition NetCondition)
 {
-	check(IsValid(SubObject));
+	if( !IsValid(SubObject) )
+	{
+		ensureMsgf(TEXT("Ignoring AddReplicatedSubObject for %s. Invalid pointer received."), *GetNameSafe(this));
+		return;
+	}
 
 	ensureMsgf(IsUsingRegisteredSubObjectList(), TEXT("%s is registering subobjects but bReplicateUsingRegisteredSubObjectList is false. Without the flag set to true the registered subobjects will not be replicated."), *GetName());
 
@@ -775,7 +779,13 @@ void AActor::TearOffReplicatedSubObjectOnRemotePeers(UObject* SubObject)
 void AActor::AddActorComponentReplicatedSubObject(UActorComponent* OwnerComponent, UObject* SubObject, ELifetimeCondition NetCondition)
 {
 	check(IsValid(OwnerComponent));
-	check(IsValid(SubObject));
+
+	if (!IsValid(SubObject))
+	{
+		ensureMsgf(TEXT("Ignoring AddReplicatedSubObject for %s::%s. Invalid pointer received."), *GetNameSafe(this), *GetNameSafe(OwnerComponent));
+		return;
+	}
+	
 	ensureMsgf(OwnerComponent->GetIsReplicated(), TEXT("Only components with replication enabled can register subobjects. %s::%s has replication disabled."), *GetName(), *SubObject->GetName());
 	ensureMsgf(NetCondition != COND_Custom, TEXT("Custom netconditions do not work with SubObjects. %s::%s will not be replicated."), *GetName(), *SubObject->GetName());
 
