Fix an oversight that caused spline meshes to not update rendering when changed at run-time.

#rnx

diff --git a/Engine/Source/Runtime/Engine/Private/Components/SplineMeshComponent.cpp b/Engine/Source/Runtime/Engine/Private/Components/SplineMeshComponent.cpp
--- a/Engine/Source/Runtime/Engine/Private/Components/SplineMeshComponent.cpp
+++ b/Engine/Source/Runtime/Engine/Private/Components/SplineMeshComponent.cpp
@@ -706,25 +706,10 @@ void USplineMeshComponent::UpdateRenderStateAndCollision_Internal(bool bConcurre
 			MarkRenderTransformDirty();
 		}
 
-		FSplineMeshShaderParams* DestParams = nullptr;
-		if (SceneProxy->IsNaniteMesh())
-		{
-			DestParams = &static_cast<FNaniteSplineMeshSceneProxy*>(SceneProxy)->SplineParams;
-		}
-		else
-		{
-			DestParams = &static_cast<FSplineMeshSceneProxy*>(SceneProxy)->SplineParams;
-		}
-
 		ENQUEUE_RENDER_COMMAND(UpdateSplineParamsRTCommand)(
-			[SceneProxy=SceneProxy, DestParams, SrcParams=CalculateShaderParams()](FRHICommandList&)
+			[SceneProxy=SceneProxy, Params=CalculateShaderParams()](FRHICommandList&)
 			{
-				*DestParams = SrcParams;
-
-				// Request a GPU scene update for this primitive so it updates its instance data
-				FSceneInterface& Scene = SceneProxy->GetScene();
-				FPrimitiveSceneInfo& SceneInfo = *SceneProxy->GetPrimitiveSceneInfo();
-				Scene.RequestGPUSceneUpdate(SceneInfo, EPrimitiveDirtyState::ChangedOther);
+				UpdateSplineMeshParams_RenderThread(SceneProxy, Params);
 			}
 		);
 	}
diff --git a/Engine/Source/Runtime/Engine/Private/SplineMeshSceneProxy.cpp b/Engine/Source/Runtime/Engine/Private/SplineMeshSceneProxy.cpp
--- a/Engine/Source/Runtime/Engine/Private/SplineMeshSceneProxy.cpp
+++ b/Engine/Source/Runtime/Engine/Private/SplineMeshSceneProxy.cpp
@@ -324,4 +324,34 @@ void FNaniteSplineMeshSceneProxy::OnTransformChanged()
 	// NOTE: The proxy's local bounds have already been padded for WPO/Displacement
 	check(InstanceLocalBounds.Num() == 1);
 	SetInstanceLocalBounds(0, GetLocalBounds(), false);
+}
+
+void UpdateSplineMeshParams_RenderThread(FPrimitiveSceneProxy* SceneProxy, const FSplineMeshShaderParams& Params)
+{
+	check(SceneProxy->IsSplineMesh());
+
+	TArrayView<FVector4f> InstancePayloadExtension;
+	if (SceneProxy->IsNaniteMesh())
+	{
+		auto* Proxy = static_cast<FNaniteSplineMeshSceneProxy*>(SceneProxy);
+		Proxy->SplineParams = Params;
+		InstancePayloadExtension = Proxy->InstancePayloadExtension;
+	}
+	else
+	{
+		auto* Proxy = static_cast<FSplineMeshSceneProxy*>(SceneProxy);
+		Proxy->SplineParams = Params;
+		InstancePayloadExtension = Proxy->InstancePayloadExtension;
+	}
+
+	// Re-pack the shader params and request a GPU Scene update for this primitive so it updates its instance data
+	// NOTE: The payload extension could be empty if not using GPU Scene
+	if (InstancePayloadExtension.Num() == SPLINE_MESH_PARAMS_FLOAT4_SIZE)
+	{
+		PackSplineMeshParams(Params, InstancePayloadExtension);
+
+		FSceneInterface& Scene = SceneProxy->GetScene();
+		FPrimitiveSceneInfo& SceneInfo = *SceneProxy->GetPrimitiveSceneInfo();
+		Scene.RequestGPUSceneUpdate(SceneInfo, EPrimitiveDirtyState::ChangedOther);
+	}
 }
\ No newline at end of file
diff --git a/Engine/Source/Runtime/Engine/Private/SplineMeshSceneProxy.h b/Engine/Source/Runtime/Engine/Private/SplineMeshSceneProxy.h
--- a/Engine/Source/Runtime/Engine/Private/SplineMeshSceneProxy.h
+++ b/Engine/Source/Runtime/Engine/Private/SplineMeshSceneProxy.h
@@ -79,9 +79,14 @@ private:
 //////////////////////////////////////////////////////////////////////////
 // SplineMeshSceneProxy
 
+/** Helper to update the parameters of the specified spline mesh scene proxy */
+ENGINE_API void UpdateSplineMeshParams_RenderThread(FPrimitiveSceneProxy* SceneProxy, const FSplineMeshShaderParams& Params);
+
 /** Scene proxy for SplineMesh instance */
 class FSplineMeshSceneProxy final : public FStaticMeshSceneProxy
 {
+	friend void UpdateSplineMeshParams_RenderThread(FPrimitiveSceneProxy* SceneProxy, const FSplineMeshShaderParams& Params);
+
 public:
 	FSplineMeshSceneProxy(USplineMeshComponent* InComponent);
 	void InitVertexFactory(USplineMeshComponent* InComponent, int32 InLODIndex, FColorVertexBuffer*);
@@ -121,6 +126,8 @@ private:
 /** Scene proxy for SplineMesh instance for Nanite */
 class FNaniteSplineMeshSceneProxy final : public Nanite::FSceneProxy
 {
+	friend void UpdateSplineMeshParams_RenderThread(FPrimitiveSceneProxy* SceneProxy, const FSplineMeshShaderParams& Params);
+	
 public:
 	FNaniteSplineMeshSceneProxy(const Nanite::FMaterialAudit& MaterialAudit, USplineMeshComponent* InComponent);
 
