Fix crash due to type casting for anim bp compile warning in a state transition node when using looping asset player with automatic rule.
Removed behaviour change regarding inertialization transitions or very low blend time standard blend transitions considering a valid transitions when their looping asset player reaches end and they are using automatic rules.

#rnx

diff --git a/Engine/Source/Editor/AnimGraph/Private/AnimStateTransitionNode.cpp b/Engine/Source/Editor/AnimGraph/Private/AnimStateTransitionNode.cpp
--- a/Engine/Source/Editor/AnimGraph/Private/AnimStateTransitionNode.cpp
+++ b/Engine/Source/Editor/AnimGraph/Private/AnimStateTransitionNode.cpp
@@ -698,7 +698,7 @@ void UAnimStateTransitionNode::ValidateNodeDuringCompilation(class FCompilerResu
 			{
 				if (UEdGraph* PreviousStateGraph = PreviousState->GetBoundGraph())
 				{
-					if (UAnimGraphNode_StateResult* PreviousStateGraphResultNode = CastChecked<UAnimationStateGraph>(PreviousStateGraph)->GetResultNode())
+					if (UAnimGraphNode_StateResult* PreviousStateGraphResultNode = Cast<UAnimationStateGraph>(PreviousStateGraph)->GetResultNode())
 					{
 						for (UEdGraphPin* TestPin : PreviousStateGraphResultNode->Pins)
 						{
diff --git a/Engine/Source/Runtime/Engine/Private/Animation/AnimNode_StateMachine.cpp b/Engine/Source/Runtime/Engine/Private/Animation/AnimNode_StateMachine.cpp
--- a/Engine/Source/Runtime/Engine/Private/Animation/AnimNode_StateMachine.cpp
+++ b/Engine/Source/Runtime/Engine/Private/Animation/AnimNode_StateMachine.cpp
@@ -762,41 +762,19 @@ bool FAnimNode_StateMachine::FindValidTransition(const FAnimationUpdateContext&
 			{
 				if (const UAnimationAsset* AnimAsset = RelevantPlayer->GetAnimAsset())
 				{
-					float AnimTimeRemaining = AnimAsset->GetPlayLength() - RelevantPlayer->GetAccumulatedTime();
-
-					// If the asset player is a looping player and has looped, then we enter the transition as though the asset player is a non-looping player that has reached its end
-					if (RelevantPlayer->IsLooping() && AnimTimeRemaining > 0.f)
-					{
-						if (const FDeltaTimeRecord* DeltaTimeRecord = RelevantPlayer->GetDeltaTimeRecord())
-						{
-							if (DeltaTimeRecord->IsPreviousValid())
-							{
-								// If the player's time difference is in the opposite direction of the time delta, then the player must have looped
-								const float DeltaTimeFromPlayer = RelevantPlayer->GetAccumulatedTime() - DeltaTimeRecord->GetPrevious();
-								const float DeltaTimeFromRecord = DeltaTimeRecord->Delta;
-								if (DeltaTimeFromPlayer * DeltaTimeFromRecord < 0.f)
-								{
-									AnimTimeRemaining = 0.f;
-								}
-							}
-						}
-					}
-
-					// Determine if we can automatically begin transitioning.
-					{
-						const FAnimationTransitionBetweenStates& TransitionInfo = GetTransitionInfo(TransitionRule.TransitionIndex);
-					
-						// For transitions that go to a conduit the user is not able to edit the transition's cross fade duration,
-						// therefore we force the cross fade duration to zero to ensure the transition is always triggerred upon reaching the end of the animation.
-						const float CrossfadeDuration = GetStateInfo(TransitionInfo.NextState).bIsAConduit ? 0.0f : TransitionInfo.CrossfadeDuration;
+					const float AnimTimeRemaining = AnimAsset->GetPlayLength() - RelevantPlayer->GetAccumulatedTime();
+					const FAnimationTransitionBetweenStates& TransitionInfo = GetTransitionInfo(TransitionRule.TransitionIndex);
+				
+					// For transitions that go to a conduit the user is not able to edit the transition's cross fade duration,
+					// therefore we force the cross fade duration to zero to ensure the transition is always triggerred upon reaching the end of the animation.
+					const float CrossfadeDuration = GetStateInfo(TransitionInfo.NextState).bIsAConduit ? 0.0f : TransitionInfo.CrossfadeDuration;
 
-						// Allow for used to determine the automatic transition trigger time or fallback to using cross fade duration.
-						const float TransitionTriggerTime = (TransitionRule.AutomaticRuleTriggerTime >= 0.0f) ? TransitionRule.AutomaticRuleTriggerTime : CrossfadeDuration;
-						CrossfadeTimeAdjustment = TransitionTriggerTime - AnimTimeRemaining;
+					// Allow for used to determine the automatic transition trigger time or fallback to using cross fade duration.
+					const float TransitionTriggerTime = (TransitionRule.AutomaticRuleTriggerTime >= 0.0f) ? TransitionRule.AutomaticRuleTriggerTime : CrossfadeDuration;
+					CrossfadeTimeAdjustment = TransitionTriggerTime - AnimTimeRemaining;
 
-						// Trigger transition only if we have "CrossfadeTimeAdjustment" seconds left before reaching animation end boundary.
-						bCanEnterTransition = (CrossfadeTimeAdjustment >= 0.f);
-					}
+					// Trigger transition only if we have "CrossfadeTimeAdjustment" seconds left before reaching animation end boundary.
+					bCanEnterTransition = (CrossfadeTimeAdjustment >= 0.f);
 				}
 			}
 			
