WorldPartition : UWorldPartitionLevelStreamingDynamic::RequestLevel propertly detect RuntimeLevel marked as garbage (similar to ULevelStreaming)

diff --git a/Engine/Source/Runtime/Engine/Private/WorldPartition/WorldPartitionLevelStreamingDynamic.cpp b/Engine/Source/Runtime/Engine/Private/WorldPartition/WorldPartitionLevelStreamingDynamic.cpp
--- a/Engine/Source/Runtime/Engine/Private/WorldPartition/WorldPartitionLevelStreamingDynamic.cpp
+++ b/Engine/Source/Runtime/Engine/Private/WorldPartition/WorldPartitionLevelStreamingDynamic.cpp
@@ -205,13 +205,6 @@ bool UWorldPartitionLevelStreamingDynamic::RequestLevel(UWorld* InPersistentWorl
 		return false;
 	}
 
-	// Check if currently loaded level is what we want right now
-	if (LoadedLevel)
-	{
-		check(GetLoadedLevelPackageName() == GetWorldAssetPackageFName());
-		return true;
-	}
-
 	// Can not load new level now, there is still level pending unload
 	if (PendingUnloadLevel)
 	{
@@ -242,6 +235,21 @@ bool UWorldPartitionLevelStreamingDynamic::RequestLevel(UWorld* InPersistentWorl
 		RuntimeLevel = FoundWorld->PersistentLevel;
 	}
 
+	if (RuntimeLevel && !IsValid(RuntimeLevel))
+	{
+		// We're trying to reload a level that has very recently been marked for garbage collection, it might not have been cleaned up yet
+		// So continue attempting to reload the package if possible
+		UE_LOG(LogLevelStreaming, Verbose, TEXT("RequestLevel: Runtime level is marked as garbage %s"), *WorldAssetPackageFName.ToString());
+		return false;
+	}
+
+	// Check if currently loaded level is what we want right now
+	if (LoadedLevel)
+	{
+		check(GetLoadedLevelPackageName() == GetWorldAssetPackageFName());
+		return true;
+	}
+
 	if (RuntimeLevel)
 	{
 		check(ULevelStreaming::ShouldReuseUnloadedButStillAroundLevels(RuntimeLevel));
@@ -287,8 +295,10 @@ bool UWorldPartitionLevelStreamingDynamic::RequestLevel(UWorld* InPersistentWorl
 
 			if (IssueLoadRequests())
 			{
+				UWorld* OuterWorld = GetStreamingWorld();
+				bool bForceSkipBlockingLoad = false;// (OuterWorld && OuterWorld->IsGameWorld() && OuterWorld->IsInstanced() && ShouldBeAlwaysLoaded());
 				// Editor immediately blocks on load and we also block if background level streaming is disabled.
-				if (InBlockPolicy == AlwaysBlock || (ShouldBeAlwaysLoaded() && InBlockPolicy != NeverBlock))
+				if (!bForceSkipBlockingLoad && ((InBlockPolicy == AlwaysBlock) || (ShouldBeAlwaysLoaded() && InBlockPolicy != NeverBlock)))
 				{
 					if (IsAsyncLoading())
 					{
