Fix race in FKeyHandle

#rnx

diff --git a/Engine/Source/Runtime/Engine/Private/Curves/KeyHandle.cpp b/Engine/Source/Runtime/Engine/Private/Curves/KeyHandle.cpp
--- a/Engine/Source/Runtime/Engine/Private/Curves/KeyHandle.cpp
+++ b/Engine/Source/Runtime/Engine/Private/Curves/KeyHandle.cpp
@@ -7,19 +7,24 @@
 
 FKeyHandle::FKeyHandle()
 {
-	static uint32 LastKeyHandleIndex = 1;
+	static std::atomic<uint32> LastKeyHandleIndex = 1;
 	Index = ++LastKeyHandleIndex;
 
-	if (LastKeyHandleIndex == 0)
+	if (Index == 0)
 	{
 		// If we are cooking, allow wrap-around
 		if (IsRunningCookCommandlet())
 		{
-			Index = LastKeyHandleIndex = 1;
+			// Skip indices until it's not 0 anymore as we can't 
+			// assign without loss of thread-safety.
+			while (Index == 0)
+			{
+				Index = ++LastKeyHandleIndex;
+			}
 		}
 		else
 		{
-			check(LastKeyHandleIndex != 0); // check in the unlikely event that this overflows
+			check(Index != 0); // check in the unlikely event that this overflows
 		}
 	}
 }
