Delay AddToWorld of a partition world if this same partitioned world hasn't finished removing it's sub-levels triggered by a prior RemoveFromWorld.

diff --git a/Engine/Source/Runtime/Engine/Private/LevelStreaming.cpp b/Engine/Source/Runtime/Engine/Private/LevelStreaming.cpp
--- a/Engine/Source/Runtime/Engine/Private/LevelStreaming.cpp
+++ b/Engine/Source/Runtime/Engine/Private/LevelStreaming.cpp
@@ -32,6 +32,7 @@
 #include "Serialization/LoadTimeTrace.h"
 #include "Streaming/LevelStreamingDelegates.h"
 #include "AssetRegistry/AssetRegistryModule.h"
+#include "WorldPartition/WorldPartition.h"
 
 #include UE_INLINE_GENERATED_CPP_BY_NAME(LevelStreaming)
 
@@ -795,12 +796,20 @@ bool ULevelStreaming::CanMakeInvisible()
 
 bool ULevelStreaming::CanMakeVisible()
 {
-	// Once the Level becomes the current pending visibility level, this function must return true or else it could
-	// block indefinitely if called inside UWorld::BlockTillLevelStreamingCompleted (same reason as CanMakeInvisible but with World's CurrentLevelPendingVisibility)
 	UWorld* World = GetWorld();
-	if (World && World->IsGameWorld() && LoadedLevel && (LoadedLevel == World->GetCurrentLevelPendingVisibility()))
+	if (World && World->IsGameWorld() && LoadedLevel)
 	{
-		return true;
+		// Once the Level becomes the current pending visibility level, this function must return true or else it could
+		// block indefinitely if called inside UWorld::BlockTillLevelStreamingCompleted (same reason as CanMakeInvisible but with World's CurrentLevelPendingVisibility)
+		if (LoadedLevel == World->GetCurrentLevelPendingVisibility())
+		{
+			return true;
+		}
+		// Delay AddToWorld of a partition world if this same partitioned world hasn't finished removing it's sub-levels triggered by a prior RemoveFromWorld
+		else if (const UWorldPartition* WorldPartition = LoadedLevel->GetWorldPartition(); WorldPartition && !WorldPartition->CanInitialize(World))
+		{
+			return false;
+		}
 	}
 
 	const bool bCanMakeVisible = RequestVisibilityChange(true);
diff --git a/Engine/Source/Runtime/Engine/Private/WorldPartition/WorldPartition.cpp b/Engine/Source/Runtime/Engine/Private/WorldPartition/WorldPartition.cpp
--- a/Engine/Source/Runtime/Engine/Private/WorldPartition/WorldPartition.cpp
+++ b/Engine/Source/Runtime/Engine/Private/WorldPartition/WorldPartition.cpp
@@ -453,6 +453,22 @@ FName UWorldPartition::GetWorldPartitionEditorName() const
 }
 #endif
 
+bool UWorldPartition::CanInitialize(UWorld* InWorld) const
+{
+	check(InWorld);
+	if (!IsInitialized() && InWorld->IsGameWorld())
+	{
+		if (UWorldPartitionSubsystem* WorldPartitionSubsystem = InWorld->GetSubsystem<UWorldPartitionSubsystem>())
+		{
+			if (WorldPartitionSubsystem->HasUninitializationPendingStreamingLevels(this))
+			{
+				return false;
+			}
+		}
+	}
+	return true;
+}
+
 void UWorldPartition::Initialize(UWorld* InWorld, const FTransform& InTransform)
 {
 	UE_SCOPED_TIMER(TEXT("WorldPartition initialize"), LogWorldPartition, Display);
@@ -470,6 +486,7 @@ void UWorldPartition::Initialize(UWorld* InWorld, const FTransform& InTransform)
 	}
 
 	check(InWorld);
+	check(CanInitialize(InWorld));
 	World = InWorld;
 
 	if (!InTransform.Equals(FTransform::Identity))
diff --git a/Engine/Source/Runtime/Engine/Private/WorldPartition/WorldPartitionSubsystem.cpp b/Engine/Source/Runtime/Engine/Private/WorldPartition/WorldPartitionSubsystem.cpp
--- a/Engine/Source/Runtime/Engine/Private/WorldPartition/WorldPartitionSubsystem.cpp
+++ b/Engine/Source/Runtime/Engine/Private/WorldPartition/WorldPartitionSubsystem.cpp
@@ -435,6 +435,7 @@ void UWorldPartitionSubsystem::Initialize(FSubsystemCollectionBase& Collection)
 		FLevelStreamingDelegates::OnLevelBeginMakingVisible.AddUObject(this, &UWorldPartitionSubsystem::OnLevelBeginMakingVisible);
 		FLevelStreamingDelegates::OnLevelBeginMakingInvisible.AddUObject(this, &UWorldPartitionSubsystem::OnLevelBeginMakingInvisible);
 		FLevelStreamingDelegates::OnLevelStreamingTargetStateChanged.AddUObject(this, &UWorldPartitionSubsystem::OnLevelStreamingTargetStateChanged);
+		FLevelStreamingDelegates::OnLevelStreamingStateChanged.AddUObject(this, &UWorldPartitionSubsystem::OnLevelStreamingStateChanged);
 	}
 }
 
@@ -455,6 +456,7 @@ void UWorldPartitionSubsystem::Deinitialize()
 		FLevelStreamingDelegates::OnLevelBeginMakingVisible.RemoveAll(this);
 		FLevelStreamingDelegates::OnLevelBeginMakingInvisible.RemoveAll(this);
 		FLevelStreamingDelegates::OnLevelStreamingTargetStateChanged.RemoveAll(this);
+		FLevelStreamingDelegates::OnLevelStreamingStateChanged.RemoveAll(this);
 	}
 
 	// At this point World Partition should be uninitialized
@@ -508,6 +510,31 @@ void UWorldPartitionSubsystem::OnWorldPartitionUninitialized(UWorldPartition* In
 {
 	check(RegisteredWorldPartitions.Contains(InWorldPartition));
 	RegisteredWorldPartitions.Remove(InWorldPartition);
+	const UWorld* OwningWorld = GetWorld();
+	if (OwningWorld->IsGameWorld())
+	{
+		TOptional<TSet<TWeakObjectPtr<ULevelStreaming>>*> PendingStreamingLevels;
+		auto GetPendingStreamingLevels = [this, InWorldPartition, &PendingStreamingLevels]() -> TSet<TWeakObjectPtr<ULevelStreaming>>&
+		{
+			if (!PendingStreamingLevels.IsSet())
+			{
+				PendingStreamingLevels = &WorldPartitionUninitializationPendingStreamingLevels.FindOrAdd(FSoftObjectPath(InWorldPartition));
+			}
+			return *PendingStreamingLevels.GetValue();
+		};
+
+		const UWorld* WorldPartitionOuterWorld = InWorldPartition->GetTypedOuter<UWorld>();
+		if (WorldPartitionOuterWorld != OwningWorld)
+		{
+			for (ULevelStreaming* StreamingLevel : OwningWorld->GetStreamingLevels())
+			{
+				if (StreamingLevel->GetStreamingWorld() == WorldPartitionOuterWorld)
+				{
+					GetPendingStreamingLevels().Add(StreamingLevel);
+				}
+			}
+		}
+	}
 	NumWorldPartitionServerStreamingEnabled -= InWorldPartition->IsServerStreamingEnabled() ? 1 : 0;
 	check(NumWorldPartitionServerStreamingEnabled >= 0);
 
@@ -527,6 +554,44 @@ void UWorldPartitionSubsystem::OnWorldPartitionUninitialized(UWorldPartition* In
 	}
 }
 
+bool UWorldPartitionSubsystem::HasUninitializationPendingStreamingLevels(const UWorldPartition* InWorldPartition) const
+{
+	if (const TSet<TWeakObjectPtr<ULevelStreaming>>* PendingStreamingLevels = InWorldPartition ? WorldPartitionUninitializationPendingStreamingLevels.Find(FSoftObjectPath(InWorldPartition)) : nullptr)
+	{
+		if (ensure(!PendingStreamingLevels->IsEmpty()))
+		{
+			return true;
+		}
+	}
+	return false;
+}
+
+void UWorldPartitionSubsystem::OnLevelStreamingStateChanged(UWorld* InWorld, const ULevelStreaming* InStreamingLevel, ULevel* LevelIfLoaded, ELevelStreamingState PreviousState, ELevelStreamingState NewState)
+{
+	if (InWorld != GetWorld())
+	{
+		return;
+	}
+
+	if (NewState == ELevelStreamingState::Removed)
+	{
+		const UWorld* OuterWorld = InStreamingLevel->GetStreamingWorld();
+		if (const UWorldPartition* OuterWorldPartition = OuterWorld && OuterWorld->IsGameWorld() ? OuterWorld->GetWorldPartition() : nullptr)
+		{
+			if (TSet<TWeakObjectPtr<ULevelStreaming>>* PendingStreamingLevels = WorldPartitionUninitializationPendingStreamingLevels.Find(FSoftObjectPath(OuterWorldPartition)))
+			{
+				if (PendingStreamingLevels->Remove(InStreamingLevel))
+				{
+					if (PendingStreamingLevels->IsEmpty())
+					{
+						WorldPartitionUninitializationPendingStreamingLevels.Remove(OuterWorldPartition);
+					}
+				}
+			}
+		}
+	}
+}
+
 void UWorldPartitionSubsystem::OnLevelBeginMakingVisible(UWorld* InWorld, const ULevelStreaming* InStreamingLevel, ULevel* InLoadedLevel)
 {
 	if ((InWorld != GetWorld()) || !GetWorld()->IsGameWorld())
diff --git a/Engine/Source/Runtime/Engine/Public/WorldPartition/WorldPartition.h b/Engine/Source/Runtime/Engine/Public/WorldPartition/WorldPartition.h
--- a/Engine/Source/Runtime/Engine/Public/WorldPartition/WorldPartition.h
+++ b/Engine/Source/Runtime/Engine/Public/WorldPartition/WorldPartition.h
@@ -325,6 +325,7 @@ public:
 	static ENGINE_API bool IsSimulating(bool bIncludeTestEnableSimulationStreamingSource = true);
 	int32 GetStreamingStateEpoch() const { return StreamingStateEpoch; }
 
+	ENGINE_API bool CanInitialize(UWorld* InWorld) const;
 	ENGINE_API void Initialize(UWorld* World, const FTransform& InTransform);
 	ENGINE_API bool IsInitialized() const;
 	ENGINE_API void Update();
@@ -535,4 +536,4 @@ private:
 
 	friend class AWorldPartitionReplay;
 	friend class UWorldPartitionSubsystem;
-};
+};
\ No newline at end of file
diff --git a/Engine/Source/Runtime/Engine/Public/WorldPartition/WorldPartitionSubsystem.h b/Engine/Source/Runtime/Engine/Public/WorldPartition/WorldPartitionSubsystem.h
--- a/Engine/Source/Runtime/Engine/Public/WorldPartition/WorldPartitionSubsystem.h
+++ b/Engine/Source/Runtime/Engine/Public/WorldPartition/WorldPartitionSubsystem.h
@@ -163,16 +163,18 @@ private:
 	ENGINE_API void GetStreamingSources(const UWorldPartition* InWorldPartition, TArray<FWorldPartitionStreamingSource>& OutStreamingSources) const;
 	uint32 GetStreamingSourcesHash() const { return StreamingSourcesHash; }
 
-	ENGINE_API void OnWorldPartitionInitialized(UWorldPartition* InWorldPartition);
-	ENGINE_API void OnWorldPartitionUninitialized(UWorldPartition* InWorldPartition);
-	ENGINE_API void OnLevelStreamingTargetStateChanged(UWorld* InWorld, const ULevelStreaming* InStreamingLevel, ULevel* InLevelIfLoaded, ELevelStreamingState InCurrentState, ELevelStreamingTargetState InPrevTarget, ELevelStreamingTargetState InNewTarget);
-	ENGINE_API void OnLevelBeginMakingVisible(UWorld* InWorld, const ULevelStreaming* InStreamingLevel, ULevel* InLoadedLevel);
-	ENGINE_API void OnLevelBeginMakingInvisible(UWorld* InWorld, const ULevelStreaming* InStreamingLevel, ULevel* InLoadedLevel);
+	void OnWorldPartitionInitialized(UWorldPartition* InWorldPartition);
+	void OnWorldPartitionUninitialized(UWorldPartition* InWorldPartition);
+	void OnLevelStreamingStateChanged(UWorld* InWorld, const ULevelStreaming* InStreamingLevel, ULevel* LevelIfLoaded, ELevelStreamingState PreviousState, ELevelStreamingState NewState);
+	void OnLevelStreamingTargetStateChanged(UWorld* InWorld, const ULevelStreaming* InStreamingLevel, ULevel* InLevelIfLoaded, ELevelStreamingState InCurrentState, ELevelStreamingTargetState InPrevTarget, ELevelStreamingTargetState InNewTarget);
+	void OnLevelBeginMakingVisible(UWorld* InWorld, const ULevelStreaming* InStreamingLevel, ULevel* InLoadedLevel);
+	void OnLevelBeginMakingInvisible(UWorld* InWorld, const ULevelStreaming* InStreamingLevel, ULevel* InLoadedLevel);
 
 	static ENGINE_API void UpdateStreamingStateInternal(const UWorld* InWorld, const TArray<TObjectPtr<UWorldPartition>>& InWorldPartitions);
 
 	ENGINE_API bool IsServer() const;
 	ENGINE_API bool HasAnyWorldPartitionServerStreamingEnabled() const;
+	ENGINE_API bool HasUninitializationPendingStreamingLevels(const UWorldPartition* InWorldPartition) const;
 
 	ENGINE_API UWorldPartition* GetWorldPartition();
 	ENGINE_API const UWorldPartition* GetWorldPartition() const;
@@ -201,6 +203,9 @@ private:
 	int32 LevelStreamingContinuouslyIncrementalGCWhileLevelsPendingPurge;
 	int32 LevelStreamingForceGCAfterLevelStreamedOut;
 
+	// Tracks streaming levels of uninitialized world partition used to delay next initialization until they're done being removed
+	TMap<FSoftObjectPath, TSet<TWeakObjectPtr<ULevelStreaming>>> WorldPartitionUninitializationPendingStreamingLevels;
+
 #if WITH_EDITOR
 	bool bIsRunningConvertWorldPartitionCommandlet;
 	mutable FActorDescContainerInstanceManager ActorDescContainerInstanceManager;
