dll-exposed SetOwnerName on FStaticMeshVertexBuffers

diff --git a/Engine/Source/Runtime/Engine/Public/StaticMeshResources.h b/Engine/Source/Runtime/Engine/Public/StaticMeshResources.h
--- a/Engine/Source/Runtime/Engine/Public/StaticMeshResources.h
+++ b/Engine/Source/Runtime/Engine/Public/StaticMeshResources.h
@@ -327,7 +327,7 @@ struct FStaticMeshVertexBuffers
 	/* This is a temporary function to refactor and convert old code, do not copy this as is and try to build your data as SoA from the beginning.*/
 	void ENGINE_API InitModelVF(FLocalVertexFactory* VertexFactory);
 
-	void SetOwnerName(const FName& OwnerName);
+	void ENGINE_API SetOwnerName(const FName& OwnerName);
 };
 
 struct FAdditionalStaticMeshIndexBuffers
