Mesh streaming is no longer marked as experimental in the renderer settings UI.

Also, Skeletal Mesh Streaming (which is a separate option) is only editable if the main Mesh Streaming option is also enabled, since the code requires both to be on to stream skeletal mesh LODs.

#rnx

diff --git a/Engine/Source/Runtime/Engine/Classes/Engine/RendererSettings.h b/Engine/Source/Runtime/Engine/Classes/Engine/RendererSettings.h
--- a/Engine/Source/Runtime/Engine/Classes/Engine/RendererSettings.h
+++ b/Engine/Source/Runtime/Engine/Classes/Engine/RendererSettings.h
@@ -832,9 +832,9 @@ class URendererSettings : public UDeveloperSettings
 		ConfigRestartRequired = false))
 	uint32 bRoundRobinOcclusion : 1;
 
-	UPROPERTY(config, EditAnywhere, Category = Experimental, meta = (
+	UPROPERTY(config, EditAnywhere, Category = "Mesh Streaming", meta = (
 		ConsoleVariable="r.MeshStreaming",DisplayName="Mesh Streaming",
-		ToolTip="When enabled mesh will stream in based on what is visible on screen.",
+		ToolTip="When enabled mesh LODs will stream in based on what is visible on screen.",
 		ConfigRestartRequired = true))
 		uint32 bMeshStreaming : 1;
 
@@ -1125,12 +1125,13 @@ class URendererSettings : public UDeveloperSettings
 		ConfigRestartRequired = true))
 		uint32 bMobileSupportsGen4TAA : 1;
 
-	UPROPERTY(config, EditAnywhere, Category="Experimental|LOD Streaming|Skeletal Mesh", meta=(
+	UPROPERTY(config, EditAnywhere, Category="Mesh Streaming|Skeletal Mesh", meta=(
+		EditCondition = "bMeshStreaming",
 		DisplayName="Stream LODs by default",
 		ToolTip="Whether to stream skeletal mesh LODs by default."))
 	FPerPlatformBool bStreamSkeletalMeshLODs;
 
-	UPROPERTY(config, EditAnywhere, Category="Experimental|LOD Streaming|Skeletal Mesh", meta=(
+	UPROPERTY(config, EditAnywhere, Category="Mesh Streaming|Skeletal Mesh", meta=(
 		DisplayName="Discard optional LODs",
 		ToolTip="Whether to discard skeletal mesh LODs below minimum LOD levels at cook time."))
 	FPerPlatformBool bDiscardSkeletalMeshOptionalLODs;
