Fixes for using ForceLODLevel of StaticMeshComponent on HISMs

diff --git a/Engine/Source/Runtime/Engine/Private/HierarchicalInstancedStaticMesh.cpp b/Engine/Source/Runtime/Engine/Private/HierarchicalInstancedStaticMesh.cpp
--- a/Engine/Source/Runtime/Engine/Private/HierarchicalInstancedStaticMesh.cpp
+++ b/Engine/Source/Runtime/Engine/Private/HierarchicalInstancedStaticMesh.cpp
@@ -1545,6 +1545,9 @@ void FHierarchicalStaticMeshSceneProxy::GetDynamicMeshElements(const TArray<cons
 	SCOPE_CYCLE_COUNTER(STAT_HISMCGetDynamicMeshElement);
 
 	bool bMultipleSections = bDitheredLODTransitions && CVarDitheredLOD.GetValueOnRenderThread() > 0;
+	// Disable multiple selections when forced LOD is set
+	bMultipleSections = bMultipleSections && ForcedLodModel <= 0 && CVarForceLOD.GetValueOnRenderThread() < 0;
+
 	bool bSingleSections = !bMultipleSections;
 	bool bOverestimate = CVarOverestimateLOD.GetValueOnRenderThread() > 0;
 
@@ -1799,7 +1802,7 @@ void FHierarchicalStaticMeshSceneProxy::GetDynamicMeshElements(const TArray<cons
 					}
 					int32 UseMaxLOD = InstanceParams.LODs;
 
-					int32 Force = CVarForceLOD.GetValueOnRenderThread();
+					int32 Force = CVarForceLOD.GetValueOnRenderThread() >= 0 ? CVarForceLOD.GetValueOnRenderThread() : (ForcedLodModel > 0 ? ForcedLodModel : -1); 
 					if (Force >= 0)
 					{
 						UseMinLOD = FMath::Clamp(Force, 0, InstanceParams.LODs - 1);
@@ -1877,7 +1880,7 @@ void FHierarchicalStaticMeshSceneProxy::GetDynamicMeshElements(const TArray<cons
 				{
 					const int32 NumLODs = RenderData->LODResources.Num();
 
-					int32 Force = CVarForceLOD.GetValueOnRenderThread();
+					int32 Force = CVarForceLOD.GetValueOnRenderThread() >= 0 ? CVarForceLOD.GetValueOnRenderThread() : (ForcedLodModel > 0 ? ForcedLodModel : -1);
 					if (Force >= 0)
 					{
 						Force = FMath::Clamp(Force, 0, NumLODs - 1);
