Add cook stat so we can get metrics on UMaterialInstance::PostLoad time.

- This accounts for 7% critical path time in certain project's cooks.

diff --git a/Engine/Source/Runtime/Engine/Private/Materials/MaterialInstance.cpp b/Engine/Source/Runtime/Engine/Private/Materials/MaterialInstance.cpp
--- a/Engine/Source/Runtime/Engine/Private/Materials/MaterialInstance.cpp
+++ b/Engine/Source/Runtime/Engine/Private/Materials/MaterialInstance.cpp
@@ -33,6 +33,7 @@
 #include "Materials/MaterialInstanceSupport.h"
 #include "Engine/SubsurfaceProfile.h"
 #include "Engine/SpecularProfile.h"
+#include "ProfilingDebugging/CookStats.h"
 #include "ProfilingDebugging/LoadTimeTracker.h"
 #include "ObjectCacheEventSink.h"
 #include "Interfaces/ITargetPlatform.h"
@@ -62,6 +63,21 @@ DECLARE_CYCLE_STAT(TEXT("MaterialInstance CopyMatInstParams"), STAT_MaterialInst
 DECLARE_CYCLE_STAT(TEXT("MaterialInstance Serialize"), STAT_MaterialInstance_Serialize, STATGROUP_Shaders);
 DECLARE_CYCLE_STAT(TEXT("MaterialInstance CopyUniformParamsInternal"), STAT_MaterialInstance_CopyUniformParamsInternal, STATGROUP_Shaders);
 
+#if ENABLE_COOK_STATS
+#include "ProfilingDebugging/ScopedTimers.h"
+namespace MaterialInstanceCookStats
+{
+	static double MaterialInstancePostLoadSec = 0.0;
+
+	static FCookStatsManager::FAutoRegisterCallback RegisterMaterialInstanceCookStats([](FCookStatsManager::AddStatFuncRef AddStat)
+	{
+		AddStat(TEXT("Material"), FCookStatsManager::CreateKeyValueArray(
+			TEXT("MaterialInstancePostLoadSec"), MaterialInstancePostLoadSec
+		));
+	});
+}
+#endif
+
 // This flag controls whether MaterialInstances parents should be restricted to be either uncooked, to be
 // user defined, part of the engine or part of the base game.
 ENGINE_API bool bEnableRestrictiveMaterialInstanceParents = false;
@@ -2978,6 +2994,7 @@ void UMaterialInstance::PostLoad()
 {
 	LLM_SCOPE(ELLMTag::MaterialInstance);
 	SCOPED_LOADTIMER(MaterialInstancePostLoad);
+	COOK_STAT(FScopedDurationTimer BlockingTimer(MaterialInstanceCookStats::MaterialInstancePostLoadSec));
 
 #if WITH_EDITORONLY_DATA // fixup serialization before everything else
 	if (IsEditorOnlyDataValid() && !GetEditorOnlyData()->StaticParameters.StaticSwitchParameters_DEPRECATED.IsEmpty())
