Iris
* Added UNetDriver::RestartIrisSystem function that can be reused to restart Iris classes while no-one is connected to it.
* Useful if you want to apply runtime config changes like an hotfix
* Support recreating Iris after downloading an hotfix change

#rnx

diff --git a/Engine/Source/Runtime/Engine/Classes/Engine/NetDriver.h b/Engine/Source/Runtime/Engine/Classes/Engine/NetDriver.h
--- a/Engine/Source/Runtime/Engine/Classes/Engine/NetDriver.h
+++ b/Engine/Source/Runtime/Engine/Classes/Engine/NetDriver.h
@@ -1785,6 +1785,13 @@ public:
 
 	/** Set a previously initialized IrisSystem into this NetDriver */
 	ENGINE_API void RestoreIrisSystem(UReplicationSystem* InReplicationSystem);
+
+	/**
+	 * Destroy and recreate the iris replication system for an active netdrive.
+	 * This will re-add all existing replicated actors back in the system.
+	 * Useful if you need to reapply hotfix configs downloaded post-initialization.
+	 */
+	ENGINE_API void RestartIrisSystem();
 #endif // UE_WITH_IRIS
 
 	template<class T>
@@ -2069,6 +2076,7 @@ private:
 #if UE_WITH_IRIS
 	void InitIrisSettings(FName NewDriverName);
 	void SetReplicationSystem(UReplicationSystem* ReplicationSystem);
+	void CreateReplicationSystem(bool bInitAsClient);
 	void UpdateReplicationViews() const;
 	void SendClientMoveAdjustments();
 #endif
diff --git a/Engine/Source/Runtime/Engine/Private/NetDriver.cpp b/Engine/Source/Runtime/Engine/Private/NetDriver.cpp
--- a/Engine/Source/Runtime/Engine/Private/NetDriver.cpp
+++ b/Engine/Source/Runtime/Engine/Private/NetDriver.cpp
@@ -1484,32 +1484,7 @@ bool UNetDriver::InitBase(bool bInitAsClient, FNetworkNotify* InNotify, const FU
 #if UE_WITH_IRIS
 	if (IsUsingIrisReplication() && !ReplicationSystem)
 	{
-		LLM_SCOPE_BYTAG(Iris);
-
-		if (ensureAlways(InitReplicationBridgeClass()))
-		{
-			const bool bAllowSend = !bInitAsClient;
-
-			UReplicationBridge* ReplicationBridge = NewObject<UActorReplicationBridge>(GetTransientPackage(), ReplicationBridgeClass);
-			if (ReplicationBridge)
-			{
-				ReplicationBridge->SetNetDriver(this);
-
-				// Create ReplicationSystem
-				UReplicationSystem::FReplicationSystemParams Params;
-				Params.ReplicationBridge = ReplicationBridge;
-				Params.bIsServer = !bInitAsClient;
-				Params.bAllowObjectReplication = !bInitAsClient;
-
-				UE::Net::Private::ApplyReplicationSystemConfig(ReplicationSystemConfig, Params);
-
-				SetReplicationSystem(UE::Net::FReplicationSystemFactory::CreateReplicationSystem(Params));
-			}
-			else
-			{
-				UE_LOG(LogNet, Error, TEXT("Failed to initialize ReplicationSystem"));
-			}
-		}
+		CreateReplicationSystem(bInitAsClient);
 	}
 #endif // UE_WITH_IRIS
 	
@@ -6288,6 +6263,61 @@ void UNetDriver::RestoreIrisSystem(UReplicationSystem* InReplicationSystem)
 	bSkipBeginReplicationForWorld = true;
 }
 
+void UNetDriver::RestartIrisSystem()
+{
+	if (ReplicationSystem == nullptr)
+	{
+		ensureMsgf(false, TEXT("RestartIrisSystem called while no system existed."));
+		return;
+	}
+
+	if (ClientConnections.Num() > 0)
+	{
+		ensureMsgf(false, TEXT("RestartIrisSystem called while there were active connections."));
+		return;
+	}
+
+	QUICK_SCOPE_CYCLE_COUNTER(STAT_NetDriver_RestartIrisSystem);
+
+	UE::Net::FReplicationSystemFactory::DestroyReplicationSystem(ReplicationSystem);
+	ReplicationSystem = nullptr;
+
+	CreateReplicationSystem(!IsServer());
+}
+
+void UNetDriver::CreateReplicationSystem(bool bInitAsClient)
+{
+	LLM_SCOPE_BYTAG(Iris);
+
+	const bool bBridgeClassExists = InitReplicationBridgeClass();
+
+	if (!bBridgeClassExists)
+	{
+		ensureMsgf(false, TEXT("InitReplicationBridgeClass could not load configured class %s"), *ReplicationBridgeClassName);
+		return;
+	}
+
+	UReplicationBridge* ReplicationBridge = NewObject<UActorReplicationBridge>(GetTransientPackage(), ReplicationBridgeClass);
+	if (ReplicationBridge)
+	{
+		ReplicationBridge->SetNetDriver(this);
+
+		// Create ReplicationSystem
+		UReplicationSystem::FReplicationSystemParams Params;
+		Params.ReplicationBridge = ReplicationBridge;
+		Params.bIsServer = !bInitAsClient;
+		Params.bAllowObjectReplication = !bInitAsClient;
+
+		UE::Net::Private::ApplyReplicationSystemConfig(ReplicationSystemConfig, Params);
+
+		SetReplicationSystem(UE::Net::FReplicationSystemFactory::CreateReplicationSystem(Params));
+	}
+	else
+	{
+		UE_LOG(LogNet, Error, TEXT("Failed to initialize ReplicationSystem"));
+	}
+}
+
 void UNetDriver::UpdateReplicationViews() const
 {
 	using namespace UE::Net;
