Remove calls to SetNetVersionsOnArchive() and SetNetVersionsFromArchive() in FInBunch

diff --git a/Engine/Source/Runtime/Core/Public/Serialization/Archive.h b/Engine/Source/Runtime/Core/Public/Serialization/Archive.h
--- a/Engine/Source/Runtime/Core/Public/Serialization/Archive.h
+++ b/Engine/Source/Runtime/Core/Public/Serialization/Archive.h
@@ -232,10 +232,10 @@ public:
 	}
 
 	/** Returns the engine-global network protocol version for this archive. */
-	uint32 EngineNetVer() const;
+	virtual uint32 EngineNetVer() const;
 
 	/** Returns the game-specific network protocol version for this archive. */
-	uint32 GameNetVer() const;
+	virtual uint32 GameNetVer() const;
 
 	/**
 	 * Queries a custom version from the archive.  If the archive is being used to write, the custom version must have already been registered.
diff --git a/Engine/Source/Runtime/Engine/Classes/Engine/NetConnection.h b/Engine/Source/Runtime/Engine/Classes/Engine/NetConnection.h
--- a/Engine/Source/Runtime/Engine/Classes/Engine/NetConnection.h
+++ b/Engine/Source/Runtime/Engine/Classes/Engine/NetConnection.h
@@ -1584,7 +1584,7 @@ private:
 	void WriteFinalPacketInfo(FBitWriter& Writer, double PacketSentTimeInS);
 	
 	/** Read extended packet header information (ServerFrameTime) */
-	bool ReadPacketInfo(FBitReader& Reader, bool bHasPacketInfoPayload);
+	bool ReadPacketInfo(FBitReader& Reader, bool bHasPacketInfoPayload, FEngineNetworkCustomVersion::Type EngineNetVer);
 
 	/** Packet was acknowledged as delivered */
 	void ReceivedAck(int32 AckPacketId, FChannelsToClose& OutChannelsToClose);
diff --git a/Engine/Source/Runtime/Engine/Private/DataBunch.cpp b/Engine/Source/Runtime/Engine/Private/DataBunch.cpp
--- a/Engine/Source/Runtime/Engine/Private/DataBunch.cpp
+++ b/Engine/Source/Runtime/Engine/Private/DataBunch.cpp
@@ -10,8 +10,7 @@
 #include "Net/Core/Trace/NetTrace.h"
 #include "Net/Core/Trace/Private/NetTraceInternal.h"
 
-const int32 MAX_BUNCH_SIZE = 1024 * 1024; 
-
+const int32 MAX_BUNCH_SIZE = 1024 * 1024;
 
 /*-----------------------------------------------------------------------------
 	FInBunch implementation.
@@ -41,9 +40,6 @@ FInBunch::FInBunch( UNetConnection* InConnection, uint8* Src, int64 CountBits )
 	check(Connection);
 	// Match the byte swapping settings of the connection
 	SetByteSwapping(Connection->bNeedsByteSwapping);
-
-	// Copy network version info
-	Connection->SetNetVersionsOnArchive(*this);
 }
 
 /** Copy constructor but with optional parameter to not copy buffer */
@@ -66,10 +62,6 @@ FInBunch::FInBunch( FInBunch &InBunch, bool CopyBuffer )
 	bHasMustBeMappedGUIDs =	InBunch.bHasMustBeMappedGUIDs;
 	bIgnoreRPCs = InBunch.bIgnoreRPCs;
 	CloseReason = InBunch.CloseReason;
-
-	// Copy network version info
-	SetNetVersionsFromArchive(InBunch);
-
 	PackageMap = InBunch.PackageMap;
 
 	if (CopyBuffer)
@@ -91,6 +83,26 @@ void FInBunch::CountMemory(FArchive& Ar) const
 	}
 }
 
+uint32 FInBunch::EngineNetVer() const
+{
+	return Connection->GetNetworkCustomVersion(FEngineNetworkCustomVersion::Guid);
+}
+
+uint32 FInBunch::GameNetVer() const
+{
+	return Connection->GetNetworkCustomVersion(FGameNetworkCustomVersion::Guid);
+}
+
+void FInBunch::SetEngineNetVer(const uint32 InEngineNetVer)
+{
+	ensureMsgf(false, TEXT("SetEngineNetVer should not be called on FInBunch"));
+}
+
+void FInBunch::SetGameNetVer(const uint32 InGameNetVer)
+{
+	ensureMsgf(false, TEXT("SetGameNetVer should not be called on FInBunch"));
+}
+
 /*-----------------------------------------------------------------------------
 	FOutBunch implementation.
 -----------------------------------------------------------------------------*/
diff --git a/Engine/Source/Runtime/Engine/Private/NetConnection.cpp b/Engine/Source/Runtime/Engine/Private/NetConnection.cpp
--- a/Engine/Source/Runtime/Engine/Private/NetConnection.cpp
+++ b/Engine/Source/Runtime/Engine/Private/NetConnection.cpp
@@ -2564,7 +2564,7 @@ void UNetConnection::WriteFinalPacketInfo(FBitWriter& Writer, const double Packe
 	CurrentMark.PopWithoutClear(Writer);
 }
 
-bool UNetConnection::ReadPacketInfo(FBitReader& Reader, bool bHasPacketInfoPayload)
+bool UNetConnection::ReadPacketInfo(FBitReader& Reader, bool bHasPacketInfoPayload, FEngineNetworkCustomVersion::Type EngineNetVer)
 {
 	using namespace UE::Net;
 
@@ -2593,7 +2593,7 @@ bool UNetConnection::ReadPacketInfo(FBitReader& Reader, bool bHasPacketInfoPaylo
 		bLastHasServerFrameTime = bHasServerFrameTime;
 	}
 
-	if (Reader.EngineNetVer() < FEngineNetworkCustomVersion::JitterInHeader)
+	if (EngineNetVer < FEngineNetworkCustomVersion::JitterInHeader)
 	{
 		uint8 RemoteInKBytesPerSecondByte = 0;
 		Reader << RemoteInKBytesPerSecondByte;
@@ -2760,6 +2760,8 @@ void UNetConnection::ReceivedPacket( FBitReader& Reader, bool bIsReinjectedPacke
 
 	FChannelsToClose ChannelsToClose;
 
+	const FEngineNetworkCustomVersion::Type PacketEngineNetVer = static_cast<FEngineNetworkCustomVersion::Type>(Reader.EngineNetVer());
+
 	if (IsInternalAck())
 	{
 		++InPacketId;
@@ -2779,7 +2781,7 @@ void UNetConnection::ReceivedPacket( FBitReader& Reader, bool bIsReinjectedPacke
 
 		bool bHasPacketInfoPayload = true;
 
-		if (Reader.EngineNetVer() >= FEngineNetworkCustomVersion::JitterInHeader)
+		if (PacketEngineNetVer >= FEngineNetworkCustomVersion::JitterInHeader)
 		{
 			bHasPacketInfoPayload = Reader.ReadBit() == 1u;
 
@@ -2940,7 +2942,7 @@ void UNetConnection::ReceivedPacket( FBitReader& Reader, bool bIsReinjectedPacke
 		PacketNotify.Update(Header, HandlePacketNotification);
 
 		// Extra information associated with the header (read only after acks have been processed)
-		if (PacketSequenceDelta > 0 && !ReadPacketInfo(Reader, bHasPacketInfoPayload))
+		if (PacketSequenceDelta > 0 && !ReadPacketInfo(Reader, bHasPacketInfoPayload, PacketEngineNetVer))
 		{
 			UE_LOG(LogNet, Warning, TEXT("Failed to read extra PacketHeader information"));
 
@@ -2990,7 +2992,7 @@ void UNetConnection::ReceivedPacket( FBitReader& Reader, bool bIsReinjectedPacke
 		while( !Reader.AtEnd() && GetConnectionState()!=USOCK_Closed )
 		{
 			// For demo backwards compatibility, old replays still have this bit
-			if (IsInternalAck() && GetNetworkCustomVersion(FEngineNetworkCustomVersion::Guid) < FEngineNetworkCustomVersion::AcksIncludedInHeader)
+			if (IsInternalAck() && PacketEngineNetVer < FEngineNetworkCustomVersion::AcksIncludedInHeader)
 			{
 				const bool IsAckDummy = Reader.ReadBit() == 1u;
 			}
@@ -3008,7 +3010,7 @@ void UNetConnection::ReceivedPacket( FBitReader& Reader, bool bIsReinjectedPacke
 				Bunch.bOpen					= bControl ? Reader.ReadBit() : 0;
 				Bunch.bClose				= bControl ? Reader.ReadBit() : 0;
 			
-				if (Bunch.EngineNetVer() < FEngineNetworkCustomVersion::ChannelCloseReason)
+				if (PacketEngineNetVer < FEngineNetworkCustomVersion::ChannelCloseReason)
 				{
 					const uint8 bDormant = Bunch.bClose ? Reader.ReadBit() : 0;
 					Bunch.CloseReason = bDormant ? EChannelCloseReason::Dormancy : EChannelCloseReason::Destroyed;
@@ -3025,7 +3027,7 @@ void UNetConnection::ReceivedPacket( FBitReader& Reader, bool bIsReinjectedPacke
 				PRAGMA_ENABLE_DEPRECATION_WARNINGS
 				Bunch.bReliable				= Reader.ReadBit();
 
-				if (Bunch.EngineNetVer() < FEngineNetworkCustomVersion::MaxActorChannelsCustomization)
+				if (PacketEngineNetVer < FEngineNetworkCustomVersion::MaxActorChannelsCustomization)
 				{
 					static const int OLD_MAX_ACTOR_CHANNELS = 10240;
 					Bunch.ChIndex = Reader.ReadInt(OLD_MAX_ACTOR_CHANNELS);
@@ -3051,7 +3053,7 @@ void UNetConnection::ReceivedPacket( FBitReader& Reader, bool bIsReinjectedPacke
 
 				// if flag is set, remap channel index values, we're fast forwarding a replay checkpoint
 				// and there should be no bunches for existing channels
-				if (IsInternalAck() && bAllowExistingChannelIndex && (Bunch.EngineNetVer() >= FEngineNetworkCustomVersion::ReplayDormancy))
+				if (IsInternalAck() && bAllowExistingChannelIndex && (PacketEngineNetVer >= FEngineNetworkCustomVersion::ReplayDormancy))
 				{
 					if (ChannelIndexMap.Contains(Bunch.ChIndex))
 					{
@@ -3130,7 +3132,7 @@ void UNetConnection::ReceivedPacket( FBitReader& Reader, bool bIsReinjectedPacke
 				Bunch.bPartialInitial = Bunch.bPartial ? Reader.ReadBit() : 0;
 				Bunch.bPartialFinal = Bunch.bPartial ? Reader.ReadBit() : 0;
 
-				if (Bunch.EngineNetVer() < FEngineNetworkCustomVersion::ChannelNames)
+				if (PacketEngineNetVer < FEngineNetworkCustomVersion::ChannelNames)
 				{
 					uint32 ChType = (Bunch.bReliable || Bunch.bOpen) ? Reader.ReadInt(CHTYPE_MAX) : CHTYPE_None;
 					switch (ChType)
@@ -3308,7 +3310,7 @@ void UNetConnection::ReceivedPacket( FBitReader& Reader, bool bIsReinjectedPacke
 				// In that case, we can generally ignore these bunches.
 				if (IsInternalAck() && bAllowExistingChannelIndex)
 				{
-					if (Bunch.EngineNetVer() < FEngineNetworkCustomVersion::ReplayDormancy)
+					if (PacketEngineNetVer < FEngineNetworkCustomVersion::ReplayDormancy)
 					{
 						if (Channel)
 						{
@@ -5457,6 +5459,7 @@ void UNetConnection::NotifyActorChannelCleanedUp(UActorChannel* Channel, EChanne
 
 void UNetConnection::SetNetVersionsOnArchive(FArchive& Ar) const
 {
+	QUICK_SCOPE_CYCLE_COUNTER(STAT_NetSetVersionsOnArchive);
 	Ar.SetEngineNetVer(GetNetworkCustomVersion(FEngineNetworkCustomVersion::Guid));
 	Ar.SetGameNetVer(GetNetworkCustomVersion(FGameNetworkCustomVersion::Guid));
 
diff --git a/Engine/Source/Runtime/Engine/Public/Net/DataBunch.h b/Engine/Source/Runtime/Engine/Public/Net/DataBunch.h
--- a/Engine/Source/Runtime/Engine/Public/Net/DataBunch.h
+++ b/Engine/Source/Runtime/Engine/Public/Net/DataBunch.h
@@ -182,6 +182,11 @@ public:
 	ENGINE_API PRAGMA_ENABLE_DEPRECATION_WARNINGS
 
 	virtual void CountMemory(FArchive& Ar) const override;
+		
+	virtual uint32 EngineNetVer() const override;
+	virtual uint32 GameNetVer() const override;
+	virtual void SetEngineNetVer(const uint32 InEngineNetVer) override;
+	virtual void SetGameNetVer(const uint32 InGameNetVer) override;
 };
 
 /** out bunch for the control channel (special restrictions) */
