WorldPartition - PartionedActor: Avoid division by zero crash if grid size is 0

diff --git a/Engine/Source/Runtime/Engine/Private/WorldPartition/ActorPartition/PartitionActorDesc.cpp b/Engine/Source/Runtime/Engine/Private/WorldPartition/ActorPartition/PartitionActorDesc.cpp
--- a/Engine/Source/Runtime/Engine/Private/WorldPartition/ActorPartition/PartitionActorDesc.cpp
+++ b/Engine/Source/Runtime/Engine/Private/WorldPartition/ActorPartition/PartitionActorDesc.cpp
@@ -27,9 +27,7 @@ void FPartitionActorDesc::Init(const AActor* InActor)
 		GridGuid = PartitionActor->GetGridGuid();
 	
 		const FVector ActorLocation = InActor->GetActorLocation();
-		GridIndexX = FMath::FloorToInt(ActorLocation.X / GridSize);
-		GridIndexY = FMath::FloorToInt(ActorLocation.Y / GridSize);
-		GridIndexZ = FMath::FloorToInt(ActorLocation.Z / GridSize);
+		SetGridIndices(ActorLocation.X, ActorLocation.Y, ActorLocation.Z);
 	}
 }
 
@@ -85,4 +83,18 @@ void FPartitionActorDesc::TransferWorldData(const FWorldPartitionActorDesc* From
 	GridIndexZ = FromPartitionActorDesc->GridIndexZ;
 }
 
+void FPartitionActorDesc::SetGridIndices(double X, double Y, double Z)
+{
+	if (GridSize > 0)
+	{
+		GridIndexX = FMath::FloorToInt(X / GridSize);
+		GridIndexY = FMath::FloorToInt(Y / GridSize);
+		GridIndexZ = FMath::FloorToInt(Z / GridSize);
+	}
+	else
+	{
+		GridIndexX = GridIndexY = GridIndexZ = 0;
+	}
+}
+
 #endif
diff --git a/Engine/Source/Runtime/Engine/Private/WorldPartition/Landscape/LandscapeActorDesc.cpp b/Engine/Source/Runtime/Engine/Private/WorldPartition/Landscape/LandscapeActorDesc.cpp
--- a/Engine/Source/Runtime/Engine/Private/WorldPartition/Landscape/LandscapeActorDesc.cpp
+++ b/Engine/Source/Runtime/Engine/Private/WorldPartition/Landscape/LandscapeActorDesc.cpp
@@ -13,9 +13,7 @@ void FLandscapeActorDesc::Init(const AActor* InActor)
 
 	const ALandscapeProxy* LandscapeProxy = CastChecked<ALandscapeProxy>(InActor);
 	check(LandscapeProxy);
-	GridIndexX = LandscapeProxy->LandscapeSectionOffset.X / (int32)LandscapeProxy->GridSize;
-	GridIndexY = LandscapeProxy->LandscapeSectionOffset.Y / (int32)LandscapeProxy->GridSize;
-	GridIndexZ = 0;
+	SetGridIndices(LandscapeProxy->LandscapeSectionOffset.X, LandscapeProxy->LandscapeSectionOffset.Y, 0);
 
 	if (!bIsDefaultActorDesc)
 	{
@@ -38,8 +36,7 @@ void FLandscapeActorDesc::Serialize(FArchive& Ar)
 	{
 		if (Ar.CustomVer(FUE5MainStreamObjectVersion::GUID) < FUE5MainStreamObjectVersion::FLandscapeActorDescFixupGridIndices)
 		{
-			GridIndexX = (int32)(GridIndexX * GridSize) / (int32)GridSize;
-			GridIndexY = (int32)(GridIndexY * GridSize) / (int32)GridSize;
+			SetGridIndices(GridIndexX * GridSize, GridIndexY * GridSize, 0);
 		}
 
 		if (Ar.CustomVer(FUE5ReleaseStreamObjectVersion::GUID) >= FUE5ReleaseStreamObjectVersion::WorldPartitionLandscapeActorDescSerializeLandscapeActorGuid)
