Added missing file
#rnx

diff --git a/Engine/Source/Runtime/Engine/Public/WorldPartition/ActorPartition/PartitionActorDesc.h b/Engine/Source/Runtime/Engine/Public/WorldPartition/ActorPartition/PartitionActorDesc.h
--- a/Engine/Source/Runtime/Engine/Public/WorldPartition/ActorPartition/PartitionActorDesc.h
+++ b/Engine/Source/Runtime/Engine/Public/WorldPartition/ActorPartition/PartitionActorDesc.h
@@ -27,5 +27,7 @@ protected:
 	ENGINE_API virtual void Serialize(FArchive& Ar) override;
 	ENGINE_API virtual FBox GetEditorBounds() const override;
 	ENGINE_API virtual void TransferWorldData(const FWorldPartitionActorDesc* From) override;
+
+	ENGINE_API void SetGridIndices(double LocationX, double LocationY, double LocationZ);
 };
 #endif
