Fix typo in log

diff --git a/Engine/Source/Runtime/Engine/Private/AudioDecompress.cpp b/Engine/Source/Runtime/Engine/Private/AudioDecompress.cpp
--- a/Engine/Source/Runtime/Engine/Private/AudioDecompress.cpp
+++ b/Engine/Source/Runtime/Engine/Private/AudioDecompress.cpp
@@ -294,7 +294,7 @@ bool IStreamedCompressedInfo::StreamCompressedData(uint8* Destination, bool bLoo
 						else // Seek failed (off the end of the chunk).
 						{
 							const float TimeInSeconds = (float)StreamSeekToAudioFrames / StreamingSoundWave->GetSampleRate();
-							UE_LOG(LogAudio, Log, TEXT("Failed seeking to %2.2f seconds as its off then end of the stream. Wave=%s"), TimeInSeconds, *StreamingSoundWave->GetFName().ToString());
+							UE_LOG(LogAudio, Log, TEXT("Failed seeking to %2.2f seconds as it's off the end of the stream. Wave=%s"), TimeInSeconds, *StreamingSoundWave->GetFName().ToString());
 						}
 					}
 				}
