Fix Source Busses hitting and failing the DDC

diff --git a/Engine/Source/Runtime/Engine/Classes/Sound/SoundSourceBus.h b/Engine/Source/Runtime/Engine/Classes/Sound/SoundSourceBus.h
--- a/Engine/Source/Runtime/Engine/Classes/Sound/SoundSourceBus.h
+++ b/Engine/Source/Runtime/Engine/Classes/Sound/SoundSourceBus.h
@@ -54,9 +54,24 @@ public:
 	}
 
 	ENGINE_API virtual void PostEditChangeProperty(FPropertyChangedEvent& PropertyChangedEvent) override;
+
+	// Override the SoundWave behavior so that we don't ask DDC for compressed data.
+	virtual void BeginCacheForCookedPlatformData(const ITargetPlatform* TargetPlatform) override {};
+	virtual bool IsCachedCookedPlatformDataLoaded(const ITargetPlatform* TargetPlatform) override { return true; };
+	virtual void ClearCachedCookedPlatformData(const ITargetPlatform* TargetPlatform) override {};
+
 #endif // WITH_EDITOR
 	//~ End UObject Interface.
 
+	//~ Begin USoundWave Interface
+#if WITH_EDITORONLY_DATA
+	virtual void CachePlatformData(bool bAsyncCache = false) override {};
+	virtual void BeginCachePlatformData() override {}
+	virtual void FinishCachePlatformData() override {};	
+#endif	//WITH_EDITOR_ONLY_DATA
+	virtual void SerializeCookedPlatformData(FArchive& Ar) override {}
+	//~ End USoundWave Interface
+
 	//~ Begin USoundBase Interface.
 	ENGINE_API virtual bool IsPlayable() const override;
 	ENGINE_API virtual float GetDuration() const override;
diff --git a/Engine/Source/Runtime/Engine/Classes/Sound/SoundWave.h b/Engine/Source/Runtime/Engine/Classes/Sound/SoundWave.h
--- a/Engine/Source/Runtime/Engine/Classes/Sound/SoundWave.h
+++ b/Engine/Source/Runtime/Engine/Classes/Sound/SoundWave.h
@@ -1158,7 +1158,7 @@ public:
 	/**
 	 * Serializes cooked platform data.
 	 */
-	ENGINE_API void SerializeCookedPlatformData(class FArchive& Ar);
+	ENGINE_API virtual void SerializeCookedPlatformData(class FArchive& Ar);
 
 	/*
 	* Returns a sample rate if there is a specific sample rate override for this platform, -1.0 otherwise.
@@ -1213,17 +1213,17 @@ public:
 	/**
 	 * Caches platform data for the sound.
 	 */
-	ENGINE_API void CachePlatformData(bool bAsyncCache = false);
+	ENGINE_API virtual void CachePlatformData(bool bAsyncCache = false);
 
 	/**
 	 * Begins caching platform data in the background.
 	 */
-	ENGINE_API void BeginCachePlatformData();
+	ENGINE_API virtual void BeginCachePlatformData();
 
 	/**
 	 * Blocks on async cache tasks and prepares platform data for use.
 	 */
-	ENGINE_API void FinishCachePlatformData();
+	ENGINE_API virtual void FinishCachePlatformData();
 
 	/**
 	 * Forces platform data to be rebuilt.
