Remove ensure in SetEngineNetVer to fix BH

diff --git a/Engine/Source/Runtime/Engine/Private/DataBunch.cpp b/Engine/Source/Runtime/Engine/Private/DataBunch.cpp
--- a/Engine/Source/Runtime/Engine/Private/DataBunch.cpp
+++ b/Engine/Source/Runtime/Engine/Private/DataBunch.cpp
@@ -95,12 +95,12 @@ uint32 FInBunch::GameNetVer() const
 
 void FInBunch::SetEngineNetVer(const uint32 InEngineNetVer)
 {
-	ensureMsgf(false, TEXT("SetEngineNetVer should not be called on FInBunch"));
+	//ensureMsgf(false, TEXT("SetEngineNetVer should not be called on FInBunch"));
 }
 
 void FInBunch::SetGameNetVer(const uint32 InGameNetVer)
 {
-	ensureMsgf(false, TEXT("SetGameNetVer should not be called on FInBunch"));
+	//ensureMsgf(false, TEXT("SetGameNetVer should not be called on FInBunch"));
 }
 
 /*-----------------------------------------------------------------------------
