Implemented occupancy bit based compression for cooked on-disk tile data of SparseVolumeTextures. This reduces the memory footprint by 30%-50%, which also greatly benefits streaming performance. This optimization is only for the on-disk representation of derived/cooked data; the runtime GPU memory is unaffected, as is the size of uncooked assets.
Prior versions accidentally always serialized the derived data, even when not cooking. This causes problems when changing the format of derived data. This is now no longer the case, but unfortunately this breaks old assets because the size of the old serialized derived data is unknown and can't be accounted for even if we know that the asset is using an older version.

#rnx

diff --git a/Engine/Shaders/Private/SparseVolumeTexture/UpdateSparseVolumeTexture.usf b/Engine/Shaders/Private/SparseVolumeTexture/UpdateSparseVolumeTexture.usf
--- a/Engine/Shaders/Private/SparseVolumeTexture/UpdateSparseVolumeTexture.usf
+++ b/Engine/Shaders/Private/SparseVolumeTexture/UpdateSparseVolumeTexture.usf
@@ -67,6 +67,112 @@ void SparseVolumeTextureUpdateFromBufferCS(uint3 GroupId : SV_GroupID, uint Grou
 
 #endif // UPDATE_TILE_TEXTURE_FROM_BUFFER
 
+#ifdef UPDATE_TILE_TEXTURE_FROM_SPARSE_BUFFER
+
+RWTexture3D<float4> DstPhysicalTileTextureA;
+RWTexture3D<float4> DstPhysicalTileTextureB;
+Buffer<float4> SrcPhysicalTileBufferA;
+Buffer<float4> SrcPhysicalTileBufferB;
+ByteAddressBuffer OccupancyBitsBuffer;
+ByteAddressBuffer TileDataOffsetsBuffer;
+ByteAddressBuffer DstTileCoordsBuffer;
+float4 FallbackValueA;
+float4 FallbackValueB;
+uint NumTilesToCopy;
+uint BufferTileStep;
+uint NumDispatchedGroups;
+uint PaddedTileSize;
+uint CopyTexureMask;
+
+uint ComputeLocalVoxelDataLoadIndex(uint VoxelIndex, uint2 OccupancyBits, out bool bValidVoxel)
+{
+	VoxelIndex &= 63u;
+	bValidVoxel = false;
+	uint PrefixSum = 0;
+	if (VoxelIndex < 32u)
+	{
+		bValidVoxel = OccupancyBits.x & (1u << VoxelIndex);
+		PrefixSum = countbits(OccupancyBits.x & ((1u << VoxelIndex) - 1u));
+	}
+	else
+	{
+		const uint WrappedVoxelIndex = VoxelIndex - 32u;
+		bValidVoxel = OccupancyBits.y & (1u << WrappedVoxelIndex);
+		PrefixSum = countbits(OccupancyBits.x) + countbits(OccupancyBits.y & ((1u << WrappedVoxelIndex) - 1u));
+	}
+	return PrefixSum;
+}
+
+[numthreads(64, 1, 1)]
+void SparseVolumeTextureUpdateFromSparseBufferCS(uint3 GroupId : SV_GroupID, uint GroupThreadIndex : SV_GroupIndex)
+{
+	const bool bCopyTextureA = CopyTexureMask & 1u;
+	const bool bCopyTextureB = CopyTexureMask & 2u;
+	const uint GroupIdx = GroupId.x;
+	const uint NumVoxelsPerPaddedTile = PaddedTileSize * PaddedTileSize * PaddedTileSize;
+	const uint NumOccupancyWordsPerPaddedTile = (NumVoxelsPerPaddedTile + 31u) / 32u;
+
+	// Process one tile per workgroup
+	for (uint TileIdx = GroupIdx; TileIdx < NumTilesToCopy; TileIdx += NumDispatchedGroups)
+	{
+		const uint PackedDstCoord = DstTileCoordsBuffer.Load(TileIdx << 2u);
+		const uint3 DstBaseCoord = uint3(PackedDstCoord & 0xFFu, (PackedDstCoord >> 8u) & 0xFFu, (PackedDstCoord >> 16u) & 0xFFu) * PaddedTileSize;
+
+		// The start offset of the voxels and the occupancy bits of any given tile are stored per texture in a contiguous section of the buffer(s).
+		// This offset is used to jump to the offset for the current texture.
+		uint TextureTileOffset = 0;
+
+		// Iterate over both possible textures.
+		UNROLL for (uint TextureIdx = 0; TextureIdx < 2; ++TextureIdx)
+		{
+			BRANCH if ((TextureIdx == 0 && bCopyTextureA) || (TextureIdx == 1 && bCopyTextureB))
+			{
+				// The offset in the source buffer at which the tile data for this tile starts.
+				const uint TileDataOffset = TileDataOffsetsBuffer.Load((TileIdx + TextureTileOffset) << 2u);
+				uint OccupancyLoadIdx = (TileIdx + TextureTileOffset) * NumOccupancyWordsPerPaddedTile;
+				// Keep track of the number of voxels read from the source buffer (for this tile only).
+				uint NumReadVoxels = 0;
+
+				// Iterate over all voxels in the tile and copy them.
+				for (uint VoxelIdx = GroupThreadIndex; VoxelIdx < NumVoxelsPerPaddedTile; VoxelIdx += 64)
+				{
+					const uint2 OccupancyBits = OccupancyBitsBuffer.Load2(OccupancyLoadIdx << 2u);
+
+					bool bValidVoxel = false;
+					const uint LocalLoadIndex = ComputeLocalVoxelDataLoadIndex(VoxelIdx, OccupancyBits, bValidVoxel);
+					const uint3 DstCoord = DstBaseCoord + UnflattenCoord(VoxelIdx, PaddedTileSize);
+
+					BRANCH if (TextureIdx == 0)
+					{
+						float4 Voxel = FallbackValueA;
+						BRANCH if (bValidVoxel)
+						{
+							Voxel = SrcPhysicalTileBufferA.Load(TileDataOffset + NumReadVoxels + LocalLoadIndex);
+						}
+						DstPhysicalTileTextureA[DstCoord] = Voxel;
+					}
+					else
+					{
+						float4 Voxel = FallbackValueB;
+						BRANCH if (bValidVoxel)
+						{
+							Voxel = SrcPhysicalTileBufferB.Load(TileDataOffset + NumReadVoxels + LocalLoadIndex);
+						}
+						DstPhysicalTileTextureB[DstCoord] = Voxel;
+					}
+					
+					OccupancyLoadIdx += 2; // We just consumed 64 bits (2 uints)
+					NumReadVoxels += countbits(OccupancyBits.x) + countbits(OccupancyBits.y);
+				}
+
+				TextureTileOffset += BufferTileStep;
+			}
+		}
+	}
+}
+
+#endif // UPDATE_TILE_TEXTURE_FROM_SPARSE_BUFFER
+
 #ifdef UPDATE_PAGE_TABLE
 
 RWTexture3D<uint> PageTable;
diff --git a/Engine/Source/Runtime/Engine/Classes/SparseVolumeTexture/SparseVolumeTexture.h b/Engine/Source/Runtime/Engine/Classes/SparseVolumeTexture/SparseVolumeTexture.h
--- a/Engine/Source/Runtime/Engine/Classes/SparseVolumeTexture/SparseVolumeTexture.h
+++ b/Engine/Source/Runtime/Engine/Classes/SparseVolumeTexture/SparseVolumeTexture.h
@@ -47,16 +47,35 @@ struct FHeader
 };
 
 // Describes a mip level of a SVT frame in terms of the sizes and offsets of the data in the built bulk data.
+// Each mip level consists of 4 buffer sections:
+// 
+// Page table data is stored as two consecutive uint32 arrays, where the first array stores packed coordinates into the page table.
+// The second array stores the physical tile indices to be written to the page table. Only non-zero page table entries are stored.
+//
+// Since often about 30%-50% of all voxels in the physical tiles have values equal to the fallback value, the tile data (voxels) is also compressed.
+// For each tile, a bit mask with one bit per voxel is stored (occupancy bits). A set bit indicates that the voxel is actually stored and not equal
+// to the fallback value. When streaming in the data, the compressed voxel data is expanded on the GPU in the upload shader.
+// This bit mask costs 183 uint32 -> 732 bytes. A full 18x18x18 tile (16+2 padding) of 8bit unorm data is 5832 bytes, so if 30%-50% of that can be saved,
+// then paying 732 bytes for the occupancy bits should almost always be worth it.
+// 
+// Due to the above compression technique, the start offset of the voxel data for a given tile can now no longer be computed as NumVoxelsPerPaddedTile * TileIndex,
+// so a precomputed offset into the shared array of voxel data is needed (tile data offsets). This means that we need to store one additional uint32 per tile,
+// arriving at a total overhead of 736 bytes per tile.
+// 
+// Finally there is the actual voxel data, which is simply the all non-fallback voxels in contiguous memory.
+// 
 struct FMipLevelStreamingInfo
 {
 	int32 BulkOffset;
 	int32 BulkSize;
 	int32 PageTableOffset; // relative to BulkOffset
 	int32 PageTableSize;
-	int32 TileDataAOffset; // relative to BulkOffset
-	int32 TileDataASize;
-	int32 TileDataBOffset; // relative to BulkOffset
-	int32 TileDataBSize;
+	TStaticArray<int32, 2> OccupancyBitsOffset; // relative to BulkOffset
+	TStaticArray<int32, 2> OccupancyBitsSize;
+	TStaticArray<int32, 2> TileDataOffsetsOffset; // relative to BulkOffset. Per-tile offset into this mip levels voxel data
+	TStaticArray<int32, 2> TileDataOffsetsSize;
+	TStaticArray<int32, 2> TileDataOffset; // relative to BulkOffset
+	TStaticArray<int32, 2> TileDataSize;
 	int32 NumPhysicalTiles;
 };
 
diff --git a/Engine/Source/Runtime/Engine/Private/SparseVolumeTexture/SparseVolumeTexture.cpp b/Engine/Source/Runtime/Engine/Private/SparseVolumeTexture/SparseVolumeTexture.cpp
--- a/Engine/Source/Runtime/Engine/Private/SparseVolumeTexture/SparseVolumeTexture.cpp
+++ b/Engine/Source/Runtime/Engine/Private/SparseVolumeTexture/SparseVolumeTexture.cpp
@@ -55,10 +55,12 @@ FArchive& operator<<(FArchive& Ar, UE::SVT::FMipLevelStreamingInfo& MipLevelStre
 	Ar << MipLevelStreamingInfo.BulkSize;
 	Ar << MipLevelStreamingInfo.PageTableOffset;
 	Ar << MipLevelStreamingInfo.PageTableSize;
-	Ar << MipLevelStreamingInfo.TileDataAOffset;
-	Ar << MipLevelStreamingInfo.TileDataASize;
-	Ar << MipLevelStreamingInfo.TileDataBOffset;
-	Ar << MipLevelStreamingInfo.TileDataBSize;
+	Ar << MipLevelStreamingInfo.OccupancyBitsOffset;
+	Ar << MipLevelStreamingInfo.OccupancyBitsSize;
+	Ar << MipLevelStreamingInfo.TileDataOffsetsOffset;
+	Ar << MipLevelStreamingInfo.TileDataOffsetsSize;
+	Ar << MipLevelStreamingInfo.TileDataOffset;
+	Ar << MipLevelStreamingInfo.TileDataSize;
 	Ar << MipLevelStreamingInfo.NumPhysicalTiles;
 	return Ar;
 }
@@ -99,7 +101,7 @@ static int32 ComputeNumMipLevels(const FIntVector3& InResolution)
 
 static const FString& GetDerivedDataVersion()
 {
-	static FString CachedVersionString = TEXT("381AE2A9-A903-4C8F-8486-891E24D6FC71");	// Bump this if you want to ignore all cached data so far.
+	static FString CachedVersionString = TEXT("381AE2A9-A903-4C8F-8486-891E24D6FC72");	// Bump this if you want to ignore all cached data so far.
 	return CachedVersionString;
 }
 
@@ -313,6 +315,76 @@ bool FResources::Build(USparseVolumeTextureFrame* Owner, UE::Serialization::FEdi
 			return NumNonZeroEntries;
 		};
 
+		auto CompressTiles = [](int32 NumTiles, const TArray64<uint8>& PhysicalTileDataA, const TArray64<uint8>& PhysicalTileDataB, const TArrayView<EPixelFormat>& Formats, const TArrayView<FVector4f>& FallbackValues, TArray<uint8>& BulkData, FMipLevelStreamingInfo& MipStreamingInfo)
+		{
+			const int64 FormatSize[] = { GPixelFormats[Formats[0]].BlockBytes, GPixelFormats[Formats[1]].BlockBytes };
+			uint8 NullTileValuesU8[2][sizeof(float) * 4] = {};
+			int32 NumTextures = 0;
+			int32 EstimatedTileDataSize = NumTiles * SVT::NumVoxelsPerPaddedTile * (FormatSize[0] + FormatSize[1]);
+			for (int32 i = 0; i < 2; ++i)
+			{
+				if (Formats[i] != PF_Unknown)
+				{
+					++NumTextures;
+					SVT::WriteVoxel(0, NullTileValuesU8[i], Formats[i], FallbackValues[i]);
+				}
+			}
+
+			const int32 OccupancySizePerTexture = NumTiles * SVT::NumOccupancyWordsPerPaddedTile * sizeof(uint32);
+			const int32 OccupancySize = NumTextures * OccupancySizePerTexture;
+			const int32 TileDataOffsetsSize = NumTextures * NumTiles * sizeof(uint32);
+			BulkData.Reserve(BulkData.Num() + OccupancySize + TileDataOffsetsSize + EstimatedTileDataSize);
+
+			MipStreamingInfo.OccupancyBitsOffset[0] = BulkData.Num() - MipStreamingInfo.BulkOffset;
+			MipStreamingInfo.OccupancyBitsSize[0] = Formats[0] != PF_Unknown ? OccupancySizePerTexture : 0;
+			BulkData.SetNum(BulkData.Num() + MipStreamingInfo.OccupancyBitsSize[0]);
+			
+			MipStreamingInfo.OccupancyBitsOffset[1] = BulkData.Num() - MipStreamingInfo.BulkOffset;
+			MipStreamingInfo.OccupancyBitsSize[1] = Formats[1] != PF_Unknown ? OccupancySizePerTexture : 0;
+			BulkData.SetNum(BulkData.Num() + MipStreamingInfo.OccupancyBitsSize[1]);
+
+			MipStreamingInfo.TileDataOffsetsOffset[0] = BulkData.Num() - MipStreamingInfo.BulkOffset;
+			MipStreamingInfo.TileDataOffsetsSize[0] = Formats[0] != PF_Unknown ? NumTiles * sizeof(uint32) : 0;
+			BulkData.SetNum(BulkData.Num() + MipStreamingInfo.TileDataOffsetsSize[0]);
+
+			MipStreamingInfo.TileDataOffsetsOffset[1] = BulkData.Num() - MipStreamingInfo.BulkOffset;
+			MipStreamingInfo.TileDataOffsetsSize[1] = Formats[1] != PF_Unknown ? NumTiles * sizeof(uint32) : 0;
+			BulkData.SetNum(BulkData.Num() + MipStreamingInfo.TileDataOffsetsSize[1]);
+
+			FMemory::Memzero(BulkData.GetData() + MipStreamingInfo.BulkOffset + MipStreamingInfo.OccupancyBitsOffset[0], MipStreamingInfo.OccupancyBitsSize[0] + MipStreamingInfo.OccupancyBitsSize[1]);
+
+			for (int32 AttributesIdx = 0; AttributesIdx < 2; ++AttributesIdx)
+			{
+				MipStreamingInfo.TileDataOffset[AttributesIdx] = BulkData.Num() - MipStreamingInfo.BulkOffset;
+
+				if (Formats[AttributesIdx] != PF_Unknown)
+				{
+					uint32 NumValidVoxelsTotal = 0;
+					for (int64 TileIndex = 0; TileIndex < NumTiles; ++TileIndex)
+					{
+						uint32* TileDataOffsets = reinterpret_cast<uint32*>(BulkData.GetData() + MipStreamingInfo.BulkOffset + MipStreamingInfo.TileDataOffsetsOffset[AttributesIdx]);
+						TileDataOffsets[TileIndex] = NumValidVoxelsTotal;
+
+						for (int64 VoxelIndex = 0; VoxelIndex < SVT::NumVoxelsPerPaddedTile; ++VoxelIndex)
+						{
+							const uint8* Src = (AttributesIdx == 0 ? PhysicalTileDataA.GetData() : PhysicalTileDataB.GetData()) + FormatSize[AttributesIdx] * (TileIndex * SVT::NumVoxelsPerPaddedTile + VoxelIndex);
+							bool bIsFallbackValue = FMemory::Memcmp(Src, NullTileValuesU8[AttributesIdx], FormatSize[AttributesIdx]) == 0;
+							if (!bIsFallbackValue)
+							{
+								const int64 WordIndex = TileIndex * SVT::NumOccupancyWordsPerPaddedTile + (VoxelIndex / 32);
+								uint32* OccupancyBits = reinterpret_cast<uint32*>(BulkData.GetData() + MipStreamingInfo.BulkOffset + MipStreamingInfo.OccupancyBitsOffset[AttributesIdx]);
+								OccupancyBits[WordIndex] |= 1u << (static_cast<uint32>(VoxelIndex) % 32u);
+								NumValidVoxelsTotal++;
+								BulkData.Append(Src, FormatSize[AttributesIdx]);
+							}
+						}
+					}
+				}
+
+				MipStreamingInfo.TileDataSize[AttributesIdx] = BulkData.Num() - MipStreamingInfo.BulkOffset - MipStreamingInfo.TileDataOffset[AttributesIdx];
+			}
+		};
+
 		TArray<uint8> StreamableBulkData;
 		for (int32 MipLevelIdx = 0; MipLevelIdx < DerivedTextureData.MipMaps.Num(); ++MipLevelIdx)
 		{
@@ -330,13 +402,7 @@ bool FResources::Build(USparseVolumeTextureFrame* Owner, UE::Serialization::FEdi
 			const int32 NumNonZeroPageTableEntries = CompressPageTable(Mip.PageTable, MipPageTableResolution, BulkData);
 			MipStreamingInfo.PageTableSize = NumNonZeroPageTableEntries * 2 * sizeof(uint32);
 			
-			MipStreamingInfo.TileDataAOffset = BulkData.Num() - MipStreamingInfo.BulkOffset;
-			MipStreamingInfo.TileDataASize = Mip.PhysicalTileDataA.Num();
-			BulkData.Append(Mip.PhysicalTileDataA.GetData(), Mip.PhysicalTileDataA.Num());
-
-			MipStreamingInfo.TileDataBOffset = BulkData.Num() - MipStreamingInfo.BulkOffset;
-			MipStreamingInfo.TileDataBSize = Mip.PhysicalTileDataB.Num();
-			BulkData.Append(Mip.PhysicalTileDataB.GetData(), Mip.PhysicalTileDataB.Num());
+			CompressTiles(Mip.NumPhysicalTiles, Mip.PhysicalTileDataA, Mip.PhysicalTileDataB, Header.AttributesFormats, Header.FallbackValues, BulkData, MipStreamingInfo);
 
 			MipStreamingInfo.BulkSize = BulkData.Num() - MipStreamingInfo.BulkOffset;
 			MipStreamingInfo.NumPhysicalTiles = Mip.NumPhysicalTiles;
@@ -771,7 +837,12 @@ void USparseVolumeTextureFrame::Serialize(FArchive& Ar)
 
 	bool bCooked = Ar.IsCooking();
 	Ar << bCooked;
-	Resources.Serialize(Ar, this, bCooked);
+
+	// Inline the derived data for cooked builds.
+	if (bCooked && !IsTemplate() && !Ar.IsCountingMemory())
+	{
+		Resources.Serialize(Ar, this, bCooked);
+	}
 }
 
 void USparseVolumeTextureFrame::GetResourceSizeEx(FResourceSizeEx& CumulativeResourceSize)
diff --git a/Engine/Source/Runtime/Engine/Private/SparseVolumeTexture/SparseVolumeTextureStreamingManager.cpp b/Engine/Source/Runtime/Engine/Private/SparseVolumeTexture/SparseVolumeTextureStreamingManager.cpp
--- a/Engine/Source/Runtime/Engine/Private/SparseVolumeTexture/SparseVolumeTextureStreamingManager.cpp
+++ b/Engine/Source/Runtime/Engine/Private/SparseVolumeTexture/SparseVolumeTextureStreamingManager.cpp
@@ -69,6 +69,14 @@ static FAutoConsoleVariableRef CVarSVTStreamingEmptyPhysicalTileTextures(
 	ECVF_RenderThreadSafe
 );
 
+static int32 GSVTStreamingPrintMemoryStats = 0;
+static FAutoConsoleVariableRef CVarSVTStreamingPrintMemoryStats(
+	TEXT("r.SparseVolumeTexture.Streaming.PrintMemoryStats"),
+	GSVTStreamingPrintMemoryStats,
+	TEXT("Prints memory sizes of all frames of all SVTs registered with the streaming system."),
+	ECVF_RenderThreadSafe
+);
+
 static int32 GSVTStreamingMaxPendingMipLevels = 128;
 static FAutoConsoleVariableRef CVarSVTStreamingMaxPendingMipLevels(
 	TEXT("r.SparseVolumeTexture.Streaming.MaxPendingMipLevels"),
@@ -157,10 +165,10 @@ static FIntVector3 ComputeLargestPossibleTileDataVolumeResolution(int32 VoxelMem
 	return Resolution;
 }
 
-class SparseVolumeTextureUpdateFromBufferCS : public FGlobalShader
+class FSparseVolumeTextureUpdateFromBufferCS : public FGlobalShader
 {
-	DECLARE_GLOBAL_SHADER(SparseVolumeTextureUpdateFromBufferCS);
-	SHADER_USE_PARAMETER_STRUCT(SparseVolumeTextureUpdateFromBufferCS, FGlobalShader)
+	DECLARE_GLOBAL_SHADER(FSparseVolumeTextureUpdateFromBufferCS);
+	SHADER_USE_PARAMETER_STRUCT(FSparseVolumeTextureUpdateFromBufferCS, FGlobalShader)
 
 	BEGIN_SHADER_PARAMETER_STRUCT(FParameters, )
 		SHADER_PARAMETER_RDG_TEXTURE_UAV(RWTexture3D, DstPhysicalTileTextureA)
@@ -187,12 +195,47 @@ class SparseVolumeTextureUpdateFromBufferCS : public FGlobalShader
 		OutEnvironment.SetDefine(TEXT("UPDATE_TILE_TEXTURE_FROM_BUFFER"), 1);
 	}
 };
-IMPLEMENT_GLOBAL_SHADER(SparseVolumeTextureUpdateFromBufferCS, "/Engine/Private/SparseVolumeTexture/UpdateSparseVolumeTexture.usf", "SparseVolumeTextureUpdateFromBufferCS", SF_Compute);
+IMPLEMENT_GLOBAL_SHADER(FSparseVolumeTextureUpdateFromBufferCS, "/Engine/Private/SparseVolumeTexture/UpdateSparseVolumeTexture.usf", "SparseVolumeTextureUpdateFromBufferCS", SF_Compute);
 
-class SparseVolumeTextureUpdatePageTableCS : public FGlobalShader
+class FSparseVolumeTextureUpdateFromSparseBufferCS : public FGlobalShader
 {
-	DECLARE_GLOBAL_SHADER(SparseVolumeTextureUpdatePageTableCS);
-	SHADER_USE_PARAMETER_STRUCT(SparseVolumeTextureUpdatePageTableCS, FGlobalShader)
+	DECLARE_GLOBAL_SHADER(FSparseVolumeTextureUpdateFromSparseBufferCS);
+	SHADER_USE_PARAMETER_STRUCT(FSparseVolumeTextureUpdateFromSparseBufferCS, FGlobalShader)
+
+	BEGIN_SHADER_PARAMETER_STRUCT(FParameters, )
+		SHADER_PARAMETER_RDG_TEXTURE_UAV(RWTexture3D, DstPhysicalTileTextureA)
+		SHADER_PARAMETER_RDG_TEXTURE_UAV(RWTexture3D, DstPhysicalTileTextureB)
+		SHADER_PARAMETER_RDG_BUFFER_SRV(Buffer<float4>, SrcPhysicalTileBufferA)
+		SHADER_PARAMETER_RDG_BUFFER_SRV(Buffer<float4>, SrcPhysicalTileBufferB)
+		SHADER_PARAMETER_RDG_BUFFER_SRV(ByteAddressBuffer, OccupancyBitsBuffer)
+		SHADER_PARAMETER_RDG_BUFFER_SRV(ByteAddressBuffer, TileDataOffsetsBuffer)
+		SHADER_PARAMETER_RDG_BUFFER_SRV(ByteAddressBuffer, DstTileCoordsBuffer)
+		SHADER_PARAMETER(FVector4f, FallbackValueA)
+		SHADER_PARAMETER(FVector4f, FallbackValueB)
+		SHADER_PARAMETER(uint32, NumTilesToCopy)
+		SHADER_PARAMETER(uint32, BufferTileStep)
+		SHADER_PARAMETER(uint32, NumDispatchedGroups)
+		SHADER_PARAMETER(uint32, PaddedTileSize)
+		SHADER_PARAMETER(uint32, CopyTexureMask)
+	END_SHADER_PARAMETER_STRUCT()
+
+	static bool ShouldCompilePermutation(const FGlobalShaderPermutationParameters& Parameters)
+	{
+		return true;
+	}
+
+	static void ModifyCompilationEnvironment(const FGlobalShaderPermutationParameters& Parameters, FShaderCompilerEnvironment& OutEnvironment)
+	{
+		FGlobalShader::ModifyCompilationEnvironment(Parameters, OutEnvironment);
+		OutEnvironment.SetDefine(TEXT("UPDATE_TILE_TEXTURE_FROM_SPARSE_BUFFER"), 1);
+	}
+};
+IMPLEMENT_GLOBAL_SHADER(FSparseVolumeTextureUpdateFromSparseBufferCS, "/Engine/Private/SparseVolumeTexture/UpdateSparseVolumeTexture.usf", "SparseVolumeTextureUpdateFromSparseBufferCS", SF_Compute);
+
+class FSparseVolumeTextureUpdatePageTableCS : public FGlobalShader
+{
+	DECLARE_GLOBAL_SHADER(FSparseVolumeTextureUpdatePageTableCS);
+	SHADER_USE_PARAMETER_STRUCT(FSparseVolumeTextureUpdatePageTableCS, FGlobalShader)
 
 	BEGIN_SHADER_PARAMETER_STRUCT(FParameters, )
 		SHADER_PARAMETER_RDG_TEXTURE_UAV(RWTexture3D, PageTable)
@@ -213,12 +256,12 @@ class SparseVolumeTextureUpdatePageTableCS : public FGlobalShader
 		OutEnvironment.SetDefine(TEXT("UPDATE_PAGE_TABLE"), 1);
 	}
 };
-IMPLEMENT_GLOBAL_SHADER(SparseVolumeTextureUpdatePageTableCS, "/Engine/Private/SparseVolumeTexture/UpdateSparseVolumeTexture.usf", "SparseVolumeTextureUpdatePageTableCS", SF_Compute);
+IMPLEMENT_GLOBAL_SHADER(FSparseVolumeTextureUpdatePageTableCS, "/Engine/Private/SparseVolumeTexture/UpdateSparseVolumeTexture.usf", "SparseVolumeTextureUpdatePageTableCS", SF_Compute);
 
-class SparseVolumeTextureUpdateStreamingInfoBufferCS : public FGlobalShader
+class FSparseVolumeTextureUpdateStreamingInfoBufferCS : public FGlobalShader
 {
-	DECLARE_GLOBAL_SHADER(SparseVolumeTextureUpdateStreamingInfoBufferCS);
-	SHADER_USE_PARAMETER_STRUCT(SparseVolumeTextureUpdateStreamingInfoBufferCS, FGlobalShader)
+	DECLARE_GLOBAL_SHADER(FSparseVolumeTextureUpdateStreamingInfoBufferCS);
+	SHADER_USE_PARAMETER_STRUCT(FSparseVolumeTextureUpdateStreamingInfoBufferCS, FGlobalShader)
 
 	BEGIN_SHADER_PARAMETER_STRUCT(FParameters, )
 		SHADER_PARAMETER_RDG_BUFFER_UAV(RWByteAddressBuffer, StreamingInfoBuffer)
@@ -238,22 +281,34 @@ class SparseVolumeTextureUpdateStreamingInfoBufferCS : public FGlobalShader
 		OutEnvironment.SetDefine(TEXT("UPDATE_STREAMING_INFO_BUFFER"), 1);
 	}
 };
-IMPLEMENT_GLOBAL_SHADER(SparseVolumeTextureUpdateStreamingInfoBufferCS, "/Engine/Private/SparseVolumeTexture/UpdateSparseVolumeTexture.usf", "SparseVolumeTextureUpdateStreamingInfoBufferCS", SF_Compute);
+IMPLEMENT_GLOBAL_SHADER(FSparseVolumeTextureUpdateStreamingInfoBufferCS, "/Engine/Private/SparseVolumeTexture/UpdateSparseVolumeTexture.usf", "SparseVolumeTextureUpdateStreamingInfoBufferCS", SF_Compute);
 
 // Utility class for uploading tiles to a physical tile data texture
 class FTileUploader
 {
 public:
+
+	struct FAddResult
+	{
+		TStaticArray<uint8*, 2> OccupancyBitsPtrs;
+		TStaticArray<uint8*, 2> TileDataOffsetsPtrs;
+		TStaticArray<uint8*, 2> TileDataPtrs;
+		TStaticArray<uint32, 2> TileDataBaseOffsets; // Caller needs to add this value to all data written to TileDataOffsetsPtrs
+		uint8* PackedPhysicalTileCoordsPtr;
+	};
+
 	FTileUploader()
 	{
 		ResetState();
 	}
 
-	void Init(FRDGBuilder& GraphBuilder, int32 InMaxNumTiles, EPixelFormat InFormatA, EPixelFormat InFormatB)
+	void Init(FRDGBuilder& GraphBuilder, int32 InMaxNumTiles, int32 InMaxNumVoxelsA, int32 InMaxNumVoxelsB, EPixelFormat InFormatA, EPixelFormat InFormatB)
 	{
 		check(InFormatA != PF_Unknown || InFormatB != PF_Unknown);
 		ResetState();
 		MaxNumTiles = InMaxNumTiles;
+		MaxNumVoxelsA = InMaxNumVoxelsA;
+		MaxNumVoxelsB = InMaxNumVoxelsB;
 		FormatA = InFormatA;
 		FormatB = InFormatB;
 		FormatSizeA = GPixelFormats[FormatA].BlockBytes;
@@ -262,75 +317,112 @@ public:
 		// Create a new set of buffers if the old set is already queued into RDG.
 		if (IsRegistered(GraphBuilder, DstTileCoordsUploadBuffer))
 		{
+			OccupancyBitsUploadBuffer = nullptr;
+			TileDataOffsetsUploadBuffer = nullptr;
 			DstTileCoordsUploadBuffer = nullptr;
 			TileDataAUploadBuffer = nullptr;
 			TileDataBUploadBuffer = nullptr;
 		}
 
+		const int32 NumTextures = (FormatA == PF_Unknown || FormatB == PF_Unknown) ? 1 : 2;
+
 		if (MaxNumTiles > 0)
 		{
+			// Occupancy bits
+			{
+				const uint32 BufferSize = NumTextures * MaxNumTiles * SVT::NumOccupancyWordsPerPaddedTile * sizeof(uint32);
+				FRDGBufferDesc BufferDesc = FRDGBufferDesc::CreateByteAddressUploadDesc(BufferSize);
+				BufferDesc.Usage |= EBufferUsageFlags::Dynamic; // Skip the unneeded copy from upload to VRAM resource on d3d12 RHI
+				AllocatePooledBuffer(BufferDesc, OccupancyBitsUploadBuffer, TEXT("SparseVolumeTexture.OccupancyBitsUploadBuffer"));
+
+				OccupancyBitsAPtr = (uint8*)RHILockBuffer(OccupancyBitsUploadBuffer->GetRHI(), 0, BufferSize, RLM_WriteOnly);
+				OccupancyBitsBPtr = OccupancyBitsAPtr + (FormatA != PF_Unknown ? (MaxNumTiles * SVT::NumOccupancyWordsPerPaddedTile * sizeof(uint32)) : 0);
+			}
+			// Tile data offsets
+			{
+				const uint32 BufferSize = NumTextures * MaxNumTiles * sizeof(uint32);
+				FRDGBufferDesc BufferDesc = FRDGBufferDesc::CreateByteAddressUploadDesc(BufferSize);
+				BufferDesc.Usage |= EBufferUsageFlags::Dynamic; // Skip the unneeded copy from upload to VRAM resource on d3d12 RHI
+				AllocatePooledBuffer(BufferDesc, TileDataOffsetsUploadBuffer, TEXT("SparseVolumeTexture.TileDataOffsetsUploadBuffer"));
+
+				TileDataOffsetsAPtr = (uint8*)RHILockBuffer(TileDataOffsetsUploadBuffer->GetRHI(), 0, BufferSize, RLM_WriteOnly);
+				TileDataOffsetsBPtr = TileDataOffsetsAPtr + (FormatA != PF_Unknown ? (MaxNumTiles * sizeof(uint32)) : 0);
+			}
 			// TileCoords
 			{
-				// Add EBufferUsageFlags::Dynamic to skip the unneeded copy from upload to VRAM resource on d3d12 RHI
-				FRDGBufferDesc BufferDesc = FRDGBufferDesc::CreateByteAddressUploadDesc(MaxNumTiles * sizeof(uint32));
-				BufferDesc.Usage |= EBufferUsageFlags::Dynamic;
+				const uint32 BufferSize = MaxNumTiles * sizeof(uint32);
+				FRDGBufferDesc BufferDesc = FRDGBufferDesc::CreateByteAddressUploadDesc(BufferSize);
+				BufferDesc.Usage |= EBufferUsageFlags::Dynamic; // Skip the unneeded copy from upload to VRAM resource on d3d12 RHI
 				AllocatePooledBuffer(BufferDesc, DstTileCoordsUploadBuffer, TEXT("SparseVolumeTexture.TileCoordsUploadBuffer"));
 
-				TileCoordsPtr = (uint8*)RHILockBuffer(DstTileCoordsUploadBuffer->GetRHI(), 0, MaxNumTiles * sizeof(uint32), RLM_WriteOnly);
+				TileCoordsPtr = (uint8*)RHILockBuffer(DstTileCoordsUploadBuffer->GetRHI(), 0, BufferSize, RLM_WriteOnly);
 			}
 
-			const int32 NumVoxels = MaxNumTiles * SVT::NumVoxelsPerPaddedTile;
-
 			// TileData
 			if (FormatSizeA > 0)
 			{
-				// Add EBufferUsageFlags::Dynamic to skip the unneeded copy from upload to VRAM resource on d3d12 RHI
-				FRDGBufferDesc BufferDesc = FRDGBufferDesc::CreateUploadDesc(FormatSizeA, NumVoxels);
-				BufferDesc.Usage |= EBufferUsageFlags::Dynamic;
+				FRDGBufferDesc BufferDesc = FRDGBufferDesc::CreateUploadDesc(FormatSizeA, FMath::Max(MaxNumVoxelsA, 1));
+				BufferDesc.Usage |= EBufferUsageFlags::Dynamic; // Skip the unneeded copy from upload to VRAM resource on d3d12 RHI
 				AllocatePooledBuffer(BufferDesc, TileDataAUploadBuffer, TEXT("SparseVolumeTexture.TileDataAUploadBuffer"));
 
-				TileDataAPtr = (uint8*)RHILockBuffer(TileDataAUploadBuffer->GetRHI(), 0, NumVoxels * FormatSizeA, RLM_WriteOnly);
+				TileDataAPtr = (uint8*)RHILockBuffer(TileDataAUploadBuffer->GetRHI(), 0, FMath::Max(MaxNumVoxelsA, 1) * FormatSizeA, RLM_WriteOnly);
 			}
 			if (FormatSizeB > 0)
 			{
-				// Add EBufferUsageFlags::Dynamic to skip the unneeded copy from upload to VRAM resource on d3d12 RHI
-				FRDGBufferDesc BufferDesc = FRDGBufferDesc::CreateUploadDesc(FormatSizeB, NumVoxels);
-				BufferDesc.Usage |= EBufferUsageFlags::Dynamic;
+				FRDGBufferDesc BufferDesc = FRDGBufferDesc::CreateUploadDesc(FormatSizeB, FMath::Max(MaxNumVoxelsB, 1));
+				BufferDesc.Usage |= EBufferUsageFlags::Dynamic; // Skip the unneeded copy from upload to VRAM resource on d3d12 RHI
 				AllocatePooledBuffer(BufferDesc, TileDataBUploadBuffer, TEXT("SparseVolumeTexture.TileDataBUploadBuffer"));
 
-				TileDataBPtr = (uint8*)RHILockBuffer(TileDataBUploadBuffer->GetRHI(), 0, NumVoxels * FormatSizeB, RLM_WriteOnly);
+				TileDataBPtr = (uint8*)RHILockBuffer(TileDataBUploadBuffer->GetRHI(), 0, FMath::Max(MaxNumVoxelsB, 1) * FormatSizeB, RLM_WriteOnly);
 			}
 		}
 	}
 
-	void Add_GetRef(int32 NumTiles, uint8*& OutPackedPhysicalTileCoordsPtr, uint8*& OutPtrA, uint8*& OutPtrB)
+	FAddResult Add_GetRef(int32 NumTiles, int32 NumVoxelsA, int32 NumVoxelsB)
 	{
 		check((NumWrittenTiles + NumTiles) <= MaxNumTiles);
+		check((NumWrittenVoxelsA + NumVoxelsA) <= MaxNumVoxelsA);
+		check((NumWrittenVoxelsB + NumVoxelsB) <= MaxNumVoxelsB);
 		check(TileCoordsPtr);
 		check(FormatSizeA <= 0 || TileDataAPtr);
 		check(FormatSizeB <= 0 || TileDataBPtr);
 
-		OutPackedPhysicalTileCoordsPtr = TileCoordsPtr + NumWrittenTiles * sizeof(uint32);
-		OutPtrA = TileDataAPtr ? TileDataAPtr + (NumWrittenTiles * UE::SVT::NumVoxelsPerPaddedTile * FormatSizeA) : nullptr;
-		OutPtrB = TileDataBPtr ? TileDataBPtr + (NumWrittenTiles * UE::SVT::NumVoxelsPerPaddedTile * FormatSizeB) : nullptr;
-
+		FAddResult Result = {};
+		Result.OccupancyBitsPtrs[0] = TileDataAPtr ? OccupancyBitsAPtr + (NumWrittenTiles * SVT::NumOccupancyWordsPerPaddedTile * sizeof(uint32)) : nullptr;
+		Result.OccupancyBitsPtrs[1] = TileDataBPtr ? OccupancyBitsBPtr + (NumWrittenTiles * SVT::NumOccupancyWordsPerPaddedTile * sizeof(uint32)) : nullptr;
+		Result.TileDataOffsetsPtrs[0] = TileDataAPtr ? TileDataOffsetsAPtr + (NumWrittenTiles * sizeof(uint32)) : nullptr;
+		Result.TileDataOffsetsPtrs[1] = TileDataBPtr ? TileDataOffsetsBPtr + (NumWrittenTiles * sizeof(uint32)) : nullptr;
+		Result.TileDataPtrs[0] = TileDataAPtr ? TileDataAPtr + (NumWrittenVoxelsA * FormatSizeA) : nullptr;
+		Result.TileDataPtrs[1] = TileDataBPtr ? TileDataBPtr + (NumWrittenVoxelsB * FormatSizeB) : nullptr;
+		Result.TileDataBaseOffsets[0] = NumWrittenVoxelsA;
+		Result.TileDataBaseOffsets[1] = NumWrittenVoxelsB;
+		Result.PackedPhysicalTileCoordsPtr = TileCoordsPtr + NumWrittenTiles * sizeof(uint32);
+		
 		NumWrittenTiles += NumTiles;
+		NumWrittenVoxelsA += NumVoxelsA;
+		NumWrittenVoxelsB += NumVoxelsB;
+
+		return Result;
 	}
 
 	void Release()
 	{
+		OccupancyBitsUploadBuffer.SafeRelease();
+		TileDataOffsetsUploadBuffer.SafeRelease();
 		DstTileCoordsUploadBuffer.SafeRelease();
 		TileDataAUploadBuffer.SafeRelease();
 		TileDataBUploadBuffer.SafeRelease();
 		ResetState();
 	}
 
-	void ResourceUploadTo(FRDGBuilder& GraphBuilder, FRHITexture* DstTextureA, FRHITexture* DstTextureB)
+	void ResourceUploadTo(FRDGBuilder& GraphBuilder, FRHITexture* DstTextureA, FRHITexture* DstTextureB, const FVector4f& FallbackValueA, const FVector4f& FallbackValueB)
 	{
 		check(DstTextureA || FormatSizeA <= 0);
 		check(DstTextureB || FormatSizeB <= 0);
 		if (MaxNumTiles > 0)
 		{
+			RHIUnlockBuffer(OccupancyBitsUploadBuffer->GetRHI());
+			RHIUnlockBuffer(TileDataOffsetsUploadBuffer->GetRHI());
 			RHIUnlockBuffer(DstTileCoordsUploadBuffer->GetRHI());
 			if (TileDataAPtr)
 			{
@@ -346,42 +438,33 @@ public:
 				FRDGTexture* DstTextureARDG = DstTextureA ? GraphBuilder.RegisterExternalTexture(CreateRenderTarget(DstTextureA, TEXT("SparseVolumeTexture.TileDataTextureA"))) : nullptr;
 				FRDGTexture* DstTextureBRDG = DstTextureB ? GraphBuilder.RegisterExternalTexture(CreateRenderTarget(DstTextureB, TEXT("SparseVolumeTexture.TileDataTextureB"))) : nullptr;
 
+				FRDGBufferSRV* OccupancyBitsBufferSRV = GraphBuilder.CreateSRV(GraphBuilder.RegisterExternalBuffer(OccupancyBitsUploadBuffer));
+				FRDGBufferSRV* TileDataOffsetsBufferSRV = GraphBuilder.CreateSRV(GraphBuilder.RegisterExternalBuffer(TileDataOffsetsUploadBuffer));
 				FRDGBufferSRV* DstTileCoordsBufferSRV = GraphBuilder.CreateSRV(GraphBuilder.RegisterExternalBuffer(DstTileCoordsUploadBuffer));
 				FRDGBufferSRV* TileDataABufferSRV = FormatSizeA ? GraphBuilder.CreateSRV(GraphBuilder.RegisterExternalBuffer(TileDataAUploadBuffer), FormatA) : nullptr;
 				FRDGBufferSRV* TileDataBBufferSRV = FormatSizeB ? GraphBuilder.CreateSRV(GraphBuilder.RegisterExternalBuffer(TileDataBUploadBuffer), FormatB) : nullptr;
 				FRDGTextureUAV* DstTextureAUAV = DstTextureARDG ? GraphBuilder.CreateUAV(DstTextureARDG) : nullptr;
 				FRDGTextureUAV* DstTextureBUAV = DstTextureBRDG ? GraphBuilder.CreateUAV(DstTextureBRDG) : nullptr;
 
-				auto ComputeShader = GetGlobalShaderMap(GMaxRHIFeatureLevel)->GetShader<SparseVolumeTextureUpdateFromBufferCS>();
-
-				SparseVolumeTextureUpdateFromBufferCS::FParameters* PassParameters = GraphBuilder.AllocParameters<SparseVolumeTextureUpdateFromBufferCS::FParameters>();
-				if (FormatSizeA > 0) // TextureA exists
-				{
-					PassParameters->DstPhysicalTileTextureA = DstTextureAUAV;
-					PassParameters->SrcPhysicalTileBufferA = TileDataABufferSRV;
-					if (FormatSizeB == 0) // B doesn't exist: fill B params with A
-					{
-						PassParameters->DstPhysicalTileTextureB = DstTextureAUAV;
-						PassParameters->SrcPhysicalTileBufferB = TileDataABufferSRV;
-					}
-				}
-				if (FormatSizeB > 0) // TextureB exists
-				{
-					PassParameters->DstPhysicalTileTextureB = DstTextureBUAV;
-					PassParameters->SrcPhysicalTileBufferB = TileDataBBufferSRV;
-					if (FormatSizeA == 0) // A doesnt't exist: fill A params with B
-					{
-						PassParameters->DstPhysicalTileTextureA = DstTextureBUAV;
-						PassParameters->SrcPhysicalTileBufferA = TileDataBBufferSRV;
-					}
-				}
-				PassParameters->DstTileCoords = DstTileCoordsBufferSRV;
-				PassParameters->TileCoordsBufferOffset = 0;
-				PassParameters->TileDataBufferOffsetInTiles = 0;
+				auto ComputeShader = GetGlobalShaderMap(GMaxRHIFeatureLevel)->GetShader<FSparseVolumeTextureUpdateFromSparseBufferCS>();
+
+				FSparseVolumeTextureUpdateFromSparseBufferCS::FParameters* PassParameters = GraphBuilder.AllocParameters<FSparseVolumeTextureUpdateFromSparseBufferCS::FParameters>();
+				PassParameters->DstPhysicalTileTextureA = FormatSizeA > 0 ? DstTextureAUAV : DstTextureBUAV;
+				PassParameters->DstPhysicalTileTextureB = FormatSizeB > 0 ? DstTextureBUAV : DstTextureAUAV;
+				PassParameters->SrcPhysicalTileBufferA = FormatSizeA > 0 ? TileDataABufferSRV : TileDataBBufferSRV;
+				PassParameters->SrcPhysicalTileBufferB = FormatSizeB > 0 ? TileDataBBufferSRV : TileDataABufferSRV;
+				PassParameters->OccupancyBitsBuffer = OccupancyBitsBufferSRV;
+				PassParameters->TileDataOffsetsBuffer = TileDataOffsetsBufferSRV;
+				PassParameters->DstTileCoordsBuffer = DstTileCoordsBufferSRV;
+				PassParameters->FallbackValueA = FallbackValueA;
+				PassParameters->FallbackValueB = FallbackValueB;
 				PassParameters->NumTilesToCopy = NumWrittenTiles;
+				PassParameters->BufferTileStep = MaxNumTiles;
 				PassParameters->NumDispatchedGroups = FMath::Min(NumWrittenTiles, 1024);
 				PassParameters->PaddedTileSize = SPARSE_VOLUME_TILE_RES_PADDED;
-				PassParameters->bCopyTexureAOnlyUI = (FormatSizeA == 0 || FormatSizeB == 0);
+				PassParameters->CopyTexureMask = 0;
+				PassParameters->CopyTexureMask |= FormatSizeA > 0 ? 0x1u : 0x0u;
+				PassParameters->CopyTexureMask |= FormatSizeB > 0 ? 0x2u : 0x0u;
 
 				const bool bAsyncCompute = GSupportsEfficientAsyncCompute && (GSVTStreamingAsyncCompute != 0);
 
@@ -399,30 +482,48 @@ public:
 	}
 
 private:
+	TRefCountPtr<FRDGPooledBuffer> OccupancyBitsUploadBuffer;
+	TRefCountPtr<FRDGPooledBuffer> TileDataOffsetsUploadBuffer;
 	TRefCountPtr<FRDGPooledBuffer> DstTileCoordsUploadBuffer;
 	TRefCountPtr<FRDGPooledBuffer> TileDataAUploadBuffer;
 	TRefCountPtr<FRDGPooledBuffer> TileDataBUploadBuffer;
+	uint8* OccupancyBitsAPtr = nullptr;
+	uint8* OccupancyBitsBPtr = nullptr;
+	uint8* TileDataOffsetsAPtr = nullptr;
+	uint8* TileDataOffsetsBPtr = nullptr;
 	uint8* TileCoordsPtr = nullptr;
 	uint8* TileDataAPtr = nullptr;
 	uint8* TileDataBPtr = nullptr;
 	int32 MaxNumTiles = 0;
+	int32 MaxNumVoxelsA = 0;
+	int32 MaxNumVoxelsB = 0;
 	EPixelFormat FormatA = PF_Unknown;
 	EPixelFormat FormatB = PF_Unknown;
 	int32 FormatSizeA = 0;
 	int32 FormatSizeB = 0;
 	int32 NumWrittenTiles = 0;
+	int32 NumWrittenVoxelsA = 0;
+	int32 NumWrittenVoxelsB = 0;
 
 	void ResetState()
 	{
+		OccupancyBitsAPtr = nullptr;
+		OccupancyBitsBPtr = nullptr;
+		TileDataOffsetsAPtr = nullptr;
+		TileDataOffsetsBPtr = nullptr;
 		TileCoordsPtr = nullptr;
 		TileDataAPtr = nullptr;
 		TileDataBPtr = nullptr;
 		MaxNumTiles = 0;
+		MaxNumVoxelsA = 0;
+		MaxNumVoxelsB = 0;
 		FormatA = PF_Unknown;
 		FormatB = PF_Unknown;
 		FormatSizeA = 0;
 		FormatSizeB = 0;
 		NumWrittenTiles = 0;
+		NumWrittenVoxelsA = 0;
+		NumWrittenVoxelsB = 0;
 	}
 };
 
@@ -489,7 +590,7 @@ public:
 
 			if (NumWrittenUpdates > 0)
 			{
-				auto ComputeShader = GetGlobalShaderMap(GMaxRHIFeatureLevel)->GetShader<SparseVolumeTextureUpdatePageTableCS>();
+				auto ComputeShader = GetGlobalShaderMap(GMaxRHIFeatureLevel)->GetShader<FSparseVolumeTextureUpdatePageTableCS>();
 				const bool bAsyncCompute = GSupportsEfficientAsyncCompute && (GSVTStreamingAsyncCompute != 0);
 
 				uint32 UpdatesOffset = 0;
@@ -499,7 +600,7 @@ public:
 					FRDGTextureUAV* PageTableUAV = GraphBuilder.CreateUAV(FRDGTextureUAVDesc(PageTableRDG, Batch.MipLevel, PF_R32_UINT));
 					FRDGBufferSRV* UpdatesBufferSRV = GraphBuilder.CreateSRV(GraphBuilder.RegisterExternalBuffer(UpdatesUploadBuffer));
 
-					SparseVolumeTextureUpdatePageTableCS::FParameters* PassParameters = GraphBuilder.AllocParameters<SparseVolumeTextureUpdatePageTableCS::FParameters>();
+					FSparseVolumeTextureUpdatePageTableCS::FParameters* PassParameters = GraphBuilder.AllocParameters<FSparseVolumeTextureUpdatePageTableCS::FParameters>();
 					PassParameters->PageTable = PageTableUAV;
 					PassParameters->PageTableUpdates = UpdatesBufferSRV;
 					PassParameters->UpdateCoordOffset = UpdatesOffset;
@@ -589,7 +690,7 @@ public:
 				RHIUnlockBuffer(UpdatesUploadBuffer->GetRHI());
 			}
 
-			auto ComputeShader = GetGlobalShaderMap(GMaxRHIFeatureLevel)->GetShader<SparseVolumeTextureUpdateStreamingInfoBufferCS>();
+			auto ComputeShader = GetGlobalShaderMap(GMaxRHIFeatureLevel)->GetShader<FSparseVolumeTextureUpdateStreamingInfoBufferCS>();
 			const bool bAsyncCompute = GSupportsEfficientAsyncCompute && (GSVTStreamingAsyncCompute != 0);
 
 			uint32 UpdatesOffset = 0;
@@ -598,7 +699,7 @@ public:
 				FRDGBufferUAV* StreamingInfoBufferUAV = GraphBuilder.CreateUAV(GraphBuilder.RegisterExternalBuffer(Batch.StreamingInfoBuffer), PF_R32_UINT);
 				FRDGBufferSRV* UpdatesBufferSRV = GraphBuilder.CreateSRV(GraphBuilder.RegisterExternalBuffer(UpdatesUploadBuffer));
 
-				SparseVolumeTextureUpdateStreamingInfoBufferCS::FParameters* PassParameters = GraphBuilder.AllocParameters<SparseVolumeTextureUpdateStreamingInfoBufferCS::FParameters>();
+				FSparseVolumeTextureUpdateStreamingInfoBufferCS::FParameters* PassParameters = GraphBuilder.AllocParameters<FSparseVolumeTextureUpdateStreamingInfoBufferCS::FParameters>();
 				PassParameters->StreamingInfoBuffer = StreamingInfoBufferUAV;
 				PassParameters->StreamingInfoBufferUpdates = UpdatesBufferSRV;
 				PassParameters->UpdateOffset = UpdatesOffset;
@@ -910,6 +1011,41 @@ void FStreamingManager::BeginAsyncUpdate(FRDGBuilder& GraphBuilder, bool bBlocki
 		GSVTStreamingEmptyPhysicalTileTextures = 0;
 	}
 
+	if (GSVTStreamingPrintMemoryStats != 0)
+	{
+		for (auto& Pair : StreamingInfo)
+		{
+			double MinFrameMiB = FLT_MAX;
+			double MaxFrameMiB = -FLT_MAX;
+			double SumFrameMiB = 0.0;
+
+			const int32 NumFrames = Pair.Value.PerFrameInfo.Num();
+			UE_LOG(LogSparseVolumeTextureStreamingManager, Display, TEXT("Memory stats for SVT '%p': Each mip level is printed as a tuple of [PageTable Size | VoxelData Size | Total]"), Pair.Key);
+
+			for (int32 FrameIdx = 0; FrameIdx < NumFrames; ++FrameIdx)
+			{
+				const FFrameInfo& FrameInfo = Pair.Value.PerFrameInfo[FrameIdx];
+				FString Str;
+				int32 TotalSize = 0;
+				for (const auto& SInfo : FrameInfo.Resources->MipLevelStreamingInfo)
+				{
+					Str += FString::Printf(TEXT("[%5.2f KiB|%5.2f KiB|%5.2f KiB] "), SInfo.PageTableSize / 1024.0f, (SInfo.TileDataSize[0] + SInfo.TileDataSize[1]) / 1024.0f, (SInfo.PageTableSize + SInfo.TileDataSize[0] + SInfo.TileDataSize[1]) / 1024.0f);
+					TotalSize += SInfo.BulkSize;
+				}
+
+				MinFrameMiB = FMath::Min(MinFrameMiB, TotalSize / 1024.0 / 1024.0);
+				MaxFrameMiB = FMath::Max(MaxFrameMiB, TotalSize / 1024.0 / 1024.0);
+				SumFrameMiB += TotalSize / 1024.0 / 1024.0;
+
+				UE_LOG(LogSparseVolumeTextureStreamingManager, Display, TEXT("SVT Frame %3i: TotalSize: %3.2f MiB %s"), FrameIdx, TotalSize / 1024.0f / 1024.0f, *Str);
+			}
+
+			UE_LOG(LogSparseVolumeTextureStreamingManager, Display, TEXT("SVT Frame Stats: Min: %3.2f MiB, Max: %3.2f MiB, Avg: %3.2f, Total All: %3.2f MiB"), (float)MinFrameMiB, (float)MaxFrameMiB, (float)(SumFrameMiB / NumFrames), SumFrameMiB);
+		}
+
+		GSVTStreamingPrintMemoryStats = 0;
+	}
+
 	AddParentRequests();
 	const int32 MaxSelectedRequests = MaxPendingMipLevels - NumPendingMipLevels;
 	SelectHighestPriorityRequestsAndUpdateLRU(MaxSelectedRequests);
@@ -932,8 +1068,12 @@ void FStreamingManager::BeginAsyncUpdate(FRDGBuilder& GraphBuilder, bool bBlocki
 				continue; // Skip mip level install. SVT no longer exists or mip level was "streamed out" before it was even installed in the first place.
 			}
 
+			const int32 FormatSizeA = GPixelFormats[SVTInfo->FormatA].BlockBytes;
+			const int32 FormatSizeB = GPixelFormats[SVTInfo->FormatB].BlockBytes;
 			const FResources* Resources = SVTInfo->PerFrameInfo[PendingMipLevel.FrameIndex].Resources;
 			SVTInfo->TileDataTexture->NumTilesToUpload += Resources->MipLevelStreamingInfo[PendingMipLevel.MipLevelIndex].NumPhysicalTiles;
+			SVTInfo->TileDataTexture->NumVoxelsToUploadA += FormatSizeA > 0 ? Resources->MipLevelStreamingInfo[PendingMipLevel.MipLevelIndex].TileDataSize[0] / FormatSizeA : 0;
+			SVTInfo->TileDataTexture->NumVoxelsToUploadB += FormatSizeB > 0 ? Resources->MipLevelStreamingInfo[PendingMipLevel.MipLevelIndex].TileDataSize[1] / FormatSizeB : 0;
 			TileDataTexturesToUpdate.Add(SVTInfo->TileDataTexture.Get());
 			NumPageTableUpdatesTotal += Resources->MipLevelStreamingInfo[PendingMipLevel.MipLevelIndex].PageTableSize / (2 * sizeof(uint32));
 		}
@@ -942,7 +1082,7 @@ void FStreamingManager::BeginAsyncUpdate(FRDGBuilder& GraphBuilder, bool bBlocki
 
 		for (FTileDataTexture* TileDataTexture : TileDataTexturesToUpdate)
 		{
-			TileDataTexture->TileUploader->Init(GraphBuilder, TileDataTexture->NumTilesToUpload, TileDataTexture->FormatA, TileDataTexture->FormatB);
+			TileDataTexture->TileUploader->Init(GraphBuilder, TileDataTexture->NumTilesToUpload, TileDataTexture->NumVoxelsToUploadA, TileDataTexture->NumVoxelsToUploadB, TileDataTexture->FormatA, TileDataTexture->FormatB);
 		}
 	}
 
@@ -996,8 +1136,10 @@ void FStreamingManager::EndAsyncUpdate(FRDGBuilder& GraphBuilder)
 	// Issue the actual data uploads
 	for (FTileDataTexture* TileDataTexture : TileDataTexturesToUpdate)
 	{
-		TileDataTexture->TileUploader->ResourceUploadTo(GraphBuilder, TileDataTexture->TileDataTextureARHIRef, TileDataTexture->TileDataTextureBRHIRef);
+		TileDataTexture->TileUploader->ResourceUploadTo(GraphBuilder, TileDataTexture->TileDataTextureARHIRef, TileDataTexture->TileDataTextureBRHIRef, TileDataTexture->FallbackValueA, TileDataTexture->FallbackValueB);
 		TileDataTexture->NumTilesToUpload = 0;
+		TileDataTexture->NumVoxelsToUploadA = 0;
+		TileDataTexture->NumVoxelsToUploadB = 0;
 	}
 
 	// Update streaming info buffers
@@ -1082,7 +1224,11 @@ void FStreamingManager::AddInternal(FRDGBuilder& GraphBuilder, FNewSparseVolumeT
 	SVTInfo.LRUNodes.SetNum(NumFrames * SVTInfo.NumMipLevelsGlobal);
 	SVTInfo.PerMipLRULists.SetNum(SVTInfo.NumMipLevelsGlobal);
 
+	const int32 FormatSizes[] = { GPixelFormats[SVTInfo.FormatA].BlockBytes, GPixelFormats[SVTInfo.FormatB].BlockBytes };
+
 	int32 NumRootPhysicalTiles = 0;
+	int32 NumRootVoxelsA = 0;
+	int32 NumRootVoxelsB = 0;
 	int32 MaxNumPhysicalTiles = 0;
 	for (int32 FrameIdx = 0; FrameIdx < NumFrames; ++FrameIdx)
 	{
@@ -1109,6 +1255,10 @@ void FStreamingManager::AddInternal(FRDGBuilder& GraphBuilder, FNewSparseVolumeT
 		if (NumPhysicalTiles > 0)
 		{
 			++NumRootPhysicalTiles;
+			check(FormatSizes[0] == 0 || (Resources->MipLevelStreamingInfo.Last().TileDataSize[0] % FormatSizes[0]) == 0);
+			check(FormatSizes[1] == 0 || (Resources->MipLevelStreamingInfo.Last().TileDataSize[1] % FormatSizes[1]) == 0);
+			NumRootVoxelsA += FormatSizes[0] > 0 ? (Resources->MipLevelStreamingInfo.Last().TileDataSize[0] / FormatSizes[0]) : 0;
+			NumRootVoxelsB += FormatSizes[1] > 0 ? (Resources->MipLevelStreamingInfo.Last().TileDataSize[1] / FormatSizes[1]) : 0;
 		}
 
 		for (int32 MipIdx = 0; MipIdx < SVTInfo.NumMipLevelsGlobal; ++MipIdx)
@@ -1132,7 +1282,7 @@ void FStreamingManager::AddInternal(FRDGBuilder& GraphBuilder, FNewSparseVolumeT
 		const FIntVector3 TileDataVolumeResolution = ComputeTileDataVolumeResolution(NumPhysicalTilesCapacity);
 		const FIntVector3 TileDataVolumeResolutionInTiles = TileDataVolumeResolution / SPARSE_VOLUME_TILE_RES_PADDED;
 
-		SVTInfo.TileDataTexture = MakeUnique<FTileDataTexture>(TileDataVolumeResolutionInTiles, SVTInfo.FormatA, SVTInfo.FormatB);
+		SVTInfo.TileDataTexture = MakeUnique<FTileDataTexture>(TileDataVolumeResolutionInTiles, SVTInfo.FormatA, SVTInfo.FormatB, SVTInfo.FallbackValueA, SVTInfo.FallbackValueB);
 		SVTInfo.TileDataTexture->InitResource(GraphBuilder.RHICmdList);
 
 		// Create streaming info buffer
@@ -1144,28 +1294,23 @@ void FStreamingManager::AddInternal(FRDGBuilder& GraphBuilder, FNewSparseVolumeT
 		}
 
 		FTileUploader RootTileUploader;
-		RootTileUploader.Init(GraphBuilder, NumRootPhysicalTiles + 1 /*null tile*/, SVTInfo.FormatA, SVTInfo.FormatB);
+		RootTileUploader.Init(GraphBuilder, NumRootPhysicalTiles + 1 /*null tile*/, NumRootVoxelsA, NumRootVoxelsB, SVTInfo.FormatA, SVTInfo.FormatB);
 
 		// Allocate null tile
 		{
 			const uint32 NullTileCoord = SVTInfo.TileDataTexture->Allocate();
 			check(NullTileCoord == 0);
-			uint8* TileCoordsPtr = nullptr;
-			uint8* DataAPtr = nullptr;
-			uint8* DataBPtr = nullptr;
-			RootTileUploader.Add_GetRef(1, TileCoordsPtr, DataAPtr, DataBPtr);
-			FMemory::Memcpy(TileCoordsPtr, &NullTileCoord, sizeof(NullTileCoord));
-			for (int32 VoxelIdx = 0; VoxelIdx < SVT::NumVoxelsPerPaddedTile; ++VoxelIdx)
+			FTileUploader::FAddResult AddResult = RootTileUploader.Add_GetRef(1 /*NumTiles*/, 0 /*NumVoxelsA*/, 0 /*NumVoxelsB*/);
+			FMemory::Memcpy(AddResult.PackedPhysicalTileCoordsPtr, &NullTileCoord, sizeof(NullTileCoord));
+			if (SVTInfo.FormatA != PF_Unknown)
 			{
-				if (SVTInfo.FormatA != PF_Unknown)
-				{
-					WriteVoxel(VoxelIdx, DataAPtr, SVTInfo.FormatA, SVTInfo.FallbackValueA);
-				}
-				if (SVTInfo.FormatB != PF_Unknown)
-				{
-					WriteVoxel(VoxelIdx, DataBPtr, SVTInfo.FormatB, SVTInfo.FallbackValueB);
-				}
+				FMemory::Memzero(AddResult.OccupancyBitsPtrs[0], SVT::NumOccupancyWordsPerPaddedTile * sizeof(uint32));
+			}
+			if (SVTInfo.FormatB != PF_Unknown)
+			{
+				FMemory::Memzero(AddResult.OccupancyBitsPtrs[1], SVT::NumOccupancyWordsPerPaddedTile * sizeof(uint32));
 			}
+			// No need to write to TileDataPtrA. TileDataPtrB, TileDataOffsetsPtrA and TileDataOffsetsPtrB because we zeroed out all the occupancy bits.
 		}
 
 		// Process frames
@@ -1208,28 +1353,37 @@ void FStreamingManager::AddInternal(FRDGBuilder& GraphBuilder, FNewSparseVolumeT
 			const FMipLevelStreamingInfo* RootStreamingInfo = !Resources->MipLevelStreamingInfo.IsEmpty() ? &Resources->MipLevelStreamingInfo.Last() : nullptr;
 			if (!Resources->RootData.IsEmpty() && RootStreamingInfo)
 			{
-				check((RootStreamingInfo->TileDataASize > 0) || (RootStreamingInfo->TileDataBSize > 0));
+				check((RootStreamingInfo->TileDataSize[0] > 0) || (RootStreamingInfo->TileDataSize[1] > 0));
 
 				const uint32 TileCoord = SVTInfo.TileDataTexture->Allocate();
 				check(TileCoord != INDEX_NONE);
 				FrameInfo.TileAllocations.Last()[0] = TileCoord;
-				uint8* TileCoordsPtr = nullptr;
-				uint8* DataAPtr = nullptr;
-				uint8* DataBPtr = nullptr;
-				RootTileUploader.Add_GetRef(1, TileCoordsPtr, DataAPtr, DataBPtr);
 
-				FMemory::Memcpy(TileCoordsPtr, &TileCoord, sizeof(TileCoord));
-				if (RootStreamingInfo->TileDataASize > 0)
-				{
-					const uint8* Src = Resources->RootData.GetData() + RootStreamingInfo->TileDataAOffset;
-					check(DataAPtr);
-					FMemory::Memcpy(DataAPtr, Src, RootStreamingInfo->TileDataASize);
-				}
-				if (RootStreamingInfo->TileDataBSize > 0)
+				const int32 NumVoxelsA = FormatSizes[0] > 0 ? RootStreamingInfo->TileDataSize[0] / FormatSizes[0] : 0;
+				const int32 NumVoxelsB = FormatSizes[1] > 0 ? RootStreamingInfo->TileDataSize[1] / FormatSizes[1] : 0;
+				FTileUploader::FAddResult AddResult = RootTileUploader.Add_GetRef(1, NumVoxelsA, NumVoxelsB);
+
+				FMemory::Memcpy(AddResult.PackedPhysicalTileCoordsPtr, &TileCoord, sizeof(TileCoord));
+				for (int32 AttributesIdx = 0; AttributesIdx < 2; ++AttributesIdx)
 				{
-					const uint8* Src = Resources->RootData.GetData() + RootStreamingInfo->TileDataBOffset;
-					check(DataBPtr);
-					FMemory::Memcpy(DataBPtr, Src, RootStreamingInfo->TileDataBSize);
+					if (FormatSizes[AttributesIdx] > 0)
+					{
+						// Occupancy bits
+						const uint8* SrcOccupancyBits = Resources->RootData.GetData() + RootStreamingInfo->OccupancyBitsOffset[AttributesIdx];
+						check(AddResult.OccupancyBitsPtrs[AttributesIdx]);
+						FMemory::Memcpy(AddResult.OccupancyBitsPtrs[AttributesIdx], SrcOccupancyBits, RootStreamingInfo->OccupancyBitsSize[AttributesIdx]);
+
+						// Per-tile offsets into tile data
+						const uint32* SrcTileDataOffsets = reinterpret_cast<const uint32*>(Resources->RootData.GetData() + RootStreamingInfo->TileDataOffsetsOffset[AttributesIdx]);
+						check(AddResult.TileDataOffsetsPtrs[AttributesIdx]);
+						check(RootStreamingInfo->TileDataOffsetsSize[AttributesIdx] == sizeof(uint32));
+						reinterpret_cast<uint32*>(AddResult.TileDataOffsetsPtrs[AttributesIdx])[0] = AddResult.TileDataBaseOffsets[AttributesIdx] + SrcTileDataOffsets[0];
+
+						// Tile data
+						const uint8* SrcTileData = Resources->RootData.GetData() + RootStreamingInfo->TileDataOffset[AttributesIdx];
+						check(AddResult.TileDataPtrs[AttributesIdx]);
+						FMemory::Memcpy(AddResult.TileDataPtrs[AttributesIdx], SrcTileData, RootStreamingInfo->TileDataSize[AttributesIdx]);
+					}
 				}
 
 				// Update highest mip (1x1x1) in page table
@@ -1238,7 +1392,7 @@ void FStreamingManager::AddInternal(FRDGBuilder& GraphBuilder, FNewSparseVolumeT
 			}
 		}
 
-		RootTileUploader.ResourceUploadTo(GraphBuilder, SVTInfo.TileDataTexture->TileDataTextureARHIRef, SVTInfo.TileDataTexture->TileDataTextureBRHIRef);
+		RootTileUploader.ResourceUploadTo(GraphBuilder, SVTInfo.TileDataTexture->TileDataTextureARHIRef, SVTInfo.TileDataTexture->TileDataTextureBRHIRef, SVTInfo.FallbackValueA, SVTInfo.FallbackValueB);
 	}
 
 	InvalidatedStreamingInfos.Add(&SVTInfo);
@@ -1821,36 +1975,16 @@ void FStreamingManager::InstallReadyMipLevels()
 
 	TRACE_CPUPROFILER_EVENT_SCOPE(SVT::InstallReadyMipLevels);
 
-	const int32 StartPendingMipLevelIndex = (NextPendingMipLevelIndex + MaxPendingMipLevels - NumPendingMipLevels) % MaxPendingMipLevels;
-
-	// Do a first pass over all the mips to be uploaded to compute the upload buffer size requirements.
-	int32 NumPageTableUpdatesTotal = 0;
-	for (int32 i = 0; i < AsyncState.NumReadyMipLevels; ++i)
-	{
-		const int32 PendingMipLevelIndex = (StartPendingMipLevelIndex + i) % MaxPendingMipLevels;
-		FPendingMipLevel& PendingMipLevel = PendingMipLevels[PendingMipLevelIndex];
-
-		FStreamingInfo* SVTInfo = StreamingInfo.Find(PendingMipLevel.SparseVolumeTexture);
-		if (!SVTInfo || (SVTInfo->PerFrameInfo[PendingMipLevel.FrameIndex].LowestRequestedMipLevel > PendingMipLevel.MipLevelIndex))
-		{
-			continue; // Skip mip level install. SVT no longer exists or mip level was "streamed out" before it was even installed in the first place.
-		}
-
-		const FResources* Resources = SVTInfo->PerFrameInfo[PendingMipLevel.FrameIndex].Resources;
-		SVTInfo->TileDataTexture->NumTilesToUpload += Resources->MipLevelStreamingInfo[PendingMipLevel.MipLevelIndex].NumPhysicalTiles;
-		NumPageTableUpdatesTotal += Resources->MipLevelStreamingInfo[PendingMipLevel.MipLevelIndex].PageTableSize / (2 * sizeof(uint32));
-	}
-
 	UploadTasks.Reset();
 	UploadTasks.Reserve(AsyncState.NumReadyMipLevels * 2 /*slack for splitting large uploads*/);
 	UploadCleanupTasks.Reset();
-	
 
 #if WITH_EDITORONLY_DATA
 	TMap<const FResources*, const uint8*> ResourceToBulkPointer;
 #endif
 
 	// Do a second pass over all ready mip levels, claiming memory in the upload buffers and creating FUploadTasks
+	const int32 StartPendingMipLevelIndex = (NextPendingMipLevelIndex + MaxPendingMipLevels - NumPendingMipLevels) % MaxPendingMipLevels;
 	for (int32 i = 0; i < AsyncState.NumReadyMipLevels; ++i)
 	{
 		const int32 PendingMipLevelIndex = (StartPendingMipLevelIndex + i) % MaxPendingMipLevels;
@@ -1902,16 +2036,18 @@ void FStreamingManager::InstallReadyMipLevels()
 
 		check(SrcPtr);
 
+		const int32 FormatSizeA = GPixelFormats[SVTInfo->FormatA].BlockBytes;
+		const int32 FormatSizeB = GPixelFormats[SVTInfo->FormatB].BlockBytes;
+
 		const int32 NumPhysicalTiles = MipLevelStreamingInfo.NumPhysicalTiles;
+		const int32 NumVoxelsA = FormatSizeA > 0 ? MipLevelStreamingInfo.TileDataSize[0] / FormatSizeA : 0;
+		const int32 NumVoxelsB = FormatSizeB > 0 ? MipLevelStreamingInfo.TileDataSize[1] / FormatSizeB : 0;
 		TArray<uint32>& TileAllocations = FrameInfo.TileAllocations[PendingMipLevel.MipLevelIndex];
 		check(TileAllocations.Num() == NumPhysicalTiles);
 		check((MipLevelStreamingInfo.PageTableSize % (sizeof(uint32) * 2)) == 0);
 		const int32 NumPageTableUpdates = MipLevelStreamingInfo.PageTableSize / (sizeof(uint32) * 2);
 
-		uint8* DstPhysicalTileCoords = nullptr;
-		uint8* DstPhysicalTileDataA = nullptr;
-		uint8* DstPhysicalTileDataB = nullptr;
-		SVTInfo->TileDataTexture->TileUploader->Add_GetRef(NumPhysicalTiles, DstPhysicalTileCoords, DstPhysicalTileDataA, DstPhysicalTileDataB);
+		FTileUploader::FAddResult TileDataAddResult = SVTInfo->TileDataTexture->TileUploader->Add_GetRef(NumPhysicalTiles, NumVoxelsA, NumVoxelsB);
 
 		uint8* DstPageCoords = nullptr;
 		uint8* DstPageEntries = nullptr;
@@ -1919,32 +2055,38 @@ void FStreamingManager::InstallReadyMipLevels()
 
 		// Tile data
 		{
-			check(MipLevelStreamingInfo.TileDataASize == (SVT::NumVoxelsPerPaddedTile * GPixelFormats[SVTInfo->TileDataTexture->FormatA].BlockBytes * NumPhysicalTiles));
-			check(MipLevelStreamingInfo.TileDataBSize == (SVT::NumVoxelsPerPaddedTile * GPixelFormats[SVTInfo->TileDataTexture->FormatB].BlockBytes * NumPhysicalTiles));
+			FUploadTask::FTileDataTask TileDataTask = {};
+			TileDataTask.DstOccupancyBitsPtrs = TileDataAddResult.OccupancyBitsPtrs;
+			TileDataTask.DstTileDataOffsetsPtrs = TileDataAddResult.TileDataOffsetsPtrs;
+			TileDataTask.DstTileDataPtrs = TileDataAddResult.TileDataPtrs;
+			TileDataTask.DstPhysicalTileCoords = TileDataAddResult.PackedPhysicalTileCoordsPtr;
+			TileDataTask.SrcOccupancyBitsPtrs[0] = SrcPtr + MipLevelStreamingInfo.OccupancyBitsOffset[0];
+			TileDataTask.SrcOccupancyBitsPtrs[1] = SrcPtr + MipLevelStreamingInfo.OccupancyBitsOffset[1];
+			TileDataTask.SrcTileDataOffsetsPtrs[0] = SrcPtr + MipLevelStreamingInfo.TileDataOffsetsOffset[0];
+			TileDataTask.SrcTileDataOffsetsPtrs[1] = SrcPtr + MipLevelStreamingInfo.TileDataOffsetsOffset[1];
+			TileDataTask.SrcTileDataPtrs[0] = SrcPtr + MipLevelStreamingInfo.TileDataOffset[0];
+			TileDataTask.SrcTileDataPtrs[1] = SrcPtr + MipLevelStreamingInfo.TileDataOffset[1];
+			TileDataTask.SrcPhysicalTileCoords = reinterpret_cast<const uint8*>(TileAllocations.GetData());
+			TileDataTask.TileDataBaseOffsets = TileDataAddResult.TileDataBaseOffsets;
+			TileDataTask.TileDataSizes = MipLevelStreamingInfo.TileDataSize;
+			TileDataTask.NumPhysicalTiles = NumPhysicalTiles;
 
 			FUploadTask& Task = UploadTasks.AddDefaulted_GetRef();
-			Task.TaskType = FUploadTask::ETaskType::TileData;
-			Task.TileDataTask.DstA = DstPhysicalTileDataA;
-			Task.TileDataTask.DstB = DstPhysicalTileDataB;
-			Task.TileDataTask.DstPhysicalTileCoords = DstPhysicalTileCoords;
-			Task.TileDataTask.SrcA = SrcPtr + MipLevelStreamingInfo.TileDataAOffset;
-			Task.TileDataTask.SrcB = SrcPtr + MipLevelStreamingInfo.TileDataBOffset;
-			Task.TileDataTask.SrcPhysicalTileCoords = reinterpret_cast<const uint8*>(TileAllocations.GetData());
-			Task.TileDataTask.SizeA = MipLevelStreamingInfo.TileDataASize;
-			Task.TileDataTask.SizeB = MipLevelStreamingInfo.TileDataBSize;
-			Task.TileDataTask.NumPhysicalTiles = NumPhysicalTiles;
+			Task.Union.SetSubtype<FUploadTask::FTileDataTask>(TileDataTask);
 		}
 
 		// Page table
 		{
+			FUploadTask::FPageTableTask PageTableTask = {};
+			PageTableTask.PendingMipLevel = &PendingMipLevel;
+			PageTableTask.DstPageCoords = DstPageCoords;
+			PageTableTask.DstPageEntries = DstPageEntries;
+			PageTableTask.SrcPageCoords = SrcPtr + MipLevelStreamingInfo.PageTableOffset;
+			PageTableTask.SrcPageEntries = SrcPtr + MipLevelStreamingInfo.PageTableOffset + NumPageTableUpdates * sizeof(uint32);
+			PageTableTask.NumPageTableUpdates = NumPageTableUpdates;
+
 			FUploadTask& Task = UploadTasks.AddDefaulted_GetRef();
-			Task.TaskType = FUploadTask::ETaskType::PageTable;
-			Task.PageTableTask.PendingMipLevel = &PendingMipLevel;
-			Task.PageTableTask.DstPageCoords = DstPageCoords;
-			Task.PageTableTask.DstPageEntries = DstPageEntries;
-			Task.PageTableTask.SrcPageCoords = SrcPtr + MipLevelStreamingInfo.PageTableOffset;
-			Task.PageTableTask.SrcPageEntries = SrcPtr + MipLevelStreamingInfo.PageTableOffset + NumPageTableUpdates * sizeof(uint32);
-			Task.PageTableTask.NumPageTableUpdates = NumPageTableUpdates;
+			Task.Union.SetSubtype<FUploadTask::FPageTableTask>(PageTableTask);
 		}
 
 		// Cleanup
@@ -1978,41 +2120,52 @@ void FStreamingManager::InstallReadyMipLevels()
 
 			FUploadTask& Task = UploadTasks[TaskIndex];
 
-			switch (Task.TaskType)
-			{
-			case FUploadTask::ETaskType::PageTable:
+			if (Task.Union.HasSubtype<FUploadTask::FPageTableTask>())
 			{
 				TRACE_CPUPROFILER_EVENT_SCOPE(SVT::PageTableUpload);
-				if (Task.PageTableTask.NumPageTableUpdates > 0)
+				FUploadTask::FPageTableTask& PageTableTask = Task.Union.GetSubtype<FUploadTask::FPageTableTask>();
+				if (PageTableTask.NumPageTableUpdates > 0)
 				{
-					FMemory::Memcpy(Task.PageTableTask.DstPageCoords, Task.PageTableTask.SrcPageCoords, Task.PageTableTask.NumPageTableUpdates * sizeof(uint32));
+					FMemory::Memcpy(PageTableTask.DstPageCoords, PageTableTask.SrcPageCoords, PageTableTask.NumPageTableUpdates * sizeof(uint32));
 
-					FStreamingInfo* SVTInfo = StreamingInfo.Find(Task.PageTableTask.PendingMipLevel->SparseVolumeTexture);
-					TArray<uint32>& TileAllocations = SVTInfo->PerFrameInfo[Task.PageTableTask.PendingMipLevel->FrameIndex].TileAllocations[Task.PageTableTask.PendingMipLevel->MipLevelIndex];
-					const uint32* SrcEntries = reinterpret_cast<const uint32*>(Task.PageTableTask.SrcPageEntries);
-					uint32* DstEntries = reinterpret_cast<uint32*>(Task.PageTableTask.DstPageEntries);
-					for (int32 i = 0; i < Task.PageTableTask.NumPageTableUpdates; ++i)
+					FStreamingInfo* SVTInfo = StreamingInfo.Find(PageTableTask.PendingMipLevel->SparseVolumeTexture);
+					TArray<uint32>& TileAllocations = SVTInfo->PerFrameInfo[PageTableTask.PendingMipLevel->FrameIndex].TileAllocations[PageTableTask.PendingMipLevel->MipLevelIndex];
+					const uint32* SrcEntries = reinterpret_cast<const uint32*>(PageTableTask.SrcPageEntries);
+					uint32* DstEntries = reinterpret_cast<uint32*>(PageTableTask.DstPageEntries);
+					for (int32 i = 0; i < PageTableTask.NumPageTableUpdates; ++i)
 					{
 						DstEntries[i] = TileAllocations[SrcEntries[i]];
 					}
 				}
-				break;
 			}
-			case FUploadTask::ETaskType::TileData:
+			else if(Task.Union.HasSubtype<FUploadTask::FTileDataTask>())
 			{
 				TRACE_CPUPROFILER_EVENT_SCOPE(SVT::TileDataUpload);
-				FMemory::Memcpy(Task.TileDataTask.DstPhysicalTileCoords, Task.TileDataTask.SrcPhysicalTileCoords, Task.TileDataTask.NumPhysicalTiles * sizeof(uint32));
-				if (Task.TileDataTask.SizeA > 0)
-				{
-					FMemory::Memcpy(Task.TileDataTask.DstA, Task.TileDataTask.SrcA, Task.TileDataTask.SizeA);
-				}
-				if (Task.TileDataTask.SizeB > 0)
+				FUploadTask::FTileDataTask& TileDataTask = Task.Union.GetSubtype<FUploadTask::FTileDataTask>();
+				for (int32 i = 0; i < 2; ++i)
 				{
-					FMemory::Memcpy(Task.TileDataTask.DstB, Task.TileDataTask.SrcB, Task.TileDataTask.SizeB);
+					if (TileDataTask.DstOccupancyBitsPtrs[i])
+					{
+						FMemory::Memcpy(TileDataTask.DstOccupancyBitsPtrs[i], TileDataTask.SrcOccupancyBitsPtrs[i], TileDataTask.NumPhysicalTiles * SVT::NumOccupancyWordsPerPaddedTile * sizeof(uint32));
+					}
+					if (TileDataTask.TileDataSizes[i] > 0)
+					{
+						FMemory::Memcpy(TileDataTask.DstTileDataPtrs[i], TileDataTask.SrcTileDataPtrs[i], TileDataTask.TileDataSizes[i]);
+					}
+					if (TileDataTask.DstTileDataOffsetsPtrs[i])
+					{
+						uint32* Dst = reinterpret_cast<uint32*>(TileDataTask.DstTileDataOffsetsPtrs[i]);
+						const uint32* Src = reinterpret_cast<const uint32*>(TileDataTask.SrcTileDataOffsetsPtrs[i]);
+						for (int32 TileIndex = 0; TileIndex < TileDataTask.NumPhysicalTiles; ++TileIndex)
+						{
+							Dst[TileIndex] = Src[TileIndex] + TileDataTask.TileDataBaseOffsets[i];
+						}
+					}
 				}
-				break;
+				FMemory::Memcpy(TileDataTask.DstPhysicalTileCoords, TileDataTask.SrcPhysicalTileCoords, TileDataTask.NumPhysicalTiles * sizeof(uint32));
 			}
-			default:
+			else
+			{
 				checkNoEntry();
 			}
 		});
@@ -2112,13 +2265,17 @@ void FStreamingManager::RequestDDCData(TConstArrayView<UE::DerivedData::FCacheGe
 
 #endif // WITH_EDITORONLY_DATA
 
-FStreamingManager::FTileDataTexture::FTileDataTexture(const FIntVector3& InResolutionInTiles, EPixelFormat InFormatA, EPixelFormat InFormatB)
+FStreamingManager::FTileDataTexture::FTileDataTexture(const FIntVector3& InResolutionInTiles, EPixelFormat InFormatA, EPixelFormat InFormatB, const FVector4f& InFallbackValueA, const FVector4f& InFallbackValueB)
 	: ResolutionInTiles(InResolutionInTiles), 
 	PhysicalTilesCapacity(InResolutionInTiles.X * InResolutionInTiles.Y * InResolutionInTiles.Z), 
 	FormatA(InFormatA),
 	FormatB(InFormatB),
+	FallbackValueA(InFallbackValueA),
+	FallbackValueB(InFallbackValueB),
 	TileUploader(MakeUnique<FTileUploader>()),
-	NumTilesToUpload(0)
+	NumTilesToUpload(0),
+	NumVoxelsToUploadA(0),
+	NumVoxelsToUploadB(0)
 {
 	const int64 MaxFormatSize = FMath::Max(GPixelFormats[FormatA].BlockBytes, GPixelFormats[FormatB].BlockBytes);
 	const FIntVector3 LargestPossibleResolution = ComputeLargestPossibleTileDataVolumeResolution(MaxFormatSize);
diff --git a/Engine/Source/Runtime/Engine/Private/SparseVolumeTexture/SparseVolumeTextureStreamingManager.h b/Engine/Source/Runtime/Engine/Private/SparseVolumeTexture/SparseVolumeTextureStreamingManager.h
--- a/Engine/Source/Runtime/Engine/Private/SparseVolumeTexture/SparseVolumeTextureStreamingManager.h
+++ b/Engine/Source/Runtime/Engine/Private/SparseVolumeTexture/SparseVolumeTextureStreamingManager.h
@@ -11,6 +11,8 @@
 #include "RenderResource.h"
 #include "Containers/IntrusiveDoubleLinkedList.h"
 #include "Containers/Map.h"
+#include "Containers/StaticArray.h"
+#include "Containers/Union.h"
 #include "RenderGraphBuilder.h"
 
 class UStreamableSparseVolumeTexture;
@@ -112,13 +114,17 @@ private:
 		int32 PhysicalTilesCapacity;
 		EPixelFormat FormatA;
 		EPixelFormat FormatB;
+		FVector4f FallbackValueA;
+		FVector4f FallbackValueB;
 		FTextureRHIRef TileDataTextureARHIRef;
 		FTextureRHIRef TileDataTextureBRHIRef;
 		TUniquePtr<class FTileUploader> TileUploader;
 		int32 NumTilesToUpload;
+		int32 NumVoxelsToUploadA;
+		int32 NumVoxelsToUploadB;
 
 		// Constructor. May change the requested ResolutionInTiles (and resulting PhysicalTilesCapacity) if it exceeds hardware limits.
-		FTileDataTexture(const FIntVector3& ResolutionInTiles, EPixelFormat FormatA, EPixelFormat FormatB);
+		FTileDataTexture(const FIntVector3& ResolutionInTiles, EPixelFormat FormatA, EPixelFormat FormatB, const FVector4f& FallbackValueA, const FVector4f& FallbackValueB);
 
 		// Allocate a tile slot in the texture. The resulting value is a packed coordinate (8 bit per component) of the allocated slot or INDEX_NONE if the allocation failed.
 		// The upper 8 bit are free to be used by the caller.
@@ -256,23 +262,21 @@ private:
 	};
 
 	// Encapsulates all the src/dst pointers for memcpy-ing the streamed data into GPU memory.
+	// See the comment on FMipLevelStreamingInfo for details about these pointers.
 	struct FUploadTask
 	{
-		enum class ETaskType
-		{
-			TileData, PageTable
-		};
-
 		struct FTileDataTask
 		{
-			uint8* DstA;
-			uint8* DstB;
+			TStaticArray<uint8*, 2> DstOccupancyBitsPtrs;
+			TStaticArray<uint8*, 2> DstTileDataOffsetsPtrs;
+			TStaticArray<uint8*, 2> DstTileDataPtrs;
 			uint8* DstPhysicalTileCoords;
-			const uint8* SrcA;
-			const uint8* SrcB;
+			TStaticArray<const uint8*, 2> SrcOccupancyBitsPtrs;
+			TStaticArray<const uint8*, 2> SrcTileDataOffsetsPtrs;
+			TStaticArray<const uint8*, 2> SrcTileDataPtrs;
 			const uint8* SrcPhysicalTileCoords;
-			int32 SizeA;
-			int32 SizeB;
+			TStaticArray<uint32, 2> TileDataBaseOffsets;
+			TStaticArray<int32, 2> TileDataSizes;
 			int32 NumPhysicalTiles;
 		};
 
@@ -286,12 +290,7 @@ private:
 			int32 NumPageTableUpdates;
 		};
 
-		ETaskType TaskType = ETaskType::TileData;
-		union
-		{
-			FTileDataTask TileDataTask;
-			FPageTableTask PageTableTask;
-		};
+		TUnion<FTileDataTask, FPageTableTask> Union;
 	};
 
 	struct FPageTableClear
diff --git a/Engine/Source/Runtime/Engine/Private/SparseVolumeTexture/SparseVolumeTextureUtility.h b/Engine/Source/Runtime/Engine/Private/SparseVolumeTexture/SparseVolumeTextureUtility.h
--- a/Engine/Source/Runtime/Engine/Private/SparseVolumeTexture/SparseVolumeTextureUtility.h
+++ b/Engine/Source/Runtime/Engine/Private/SparseVolumeTexture/SparseVolumeTextureUtility.h
@@ -15,6 +15,7 @@ namespace SVT
 	static constexpr int64 MaxResourceSize = 2048LL * 1024LL * 1024LL;
 	static constexpr int32 NumVoxelsPerTile = SPARSE_VOLUME_TILE_RES * SPARSE_VOLUME_TILE_RES * SPARSE_VOLUME_TILE_RES;
 	static constexpr int32 NumVoxelsPerPaddedTile = SPARSE_VOLUME_TILE_RES_PADDED * SPARSE_VOLUME_TILE_RES_PADDED * SPARSE_VOLUME_TILE_RES_PADDED;
+	static constexpr int32 NumOccupancyWordsPerPaddedTile = (NumVoxelsPerPaddedTile + 31) / 32;
 
 	uint32 PackPageTableEntry(const FIntVector3& Coord);
 	FIntVector3 UnpackPageTableEntry(uint32 Packed);
