Packed Level Actor: code cleanup

#rnx

diff --git a/Engine/Source/Runtime/Engine/Private/LevelInstance/LevelInstanceComponent.cpp b/Engine/Source/Runtime/Engine/Private/LevelInstance/LevelInstanceComponent.cpp
--- a/Engine/Source/Runtime/Engine/Private/LevelInstance/LevelInstanceComponent.cpp
+++ b/Engine/Source/Runtime/Engine/Private/LevelInstance/LevelInstanceComponent.cpp
@@ -71,14 +71,18 @@ void ULevelInstanceComponent::OnRegister()
 #endif //WITH_EDITORONLY_DATA
 }
 
-void ULevelInstanceComponent::SetFilter(const FWorldPartitionActorFilter& InFilter)
+void ULevelInstanceComponent::SetFilter(const FWorldPartitionActorFilter& InFilter, bool bNotify)
 {
 	if (GetFilter() != InFilter)
 	{
 		Modify(!IsEditFilter());
 		SetActiveFilter(InFilter);
 		OnFilterChanged();
-		FWorldPartitionActorFilter::GetOnWorldPartitionActorFilterChanged().Broadcast();
+
+		if (bNotify)
+		{
+			FWorldPartitionActorFilter::GetOnWorldPartitionActorFilterChanged().Broadcast();
+		}
 	}
 }
 
diff --git a/Engine/Source/Runtime/Engine/Private/LevelInstance/LevelInstanceInterface.cpp b/Engine/Source/Runtime/Engine/Private/LevelInstance/LevelInstanceInterface.cpp
--- a/Engine/Source/Runtime/Engine/Private/LevelInstance/LevelInstanceInterface.cpp
+++ b/Engine/Source/Runtime/Engine/Private/LevelInstance/LevelInstanceInterface.cpp
@@ -318,11 +318,11 @@ const TMap<FActorContainerID, TSet<FGuid>>& ILevelInstanceInterface::GetFiltered
 	return NoFilteredActors;
 }
 
-void ILevelInstanceInterface::SetFilter(const FWorldPartitionActorFilter& InFilter)
+void ILevelInstanceInterface::SetFilter(const FWorldPartitionActorFilter& InFilter, bool bNotify)
 {
 	if (ULevelInstanceComponent* LevelInstanceComponent = GetLevelInstanceComponent())
 	{
-		LevelInstanceComponent->SetFilter(InFilter);
+		LevelInstanceComponent->SetFilter(InFilter, bNotify);
 	}
 }
 
diff --git a/Engine/Source/Runtime/Engine/Private/PackedLevelActor/PackedLevelActorBuilder.cpp b/Engine/Source/Runtime/Engine/Private/PackedLevelActor/PackedLevelActorBuilder.cpp
--- a/Engine/Source/Runtime/Engine/Private/PackedLevelActor/PackedLevelActorBuilder.cpp
+++ b/Engine/Source/Runtime/Engine/Private/PackedLevelActor/PackedLevelActorBuilder.cpp
@@ -36,19 +36,6 @@
 
 #define LOCTEXT_NAMESPACE "FPackedLevelActorBuilder"
 
-void FPackedLevelActorBuilderContext::ClusterLevelActor(AActor* InActor)
-{
-	if (!ActorDiscards.Contains(InActor))
-	{
-		PerActorClusteredComponents.FindOrAdd(InActor);
-
-		for (const auto& Pair : Builders)
-		{
-			Pair.Value->GetPackClusters(*this, InActor);
-		}
-	}
-}
-
 void FPackedLevelActorBuilderContext::FindOrAddCluster(FPackedLevelActorBuilderClusterID&& InClusterID, UActorComponent* InComponent)
 {
 	TArray<UActorComponent*>& ClusterComponents = Clusters.FindOrAdd(MoveTemp(InClusterID));
@@ -139,7 +126,6 @@ TSharedPtr<FPackedLevelActorBuilder> FPackedLevelActorBuilder::CreateDefaultBuil
 
 	// Class Discards are used to validate the packing result.
 	// Components or Actor classes in this set will not generate warnings
-	Builder->ClassDiscards.Add(ALevelBounds::StaticClass());
 	// Avoid dependency (find class)
 	UClass* ChaosDebugClass = FindObject<UClass>(nullptr, TEXT("/Script/ChaosSolverEngine.ChaosDebugDrawComponent"));
 	if (ChaosDebugClass)
@@ -150,12 +136,25 @@ TSharedPtr<FPackedLevelActorBuilder> FPackedLevelActorBuilder::CreateDefaultBuil
 	// Root Components that are SceneComponents (not child class of)
 	Builder->ClassDiscards.Add(USceneComponent::StaticClass());
 	
-	Builder->Builders.Add(FPackedLevelActorRecursiveBuilder::BuilderID, MakeUnique<FPackedLevelActorRecursiveBuilder>());
-	Builder->Builders.Add(FPackedLevelActorISMBuilder::BuilderID, MakeUnique<FPackedLevelActorISMBuilder>());
+	Builder->AddBuilder<FPackedLevelActorRecursiveBuilder>();
+	Builder->AddBuilder<FPackedLevelActorISMBuilder>();
 
 	return Builder;
 }
 
+void FPackedLevelActorBuilder::ClusterActor(FPackedLevelActorBuilderContext& InContext, AActor* InActor)
+{
+	if (!InContext.ActorDiscards.Contains(InActor))
+	{
+		InContext.PerActorClusteredComponents.FindOrAdd(InActor);
+
+		for (const auto& Pair : Builders)
+		{
+			Pair.Value->GetPackClusters(InContext, InActor);
+		}
+	}
+}
+
 bool FPackedLevelActorBuilder::PackActor(APackedLevelActor* InPackedLevelActor)
 {
 	return PackActor(InPackedLevelActor, InPackedLevelActor);
@@ -163,7 +162,7 @@ bool FPackedLevelActorBuilder::PackActor(APackedLevelActor* InPackedLevelActor)
 
 bool FPackedLevelActorBuilder::PackActor(APackedLevelActor* InPackedLevelActor, ILevelInstanceInterface* InLevelInstanceToPack)
 {
-	FPackedLevelActorBuilderContext Context(*this, InPackedLevelActor, InLevelInstanceToPack);
+	FPackedLevelActorBuilderContext Context(InPackedLevelActor, InLevelInstanceToPack, ClassDiscards);
 	return PackActor(Context);
 }
 
@@ -193,7 +192,7 @@ bool FPackedLevelActorBuilder::PackActor(FPackedLevelActorBuilderContext& InCont
 	FTransform RelativePivotTransform(NewPivotTransform.GetRelativeTransform(CurrentPivotTransform));
 		
 	InContext.SetRelativePivotTransform(RelativePivotTransform);	
-	InContext.ClusterLevelActor(CastChecked<AActor>(InContext.GetLevelInstanceToPack()));
+	ClusterActor(InContext, CastChecked<AActor>(InContext.GetLevelInstanceToPack()));
 
 	for (const auto& Pair : InContext.GetClusters())
 	{
@@ -297,7 +296,7 @@ ALevelInstance* FPackedLevelActorBuilder::CreateTransientLevelInstanceForPacking
 	check(World);
 	SpawnParams.OverrideLevel = World->PersistentLevel;
 	APackedLevelActor* LevelInstance = World->SpawnActor<APackedLevelActor>(InLocation, InRotator, SpawnParams);
-	LevelInstance->SetFilter(InFilter);
+	LevelInstance->SetFilter(InFilter, false);
 	LevelInstance->SetShouldLoadForPacking(true);
 	LevelInstance->SetWorldAsset(InWorldAsset);
 
@@ -376,8 +375,8 @@ bool FPackedLevelActorBuilder::CreateOrUpdateBlueprintFromUnpacked(ILevelInstanc
 	SpawnParams.OverrideLevel = World->PersistentLevel;
 
 	APackedLevelActor* PackedLevelActor = World->SpawnActor<APackedLevelActor>(LevelInstanceActor->GetActorLocation(), LevelInstanceActor->GetActorRotation(), SpawnParams);
+	PackedLevelActor->SetFilter(InLevelInstance->GetFilter(), false);
 	PackedLevelActor->SetWorldAsset(InLevelInstance->GetWorldAsset());
-	PackedLevelActor->SetFilter(InLevelInstance->GetFilter());
 	ON_SCOPE_EXIT
 	{
 		LevelInstanceActor->GetWorld()->DestroyActor(PackedLevelActor);
@@ -472,8 +471,7 @@ bool FPackedLevelActorBuilder::CreateOrUpdateBlueprintFromPacked(APackedLevelAct
 
 		PackedLevelActor->Modify(false);
 		PackedLevelActor->SetWorldAsset(InActor->GetWorldAsset());
-		PackedLevelActor->SetFilter(InActor->GetFilter());
-		PackedLevelActor->PackedBPDependencies = InActor->PackedBPDependencies;
+		PackedLevelActor->SetFilter(InActor->GetFilter(), false);
 		PackedLevelActor->SetPackedVersion(NewVersion);
 		// match root component mobility to source actor
 		USceneComponent* Root = PackedLevelActor->GetRootComponent();
@@ -488,6 +486,8 @@ bool FPackedLevelActorBuilder::CreateOrUpdateBlueprintFromPacked(APackedLevelAct
 			
 	// Synchronous compile
 	FKismetEditorUtilities::CompileBlueprint(BP, EBlueprintCompileOptions::SkipGarbageCollection);
+	FMessageLog Log("PackedLevelActor");
+	Log.Info(FText::Format(LOCTEXT("PackedBlueprintUpdated", "Updated Packed Blueprint '{0}'"), FText::FromString(InBlueprintAsset.ToString())));
 
 	if (bCheckoutAndSave)
 	{
diff --git a/Engine/Source/Runtime/Engine/Private/PackedLevelActor/PackedLevelActorISMBuilder.h b/Engine/Source/Runtime/Engine/Private/PackedLevelActor/PackedLevelActorISMBuilder.h
--- a/Engine/Source/Runtime/Engine/Private/PackedLevelActor/PackedLevelActorISMBuilder.h
+++ b/Engine/Source/Runtime/Engine/Private/PackedLevelActor/PackedLevelActorISMBuilder.h
@@ -13,6 +13,9 @@ class UStaticMeshComponent;
 class FPackedLevelActorISMBuilder : public IPackedLevelActorBuilder
 {
 public:
+	FPackedLevelActorISMBuilder(FPackedLevelActorBuilder& InOwner)
+		: IPackedLevelActorBuilder(InOwner){}
+
 	static FPackedLevelActorBuilderID BuilderID;
 
 	virtual FPackedLevelActorBuilderID GetID() const override;
diff --git a/Engine/Source/Runtime/Engine/Private/PackedLevelActor/PackedLevelActorRecursiveBuilder.cpp b/Engine/Source/Runtime/Engine/Private/PackedLevelActor/PackedLevelActorRecursiveBuilder.cpp
--- a/Engine/Source/Runtime/Engine/Private/PackedLevelActor/PackedLevelActorRecursiveBuilder.cpp
+++ b/Engine/Source/Runtime/Engine/Private/PackedLevelActor/PackedLevelActorRecursiveBuilder.cpp
@@ -26,9 +26,6 @@ void FPackedLevelActorRecursiveBuilder::GetPackClusters(FPackedLevelActorBuilder
 {
 	if (ILevelInstanceInterface* LevelInstance = Cast<ILevelInstanceInterface>(InActor))
 	{
-		FPackedLevelActorBuilderClusterID ClusterID(MakeUnique<FPackedLevelActorRecursiveBuilderCluster>(GetID(), LevelInstance));
-		InContext.FindOrAddCluster(MoveTemp(ClusterID));
-
 		// This Actor can be safely discarded without warning because it is a container
 		InContext.DiscardActor(InActor);
 
@@ -40,17 +37,15 @@ void FPackedLevelActorRecursiveBuilder::GetPackClusters(FPackedLevelActorBuilder
 			{
 				if (LevelActor)
 				{
-					if (LevelActor == Level->GetDefaultBrush())
-					{
-						InContext.DiscardActor(LevelActor);
-					}
-					else if (LevelActor == Level->GetWorldSettings())
+					if (LevelActor == Level->GetDefaultBrush() || 
+						LevelActor == Level->GetWorldSettings() ||
+						LevelActor->IsMainWorldOnly())
 					{
 						InContext.DiscardActor(LevelActor);
 					}
 					else
 					{
-						InContext.ClusterLevelActor(LevelActor);
+						Owner.ClusterActor(InContext, LevelActor);
 					}
 				}
 			}
@@ -58,40 +53,4 @@ void FPackedLevelActorRecursiveBuilder::GetPackClusters(FPackedLevelActorBuilder
 	}
 }
 
-void FPackedLevelActorRecursiveBuilder::PackActors(FPackedLevelActorBuilderContext& InContext, const FPackedLevelActorBuilderClusterID& InClusterID, const TArray<UActorComponent*>& InComponents) const
-{
-	check(InClusterID.GetBuilderID() == GetID());
-	FPackedLevelActorRecursiveBuilderCluster* LevelInstanceCluster = (FPackedLevelActorRecursiveBuilderCluster*)InClusterID.GetData();
-	check(LevelInstanceCluster);
-
-	if (APackedLevelActor* PackedLevelActor = Cast<APackedLevelActor>(LevelInstanceCluster->LevelInstance))
-	{
-		if (UBlueprint* GeneratedBy = PackedLevelActor->GetRootBlueprint())
-		{
-			InContext.GetPackedLevelActor()->PackedBPDependencies.AddUnique(GeneratedBy);
-		}
-	}
-}
-
-FPackedLevelActorRecursiveBuilderCluster::FPackedLevelActorRecursiveBuilderCluster(FPackedLevelActorBuilderID InBuilderID, ILevelInstanceInterface* InLevelInstance)
-	: FPackedLevelActorBuilderCluster(InBuilderID), LevelInstance(InLevelInstance)
-{
-}
-
-bool FPackedLevelActorRecursiveBuilderCluster::Equals(const FPackedLevelActorBuilderCluster& InOther) const
-{
-	if (!FPackedLevelActorBuilderCluster::Equals(InOther))
-	{
-		return false;
-	}
-
-	const FPackedLevelActorRecursiveBuilderCluster& LevelInstanceCluster = (const FPackedLevelActorRecursiveBuilderCluster&)InOther;
-	return LevelInstance == LevelInstanceCluster.LevelInstance;
-}
-	
-uint32 FPackedLevelActorRecursiveBuilderCluster::ComputeHash() const
-{
-	return FCrc::TypeCrc32(LevelInstance->GetLevelInstanceID(), FPackedLevelActorBuilderCluster::ComputeHash());
-}
-
 #endif
diff --git a/Engine/Source/Runtime/Engine/Private/PackedLevelActor/PackedLevelActorRecursiveBuilder.h b/Engine/Source/Runtime/Engine/Private/PackedLevelActor/PackedLevelActorRecursiveBuilder.h
--- a/Engine/Source/Runtime/Engine/Private/PackedLevelActor/PackedLevelActorRecursiveBuilder.h
+++ b/Engine/Source/Runtime/Engine/Private/PackedLevelActor/PackedLevelActorRecursiveBuilder.h
@@ -12,22 +12,14 @@ class ILevelInstanceInterface;
 class FPackedLevelActorRecursiveBuilder : public IPackedLevelActorBuilder
 {
 public:
+	FPackedLevelActorRecursiveBuilder(FPackedLevelActorBuilder& InOwner)
+		: IPackedLevelActorBuilder(InOwner) {}
+
 	static FPackedLevelActorBuilderID BuilderID;
 
 	virtual FPackedLevelActorBuilderID GetID() const override;
 	virtual void GetPackClusters(FPackedLevelActorBuilderContext& InContext, AActor* InActor) const override;
-	virtual void PackActors(FPackedLevelActorBuilderContext& InContext, const FPackedLevelActorBuilderClusterID& InClusterID, const TArray<UActorComponent*>& InComponents) const override;
-};
-
-class FPackedLevelActorRecursiveBuilderCluster : public FPackedLevelActorBuilderCluster
-{
-public:
-	FPackedLevelActorRecursiveBuilderCluster(FPackedLevelActorBuilderID InBuilderID, ILevelInstanceInterface* InLevelInstance);
-
-	virtual bool Equals(const FPackedLevelActorBuilderCluster& InOther) const override;
-	virtual uint32 ComputeHash() const override;
-
-	ILevelInstanceInterface* LevelInstance = nullptr;
+	virtual void PackActors(FPackedLevelActorBuilderContext& InContext, const FPackedLevelActorBuilderClusterID& InClusterID, const TArray<UActorComponent*>& InComponents) const override { }
 };
 
 #endif
\ No newline at end of file
diff --git a/Engine/Source/Runtime/Engine/Public/LevelInstance/LevelInstanceComponent.h b/Engine/Source/Runtime/Engine/Public/LevelInstance/LevelInstanceComponent.h
--- a/Engine/Source/Runtime/Engine/Public/LevelInstance/LevelInstanceComponent.h
+++ b/Engine/Source/Runtime/Engine/Public/LevelInstance/LevelInstanceComponent.h
@@ -42,7 +42,7 @@ public:
 	ENGINE_API void OnCommit();
 
 	const FWorldPartitionActorFilter& GetFilter() const { return IsEditFilter() ? EditFilter : Filter; }
-	ENGINE_API void SetFilter(const FWorldPartitionActorFilter& InFilter);
+	ENGINE_API void SetFilter(const FWorldPartitionActorFilter& InFilter, bool bNotify = true);
 	ENGINE_API const TMap<FActorContainerID, TSet<FGuid>>& GetFilteredActorsPerContainer() const;
 	ENGINE_API void UpdateEditFilter();
 private:
diff --git a/Engine/Source/Runtime/Engine/Public/LevelInstance/LevelInstanceInterface.h b/Engine/Source/Runtime/Engine/Public/LevelInstance/LevelInstanceInterface.h
--- a/Engine/Source/Runtime/Engine/Public/LevelInstance/LevelInstanceInterface.h
+++ b/Engine/Source/Runtime/Engine/Public/LevelInstance/LevelInstanceInterface.h
@@ -125,7 +125,7 @@ class ILevelInstanceInterface
 
 	ENGINE_API virtual const TMap<FActorContainerID, TSet<FGuid>>& GetFilteredActorsPerContainer() const;
 
-	ENGINE_API virtual void SetFilter(const FWorldPartitionActorFilter& InFilter);
+	ENGINE_API virtual void SetFilter(const FWorldPartitionActorFilter& InFilter, bool bNotify = true);
 
 	virtual void OnFilterChanged() {}
 
diff --git a/Engine/Source/Runtime/Engine/Public/PackedLevelActor/IPackedLevelActorBuilder.h b/Engine/Source/Runtime/Engine/Public/PackedLevelActor/IPackedLevelActorBuilder.h
--- a/Engine/Source/Runtime/Engine/Public/PackedLevelActor/IPackedLevelActorBuilder.h
+++ b/Engine/Source/Runtime/Engine/Public/PackedLevelActor/IPackedLevelActorBuilder.h
@@ -12,15 +12,18 @@ class APackedLevelActor;
 class AActor;
 class UActorComponent;
 class FPackedLevelActorBuilderContext;
+class FPackedLevelActorBuilder;
 
 class IPackedLevelActorBuilder
 {
 public:
-	IPackedLevelActorBuilder() {}
+	IPackedLevelActorBuilder(FPackedLevelActorBuilder& InOwner) : Owner(InOwner) {}
 	virtual ~IPackedLevelActorBuilder() {}
 	virtual FPackedLevelActorBuilderID GetID() const = 0;
 	virtual void GetPackClusters(FPackedLevelActorBuilderContext& InContext, AActor* InActor) const = 0;
 	virtual void PackActors(FPackedLevelActorBuilderContext& InBuilder, const FPackedLevelActorBuilderClusterID& InClusterID, const TArray<UActorComponent*>& InComponents) const = 0;
+protected:
+	FPackedLevelActorBuilder& Owner;
 };
 
 #endif
diff --git a/Engine/Source/Runtime/Engine/Public/PackedLevelActor/PackedLevelActor.h b/Engine/Source/Runtime/Engine/Public/PackedLevelActor/PackedLevelActor.h
--- a/Engine/Source/Runtime/Engine/Public/PackedLevelActor/PackedLevelActor.h
+++ b/Engine/Source/Runtime/Engine/Public/PackedLevelActor/PackedLevelActor.h
@@ -90,9 +90,6 @@ public:
 #endif
 
 #if WITH_EDITORONLY_DATA
-	UPROPERTY()
-	TArray<TSoftObjectPtr<UBlueprint>> PackedBPDependencies;
-
 private:
 	UPROPERTY()
 	FGuid PackedVersion;
diff --git a/Engine/Source/Runtime/Engine/Public/PackedLevelActor/PackedLevelActorBuilder.h b/Engine/Source/Runtime/Engine/Public/PackedLevelActor/PackedLevelActorBuilder.h
--- a/Engine/Source/Runtime/Engine/Public/PackedLevelActor/PackedLevelActorBuilder.h
+++ b/Engine/Source/Runtime/Engine/Public/PackedLevelActor/PackedLevelActorBuilder.h
@@ -50,6 +50,14 @@ public:
 	/* Creates a new APackedLevelActor Blueprint using InPackagePath/InAssetName */
 	static ENGINE_API UBlueprint* CreatePackedLevelActorBlueprint(TSoftObjectPtr<UBlueprint> InBlueprintAsset, TSoftObjectPtr<UWorld> InWorldAsset, bool bInCompile);
 	
+	template<class T>
+	void AddBuilder()
+	{
+		check(!Builders.Contains(T::BuilderID));
+		Builders.Add(T::BuilderID, MakeUnique<T>(*this));
+	}
+
+	void ClusterActor(FPackedLevelActorBuilderContext& InContext, AActor* InActor);
 private:
 	ENGINE_API bool PackActor(FPackedLevelActorBuilderContext& InContext);
 
@@ -76,11 +84,10 @@ private:
 class FPackedLevelActorBuilderContext
 {
 public:
-	FPackedLevelActorBuilderContext(const FPackedLevelActorBuilder& InBuilder, APackedLevelActor* InPackedLevelActor, ILevelInstanceInterface* InLevelInstanceToPack) 
-		: Builders(InBuilder.Builders), ClassDiscards(InBuilder.ClassDiscards), PackedLevelActor(InPackedLevelActor), LevelInstanceToPack(InLevelInstanceToPack), RelativePivotTransform(FTransform::Identity) {}
+	FPackedLevelActorBuilderContext(APackedLevelActor* InPackedLevelActor, ILevelInstanceInterface* InLevelInstanceToPack, const TSet<UClass*>& InClassDiscards) 
+		: ClassDiscards(InClassDiscards), PackedLevelActor(InPackedLevelActor), LevelInstanceToPack(InLevelInstanceToPack), RelativePivotTransform(FTransform::Identity) {}
 
 	/* Interface for IPackedLevelActorBuilder's to use */
-	void ClusterLevelActor(AActor* InLevelActor);
 	void FindOrAddCluster(FPackedLevelActorBuilderClusterID&& InClusterID, UActorComponent* InComponent = nullptr);
 	void DiscardActor(AActor* InActor);
 	void Report(FMessageLog& Log) const;
@@ -95,7 +102,6 @@ public:
 	APackedLevelActor* GetPackedLevelActor() const { return PackedLevelActor; }
 	ILevelInstanceInterface* GetLevelInstanceToPack() const { return LevelInstanceToPack; }
 private:
-	const TMap<FPackedLevelActorBuilderID, TUniquePtr<IPackedLevelActorBuilder>>& Builders;
 	const TSet<UClass*>& ClassDiscards;
 
 	APackedLevelActor* PackedLevelActor;
@@ -107,6 +113,8 @@ private:
 	TSet<AActor*> ActorDiscards;
 
 	FTransform RelativePivotTransform;
+
+	friend FPackedLevelActorBuilder;
 };
 
 #endif
