Crash When Updating/Replacing Skeleton on AnimSequence with Animated Bone Attributes

diff --git a/Engine/Source/Runtime/Engine/Classes/Animation/AnimData/AttributeIdentifier.h b/Engine/Source/Runtime/Engine/Classes/Animation/AnimData/AttributeIdentifier.h
--- a/Engine/Source/Runtime/Engine/Classes/Animation/AnimData/AttributeIdentifier.h
+++ b/Engine/Source/Runtime/Engine/Classes/Animation/AnimData/AttributeIdentifier.h
@@ -18,7 +18,7 @@ struct FAnimationAttributeIdentifier
 	GENERATED_USTRUCT_BODY();
 
 	FAnimationAttributeIdentifier() : Name(NAME_None), BoneName(NAME_None), BoneIndex(INDEX_NONE), ScriptStruct(nullptr), ScriptStructPath(nullptr) {}
-	FAnimationAttributeIdentifier(const FName& InName, const int32 InBoneIndex, const FName InBoneName, UScriptStruct* InStruct) : Name(InName), BoneName(InBoneName), BoneIndex(InBoneIndex), ScriptStruct(InStruct) {}
+	FAnimationAttributeIdentifier(const FName& InName, const int32 InBoneIndex, const FName InBoneName, UScriptStruct* InStruct) : Name(InName), BoneName(InBoneName), BoneIndex(InBoneIndex), ScriptStruct(InStruct), ScriptStructPath(InStruct) {}
 
 	const FName& GetName() const { return Name; }
 	const FName& GetBoneName() const { return BoneName; }
