Add material full path name to the metadata for the Material Translate CPU marker.

diff --git a/Engine/Source/Runtime/Engine/Private/Materials/HLSLMaterialTranslator.cpp b/Engine/Source/Runtime/Engine/Private/Materials/HLSLMaterialTranslator.cpp
--- a/Engine/Source/Runtime/Engine/Private/Materials/HLSLMaterialTranslator.cpp
+++ b/Engine/Source/Runtime/Engine/Private/Materials/HLSLMaterialTranslator.cpp
@@ -806,10 +806,23 @@ void FHLSLMaterialTranslator::ValidateVtPropertyLimits()
 		}
 	}
 }
+
+UE_TRACE_EVENT_BEGIN(Cpu, FHLSLMaterialTranslatorTranslate, NoSync)
+	UE_TRACE_EVENT_FIELD(UE::Trace::WideString, MaterialName)
+UE_TRACE_EVENT_END()
  
 bool FHLSLMaterialTranslator::Translate()
 {
-	TRACE_CPUPROFILER_EVENT_SCOPE(FHLSLMaterialTranslator::Translate);
+#if CPUPROFILERTRACE_ENABLED
+	FString TraceMaterialName;
+	if (UE_TRACE_CHANNELEXPR_IS_ENABLED(CpuChannel))
+	{
+		TraceMaterialName = Material->GetMaterialInterface()->GetFullName();
+	}
+	UE_TRACE_LOG_SCOPED_T(Cpu, FHLSLMaterialTranslatorTranslate, CpuChannel)
+		<< FHLSLMaterialTranslatorTranslate.MaterialName(*TraceMaterialName);
+#endif
+
 	COOK_STAT(MaterialTranslatorCookStats::MaterialTranslateCalls++);
 	COOK_STAT(FScopedDurationTimer DurationTimer(MaterialTranslatorCookStats::MaterialTranslateTimeSec));
 
