Fix for SparseVolumeTextures not rendering correctly with dx11. Fixed by using a dummy UAV for unused UAV parameters in the upload shader.

#rnx

diff --git a/Engine/Source/Runtime/Engine/Private/SparseVolumeTexture/SparseVolumeTextureStreamingManager.cpp b/Engine/Source/Runtime/Engine/Private/SparseVolumeTexture/SparseVolumeTextureStreamingManager.cpp
--- a/Engine/Source/Runtime/Engine/Private/SparseVolumeTexture/SparseVolumeTextureStreamingManager.cpp
+++ b/Engine/Source/Runtime/Engine/Private/SparseVolumeTexture/SparseVolumeTextureStreamingManager.cpp
@@ -435,6 +435,7 @@ public:
 
 			if (NumWrittenTiles > 0)
 			{
+				FRDGTexture* DummyTexture = GraphBuilder.CreateTexture(FRDGTextureDesc::Create3D(FIntVector(1), PF_R8, FClearValueBinding::None, ETextureCreateFlags::UAV), TEXT("SparseVolumeTexture.DummyTexture"));
 				FRDGTexture* DstTextureARDG = DstTextureA ? GraphBuilder.RegisterExternalTexture(CreateRenderTarget(DstTextureA, TEXT("SparseVolumeTexture.TileDataTextureA"))) : nullptr;
 				FRDGTexture* DstTextureBRDG = DstTextureB ? GraphBuilder.RegisterExternalTexture(CreateRenderTarget(DstTextureB, TEXT("SparseVolumeTexture.TileDataTextureB"))) : nullptr;
 
@@ -443,14 +444,15 @@ public:
 				FRDGBufferSRV* DstTileCoordsBufferSRV = GraphBuilder.CreateSRV(GraphBuilder.RegisterExternalBuffer(DstTileCoordsUploadBuffer));
 				FRDGBufferSRV* TileDataABufferSRV = FormatSizeA ? GraphBuilder.CreateSRV(GraphBuilder.RegisterExternalBuffer(TileDataAUploadBuffer), FormatA) : nullptr;
 				FRDGBufferSRV* TileDataBBufferSRV = FormatSizeB ? GraphBuilder.CreateSRV(GraphBuilder.RegisterExternalBuffer(TileDataBUploadBuffer), FormatB) : nullptr;
+				FRDGTextureUAV* DummyTextureUAV = GraphBuilder.CreateUAV(DummyTexture);
 				FRDGTextureUAV* DstTextureAUAV = DstTextureARDG ? GraphBuilder.CreateUAV(DstTextureARDG) : nullptr;
 				FRDGTextureUAV* DstTextureBUAV = DstTextureBRDG ? GraphBuilder.CreateUAV(DstTextureBRDG) : nullptr;
 
 				auto ComputeShader = GetGlobalShaderMap(GMaxRHIFeatureLevel)->GetShader<FSparseVolumeTextureUpdateFromSparseBufferCS>();
 
 				FSparseVolumeTextureUpdateFromSparseBufferCS::FParameters* PassParameters = GraphBuilder.AllocParameters<FSparseVolumeTextureUpdateFromSparseBufferCS::FParameters>();
-				PassParameters->DstPhysicalTileTextureA = FormatSizeA > 0 ? DstTextureAUAV : DstTextureBUAV;
-				PassParameters->DstPhysicalTileTextureB = FormatSizeB > 0 ? DstTextureBUAV : DstTextureAUAV;
+				PassParameters->DstPhysicalTileTextureA = FormatSizeA > 0 ? DstTextureAUAV : DummyTextureUAV;
+				PassParameters->DstPhysicalTileTextureB = FormatSizeB > 0 ? DstTextureBUAV : DummyTextureUAV;
 				PassParameters->SrcPhysicalTileBufferA = FormatSizeA > 0 ? TileDataABufferSRV : TileDataBBufferSRV;
 				PassParameters->SrcPhysicalTileBufferB = FormatSizeB > 0 ? TileDataBBufferSRV : TileDataABufferSRV;
 				PassParameters->OccupancyBitsBuffer = OccupancyBitsBufferSRV;
