Exporting a couple of helpers

diff --git a/Engine/Source/Runtime/Engine/Public/PluginBlueprintLibrary.h b/Engine/Source/Runtime/Engine/Public/PluginBlueprintLibrary.h
--- a/Engine/Source/Runtime/Engine/Public/PluginBlueprintLibrary.h
+++ b/Engine/Source/Runtime/Engine/Public/PluginBlueprintLibrary.h
@@ -57,7 +57,7 @@ public:
 	 *         name was stored in OutPluginName, or false otherwise
 	 */
 	UFUNCTION(BlueprintCallable, Category="Engine Scripting | Plugin Utilities")
-	static bool GetPluginNameForObjectPath(
+	static ENGINE_API bool GetPluginNameForObjectPath(
 			const FSoftObjectPath& ObjectPath,
 			FString& OutPluginName);
 
@@ -99,7 +99,7 @@ public:
 	 *         directory filesystem path was stored in OutContentDir, or false otherwise
 	 */
 	UFUNCTION(BlueprintCallable, Category="Engine Scripting | Plugin Utilities")
-	static bool GetPluginContentDir(
+	static ENGINE_API bool GetPluginContentDir(
 			const FString& PluginName,
 			FString& OutContentDir);
 
