Anim blueprint logging improvements:
-Log to PIE if the error is found in a PIE session.
-Collect tokenized messages during the graph update to allow for hyperlinking

diff --git a/Engine/Source/Runtime/Engine/Classes/Animation/AnimNodeBase.h b/Engine/Source/Runtime/Engine/Classes/Animation/AnimNodeBase.h
--- a/Engine/Source/Runtime/Engine/Classes/Animation/AnimNodeBase.h
+++ b/Engine/Source/Runtime/Engine/Classes/Animation/AnimNodeBase.h
@@ -288,7 +288,7 @@ protected:
 protected:
 
 	/** Interface for node contexts to register log messages with the proxy */
-	ENGINE_API void LogMessageInternal(FName InLogType, EMessageSeverity::Type InSeverity, FText InMessage) const;
+	ENGINE_API void LogMessageInternal(FName InLogType, const TSharedRef<FTokenizedMessage>& InMessage) const;
 };
 
 
@@ -434,7 +434,8 @@ public:
 	float GetDeltaTime() const { return DeltaTime; }
 
 	// Log update message
-	void LogMessage(EMessageSeverity::Type InSeverity, FText InMessage) const { LogMessageInternal("Update", InSeverity, InMessage); }
+	void LogMessage(const TSharedRef<FTokenizedMessage>& InMessage) const { LogMessageInternal("Update", InMessage); }
+	void LogMessage(EMessageSeverity::Type InSeverity, FText InMessage) const { LogMessage(FTokenizedMessage::Create(InSeverity, InMessage)); }
 };
 
 
@@ -473,7 +474,8 @@ public:
 	void Initialize(FAnimInstanceProxy* InAnimInstanceProxy) { InitializeImpl(InAnimInstanceProxy); }
 
 	// Log evaluation message
-	void LogMessage(EMessageSeverity::Type InSeverity, FText InMessage) const { LogMessageInternal("Evaluate", InSeverity, InMessage); }
+	void LogMessage(const TSharedRef<FTokenizedMessage>& InMessage) const { LogMessageInternal("Evaluate", InMessage); }
+	void LogMessage(EMessageSeverity::Type InSeverity, FText InMessage) const { LogMessage(FTokenizedMessage::Create(InSeverity, InMessage)); }
 
 	void ResetToRefPose()
 	{
diff --git a/Engine/Source/Runtime/Engine/Private/Animation/AnimInstanceProxy.cpp b/Engine/Source/Runtime/Engine/Private/Animation/AnimInstanceProxy.cpp
--- a/Engine/Source/Runtime/Engine/Private/Animation/AnimInstanceProxy.cpp
+++ b/Engine/Source/Runtime/Engine/Private/Animation/AnimInstanceProxy.cpp
@@ -47,6 +47,7 @@
 #define LOCTEXT_NAMESPACE "AnimInstance"
 
 const FName NAME_AnimBlueprintLog(TEXT("AnimBlueprintLog"));
+const FName NAME_PIELog(TEXT("PIE"));
 const FName NAME_Evaluate(TEXT("Evaluate"));
 const FName NAME_Update(TEXT("Update"));
 const FName NAME_AnimGraph(TEXT("AnimGraph"));
@@ -73,6 +74,7 @@ FAnimInstanceProxy::FAnimInstanceProxy()
 	, bDeferRootNodeInitialization(false)
 #if WITH_EDITORONLY_DATA
 	, bIsBeingDebugged(false)
+	, bIsGameWorld(false)
 #endif
 	, bInitializeSubsystems(false)
 	, bUseMainInstanceMontageEvaluationData(false)
@@ -410,21 +412,35 @@ FGuid MakeGuidForMessage(const FText& Message)
 	return FGuid(Hash[0] ^ Hash[4], Hash[1], Hash[2], Hash[3]);
 }
 
-void FAnimInstanceProxy::LogMessage(FName InLogType, EMessageSeverity::Type InSeverity, const FText& InMessage) const
+FName FAnimInstanceProxy::GetTargetLogNameForCurrentWorldType() const
+{
+#if WITH_EDITORONLY_DATA
+	return bIsGameWorld ? NAME_PIELog : NAME_AnimBlueprintLog;
+#else
+	return NAME_AnimBlueprintLog;
+#endif
+}
+
+void FAnimInstanceProxy::LogMessage(FName InLogType, const TSharedRef<FTokenizedMessage>& InMessage) const
 {
 #if ENABLE_ANIM_LOGGING
-	FGuid CurrentMessageGuid = MakeGuidForMessage(InMessage);
+	FGuid CurrentMessageGuid = MakeGuidForMessage(InMessage->ToText());
 	if(!PreviouslyLoggedMessages.Contains(CurrentMessageGuid))
 	{
 		PreviouslyLoggedMessages.Add(CurrentMessageGuid);
 		if (TArray<FLogMessageEntry>* LoggedMessages = LoggedMessagesMap.Find(InLogType))
 		{
-			LoggedMessages->Emplace(InSeverity, InMessage);
+			LoggedMessages->Emplace(InMessage);
 		}
 	}
 #endif
 }
 
+void FAnimInstanceProxy::LogMessage(FName InLogType, EMessageSeverity::Type InSeverity, const FText& InMessage) const
+{
+	LogMessage(InLogType, FTokenizedMessage::Create(InSeverity, InMessage));
+}
+
 void FAnimInstanceProxy::Uninitialize(UAnimInstance* InAnimInstance)
 {
 	DECLARE_SCOPE_HIERARCHICAL_COUNTER_FUNC()
@@ -459,6 +475,8 @@ void FAnimInstanceProxy::PreUpdate(UAnimInstance* InAnimInstance, float DeltaSec
 	bShouldExtractRootMotion = InAnimInstance->ShouldExtractRootMotion();
 
 #if WITH_EDITORONLY_DATA
+	bIsGameWorld = World ? World->IsGameWorld() : false;
+
 	if (FAnimBlueprintDebugData* DebugData = GetAnimBlueprintDebugData())
 	{
 		DebugData->ResetNodeVisitSites();
@@ -663,13 +681,13 @@ void FAnimInstanceProxy::PostUpdate(UAnimInstance* InAnimInstance) const
 #endif
 
 #if ENABLE_ANIM_LOGGING
-	FMessageLog MessageLog(NAME_AnimBlueprintLog);
+	FMessageLog MessageLog(GetTargetLogNameForCurrentWorldType());
 	const TArray<FLogMessageEntry>* Messages = LoggedMessagesMap.Find(NAME_Update);
 	if (ensureMsgf(Messages, TEXT("PreUpdate isn't called. This could potentially cause other issues.")))
 	{
 		for (const FLogMessageEntry& Message : *Messages)
 		{
-			MessageLog.Message(Message.Key, Message.Value);
+			MessageLog.AddMessage(Message);
 		}
 	}
 #endif
@@ -682,12 +700,12 @@ void FAnimInstanceProxy::PostEvaluate(UAnimInstance* InAnimInstance)
 	ClearObjects();
 
 #if ENABLE_ANIM_LOGGING
-	FMessageLog MessageLog(NAME_AnimBlueprintLog);
+	FMessageLog MessageLog(GetTargetLogNameForCurrentWorldType());
 	if(const TArray<FLogMessageEntry>* Messages = LoggedMessagesMap.Find(NAME_Evaluate))
 	{
 		for (const FLogMessageEntry& Message : *Messages)
 		{
-			MessageLog.Message(Message.Key, Message.Value);
+			MessageLog.AddMessage(Message);
 		}
 	}
 #endif
diff --git a/Engine/Source/Runtime/Engine/Private/Animation/AnimNodeBase.cpp b/Engine/Source/Runtime/Engine/Private/Animation/AnimNodeBase.cpp
--- a/Engine/Source/Runtime/Engine/Private/Animation/AnimNodeBase.cpp
+++ b/Engine/Source/Runtime/Engine/Private/Animation/AnimNodeBase.cpp
@@ -93,9 +93,9 @@ UAnimBlueprint* FAnimationBaseContext::GetAnimBlueprint() const
 }
 #endif //WITH_EDITORONLY_DATA
 
-void FAnimationBaseContext::LogMessageInternal(FName InLogType, EMessageSeverity::Type InSeverity, FText InMessage) const
+void FAnimationBaseContext::LogMessageInternal(FName InLogType, const TSharedRef<FTokenizedMessage>& InMessage) const
 {
-	AnimInstanceProxy->LogMessage(InLogType, InSeverity, InMessage);
+	AnimInstanceProxy->LogMessage(InLogType, InMessage);
 }
 /////////////////////////////////////////////////////
 // FPoseContext
diff --git a/Engine/Source/Runtime/Engine/Public/Animation/AnimInstanceProxy.h b/Engine/Source/Runtime/Engine/Public/Animation/AnimInstanceProxy.h
--- a/Engine/Source/Runtime/Engine/Public/Animation/AnimInstanceProxy.h
+++ b/Engine/Source/Runtime/Engine/Public/Animation/AnimInstanceProxy.h
@@ -930,6 +930,7 @@ protected:
 
 	/** Allow nodes to register log messages to be processed on the game thread */
 	ENGINE_API void LogMessage(FName InLogType, EMessageSeverity::Type InSeverity, const FText& InMessage) const;
+	ENGINE_API void LogMessage(FName InLogType, const TSharedRef<FTokenizedMessage>& InMessage) const;
 
 	/** Get the current value of all animation curves **/
 	TMap<FName, float>& GetAnimationCurves(EAnimCurveType InCurveType) { return AnimationCurves[(uint8)InCurveType]; }
@@ -970,6 +971,8 @@ protected:
 
 private:
 
+	FName GetTargetLogNameForCurrentWorldType() const;
+
 	/** Executes the provided functor on each valid state machine on this anim instance */
 	ENGINE_API void ForEachStateMachine(const TFunctionRef<void(FAnimNode_StateMachine&)>& Functor);
 
@@ -1136,7 +1139,7 @@ private:
 	/** Logged message queues. Allows nodes to report messages to MessageLog even though they may be running
 	 *  on a worked thread
 	 */
-	typedef TPair<EMessageSeverity::Type, FText> FLogMessageEntry;
+	typedef TSharedRef<FTokenizedMessage> FLogMessageEntry;
 	mutable TMap<FName, TArray<FLogMessageEntry>> LoggedMessagesMap;
 
 	/** Cache of guids generated from previously sent messages so we can stop spam*/
@@ -1162,6 +1165,9 @@ private:
 #if WITH_EDITORONLY_DATA
 	/** Whether this UAnimInstance is currently being debugged in the editor */
 	uint8 bIsBeingDebugged : 1;
+
+	/** Whether the world of this proxy's skeletal mesh component is a game world. */
+	uint8 bIsGameWorld : 1;
 #endif
 
 	// Whether subsystems should be initialized
