Allow for multiple weight map samples per weight map.  Only occurs when material swaps are used to change the meaning of a layer. This will not occur with landscape materials generated by the landscape system.

diff --git a/Engine/Source/Runtime/Engine/Private/Materials/HLSLMaterialTranslator.cpp b/Engine/Source/Runtime/Engine/Private/Materials/HLSLMaterialTranslator.cpp
--- a/Engine/Source/Runtime/Engine/Private/Materials/HLSLMaterialTranslator.cpp
+++ b/Engine/Source/Runtime/Engine/Private/Materials/HLSLMaterialTranslator.cpp
@@ -7973,37 +7973,61 @@ bool FHLSLMaterialTranslator::GetStaticBoolValue(int32 BoolIndex, bool& bSucceed
 int32 FHLSLMaterialTranslator::StaticTerrainLayerWeight(FName LayerName,int32 Default)
 {
 	// Look up the weight-map index for this static parameter.
-	int32 WeightmapIndex = INDEX_NONE;
-	bool bFoundParameter = false;
-
+	int32 WeightmapCode = INDEX_NONE;
+	int32 NumWeightmapParameters = 0;
 	for(int32 ParameterIndex = 0;ParameterIndex < StaticParameters.EditorOnly.TerrainLayerWeightParameters.Num(); ++ParameterIndex)
 	{
 		const FStaticTerrainLayerWeightParameter& Parameter = StaticParameters.EditorOnly.TerrainLayerWeightParameters[ParameterIndex];
-		if(Parameter.LayerName == LayerName)
+
+		// If there are multiple weight maps with the same name, they should be numbered to allow for unique masks
+		FName LayerNameTest = Parameter.LayerName;
+		LayerNameTest.SetNumber(0);
+
+		if(LayerNameTest != LayerName)
 		{
-			WeightmapIndex = Parameter.WeightmapIndex;
-			bFoundParameter = true;
+			continue;
 		}
+
+		int32 WeightmapIndex =  Parameter.WeightmapIndex;
+
+		if(WeightmapIndex == INDEX_NONE)
+		{
+			continue;
+		}
+
+		constexpr EMaterialSamplerType SamplerType = SAMPLERTYPE_Masks;
+		FString WeightmapName = FString::Printf(TEXT("Weightmap%d"), WeightmapIndex);
+		int32 TextureReferenceIndex = INDEX_NONE;
+		int32 TextureCodeIndex = TextureParameter(FName(*WeightmapName), GEngine->WeightMapPlaceholderTexture, TextureReferenceIndex, SamplerType);
+		int32 SampleCodeIndex = TextureSample(TextureCodeIndex, TextureCoordinate(3, false, false), SamplerType, /*MipValue0Index = */INDEX_NONE, /*MipValue1Index = */INDEX_NONE, /*MipValueMode = */TMVM_None, /*SamplerSource = */SSM_TerrainWeightmapGroupSettings);
+
+		FString LayerMaskName = FString::Printf(TEXT("LayerMask_%s"), *Parameter.LayerName.ToString());
+		int32 CurrentWeightmapCode = Dot(SampleCodeIndex, VectorParameter(FName(*LayerMaskName), FLinearColor(1.f, 0.f, 0.f, 0.f)));
+
+		if(WeightmapCode == INDEX_NONE)
+		{
+			WeightmapCode = CurrentWeightmapCode;
+		}
+		else
+		{
+			WeightmapCode = Add(WeightmapCode, CurrentWeightmapCode);
+		}
+
+		++NumWeightmapParameters;
 	}
 
-	if(!bFoundParameter && Material->IsPreview())
+	if((NumWeightmapParameters == 0) && Material->IsPreview())
 	{
 		return Default;
 	}
-	else if(WeightmapIndex == INDEX_NONE)
+	else 
 	{
-		return INDEX_NONE;
-	}
-	else
-	{			
-		constexpr EMaterialSamplerType SamplerType = SAMPLERTYPE_Masks;
-		FString WeightmapName = FString::Printf(TEXT("Weightmap%d"),WeightmapIndex);
-		int32 TextureReferenceIndex = INDEX_NONE;
-		int32 TextureCodeIndex = TextureParameter(FName(*WeightmapName), GEngine->WeightMapPlaceholderTexture, TextureReferenceIndex, SamplerType);
-		int32 WeightmapCode = TextureSample(TextureCodeIndex, TextureCoordinate(3, false, false), SamplerType, /*MipValue0Index = */INDEX_NONE, /*MipValue1Index = */INDEX_NONE, /*MipValueMode = */TMVM_None, /*SamplerSource = */SSM_TerrainWeightmapGroupSettings);
+		if (NumWeightmapParameters > 1)
+		{
+			WeightmapCode = Clamp(WeightmapCode, Constant(0.0), Constant(1.0));
+		}
 
-		FString LayerMaskName = FString::Printf(TEXT("LayerMask_%s"),*LayerName.ToString());
-		return Dot(WeightmapCode,VectorParameter(FName(*LayerMaskName), FLinearColor(1.f,0.f,0.f,0.f)));
+		return WeightmapCode;
 	}
 }
 
