Exposed GetClassPathName, GetStructPathName, and GetEnumPathName to Editor Scripting

These functions return a TopLevelAssetPath containing the type info.

The API in Python has the same form as C++, eg):
 unreal.MyClass.static_class().get_class_path_name()
 unreal.MyStruct.static_struct().get_struct_path_name()
 unreal.MyEnum.static_enum().get_enum_path_name()

diff --git a/Engine/Source/Runtime/Engine/Classes/Kismet/KismetSystemLibrary.h b/Engine/Source/Runtime/Engine/Classes/Kismet/KismetSystemLibrary.h
--- a/Engine/Source/Runtime/Engine/Classes/Kismet/KismetSystemLibrary.h
+++ b/Engine/Source/Runtime/Engine/Classes/Kismet/KismetSystemLibrary.h
@@ -134,9 +134,17 @@ class UKismetSystemLibrary : public UBlueprintFunctionLibrary
 	static ENGINE_API FSoftClassPath GetSoftClassPath(const UClass* Class);
 
 	// Returns the full path to the specified class as a Top Level Asset Path used by asset utilities
-	UFUNCTION(BlueprintPure, Category = "Utilities")
+	UFUNCTION(BlueprintPure, Category = "Utilities", DisplayName="Get Class Path Name (Top Level Asset Path)", meta=(ScriptMethod="GetClassPathName"))
 	static ENGINE_API FTopLevelAssetPath GetClassTopLevelAssetPath(const UClass* Class);
 
+	// Returns the full path to the specified struct as a Top Level Asset Path used by asset utilities
+	UFUNCTION(BlueprintPure, Category = "Utilities", DisplayName = "Get Struct Path Name (Top Level Asset Path)", meta=(ScriptMethod="GetStructPathName"))
+	static FTopLevelAssetPath GetStructTopLevelAssetPath(const UScriptStruct* Struct);
+
+	// Returns the full path to the specified enum as a Top Level Asset Path used by asset utilities
+	UFUNCTION(BlueprintPure, Category = "Utilities", DisplayName = "Get Enum Path Name (Top Level Asset Path)", meta=(ScriptMethod="GetEnumPathName"))
+	static FTopLevelAssetPath GetEnumTopLevelAssetPath(const UEnum* Enum);
+
 	// Returns the outer object of an object.
 	UFUNCTION(BlueprintPure, Category = "Utilities")
 	static ENGINE_API UObject* GetOuterObject(const UObject* Object);
diff --git a/Engine/Source/Runtime/Engine/Private/KismetSystemLibrary.cpp b/Engine/Source/Runtime/Engine/Private/KismetSystemLibrary.cpp
--- a/Engine/Source/Runtime/Engine/Private/KismetSystemLibrary.cpp
+++ b/Engine/Source/Runtime/Engine/Private/KismetSystemLibrary.cpp
@@ -169,6 +169,18 @@ FTopLevelAssetPath UKismetSystemLibrary::GetClassTopLevelAssetPath(const UClass*
 	return FTopLevelAssetPath(Class);
 }
 
+FTopLevelAssetPath UKismetSystemLibrary::GetStructTopLevelAssetPath(const UScriptStruct* Struct)
+{
+	// This will succeed for all valid structs as they are never subobjects
+	return FTopLevelAssetPath(Struct);
+}
+
+FTopLevelAssetPath UKismetSystemLibrary::GetEnumTopLevelAssetPath(const UEnum* Enum)
+{
+	// This will succeed for all valid enums as they are never subobjects
+	return FTopLevelAssetPath(Enum);
+}
+
 UObject* UKismetSystemLibrary::GetOuterObject(const UObject* Object)
 {
 	return Object ? Object->GetOuter() : nullptr;
