Added a method (similar to getting Primitive Components) to get the Actors in the cluster union component

diff --git a/Engine/Source/Runtime/Engine/Classes/PhysicsEngine/ClusterUnionComponent.h b/Engine/Source/Runtime/Engine/Classes/PhysicsEngine/ClusterUnionComponent.h
--- a/Engine/Source/Runtime/Engine/Classes/PhysicsEngine/ClusterUnionComponent.h
+++ b/Engine/Source/Runtime/Engine/Classes/PhysicsEngine/ClusterUnionComponent.h
@@ -155,6 +155,9 @@ public:
 	UFUNCTION(BlueprintCallable, Category = "Cluster Union")
 	ENGINE_API TArray<UPrimitiveComponent*> GetPrimitiveComponents();
 
+	UFUNCTION(BlueprintCallable, Category = "Cluster Union")
+	ENGINE_API TArray<AActor*> GetActors();
+
 	UFUNCTION(BlueprintCallable, Category = "Cluster Union")
 	ENGINE_API void SetIsAnchored(bool bIsAnchored);
 
@@ -181,6 +184,7 @@ public:
 
 	// Lambda returns whether or not iteration should continue;
 	ENGINE_API void VisitAllCurrentChildComponents(const TFunction<bool(UPrimitiveComponent*)>& Lambda) const;
+	ENGINE_API void VisitAllCurrentActors(const TFunction<bool(AActor*)>& Lambda) const;
 
 	ENGINE_API int32 NumChildClusterComponents() const;
 
@@ -266,6 +270,7 @@ private:
 	ENGINE_API void HandleRemovedClusteredComponent(UPrimitiveComponent* ChangedComponent, bool bDestroyReplicatedProxy);
 
 	ENGINE_API TArray<UPrimitiveComponent*> GetAllCurrentChildComponents() const;
+	ENGINE_API TArray<AActor*> GetAllCurrentActors() const;
 	ENGINE_API void VisitAllCurrentChildComponentsForCollision(ECollisionChannel TraceChannel, const struct FCollisionQueryParams& Params, const struct FCollisionResponseParams& ResponseParams, const struct FCollisionObjectQueryParams& ObjectParams, const TFunction<bool(UPrimitiveComponent*)>& Lambda) const;
 
 	// Whether or not this code is running on the server.
diff --git a/Engine/Source/Runtime/Engine/Private/PhysicsEngine/ClusterUnionComponent.cpp b/Engine/Source/Runtime/Engine/Private/PhysicsEngine/ClusterUnionComponent.cpp
--- a/Engine/Source/Runtime/Engine/Private/PhysicsEngine/ClusterUnionComponent.cpp
+++ b/Engine/Source/Runtime/Engine/Private/PhysicsEngine/ClusterUnionComponent.cpp
@@ -299,6 +299,8 @@ void UClusterUnionComponent::ForceRebuildGTParticleGeometry()
 TArray<UPrimitiveComponent*> UClusterUnionComponent::GetPrimitiveComponents()
 {
 	TArray<UPrimitiveComponent*> PrimitiveComponents;
+	PrimitiveComponents.Reserve(ComponentToPhysicsObjects.Num());
+
 	for (auto Iter = ComponentToPhysicsObjects.CreateIterator(); Iter; ++Iter)
 	{
 		PrimitiveComponents.Add(Iter.Key().ResolveObjectPtr());
@@ -307,6 +309,19 @@ TArray<UPrimitiveComponent*> UClusterUnionComponent::GetPrimitiveComponents()
 	return PrimitiveComponents;
 }
 
+TArray<AActor*> UClusterUnionComponent::GetActors()
+{
+	TArray<AActor*> Actors;
+	Actors.Reserve(ActorToComponents.Num());
+
+	for (auto Iter = ActorToComponents.CreateIterator(); Iter; ++Iter)
+	{
+		Actors.Add(Iter.Key().ResolveObjectPtr());
+	}
+
+	return Actors;
+}
+
 void UClusterUnionComponent::SetIsAnchored(bool bIsAnchored)
 {
 	if (!PhysicsProxy)
@@ -1085,6 +1100,22 @@ TArray<UPrimitiveComponent*> UClusterUnionComponent::GetAllCurrentChildComponent
 	return Components;
 }
 
+TArray<AActor*> UClusterUnionComponent::GetAllCurrentActors() const
+{
+	TArray<AActor*> Actors;
+	Actors.Reserve(ActorToComponents.Num());
+
+	VisitAllCurrentActors(
+		[&Actors](AActor* Actor)
+		{
+			Actors.Add(Actor);
+			return true;
+		}
+	);
+
+	return Actors;
+}
+
 void UClusterUnionComponent::VisitAllCurrentChildComponents(const TFunction<bool(UPrimitiveComponent*)>& Lambda) const
 {
 	for (const TPair<TObjectKey<UPrimitiveComponent>, FClusterUnionPendingAddData>& Kvp : PendingComponentSync)
@@ -1110,6 +1141,20 @@ void UClusterUnionComponent::VisitAllCurrentChildComponents(const TFunction<bool
 	}
 }
 
+void UClusterUnionComponent::VisitAllCurrentActors(const TFunction<bool(AActor*)>& Lambda) const
+{
+	for (const TPair<TObjectKey<AActor>, FClusteredActorData>& Kvp : ActorToComponents)
+	{
+		if (AActor* Actor = Kvp.Key.ResolveObjectPtr())
+		{
+			if (!Lambda(Actor))
+			{
+				return;
+			}
+		}
+	}
+}
+
 void UClusterUnionComponent::VisitAllCurrentChildComponentsForCollision(ECollisionChannel TraceChannel, const struct FCollisionQueryParams& Params, const struct FCollisionResponseParams& ResponseParams, const struct FCollisionObjectQueryParams& ObjectParams, const TFunction<bool(UPrimitiveComponent*)>& Lambda) const
 {
 	TSet<uint32> IgnoredActors;
