Fix for CIS error in AnimInstanceProxy

diff --git a/Engine/Source/Runtime/Engine/Public/Animation/AnimInstanceProxy.h b/Engine/Source/Runtime/Engine/Public/Animation/AnimInstanceProxy.h
--- a/Engine/Source/Runtime/Engine/Public/Animation/AnimInstanceProxy.h
+++ b/Engine/Source/Runtime/Engine/Public/Animation/AnimInstanceProxy.h
@@ -29,6 +29,7 @@
 class UAnimInstance;
 class UBlendProfile;
 class UBlendSpace;
+class FTokenizedMessage;
 namespace EMessageSeverity { enum Type : int; }
 struct FAnimationPoseData;
 struct FAnimBlueprintDebugData_NodeVisit;
