Cameras: add aspect ratio axis constraint override on camera components

diff --git a/Engine/Source/Runtime/Engine/Classes/Camera/CameraComponent.h b/Engine/Source/Runtime/Engine/Classes/Camera/CameraComponent.h
--- a/Engine/Source/Runtime/Engine/Classes/Camera/CameraComponent.h
+++ b/Engine/Source/Runtime/Engine/Classes/Camera/CameraComponent.h
@@ -52,18 +52,28 @@ class UCameraComponent : public USceneComponent
 	UFUNCTION(BlueprintCallable, Category = Camera)
 	void SetOrthoFarClipPlane(float InOrthoFarClipPlane) { OrthoFarClipPlane = InOrthoFarClipPlane; }
 
-	// Aspect Ratio (Width/Height)
+	/** Aspect Ratio (Width/Height) */
 	UPROPERTY(Interp, EditAnywhere, BlueprintReadWrite, Category = CameraSettings, meta = (ClampMin = "0.1", ClampMax = "100.0", EditCondition = "bConstrainAspectRatio"))
 	float AspectRatio;
 	UFUNCTION(BlueprintCallable, Category = Camera)
 	void SetAspectRatio(float InAspectRatio) { AspectRatio = InAspectRatio; }
 
-	// If bConstrainAspectRatio is true, black bars will be added if the destination view has a different aspect ratio than this camera requested.
+	/** Override for the default aspect ratio axis constraint defined on the local player */
+	UPROPERTY(Interp, EditAnywhere, BlueprintReadWrite, Category = CameraOptions, meta = (EditCondition = "bOverrideAspectRatioAxisConstraint"))
+	TEnumAsByte<EAspectRatioAxisConstraint> AspectRatioAxisConstraint = EAspectRatioAxisConstraint::AspectRatio_MaintainXFOV;
+	UFUNCTION(BlueprintCallable, Category = Camera)
+	void SetAspectRatioAxisConstraint(EAspectRatioAxisConstraint InAspectRatioAxisConstraint) { AspectRatioAxisConstraint = InAspectRatioAxisConstraint; }
+
+	/** If bConstrainAspectRatio is true, black bars will be added if the destination view has a different aspect ratio than this camera requested. */
 	UPROPERTY(Interp, EditAnywhere, BlueprintReadWrite, Category = CameraOptions)
 	uint8 bConstrainAspectRatio : 1;
 	UFUNCTION(BlueprintCallable, Category = Camera)
 	void SetConstraintAspectRatio(bool bInConstrainAspectRatio) { bConstrainAspectRatio = bInConstrainAspectRatio; }
 
+	/** Whether to override the default aspect ratio axis constraint defined on the local player */
+	UPROPERTY(Interp, EditAnywhere, BlueprintReadWrite, Category = CameraOptions)
+	uint8 bOverrideAspectRatioAxisConstraint : 1;
+
 	// If true, account for the field of view angle when computing which level of detail to use for meshes.
 	UPROPERTY(Interp, EditAnywhere, AdvancedDisplay, BlueprintReadWrite, Category = CameraOptions)
 	uint8 bUseFieldOfViewForLOD : 1;
diff --git a/Engine/Source/Runtime/Engine/Classes/Camera/CameraTypes.h b/Engine/Source/Runtime/Engine/Classes/Camera/CameraTypes.h
--- a/Engine/Source/Runtime/Engine/Classes/Camera/CameraTypes.h
+++ b/Engine/Source/Runtime/Engine/Classes/Camera/CameraTypes.h
@@ -86,6 +86,9 @@ struct FMinimalViewInfo
 	UPROPERTY(EditAnywhere, BlueprintReadWrite, Category=Camera)
 	float AspectRatio;
 
+	// Aspect ratio axis constraint override
+	TOptional<EAspectRatioAxisConstraint> AspectRatioAxisConstraint;
+
 	// If bConstrainAspectRatio is true, black bars will be added if the destination view has a different aspect ratio than this camera requested.
 	UPROPERTY(EditAnywhere, BlueprintReadWrite, Category=Camera)
 	uint32 bConstrainAspectRatio:1; 
diff --git a/Engine/Source/Runtime/Engine/Private/Camera/CameraComponent.cpp b/Engine/Source/Runtime/Engine/Private/Camera/CameraComponent.cpp
--- a/Engine/Source/Runtime/Engine/Private/Camera/CameraComponent.cpp
+++ b/Engine/Source/Runtime/Engine/Private/Camera/CameraComponent.cpp
@@ -396,6 +396,11 @@ void UCameraComponent::GetCameraView(float DeltaTime, FMinimalViewInfo& DesiredV
 	DesiredView.OrthoNearClipPlane = OrthoNearClipPlane;
 	DesiredView.OrthoFarClipPlane = OrthoFarClipPlane;
 
+	if (bOverrideAspectRatioAxisConstraint)
+	{
+		DesiredView.AspectRatioAxisConstraint = AspectRatioAxisConstraint;
+	}
+
 	// See if the CameraActor wants to override the PostProcess settings used.
 	DesiredView.PostProcessBlendWeight = PostProcessBlendWeight;
 	if (PostProcessBlendWeight > 0.0f)
diff --git a/Engine/Source/Runtime/Engine/Private/Camera/CameraStackTypes.cpp b/Engine/Source/Runtime/Engine/Private/Camera/CameraStackTypes.cpp
--- a/Engine/Source/Runtime/Engine/Private/Camera/CameraStackTypes.cpp
+++ b/Engine/Source/Runtime/Engine/Private/Camera/CameraStackTypes.cpp
@@ -158,6 +158,9 @@ void FMinimalViewInfo::CalculateProjectionMatrixGivenViewRectangle(const FMinima
 		const int32 SizeX = ViewRect.Width();
 		const int32 SizeY = ViewRect.Height();
 
+		// Get effective aspect ratio axis constraint.
+		AspectRatioAxisConstraint = ViewInfo.AspectRatioAxisConstraint.Get(AspectRatioAxisConstraint);
+
 		// If x is bigger, and we're respecting x or major axis, AND mobile isn't forcing us to be Y axis aligned
 		const bool bMaintainXFOV = 
 			((SizeX > SizeY) && (AspectRatioAxisConstraint == AspectRatio_MajorAxisFOV)) || 
