Fixed Atom regression when importing flex elements. It's good practice to test validity of MeshDescription attributes before using them, in any case.

diff --git a/Engine/Source/Runtime/Engine/Private/Rendering/SkeletalMeshLODImporterData.cpp b/Engine/Source/Runtime/Engine/Private/Rendering/SkeletalMeshLODImporterData.cpp
--- a/Engine/Source/Runtime/Engine/Private/Rendering/SkeletalMeshLODImporterData.cpp
+++ b/Engine/Source/Runtime/Engine/Private/Rendering/SkeletalMeshLODImporterData.cpp
@@ -2243,30 +2243,32 @@ FSkeletalMeshImportData FSkeletalMeshImportData::CreateFromMeshDescription(const
 	FSkeletalMeshAttributes::FBoneNameAttributesConstRef BoneNames = Attributes.GetBoneNames();
 	FSkeletalMeshAttributes::FBoneParentIndexAttributesConstRef BoneParentIndices = Attributes.GetBoneParentIndices();
 	FSkeletalMeshAttributes::FBonePoseAttributesConstRef BonePoses = Attributes.GetBonePoses();
-
-	auto GetNumChildren = [&](const int32 ParentIndex)
+	if (BoneNames.IsValid() && BoneParentIndices.IsValid() && BonePoses.IsValid())
 	{
-		int32 NumChildren = 0;
-		for (const FBoneID& BoneId: Bones.GetElementIDs())
+		auto GetNumChildren = [&](const int32 ParentIndex)
 		{
-			if (BoneParentIndices.Get(BoneId) == ParentIndex)
+			int32 NumChildren = 0;
+			for (const FBoneID& BoneId : Bones.GetElementIDs())
 			{
-				NumChildren++;
+				if (BoneParentIndices.Get(BoneId) == ParentIndex)
+				{
+					NumChildren++;
+				}
 			}
-		}
-		return NumChildren;
-	};
+			return NumChildren;
+		};
 
-	int32 Index = 0;
-	for (const FBoneID& BoneId: Bones.GetElementIDs())
-	{
-		SkeletalMeshImportData::FBone NewBone;
+		int32 Index = 0;
+		for (const FBoneID& BoneId : Bones.GetElementIDs())
+		{
+			SkeletalMeshImportData::FBone NewBone;
 			NewBone.Name = BoneNames.Get(BoneId).ToString();
 			NewBone.ParentIndex = BoneParentIndices.Get(BoneId);
 			NewBone.BonePos.Transform = FTransform3f(BonePoses.Get(BoneId));
 			NewBone.NumChildren = GetNumChildren(Index);
-		SkelMeshImportData.RefBonesBinary.Emplace( MoveTemp(NewBone) );
-		Index++;
+			SkelMeshImportData.RefBonesBinary.Emplace(MoveTemp(NewBone));
+			Index++;
+		}
 	}
 
 	SkelMeshImportData.CleanUpUnusedMaterials();
