Add mesh streaming update breadcrumb to get more information out of the crash

#localization none
#tests verified that the breadcrumb code worked locally on a platform

diff --git a/Engine/Source/Runtime/Engine/Private/Streaming/SkeletalMeshUpdate.cpp b/Engine/Source/Runtime/Engine/Private/Streaming/SkeletalMeshUpdate.cpp
--- a/Engine/Source/Runtime/Engine/Private/Streaming/SkeletalMeshUpdate.cpp
+++ b/Engine/Source/Runtime/Engine/Private/Streaming/SkeletalMeshUpdate.cpp
@@ -152,6 +152,16 @@ void FSkeletalMeshStreamIn::CreateBuffers_Internal(const FContext& Context)
 		for (int32 LODIndex = PendingFirstLODIdx; LODIndex < CurrentFirstLODIdx; ++LODIndex)
 		{
 			FSkeletalMeshLODRenderData& LODResource = *Context.LODResourcesView[LODIndex];
+
+#if ALLOW_RENDER_ASSET_STREAMING_BREADCRUMB
+			const FStaticMeshVertexBuffer& TangentTexCoordBuffers = LODResource.StaticVertexBuffers.StaticMeshVertexBuffer;
+			if (GStreamingEnableRenderAssetUpdateBreadcrumb != 0 && TangentTexCoordBuffers.GetNumVertices() > 0 && !TangentTexCoordBuffers.GetTangentData())
+			{
+				FRenderAssetUpdateBreadcrumb* Breadcrumb = (FRenderAssetUpdateBreadcrumb*)alloca(sizeof(FRenderAssetUpdateBreadcrumb));
+				Breadcrumb->SetDebugData(Mesh->GetPathName(), LODIndex, TangentTexCoordBuffers.GetNumVertices(), false);
+				*(uint32*)3 = 0xed000001;
+			}
+#endif
 			if (bRenderThread)
 			{
 				IntermediateBuffersArray[LODIndex].CreateFromCPUData_RenderThread(LODResource);
@@ -530,6 +540,16 @@ void FSkeletalMeshStreamIn_IO::SerializeLODData(const FContext& Context)
 			const bool bNeedsCPUAccess = FSkeletalMeshLODRenderData::ShouldKeepCPUResources(Mesh, LODIndex + Context.AssetLODBias, bForceKeepCPUResources);
 			constexpr uint8 DummyStripFlags = 0;
 			LODResource.SerializeStreamedData(Ar, const_cast<USkeletalMesh*>(Mesh), LODIndex + Context.AssetLODBias, DummyStripFlags, bNeedsCPUAccess, bForceKeepCPUResources);
+
+#if ALLOW_RENDER_ASSET_STREAMING_BREADCRUMB
+			const FStaticMeshVertexBuffer& TangentTexCoordBuffers = LODResource.StaticVertexBuffers.StaticMeshVertexBuffer;
+			if (GStreamingEnableRenderAssetUpdateBreadcrumb != 0 && TangentTexCoordBuffers.GetNumVertices() > 0 && !TangentTexCoordBuffers.GetTangentData())
+			{
+				FRenderAssetUpdateBreadcrumb* Breadcrumb = (FRenderAssetUpdateBreadcrumb*)alloca(sizeof(FRenderAssetUpdateBreadcrumb));
+				Breadcrumb->SetDebugData(Mesh->GetPathName(), LODIndex, TangentTexCoordBuffers.GetNumVertices(), Ar.GetError());
+				*(uint32*)3 = 0xed000001;
+			}
+#endif
 		}
 
 		BulkData = FIoBuffer();
diff --git a/Engine/Source/Runtime/Engine/Private/Streaming/StaticMeshUpdate.cpp b/Engine/Source/Runtime/Engine/Private/Streaming/StaticMeshUpdate.cpp
--- a/Engine/Source/Runtime/Engine/Private/Streaming/StaticMeshUpdate.cpp
+++ b/Engine/Source/Runtime/Engine/Private/Streaming/StaticMeshUpdate.cpp
@@ -168,6 +168,16 @@ void FStaticMeshStreamIn::CreateBuffers_Internal(const FContext& Context)
 		for (int32 LODIdx = PendingFirstLODIdx; LODIdx < CurrentFirstLODIdx; ++LODIdx)
 		{
 			FStaticMeshLODResources& LODResource = *Context.LODResourcesView[LODIdx];
+
+#if ALLOW_RENDER_ASSET_STREAMING_BREADCRUMB
+			const FStaticMeshVertexBuffer& TangentTexCoordBuffers = LODResource.VertexBuffers.StaticMeshVertexBuffer;
+			if (GStreamingEnableRenderAssetUpdateBreadcrumb != 0 && TangentTexCoordBuffers.GetNumVertices() > 0 && !TangentTexCoordBuffers.GetTangentData() && Context.Mesh)
+			{
+				FRenderAssetUpdateBreadcrumb* Breadcrumb = (FRenderAssetUpdateBreadcrumb*)alloca(sizeof(FRenderAssetUpdateBreadcrumb));
+				Breadcrumb->SetDebugData(Context.Mesh->GetPathName(), LODIdx, TangentTexCoordBuffers.GetNumVertices(), false);
+				*(uint32*)3 = 0xed000002;
+			}
+#endif
 			if (bRenderThread)
 			{
 				IntermediateBuffersArray[LODIdx].CreateFromCPUData_RenderThread(LODResource);
@@ -571,6 +581,16 @@ void FStaticMeshStreamIn_IO::SerializeLODData(const FContext& Context)
 			typename FStaticMeshLODResources::FStaticMeshBuffersSize DummyBuffersSize;
 			LODResource.SerializeBuffers(Ar, const_cast<UStaticMesh*>(Mesh), DummyStripFlags, DummyBuffersSize);
 			check(DummyBuffersSize.CalcBuffersSize() == LODResource.BuffersSize);
+
+#if ALLOW_RENDER_ASSET_STREAMING_BREADCRUMB
+			const FStaticMeshVertexBuffer& TangentTexCoordBuffers = LODResource.VertexBuffers.StaticMeshVertexBuffer;
+			if (GStreamingEnableRenderAssetUpdateBreadcrumb != 0 && TangentTexCoordBuffers.GetNumVertices() > 0 && !TangentTexCoordBuffers.GetTangentData())
+			{
+				FRenderAssetUpdateBreadcrumb* Breadcrumb = (FRenderAssetUpdateBreadcrumb*)alloca(sizeof(FRenderAssetUpdateBreadcrumb));
+				Breadcrumb->SetDebugData(Mesh->GetPathName(), LODIdx, TangentTexCoordBuffers.GetNumVertices(), Ar.GetError());
+				*(uint32*)3 = 0xed000002;
+			}
+#endif
 		}
 
 		BulkData = FIoBuffer();
diff --git a/Engine/Source/Runtime/Engine/Private/Streaming/TextureStreamingHelpers.cpp b/Engine/Source/Runtime/Engine/Private/Streaming/TextureStreamingHelpers.cpp
--- a/Engine/Source/Runtime/Engine/Private/Streaming/TextureStreamingHelpers.cpp
+++ b/Engine/Source/Runtime/Engine/Private/Streaming/TextureStreamingHelpers.cpp
@@ -299,6 +299,28 @@ static TAutoConsoleVariable<int32> CVarStreamingVRAMPercentageClamp(
 	TEXT("This avoids reserving too much space for systems with a lot of VRAM. (default=1024)"),
 	ECVF_Cheat);
 
+#if ALLOW_RENDER_ASSET_STREAMING_BREADCRUMB
+int32 GStreamingEnableRenderAssetUpdateBreadcrumb = 1;
+static FAutoConsoleVariableRef CVarStreamingEnableRenderAssetUpdateBreadcrumb(
+	TEXT("r.Streaming.EnableRenderAssetUpdateBreadcrumb"),
+	GStreamingEnableRenderAssetUpdateBreadcrumb,
+	TEXT("Whether to detect invalid streaming states, crash when detected, and store debug info onto stack"),
+	ECVF_Default);
+
+void FRenderAssetUpdateBreadcrumb::SetDebugData(const FString& InAssetName, int32 InLODIndex, uint32 InNumVertices, bool bInIsArError)
+{
+	FMemory::Memzero(this, sizeof(FRenderAssetUpdateBreadcrumb));
+	const int32 CopySize = FMath::Min<int32>(InAssetName.Len(), UE_ARRAY_COUNT(AssetName) - 1);
+	if (CopySize > 0)
+	{
+		FCStringAnsi::Strncpy(AssetName, TCHAR_TO_ANSI(*InAssetName), CopySize);
+	}
+	LODIndex = InLODIndex;
+	NumVertices = InNumVertices;
+	bIsArError = bInIsArError;
+}
+#endif
+
 void FRenderAssetStreamingSettings::Update()
 {
 	MaxEffectiveScreenSize = CVarStreamingScreenSizeEffectiveMax.GetValueOnAnyThread();
diff --git a/Engine/Source/Runtime/Engine/Private/Streaming/TextureStreamingHelpers.h b/Engine/Source/Runtime/Engine/Private/Streaming/TextureStreamingHelpers.h
--- a/Engine/Source/Runtime/Engine/Private/Streaming/TextureStreamingHelpers.h
+++ b/Engine/Source/Runtime/Engine/Private/Streaming/TextureStreamingHelpers.h
@@ -12,6 +12,10 @@ TextureStreamingHelpers.h: Definitions of classes used for texture streaming.
 #include "Misc/MemStack.h"
 #include "Engine/TextureDefines.h"
 
+#ifndef ALLOW_RENDER_ASSET_STREAMING_BREADCRUMB
+#define ALLOW_RENDER_ASSET_STREAMING_BREADCRUMB 0
+#endif
+
 class AActor;
 class UStreamableRenderAsset;
 
@@ -187,3 +191,17 @@ struct FRenderAssetStreamingStats
 
 // Helper to access the level bStaticComponentsRegisteredInStreamingManager flag.
 extern bool OwnerLevelHasRegisteredStaticComponentsInStreamingManager(const class AActor* Owner);
+
+#if ALLOW_RENDER_ASSET_STREAMING_BREADCRUMB
+struct FRenderAssetUpdateBreadcrumb
+{
+	ANSICHAR AssetName[256];
+	int32 LODIndex;
+	uint32 NumVertices;
+	bool bIsArError;
+
+	FORCENOINLINE void SetDebugData(const FString& InAssetName, int32 InLODIndex, uint32 InNumVertices, bool bInIsArError);
+};
+
+extern int32 GStreamingEnableRenderAssetUpdateBreadcrumb;
+#endif
