Substrate - fixed simplification code issues due to logic.
Code clean up.

diff --git a/Engine/Source/Runtime/Engine/Private/Materials/HLSLMaterialTranslator.cpp b/Engine/Source/Runtime/Engine/Private/Materials/HLSLMaterialTranslator.cpp
--- a/Engine/Source/Runtime/Engine/Private/Materials/HLSLMaterialTranslator.cpp
+++ b/Engine/Source/Runtime/Engine/Private/Materials/HLSLMaterialTranslator.cpp
@@ -900,7 +900,7 @@ bool FHLSLMaterialTranslator::Translate()
 		#if DEBUG_STRATA_TREE_STACK
 			UE_LOG(LogMaterial, Display, TEXT(" StrataTreeStack: StrataGenerateMaterialTopologyTree"));
 		#endif
-			for(uint32 StrataCompilationContextIndex = 0; StrataCompilationContextIndex < EStrataCompilationContext::SCC_MAX; ++StrataCompilationContextIndex) // sebok
+			for(uint32 StrataCompilationContextIndex = 0; StrataCompilationContextIndex < EStrataCompilationContext::SCC_MAX; ++StrataCompilationContextIndex)
 			{
 				// This is needed because StrataGenerateMaterialTopologyTree will call some compiler context though material expressions.
 				CurrentStrataCompilationContext = EStrataCompilationContext(StrataCompilationContextIndex);
@@ -2328,7 +2328,7 @@ void FHLSLMaterialTranslator::GetMaterialEnvironment(EShaderPlatform InPlatform,
 
 
 		FString StrataMaterialDescription;
-		for (uint32 StrataCompilationContextIndex = 0; StrataCompilationContextIndex < EStrataCompilationContext::SCC_MAX; ++StrataCompilationContextIndex) // sebok
+		for (uint32 StrataCompilationContextIndex = 0; StrataCompilationContextIndex < EStrataCompilationContext::SCC_MAX; ++StrataCompilationContextIndex)
 		{
 			FStrataCompilationContext& StrataCtx = StrataCompilationContext[StrataCompilationContextIndex];
 			const FStrataSimplificationStatus& StrataSimplificationStatus = StrataCtx.StrataSimplificationStatus;
@@ -2566,7 +2566,7 @@ void FHLSLMaterialTranslator::GetSharedInputsMaterialCode(FString& PixelMembersD
 				PixelInputInitializerValues += FString::Printf(TEXT("\t#endif\n"));
 			}
 
-			for (uint32 StrataCompilationContextIndex = 0; StrataCompilationContextIndex < EStrataCompilationContext::SCC_MAX; ++StrataCompilationContextIndex) // sebok
+			for (uint32 StrataCompilationContextIndex = 0; StrataCompilationContextIndex < EStrataCompilationContext::SCC_MAX; ++StrataCompilationContextIndex)
 			{
 				FStrataCompilationContext& StrataCtx = StrataCompilationContext[StrataCompilationContextIndex];
 				if (StrataCtx.CodeChunkToStrataSharedLocalBasis.Num() > 0)
@@ -11416,7 +11416,7 @@ bool FHLSLMaterialTranslator::FStrataCompilationContext::StrataGenerateDerivedMa
 		}
 	}
 
-	StrataSimplificationStatus.bRunFullSimplification = (Compiler->StrataCompilationConfig.bFullSimplify || CompilationContextIndex == EStrataCompilationContext::SCC_FullySimplified) && TotalBSDFCount > 1 && !Compiler->bStrataUsesConversionFromLegacy ;
+	StrataSimplificationStatus.bRunFullSimplification = Compiler->StrataCompilationConfig.bFullSimplify || CompilationContextIndex == EStrataCompilationContext::SCC_FullySimplified;
 
 	const uint32 StrataBytePerPixel = Strata::GetBytePerPixel(Compiler->GetShaderPlatform());
 	do 
@@ -11492,7 +11492,7 @@ bool FHLSLMaterialTranslator::FStrataCompilationContext::StrataGenerateDerivedMa
 
 				// We mark the top of a parameter blending tree as a BSDF now to allocate a slot for it that can then be used next for non-parameter blending operations.
 				// Intermediate parameter blending BSDF and operation will be done inline and stored in FStrataData.
-				if (bRootOfParameterBlendingSubTree)
+				if (CurrentOperator.OperatorType != STRATA_OPERATOR_BSDF && bRootOfParameterBlendingSubTree)
 				{
 					CurrentOperator.OperatorType = STRATA_OPERATOR_BSDF;
 					CurrentOperator.BSDFIndex = StrataMaterialBSDFCount++;
