Add Subsutrate topology previsualization on the Substrate node within material editor.

diff --git a/Engine/Source/Editor/GraphEditor/Private/MaterialNodes/SGraphNodeMaterialBase.cpp b/Engine/Source/Editor/GraphEditor/Private/MaterialNodes/SGraphNodeMaterialBase.cpp
--- a/Engine/Source/Editor/GraphEditor/Private/MaterialNodes/SGraphNodeMaterialBase.cpp
+++ b/Engine/Source/Editor/GraphEditor/Private/MaterialNodes/SGraphNodeMaterialBase.cpp
@@ -25,6 +25,7 @@
 #include "Materials/MaterialExpression.h"
 #include "Materials/MaterialExpressionBreakMaterialAttributes.h"
 #include "Materials/MaterialExpressionMakeMaterialAttributes.h"
+#include "Materials/MaterialExpressionStrata.h"
 #include "Materials/MaterialFunction.h"
 #include "Math/Color.h"
 #include "Math/IntPoint.h"
@@ -60,6 +61,8 @@
 #include "Widgets/SOverlay.h"
 #include "Widgets/SViewport.h"
 #include "Widgets/Text/STextBlock.h"
+#include "Rendering/StrataMaterialShared.h"
+#include "SGraphSubstrateMaterial.h"
 
 class FWidgetStyle;
 class SWidget;
@@ -346,6 +349,12 @@ void SGraphNodeMaterialBase::CreatePinWidgets()
 				NewPin->SetCustomPinIcon(CacheImg_Pin_NotConnectable, CacheImg_Pin_NotConnectable);
 			}
 
+			// Override pin color for Substrate node
+			if (Strata::IsStrataEnabled())
+			{
+				FSubstrateWidget::GetPinColor(NewPin, MaterialNode);
+			}
+
 			this->AddPin(NewPin.ToSharedRef());
 		}
 	}
@@ -471,6 +480,37 @@ void SGraphNodeMaterialBase::CreateBelowPinControls(TSharedPtr<SVerticalBox> Mai
 			];
 		}
 	}
+
+	// Preview of Substrate nodes topology
+	if (Strata::IsStrataEnabled() && MaterialNode && MaterialNode->MaterialExpression->IsA(UMaterialExpressionStrataBSDF::StaticClass()))
+	{
+		if (const UMaterialExpressionStrataBSDF* StrataExpression = (const UMaterialExpressionStrataBSDF*)MaterialNode->MaterialExpression)
+		{		
+			if (UMaterial* MaterialForStats = StrataExpression->Material)
+			{
+				if (const FMaterialResource* MaterialResource = MaterialForStats->GetMaterialResource(GMaxRHIFeatureLevel))
+				{
+					if (FMaterialShaderMap* ShaderMap = MaterialResource->GetGameThreadShaderMap())
+					{
+						const FStrataMaterialCompilationOutput& CompilationOutput = ShaderMap->GetStrataMaterialCompilationOutput();
+						MainBox->AddSlot()
+						.Padding(Settings->GetNonPinNodeBodyPadding())
+						.AutoHeight()
+						[
+							SNew(SHorizontalBox)
+							+SHorizontalBox::Slot()
+							.VAlign(VAlign_Center)
+							.HAlign(HAlign_Center)
+							[							
+								FSubstrateWidget::ProcessOperator(CompilationOutput, StrataExpression->MaterialExpressionGuid)
+							]
+						];
+					}
+				}
+	
+			}
+		}
+	}
 }
 
 void SGraphNodeMaterialBase::SetDefaultTitleAreaWidget(TSharedRef<SOverlay> DefaultTitleAreaWidget)
diff --git a/Engine/Source/Editor/GraphEditor/Private/MaterialNodes/SGraphSubstrateMaterial.cpp b/Engine/Source/Editor/GraphEditor/Private/MaterialNodes/SGraphSubstrateMaterial.cpp
new file mode 100644
--- /dev/null
+++ b/Engine/Source/Editor/GraphEditor/Private/MaterialNodes/SGraphSubstrateMaterial.cpp
@@ -0,0 +1,209 @@
+// Copyright Epic Games, Inc. All Rights Reserved.
+
+#pragma once
+
+#include "SGraphSubstrateMaterial.h"
+#include "Internationalization/Text.h"
+#include "Rendering/RenderingCommon.h"
+#include "MaterialGraph/MaterialGraphNode.h"
+#include "Materials/MaterialExpression.h"
+#include "Materials/MaterialExpressionStrata.h"
+#include "SGraphPin.h"
+#include "Math/Color.h"
+
+enum class ESubstrateWidgetOutputType : uint8
+{
+	Node,
+	DetailPanel,
+};
+
+static const TSharedRef<SWidget> InternalProcessOperator(
+	const FStrataMaterialCompilationOutput& CompilationOutput, 
+	const FStrataOperator& Op, 
+	ESubstrateWidgetOutputType OutputType,
+	const FGuid& InGuid, 
+	EStyleColor OverrideColor)
+{
+	const bool bIsCurrent = OutputType == ESubstrateWidgetOutputType::Node ? InGuid == Op.MaterialExpressionGuid : true;
+	const EStyleColor Color0 = bIsCurrent ? EStyleColor::AccentPurple : OverrideColor;
+	const EStyleColor Color1 = bIsCurrent ? EStyleColor::AccentGreen : OverrideColor;
+	switch (Op.OperatorType)
+	{
+		case STRATA_OPERATOR_WEIGHT:
+			return InternalProcessOperator(CompilationOutput, CompilationOutput.Operators[Op.LeftIndex], OutputType, InGuid, OverrideColor);
+			break;
+		case STRATA_OPERATOR_VERTICAL:
+		{
+			auto VerticalOperator = SNew(SVerticalBox)
+				+SVerticalBox::Slot()
+				.AutoHeight()
+				.VAlign(VAlign_Fill)
+				.HAlign(HAlign_Fill)
+				.Padding(0.0f, 0.0f, 1.0f, 1.0f)
+				[
+					InternalProcessOperator(CompilationOutput, CompilationOutput.Operators[Op.LeftIndex], OutputType, InGuid, Color0)
+				]
+				+ SVerticalBox::Slot()
+				.AutoHeight()
+				.VAlign(VAlign_Fill)
+				.HAlign(HAlign_Fill)
+				.Padding(0.0f, 0.0f, 1.0f, 1.0f)
+				[
+					InternalProcessOperator(CompilationOutput, CompilationOutput.Operators[Op.RightIndex], OutputType, InGuid, Color1)
+				];
+			return VerticalOperator->AsShared();
+		}
+		break;
+		case STRATA_OPERATOR_HORIZONTAL:
+		{
+			auto HorizontalOperator = SNew(SHorizontalBox)
+				+SHorizontalBox::Slot()
+				.AutoWidth()
+				.VAlign(VAlign_Fill)
+				.HAlign(HAlign_Fill)
+				.Padding(0.0f, 0.0f, 1.0f, 1.0f)
+				[
+					InternalProcessOperator(CompilationOutput, CompilationOutput.Operators[Op.LeftIndex], OutputType, InGuid, Color0)
+				]
+				+SHorizontalBox::Slot()
+				.AutoWidth()
+				.VAlign(VAlign_Fill)
+				.HAlign(HAlign_Fill)
+				.Padding(0.0f, 0.0f, 1.0f, 1.0f)
+				[
+					InternalProcessOperator(CompilationOutput, CompilationOutput.Operators[Op.RightIndex], OutputType, InGuid, Color1)
+				];
+			return HorizontalOperator->AsShared();
+		}
+		break;
+		case STRATA_OPERATOR_ADD:
+		{
+			auto HorizontalOperator = SNew(SHorizontalBox)
+				+SHorizontalBox::Slot()
+				.AutoWidth()
+				.VAlign(VAlign_Fill)
+				.HAlign(HAlign_Fill)
+				.Padding(0.0f, 0.0f, 1.0f, 1.0f)
+				[
+					InternalProcessOperator(CompilationOutput, CompilationOutput.Operators[Op.LeftIndex], OutputType, InGuid, Color0)
+				]
+				+SHorizontalBox::Slot()
+				.AutoWidth()
+				.VAlign(VAlign_Fill)
+				.HAlign(HAlign_Fill)
+				.Padding(0.0f, 0.0f, 1.0f, 1.0f)
+				[
+					InternalProcessOperator(CompilationOutput, CompilationOutput.Operators[Op.RightIndex], OutputType, InGuid, Color1)
+				];
+			return HorizontalOperator->AsShared();
+		}
+		break;
+		case STRATA_OPERATOR_BSDF_LEGACY:	// legacy BSDF should have been converted to BSDF already.
+		case STRATA_OPERATOR_BSDF:
+		{
+			FString BSDFDesc = OutputType == ESubstrateWidgetOutputType::Node ? 
+											 FString(TEXT("BSDF")) : 
+											 FString::Printf(TEXT("BSDF (%s%s%s%s%s%s%s%s)")
+											, Op.bBSDFHasEdgeColor ? TEXT("F90 ") : TEXT("")
+											, Op.bBSDFHasSSS ? TEXT("SSS ") : TEXT("")
+											, Op.bBSDFHasMFPPluggedIn ? TEXT("MFP ") : TEXT("")
+											, Op.bBSDFHasAnisotropy ? TEXT("Ani ") : TEXT("")
+											, Op.bBSDFHasSecondRoughnessOrSimpleClearCoat ? TEXT("2Ro ") : TEXT("")
+											, Op.bBSDFHasFuzz ? TEXT("Fuz ") : TEXT("")
+											, Op.bBSDFHasGlint ? TEXT("Gli ") : TEXT("")
+											, Op.bBSDFHasSpecularProfile ? TEXT("Spc ") : TEXT("")
+											);
+
+			static FString ToolTip;
+			if (OutputType == ESubstrateWidgetOutputType::DetailPanel && ToolTip.IsEmpty())
+			{
+				ToolTip += TEXT("SSS means the BSDF features subsurface profile or subsurface setup using MFP.\n");
+				ToolTip += TEXT("MFP means the BSDF MFP is specified by the user.\n");
+				ToolTip += TEXT("F90 means the BSDF edge specular color representing reflectivity at grazing angle is used.\n");
+				ToolTip += TEXT("Fuz means the BSDF fuzz layer is enabled.\n");
+				ToolTip += TEXT("2Ro means the BSDF either uses a second specular lob with a second roughness, or the legacy simple clear coat.\n");
+				ToolTip += TEXT("Ani means the BSDF anisotropic specular lighting is used.\n");
+				ToolTip += TEXT("Gli means the BSDF features glints.\n");
+				ToolTip += TEXT("Spc means the BSDF features specular profile.\n");
+			}
+
+			const EStyleColor Color = OverrideColor != EStyleColor::MAX ? OverrideColor : (bIsCurrent ? EStyleColor::AccentGreen : EStyleColor::AccentGray);
+			auto BSDF = SNew(SErrorText)
+				.ErrorText(FText::FromString(BSDFDesc))
+				.BackgroundColor(FSlateColor(Color))
+				.ToolTipText(FText::FromString(ToolTip));
+			return BSDF->AsShared();
+		}
+		break;
+	}
+
+	auto TreeOperatorError = SNew(SErrorText)
+		.ErrorText(FText::FromString(TEXT("Tree Operator Error")))
+		.BackgroundColor(FSlateColor(EStyleColor::AccentRed));
+	return TreeOperatorError->AsShared();
+}
+
+const TSharedRef<SWidget> FSubstrateWidget::ProcessOperator(const FStrataMaterialCompilationOutput& CompilationOutput)
+{
+	return InternalProcessOperator(CompilationOutput, CompilationOutput.Operators[CompilationOutput.RootOperatorIndex], ESubstrateWidgetOutputType::DetailPanel, FGuid(), EStyleColor::MAX);
+}
+
+const TSharedRef<SWidget> FSubstrateWidget::ProcessOperator(const FStrataMaterialCompilationOutput& CompilationOutput, const FGuid& InGuid)
+{
+	return InternalProcessOperator(CompilationOutput, CompilationOutput.Operators[CompilationOutput.RootOperatorIndex], ESubstrateWidgetOutputType::Node, InGuid, EStyleColor::MAX);
+}
+
+void FSubstrateWidget::GetPinColor(TSharedPtr<SGraphPin>& Out, const UMaterialGraphNode* InNode)
+{	
+	const FLinearColor Color0 = FLinearColor(FColor::Purple);
+	const FLinearColor Color1 = FLinearColor(FColor::Green);
+
+	FLinearColor ColorModifier;
+	bool bHasColorModifier = false;
+	// Strata operator override pin color to ease material topology visualization
+	const FName PinName = Out->GetPinObj()->PinName;
+	if (InNode->MaterialExpression->IsA(UMaterialExpressionStrataVerticalLayering::StaticClass()))
+	{			
+		if (PinName == InNode->MaterialExpression->GetInputName(0)) // Top
+		{
+			bHasColorModifier = true;
+			ColorModifier = Color0;
+		}
+		else if (PinName == InNode->MaterialExpression->GetInputName(1)) // Base
+		{
+			bHasColorModifier = true;
+			ColorModifier = Color1;
+		}
+	}
+	else if (InNode->MaterialExpression->IsA(UMaterialExpressionStrataHorizontalMixing::StaticClass()))
+	{
+		if (PinName == InNode->MaterialExpression->GetInputName(1)) // Foreground
+		{
+			bHasColorModifier = true;
+			ColorModifier = Color1;
+		}
+		else if (PinName == InNode->MaterialExpression->GetInputName(0)) // Background
+		{
+			bHasColorModifier = true;
+			ColorModifier = Color0;
+		}
+	}
+	else if (InNode->MaterialExpression->IsA(UMaterialExpressionStrataAdd::StaticClass()))
+	{
+		if (PinName == InNode->MaterialExpression->GetInputName(0)) // A
+		{
+			bHasColorModifier = true;
+			ColorModifier = Color0;
+		}
+		else if (PinName == InNode->MaterialExpression->GetInputName(1)) // B
+		{
+			bHasColorModifier = true;
+			ColorModifier = Color1;
+		}
+	}
+
+	if (bHasColorModifier)
+	{
+		Out->SetPinColorModifier(ColorModifier);
+	}
+}
\ No newline at end of file
diff --git a/Engine/Source/Editor/GraphEditor/Public/SGraphSubstrateMaterial.h b/Engine/Source/Editor/GraphEditor/Public/SGraphSubstrateMaterial.h
new file mode 100644
--- /dev/null
+++ b/Engine/Source/Editor/GraphEditor/Public/SGraphSubstrateMaterial.h
@@ -0,0 +1,17 @@
+// Copyright Epic Games, Inc. All Rights Reserved.
+
+#pragma once
+
+#include "Internationalization/Text.h"
+#include "Rendering/RenderingCommon.h"
+#include "SGraphPin.h"
+
+class UMaterialGraphNode;
+struct FStrataMaterialCompilationOutput;
+
+struct GRAPHEDITOR_API FSubstrateWidget
+{
+	static const TSharedRef<SWidget> ProcessOperator(const FStrataMaterialCompilationOutput& CompilationOutput);
+	static const TSharedRef<SWidget> ProcessOperator(const FStrataMaterialCompilationOutput& CompilationOutput, const FGuid& InGuid);
+	static void GetPinColor(TSharedPtr<SGraphPin>& Out, const UMaterialGraphNode* InNode);
+};
\ No newline at end of file
diff --git a/Engine/Source/Editor/MaterialEditor/Private/MaterialEditor.cpp b/Engine/Source/Editor/MaterialEditor/Private/MaterialEditor.cpp
--- a/Engine/Source/Editor/MaterialEditor/Private/MaterialEditor.cpp
+++ b/Engine/Source/Editor/MaterialEditor/Private/MaterialEditor.cpp
@@ -6062,6 +6062,22 @@ FText FMaterialEditor::GetOriginalObjectName() const
 	return FText::FromString(GetEditingObjects()[0]->GetName());
 }
 
+void FMaterialEditor::UpdateStrataTopologyPreview()
+{
+	if (Strata::IsGlintEnabled())
+	{
+		// Update all Strata node which have a preview
+		for (int32 Index = 0; Index < Material->MaterialGraph->Nodes.Num(); ++Index)
+		{
+			UMaterialGraphNode* MaterialNode = Cast<UMaterialGraphNode>(Material->MaterialGraph->Nodes[Index]);
+			if (MaterialNode && MaterialNode->MaterialExpression && MaterialNode->MaterialExpression->IsA(UMaterialExpressionStrataBSDF::StaticClass()))
+			{
+				MaterialNode->ReconstructNode();
+			}
+		}
+	}
+}
+
 void FMaterialEditor::UpdateMaterialAfterGraphChange()
 {
 	FlushRenderingCommands();
@@ -6092,6 +6108,8 @@ void FMaterialEditor::UpdateMaterialAfterGraphChange()
 	}
 
 	Material->MaterialGraph->UpdatePinTypes();
+
+	UpdateStrataTopologyPreview();
 }
 
 void FMaterialEditor::JumpToHyperlink(const UObject* ObjectReference)
diff --git a/Engine/Source/Editor/MaterialEditor/Private/MaterialEditor.h b/Engine/Source/Editor/MaterialEditor/Private/MaterialEditor.h
--- a/Engine/Source/Editor/MaterialEditor/Private/MaterialEditor.h
+++ b/Engine/Source/Editor/MaterialEditor/Private/MaterialEditor.h
@@ -697,6 +697,8 @@ private:
 	bool IsFeaturePreviewChecked(ERHIFeatureLevel::Type TestFeatureLevel) const;
 	bool IsFeaturePreviewAvailable(ERHIFeatureLevel::Type TestFeatureLevel) const;
 
+	/** Update Substrate topology preview */
+	void UpdateStrataTopologyPreview();
 
 public:
 
diff --git a/Engine/Source/Editor/MaterialEditor/Private/SMaterialEditorStrataWidget.cpp b/Engine/Source/Editor/MaterialEditor/Private/SMaterialEditorStrataWidget.cpp
--- a/Engine/Source/Editor/MaterialEditor/Private/SMaterialEditorStrataWidget.cpp
+++ b/Engine/Source/Editor/MaterialEditor/Private/SMaterialEditorStrataWidget.cpp
@@ -14,6 +14,7 @@
 #include "Widgets/Layout/SScrollBox.h"
 #include "Widgets/SNullWidget.h"
 #include "Widgets/Notifications/SErrorText.h"
+#include "SGraphSubstrateMaterial.h"
 
 #define LOCTEXT_NAMESPACE "SMaterialEditorStrataWidget"
 
@@ -303,124 +304,8 @@ void SMaterialEditorStrataWidget::Tick(const FGeometry& AllottedGeometry, const
 				// Now generate a visual representation of the material from the topology tree of operators.
 				{
 					if (CompilationOutput.RootOperatorIndex >= 0)
-					{
-						std::function<const TSharedRef<SWidget>(const FStrataOperator&)> ProcessOperator = [&](const FStrataOperator& Op) -> const TSharedRef<SWidget>
-						{
-							switch (Op.OperatorType)
-							{
-							case STRATA_OPERATOR_WEIGHT:
-								return ProcessOperator(CompilationOutput.Operators[Op.LeftIndex]);
-								break;
-							case STRATA_OPERATOR_VERTICAL:
-							{
-								auto VerticalOperator = SNew(SVerticalBox)
-									+SVerticalBox::Slot()
-									.AutoHeight()
-									.VAlign(VAlign_Fill)
-									.HAlign(HAlign_Fill)
-									.Padding(0.0f, 0.0f, 1.0f, 1.0f)
-									[
-										ProcessOperator(CompilationOutput.Operators[Op.LeftIndex])
-									]
-									+ SVerticalBox::Slot()
-									.AutoHeight()
-									.VAlign(VAlign_Fill)
-									.HAlign(HAlign_Fill)
-									.Padding(0.0f, 0.0f, 1.0f, 1.0f)
-									[
-										ProcessOperator(CompilationOutput.Operators[Op.RightIndex])
-									];
-								return VerticalOperator->AsShared();
-							}
-							break;
-							case STRATA_OPERATOR_HORIZONTAL:
-							{
-								auto HorizontalOperator = SNew(SHorizontalBox)
-									+SHorizontalBox::Slot()
-									.AutoWidth()
-									.VAlign(VAlign_Fill)
-									.HAlign(HAlign_Fill)
-									.Padding(0.0f, 0.0f, 1.0f, 1.0f)
-									[
-										ProcessOperator(CompilationOutput.Operators[Op.LeftIndex])
-									]
-									+SHorizontalBox::Slot()
-									.AutoWidth()
-									.VAlign(VAlign_Fill)
-									.HAlign(HAlign_Fill)
-									.Padding(0.0f, 0.0f, 1.0f, 1.0f)
-									[
-										ProcessOperator(CompilationOutput.Operators[Op.RightIndex])
-									];
-								return HorizontalOperator->AsShared();
-							}
-							break;
-							case STRATA_OPERATOR_ADD:
-							{
-								auto HorizontalOperator = SNew(SHorizontalBox)
-									+SHorizontalBox::Slot()
-									.AutoWidth()
-									.VAlign(VAlign_Fill)
-									.HAlign(HAlign_Fill)
-									.Padding(0.0f, 0.0f, 1.0f, 1.0f)
-									[
-										ProcessOperator(CompilationOutput.Operators[Op.LeftIndex])
-									]
-									+SHorizontalBox::Slot()
-									.AutoWidth()
-									.VAlign(VAlign_Fill)
-									.HAlign(HAlign_Fill)
-									.Padding(0.0f, 0.0f, 1.0f, 1.0f)
-									[
-										ProcessOperator(CompilationOutput.Operators[Op.RightIndex])
-									];
-								return HorizontalOperator->AsShared();
-							}
-							break;
-							case STRATA_OPERATOR_BSDF_LEGACY:	// legacy BSDF should have been converted to BSDF already.
-							case STRATA_OPERATOR_BSDF:
-							{
-								FString BSDFDesc = FString::Printf(TEXT("BSDF (%s%s%s%s%s%s%s%s)")
-									, Op.bBSDFHasEdgeColor ? TEXT("F90 ") : TEXT("")
-									, Op.bBSDFHasSSS ? TEXT("SSS ") : TEXT("")
-									, Op.bBSDFHasMFPPluggedIn ? TEXT("MFP ") : TEXT("")
-									, Op.bBSDFHasAnisotropy ? TEXT("Ani ") : TEXT("")
-									, Op.bBSDFHasSecondRoughnessOrSimpleClearCoat ? TEXT("2Ro ") : TEXT("")
-									, Op.bBSDFHasFuzz ? TEXT("Fuz ") : TEXT("")
-									, Op.bBSDFHasGlint ? TEXT("Gli ") : TEXT("")
-									, Op.bBSDFHasSpecularProfile ? TEXT("Spc ") : TEXT("")
-								);
-
-								static FString ToolTip;
-								if (ToolTip.IsEmpty())
-								{
-									ToolTip += TEXT("SSS means the BSDF features subsurface profile or subsurface setup using MFP.\n");
-									ToolTip += TEXT("MFP means the BSDF MFP is specified by the user.\n");
-									ToolTip += TEXT("F90 means the BSDF edge specular color representing reflectivity at grazing angle is used.\n");
-									ToolTip += TEXT("Fuz means the BSDF fuzz layer is enabled.\n");
-									ToolTip += TEXT("2Ro means the BSDF either uses a second specular lob with a second roughness, or the legacy simple clear coat.\n");
-									ToolTip += TEXT("Ani means the BSDF anisotropic specular lighting is used.\n");
-									ToolTip += TEXT("Gli means the BSDF features glints.\n");
-									ToolTip += TEXT("Spc means the BSDF features specular profile.\n");
-								}
-
-								auto BSDF = SNew(SErrorText)
-									.ErrorText(FText::FromString(BSDFDesc))
-									.BackgroundColor(FSlateColor(EStyleColor::AccentGreen))
-									.ToolTipText(FText::FromString(ToolTip));
-								return BSDF->AsShared();
-							}
-							break;
-							}
-
-							auto TreeOperatorError = SNew(SErrorText)
-								.ErrorText(LOCTEXT("TreeOperatorError", "Tree Operator Error"))
-								.BackgroundColor(FSlateColor(EStyleColor::AccentRed));
-							return TreeOperatorError->AsShared();
-						};
-
-						const FStrataOperator& RootOperator = CompilationOutput.Operators[CompilationOutput.RootOperatorIndex];
-						MaterialBox->SetContent(ProcessOperator(RootOperator));
+					{										
+						MaterialBox->SetContent(FSubstrateWidget::ProcessOperator(CompilationOutput));
 					}
 					else
 					{
diff --git a/Engine/Source/Runtime/Engine/Private/Materials/HLSLMaterialTranslator.cpp b/Engine/Source/Runtime/Engine/Private/Materials/HLSLMaterialTranslator.cpp
--- a/Engine/Source/Runtime/Engine/Private/Materials/HLSLMaterialTranslator.cpp
+++ b/Engine/Source/Runtime/Engine/Private/Materials/HLSLMaterialTranslator.cpp
@@ -11105,7 +11105,7 @@ int32 FHLSLMaterialTranslator::GetLocal(const FName& LocalName)
 	return AddInlinedCodeChunk(MCT_Float1, TEXT("%s"), *Entry->Name);
 }
 
-FStrataOperator& FHLSLMaterialTranslator::StrataCompilationRegisterOperator(int32 OperatorType, FGuid StrataExpressionGuid, UMaterialExpression* Parent, FGuid StrataParentExpressionGuid, bool bUseParameterBlending)
+FStrataOperator& FHLSLMaterialTranslator::StrataCompilationRegisterOperator(int32 OperatorType, FGuid StrataExpressionGuid, UMaterialExpression* Child, UMaterialExpression* Parent, FGuid StrataParentExpressionGuid, bool bUseParameterBlending)
 {
 	FStrataCompilationContext& StrataCtx = StrataCompilationContext[CurrentStrataCompilationContext];
 
@@ -11155,6 +11155,7 @@ FStrataOperator& FHLSLMaterialTranslator::StrataCompilationRegisterOperator(int3
 	NewOperator.MaxDistanceFromLeaves = 0;
 	NewOperator.bIsBottom = false;
 	NewOperator.bIsTop = false;
+	NewOperator.MaterialExpressionGuid = Child->MaterialExpressionGuid;
 
 	return NewOperator;
 }
@@ -11534,7 +11535,7 @@ bool FHLSLMaterialTranslator::FStrataCompilationContext::StrataGenerateDerivedMa
 		//
 		// Make sure all the types have valid children operator indices
 		//
-		for (auto& It : StrataMaterialExpressionRegisteredOperators)
+		for (const auto& It : StrataMaterialExpressionRegisteredOperators)
 		{
 			if (It.IsDiscarded())
 			{
@@ -11692,7 +11693,7 @@ bool FHLSLMaterialTranslator::FStrataCompilationContext::StrataGenerateDerivedMa
 			bool bIsFastWaterPath = false;
 			bool bCustomEncoding = false;
 			bool bUsesComplexSpecialRenderPath = false;
-			for (auto& It : StrataMaterialExpressionRegisteredOperators)
+			for (const auto& It : StrataMaterialExpressionRegisteredOperators)
 			{
 				if (It.IsDiscarded())
 				{
diff --git a/Engine/Source/Runtime/Engine/Private/Materials/HLSLMaterialTranslator.h b/Engine/Source/Runtime/Engine/Private/Materials/HLSLMaterialTranslator.h
--- a/Engine/Source/Runtime/Engine/Private/Materials/HLSLMaterialTranslator.h
+++ b/Engine/Source/Runtime/Engine/Private/Materials/HLSLMaterialTranslator.h
@@ -1223,7 +1223,7 @@ protected:
 	virtual int32 StrataThicknessStackPush(UMaterialExpression* Expression, FExpressionInput* Input) override;
 	virtual void StrataThicknessStackPop() override;
 
-	virtual FStrataOperator& StrataCompilationRegisterOperator(int32 OperatorType, FGuid StrataExpressionGuid, UMaterialExpression* Parent, FGuid StrataParentExpressionGuid, bool bUseParameterBlending = false) override;
+	virtual FStrataOperator& StrataCompilationRegisterOperator(int32 OperatorType, FGuid StrataExpressionGuid, UMaterialExpression* Child, UMaterialExpression* Parent, FGuid StrataParentExpressionGuid, bool bUseParameterBlending = false) override;
 	virtual FStrataOperator& StrataCompilationGetOperator(FGuid StrataExpressionGuid) override;
 	virtual FStrataOperator* StrataCompilationGetOperatorFromIndex(int32 OperatorIndex) override;
 
diff --git a/Engine/Source/Runtime/Engine/Private/Materials/MaterialExpressions.cpp b/Engine/Source/Runtime/Engine/Private/Materials/MaterialExpressions.cpp
--- a/Engine/Source/Runtime/Engine/Private/Materials/MaterialExpressions.cpp
+++ b/Engine/Source/Runtime/Engine/Private/Materials/MaterialExpressions.cpp
@@ -7618,7 +7618,7 @@ FStrataOperator* UMaterialExpressionBlendMaterialAttributes::StrataGenerateMater
 {
 	// STRATA_TODO: this likely no longer work. We would need to stop parsing and always do parameter blending at this stage.
 	const bool bUseParameterBlending = false;
-	FStrataOperator& StrataOperator = Compiler->StrataCompilationRegisterOperator(STRATA_OPERATOR_HORIZONTAL, Compiler->StrataTreeStackGetPathUniqueId(), Parent, Compiler->StrataTreeStackGetParentPathUniqueId(), bUseParameterBlending);
+	FStrataOperator& StrataOperator = Compiler->StrataCompilationRegisterOperator(STRATA_OPERATOR_HORIZONTAL, Compiler->StrataTreeStackGetPathUniqueId(), this, Parent, Compiler->StrataTreeStackGetParentPathUniqueId(), bUseParameterBlending);
 
 	UMaterialExpression* ChildAExpression = A.GetTracedInput().Expression;
 	UMaterialExpression* ChildBExpression = B.GetTracedInput().Expression;
@@ -17012,7 +17012,7 @@ FStrataOperator* UMaterialExpressionFunctionInput::StrataGenerateMaterialTopolog
 	// If we are parsing for a material function input we always needs to return a default valid BSDF operator at least 
 	// If the material requires an input thenthe UI will forcethe user to proviced one.
 	// If this is not the case however, we need to compile a default material to compile the shader, or to preview the material function.
-	FStrataOperator& DefaultSlabOperator = Compiler->StrataCompilationRegisterOperator(STRATA_OPERATOR_BSDF, Compiler->StrataTreeStackGetPathUniqueId(), Parent, Compiler->StrataTreeStackGetParentPathUniqueId());
+	FStrataOperator& DefaultSlabOperator = Compiler->StrataCompilationRegisterOperator(STRATA_OPERATOR_BSDF, Compiler->StrataTreeStackGetPathUniqueId(), this, Parent, Compiler->StrataTreeStackGetParentPathUniqueId());
 	DefaultSlabOperator.BSDFType = STRATA_BSDF_TYPE_SLAB;
 	DefaultSlabOperator.ThicknessIndex = Compiler->StrataThicknessStackGetThicknessIndex();
 	return &DefaultSlabOperator;
@@ -23157,7 +23157,7 @@ FStrataOperator* UMaterialExpressionStrataLegacyConversion::StrataGenerateMateri
 
 	auto AddDefaultWorstCase = [&](bool bSSS)
 	{
-		FStrataOperator& SlabOperator = Compiler->StrataCompilationRegisterOperator(STRATA_OPERATOR_BSDF_LEGACY, Compiler->StrataTreeStackGetPathUniqueId(), Parent, Compiler->StrataTreeStackGetParentPathUniqueId());
+		FStrataOperator& SlabOperator = Compiler->StrataCompilationRegisterOperator(STRATA_OPERATOR_BSDF_LEGACY, Compiler->StrataTreeStackGetPathUniqueId(), this, Parent, Compiler->StrataTreeStackGetParentPathUniqueId());
 		SlabOperator.BSDFType = STRATA_BSDF_TYPE_SLAB;
 		SlabOperator.bBSDFHasSSS = true;
 		SlabOperator.bBSDFHasMFPPluggedIn = true;
@@ -23179,7 +23179,7 @@ FStrataOperator* UMaterialExpressionStrataLegacyConversion::StrataGenerateMateri
 
 		if (ConvertedStrataMaterialInfo.HasShadingModel(SSM_Unlit))
 		{
-			FStrataOperator& Operator = Compiler->StrataCompilationRegisterOperator(STRATA_OPERATOR_BSDF_LEGACY, Compiler->StrataTreeStackGetPathUniqueId(), Parent, Compiler->StrataTreeStackGetParentPathUniqueId());
+			FStrataOperator& Operator = Compiler->StrataCompilationRegisterOperator(STRATA_OPERATOR_BSDF_LEGACY, Compiler->StrataTreeStackGetPathUniqueId(), this, Parent, Compiler->StrataTreeStackGetParentPathUniqueId());
 			Operator.BSDFType = STRATA_BSDF_TYPE_UNLIT;
 			Operator.ThicknessIndex = ThicknessIndex;
 			return &Operator;
@@ -23194,21 +23194,21 @@ FStrataOperator* UMaterialExpressionStrataLegacyConversion::StrataGenerateMateri
 		}
 		else if (ConvertedStrataMaterialInfo.HasShadingModel(SSM_Hair))
 		{
-			FStrataOperator& Operator = Compiler->StrataCompilationRegisterOperator(STRATA_OPERATOR_BSDF_LEGACY, Compiler->StrataTreeStackGetPathUniqueId(), Parent, Compiler->StrataTreeStackGetParentPathUniqueId());
+			FStrataOperator& Operator = Compiler->StrataCompilationRegisterOperator(STRATA_OPERATOR_BSDF_LEGACY, Compiler->StrataTreeStackGetPathUniqueId(), this, Parent, Compiler->StrataTreeStackGetParentPathUniqueId());
 			Operator.BSDFType = STRATA_BSDF_TYPE_HAIR;
 			Operator.ThicknessIndex = ThicknessIndex;
 			return &Operator;
 		}
 		else if (ConvertedStrataMaterialInfo.HasShadingModel(SSM_Eye))
 		{
-			FStrataOperator& Operator = Compiler->StrataCompilationRegisterOperator(STRATA_OPERATOR_BSDF_LEGACY, Compiler->StrataTreeStackGetPathUniqueId(), Parent, Compiler->StrataTreeStackGetParentPathUniqueId());
+			FStrataOperator& Operator = Compiler->StrataCompilationRegisterOperator(STRATA_OPERATOR_BSDF_LEGACY, Compiler->StrataTreeStackGetPathUniqueId(), this, Parent, Compiler->StrataTreeStackGetParentPathUniqueId());
 			Operator.BSDFType = STRATA_BSDF_TYPE_EYE;
 			Operator.ThicknessIndex = ThicknessIndex;
 			return &Operator;
 		}
 		else if (ConvertedStrataMaterialInfo.HasShadingModel(SSM_SingleLayerWater))
 		{
-			FStrataOperator& Operator = Compiler->StrataCompilationRegisterOperator(STRATA_OPERATOR_BSDF_LEGACY, Compiler->StrataTreeStackGetPathUniqueId(), Parent, Compiler->StrataTreeStackGetParentPathUniqueId());
+			FStrataOperator& Operator = Compiler->StrataCompilationRegisterOperator(STRATA_OPERATOR_BSDF_LEGACY, Compiler->StrataTreeStackGetPathUniqueId(), this, Parent, Compiler->StrataTreeStackGetParentPathUniqueId());
 			Operator.BSDFType = STRATA_BSDF_TYPE_SINGLELAYERWATER;
 			Operator.ThicknessIndex = ThicknessIndex;
 			return &Operator;
@@ -23611,7 +23611,7 @@ void UMaterialExpressionStrataSlabBSDF::GatherStrataMaterialInfo(FStrataMaterial
 
 FStrataOperator* UMaterialExpressionStrataSlabBSDF::StrataGenerateMaterialTopologyTree(class FMaterialCompiler* Compiler, class UMaterialExpression* Parent, int32 OutputIndex)
 {
-	FStrataOperator& StrataOperator = Compiler->StrataCompilationRegisterOperator(STRATA_OPERATOR_BSDF, Compiler->StrataTreeStackGetPathUniqueId(), Parent, Compiler->StrataTreeStackGetParentPathUniqueId());
+	FStrataOperator& StrataOperator = Compiler->StrataCompilationRegisterOperator(STRATA_OPERATOR_BSDF, Compiler->StrataTreeStackGetPathUniqueId(), this, Parent, Compiler->StrataTreeStackGetParentPathUniqueId());
 	StrataOperator.BSDFType = STRATA_BSDF_TYPE_SLAB;
 	StrataOperator.bBSDFHasEdgeColor = HasEdgeColor();
 	StrataOperator.bBSDFHasFuzz = HasFuzz();
@@ -23843,7 +23843,7 @@ void UMaterialExpressionStrataSimpleClearCoatBSDF::GatherStrataMaterialInfo(FStr
 
 FStrataOperator* UMaterialExpressionStrataSimpleClearCoatBSDF::StrataGenerateMaterialTopologyTree(class FMaterialCompiler* Compiler, class UMaterialExpression* Parent, int32 OutputIndex)
 {
-	FStrataOperator& StrataOperator = Compiler->StrataCompilationRegisterOperator(STRATA_OPERATOR_BSDF, Compiler->StrataTreeStackGetPathUniqueId(), Parent, Compiler->StrataTreeStackGetParentPathUniqueId());
+	FStrataOperator& StrataOperator = Compiler->StrataCompilationRegisterOperator(STRATA_OPERATOR_BSDF, Compiler->StrataTreeStackGetPathUniqueId(), this, Parent, Compiler->StrataTreeStackGetParentPathUniqueId());
 	StrataOperator.BSDFType = STRATA_BSDF_TYPE_SLAB;
 	StrataOperator.bBSDFHasEdgeColor = false;
 	StrataOperator.bBSDFHasFuzz = false;
@@ -23930,7 +23930,7 @@ void UMaterialExpressionStrataVolumetricFogCloudBSDF::GatherStrataMaterialInfo(F
 
 FStrataOperator* UMaterialExpressionStrataVolumetricFogCloudBSDF::StrataGenerateMaterialTopologyTree(class FMaterialCompiler* Compiler, class UMaterialExpression* Parent, int32 OutputIndex)
 {
-	FStrataOperator& StrataOperator = Compiler->StrataCompilationRegisterOperator(STRATA_OPERATOR_BSDF, Compiler->StrataTreeStackGetPathUniqueId(), Parent, Compiler->StrataTreeStackGetParentPathUniqueId());
+	FStrataOperator& StrataOperator = Compiler->StrataCompilationRegisterOperator(STRATA_OPERATOR_BSDF, Compiler->StrataTreeStackGetPathUniqueId(), this, Parent, Compiler->StrataTreeStackGetParentPathUniqueId());
 	StrataOperator.BSDFType = STRATA_BSDF_TYPE_VOLUMETRICFOGCLOUD;
 	StrataOperator.ThicknessIndex = Compiler->StrataThicknessStackGetThicknessIndex();
 	return &StrataOperator;
@@ -23999,7 +23999,7 @@ void UMaterialExpressionStrataLightFunction::GatherStrataMaterialInfo(FStrataMat
 
 FStrataOperator* UMaterialExpressionStrataLightFunction::StrataGenerateMaterialTopologyTree(class FMaterialCompiler* Compiler, class UMaterialExpression* Parent, int32 OutputIndex)
 {
-	FStrataOperator& StrataOperator = Compiler->StrataCompilationRegisterOperator(STRATA_OPERATOR_BSDF, Compiler->StrataTreeStackGetPathUniqueId(), Parent, Compiler->StrataTreeStackGetParentPathUniqueId());
+	FStrataOperator& StrataOperator = Compiler->StrataCompilationRegisterOperator(STRATA_OPERATOR_BSDF, Compiler->StrataTreeStackGetPathUniqueId(), this, Parent, Compiler->StrataTreeStackGetParentPathUniqueId());
 	StrataOperator.BSDFType = STRATA_BSDF_TYPE_UNLIT;
 	StrataOperator.ThicknessIndex = Compiler->StrataThicknessStackGetThicknessIndex();
 	return &StrataOperator;
@@ -24074,7 +24074,7 @@ void UMaterialExpressionStrataPostProcess::GatherStrataMaterialInfo(FStrataMater
 
 FStrataOperator* UMaterialExpressionStrataPostProcess::StrataGenerateMaterialTopologyTree(class FMaterialCompiler* Compiler, class UMaterialExpression* Parent, int32 OutputIndex)
 {
-	FStrataOperator& StrataOperator = Compiler->StrataCompilationRegisterOperator(STRATA_OPERATOR_BSDF, Compiler->StrataTreeStackGetPathUniqueId(), Parent, Compiler->StrataTreeStackGetParentPathUniqueId());
+	FStrataOperator& StrataOperator = Compiler->StrataCompilationRegisterOperator(STRATA_OPERATOR_BSDF, Compiler->StrataTreeStackGetPathUniqueId(), this, Parent, Compiler->StrataTreeStackGetParentPathUniqueId());
 	StrataOperator.BSDFType = STRATA_BSDF_TYPE_UNLIT;
 	StrataOperator.ThicknessIndex = Compiler->StrataThicknessStackGetThicknessIndex();
 	return &StrataOperator;
@@ -24149,7 +24149,7 @@ void UMaterialExpressionStrataUI::GatherStrataMaterialInfo(FStrataMaterialInfo&
 
 FStrataOperator* UMaterialExpressionStrataUI::StrataGenerateMaterialTopologyTree(class FMaterialCompiler* Compiler, class UMaterialExpression* Parent, int32 OutputIndex)
 {
-	FStrataOperator& StrataOperator = Compiler->StrataCompilationRegisterOperator(STRATA_OPERATOR_BSDF, Compiler->StrataTreeStackGetPathUniqueId(), Parent, Compiler->StrataTreeStackGetParentPathUniqueId());
+	FStrataOperator& StrataOperator = Compiler->StrataCompilationRegisterOperator(STRATA_OPERATOR_BSDF, Compiler->StrataTreeStackGetPathUniqueId(), this, Parent, Compiler->StrataTreeStackGetParentPathUniqueId());
 	StrataOperator.BSDFType = STRATA_BSDF_TYPE_UNLIT;
 	StrataOperator.ThicknessIndex = Compiler->StrataThicknessStackGetThicknessIndex();
 	return &StrataOperator;
@@ -24258,7 +24258,7 @@ void UMaterialExpressionStrataConvertToDecal::GatherStrataMaterialInfo(FStrataMa
 FStrataOperator* UMaterialExpressionStrataConvertToDecal::StrataGenerateMaterialTopologyTree(class FMaterialCompiler* Compiler, class UMaterialExpression* Parent, int32 OutputIndex)
 {
 	const bool bUseParameterBlending = true;
-	FStrataOperator& StrataOperator = Compiler->StrataCompilationRegisterOperator(STRATA_OPERATOR_WEIGHT, Compiler->StrataTreeStackGetPathUniqueId(), Parent, Compiler->StrataTreeStackGetParentPathUniqueId(), bUseParameterBlending);
+	FStrataOperator& StrataOperator = Compiler->StrataCompilationRegisterOperator(STRATA_OPERATOR_WEIGHT, Compiler->StrataTreeStackGetPathUniqueId(), this, Parent, Compiler->StrataTreeStackGetParentPathUniqueId(), bUseParameterBlending);
 	if (Compiler->GetStrataTreeOutOfStackDepthOccurred())
 	{
 		return &StrataOperator; // Out ot stack space, return now to fail the compilation
@@ -24349,7 +24349,7 @@ void UMaterialExpressionStrataUnlitBSDF::GatherStrataMaterialInfo(FStrataMateria
 
 FStrataOperator* UMaterialExpressionStrataUnlitBSDF::StrataGenerateMaterialTopologyTree(class FMaterialCompiler* Compiler, class UMaterialExpression* Parent, int32 OutputIndex)
 {
-	FStrataOperator& StrataOperator = Compiler->StrataCompilationRegisterOperator(STRATA_OPERATOR_BSDF, Compiler->StrataTreeStackGetPathUniqueId(), Parent, Compiler->StrataTreeStackGetParentPathUniqueId());
+	FStrataOperator& StrataOperator = Compiler->StrataCompilationRegisterOperator(STRATA_OPERATOR_BSDF, Compiler->StrataTreeStackGetPathUniqueId(), this, Parent, Compiler->StrataTreeStackGetParentPathUniqueId());
 	StrataOperator.BSDFType = STRATA_BSDF_TYPE_UNLIT;
 	StrataOperator.ThicknessIndex = Compiler->StrataThicknessStackGetThicknessIndex();
 	return &StrataOperator;
@@ -24465,7 +24465,7 @@ void UMaterialExpressionStrataHairBSDF::GatherStrataMaterialInfo(FStrataMaterial
 
 FStrataOperator* UMaterialExpressionStrataHairBSDF::StrataGenerateMaterialTopologyTree(class FMaterialCompiler* Compiler, class UMaterialExpression* Parent, int32 OutputIndex)
 {
-	FStrataOperator& StrataOperator = Compiler->StrataCompilationRegisterOperator(STRATA_OPERATOR_BSDF, Compiler->StrataTreeStackGetPathUniqueId(), Parent, Compiler->StrataTreeStackGetParentPathUniqueId());
+	FStrataOperator& StrataOperator = Compiler->StrataCompilationRegisterOperator(STRATA_OPERATOR_BSDF, Compiler->StrataTreeStackGetPathUniqueId(), this, Parent, Compiler->StrataTreeStackGetParentPathUniqueId());
 	StrataOperator.BSDFType = STRATA_BSDF_TYPE_HAIR;
 	StrataOperator.ThicknessIndex = Compiler->StrataThicknessStackGetThicknessIndex();
 	return &StrataOperator;
@@ -24579,7 +24579,7 @@ void UMaterialExpressionStrataEyeBSDF::GatherStrataMaterialInfo(FStrataMaterialI
 
 FStrataOperator* UMaterialExpressionStrataEyeBSDF::StrataGenerateMaterialTopologyTree(class FMaterialCompiler* Compiler, class UMaterialExpression* Parent, int32 OutputIndex)
 {
-	FStrataOperator& StrataOperator = Compiler->StrataCompilationRegisterOperator(STRATA_OPERATOR_BSDF, Compiler->StrataTreeStackGetPathUniqueId(), Parent, Compiler->StrataTreeStackGetParentPathUniqueId());
+	FStrataOperator& StrataOperator = Compiler->StrataCompilationRegisterOperator(STRATA_OPERATOR_BSDF, Compiler->StrataTreeStackGetPathUniqueId(), this, Parent, Compiler->StrataTreeStackGetParentPathUniqueId());
 	StrataOperator.BSDFType = STRATA_BSDF_TYPE_EYE;
 	StrataOperator.ThicknessIndex = Compiler->StrataThicknessStackGetThicknessIndex();
 	return &StrataOperator;
@@ -24710,7 +24710,7 @@ void UMaterialExpressionStrataSingleLayerWaterBSDF::GatherStrataMaterialInfo(FSt
 
 FStrataOperator* UMaterialExpressionStrataSingleLayerWaterBSDF::StrataGenerateMaterialTopologyTree(class FMaterialCompiler* Compiler, class UMaterialExpression* Parent, int32 OutputIndex)
 {
-	FStrataOperator& StrataOperator = Compiler->StrataCompilationRegisterOperator(STRATA_OPERATOR_BSDF, Compiler->StrataTreeStackGetPathUniqueId(), Parent, Compiler->StrataTreeStackGetParentPathUniqueId());
+	FStrataOperator& StrataOperator = Compiler->StrataCompilationRegisterOperator(STRATA_OPERATOR_BSDF, Compiler->StrataTreeStackGetPathUniqueId(), this, Parent, Compiler->StrataTreeStackGetParentPathUniqueId());
 	StrataOperator.BSDFType = STRATA_BSDF_TYPE_SINGLELAYERWATER;
 	StrataOperator.ThicknessIndex = Compiler->StrataThicknessStackGetThicknessIndex();
 	return &StrataOperator;
@@ -24864,7 +24864,7 @@ void UMaterialExpressionStrataHorizontalMixing::GatherStrataMaterialInfo(FStrata
 
 FStrataOperator* UMaterialExpressionStrataHorizontalMixing::StrataGenerateMaterialTopologyTree(class FMaterialCompiler* Compiler, class UMaterialExpression* Parent, int32 OutputIndex)
 {
-	FStrataOperator& StrataOperator = Compiler->StrataCompilationRegisterOperator(STRATA_OPERATOR_HORIZONTAL, Compiler->StrataTreeStackGetPathUniqueId(), Parent, Compiler->StrataTreeStackGetParentPathUniqueId(), bUseParameterBlending);
+	FStrataOperator& StrataOperator = Compiler->StrataCompilationRegisterOperator(STRATA_OPERATOR_HORIZONTAL, Compiler->StrataTreeStackGetPathUniqueId(), this, Parent, Compiler->StrataTreeStackGetParentPathUniqueId(), bUseParameterBlending);
 	if (Compiler->GetStrataTreeOutOfStackDepthOccurred())
 	{
 		return &StrataOperator; // Out ot stack space, return now to fail the compilation
@@ -25062,7 +25062,7 @@ void UMaterialExpressionStrataVerticalLayering::GatherStrataMaterialInfo(FStrata
 
 FStrataOperator* UMaterialExpressionStrataVerticalLayering::StrataGenerateMaterialTopologyTree(class FMaterialCompiler* Compiler, class UMaterialExpression* Parent, int32 OutputIndex)
 {
-	FStrataOperator& StrataOperator = Compiler->StrataCompilationRegisterOperator(STRATA_OPERATOR_VERTICAL, Compiler->StrataTreeStackGetPathUniqueId(), Parent, Compiler->StrataTreeStackGetParentPathUniqueId(), bUseParameterBlending);
+	FStrataOperator& StrataOperator = Compiler->StrataCompilationRegisterOperator(STRATA_OPERATOR_VERTICAL, Compiler->StrataTreeStackGetPathUniqueId(), this, Parent, Compiler->StrataTreeStackGetParentPathUniqueId(), bUseParameterBlending);
 	if (Compiler->GetStrataTreeOutOfStackDepthOccurred())
 	{
 		return &StrataOperator; // Out ot stack space, return now to fail the compilation
@@ -25239,7 +25239,7 @@ void UMaterialExpressionStrataAdd::GatherStrataMaterialInfo(FStrataMaterialInfo&
 
 FStrataOperator* UMaterialExpressionStrataAdd::StrataGenerateMaterialTopologyTree(class FMaterialCompiler* Compiler, class UMaterialExpression* Parent, int32 OutputIndex)
 {
-	FStrataOperator& StrataOperator = Compiler->StrataCompilationRegisterOperator(STRATA_OPERATOR_ADD, Compiler->StrataTreeStackGetPathUniqueId(), Parent, Compiler->StrataTreeStackGetParentPathUniqueId(), bUseParameterBlending);
+	FStrataOperator& StrataOperator = Compiler->StrataCompilationRegisterOperator(STRATA_OPERATOR_ADD, Compiler->StrataTreeStackGetPathUniqueId(), this, Parent, Compiler->StrataTreeStackGetParentPathUniqueId(), bUseParameterBlending);
 	if (Compiler->GetStrataTreeOutOfStackDepthOccurred())
 	{
 		return &StrataOperator; // Out ot stack space, return now to fail the compilation
@@ -25376,7 +25376,7 @@ void UMaterialExpressionStrataWeight::GatherStrataMaterialInfo(FStrataMaterialIn
 
 FStrataOperator* UMaterialExpressionStrataWeight::StrataGenerateMaterialTopologyTree(class FMaterialCompiler* Compiler, class UMaterialExpression* Parent, int32 OutputIndex)
 {
-	FStrataOperator& StrataOperator = Compiler->StrataCompilationRegisterOperator(STRATA_OPERATOR_WEIGHT, Compiler->StrataTreeStackGetPathUniqueId(), Parent, Compiler->StrataTreeStackGetParentPathUniqueId());
+	FStrataOperator& StrataOperator = Compiler->StrataCompilationRegisterOperator(STRATA_OPERATOR_WEIGHT, Compiler->StrataTreeStackGetPathUniqueId(), this, Parent, Compiler->StrataTreeStackGetParentPathUniqueId());
 	if (Compiler->GetStrataTreeOutOfStackDepthOccurred())
 	{
 		return &StrataOperator; // Out ot stack space, return now to fail the compilation
diff --git a/Engine/Source/Runtime/Engine/Public/MaterialCompiler.h b/Engine/Source/Runtime/Engine/Public/MaterialCompiler.h
--- a/Engine/Source/Runtime/Engine/Public/MaterialCompiler.h
+++ b/Engine/Source/Runtime/Engine/Public/MaterialCompiler.h
@@ -634,7 +634,7 @@ public:
 	/**
 	 * Register an operator of the tree representation the Strata material and its topology.
 	 */
-	virtual FStrataOperator& StrataCompilationRegisterOperator(int32 OperatorType, FGuid StrataExpressionGuid, UMaterialExpression* Parent, FGuid StrataParentExpressionGuid, bool bUseParameterBlending = false) = 0;
+	virtual FStrataOperator& StrataCompilationRegisterOperator(int32 OperatorType, FGuid StrataExpressionGuid, UMaterialExpression* Child, UMaterialExpression* Parent, FGuid StrataParentExpressionGuid, bool bUseParameterBlending = false) = 0;
 	/**
 	 * Return the operator information for a given expression.
 	 */
@@ -1369,9 +1369,9 @@ public:
 		Compiler->StrataThicknessStackPop();
 	}
 
-	virtual FStrataOperator& StrataCompilationRegisterOperator(int32 OperatorType, FGuid StrataExpressionGuid, UMaterialExpression* Parent, FGuid StrataParentExpressionGuid, bool bUseParameterBlending = false) override
+	virtual FStrataOperator& StrataCompilationRegisterOperator(int32 OperatorType, FGuid StrataExpressionGuid, UMaterialExpression* Child, UMaterialExpression* Parent, FGuid StrataParentExpressionGuid, bool bUseParameterBlending = false) override
 	{
-		return Compiler->StrataCompilationRegisterOperator(OperatorType, StrataExpressionGuid, Parent, StrataParentExpressionGuid, bUseParameterBlending);
+		return Compiler->StrataCompilationRegisterOperator(OperatorType, StrataExpressionGuid, Child, Parent, StrataParentExpressionGuid, bUseParameterBlending);
 	}
 
 	virtual FStrataOperator& StrataCompilationGetOperator(FGuid StrataExpressionGuid) override
diff --git a/Engine/Source/Runtime/Engine/Public/Rendering/StrataMaterialShared.cpp b/Engine/Source/Runtime/Engine/Public/Rendering/StrataMaterialShared.cpp
--- a/Engine/Source/Runtime/Engine/Public/Rendering/StrataMaterialShared.cpp
+++ b/Engine/Source/Runtime/Engine/Public/Rendering/StrataMaterialShared.cpp
@@ -61,6 +61,7 @@ FStrataOperator::FStrataOperator()
 	bIsBottom = false;
 	bUseParameterBlending = false;
 	bRootOfParameterBlendingSubTree = false;
+	MaterialExpressionGuid = FGuid();
 #endif
 }
 
diff --git a/Engine/Source/Runtime/Engine/Public/Rendering/StrataMaterialShared.h b/Engine/Source/Runtime/Engine/Public/Rendering/StrataMaterialShared.h
--- a/Engine/Source/Runtime/Engine/Public/Rendering/StrataMaterialShared.h
+++ b/Engine/Source/Runtime/Engine/Public/Rendering/StrataMaterialShared.h
@@ -64,7 +64,8 @@ public:
 	LAYOUT_FIELD_EDITORONLY(uint8, bIsBottom);
 
 	LAYOUT_FIELD_EDITORONLY(uint8, bUseParameterBlending);			// True when part of a sub tree where parameter blending is in use
-	LAYOUT_FIELD_EDITORONLY(uint8, bRootOfParameterBlendingSubTree);	// True when the root of a sub tree where parameter blending is in use. Only this node will register a BSDF
+	LAYOUT_FIELD_EDITORONLY(uint8, bRootOfParameterBlendingSubTree);// True when the root of a sub tree where parameter blending is in use. Only this node will register a BSDF
+	LAYOUT_FIELD_EDITORONLY(FGuid, MaterialExpressionGuid);			// Material expression Guid for mapping between UMaterialExpression and FStrataOperator
 
 	void CombineFlagsForParameterBlending(FStrataOperator& A, FStrataOperator& B);
 
