Fix linking error.

diff --git a/Engine/Source/Runtime/Engine/Classes/Engine/Engine.h b/Engine/Source/Runtime/Engine/Classes/Engine/Engine.h
--- a/Engine/Source/Runtime/Engine/Classes/Engine/Engine.h
+++ b/Engine/Source/Runtime/Engine/Classes/Engine/Engine.h
@@ -1934,7 +1934,7 @@ public:
 	/** Delay loading the SimpleVolume texture until it is needed by the renderer.
 	* This texture is not going to be streamed to be available right away.
 	*/
-	void LoadSimpleVolumeTextures();
+	ENGINE_API void LoadSimpleVolumeTextures();
 
 private:
 	#if WITH_DYNAMIC_RESOLUTION
