Substrate - added compiler error when requisting to omany bits.

diff --git a/Engine/Shaders/Private/Strata/StrataTile.ush b/Engine/Shaders/Private/Strata/StrataTile.ush
--- a/Engine/Shaders/Private/Strata/StrataTile.ush
+++ b/Engine/Shaders/Private/Strata/StrataTile.ush
@@ -89,11 +89,6 @@ FStrataBSDFTile UnpackBSDFTile(uint In)
 ////////////////////////////////////////////////////////////////////////////////////////////////////////////
 // Strata BSDF offset
 
-// BSDF offsets are packed into 32bits to max. BSDF is < STRATA_MAX_BSDF_COUNT, which could leads to some issue.
-#define STRATA_MAX_BSDF_COUNT_FOR_BDSFOFFSET 8u
-#define STRATA_BSDF_OFFSET_BIT_COUNT 4u
-#define STRATA_BSDF_OFFSET_BIT_MASK  0xF
-
 // Store the material slice index offset for each BSDF, for pixels containing several BSDFs
 struct FStrataBSDFOffset
 {
diff --git a/Engine/Shaders/Shared/StrataDefinitions.h b/Engine/Shaders/Shared/StrataDefinitions.h
--- a/Engine/Shaders/Shared/StrataDefinitions.h
+++ b/Engine/Shaders/Shared/StrataDefinitions.h
@@ -10,9 +10,13 @@
 // Change this to force recompilation of all strata dependent shaders (use https://www.random.org/cgi-bin/randbyte?nbytes=4&format=h)
 #define STRATA_SHADER_VERSION 0x56F06029 
 
+// BSDF offsets are packed into 32bits to max. BSDF is < STRATA_MAX_BSDF_COUNT, which could leads to some issue.
+#define STRATA_MAX_BSDF_COUNT_FOR_BDSFOFFSET	8u
+#define STRATA_BSDF_OFFSET_BIT_COUNT			4u
+#define STRATA_BSDF_OFFSET_BIT_MASK				0xF
 
-
-#define STRATA_MAX_BSDF_COUNT				15
+// We can only ever use STRATA_MAX_BSDF_COUNT_FOR_BDSFOFFSET BSDFs for Lumen, se we use that as a global BSDF count limit today.
+#define STRATA_MAX_BSDF_COUNT				STRATA_MAX_BSDF_COUNT_FOR_BDSFOFFSET
 #define STRATA_MAX_OPERATOR_COUNT			15
 
 // It this is changed, STATE_BIT_COUNT_SHAREDLOCALBASESID and HEADER_BIT_COUNT_SHAREDLOCALBASES_COUNT also needs to be updated
diff --git a/Engine/Source/Runtime/Engine/Private/Materials/HLSLMaterialTranslator.cpp b/Engine/Source/Runtime/Engine/Private/Materials/HLSLMaterialTranslator.cpp
--- a/Engine/Source/Runtime/Engine/Private/Materials/HLSLMaterialTranslator.cpp
+++ b/Engine/Source/Runtime/Engine/Private/Materials/HLSLMaterialTranslator.cpp
@@ -11868,6 +11868,7 @@ bool FHLSLMaterialTranslator::FStrataCompilationContext::StrataGenerateDerivedMa
 
 			// 2. Process the list of BSDFs for worst case memory usage and count operators.
 			uint32 OperatorCount = 0;
+			const uint32 BSDFMaxByteCount = uint32(FMath::Pow(2.0f, float(STRATA_BSDF_OFFSET_BIT_COUNT))) * sizeof(uint32);
 			for (auto& It : StrataMaterialExpressionRegisteredOperators)
 			{
 				if (It.IsDiscarded())
@@ -11879,6 +11880,7 @@ bool FHLSLMaterialTranslator::FStrataCompilationContext::StrataGenerateDerivedMa
 				// Be aware that BSDFs also count as Operators when they are promoted from parameter blending!
 				OperatorCount++;
 
+				const uint32 PreStrataMaterialRequestedSizeByte = StrataMaterialRequestedSizeByte;
 				switch (It.OperatorType)
 				{
 				case STRATA_OPERATOR_BSDF:
@@ -11996,6 +11998,13 @@ bool FHLSLMaterialTranslator::FStrataCompilationContext::StrataGenerateDerivedMa
 					break;
 				} // case STRATA_OPERATOR_BSDF
 				} // switch (It.OperatorType)
+
+				const uint32 BSDFRequestedSizeByte = StrataMaterialRequestedSizeByte - PreStrataMaterialRequestedSizeByte;
+				if (BSDFRequestedSizeByte > BSDFMaxByteCount)
+				{
+					Compiler->Errorf(TEXT("A BSDF is requesting more bytes than our BSDF offset system can handle (%d/%d bytes). Notify your rendering engineer. Material %s (asset: %s).\r\n"), BSDFRequestedSizeByte, BSDFMaxByteCount, *CompilerMaterial->GetDebugName(), *CompilerMaterial->GetAssetPath().ToString(), OperatorCount, STRATA_MAX_OPERATOR_COUNT);
+					return false;
+				}
 			}
 
 			if (OperatorCount > STRATA_MAX_OPERATOR_COUNT)
