UAssetManager::PreloadPrimaryAssets now reports names of invalid/malformed primary assets (warning) and still ensures on the first failure. If there was a partial success these are reported at the end (if there is room).
#rnx

diff --git a/Engine/Source/Runtime/Engine/Private/AssetManager.cpp b/Engine/Source/Runtime/Engine/Private/AssetManager.cpp
--- a/Engine/Source/Runtime/Engine/Private/AssetManager.cpp
+++ b/Engine/Source/Runtime/Engine/Private/AssetManager.cpp
@@ -1936,38 +1936,63 @@ bool UAssetManager::GetPrimaryAssetLoadSet(TSet<FSoftObjectPath>& OutAssetLoadSe
 TSharedPtr<FStreamableHandle> UAssetManager::PreloadPrimaryAssets(const TArray<FPrimaryAssetId>& AssetsToLoad, const TArray<FName>& LoadBundles, bool bLoadRecursive, FStreamableDelegate DelegateToCall, TAsyncLoadPriority Priority)
 {
 	TSet<FSoftObjectPath> PathsToLoad;
-	TStringBuilder<256> DebugName;
+	TStringBuilder<256> DebugValid;
+	TStringBuilder<256> DebugInvalid;
 	TSharedPtr<FStreamableHandle> ReturnHandle;
+	const int MaxDebugWarningLen = 1024;
+	const int MaxEnsureMessageLen = 256;
+	bool ExeededMaxLength = false;
 
 	for (const FPrimaryAssetId& PrimaryAssetId : AssetsToLoad)
 	{
-		if (GetPrimaryAssetLoadSet(PathsToLoad, PrimaryAssetId, LoadBundles, bLoadRecursive) && DebugName.Len() < 1024)
+		if (GetPrimaryAssetLoadSet(PathsToLoad, PrimaryAssetId, LoadBundles, bLoadRecursive))
 		{
-			if (DebugName.Len() == 0)
+			if (DebugValid.Len() < MaxDebugWarningLen)
 			{
-				DebugName << FStreamableHandle::HandleDebugName_Preloading;
-				DebugName << TEXT(" ");
+				DebugValid << (DebugValid.Len() > 0 ? TEXT(", ") : TEXT("")) << PrimaryAssetId.ToString();
 			}
 			else
 			{
-				DebugName << TEXT(", ");
+				ExeededMaxLength = true;
+			}
+		}
+		else
+		{
+			if (DebugInvalid.Len() < MaxDebugWarningLen)
+			{
+				DebugInvalid << (DebugInvalid.Len() > 0? TEXT(", ") : TEXT("")) << PrimaryAssetId.ToString();
+			}
+			else
+			{
+				ExeededMaxLength = true;
 			}
-			DebugName << PrimaryAssetId.ToString();
 		}
-	}
-	
-	if (DebugName.Len() >= 1024)
-	{
-		DebugName.ReplaceAt(DebugName.Len() - 3, 3, TEXTVIEW("..."));
 	}
 
-	ReturnHandle = LoadAssetList(PathsToLoad.Array(), MoveTemp(DelegateToCall), Priority, FString(*DebugName));
+	ReturnHandle = LoadAssetList(PathsToLoad.Array(), MoveTemp(DelegateToCall), Priority, FString(*DebugValid));
 
-	if (!ensureMsgf(ReturnHandle.IsValid(), TEXT("Requested preload of Primary Asset with no referenced assets! DebugName:%s"), *DebugName))
+	if (DebugInvalid.Len() > 0)
 	{
-		return nullptr;
+		if (DebugValid.Len() > 0)
+		{
+			DebugInvalid << TEXT(" (Valid: ") << DebugValid << TEXT(")");
+		}
+
+		DebugInvalid << (ExeededMaxLength ? TEXT("...") : TEXT(""));
+
+		UE_LOG(LogAssetManager, Warning, TEXT("Requested preload of some Primary Assets failed: %s"), *DebugInvalid);
+
+		// Trim for shorter ensure message
+		int TrimEndLen = FMath::Max(0, DebugInvalid.Len() - MaxEnsureMessageLen);
+		if (TrimEndLen > 0)
+		{
+			DebugInvalid.RemoveSuffix(TrimEndLen);
+			DebugInvalid << TEXT("...");
+		}
 	}
 
+	ensureMsgf(ReturnHandle.IsValid(), TEXT("Requested preload of Primary Assets failed: %s"), *DebugInvalid);
+
 	return ReturnHandle;
 }
 
