Fix for pointer invalidation in SVT streaming manager. FStreamingInfo is now stored as a UniquePtr in the map so that pointers should now always be stable.

#rnx

diff --git a/Engine/Source/Runtime/Engine/Private/SparseVolumeTexture/SparseVolumeTextureStreamingManager.cpp b/Engine/Source/Runtime/Engine/Private/SparseVolumeTexture/SparseVolumeTextureStreamingManager.cpp
--- a/Engine/Source/Runtime/Engine/Private/SparseVolumeTexture/SparseVolumeTextureStreamingManager.cpp
+++ b/Engine/Source/Runtime/Engine/Private/SparseVolumeTexture/SparseVolumeTextureStreamingManager.cpp
@@ -888,7 +888,7 @@ void FStreamingManager::Request(UStreamableSparseVolumeTexture* SparseVolumeText
 		return;
 	}
 
-	FStreamingInfo* SVTInfo = StreamingInfo.Find(SparseVolumeTexture);
+	FStreamingInfo* SVTInfo = FindStreamingInfo(SparseVolumeTexture);
 	if (SVTInfo)
 	{
 		const int32 NumFrames = SVTInfo->PerFrameInfo.Num();
@@ -993,13 +993,13 @@ void FStreamingManager::BeginAsyncUpdate(FRDGBuilder& GraphBuilder, bool bBlocki
 		for (auto& Pair : StreamingInfo)
 		{
 			MipLevelsToFree.Reset();
-			FStreamingInfo& SVTInfo = Pair.Value;
-			const int32 NumFrames = SVTInfo.PerFrameInfo.Num();
-			const int32 NumMipLevelsGlobal = SVTInfo.NumMipLevelsGlobal;
+			FStreamingInfo* SVTInfo = Pair.Value.Get();
+			const int32 NumFrames = SVTInfo->PerFrameInfo.Num();
+			const int32 NumMipLevelsGlobal = SVTInfo->NumMipLevelsGlobal;
 			
 			for (int32 MipLevel = 0; MipLevel < NumMipLevelsGlobal; ++MipLevel)
 			{
-				for (auto& Node : SVTInfo.PerMipLRULists[MipLevel])
+				for (auto& Node : SVTInfo->PerMipLRULists[MipLevel])
 				{
 					MipLevelsToFree.Add(&Node);
 				}
@@ -1021,12 +1021,12 @@ void FStreamingManager::BeginAsyncUpdate(FRDGBuilder& GraphBuilder, bool bBlocki
 			double MaxFrameMiB = -FLT_MAX;
 			double SumFrameMiB = 0.0;
 
-			const int32 NumFrames = Pair.Value.PerFrameInfo.Num();
+			const int32 NumFrames = Pair.Value->PerFrameInfo.Num();
 			UE_LOG(LogSparseVolumeTextureStreamingManager, Display, TEXT("Memory stats for SVT '%p': Each mip level is printed as a tuple of [PageTable Size | VoxelData Size | Total]"), Pair.Key);
 
 			for (int32 FrameIdx = 0; FrameIdx < NumFrames; ++FrameIdx)
 			{
-				const FFrameInfo& FrameInfo = Pair.Value.PerFrameInfo[FrameIdx];
+				const FFrameInfo& FrameInfo = Pair.Value->PerFrameInfo[FrameIdx];
 				FString Str;
 				int32 TotalSize = 0;
 				for (const auto& SInfo : FrameInfo.Resources->MipLevelStreamingInfo)
@@ -1064,7 +1064,7 @@ void FStreamingManager::BeginAsyncUpdate(FRDGBuilder& GraphBuilder, bool bBlocki
 			const int32 PendingMipLevelIndex = (StartPendingMipLevelIndex + i) % MaxPendingMipLevels;
 			FPendingMipLevel& PendingMipLevel = PendingMipLevels[PendingMipLevelIndex];
 
-			FStreamingInfo* SVTInfo = StreamingInfo.Find(PendingMipLevel.SparseVolumeTexture);
+			FStreamingInfo* SVTInfo = FindStreamingInfo(PendingMipLevel.SparseVolumeTexture);
 			if (!SVTInfo || (SVTInfo->PerFrameInfo[PendingMipLevel.FrameIndex].LowestRequestedMipLevel > PendingMipLevel.MipLevelIndex))
 			{
 				continue; // Skip mip level install. SVT no longer exists or mip level was "streamed out" before it was even installed in the first place.
@@ -1145,13 +1145,13 @@ void FStreamingManager::EndAsyncUpdate(FRDGBuilder& GraphBuilder)
 	}
 
 	// Update streaming info buffers
-	for (FStreamingInfo* SVTInfo : InvalidatedStreamingInfos)
+	for (FStreamingInfo* SVTInfo : SVTsWithInvalidatedStreamingInfoBuffer)
 	{
 #if DO_CHECK
 		bool bSVTInfoExists = false;
 		for (const auto& Pair : StreamingInfo)
 		{
-			if (&Pair.Value == SVTInfo)
+			if (Pair.Value.Get() == SVTInfo)
 			{
 				bSVTInfoExists = true;
 				break;
@@ -1166,7 +1166,7 @@ void FStreamingManager::EndAsyncUpdate(FRDGBuilder& GraphBuilder)
 			StreamingInfoBufferUpdater->Add(SVTInfo->StreamingInfoBuffer, FrameIndex, SVTInfo->PerFrameInfo[FrameIndex].LowestResidentMipLevel);
 		}
 	}
-	InvalidatedStreamingInfos.Reset();
+	SVTsWithInvalidatedStreamingInfoBuffer.Reset();
 	StreamingInfoBufferUpdater->Apply(GraphBuilder);
 
 	PageTableUpdater->Apply(GraphBuilder);
@@ -1183,10 +1183,10 @@ void FStreamingManager::EndAsyncUpdate(FRDGBuilder& GraphBuilder)
 #if SVT_STREAMING_LOG_VERBOSE
 		FString ResidentMipLevelsStr = TEXT("");
 #endif
-		const int32 NumFrames = Pair.Value.PerFrameInfo.Num();
+		const int32 NumFrames = Pair.Value->PerFrameInfo.Num();
 		for (int32 FrameIndex = 0; FrameIndex < NumFrames; ++FrameIndex)
 		{
-			const auto& FrameInfo = Pair.Value.PerFrameInfo[FrameIndex];
+			const auto& FrameInfo = Pair.Value->PerFrameInfo[FrameIndex];
 			check(FrameInfo.LowestResidentMipLevel <= (FrameInfo.NumMipLevels - 1));
 			check(FrameInfo.LowestRequestedMipLevel <= FrameInfo.LowestResidentMipLevel);
 			check(FrameInfo.TextureRenderResources->GetNumLogicalMipLevels() == FrameInfo.NumMipLevels);
@@ -1215,7 +1215,7 @@ void FStreamingManager::AddInternal(FRDGBuilder& GraphBuilder, FNewSparseVolumeT
 
 	const int32 NumFrames = NewSVTInfo.FrameInfo.Num();
 
-	FStreamingInfo& SVTInfo = StreamingInfo.Add(NewSVTInfo.SVT);
+	FStreamingInfo& SVTInfo = *StreamingInfo.Emplace(NewSVTInfo.SVT, MakeUnique<FStreamingInfo>());
 	SVTInfo.FormatA = NewSVTInfo.FormatA;
 	SVTInfo.FormatB = NewSVTInfo.FormatB;
 	SVTInfo.FallbackValueA = NewSVTInfo.FallbackValueA;
@@ -1397,7 +1397,7 @@ void FStreamingManager::AddInternal(FRDGBuilder& GraphBuilder, FNewSparseVolumeT
 		RootTileUploader.ResourceUploadTo(GraphBuilder, SVTInfo.TileDataTexture->TileDataTextureARHIRef, SVTInfo.TileDataTexture->TileDataTextureBRHIRef, SVTInfo.FallbackValueA, SVTInfo.FallbackValueB);
 	}
 
-	InvalidatedStreamingInfos.Add(&SVTInfo);
+	SVTsWithInvalidatedStreamingInfoBuffer.Add(&SVTInfo);
 
 	// Add requests for all mips the first frame. This is necessary for cases where UAnimatedSparseVolumeTexture or UStaticSparseVolumeTexture
 	// are directly bound to the material without getting a specific frame through USparseVolumeTextureFrame::GetFrameAndIssueStreamingRequest().
@@ -1417,7 +1417,7 @@ void FStreamingManager::RemoveInternal(UStreamableSparseVolumeTexture* SparseVol
 {
 	check(IsInRenderingThread());
 	check(!AsyncState.bUpdateActive);
-	FStreamingInfo* SVTInfo = StreamingInfo.Find(SparseVolumeTexture);
+	FStreamingInfo* SVTInfo = FindStreamingInfo(SparseVolumeTexture);
 	if (SVTInfo)
 	{
 		// Remove any requests for this SVT
@@ -1454,6 +1454,8 @@ void FStreamingManager::RemoveInternal(UStreamableSparseVolumeTexture* SparseVol
 			SVTInfo->TileDataTexture.Reset();
 		}
 
+		SVTsWithInvalidatedStreamingInfoBuffer.Remove(SVTInfo);
+
 		StreamingInfo.Remove(SparseVolumeTexture);
 	}
 }
@@ -1487,7 +1489,7 @@ void FStreamingManager::AddParentRequests()
 	ParentRequestsToAdd.Reset();
 	for (const auto& Request : RequestsHashTable)
 	{
-		FStreamingInfo* SVTInfo = StreamingInfo.Find(Request.Key.SVT);
+		FStreamingInfo* SVTInfo = FindStreamingInfo(Request.Key.SVT);
 		check(SVTInfo);
 		const int32 NumStreamableMipLevels = SVTInfo->PerFrameInfo[Request.Key.FrameIndex].NumMipLevels - 1;
 		uint32 Priority = Request.Value == FStreamingRequest::BlockingPriority ? FStreamingRequest::BlockingPriority : (Request.Value + 1);
@@ -1530,7 +1532,7 @@ void FStreamingManager::SelectHighestPriorityRequestsAndUpdateLRU(int32 MaxSelec
 
 		for (const auto& Request : RequestsHashTable)
 		{
-			FStreamingInfo* SVTInfo = StreamingInfo.Find(Request.Key.SVT);
+			FStreamingInfo* SVTInfo = FindStreamingInfo(Request.Key.SVT);
 			check(SVTInfo);
 
 			// Discard invalid requests: frame index out of bounds, mip level index out of bounds (or root mip level) and mip levels without any data.
@@ -1587,7 +1589,7 @@ void FStreamingManager::SelectHighestPriorityRequestsAndUpdateLRU(int32 MaxSelec
 			FStreamingRequest SelectedRequest;
 			PrioritizedRequestsHeap.HeapPop(SelectedRequest, PriorityPredicate, false /*bAllowShrinking*/);
 
-			FStreamingInfo* SVTInfo = StreamingInfo.Find(SelectedRequest.Key.SVT);
+			FStreamingInfo* SVTInfo = FindStreamingInfo(SelectedRequest.Key.SVT);
 			if (SVTInfo)
 			{
 				check(SelectedRequest.Key.FrameIndex < SVTInfo->PerFrameInfo.Num());
@@ -1624,7 +1626,7 @@ void FStreamingManager::IssueRequests(int32 MaxSelectedRequests)
 	for (const FStreamingRequest& SelectedRequest : SelectedRequests)
 	{
 		const FMipLevelKey& SelectedKey = SelectedRequest.Key;
-		FStreamingInfo* SVTInfo = StreamingInfo.Find(SelectedKey.SVT);
+		FStreamingInfo* SVTInfo = FindStreamingInfo(SelectedKey.SVT);
 		check(SVTInfo);
 		check(SVTInfo->PerFrameInfo.Num() > SelectedKey.FrameIndex && SelectedKey.FrameIndex >= 0);
 		check(SVTInfo->PerFrameInfo[SelectedKey.FrameIndex].LowestRequestedMipLevel > SelectedKey.MipLevelIndex);
@@ -1681,7 +1683,7 @@ void FStreamingManager::IssueRequests(int32 MaxSelectedRequests)
 			// Free mip levels
 			for (FLRUNode* MipLevelToFree : MipLevelsToFree)
 			{
-				StreamOutMipLevel(*SVTInfo, MipLevelToFree);
+				StreamOutMipLevel(SVTInfo, MipLevelToFree);
 			}
 
 			// Couldn't free enough tiles, so skip this mip level
@@ -1817,12 +1819,12 @@ void FStreamingManager::IssueRequests(int32 MaxSelectedRequests)
 	}
 }
 
-void FStreamingManager::StreamOutMipLevel(FStreamingInfo& SVTInfo, FLRUNode* LRUNode)
+void FStreamingManager::StreamOutMipLevel(FStreamingInfo* SVTInfo, FLRUNode* LRUNode)
 {
 	const int32 FrameIndex = LRUNode->FrameIndex;
 	const int32 MipLevelIndex = LRUNode->MipLevelIndex;
 
-	FFrameInfo& FrameInfo = SVTInfo.PerFrameInfo[FrameIndex];
+	FFrameInfo& FrameInfo = SVTInfo->PerFrameInfo[FrameIndex];
 
 	check(FrameInfo.LowestResidentMipLevel >= MipLevelIndex); // mip might not have streamed in yet, so use >= instead of ==
 	check(FrameInfo.LowestRequestedMipLevel == MipLevelIndex);
@@ -1851,8 +1853,8 @@ void FStreamingManager::StreamOutMipLevel(FStreamingInfo& SVTInfo, FLRUNode* LRU
 	FrameInfo.LowestResidentMipLevel = NewLowestResidentMipLevel;
 	
 	// Update the streaming info buffer data
-	SVTInfo.DirtyStreamingInfoData[FrameIndex] = true;
-	InvalidatedStreamingInfos.Add(&SVTInfo);
+	SVTInfo->DirtyStreamingInfoData[FrameIndex] = true;
+	SVTsWithInvalidatedStreamingInfoBuffer.Add(SVTInfo);
 
 	// Unlink
 	LRUNode->Remove();
@@ -1866,7 +1868,7 @@ void FStreamingManager::StreamOutMipLevel(FStreamingInfo& SVTInfo, FLRUNode* LRU
 	// Free allocated tiles
 	for (uint32& TileCoord : FrameInfo.TileAllocations[MipLevelIndex])
 	{
-		SVTInfo.TileDataTexture->Free(TileCoord);
+		SVTInfo->TileDataTexture->Free(TileCoord);
 		TileCoord = 0;
 	}
 }
@@ -1885,7 +1887,7 @@ int32 FStreamingManager::DetermineReadyMipLevels()
 		const int32 PendingMipLevelIndex = (StartPendingMipLevelIndex + i) % MaxPendingMipLevels;
 		FPendingMipLevel& PendingMipLevel = PendingMipLevels[PendingMipLevelIndex];
 
-		FStreamingInfo* SVTInfo = StreamingInfo.Find(PendingMipLevel.SparseVolumeTexture);
+		FStreamingInfo* SVTInfo = FindStreamingInfo(PendingMipLevel.SparseVolumeTexture);
 		if (!SVTInfo)
 		{
 #if WITH_EDITORONLY_DATA
@@ -1992,7 +1994,7 @@ void FStreamingManager::InstallReadyMipLevels()
 		const int32 PendingMipLevelIndex = (StartPendingMipLevelIndex + i) % MaxPendingMipLevels;
 		FPendingMipLevel& PendingMipLevel = PendingMipLevels[PendingMipLevelIndex];
 
-		FStreamingInfo* SVTInfo = StreamingInfo.Find(PendingMipLevel.SparseVolumeTexture);
+		FStreamingInfo* SVTInfo = FindStreamingInfo(PendingMipLevel.SparseVolumeTexture);
 		if (!SVTInfo || (SVTInfo->PerFrameInfo[PendingMipLevel.FrameIndex].LowestRequestedMipLevel > PendingMipLevel.MipLevelIndex))
 		{
 			PendingMipLevel.Reset();
@@ -2109,7 +2111,7 @@ void FStreamingManager::InstallReadyMipLevels()
 
 		// Update the streaming info buffer data
 		SVTInfo->DirtyStreamingInfoData[PendingMipLevel.FrameIndex] = true;
-		InvalidatedStreamingInfos.Add(SVTInfo);
+		SVTsWithInvalidatedStreamingInfoBuffer.Add(SVTInfo);
 
 		const int32 LRUNodeIndex = PendingMipLevel.FrameIndex * SVTInfo->NumMipLevelsGlobal + PendingMipLevel.MipLevelIndex;
 		SVTInfo->LRUNodes[LRUNodeIndex].PendingMipLevelIndex = INDEX_NONE;
@@ -2130,7 +2132,7 @@ void FStreamingManager::InstallReadyMipLevels()
 				{
 					FMemory::Memcpy(PageTableTask.DstPageCoords, PageTableTask.SrcPageCoords, PageTableTask.NumPageTableUpdates * sizeof(uint32));
 
-					FStreamingInfo* SVTInfo = StreamingInfo.Find(PageTableTask.PendingMipLevel->SparseVolumeTexture);
+					FStreamingInfo* SVTInfo = FindStreamingInfo(PageTableTask.PendingMipLevel->SparseVolumeTexture);
 					TArray<uint32>& TileAllocations = SVTInfo->PerFrameInfo[PageTableTask.PendingMipLevel->FrameIndex].TileAllocations[PageTableTask.PendingMipLevel->MipLevelIndex];
 					const uint32* SrcEntries = reinterpret_cast<const uint32*>(PageTableTask.SrcPageEntries);
 					uint32* DstEntries = reinterpret_cast<uint32*>(PageTableTask.DstPageEntries);
@@ -2204,6 +2206,13 @@ void FStreamingManager::InstallReadyMipLevels()
 #endif
 }
 
+FStreamingManager::FStreamingInfo* FStreamingManager::FindStreamingInfo(UStreamableSparseVolumeTexture* Key)
+{
+	TUniquePtr<FStreamingInfo>* SVTInfoPtr = StreamingInfo.Find(Key);
+	check(!SVTInfoPtr || SVTInfoPtr->Get());
+	return SVTInfoPtr ? SVTInfoPtr->Get() : nullptr;
+}
+
 #if WITH_EDITORONLY_DATA
 
 UE::DerivedData::FCacheGetChunkRequest FStreamingManager::BuildDDCRequest(const FResources& Resources, const FMipLevelStreamingInfo& MipLevelStreamingInfo, const uint32 PendingMipLevelIndex)
diff --git a/Engine/Source/Runtime/Engine/Private/SparseVolumeTexture/SparseVolumeTextureStreamingManager.h b/Engine/Source/Runtime/Engine/Private/SparseVolumeTexture/SparseVolumeTextureStreamingManager.h
--- a/Engine/Source/Runtime/Engine/Private/SparseVolumeTexture/SparseVolumeTextureStreamingManager.h
+++ b/Engine/Source/Runtime/Engine/Private/SparseVolumeTexture/SparseVolumeTextureStreamingManager.h
@@ -318,7 +318,7 @@ private:
 		TArray<FFrameInfo> FrameInfo; // Only Resources and TextureRenderResources are initialized
 	};
 
-	TMap<UStreamableSparseVolumeTexture*, FStreamingInfo> StreamingInfo; // Do not dereference the key! We just read the pointer itself.
+	TMap<UStreamableSparseVolumeTexture*, TUniquePtr<FStreamingInfo>> StreamingInfo; // Do not dereference the key! We just read the pointer itself.
 	TMap<FMipLevelKey, uint32> RequestsHashTable;
 	TArray<FPendingMipLevel> PendingMipLevels;
 #if WITH_EDITORONLY_DATA
@@ -338,7 +338,7 @@ private:
 	// Transient lifetime
 	TArray<FStreamingRequest> ParentRequestsToAdd;
 	TSet<FTileDataTexture*> TileDataTexturesToUpdate;
-	TSet<FStreamingInfo*> InvalidatedStreamingInfos;
+	TSet<FStreamingInfo*> SVTsWithInvalidatedStreamingInfoBuffer; // Changes in the lowest resident mip level cause invalidation
 	TArray<FStreamingRequest> PrioritizedRequestsHeap;
 	TArray<FStreamingRequest> SelectedRequests;
 	TArray<FPageTableClear> PageTableClears;
@@ -352,9 +352,10 @@ private:
 	void AddParentRequests(); // Add requests for all parent mip levels
 	void SelectHighestPriorityRequestsAndUpdateLRU(int32 MaxSelectedRequests);
 	void IssueRequests(int32 MaxSelectedRequests);
-	void StreamOutMipLevel(FStreamingInfo& SVTInfo, FLRUNode* LRUNode);
+	void StreamOutMipLevel(FStreamingInfo* SVTInfo, FLRUNode* LRUNode);
 	int32 DetermineReadyMipLevels();
 	void InstallReadyMipLevels();
+	FStreamingInfo* FindStreamingInfo(UStreamableSparseVolumeTexture* Key); // Returns nullptr if the key can't be found
 
 #if WITH_EDITORONLY_DATA
 	UE::DerivedData::FCacheGetChunkRequest BuildDDCRequest(const FResources& Resources, const FMipLevelStreamingInfo& MipLevelStreamingInfo, const uint32 PendingMipLevelIndex);
