WorldPartition: Fixed how MaxCellsToLoad is computed by WorldPartition to limit the total number of loading and pending load streaming cells.

diff --git a/Engine/Source/Runtime/Engine/Private/WorldPartition/WorldPartitionSubsystem.cpp b/Engine/Source/Runtime/Engine/Private/WorldPartition/WorldPartitionSubsystem.cpp
--- a/Engine/Source/Runtime/Engine/Private/WorldPartition/WorldPartitionSubsystem.cpp
+++ b/Engine/Source/Runtime/Engine/Private/WorldPartition/WorldPartitionSubsystem.cpp
@@ -566,6 +566,60 @@ bool UWorldPartitionSubsystem::HasUninitializationPendingStreamingLevels(const U
 	return false;
 }
 
+static FORCEINLINE bool IsLoadingOrPendingLoadStreamingLevel(const ULevelStreaming* InStreamingLevel)
+{
+	// Consider loading and pending to load world partition streaming levels
+	if (InStreamingLevel && InStreamingLevel->IsA<UWorldPartitionLevelStreamingDynamic>())
+	{
+		const ELevelStreamingState CurrentState = InStreamingLevel->GetLevelStreamingState();
+		return ((CurrentState == ELevelStreamingState::Loading) || ((CurrentState == ELevelStreamingState::Removed || CurrentState == ELevelStreamingState::Unloaded) && InStreamingLevel->ShouldBeLoaded()));
+	}
+	return false;
+}
+
+void UWorldPartitionSubsystem::UpdateLoadingAndPendingLoadStreamingLevels(const ULevelStreaming* InStreamingLevel)
+{
+	if (IsLoadingOrPendingLoadStreamingLevel(InStreamingLevel))
+	{
+		WorldPartitionLoadingAndPendingLoadStreamingLevels.Add(InStreamingLevel);
+	}
+	else
+	{
+		WorldPartitionLoadingAndPendingLoadStreamingLevels.Remove(InStreamingLevel);
+	}
+}
+
+int32 UWorldPartitionSubsystem::GetMaxCellsToLoad(const UWorld* InWorld)
+{
+	TRACE_CPUPROFILER_EVENT_SCOPE(UWorldPartitionSubsystem::GetMaxCellsToLoad);
+	int32 MaxCellsToLoad = MAX_int32;
+	const ENetMode NetMode = InWorld->GetNetMode();
+	const bool bIsServer = (NetMode == NM_DedicatedServer || NetMode == NM_ListenServer);
+	if (!bIsServer && !InWorld->GetIsInBlockTillLevelStreamingCompleted() && (GMaxLoadingStreamingCells > 0))
+	{
+		if (UWorldPartitionSubsystem* WorldPartitionSubsystem = UWorld::GetSubsystem<UWorldPartitionSubsystem>(InWorld))
+		{
+			MaxCellsToLoad = GMaxLoadingStreamingCells;
+			for (auto It = WorldPartitionSubsystem->WorldPartitionLoadingAndPendingLoadStreamingLevels.CreateIterator(); It; ++It)
+			{
+				const ULevelStreaming* InStreamingLevel = It->Get();
+				if (IsLoadingOrPendingLoadStreamingLevel(InStreamingLevel))
+				{
+					if (!--MaxCellsToLoad)
+					{
+						break;
+					}
+				}
+				else
+				{
+					It.RemoveCurrent();
+				}
+			}
+		}
+	}
+	return MaxCellsToLoad;
+};
+
 void UWorldPartitionSubsystem::OnLevelStreamingStateChanged(UWorld* InWorld, const ULevelStreaming* InStreamingLevel, ULevel* LevelIfLoaded, ELevelStreamingState PreviousState, ELevelStreamingState NewState)
 {
 	if (InWorld != GetWorld())
@@ -573,6 +627,8 @@ void UWorldPartitionSubsystem::OnLevelStreamingStateChanged(UWorld* InWorld, con
 		return;
 	}
 
+	UpdateLoadingAndPendingLoadStreamingLevels(InStreamingLevel);
+
 	if (NewState == ELevelStreamingState::Removed)
 	{
 		const UWorld* OuterWorld = InStreamingLevel->GetStreamingWorld();
@@ -629,6 +685,8 @@ void UWorldPartitionSubsystem::OnLevelStreamingTargetStateChanged(UWorld* InWorl
 		return;
 	}
 
+	UpdateLoadingAndPendingLoadStreamingLevels(InStreamingLevel);
+
 	// Make sure when a WorldPartition is LevelStreamed that changing its state to remove it from world will update the target states of its Cells right away.
 	if(InLevelIfLoaded && InNewTarget != ELevelStreamingTargetState::LoadedVisible)
 	{
@@ -952,9 +1010,7 @@ void UWorldPartitionSubsystem::UpdateStreamingStateInternal(const UWorld* InWorl
 	}
 
 	// Compute maximum number of cells to load
-	const ENetMode NetMode = World->GetNetMode();
-	const bool bIsServer = (NetMode == NM_DedicatedServer || NetMode == NM_ListenServer);
-	int32 MaxCellsToLoad = (!bIsServer && !World->GetIsInBlockTillLevelStreamingCompleted() && (GMaxLoadingStreamingCells > 0)) ? (GMaxLoadingStreamingCells - (int32)World->GetNumStreamingLevelsBeingLoaded()) : MAX_int32;
+	int32 MaxCellsToLoad = GetMaxCellsToLoad(World);
 
 	// Process cells to activate
 	{
diff --git a/Engine/Source/Runtime/Engine/Public/WorldPartition/WorldPartitionSubsystem.h b/Engine/Source/Runtime/Engine/Public/WorldPartition/WorldPartitionSubsystem.h
--- a/Engine/Source/Runtime/Engine/Public/WorldPartition/WorldPartitionSubsystem.h
+++ b/Engine/Source/Runtime/Engine/Public/WorldPartition/WorldPartitionSubsystem.h
@@ -159,18 +159,20 @@ protected:
 private:
 
 	// Streaming Sources
-	ENGINE_API void UpdateStreamingSources();
+	void UpdateStreamingSources();
 	ENGINE_API void GetStreamingSources(const UWorldPartition* InWorldPartition, TArray<FWorldPartitionStreamingSource>& OutStreamingSources) const;
 	uint32 GetStreamingSourcesHash() const { return StreamingSourcesHash; }
 
-	ENGINE_API void OnWorldPartitionInitialized(UWorldPartition* InWorldPartition);
-	ENGINE_API void OnWorldPartitionUninitialized(UWorldPartition* InWorldPartition);
-	ENGINE_API void OnLevelStreamingStateChanged(UWorld* InWorld, const ULevelStreaming* InStreamingLevel, ULevel* LevelIfLoaded, ELevelStreamingState PreviousState, ELevelStreamingState NewState);
-	ENGINE_API void OnLevelStreamingTargetStateChanged(UWorld* InWorld, const ULevelStreaming* InStreamingLevel, ULevel* InLevelIfLoaded, ELevelStreamingState InCurrentState, ELevelStreamingTargetState InPrevTarget, ELevelStreamingTargetState InNewTarget);
-	ENGINE_API void OnLevelBeginMakingVisible(UWorld* InWorld, const ULevelStreaming* InStreamingLevel, ULevel* InLoadedLevel);
-	ENGINE_API void OnLevelBeginMakingInvisible(UWorld* InWorld, const ULevelStreaming* InStreamingLevel, ULevel* InLoadedLevel);
+	void OnWorldPartitionInitialized(UWorldPartition* InWorldPartition);
+	void OnWorldPartitionUninitialized(UWorldPartition* InWorldPartition);
+	void OnLevelStreamingStateChanged(UWorld* InWorld, const ULevelStreaming* InStreamingLevel, ULevel* LevelIfLoaded, ELevelStreamingState PreviousState, ELevelStreamingState NewState);
+	void OnLevelStreamingTargetStateChanged(UWorld* InWorld, const ULevelStreaming* InStreamingLevel, ULevel* InLevelIfLoaded, ELevelStreamingState InCurrentState, ELevelStreamingTargetState InPrevTarget, ELevelStreamingTargetState InNewTarget);
+	void OnLevelBeginMakingVisible(UWorld* InWorld, const ULevelStreaming* InStreamingLevel, ULevel* InLoadedLevel);
+	void OnLevelBeginMakingInvisible(UWorld* InWorld, const ULevelStreaming* InStreamingLevel, ULevel* InLoadedLevel);
+	void UpdateLoadingAndPendingLoadStreamingLevels(const ULevelStreaming* InStreamingLevel);
 
 	static ENGINE_API void UpdateStreamingStateInternal(const UWorld* InWorld, const TArray<TObjectPtr<UWorldPartition>>& InWorldPartitions);
+	static int32 GetMaxCellsToLoad(const UWorld* InWorld);
 
 	ENGINE_API bool IsServer() const;
 	ENGINE_API bool HasAnyWorldPartitionServerStreamingEnabled() const;
@@ -206,6 +208,9 @@ private:
 	// Tracks streaming levels of uninitialized world partition used to delay next initialization until they're done being removed
 	TMap<FSoftObjectPath, TSet<TWeakObjectPtr<ULevelStreaming>>> WorldPartitionUninitializationPendingStreamingLevels;
 
+	// Tracks world partition loading and pending loads
+	TSet<TWeakObjectPtr<const ULevelStreaming>> WorldPartitionLoadingAndPendingLoadStreamingLevels;
+
 #if WITH_EDITOR
 	bool bIsRunningConvertWorldPartitionCommandlet;
 	mutable FActorDescContainerInstanceManager ActorDescContainerInstanceManager;
