LevelInstance: add ShouldCookWorldAsset() virtual to make it clearer why we aren't cooking a PackedLevelActor's WorldAsset.
- Embedded LevelInstances already avoid cooking their WorldAsset as those actors aren't loaded during cook of a WP level
- Make sure we Cook WorldAsset on CDOs (Blueprints) only if they have a Streaming runtime behavior

#rnx

diff --git a/Engine/Source/Runtime/Engine/Private/LevelInstance/LevelInstanceActor.cpp b/Engine/Source/Runtime/Engine/Private/LevelInstance/LevelInstanceActor.cpp
--- a/Engine/Source/Runtime/Engine/Private/LevelInstance/LevelInstanceActor.cpp
+++ b/Engine/Source/Runtime/Engine/Private/LevelInstance/LevelInstanceActor.cpp
@@ -154,8 +154,22 @@ void ALevelInstance::PostLoad()
 	LevelInstanceComponent = Cast<ULevelInstanceComponent>(RootComponent);
 	OnLevelInstanceActorPostLoad.Broadcast(this);
 
-	UpdateCookedAsset();
+#if WITH_EDITORONLY_DATA
+	if (IsRunningCookCommandlet() && ShouldCookWorldAsset())
+	{
+		CookedWorldAsset = WorldAsset;
+	}
+#endif
+}
+
+#if WITH_EDITOR
+bool ALevelInstance::ShouldCookWorldAsset() const
+{
+	// If ALevelInstance actor gets loaded it means it needs to Cook its WorldAsset (World Partition Embedded Level Instances don't get loaded as they aren't runtime relevant)
+	// If ALevelInstnace is a template then we only need to Cook its WorldAsset if it's desired runtime behavior is to be Level Streamed
+	return !IsTemplate() || GetDesiredRuntimeBehavior() == ELevelInstanceRuntimeBehavior::LevelStreaming;
 }
+#endif
 
 void ALevelInstance::PreEditUndo()
 {
@@ -310,16 +324,6 @@ void ALevelInstance::PushLevelInstanceEditingStateToProxies(bool bInEditingState
 	LevelInstanceActorImpl.PushLevelInstanceEditingStateToProxies(bInEditingState);
 }
 
-void ALevelInstance::UpdateCookedAsset()
-{
-#if WITH_EDITORONLY_DATA
-	if (IsRunningCookCommandlet() && IsLoadingEnabled())
-	{
-		CookedWorldAsset = WorldAsset;
-	}
-#endif
-}
-
 #endif
 
 #undef LOCTEXT_NAMESPACE
diff --git a/Engine/Source/Runtime/Engine/Public/LevelInstance/LevelInstanceActor.h b/Engine/Source/Runtime/Engine/Public/LevelInstance/LevelInstanceActor.h
--- a/Engine/Source/Runtime/Engine/Public/LevelInstance/LevelInstanceActor.h
+++ b/Engine/Source/Runtime/Engine/Public/LevelInstance/LevelInstanceActor.h
@@ -105,7 +105,7 @@ public:
 	static ENGINE_API FOnLevelInstanceActorPostLoad OnLevelInstanceActorPostLoad;
 
 private:
-	ENGINE_API void UpdateCookedAsset();
+	ENGINE_API virtual bool ShouldCookWorldAsset() const;
 #endif
 };
 
diff --git a/Engine/Source/Runtime/Engine/Public/PackedLevelActor/PackedLevelActor.h b/Engine/Source/Runtime/Engine/Public/PackedLevelActor/PackedLevelActor.h
--- a/Engine/Source/Runtime/Engine/Public/PackedLevelActor/PackedLevelActor.h
+++ b/Engine/Source/Runtime/Engine/Public/PackedLevelActor/PackedLevelActor.h
@@ -96,6 +96,8 @@ private:
 #endif
 
 #if WITH_EDITOR
+	bool ShouldCookWorldAsset() const override { return false; }
+
 	bool bChildChanged;
 	bool bLoadForPacking;
 #endif
