World Partition - HLOD: Use global frame counter when performing HLOD warmup
* Prevent stall when warmup is performed behind a loading screen world

diff --git a/Engine/Source/Runtime/Engine/Private/WorldPartition/HLOD/HLODSubsystem.cpp b/Engine/Source/Runtime/Engine/Private/WorldPartition/HLOD/HLODSubsystem.cpp
--- a/Engine/Source/Runtime/Engine/Private/WorldPartition/HLOD/HLODSubsystem.cpp
+++ b/Engine/Source/Runtime/Engine/Private/WorldPartition/HLOD/HLODSubsystem.cpp
@@ -556,7 +556,7 @@ bool UHLODSubsystem::PrepareToWarmup(const UWorldPartitionRuntimeCell* InCell, A
 
 		// In case a previous request to unload was aborted and the cell never unloaded... assume warmup requests are expired after a given amount of frames.
 		const uint32 WarmupExpiredFrames = 30;
-		const uint32 CurrentFrameNumber = GetWorld()->Scene->GetFrameNumber();
+		const uint32 CurrentFrameNumber = GFrameNumberRenderThread;
 
 		// Trigger warmup on the first request to unload, or if a warmup request expired
 		const bool bInitiateWarmup = WarmupState.WarmupEndFrame == INDEX_NONE || CurrentFrameNumber > (WarmupState.WarmupEndFrame + WarmupExpiredFrames);
