PR #10413: Get TrackTimeToGo only Within the Segment Time

diff --git a/Engine/Source/Runtime/Engine/Private/Animation/AnimCompositeBase.cpp b/Engine/Source/Runtime/Engine/Private/Animation/AnimCompositeBase.cpp
--- a/Engine/Source/Runtime/Engine/Private/Animation/AnimCompositeBase.cpp
+++ b/Engine/Source/Runtime/Engine/Private/Animation/AnimCompositeBase.cpp
@@ -128,7 +128,10 @@ void FAnimSegment::GetAnimNotifiesFromTrackPositions(const float& PreviousTrackP
 			// Get starting position, closest overlap.
 			float AnimStartPosition = ConvertTrackPosToAnimPos( bTrackPlayingBackwards ? FMath::Min(PreviousTrackPosition, SegmentEndPos) : FMath::Max(PreviousTrackPosition, SegmentStartPos) );
 			AnimStartPosition = FMath::Clamp(AnimStartPosition, AnimStartTime, AnimEndTime);
-			float TrackTimeToGo = FMath::Abs(CurrentTrackPosition - PreviousTrackPosition);
+
+			// When looping, the current track position could exceed the current segment (anim montage loops the track position after firing notifies)
+			// We need to make sure to clamp the current/previous track positions within our segment
+			float TrackTimeToGo = FMath::Abs(FMath::Clamp(CurrentTrackPosition, SegmentStartPos, SegmentEndPos) - FMath::Clamp(PreviousTrackPosition, SegmentStartPos, SegmentEndPos));
 
 			// The track can be playing backwards and the animation can be playing backwards, so we
 			// need to combine to work out what direction we are traveling through the animation
