Fix for audio scrubbing in sequencer (or any StopDelayed() sound) not playing at low frame rates.

diff --git a/Engine/Source/Runtime/Engine/Private/AudioDevice.cpp b/Engine/Source/Runtime/Engine/Private/AudioDevice.cpp
--- a/Engine/Source/Runtime/Engine/Private/AudioDevice.cpp
+++ b/Engine/Source/Runtime/Engine/Private/AudioDevice.cpp
@@ -4464,6 +4464,11 @@ void FAudioDevice::StartSources(TArray<FWaveInstance*>& WaveInstances, int32 Fir
 							}
 						}
 					}
+					else
+					{
+						// This sound is not yet prepared to play, perform any necessary operations 
+						UpdateUnpreparedSound(WaveInstance, bGameTicking);
+					}
 				}
 
 				// If we succeeded above then we need to map the wave instance to the source
@@ -7382,3 +7387,28 @@ int32 FAudioDevice::GetNumPrecacheFrames() const
 	// Otherwise, use the default value or value set in ini file
 	return NumPrecacheFrames;
 }
+
+void FAudioDevice::UpdateUnpreparedSound(FWaveInstance* WaveInstance, bool bGameTicking) const
+{
+	// If this source is not playing yet due to a pending async decode, push it's active duration out 
+	// by the amount of time that has elapsed.
+	UpdateSoundDuration(WaveInstance, bGameTicking);
+}
+
+void FAudioDevice::UpdateSoundDuration(FWaveInstance* WaveInstance, bool bGameTicking) const
+{
+	if (FActiveSound* ActiveSound = WaveInstance->ActiveSound)
+	{
+		if (bGameTicking || ActiveSound->bIsUISound)
+		{
+			float ActiveDuration = ActiveSound->ComponentVolumeFader.GetActiveDuration();
+			if (ActiveDuration > 0.0f)
+			{
+				float DeltaTime = GetGameDeltaTime();
+				// Push the duration out by the previous frame's delta amount as the fader's elapsed time
+				// has already been updated at this point with the same value.
+				ActiveSound->ComponentVolumeFader.SetActiveDuration(ActiveDuration + DeltaTime);
+			}
+		}
+	}
+}
diff --git a/Engine/Source/Runtime/Engine/Public/AudioDevice.h b/Engine/Source/Runtime/Engine/Public/AudioDevice.h
--- a/Engine/Source/Runtime/Engine/Public/AudioDevice.h
+++ b/Engine/Source/Runtime/Engine/Public/AudioDevice.h
@@ -1879,6 +1879,12 @@ private:
 	/** Returns the number of frames to use per precache buffer. */
 	int32 GetNumPrecacheFrames() const;
 
+	/** Called by StartSources when sounds are not prepared to Init for whatever reason */
+	void UpdateUnpreparedSound(FWaveInstance* WaveInstance, bool bGameTicking) const;
+
+	/** Adjusts the active sound duration to make up for decode latency */
+	void UpdateSoundDuration(FWaveInstance* WaveInstance, bool bGameTicking) const;
+
 	bool RemoveVirtualLoop(FActiveSound& ActiveSound);
 public:
 
diff --git a/Engine/Source/Runtime/MovieSceneTracks/Private/Systems/MovieSceneAudioSystem.cpp b/Engine/Source/Runtime/MovieSceneTracks/Private/Systems/MovieSceneAudioSystem.cpp
--- a/Engine/Source/Runtime/MovieSceneTracks/Private/Systems/MovieSceneAudioSystem.cpp
+++ b/Engine/Source/Runtime/MovieSceneTracks/Private/Systems/MovieSceneAudioSystem.cpp
@@ -629,8 +629,22 @@ private:
 
 		if (Context.GetStatus() != EMovieScenePlayerStatus::Playing)
 		{
+			float ScrubDuration = AudioTrackConstants::ScrubDuration;
+			if (FAudioDevice* AudioDevice = AudioComponent.GetAudioDevice())
+			{
+				constexpr float MinScrubFrameRateFactor = 1.5f;
+				float DeviceDeltaTime = AudioDevice->GetGameDeltaTime();
+
+				// When operating at very low frame-rates (<20fps), a single frame will be
+				// longer than the hard coded scrub duration of 50ms in which case the delayed
+				// stop will trigger on the same frame that the sound starts playing and
+				// no audio will be heard. Here we increase the scrub duration to be greater than
+				// a single frame if necessary.
+				ScrubDuration = FMath::Max(ScrubDuration, DeviceDeltaTime * MinScrubFrameRateFactor);
+			}
+
 			// While scrubbing, play the sound for a short time and then cut it.
-			AudioComponent.StopDelayed(AudioTrackConstants::ScrubDuration);
+			AudioComponent.StopDelayed(ScrubDuration);
 		}
 
 		if (AudioComponent.IsPlaying() && AudioInputs)
