WorldPartition: Optimized UpdateStreamingState for Server streaming, compute ClientsVisibleLevelNames once for all registered world partitions

diff --git a/Engine/Source/Runtime/Engine/Private/WorldPartition/WorldPartitionStreamingPolicy.cpp b/Engine/Source/Runtime/Engine/Private/WorldPartition/WorldPartitionStreamingPolicy.cpp
--- a/Engine/Source/Runtime/Engine/Private/WorldPartition/WorldPartitionStreamingPolicy.cpp
+++ b/Engine/Source/Runtime/Engine/Private/WorldPartition/WorldPartitionStreamingPolicy.cpp
@@ -173,9 +173,7 @@ uint32 UWorldPartitionStreamingPolicy::ComputeUpdateStreamingHash(bool bCanOptim
 
 		if (WorldPartition->IsServer())
 		{
-			TArray<FName> ClientsVisibleLevelNames = ServerClientsVisibleLevelNames.Array();
-			ClientsVisibleLevelNames.Sort(FNameFastLess());
-			HashBuilder << ClientsVisibleLevelNames;
+			HashBuilder << GetWorld()->GetSubsystem<UWorldPartitionSubsystem>()->GetServerClientsVisibleLevelsHash();
 		}
 
 		return HashBuilder.GetHash();
@@ -298,21 +296,6 @@ void UWorldPartitionStreamingPolicy::UpdateStreamingState()
 	// Update streaming sources
 	UpdateStreamingSources(bCanOptimizeUpdate);
 
-	if (bIsServer)
-	{
-		// Gather Client visible level names
-		ServerClientsVisibleLevelNames.Reset();
-		if (const UNetDriver* NetDriver = GetWorld()->GetNetDriver())
-		{
-			for (UNetConnection* Connection : NetDriver->ClientConnections)
-			{
-				ServerClientsVisibleLevelNames.Add(Connection->GetClientWorldPackageName());
-				ServerClientsVisibleLevelNames.Append(Connection->ClientVisibleLevelNames);
-				ServerClientsVisibleLevelNames.Append(Connection->GetClientMakingVisibleLevelNames());
-			}
-		}
-	}
-
 	// Detect if nothing relevant changed and early out
 	const uint32 NewUpdateStreamingHash = ComputeUpdateStreamingHash(bCanOptimizeUpdate);
 	const bool bShouldSkipUpdate = NewUpdateStreamingHash && (UpdateStreamingHash == NewUpdateStreamingHash);
@@ -468,7 +451,8 @@ void UWorldPartitionStreamingPolicy::UpdateStreamingState()
 		}
 	}
 
-	auto ShouldWaitForClientVisibility = [bIsServer, this, &bUpdateEpoch](const UWorldPartitionRuntimeCell* Cell)
+	const TSet<FName>& ServerClientsVisibleLevelNames = GetWorld()->GetSubsystem<UWorldPartitionSubsystem>()->ServerClientsVisibleLevelNames;
+	auto ShouldWaitForClientVisibility = [bIsServer, this, &bUpdateEpoch, &ServerClientsVisibleLevelNames](const UWorldPartitionRuntimeCell* Cell)
 	{
 		check(bIsServer);
 		if (ULevel* Level = Cell->GetLevel())
diff --git a/Engine/Source/Runtime/Engine/Private/WorldPartition/WorldPartitionSubsystem.cpp b/Engine/Source/Runtime/Engine/Private/WorldPartition/WorldPartitionSubsystem.cpp
--- a/Engine/Source/Runtime/Engine/Private/WorldPartition/WorldPartitionSubsystem.cpp
+++ b/Engine/Source/Runtime/Engine/Private/WorldPartition/WorldPartitionSubsystem.cpp
@@ -21,12 +21,15 @@
 #include "Engine/CoreSettings.h"
 #include "Engine/LevelStreaming.h"
 #include "Engine/LevelStreamingGCHelper.h"
+#include "Engine/NetDriver.h"
+#include "Engine/NetConnection.h"
 #include "Streaming/LevelStreamingDelegates.h"
 #include "Engine/LevelBounds.h"
 #include "Debug/DebugDrawService.h"
 #include "GameFramework/PlayerController.h"
 #include "Async/ParallelFor.h"
 #include "Algo/ForEach.h"
+#include "Misc/HashBuilder.h"
 
 #if WITH_EDITOR
 #include "Editor.h"
@@ -135,6 +138,7 @@ static FAutoConsoleVariableRef CMaxLoadingStreamingCells(
 UWorldPartitionSubsystem::UWorldPartitionSubsystem()
 : StreamingSourcesHash(0)
 , NumWorldPartitionServerStreamingEnabled(0)
+, ServerClientsVisibleLevelsHash(0)
 {}
 
 UWorldPartition* UWorldPartitionSubsystem::GetWorldPartition()
@@ -593,9 +597,7 @@ int32 UWorldPartitionSubsystem::GetMaxCellsToLoad(const UWorld* InWorld)
 {
 	TRACE_CPUPROFILER_EVENT_SCOPE(UWorldPartitionSubsystem::GetMaxCellsToLoad);
 	int32 MaxCellsToLoad = MAX_int32;
-	const ENetMode NetMode = InWorld->GetNetMode();
-	const bool bIsServer = (NetMode == NM_DedicatedServer || NetMode == NM_ListenServer);
-	if (!bIsServer && !InWorld->GetIsInBlockTillLevelStreamingCompleted() && (GMaxLoadingStreamingCells > 0))
+	if (!IsServer(InWorld) && !InWorld->GetIsInBlockTillLevelStreamingCompleted() && (GMaxLoadingStreamingCells > 0))
 	{
 		if (UWorldPartitionSubsystem* WorldPartitionSubsystem = UWorld::GetSubsystem<UWorldPartitionSubsystem>(InWorld))
 		{
@@ -829,9 +831,9 @@ static FAutoConsoleVariableRef CVarUpdateStreamingSources(
 static const FName NAME_SIEStreamingSource(TEXT("SIE"));
 #endif
 
-bool UWorldPartitionSubsystem::IsServer() const
+bool UWorldPartitionSubsystem::IsServer(const UWorld* InWorld)
 {
-	const ENetMode NetMode = GetWorld()->GetNetMode();
+	const ENetMode NetMode = InWorld->GetNetMode();
 	const bool bIsServer = (NetMode == NM_DedicatedServer || NetMode == NM_ListenServer);
 	return bIsServer;
 }
@@ -841,6 +843,30 @@ bool UWorldPartitionSubsystem::HasAnyWorldPartitionServerStreamingEnabled() cons
 	return (NumWorldPartitionServerStreamingEnabled > 0);
 }
 
+void UWorldPartitionSubsystem::UpdateServerClientsVisibleLevelNames()
+{
+	TRACE_CPUPROFILER_EVENT_SCOPE(UWorldPartitionSubsystem::UpdateServerClientsVisibleLevelNames);
+	UWorld* World = GetWorld();
+	if (IsServer(World))
+	{
+		ServerClientsVisibleLevelNames.Reset();
+		if (const UNetDriver* NetDriver = World->GetNetDriver())
+		{
+			for (UNetConnection* Connection : NetDriver->ClientConnections)
+			{
+				ServerClientsVisibleLevelNames.Add(Connection->GetClientWorldPackageName());
+				ServerClientsVisibleLevelNames.Append(Connection->ClientVisibleLevelNames);
+				ServerClientsVisibleLevelNames.Append(Connection->GetClientMakingVisibleLevelNames());
+			}
+		}
+		FHashBuilder HashBuilder;
+		TArray<FName> SortedServerClientsVisibleLevelNames = ServerClientsVisibleLevelNames.Array();
+		SortedServerClientsVisibleLevelNames.Sort(FNameFastLess());
+		HashBuilder << SortedServerClientsVisibleLevelNames;
+		ServerClientsVisibleLevelsHash = HashBuilder.GetHash();
+	}
+}
+
 void UWorldPartitionSubsystem::UpdateStreamingSources()
 {
 	if (!GUpdateStreamingSources)
@@ -864,7 +890,7 @@ void UWorldPartitionSubsystem::UpdateStreamingSources()
 
 	if (!bIsUsingReplayStreamingSources)
 	{
-		if (!IsServer() || HasAnyWorldPartitionServerStreamingEnabled() || AWorldPartitionReplay::IsRecordingEnabled(World))
+		if (!IsServer(World) || HasAnyWorldPartitionServerStreamingEnabled() || AWorldPartitionReplay::IsRecordingEnabled(World))
 		{
 			bool bAllowPlayerControllerStreamingSources = true;
 #if WITH_EDITOR
@@ -983,6 +1009,8 @@ void UWorldPartitionSubsystem::UpdateStreamingStateInternal(const UWorld* InWorl
 	{
 		// Update streaming sources
 		WorldPartitionSubsystem->UpdateStreamingSources();
+		// Update server's clients visible levels
+		WorldPartitionSubsystem->UpdateServerClientsVisibleLevelNames();
 	}
 	
 	// Make temp copy of array as UpdateStreamingState may FlushAsyncLoading, which may add a new world partition to RegisteredWorldPartitions while iterating
diff --git a/Engine/Source/Runtime/Engine/Public/WorldPartition/WorldPartitionStreamingPolicy.h b/Engine/Source/Runtime/Engine/Public/WorldPartition/WorldPartitionStreamingPolicy.h
--- a/Engine/Source/Runtime/Engine/Public/WorldPartition/WorldPartitionStreamingPolicy.h
+++ b/Engine/Source/Runtime/Engine/Public/WorldPartition/WorldPartitionStreamingPolicy.h
@@ -132,7 +132,6 @@ private:
 	int32 ServerDataLayersStatesEpoch;
 	int32 ServerStreamingStateEpoch;
 	int32 ServerStreamingEnabledEpoch;
-	TSet<FName> ServerClientsVisibleLevelNames;
 	uint32 UpdateStreamingHash;
 	uint32 UpdateStreamingSourcesHash;
 	uint32 UpdateStreamingStateCalls;
diff --git a/Engine/Source/Runtime/Engine/Public/WorldPartition/WorldPartitionSubsystem.h b/Engine/Source/Runtime/Engine/Public/WorldPartition/WorldPartitionSubsystem.h
--- a/Engine/Source/Runtime/Engine/Public/WorldPartition/WorldPartitionSubsystem.h
+++ b/Engine/Source/Runtime/Engine/Public/WorldPartition/WorldPartitionSubsystem.h
@@ -170,11 +170,15 @@ private:
 	void OnLevelBeginMakingVisible(UWorld* InWorld, const ULevelStreaming* InStreamingLevel, ULevel* InLoadedLevel);
 	void OnLevelBeginMakingInvisible(UWorld* InWorld, const ULevelStreaming* InStreamingLevel, ULevel* InLoadedLevel);
 	void UpdateLoadingAndPendingLoadStreamingLevels(const ULevelStreaming* InStreamingLevel);
+	
+	// Server information
+	uint32 GetServerClientsVisibleLevelsHash() const { return ServerClientsVisibleLevelsHash; }
+	void UpdateServerClientsVisibleLevelNames();
 
 	static ENGINE_API void UpdateStreamingStateInternal(const UWorld* InWorld, const TArray<TObjectPtr<UWorldPartition>>& InWorldPartitions);
 	static int32 GetMaxCellsToLoad(const UWorld* InWorld);
+	static bool IsServer(const UWorld* InWorld);
 
-	ENGINE_API bool IsServer() const;
 	ENGINE_API bool HasAnyWorldPartitionServerStreamingEnabled() const;
 	ENGINE_API bool HasUninitializationPendingStreamingLevels(const UWorldPartition* InWorldPartition) const;
 
@@ -196,7 +200,11 @@ private:
 	TArray<FWorldPartitionStreamingSource> StreamingSources;
 	TMap<FName, FStreamingSourceVelocity> StreamingSourcesVelocity;
 	uint32 StreamingSourcesHash;
+
+	// Server information
 	int32 NumWorldPartitionServerStreamingEnabled;
+	TSet<FName> ServerClientsVisibleLevelNames;
+	uint32 ServerClientsVisibleLevelsHash;
 
 	TArray<FWorldPartitionDraw2DContext> WorldPartitionsDraw2DContext;
 	FDelegateHandle	DrawHandle;
