Moving material expressions out from Interchange/MaterialX into Engine:
- Exponential
- Logarithm (natural)
- Length
- RgbToHsv
- HsvToRgb

Adding redirects for material expressions going from Interchange/MaterialX to Engine

Adding Material Expression:
- Exponential2

-Add a tag 'private' to hide material expression in the material editor

diff --git a/Engine/Config/BaseEngine.ini b/Engine/Config/BaseEngine.ini
--- a/Engine/Config/BaseEngine.ini
+++ b/Engine/Config/BaseEngine.ini
@@ -1164,7 +1164,12 @@ DefaultWorldPartitionHLODLayer=/Engine/Maps/Templates/HLODs/HLODLayer_Instanced.
 +ClassRedirects=(OldName="XRLoadingScreenFunctionLibrary", NewName="/Script/XRBase.XRLoadingScreenFunctionLibrary")
 
 
-
+; Moving some MaterialX Material Expressions into Engine
++ClassRedirects=(OldName="/Script/InterchangeImport.MaterialExpressionMaterialXExponential", NewName="/Script/Engine.MaterialExpressionExponential")
++ClassRedirects=(OldName="/Script/InterchangeImport.MaterialExpressionMaterialXHsvToRgb", NewName="/Script/Engine.MaterialExpressionHsvToRgb")
++ClassRedirects=(OldName="/Script/InterchangeImport.MaterialExpressionMaterialXLength", NewName="/Script/Engine.MaterialExpressionLength")
++ClassRedirects=(OldName="/Script/InterchangeImport.MaterialExpressionMaterialXLogarithm", NewName="/Script/Engine.MaterialExpressionLogarithm")
++ClassRedirects=(OldName="/Script/InterchangeImport.MaterialExpressionMaterialXRgbToHsv", NewName="/Script/Engine.MaterialExpressionRgbToHsv")
 [CoreUObject.Metadata]
 ; Note: UnrealHeaderTool should be rerun after making changes to MetadataRedirects to catch any keys specified in class headers
 +MetadataRedirects=(OldKey="K2Protected", NewKey="BlueprintProtected")
diff --git a/Engine/Plugins/Interchange/Runtime/Source/Import/Private/MaterialX/MaterialExpressions/MaterialExpressionExponential.cpp b/Engine/Plugins/Interchange/Runtime/Source/Import/Private/MaterialX/MaterialExpressions/MaterialExpressionExponential.cpp
deleted file mode 100644
diff --git a/Engine/Plugins/Interchange/Runtime/Source/Import/Private/MaterialX/MaterialExpressions/MaterialExpressionExponential.h b/Engine/Plugins/Interchange/Runtime/Source/Import/Private/MaterialX/MaterialExpressions/MaterialExpressionExponential.h
deleted file mode 100644
diff --git a/Engine/Plugins/Interchange/Runtime/Source/Import/Private/MaterialX/MaterialExpressions/MaterialExpressionHsvToRgb.cpp b/Engine/Plugins/Interchange/Runtime/Source/Import/Private/MaterialX/MaterialExpressions/MaterialExpressionHsvToRgb.cpp
deleted file mode 100644
diff --git a/Engine/Plugins/Interchange/Runtime/Source/Import/Private/MaterialX/MaterialExpressions/MaterialExpressionHsvToRgb.h b/Engine/Plugins/Interchange/Runtime/Source/Import/Private/MaterialX/MaterialExpressions/MaterialExpressionHsvToRgb.h
deleted file mode 100644
diff --git a/Engine/Plugins/Interchange/Runtime/Source/Import/Private/MaterialX/MaterialExpressions/MaterialExpressionLength.cpp b/Engine/Plugins/Interchange/Runtime/Source/Import/Private/MaterialX/MaterialExpressions/MaterialExpressionLength.cpp
deleted file mode 100644
diff --git a/Engine/Plugins/Interchange/Runtime/Source/Import/Private/MaterialX/MaterialExpressions/MaterialExpressionLength.h b/Engine/Plugins/Interchange/Runtime/Source/Import/Private/MaterialX/MaterialExpressions/MaterialExpressionLength.h
deleted file mode 100644
diff --git a/Engine/Plugins/Interchange/Runtime/Source/Import/Private/MaterialX/MaterialExpressions/MaterialExpressionLogarithm.cpp b/Engine/Plugins/Interchange/Runtime/Source/Import/Private/MaterialX/MaterialExpressions/MaterialExpressionLogarithm.cpp
deleted file mode 100644
diff --git a/Engine/Plugins/Interchange/Runtime/Source/Import/Private/MaterialX/MaterialExpressions/MaterialExpressionLogarithm.h b/Engine/Plugins/Interchange/Runtime/Source/Import/Private/MaterialX/MaterialExpressions/MaterialExpressionLogarithm.h
deleted file mode 100644
diff --git a/Engine/Plugins/Interchange/Runtime/Source/Import/Private/MaterialX/MaterialExpressions/MaterialExpressionRgbToHsv.cpp b/Engine/Plugins/Interchange/Runtime/Source/Import/Private/MaterialX/MaterialExpressions/MaterialExpressionRgbToHsv.cpp
deleted file mode 100644
diff --git a/Engine/Plugins/Interchange/Runtime/Source/Import/Private/MaterialX/MaterialExpressions/MaterialExpressionRgbToHsv.h b/Engine/Plugins/Interchange/Runtime/Source/Import/Private/MaterialX/MaterialExpressions/MaterialExpressionRgbToHsv.h
deleted file mode 100644
diff --git a/Engine/Source/Editor/MaterialEditor/Private/MaterialExpressionClasses.cpp b/Engine/Source/Editor/MaterialEditor/Private/MaterialExpressionClasses.cpp
--- a/Engine/Source/Editor/MaterialEditor/Private/MaterialExpressionClasses.cpp
+++ b/Engine/Source/Editor/MaterialEditor/Private/MaterialExpressionClasses.cpp
@@ -123,6 +123,12 @@ void MaterialExpressionClasses::InitMaterialExpressionClasses()
 							continue;
 						}
 
+						// Hide node types that are tagged private
+						if(Class->HasMetaData(TEXT("Private")))
+						{
+							continue;
+						}
+
 						// Exclude comments from the expression list, as well as the base parameter expression and local variable usage, as they should not be used directly
 						if ( Class != UMaterialExpressionComment::StaticClass() 
 							&& Class != UMaterialExpressionParameter::StaticClass()
diff --git a/Engine/Source/Runtime/Engine/Classes/Materials/MaterialExpressionExponential.h b/Engine/Source/Runtime/Engine/Classes/Materials/MaterialExpressionExponential.h
new file mode 100644
--- /dev/null
+++ b/Engine/Source/Runtime/Engine/Classes/Materials/MaterialExpressionExponential.h
@@ -0,0 +1,27 @@
+// Copyright Epic Games, Inc. All Rights Reserved.
+
+
+#pragma once
+
+#include "Materials/MaterialExpression.h"
+#include "MaterialExpressionExponential.generated.h"
+
+UCLASS(collapsecategories, hidecategories = Object, MinimalAPI)
+class UMaterialExpressionExponential : public UMaterialExpression
+{
+	GENERATED_UCLASS_BODY()
+
+	UPROPERTY()
+	FExpressionInput Input;
+
+	//~ Begin UMaterialExpression Interface
+#if WITH_EDITOR
+	virtual int32 Compile(class FMaterialCompiler* Compiler, int32 OutputIndex) override;
+	virtual void GetCaption(TArray<FString>& OutCaptions) const override;
+	virtual void GetExpressionToolTip(TArray<FString>& OutToolTip) override;
+	virtual FText GetKeywords() const override {return FText::FromString(TEXT("exp"));}
+	virtual bool GenerateHLSLExpression(FMaterialHLSLGenerator& Generator, UE::HLSLTree::FScope& Scope, int32 OutputIndex, UE::HLSLTree::FExpression const*& OutExpression) const override;
+#endif
+	//~ End UMaterialExpression Interface
+
+};
diff --git a/Engine/Source/Runtime/Engine/Classes/Materials/MaterialExpressionExponential2.h b/Engine/Source/Runtime/Engine/Classes/Materials/MaterialExpressionExponential2.h
new file mode 100644
--- /dev/null
+++ b/Engine/Source/Runtime/Engine/Classes/Materials/MaterialExpressionExponential2.h
@@ -0,0 +1,27 @@
+// Copyright Epic Games, Inc. All Rights Reserved.
+
+
+#pragma once
+
+#include "Materials/MaterialExpression.h"
+#include "MaterialExpressionExponential2.generated.h"
+
+UCLASS(collapsecategories, hidecategories = Object, MinimalAPI)
+class UMaterialExpressionExponential2 : public UMaterialExpression
+{
+	GENERATED_UCLASS_BODY()
+
+	UPROPERTY()
+	FExpressionInput Input;
+
+	//~ Begin UMaterialExpression Interface
+#if WITH_EDITOR
+	virtual int32 Compile(class FMaterialCompiler* Compiler, int32 OutputIndex) override;
+	virtual void GetCaption(TArray<FString>& OutCaptions) const override;
+	virtual void GetExpressionToolTip(TArray<FString>& OutToolTip) override;
+	virtual FText GetKeywords() const override	{return FText::FromString(TEXT("exp2"));}
+	virtual bool GenerateHLSLExpression(FMaterialHLSLGenerator& Generator, UE::HLSLTree::FScope& Scope, int32 OutputIndex, UE::HLSLTree::FExpression const*& OutExpression) const override;
+#endif
+	//~ End UMaterialExpression Interface
+
+};
diff --git a/Engine/Source/Runtime/Engine/Classes/Materials/MaterialExpressionHsvToRgb.h b/Engine/Source/Runtime/Engine/Classes/Materials/MaterialExpressionHsvToRgb.h
new file mode 100644
--- /dev/null
+++ b/Engine/Source/Runtime/Engine/Classes/Materials/MaterialExpressionHsvToRgb.h
@@ -0,0 +1,24 @@
+// Copyright Epic Games, Inc. All Rights Reserved.
+#pragma once
+
+#include "Materials/MaterialExpression.h"
+#include "MaterialExpressionHsvToRgb.generated.h"
+
+UCLASS(collapsecategories, hidecategories = Object, MinimalAPI)
+class UMaterialExpressionHsvToRgb : public UMaterialExpression
+{
+	GENERATED_UCLASS_BODY()
+
+	UPROPERTY()
+	FExpressionInput Input;
+
+	//~ Begin UMaterialExpressionMaterialX Interface
+#if WITH_EDITOR
+	virtual int32 Compile(class FMaterialCompiler* Compiler, int32 OutputIndex) override;
+	virtual void GetCaption(TArray<FString>& OutCaptions) const override;
+	virtual void GetExpressionToolTip(TArray<FString>& OutToolTip) override;
+	virtual FText GetKeywords() const override {return FText::FromString(TEXT("hsvtorgb"));}
+#endif
+	//~ End UMaterialExpressionMaterialX Interface
+};
+
diff --git a/Engine/Source/Runtime/Engine/Classes/Materials/MaterialExpressionLength.h b/Engine/Source/Runtime/Engine/Classes/Materials/MaterialExpressionLength.h
new file mode 100644
--- /dev/null
+++ b/Engine/Source/Runtime/Engine/Classes/Materials/MaterialExpressionLength.h
@@ -0,0 +1,25 @@
+// Copyright Epic Games, Inc. All Rights Reserved.
+#pragma once
+
+#include "Materials/MaterialExpression.h"
+#include "MaterialExpressionLength.generated.h"
+
+UCLASS(collapsecategories, hidecategories = Object, MinimalAPI)
+class UMaterialExpressionLength : public UMaterialExpression
+{
+	GENERATED_UCLASS_BODY()
+
+	UPROPERTY()
+	FExpressionInput Input;
+
+	//~ Begin UMaterialExpressionMaterialX Interface
+#if WITH_EDITOR
+	virtual int32 Compile(class FMaterialCompiler* Compiler, int32 OutputIndex) override;
+	virtual void GetCaption(TArray<FString>& OutCaptions) const override;
+	virtual void GetExpressionToolTip(TArray<FString>& OutToolTip) override;
+	virtual FText GetKeywords() const override {return FText::FromString(TEXT("length"));}
+	virtual bool GenerateHLSLExpression(FMaterialHLSLGenerator& Generator, UE::HLSLTree::FScope& Scope, int32 OutputIndex, UE::HLSLTree::FExpression const*& OutExpression) const override;
+#endif
+	//~ End UMaterialExpressionMaterialX Interface
+};
+
diff --git a/Engine/Source/Runtime/Engine/Classes/Materials/MaterialExpressionLogarithm.h b/Engine/Source/Runtime/Engine/Classes/Materials/MaterialExpressionLogarithm.h
new file mode 100644
--- /dev/null
+++ b/Engine/Source/Runtime/Engine/Classes/Materials/MaterialExpressionLogarithm.h
@@ -0,0 +1,27 @@
+// Copyright Epic Games, Inc. All Rights Reserved.
+
+
+#pragma once
+
+#include "Materials/MaterialExpression.h"
+#include "MaterialExpressionLogarithm.generated.h"
+
+UCLASS(collapsecategories, hidecategories = Object, MinimalAPI)
+class UMaterialExpressionLogarithm : public UMaterialExpression
+{
+	GENERATED_UCLASS_BODY()
+
+	UPROPERTY()
+	FExpressionInput Input;
+
+	//~ Begin UMaterialExpression Interface
+#if WITH_EDITOR
+	virtual int32 Compile(class FMaterialCompiler* Compiler, int32 OutputIndex) override;
+	virtual void GetCaption(TArray<FString>& OutCaptions) const override;
+	virtual void GetExpressionToolTip(TArray<FString>& OutToolTip) override;
+	virtual FText GetKeywords() const override {return FText::FromString(TEXT("log2"));}
+	virtual bool GenerateHLSLExpression(FMaterialHLSLGenerator& Generator, UE::HLSLTree::FScope& Scope, int32 OutputIndex, UE::HLSLTree::FExpression const*& OutExpression) const override;
+#endif
+	//~ End UMaterialExpression Interface
+
+};
diff --git a/Engine/Source/Runtime/Engine/Classes/Materials/MaterialExpressionRgbToHsv.h b/Engine/Source/Runtime/Engine/Classes/Materials/MaterialExpressionRgbToHsv.h
new file mode 100644
--- /dev/null
+++ b/Engine/Source/Runtime/Engine/Classes/Materials/MaterialExpressionRgbToHsv.h
@@ -0,0 +1,24 @@
+// Copyright Epic Games, Inc. All Rights Reserved.
+#pragma once
+
+#include "Materials/MaterialExpression.h"
+#include "MaterialExpressionRgbToHsv.generated.h"
+
+UCLASS(collapsecategories, hidecategories = Object, MinimalAPI)
+class UMaterialExpressionRgbToHsv : public UMaterialExpression
+{
+	GENERATED_UCLASS_BODY()
+
+	UPROPERTY()
+	FExpressionInput Input;
+
+	//~ Begin UMaterialExpressionMaterialX Interface
+#if WITH_EDITOR
+	virtual int32 Compile(class FMaterialCompiler* Compiler, int32 OutputIndex) override;
+	virtual void GetCaption(TArray<FString>& OutCaptions) const override;
+	virtual void GetExpressionToolTip(TArray<FString>& OutToolTip) override;
+	virtual FText GetKeywords() const override {return FText::FromString(TEXT("rgbtohsv"));}
+#endif
+	//~ End UMaterialExpressionMaterialX Interface
+};
+
diff --git a/Engine/Source/Runtime/Engine/Private/HLSLTree/HLSLTreeOperation.cpp b/Engine/Source/Runtime/Engine/Private/HLSLTree/HLSLTreeOperation.cpp
--- a/Engine/Source/Runtime/Engine/Private/HLSLTree/HLSLTreeOperation.cpp
+++ b/Engine/Source/Runtime/Engine/Private/HLSLTree/HLSLTreeOperation.cpp
@@ -46,8 +46,10 @@ FOperationDescription GetOperationDescription(EOperation Op)
 	case EOperation::Rcp: return FOperationDescription(TEXT("Rcp"), TEXT("/"), 1, Shader::EPreshaderOpcode::Rcp); break;
 	case EOperation::Sqrt: return FOperationDescription(TEXT("Sqrt"), TEXT("sqrt"), 1, Shader::EPreshaderOpcode::Sqrt); break;
 	case EOperation::Rsqrt: return FOperationDescription(TEXT("Rsqrt"), TEXT("rsqrt"), 1, Shader::EPreshaderOpcode::Nop); break; // TODO
+	case EOperation::Log: return FOperationDescription(TEXT("Log"), TEXT("log"), 1, Shader::EPreshaderOpcode::Log); break;
 	case EOperation::Log2: return FOperationDescription(TEXT("Log2"), TEXT("log2"), 1, Shader::EPreshaderOpcode::Log2); break;
-	case EOperation::Exp2: return FOperationDescription(TEXT("Exp2"), TEXT("exp2"), 1, Shader::EPreshaderOpcode::Nop); break; // TODO
+	case EOperation::Exp: return FOperationDescription(TEXT("Exp"), TEXT("exp"), 1, Shader::EPreshaderOpcode::Exp); break;
+	case EOperation::Exp2: return FOperationDescription(TEXT("Exp2"), TEXT("exp2"), 1, Shader::EPreshaderOpcode::Exp2); break; // TODO
 	case EOperation::Frac: return FOperationDescription(TEXT("Frac"), TEXT("frac"), 1, Shader::EPreshaderOpcode::Frac); break;
 	case EOperation::Floor: return FOperationDescription(TEXT("Floor"), TEXT("floor"), 1, Shader::EPreshaderOpcode::Floor); break;
 	case EOperation::Ceil: return FOperationDescription(TEXT("Ceil"), TEXT("ceil"), 1, Shader::EPreshaderOpcode::Ceil); break;
@@ -329,7 +331,9 @@ FOperationTypes GetOperationTypes(EOperation Op, TConstArrayView<FPreparedType>
 			}
 			Types.ResultType = MakeNonLWCType(IntermediateType);
 			break;
+		case EOperation::Log:
 		case EOperation::Log2:
+		case EOperation::Exp:
 		case EOperation::Exp2:
 			// No LWC support yet
 			Types.InputType[0] = Shader::MakeNonLWCType(Types.InputType[0]);
@@ -520,6 +524,13 @@ void FExpressionOperation::ComputeAnalyticDerivatives(FTree& Tree, FExpressionDe
 		OutResult.ExpressionDdy = Tree.NewExpression<FExpressionSelect>(Cond0, Tree.NewExpression<FExpressionSelect>(Cond1, InputDerivatives[0].ExpressionDdy, ConstantOne), ConstantZero);
 		break;
 	}
+	case EOperation::Log:
+	{
+		const FExpression* dFdA = Tree.NewRcp(Inputs[0]);
+		OutResult.ExpressionDdx = Tree.NewMul(dFdA, InputDerivatives[0].ExpressionDdx);
+		OutResult.ExpressionDdy = Tree.NewMul(dFdA, InputDerivatives[0].ExpressionDdy);
+		break;
+	}
 	case EOperation::Log2:
 	{
 		const FExpression* dFdA = Tree.NewMul(Tree.NewRcp(Inputs[0]), Tree.NewConstant(1.442695f));
@@ -527,6 +538,13 @@ void FExpressionOperation::ComputeAnalyticDerivatives(FTree& Tree, FExpressionDe
 		OutResult.ExpressionDdy = Tree.NewMul(dFdA, InputDerivatives[0].ExpressionDdy);
 		break;
 	}
+	case EOperation::Exp:
+	{
+		const FExpression* dFdA = Tree.NewExp(Inputs[0]);
+		OutResult.ExpressionDdx = Tree.NewMul(dFdA, InputDerivatives[0].ExpressionDdx);
+		OutResult.ExpressionDdy = Tree.NewMul(dFdA, InputDerivatives[0].ExpressionDdy);
+		break;
+	}
 	case EOperation::Exp2:
 	{
 		const FExpression* dFdA = Tree.NewMul(Tree.NewExp2(Inputs[0]), Tree.NewConstant(0.693147f));
@@ -753,7 +771,9 @@ void FExpressionOperation::EmitValueShader(FEmitContext& Context, FEmitScope& Sc
 	case EOperation::Rcp: OutResult.Code = Context.EmitExpression(Scope, ResultType, Types.bIsLWC ? TEXT("LWCRcp(%)") : TEXT("rcp(%)"), InputValue[0]); break;
 	case EOperation::Sqrt: OutResult.Code = Context.EmitExpression(Scope, ResultType, Types.bIsLWC ? TEXT("LWCSqrt(%)") : TEXT("sqrt(%)"), InputValue[0]); break;
 	case EOperation::Rsqrt: OutResult.Code = Context.EmitExpression(Scope, ResultType, Types.bIsLWC ? TEXT("LWCRsqrt(%)") : TEXT("rsqrt(%)"), InputValue[0]); break;
+	case EOperation::Log: OutResult.Code = Context.EmitExpression(Scope, ResultType, TEXT("log(%)"), InputValue[0]); break;
 	case EOperation::Log2: OutResult.Code = Context.EmitExpression(Scope, ResultType, TEXT("log2(%)"), InputValue[0]); break;
+	case EOperation::Exp: OutResult.Code = Context.EmitExpression(Scope, ResultType, TEXT("exp(%)"), InputValue[0]); break;
 	case EOperation::Exp2: OutResult.Code = Context.EmitExpression(Scope, ResultType, TEXT("exp2(%)"), InputValue[0]); break;
 	case EOperation::Frac: OutResult.Code = Context.EmitExpression(Scope, ResultType, Types.bIsLWC ? TEXT("LWCFrac(%)") : TEXT("frac(%)"), InputValue[0]); break;
 	case EOperation::Floor: OutResult.Code = Context.EmitExpression(Scope, ResultType, Types.bIsLWC ? TEXT("LWCFloor(%)") : TEXT("floor(%)"), InputValue[0]); break;
diff --git a/Engine/Source/Runtime/Engine/Private/Materials/HLSLMaterialDerivativeAutogen.cpp b/Engine/Source/Runtime/Engine/Private/Materials/HLSLMaterialDerivativeAutogen.cpp
--- a/Engine/Source/Runtime/Engine/Private/Materials/HLSLMaterialDerivativeAutogen.cpp
+++ b/Engine/Source/Runtime/Engine/Private/Materials/HLSLMaterialDerivativeAutogen.cpp
@@ -337,6 +337,7 @@ FMaterialDerivativeAutogen::FOperationType1 FMaterialDerivativeAutogen::GetFunc1
 	case EFunc1::InvLength:
 		// Operate on LWCType, but return a float1
 		return FOperationType1(EDerivativeType::Float1, SrcType);
+	case EFunc1::Log:
 	case EFunc1::Log2:
 	case EFunc1::Log10:
 	case EFunc1::Exp:
@@ -492,6 +493,9 @@ int32 FMaterialDerivativeAutogen::GenerateExpressionFunc1(FHLSLMaterialTranslato
 				DstToken = TEXT("abs(") + SrcToken + TEXT(")");
 			}
 			break;
+		case EFunc1::Log:
+			DstToken = TEXT("log(") + SrcToken + TEXT(")");
+			break;
 		case EFunc1::Log2:
 			DstToken = TEXT("log2(") + SrcToken + TEXT(")");
 			break;
@@ -501,6 +505,9 @@ int32 FMaterialDerivativeAutogen::GenerateExpressionFunc1(FHLSLMaterialTranslato
 		case EFunc1::Exp:
 			DstToken = TEXT("exp(") + SrcToken + TEXT(")");
 			break;
+		case EFunc1::Exp2:
+			DstToken = TEXT("exp2(") + SrcToken + TEXT(")");
+			break;
 		case EFunc1::Sin:
 			if (bIsLWC)
 			{
@@ -720,6 +727,9 @@ int32 FMaterialDerivativeAutogen::GenerateExpressionFunc1(FHLSLMaterialTranslato
 		case EFunc1::Abs:
 			DstToken = TEXT("AbsDeriv") + Suffix + TEXT("(") + SrcToken + TEXT(")");
 			break;
+		case EFunc1::Log:
+			DstToken = TEXT("LogDeriv") + Suffix + TEXT("(") + SrcToken + TEXT(")");
+			break;
 		case EFunc1::Log2:
 			DstToken = TEXT("Log2Deriv") + Suffix + TEXT("(") + SrcToken + TEXT(")");
 			break;
@@ -729,6 +739,9 @@ int32 FMaterialDerivativeAutogen::GenerateExpressionFunc1(FHLSLMaterialTranslato
 		case EFunc1::Exp:
 			DstToken = TEXT("ExpDeriv") + Suffix + TEXT("(") + SrcToken + TEXT(")");
 			break;
+		case EFunc1::Exp2:
+			DstToken = TEXT("Exp2Deriv") + Suffix + TEXT("(") + SrcToken + TEXT(")");
+			break;
 		case EFunc1::Sin:
 			DstToken = TEXT("SinDeriv") + Suffix + TEXT("(") + SrcToken + TEXT(")");
 			break;
@@ -2065,6 +2078,18 @@ FString FMaterialDerivativeAutogen::GenerateUsedFunctions(FHLSLMaterialTranslato
 					Ret += TEXT("}") LINE_TERMINATOR;
 					Ret += TEXT("") LINE_TERMINATOR;
 					break;
+				case EFunc1::Log:
+					Ret += BaseName + TEXT(" LogDeriv") + Suffix + TEXT("(") + BaseName + TEXT(" A)") LINE_TERMINATOR;
+					Ret += TEXT("{") LINE_TERMINATOR;
+					Ret += TEXT("\t") + BaseName + TEXT(" Ret;") LINE_TERMINATOR;
+					Ret += TEXT("\tRet.Value = log(A.Value);") LINE_TERMINATOR;
+					Ret += TEXT("\t") + FieldName + TEXT(" dFdA = rcp(A.Value);");
+					Ret += TEXT("\tRet.Ddx = dFdA * A.Ddx ;") LINE_TERMINATOR;
+					Ret += TEXT("\tRet.Ddy = dFdA * A.Ddy;") LINE_TERMINATOR;
+					Ret += TEXT("\treturn Ret;") LINE_TERMINATOR;
+					Ret += TEXT("}") LINE_TERMINATOR;
+					Ret += TEXT("") LINE_TERMINATOR;
+					break;
 				case EFunc1::Log2:
 					Ret += BaseName + TEXT(" Log2Deriv") + Suffix + TEXT("(") + BaseName + TEXT(" A)") LINE_TERMINATOR;
 					Ret += TEXT("{") LINE_TERMINATOR;
@@ -2100,6 +2125,18 @@ FString FMaterialDerivativeAutogen::GenerateUsedFunctions(FHLSLMaterialTranslato
 					Ret += TEXT("}") LINE_TERMINATOR;
 					Ret += TEXT("") LINE_TERMINATOR;
 					break;
+				case EFunc1::Exp2:
+					Ret += BaseName + TEXT(" Exp2Deriv") + Suffix + TEXT("(") + BaseName + TEXT(" A)") LINE_TERMINATOR;
+					Ret += TEXT("{") LINE_TERMINATOR;
+					Ret += TEXT("\t") + BaseName + TEXT(" Ret;") LINE_TERMINATOR;
+					Ret += TEXT("\tRet.Value = exp2(A.Value);") LINE_TERMINATOR;
+					Ret += TEXT("\t") + FieldName + TEXT(" dFdA = exp2(A.Value) * 0.693147f;");
+					Ret += TEXT("\tRet.Ddx = dFdA * A.Ddx;") LINE_TERMINATOR;
+					Ret += TEXT("\tRet.Ddy = dFdA * A.Ddy;") LINE_TERMINATOR;
+					Ret += TEXT("\treturn Ret;") LINE_TERMINATOR;
+					Ret += TEXT("}") LINE_TERMINATOR;
+					Ret += TEXT("") LINE_TERMINATOR;
+					break;
 				case EFunc1::Length:
 					Ret += ScalarName + TEXT(" LengthDeriv") + Suffix + TEXT("(") + BaseName + TEXT(" A)") LINE_TERMINATOR;
 					Ret += TEXT("{") LINE_TERMINATOR;
diff --git a/Engine/Source/Runtime/Engine/Private/Materials/HLSLMaterialDerivativeAutogen.h b/Engine/Source/Runtime/Engine/Private/Materials/HLSLMaterialDerivativeAutogen.h
--- a/Engine/Source/Runtime/Engine/Private/Materials/HLSLMaterialDerivativeAutogen.h
+++ b/Engine/Source/Runtime/Engine/Private/Materials/HLSLMaterialDerivativeAutogen.h
@@ -90,9 +90,11 @@ public:
 	enum class EFunc1
 	{
 		Abs,
+		Log,
 		Log2,
 		Log10,
 		Exp,
+		Exp2,
 		Sin,
 		Cos,
 		Tan,
diff --git a/Engine/Source/Runtime/Engine/Private/Materials/HLSLMaterialTranslator.cpp b/Engine/Source/Runtime/Engine/Private/Materials/HLSLMaterialTranslator.cpp
--- a/Engine/Source/Runtime/Engine/Private/Materials/HLSLMaterialTranslator.cpp
+++ b/Engine/Source/Runtime/Engine/Private/Materials/HLSLMaterialTranslator.cpp
@@ -8778,6 +8778,69 @@ int32 FHLSLMaterialTranslator::Power(int32 Base,int32 Exponent)
 		return AddCodeChunk(GetParameterType(Base),TEXT("PositiveClampedPow(%s,%s)"),*GetParameterCode(Base),*CoerceParameter(Exponent,MCT_Float));
 	}
 }
+
+int32 FHLSLMaterialTranslator::Exponential(int32 X)
+{
+	if(X == INDEX_NONE)
+	{
+		return INDEX_NONE;
+	}
+
+	if(GetParameterUniformExpression(X))
+	{
+		return AddUniformExpression(new FMaterialUniformExpressionExponential(GetParameterUniformExpression(X)), GetParameterType(X), TEXT("exp(%s)"), *GetParameterCode(X));
+	}
+	else if(IsAnalyticDerivEnabled())
+	{
+		return DerivativeAutogen.GenerateExpressionFunc1(*this, FMaterialDerivativeAutogen::EFunc1::Exp, X);
+	}
+	else
+	{
+		return AddCodeChunk(GetParameterType(X), TEXT("exp(%s)"), *GetParameterCode(X));
+	}
+}
+
+int32 FHLSLMaterialTranslator::Exponential2(int32 X)
+{
+	if(X == INDEX_NONE)
+	{
+		return INDEX_NONE;
+	}
+
+	if(GetParameterUniformExpression(X))
+	{
+		return AddUniformExpression(new FMaterialUniformExpressionExponential2(GetParameterUniformExpression(X)), GetParameterType(X), TEXT("exp2(%s)"), *GetParameterCode(X));
+	}
+	else if(IsAnalyticDerivEnabled())
+	{
+		return DerivativeAutogen.GenerateExpressionFunc1(*this, FMaterialDerivativeAutogen::EFunc1::Exp2, X);
+	}
+	else
+	{
+		return AddCodeChunk(GetParameterType(X), TEXT("exp2(%s)"), *GetParameterCode(X));
+	}
+}
+
+int32 FHLSLMaterialTranslator::Logarithm(int32 X)
+{
+	if(X == INDEX_NONE)
+	{
+		return INDEX_NONE;
+	}
+
+	if(GetParameterUniformExpression(X))
+	{
+		return AddUniformExpression(new FMaterialUniformExpressionLogarithm(GetParameterUniformExpression(X)), GetParameterType(X), TEXT("log(%s)"), *GetParameterCode(X));
+	}
+	else if(IsAnalyticDerivEnabled())
+	{
+		return DerivativeAutogen.GenerateExpressionFunc1(*this, FMaterialDerivativeAutogen::EFunc1::Log, X);
+	}
+	else
+	{
+		return AddCodeChunk(GetParameterType(X), TEXT("log(%s)"), *GetParameterCode(X));
+	}
+}
 	
 int32 FHLSLMaterialTranslator::Logarithm2(int32 X)
 {
@@ -9670,6 +9733,80 @@ int32 FHLSLMaterialTranslator::AppendVector(int32 A,int32 B)
 	}
 }
 
+int32 FHLSLMaterialTranslator::HsvToRgb(int32 X)
+{
+	if(X == INDEX_NONE)
+	{
+		return INDEX_NONE;
+	}
+
+	FString HSV = CreateSymbolName(TEXT("HSV"));
+	FString H = CreateSymbolName(TEXT("H"));
+	FString R = CreateSymbolName(TEXT("R"));
+	FString G = CreateSymbolName(TEXT("G"));
+	FString B = CreateSymbolName(TEXT("B"));
+	FString RGB = CreateSymbolName(TEXT("RGB"));
+
+	EMaterialValueType ValueType = GetType(X);
+	FString Type = HLSLTypeString(ValueType);
+	FString CodeStr{ Type + TEXT(" ") + HSV + TEXT(" = %s; ") +
+	TEXT("float ") + H + TEXT(" = ") + HSV + TEXT(".r; ") +
+	TEXT("float ") + R + TEXT(" = abs(") + H + TEXT(" * 6 - 3) - 1; ") +
+	TEXT("float ") + G + TEXT(" = 2 - abs(") + H + TEXT(" * 6 - 2); ") +
+	TEXT("float ") + B + TEXT(" = 2 - abs(") + H + TEXT(" * 6 - 4); ") +
+	TEXT("float3 ") + RGB + TEXT(" = saturate(float3(") + R + TEXT(",") + G + TEXT(",") + B + TEXT(")); ") };
+
+	FString Code{ TEXT("((") + RGB + TEXT(" - 1) * ") + HSV + TEXT(".y + 1) * ") + HSV + TEXT(".z") };
+
+
+	if(ValueType == EMaterialValueType::MCT_Float4)
+	{
+		Code = Type + TEXT("(") + Code + TEXT(", ") + HSV + TEXT(".w);");
+	}
+
+	AddInlinedCodeChunk(EMaterialValueType::MCT_VoidStatement, *CodeStr, *GetParameterCode(X));
+	return AddCodeChunk(GetParameterType(X), *Code);
+}
+
+int32 FHLSLMaterialTranslator::RgbToHsv(int32 X)
+{
+	if(X == INDEX_NONE)
+	{
+		return INDEX_NONE;
+	}
+
+	FString RGB = CreateSymbolName(TEXT("RGB"));
+	FString P = CreateSymbolName(TEXT("P"));
+	FString Q = CreateSymbolName(TEXT("Q"));
+	FString Chroma = CreateSymbolName(TEXT("Chroma"));
+	FString Hue = CreateSymbolName(TEXT("Hue"));
+	FString HCV = CreateSymbolName(TEXT("HCV"));
+	FString S = CreateSymbolName(TEXT("s"));
+
+	EMaterialValueType ValueType = GetType(X);
+
+	FString Type = HLSLTypeString(ValueType);
+	FString CodeStr{ Type + TEXT(" ") + RGB + TEXT("= %s;") +
+	TEXT("float4 ") + P + TEXT(" = (") + RGB + TEXT(".g < ") + RGB + TEXT(".b) ? float4(") + RGB + TEXT(".bg, -1.0f, 2.0f / 3.0f) : float4(") + RGB + TEXT(".gb, 0.0f, -1.0f / 3.0f); ") +
+	TEXT("float4 ") + Q + TEXT(" = (") + RGB + TEXT(".r < ") + P + TEXT(".x) ? float4(") + P + TEXT(".xyw, ") + RGB + TEXT(".r)	: float4(") + RGB + TEXT(".r, ") + P + TEXT(".yzx);") +
+	TEXT("float ") + Chroma + TEXT("= ") + Q + TEXT(".x - min(") + Q + TEXT(".w, ") + Q + TEXT(".y);") +
+	TEXT("float ") + Hue + TEXT(" = abs((") + Q + TEXT(".w - ") + Q + TEXT(".y) / (6.0f * ") + Chroma + TEXT(" + 1e-10f) + ") + Q + TEXT(".z);") +
+	TEXT("float3 ") + HCV + TEXT(" = float3(") + Hue + TEXT(", ") + Chroma + TEXT(", ") + Q + TEXT(".x);") +
+	TEXT("float ") + S + TEXT(" = ") + HCV + TEXT(".y / (") + HCV + TEXT(".z + 1e-10f);") };
+
+	FString Code{ Type + TEXT("(") + HCV + TEXT(".x, ") + S + TEXT(", ") + HCV + TEXT(".z") };
+
+	if(ValueType == EMaterialValueType::MCT_Float4)
+	{
+		Code += TEXT(",") + RGB + TEXT(".w");
+	}
+
+	Code += TEXT(");");
+
+	AddInlinedCodeChunk(EMaterialValueType::MCT_VoidStatement, *CodeStr, *GetParameterCode(X));
+	return AddCodeChunk(GetParameterType(X), *Code);
+}
+
 static FString MultiplyMatrix(const TCHAR* Vector, const TCHAR* Matrix, int AWComponent)
 {
 	if (AWComponent)
diff --git a/Engine/Source/Runtime/Engine/Private/Materials/HLSLMaterialTranslator.h b/Engine/Source/Runtime/Engine/Private/Materials/HLSLMaterialTranslator.h
--- a/Engine/Source/Runtime/Engine/Private/Materials/HLSLMaterialTranslator.h
+++ b/Engine/Source/Runtime/Engine/Private/Materials/HLSLMaterialTranslator.h
@@ -1056,6 +1056,9 @@ protected:
 	virtual int32 Dot(int32 A, int32 B) override;
 	virtual int32 Cross(int32 A, int32 B) override;
 	virtual int32 Power(int32 Base, int32 Exponent) override;
+	virtual int32 Exponential(int32 X) override;
+	virtual int32 Exponential2(int32 X) override;
+	virtual int32 Logarithm(int32 X) override;
 	virtual int32 Logarithm2(int32 X) override;
 	virtual int32 Logarithm10(int32 X) override;
 	virtual int32 SquareRoot(int32 X) override;
@@ -1071,6 +1074,8 @@ protected:
 	virtual int32 Saturate(int32 X) override;
 	virtual int32 ComponentMask(int32 Vector, bool R, bool G, bool B, bool A) override;
 	virtual int32 AppendVector(int32 A, int32 B) override;
+	virtual int32 HsvToRgb(int32 X) override;
+	virtual int32 RgbToHsv(int32 X) override;
 
 	int32 TransformBase(EMaterialCommonBasis SourceCoordBasis, EMaterialCommonBasis DestCoordBasis, int32 A, int AWComponent);
 	
diff --git a/Engine/Source/Runtime/Engine/Private/Materials/MaterialExpressionHLSL.cpp b/Engine/Source/Runtime/Engine/Private/Materials/MaterialExpressionHLSL.cpp
--- a/Engine/Source/Runtime/Engine/Private/Materials/MaterialExpressionHLSL.cpp
+++ b/Engine/Source/Runtime/Engine/Private/Materials/MaterialExpressionHLSL.cpp
@@ -77,6 +77,10 @@
 #include "Materials/MaterialExpressionOneMinus.h"
 #include "Materials/MaterialExpressionAbs.h"
 #include "Materials/MaterialExpressionSquareRoot.h"
+#include "Materials/MaterialExpressionExponential.h"
+#include "Materials/MaterialExpressionExponential2.h"
+#include "Materials/MaterialExpressionLength.h"
+#include "Materials/MaterialExpressionLogarithm.h"
 #include "Materials/MaterialExpressionLogarithm2.h"
 #include "Materials/MaterialExpressionLogarithm10.h"
 #include "Materials/MaterialExpressionFrac.h"
@@ -1166,6 +1170,53 @@ bool UMaterialExpressionSquareRoot::GenerateHLSLExpression(FMaterialHLSLGenerato
 	return true;
 }
 
+bool UMaterialExpressionExponential::GenerateHLSLExpression(FMaterialHLSLGenerator& Generator, UE::HLSLTree::FScope& Scope, int32 OutputIndex, UE::HLSLTree::FExpression const*& OutExpression) const
+{
+	using namespace UE::HLSLTree;
+	const FExpression* InputExpression = Input.AcquireHLSLExpression(Generator, Scope);
+	if(!InputExpression)
+	{
+		return false;
+	}
+	OutExpression = Generator.GetTree().NewExp(InputExpression);
+	return true;
+}
+
+bool UMaterialExpressionExponential2::GenerateHLSLExpression(FMaterialHLSLGenerator& Generator, UE::HLSLTree::FScope& Scope, int32 OutputIndex, UE::HLSLTree::FExpression const*& OutExpression) const
+{
+	using namespace UE::HLSLTree;
+	const FExpression* InputExpression = Input.AcquireHLSLExpression(Generator, Scope);
+	if(!InputExpression)
+	{
+		return false;
+	}
+	OutExpression = Generator.GetTree().NewExp2(InputExpression);
+	return true;
+}
+
+bool UMaterialExpressionLength::GenerateHLSLExpression(FMaterialHLSLGenerator& Generator, UE::HLSLTree::FScope& Scope, int32 OutputIndex, UE::HLSLTree::FExpression const*& OutExpression) const
+{
+	const UE::HLSLTree::FExpression* InputExpression = Input.AcquireHLSLExpression(Generator, Scope);
+	if(!InputExpression)
+	{
+		return false;
+	}
+	OutExpression = Generator.GetTree().NewLength(InputExpression);
+	return true;
+}
+
+bool UMaterialExpressionLogarithm::GenerateHLSLExpression(FMaterialHLSLGenerator& Generator, UE::HLSLTree::FScope& Scope, int32 OutputIndex, UE::HLSLTree::FExpression const*& OutExpression) const
+{
+	using namespace UE::HLSLTree;
+	const FExpression* InputExpression = Input.AcquireHLSLExpression(Generator, Scope);
+	if(!InputExpression)
+	{
+		return false;
+	}
+	OutExpression = Generator.GetTree().NewUnaryOp(EOperation::Log, InputExpression);
+	return true;
+}
+
 bool UMaterialExpressionLogarithm2::GenerateHLSLExpression(FMaterialHLSLGenerator& Generator, UE::HLSLTree::FScope& Scope, int32 OutputIndex, UE::HLSLTree::FExpression const*& OutExpression) const
 {
 	using namespace UE::HLSLTree;
diff --git a/Engine/Source/Runtime/Engine/Private/Materials/MaterialExpressions.cpp b/Engine/Source/Runtime/Engine/Private/Materials/MaterialExpressions.cpp
--- a/Engine/Source/Runtime/Engine/Private/Materials/MaterialExpressions.cpp
+++ b/Engine/Source/Runtime/Engine/Private/Materials/MaterialExpressions.cpp
@@ -124,6 +124,12 @@
 #include "Materials/MaterialExpressionLightmassReplace.h"
 #include "Materials/MaterialExpressionLightVector.h"
 #include "Materials/MaterialExpressionLinearInterpolate.h"
+#include "Materials/MaterialExpressionRgbToHsv.h"
+#include "Materials/MaterialExpressionHsvToRgb.h"
+#include "Materials/MaterialExpressionExponential.h"
+#include "Materials/MaterialExpressionExponential2.h"
+#include "Materials/MaterialExpressionLength.h"
+#include "Materials/MaterialExpressionLogarithm.h"
 #include "Materials/MaterialExpressionLogarithm2.h"
 #include "Materials/MaterialExpressionLogarithm10.h"
 #include "Materials/MaterialExpressionMakeMaterialAttributes.h"
@@ -11236,6 +11242,271 @@ void UMaterialExpressionPower::GetExpressionToolTip(TArray<FString>& OutToolTip)
 }
 #endif // WITH_EDITOR
 
+UMaterialExpressionLength::UMaterialExpressionLength(const FObjectInitializer& ObjectInitializer)
+	: Super(ObjectInitializer)
+{
+	// Structure to hold one-time initialization
+	struct FConstructorStatics
+	{
+		FText NAME_Utility;
+		FConstructorStatics()
+			: NAME_Utility(LOCTEXT("Utility", "Utility"))
+		{}
+	};
+	static FConstructorStatics ConstructorStatics;
+
+#if WITH_EDITORONLY_DATA
+	MenuCategories.Add(ConstructorStatics.NAME_Utility);
+#endif
+}
+
+#if WITH_EDITOR
+int32 UMaterialExpressionLength::Compile(class FMaterialCompiler* Compiler, int32 OutputIndex)
+{
+	if(!Input.GetTracedInput().Expression)
+	{
+		return Compiler->Errorf(TEXT("Missing Length input"));
+	}
+
+	int32 Index = Input.Compile(Compiler);
+	if(Compiler->GetType(Index) == MCT_Float)
+	{
+		// optimized
+		return Compiler->Abs(Index);
+	}
+
+	return Compiler->Length(Index);
+}
+
+void UMaterialExpressionLength::GetCaption(TArray<FString>& OutCaptions) const
+{
+	OutCaptions.Add(TEXT("Length"));
+}
+
+void UMaterialExpressionLength::GetExpressionToolTip(TArray<FString>& OutToolTip)
+{
+	ConvertToMultilineToolTip(TEXT("Returns the length of input."), 40, OutToolTip);
+}
+#endif // WITH_EDITOR
+
+UMaterialExpressionHsvToRgb::UMaterialExpressionHsvToRgb(const FObjectInitializer& ObjectInitializer)
+	: Super(ObjectInitializer)
+{
+	// Structure to hold one-time initialization
+	struct FConstructorStatics
+	{
+		FText NAME_Math;
+		FConstructorStatics()
+			: NAME_Math(LOCTEXT("Math", "Math"))
+		{}
+	};
+	static FConstructorStatics ConstructorStatics;
+
+#if WITH_EDITORONLY_DATA
+	MenuCategories.Add(ConstructorStatics.NAME_Math);
+#endif
+}
+
+#if WITH_EDITOR
+int32 UMaterialExpressionHsvToRgb::Compile(class FMaterialCompiler* Compiler, int32 OutputIndex)
+{
+	if(!Input.GetTracedInput().Expression)
+	{
+		return Compiler->Errorf(TEXT("Missing HSVToRGB input"));
+	}
+
+	int32 InputIndex = Input.Compile(Compiler);
+	EMaterialValueType InputType = Compiler->GetType(InputIndex);
+	if(InputType != EMaterialValueType::MCT_Float3 && InputType != EMaterialValueType::MCT_Float4)
+	{
+		return InputIndex;
+	}
+
+	int32 Result = Compiler->HsvToRgb(InputIndex);
+
+	return Result;
+}
+
+void UMaterialExpressionHsvToRgb::GetCaption(TArray<FString>& OutCaptions) const
+{
+	OutCaptions.Add(TEXT("HSVToRGB"));
+}
+
+void UMaterialExpressionHsvToRgb::GetExpressionToolTip(TArray<FString>& OutToolTip)
+{
+	ConvertToMultilineToolTip(TEXT("Convert an incoming color from HSV to RGB space."), 40, OutToolTip);
+}
+#endif // WITH_EDITOR
+
+UMaterialExpressionRgbToHsv::UMaterialExpressionRgbToHsv(const FObjectInitializer& ObjectInitializer)
+	: Super(ObjectInitializer)
+{
+	// Structure to hold one-time initialization
+	struct FConstructorStatics
+	{
+		FText NAME_Math;
+		FConstructorStatics()
+			: NAME_Math(LOCTEXT("Math", "Math"))
+		{}
+	};
+	static FConstructorStatics ConstructorStatics;
+
+#if WITH_EDITORONLY_DATA
+	MenuCategories.Add(ConstructorStatics.NAME_Math);
+#endif
+}
+
+#if WITH_EDITOR
+int32 UMaterialExpressionRgbToHsv::Compile(class FMaterialCompiler* Compiler, int32 OutputIndex)
+{
+	if(!Input.GetTracedInput().Expression)
+	{
+		return Compiler->Errorf(TEXT("Missing RGBToHSV input"));
+	}
+
+	int32 InputIndex = Input.Compile(Compiler);
+	EMaterialValueType InputType = Compiler->GetType(InputIndex);
+	if(InputType != EMaterialValueType::MCT_Float3 && InputType != EMaterialValueType::MCT_Float4)
+	{
+		return InputIndex;
+	}
+
+	int32 Result = Compiler->RgbToHsv(InputIndex);
+
+	return Result;
+}
+
+void UMaterialExpressionRgbToHsv::GetCaption(TArray<FString>& OutCaptions) const
+{
+	OutCaptions.Add(TEXT("RGBToHSV"));
+}
+
+void UMaterialExpressionRgbToHsv::GetExpressionToolTip(TArray<FString>& OutToolTip)
+{
+	ConvertToMultilineToolTip(TEXT("Convert an incoming color from RGB to HSV space."), 40, OutToolTip);
+}
+#endif // WITH_EDITOR
+
+UMaterialExpressionExponential::UMaterialExpressionExponential(const FObjectInitializer& ObjectInitializer)
+	: Super(ObjectInitializer)
+{
+	// Structure to hold one-time initialization
+	struct FConstructorStatics
+	{
+		FText NAME_Math;
+		FConstructorStatics()
+			: NAME_Math(LOCTEXT("Math", "Math"))
+		{}
+	};
+	static FConstructorStatics ConstructorStatics;
+
+#if WITH_EDITORONLY_DATA
+	MenuCategories.Add(ConstructorStatics.NAME_Math);
+#endif
+}
+
+#if WITH_EDITOR
+int32 UMaterialExpressionExponential::Compile(class FMaterialCompiler* Compiler, int32 OutputIndex)
+{
+	if(!Input.GetTracedInput().Expression)
+	{
+		return Compiler->Errorf(TEXT("Missing Exp input"));
+	}
+
+	return Compiler->Exponential(Input.Compile(Compiler));
+}
+
+void UMaterialExpressionExponential::GetCaption(TArray<FString>& OutCaptions) const
+{
+	OutCaptions.Add(TEXT("Exp"));
+}
+
+void UMaterialExpressionExponential::GetExpressionToolTip(TArray<FString>& OutToolTip)
+{
+	ConvertToMultilineToolTip(TEXT("Returns the base-e exponential, or e^x, of the input."), 40, OutToolTip);
+}
+#endif // WITH_EDITOR
+
+UMaterialExpressionExponential2::UMaterialExpressionExponential2(const FObjectInitializer& ObjectInitializer)
+	: Super(ObjectInitializer)
+{
+	// Structure to hold one-time initialization
+	struct FConstructorStatics
+	{
+		FText NAME_Math;
+		FConstructorStatics()
+			: NAME_Math(LOCTEXT("Math", "Math"))
+		{}
+	};
+	static FConstructorStatics ConstructorStatics;
+
+#if WITH_EDITORONLY_DATA
+	MenuCategories.Add(ConstructorStatics.NAME_Math);
+#endif
+}
+
+#if WITH_EDITOR
+int32 UMaterialExpressionExponential2::Compile(class FMaterialCompiler* Compiler, int32 OutputIndex)
+{
+	if(!Input.GetTracedInput().Expression)
+	{
+		return Compiler->Errorf(TEXT("Missing Exp2 input"));
+	}
+
+	return Compiler->Exponential2(Input.Compile(Compiler));
+}
+
+void UMaterialExpressionExponential2::GetCaption(TArray<FString>& OutCaptions) const
+{
+	OutCaptions.Add(TEXT("Exp2"));
+}
+
+void UMaterialExpressionExponential2::GetExpressionToolTip(TArray<FString>& OutToolTip)
+{
+	ConvertToMultilineToolTip(TEXT("Returns the base 2 exponential, or 2^x, of the input."), 40, OutToolTip);
+}
+#endif // WITH_EDITOR
+
+UMaterialExpressionLogarithm::UMaterialExpressionLogarithm(const FObjectInitializer& ObjectInitializer)
+	: Super(ObjectInitializer)
+{
+	// Structure to hold one-time initialization
+	struct FConstructorStatics
+	{
+		FText NAME_Math;
+		FConstructorStatics()
+			: NAME_Math(LOCTEXT("Math", "Math"))
+		{}
+	};
+	static FConstructorStatics ConstructorStatics;
+
+#if WITH_EDITORONLY_DATA
+	MenuCategories.Add(ConstructorStatics.NAME_Math);
+#endif
+}
+
+#if WITH_EDITOR
+int32 UMaterialExpressionLogarithm::Compile(class FMaterialCompiler* Compiler, int32 OutputIndex)
+{
+	if(!Input.GetTracedInput().Expression)
+	{
+		return Compiler->Errorf(TEXT("Missing Log input"));
+	}
+
+	return Compiler->Logarithm(Input.Compile(Compiler));
+}
+
+void UMaterialExpressionLogarithm::GetCaption(TArray<FString>& OutCaptions) const
+{
+	OutCaptions.Add(TEXT("Log"));
+}
+
+void UMaterialExpressionLogarithm::GetExpressionToolTip(TArray<FString>& OutToolTip)
+{
+	ConvertToMultilineToolTip(TEXT("Returns the base-e logarithm, or natural logarithm, of the input. Input should be greater than 0."), 40, OutToolTip);
+}
+#endif // WITH_EDITOR
+
 UMaterialExpressionLogarithm2::UMaterialExpressionLogarithm2(const FObjectInitializer& ObjectInitializer)
 	: Super(ObjectInitializer)
 {
diff --git a/Engine/Source/Runtime/Engine/Private/Materials/MaterialUniformExpressions.cpp b/Engine/Source/Runtime/Engine/Private/Materials/MaterialUniformExpressions.cpp
--- a/Engine/Source/Runtime/Engine/Private/Materials/MaterialUniformExpressions.cpp
+++ b/Engine/Source/Runtime/Engine/Private/Materials/MaterialUniformExpressions.cpp
@@ -1797,6 +1797,9 @@ IMPLEMENT_MATERIALUNIFORMEXPRESSION_TYPE(FMaterialUniformExpressionSquareRoot);
 IMPLEMENT_MATERIALUNIFORMEXPRESSION_TYPE(FMaterialUniformExpressionRcp);
 IMPLEMENT_MATERIALUNIFORMEXPRESSION_TYPE(FMaterialUniformExpressionLength);
 IMPLEMENT_MATERIALUNIFORMEXPRESSION_TYPE(FMaterialUniformExpressionNormalize);
+IMPLEMENT_MATERIALUNIFORMEXPRESSION_TYPE(FMaterialUniformExpressionExponential);
+IMPLEMENT_MATERIALUNIFORMEXPRESSION_TYPE(FMaterialUniformExpressionExponential2);
+IMPLEMENT_MATERIALUNIFORMEXPRESSION_TYPE(FMaterialUniformExpressionLogarithm);
 IMPLEMENT_MATERIALUNIFORMEXPRESSION_TYPE(FMaterialUniformExpressionLogarithm2);
 IMPLEMENT_MATERIALUNIFORMEXPRESSION_TYPE(FMaterialUniformExpressionLogarithm10);
 IMPLEMENT_MATERIALUNIFORMEXPRESSION_TYPE(FMaterialUniformExpressionFoldedMath);
diff --git a/Engine/Source/Runtime/Engine/Private/Materials/MaterialUniformExpressions.h b/Engine/Source/Runtime/Engine/Private/Materials/MaterialUniformExpressions.h
--- a/Engine/Source/Runtime/Engine/Private/Materials/MaterialUniformExpressions.h
+++ b/Engine/Source/Runtime/Engine/Private/Materials/MaterialUniformExpressions.h
@@ -753,6 +753,119 @@ private:
 	TRefCountPtr<FMaterialUniformExpression> X;
 };
 
+/**
+ */
+class FMaterialUniformExpressionExponential : public FMaterialUniformExpression
+{
+	DECLARE_MATERIALUNIFORMEXPRESSION_TYPE(FMaterialUniformExpressionExponential);
+public:
+
+	FMaterialUniformExpressionExponential()
+	{}
+	FMaterialUniformExpressionExponential(FMaterialUniformExpression* InX) :
+		X(InX)
+	{}
+
+	// FMaterialUniformExpression interface.
+	virtual void WriteNumberOpcodes(UE::Shader::FPreshaderData& OutData) const override
+	{
+		X->WriteNumberOpcodes(OutData);
+		OutData.WriteOpcode(UE::Shader::EPreshaderOpcode::Exp);
+	}
+	bool IsConstant() const override
+	{
+		return X->IsConstant();
+	}
+	bool IsIdentical(const FMaterialUniformExpression* OtherExpression) const override
+	{
+		if(GetType() != OtherExpression->GetType())
+		{
+			return false;
+		}
+
+		const FMaterialUniformExpressionExponential* OtherExp = static_cast<const FMaterialUniformExpressionExponential*>(OtherExpression);
+		return X->IsIdentical(OtherExp->X);
+	}
+
+private:
+	TRefCountPtr<FMaterialUniformExpression> X;
+};
+
+/**
+ */
+class FMaterialUniformExpressionExponential2 : public FMaterialUniformExpression
+{
+	DECLARE_MATERIALUNIFORMEXPRESSION_TYPE(FMaterialUniformExpressionExponential2);
+public:
+
+	FMaterialUniformExpressionExponential2()
+	{}
+	FMaterialUniformExpressionExponential2(FMaterialUniformExpression* InX) :
+		X(InX)
+	{}
+
+	// FMaterialUniformExpression interface.
+	virtual void WriteNumberOpcodes(UE::Shader::FPreshaderData& OutData) const override
+	{
+		X->WriteNumberOpcodes(OutData);
+		OutData.WriteOpcode(UE::Shader::EPreshaderOpcode::Exp2);
+	}
+	bool IsConstant() const override
+	{
+		return X->IsConstant();
+	}
+	bool IsIdentical(const FMaterialUniformExpression* OtherExpression) const override
+	{
+		if(GetType() != OtherExpression->GetType())
+		{
+			return false;
+		}
+
+		const FMaterialUniformExpressionExponential2* OtherExp2= static_cast<const FMaterialUniformExpressionExponential2*>(OtherExpression);
+		return X->IsIdentical(OtherExp2->X);
+	}
+
+private:
+	TRefCountPtr<FMaterialUniformExpression> X;
+};
+
+/**
+ */
+class FMaterialUniformExpressionLogarithm : public FMaterialUniformExpression
+{
+	DECLARE_MATERIALUNIFORMEXPRESSION_TYPE(FMaterialUniformExpressionLogarithm);
+public:
+
+	FMaterialUniformExpressionLogarithm()
+	{}
+	FMaterialUniformExpressionLogarithm(FMaterialUniformExpression* InX) :
+		X(InX)
+	{}
+
+	// FMaterialUniformExpression interface.
+	virtual void WriteNumberOpcodes(UE::Shader::FPreshaderData& OutData) const override
+	{
+		X->WriteNumberOpcodes(OutData);
+		OutData.WriteOpcode(UE::Shader::EPreshaderOpcode::Log);
+	}
+	bool IsConstant() const override
+	{
+		return X->IsConstant();
+	}
+	bool IsIdentical(const FMaterialUniformExpression* OtherExpression) const override
+	{
+		if(GetType() != OtherExpression->GetType())
+		{
+			return false;
+		}
+
+		const FMaterialUniformExpressionLogarithm* OtherLog = static_cast<const FMaterialUniformExpressionLogarithm*>(OtherExpression);
+		return X->IsIdentical(OtherLog->X);
+	}
+
+private:
+	TRefCountPtr<FMaterialUniformExpression> X;
+};
 
 /**
  */
diff --git a/Engine/Source/Runtime/Engine/Private/Shader/Preshader.cpp b/Engine/Source/Runtime/Engine/Private/Shader/Preshader.cpp
--- a/Engine/Source/Runtime/Engine/Private/Shader/Preshader.cpp
+++ b/Engine/Source/Runtime/Engine/Private/Shader/Preshader.cpp
@@ -927,6 +927,9 @@ FPreshaderValue EvaluatePreshader(const FUniformExpressionSet* UniformExpression
 		case EPreshaderOpcode::Sign: EvaluateUnaryOpInPlace(Stack, SignInPlace); break;
 		case EPreshaderOpcode::Frac: EvaluateUnaryOpInPlace(Stack, FracInPlace); break;
 		case EPreshaderOpcode::Fractional: EvaluateUnaryOpInPlace(Stack, FractionalInPlace); break;
+		case EPreshaderOpcode::Exp: EvaluateUnaryOpInPlace(Stack, ExpInPlace); break;
+		case EPreshaderOpcode::Exp2: EvaluateUnaryOpInPlace(Stack, Exp2InPlace); break;
+		case EPreshaderOpcode::Log: EvaluateUnaryOpInPlace(Stack, LogInPlace); break;
 		case EPreshaderOpcode::Log2: EvaluateUnaryOpInPlace(Stack, Log2InPlace); break;
 		case EPreshaderOpcode::Log10: EvaluateUnaryOpInPlace(Stack, Log10InPlace); break;
 		case EPreshaderOpcode::ComponentSwizzle: EvaluateComponentSwizzle(Stack, Data); break;
diff --git a/Engine/Source/Runtime/Engine/Private/Shader/ShaderValue.cpp b/Engine/Source/Runtime/Engine/Private/Shader/ShaderValue.cpp
--- a/Engine/Source/Runtime/Engine/Private/Shader/ShaderValue.cpp
+++ b/Engine/Source/Runtime/Engine/Private/Shader/ShaderValue.cpp
@@ -974,6 +974,9 @@ struct FOpFrac : public FOpBaseNoInt { template<typename T> T operator()(T Value
 struct FOpFractional : public FOpBaseNoInt { template<typename T> T operator()(T Value) const { return FMath::Fractional(Value); } };
 struct FOpSqrt : public FOpBaseNoInt { template<typename T> T operator()(T Value) const { return FMath::Sqrt(Value); } };
 struct FOpRcp : public FOpBaseNoInt { template<typename T> T operator()(T Value) const { return (T)1 / GetSafeDivisor(Value); } };
+struct FOpExp : public FOpBaseNoInt { template<typename T> T operator()(T Value) const { return FMath::Exp(Value); } };
+struct FOpExp2 : public FOpBaseNoInt { template<typename T> T operator()(T Value) const { return FMath::Exp2(Value); } };
+struct FOpLog : public FOpBaseNoInt { template<typename T> T operator()(T Value) const { return FMath::Loge(Value); } };
 struct FOpLog2 : public FOpBaseNoInt { template<typename T> T operator()(T Value) const { return FMath::Log2(Value); } };
 struct FOpLog10 : public FOpBaseNoInt
 {
@@ -1432,6 +1435,21 @@ FValue Rcp(const FValue& Value)
 	return Private::UnaryOp(Private::FOpRcp(), Value);
 }
 
+FValue Exp(const FValue& Value)
+{
+	return Private::UnaryOp(Private::FOpExp(), Value);
+}
+
+FValue Exp2(const FValue& Value)
+{
+	return Private::UnaryOp(Private::FOpExp2(), Value);
+}
+
+FValue Log(const FValue& Value)
+{
+	return Private::UnaryOp(Private::FOpLog(), Value);
+}
+
 FValue Log2(const FValue& Value)
 {
 	return Private::UnaryOp(Private::FOpLog2(), Value);
@@ -1785,6 +1803,21 @@ EValueType RcpInPlace(EValueType Type, TArrayView<FValueComponent> Component)
 	return Private::UnaryOpInPlace(Private::FOpRcp(), Type, Component);
 }
 
+EValueType ExpInPlace(EValueType Type, TArrayView<FValueComponent> Component)
+{
+	return Private::UnaryOpInPlace(Private::FOpExp(), Type, Component);
+}
+
+EValueType Exp2InPlace(EValueType Type, TArrayView<FValueComponent> Component)
+{
+	return Private::UnaryOpInPlace(Private::FOpExp2(), Type, Component);
+}
+
+EValueType LogInPlace(EValueType Type, TArrayView<FValueComponent> Component)
+{
+	return Private::UnaryOpInPlace(Private::FOpLog(), Type, Component);
+}
+
 EValueType Log2InPlace(EValueType Type, TArrayView<FValueComponent> Component)
 {
 	return Private::UnaryOpInPlace(Private::FOpLog2(), Type, Component);
diff --git a/Engine/Source/Runtime/Engine/Public/HLSLTree/HLSLTree.h b/Engine/Source/Runtime/Engine/Public/HLSLTree/HLSLTree.h
--- a/Engine/Source/Runtime/Engine/Public/HLSLTree/HLSLTree.h
+++ b/Engine/Source/Runtime/Engine/Public/HLSLTree/HLSLTree.h
@@ -646,6 +646,7 @@ public:
 	const FExpression* NewSqrt(const FExpression* Input) { return NewUnaryOp(EOperation::Sqrt, Input); }
 	const FExpression* NewRsqrt(const FExpression* Input) { return NewUnaryOp(EOperation::Rsqrt, Input); }
 	const FExpression* NewLog2(const FExpression* Input) { return NewUnaryOp(EOperation::Log2, Input); }
+	const FExpression* NewExp(const FExpression* Input) { return NewUnaryOp(EOperation::Exp, Input); }
 	const FExpression* NewExp2(const FExpression* Input) { return NewUnaryOp(EOperation::Exp2, Input); }
 	const FExpression* NewFrac(const FExpression* Input) { return NewUnaryOp(EOperation::Frac, Input); }
 	const FExpression* NewLength(const FExpression* Input) { return NewUnaryOp(EOperation::Length, Input); }
diff --git a/Engine/Source/Runtime/Engine/Public/HLSLTree/HLSLTreeTypes.h b/Engine/Source/Runtime/Engine/Public/HLSLTree/HLSLTreeTypes.h
--- a/Engine/Source/Runtime/Engine/Public/HLSLTree/HLSLTreeTypes.h
+++ b/Engine/Source/Runtime/Engine/Public/HLSLTree/HLSLTreeTypes.h
@@ -80,7 +80,9 @@ enum class EOperation : uint8
 	Rcp,
 	Sqrt,
 	Rsqrt,
+	Log,
 	Log2,
+	Exp,
 	Exp2,
 	Frac,
 	Floor,
diff --git a/Engine/Source/Runtime/Engine/Public/MaterialCompiler.h b/Engine/Source/Runtime/Engine/Public/MaterialCompiler.h
--- a/Engine/Source/Runtime/Engine/Public/MaterialCompiler.h
+++ b/Engine/Source/Runtime/Engine/Public/MaterialCompiler.h
@@ -421,13 +421,18 @@ public:
 	virtual int32 Div(int32 A,int32 B) = 0;
 	virtual int32 Dot(int32 A,int32 B) = 0;
 	virtual int32 Cross(int32 A,int32 B) = 0;
-
+		
 	virtual int32 Power(int32 Base,int32 Exponent) = 0;
+	virtual int32 Exponential(int32 X) = 0;
+	virtual int32 Exponential2(int32 X) = 0;
+	virtual int32 Logarithm(int32 X) = 0;
 	virtual int32 Logarithm2(int32 X) = 0;
 	virtual int32 Logarithm10(int32 X) = 0;
 	virtual int32 SquareRoot(int32 X) = 0;
 	virtual int32 Length(int32 X) = 0;
 	virtual int32 Normalize(int32 X) = 0;
+	virtual int32 HsvToRgb(int32 X) = 0;
+	virtual int32 RgbToHsv(int32 X) = 0;
 
 	virtual int32 Lerp(int32 X,int32 Y,int32 A) = 0;
 	virtual int32 Min(int32 A,int32 B) = 0;
@@ -878,11 +883,16 @@ public:
 	virtual int32 Cross(int32 A, int32 B) override { return Compiler->Cross(A, B); }
 
 	virtual int32 Power(int32 Base, int32 Exponent) override { return Compiler->Power(Base, Exponent); }
+	virtual int32 Exponential(int32 X) override { return Compiler->Exponential(X); }
+	virtual int32 Exponential2(int32 X) override { return Compiler->Exponential2(X); }
+	virtual int32 Logarithm(int32 X) override { return Compiler->Logarithm(X); }
 	virtual int32 Logarithm2(int32 X) override { return Compiler->Logarithm2(X); }
 	virtual int32 Logarithm10(int32 X) override { return Compiler->Logarithm10(X); }
 	virtual int32 SquareRoot(int32 X) override { return Compiler->SquareRoot(X); }
 	virtual int32 Length(int32 X) override { return Compiler->Length(X); }
 	virtual int32 Normalize(int32 X) override { return Compiler->Normalize(X); }
+	virtual int32 HsvToRgb(int32 X) override { return Compiler->HsvToRgb(X); }
+	virtual int32 RgbToHsv(int32 X) override { return Compiler->RgbToHsv(X); }
 
 	virtual int32 Lerp(int32 X, int32 Y, int32 A) override { return Compiler->Lerp(X, Y, A); }
 	virtual int32 Min(int32 A, int32 B) override { return Compiler->Min(A, B); }
diff --git a/Engine/Source/Runtime/Engine/Public/Shader/Preshader.h b/Engine/Source/Runtime/Engine/Public/Shader/Preshader.h
--- a/Engine/Source/Runtime/Engine/Public/Shader/Preshader.h
+++ b/Engine/Source/Runtime/Engine/Public/Shader/Preshader.h
@@ -73,6 +73,9 @@ enum class EPreshaderOpcode : uint8
 	Greater,
 	LessEqual,
 	GreaterEqual,
+	Exp,
+	Exp2,
+	Log
 };
 
 struct FPreshaderStructType
diff --git a/Engine/Source/Runtime/Engine/Public/Shader/ShaderTypes.h b/Engine/Source/Runtime/Engine/Public/Shader/ShaderTypes.h
--- a/Engine/Source/Runtime/Engine/Public/Shader/ShaderTypes.h
+++ b/Engine/Source/Runtime/Engine/Public/Shader/ShaderTypes.h
@@ -545,6 +545,9 @@ ENGINE_API FValue Frac(const FValue& Value);
 ENGINE_API FValue Fractional(const FValue& Value);
 ENGINE_API FValue Sqrt(const FValue& Value);
 ENGINE_API FValue Rcp(const FValue& Value);
+ENGINE_API FValue Exp(const FValue& Value);
+ENGINE_API FValue Exp2(const FValue& Value);
+ENGINE_API FValue Log(const FValue& Value);
 ENGINE_API FValue Log2(const FValue& Value);
 ENGINE_API FValue Log10(const FValue& Value);
 ENGINE_API FValue Sin(const FValue& Value);
@@ -587,6 +590,9 @@ ENGINE_API EValueType FracInPlace(EValueType Type, TArrayView<FValueComponent> C
 ENGINE_API EValueType FractionalInPlace(EValueType Type, TArrayView<FValueComponent> Component);
 ENGINE_API EValueType SqrtInPlace(EValueType Type, TArrayView<FValueComponent> Component);
 ENGINE_API EValueType RcpInPlace(EValueType Type, TArrayView<FValueComponent> Component);
+ENGINE_API EValueType ExpInPlace(EValueType Type, TArrayView<FValueComponent> Component);
+ENGINE_API EValueType Exp2InPlace(EValueType Type, TArrayView<FValueComponent> Component);
+ENGINE_API EValueType LogInPlace(EValueType Type, TArrayView<FValueComponent> Component);
 ENGINE_API EValueType Log2InPlace(EValueType Type, TArrayView<FValueComponent> Component);
 ENGINE_API EValueType Log10InPlace(EValueType Type, TArrayView<FValueComponent> Component);
 ENGINE_API EValueType SinInPlace(EValueType Type, TArrayView<FValueComponent> Component);
