Add full memory stats to OOM error handler in SCW.

diff --git a/Engine/Source/Programs/ShaderCompileWorker/Private/ShaderCompileWorker.cpp b/Engine/Source/Programs/ShaderCompileWorker/Private/ShaderCompileWorker.cpp
--- a/Engine/Source/Programs/ShaderCompileWorker/Private/ShaderCompileWorker.cpp
+++ b/Engine/Source/Programs/ShaderCompileWorker/Private/ShaderCompileWorker.cpp
@@ -114,13 +114,6 @@ static const TCHAR* GetLocalHostname(int32* OutHostnameLength = nullptr)
 	return *Hostname;
 }
 
-static void GetMemoryStats(uint64& OutAvailablePhysical, uint64& OutUsedPhysical)
-{
-	const FPlatformMemoryStats MemoryStats = FPlatformMemory::GetStats();
-	OutAvailablePhysical = MemoryStats.AvailablePhysical;
-	OutUsedPhysical = MemoryStats.UsedPhysical;
-}
-
 static int64 WriteOutputFileHeader(FArchive& OutputFile, int32 ErrorCode, int32 CallstackLength, const TCHAR* Callstack,
 	int32 ExceptionInfoLength, const TCHAR* ExceptionInfo)
 {
@@ -168,9 +161,16 @@ static int64 WriteOutputFileHeader(FArchive& OutputFile, int32 ErrorCode, int32
 		// Store available and used physical memory of host machine on OOM error
 		if (ErrorCode == FSCWErrorCode::OutOfMemory)
 		{
-			uint64 AvailablePhysical = 0, UsedPhysical = 0;
-			GetMemoryStats(AvailablePhysical, UsedPhysical);
-			OutputFile << AvailablePhysical << UsedPhysical;
+			FPlatformMemoryStats MemoryStats = FPlatformMemory::GetStats();
+
+			OutputFile
+				<< MemoryStats.AvailablePhysical
+				<< MemoryStats.AvailableVirtual
+				<< MemoryStats.UsedPhysical
+				<< MemoryStats.PeakUsedPhysical
+				<< MemoryStats.UsedVirtual
+				<< MemoryStats.PeakUsedVirtual
+				;
 		}
 	}
 
diff --git a/Engine/Source/Runtime/Engine/Private/ShaderCompiler/ShaderCompiler.cpp b/Engine/Source/Runtime/Engine/Private/ShaderCompiler/ShaderCompiler.cpp
--- a/Engine/Source/Runtime/Engine/Private/ShaderCompiler/ShaderCompiler.cpp
+++ b/Engine/Source/Runtime/Engine/Private/ShaderCompiler/ShaderCompiler.cpp
@@ -936,7 +936,7 @@ static void ModalErrorOrLog(const FString& Title, const FString& Text, int64 Cur
 
 	if (FPlatformProperties::SupportsWindowedMode())
 	{
-		UE_LOG(LogShaderCompilers, Error, TEXT("%s%s"), *Text, *BadFile);
+		UE_LOG(LogShaderCompilers, Error, TEXT("%s\n%s"), *Text, *BadFile);
 		if (!bModalReported.AtomicSet(true))
 		{
 			// Show dialog box with error message and request exit
@@ -951,7 +951,7 @@ static void ModalErrorOrLog(const FString& Title, const FString& Text, int64 Cur
 	}
 	else
 	{
-		UE_LOG(LogShaderCompilers, Fatal, TEXT("%s\n%s%s"), *Title, *Text, *BadFile);
+		UE_LOG(LogShaderCompilers, Fatal, TEXT("%s\n%s\n%s"), *Title, *Text, *BadFile);
 	}
 }
 
@@ -1298,42 +1298,42 @@ namespace ShaderCompileWorkerError
 
 	void HandleBadShaderFormatVersion(const TCHAR* Data)
 	{
-		ModalErrorOrLog(TEXT("ShaderCompileWorker failed"), FString::Printf(TEXT("%s\n"), Data));
+		ModalErrorOrLog(TEXT("ShaderCompileWorker failed"), Data);
 	}
 
 	void HandleBadInputVersion(const TCHAR* Data)
 	{
-		ModalErrorOrLog(TEXT("ShaderCompileWorker failed"), FString::Printf(TEXT("%s\n"), Data));
+		ModalErrorOrLog(TEXT("ShaderCompileWorker failed"), Data);
 	}
 
 	void HandleBadSingleJobHeader(const TCHAR* Data)
 	{
-		ModalErrorOrLog(TEXT("ShaderCompileWorker failed"), FString::Printf(TEXT("%s\n"), Data));
+		ModalErrorOrLog(TEXT("ShaderCompileWorker failed"), Data);
 	}
 
 	void HandleBadPipelineJobHeader(const TCHAR* Data)
 	{
-		ModalErrorOrLog(TEXT("ShaderCompileWorker failed"), FString::Printf(TEXT("%s\n"), Data));
+		ModalErrorOrLog(TEXT("ShaderCompileWorker failed"), Data);
 	}
 
 	void HandleCantDeleteInputFile(const TCHAR* Data)
 	{
-		ModalErrorOrLog(TEXT("ShaderCompileWorker failed"), FString::Printf(TEXT("%s\n"), Data));
+		ModalErrorOrLog(TEXT("ShaderCompileWorker failed"), Data);
 	}
 
 	void HandleCantSaveOutputFile(const TCHAR* Data)
 	{
-		ModalErrorOrLog(TEXT("ShaderCompileWorker failed"), FString::Printf(TEXT("%s\n"), Data));
+		ModalErrorOrLog(TEXT("ShaderCompileWorker failed"), Data);
 	}
 
 	void HandleNoTargetShaderFormatsFound(const TCHAR* Data)
 	{
-		ModalErrorOrLog(TEXT("ShaderCompileWorker failed"), FString::Printf(TEXT("%s\n"), Data));
+		ModalErrorOrLog(TEXT("ShaderCompileWorker failed"), Data);
 	}
 
 	void HandleCantCompileForSpecificFormat(const TCHAR* Data)
 	{
-		ModalErrorOrLog(TEXT("ShaderCompileWorker failed"), FString::Printf(TEXT("%s\n"), Data));
+		ModalErrorOrLog(TEXT("ShaderCompileWorker failed"), Data);
 	}
 
 	void HandleOutputFileEmpty(const TCHAR* Filename)
@@ -1353,19 +1353,33 @@ namespace ShaderCompileWorkerError
 
 	void HandleBadInputFile(const TCHAR* Data)
 	{
-		ModalErrorOrLog(TEXT("ShaderCompileWorker failed"), FString::Printf(TEXT("Bad-input-file exception:\n%s\n"), Data));
-	}
-
-	bool HandleOutOfMemory(const TCHAR* ExceptionInfo, const TCHAR* Hostname, uint64 AvailableMemory, uint64 UsedMemory, const TArray<FShaderCommonCompileJobPtr>& QueuedJobs)
-	{
-		const FString AvailableMemoryStr = FString::Printf(TEXT("%d MB"), FUnitConversion::Convert(AvailableMemory, EUnit::Bytes, EUnit::Megabytes));
-		const FString UsedMemoryStr = FString::Printf(TEXT("%d MB"), FUnitConversion::Convert(UsedMemory, EUnit::Bytes, EUnit::Megabytes));
-		const FString ErrorReport = FString::Printf(TEXT("ShaderCompileWorker failed with out-of-memory (OOM) exception on machine \"%s\"\n%s; Used physical memory %s of %s\n"), Hostname, (ExceptionInfo[0] == TEXT('\0') ? TEXT("No exception information") : ExceptionInfo), *UsedMemoryStr, *AvailableMemoryStr);
+		ModalErrorOrLog(TEXT("ShaderCompileWorker failed"), FString::Printf(TEXT("Bad-input-file exception:\n%s"), Data));
+	}
+
+	bool HandleOutOfMemory(const TCHAR* ExceptionInfo, const TCHAR* Hostname, const FPlatformMemoryStats& MemoryStats, const TArray<FShaderCommonCompileJobPtr>& QueuedJobs)
+	{
+		const FString ErrorReport = FString::Printf(
+			TEXT("ShaderCompileWorker failed with out-of-memory (OOM) exception on machine \"%s\" (%s); MemoryStats:")
+			TEXT("\n\tAvailablePhysical %llu")
+			TEXT("\n\t AvailableVirtual %llu")
+			TEXT("\n\t     UsedPhysical %llu")
+			TEXT("\n\t PeakUsedPhysical %llu")
+			TEXT("\n\t      UsedVirtual %llu")
+			TEXT("\n\t  PeakUsedVirtual %llu"),
+			Hostname,
+			(ExceptionInfo[0] == TEXT('\0') ? TEXT("No exception information") : ExceptionInfo),
+			MemoryStats.AvailablePhysical,
+			MemoryStats.AvailableVirtual,
+			MemoryStats.UsedPhysical,
+			MemoryStats.PeakUsedPhysical,
+			MemoryStats.UsedVirtual,
+			MemoryStats.PeakUsedVirtual
+		);
 
 		if (TryReissueShaderCompileJobs(QueuedJobs))
 		{
 			// We recovered from this error
-			UE_LOG(LogShaderCompilers, Warning, TEXT("%sReissue %d shader compile %s to remaining workers"), *ErrorReport, QueuedJobs.Num(), (QueuedJobs.Num() == 1 ? TEXT("job") : TEXT("jobs")));
+			UE_LOG(LogShaderCompilers, Warning, TEXT("%s\nReissue %d shader compile %s to remaining workers"), *ErrorReport, QueuedJobs.Num(), (QueuedJobs.Num() == 1 ? TEXT("job") : TEXT("jobs")));
 			return true;
 		}
 		else
@@ -1916,10 +1930,17 @@ static bool HandleWorkerCrash(const TArray<FShaderCommonCompileJobPtr>& QueuedJo
 	Hostname[HostnameLength] = TEXT('\0');
 
 	// Read available and used physical memory from worker machine on OOM error
-	uint64 AvailableMemory = 0, UsedMemory = 0;
+	FPlatformMemoryStats MemoryStats;
 	if (ErrorCode == FSCWErrorCode::OutOfMemory)
 	{
-		OutputFile << AvailableMemory << UsedMemory;
+		OutputFile
+			<< MemoryStats.AvailablePhysical
+			<< MemoryStats.AvailableVirtual
+			<< MemoryStats.UsedPhysical
+			<< MemoryStats.PeakUsedPhysical
+			<< MemoryStats.UsedVirtual
+			<< MemoryStats.PeakUsedVirtual
+			;
 	}
 
 	// Store primary job information onto stack to make it part of a crash dump
@@ -2008,7 +2029,7 @@ static bool HandleWorkerCrash(const TArray<FShaderCommonCompileJobPtr>& QueuedJo
 		ShaderCompileWorkerError::HandleBadInputFile(ExceptionInfo.GetData());
 		break;
 	case FSCWErrorCode::OutOfMemory:
-		return ShaderCompileWorkerError::HandleOutOfMemory(ExceptionInfo.GetData(), Hostname.GetData(), AvailableMemory, UsedMemory, QueuedJobs);
+		return ShaderCompileWorkerError::HandleOutOfMemory(ExceptionInfo.GetData(), Hostname.GetData(), MemoryStats, QueuedJobs);
 	case FSCWErrorCode::Success:
 		// Can't get here...
 		return true;
diff --git a/Engine/Source/Runtime/RenderCore/Public/ShaderCompilerCore.h b/Engine/Source/Runtime/RenderCore/Public/ShaderCompilerCore.h
--- a/Engine/Source/Runtime/RenderCore/Public/ShaderCompilerCore.h
+++ b/Engine/Source/Runtime/RenderCore/Public/ShaderCompilerCore.h
@@ -29,7 +29,7 @@ class FShaderPipelineCompileJob;
 // this is for the protocol, not the data, bump if FShaderCompilerInput or ProcessInputFromArchive changes.
 inline const int32 ShaderCompileWorkerInputVersion = 21;
 // this is for the protocol, not the data, bump if FShaderCompilerOutput or WriteToOutputArchive changes.
-inline const int32 ShaderCompileWorkerOutputVersion = 15;
+inline const int32 ShaderCompileWorkerOutputVersion = 16;
 // this is for the protocol, not the data.
 inline const int32 ShaderCompileWorkerSingleJobHeader = 'S';
 // this is for the protocol, not the data.
