Substrate - renamed slab input from Glint Value to Glint Density

diff --git a/Engine/Source/Runtime/Engine/Private/Materials/MaterialExpressions.cpp b/Engine/Source/Runtime/Engine/Private/Materials/MaterialExpressions.cpp
--- a/Engine/Source/Runtime/Engine/Private/Materials/MaterialExpressions.cpp
+++ b/Engine/Source/Runtime/Engine/Private/Materials/MaterialExpressions.cpp
@@ -23822,7 +23822,7 @@ FName UMaterialExpressionStrataSlabBSDF::GetInputName(int32 InputIndex) const
 	}
 	else if (InputIndex == 16)
 	{
-		return Strata::IsGlintEnabled() ? TEXT("Glint Value") : TEXT("Glint Value (Disabled)");
+		return Strata::IsGlintEnabled() ? TEXT("Glint Density") : TEXT("Glint Density (Disabled)");
 	}
 	else if (InputIndex == 17)
 	{
