Text Render:
- Added UFUNCTION to SetText to make "Text" property keyable in Sequencer as it currently stands this text property does not meet any 1 requirement to make it keyable in Sequencer
- Adding this UFUNCTION will cause Sequencer to find the Property Setter (via "Set"+PropertyName), and thus making it keyable.

diff --git a/Engine/Source/Runtime/Engine/Classes/Components/TextRenderComponent.h b/Engine/Source/Runtime/Engine/Classes/Components/TextRenderComponent.h
--- a/Engine/Source/Runtime/Engine/Classes/Components/TextRenderComponent.h
+++ b/Engine/Source/Runtime/Engine/Classes/Components/TextRenderComponent.h
@@ -100,6 +100,7 @@ class UTextRenderComponent : public UPrimitiveComponent
 	// -----------------------------
 	
 	/** Change the text value and signal the primitives to be rebuilt */
+	UFUNCTION()
 	ENGINE_API void SetText(const FText& Value);
 
 	/** Change the text value and signal the primitives to be rebuilt */
