LevelInstance: Fix BreakLI not assigning DataLayers properly

#rnx

diff --git a/Engine/Source/Runtime/Engine/Classes/GameFramework/Actor.h b/Engine/Source/Runtime/Engine/Classes/GameFramework/Actor.h
--- a/Engine/Source/Runtime/Engine/Classes/GameFramework/Actor.h
+++ b/Engine/Source/Runtime/Engine/Classes/GameFramework/Actor.h
@@ -1298,6 +1298,8 @@ private:
 public:
 	ENGINE_API bool AddDataLayer(const UDataLayerInstance* DataLayerInstance);
 	ENGINE_API bool RemoveDataLayer(const UDataLayerInstance* DataLayerInstance);
+	ENGINE_API bool CanAddDataLayer(const UDataLayerInstance* InDataLayerInstance) const;
+
 	ENGINE_API TArray<const UDataLayerInstance*> RemoveAllDataLayers();
 	ENGINE_API bool SupportsDataLayerType(TSubclassOf<UDataLayerInstance> DataLayerType) const;
 	
diff --git a/Engine/Source/Runtime/Engine/Private/ActorEditor.cpp b/Engine/Source/Runtime/Engine/Private/ActorEditor.cpp
--- a/Engine/Source/Runtime/Engine/Private/ActorEditor.cpp
+++ b/Engine/Source/Runtime/Engine/Private/ActorEditor.cpp
@@ -1829,11 +1829,28 @@ bool AActor::SupportsDataLayerType(TSubclassOf<UDataLayerInstance> InDataLayerTy
 		!GetClass()->GetDefaultObject<AActor>()->bHiddenEd);
 }
 
+bool AActor::CanAddDataLayer(const UDataLayerInstance* InDataLayerInstance) const
+{
+	if (InDataLayerInstance && SupportsDataLayerType(InDataLayerInstance->GetClass()))
+	{
+		if (const UDataLayerAsset* DataLayerAsset = InDataLayerInstance->GetAsset())
+		{
+			return !DataLayerAssets.Contains(DataLayerAsset);
+		}
+		else if (const UDeprecatedDataLayerInstance* DataLayerInstance = Cast<UDeprecatedDataLayerInstance>(InDataLayerInstance))
+		{
+			return !DataLayers.Contains(DataLayerInstance->GetActorDataLayer());
+		}
+	}
+
+	return false;
+}
+
 bool FAssignActorDataLayer::AddDataLayerAsset(AActor* InActor, const UDataLayerAsset* InDataLayerAsset)
 {
 	check(InDataLayerAsset != nullptr);
 
-	if (!InActor->ContainsDataLayer(InDataLayerAsset))
+	if (!InActor->DataLayerAssets.Contains(InDataLayerAsset))
 	{
 		InActor->Modify();
 		InActor->DataLayerAssets.Add(InDataLayerAsset);
@@ -1868,7 +1885,7 @@ bool AActor::SupportsDataLayer() const
 
 bool AActor::AddDataLayer(const FActorDataLayer& ActorDataLayer)
 {
-	if (SupportsDataLayerType(UDataLayerInstance::StaticClass()) && !ContainsDataLayer(ActorDataLayer))
+	if (SupportsDataLayerType(UDataLayerInstance::StaticClass()) && !DataLayers.Contains(ActorDataLayer))
 	{
 		Modify();
 		DataLayers.Add(ActorDataLayer);
@@ -1880,7 +1897,7 @@ bool AActor::AddDataLayer(const FActorDataLayer& ActorDataLayer)
 
 bool AActor::RemoveDataLayer(const FActorDataLayer& ActorDataLayer)
 {
-	if (ContainsDataLayer(ActorDataLayer))
+	if (DataLayers.Contains(ActorDataLayer))
 	{
 		Modify();
 		DataLayers.Remove(ActorDataLayer);
diff --git a/Engine/Source/Runtime/Engine/Private/LevelInstance/LevelInstanceSubsystem.cpp b/Engine/Source/Runtime/Engine/Private/LevelInstance/LevelInstanceSubsystem.cpp
--- a/Engine/Source/Runtime/Engine/Private/LevelInstance/LevelInstanceSubsystem.cpp
+++ b/Engine/Source/Runtime/Engine/Private/LevelInstance/LevelInstanceSubsystem.cpp
@@ -1312,12 +1312,12 @@ void ULevelInstanceSubsystem::BreakLevelInstance_Impl(ILevelInstanceInterface* L
 					{
 						for (const UDataLayerInstance* DataLayerInstance : LevelInstanceDataLayerInstances)
 						{
-							if (const UDataLayerInstanceWithAsset* DataLayerInstanceWithAsset = Cast<UDataLayerInstanceWithAsset>(DataLayerInstance))
+							if (const UDataLayerAsset* DataLayerAsset = DataLayerInstance->GetAsset())
 							{
 								// For DataLayerInstanceAsset we add the Asset to the actor instead of the DataLayerInstance because Actor hasn't moved yet and this
 								// will fail on a UDataLayerInstanceWithAsset when comparing the Actor's outer to the DataLayerInstance outer.
 								// todo_ow : Assign the actor to the DataLayerInstance in the destination level via UDataLayerInstance::AddActor to go through the DataLayerInstance process.
-								FAssignActorDataLayer::AddDataLayerAsset(Actor, DataLayerInstanceWithAsset->GetAsset());
+								FAssignActorDataLayer::AddDataLayerAsset(Actor, DataLayerAsset);
 							}
 							else
 							{
diff --git a/Engine/Source/Runtime/Engine/Private/WorldPartition/DataLayer/DataLayerInstance.cpp b/Engine/Source/Runtime/Engine/Private/WorldPartition/DataLayer/DataLayerInstance.cpp
--- a/Engine/Source/Runtime/Engine/Private/WorldPartition/DataLayer/DataLayerInstance.cpp
+++ b/Engine/Source/Runtime/Engine/Private/WorldPartition/DataLayer/DataLayerInstance.cpp
@@ -157,7 +157,7 @@ bool UDataLayerInstance::CanUserAddActors() const
 
 bool UDataLayerInstance::CanAddActor(AActor* InActor) const
 {
-	return InActor != nullptr && InActor->SupportsDataLayerType(GetClass()) && !InActor->ContainsDataLayer(this);
+	return InActor != nullptr && InActor->CanAddDataLayer(this);
 }
 
 bool UDataLayerInstance::AddActor(AActor* InActor) const
diff --git a/Engine/Source/Runtime/Engine/Private/WorldPartition/DataLayer/DeprecatedDataLayerInstance.cpp b/Engine/Source/Runtime/Engine/Private/WorldPartition/DataLayer/DeprecatedDataLayerInstance.cpp
--- a/Engine/Source/Runtime/Engine/Private/WorldPartition/DataLayer/DeprecatedDataLayerInstance.cpp
+++ b/Engine/Source/Runtime/Engine/Private/WorldPartition/DataLayer/DeprecatedDataLayerInstance.cpp
@@ -72,14 +72,19 @@ void UDeprecatedDataLayerInstance::OnCreated(const UDEPRECATED_DataLayer* Deprec
 	InitialRuntimeState = DeprecatedDataLayer->InitialRuntimeState;
 }
 
+FActorDataLayer UDeprecatedDataLayerInstance::GetActorDataLayer() const
+{
+	return FActorDataLayer(GetDataLayerFName());
+}
+
 bool UDeprecatedDataLayerInstance::PerformAddActor(AActor* InActor) const
 {
-	return InActor->AddDataLayer(FActorDataLayer(GetDataLayerFName()));
+	return InActor->AddDataLayer(GetActorDataLayer());
 }
 
 bool UDeprecatedDataLayerInstance::PerformRemoveActor(AActor* InActor) const
 {
-	return InActor->RemoveDataLayer(FActorDataLayer(GetDataLayerFName()));
+	return InActor->RemoveDataLayer(GetActorDataLayer());
 }
 
 bool UDeprecatedDataLayerInstance::RelabelDataLayer(FName InDataLayerLabel)
diff --git a/Engine/Source/Runtime/Engine/Public/WorldPartition/DataLayer/DeprecatedDataLayerInstance.h b/Engine/Source/Runtime/Engine/Public/WorldPartition/DataLayer/DeprecatedDataLayerInstance.h
--- a/Engine/Source/Runtime/Engine/Public/WorldPartition/DataLayer/DeprecatedDataLayerInstance.h
+++ b/Engine/Source/Runtime/Engine/Public/WorldPartition/DataLayer/DeprecatedDataLayerInstance.h
@@ -2,7 +2,7 @@
 #pragma once
 
 #include "WorldPartition/DataLayer/DataLayerInstance.h"
-
+#include "WorldPartition/DataLayer/ActorDataLayer.h"
 #include "DeprecatedDataLayerInstance.generated.h"
 
 class UDEPRECATED_DataLayer;
@@ -32,6 +32,7 @@ public:
 #if WITH_EDITOR
 	static ENGINE_API FName MakeName();
 	ENGINE_API void OnCreated();
+	FActorDataLayer GetActorDataLayer() const;
 
 	ENGINE_API PRAGMA_DISABLE_DEPRECATION_WARNINGS
 	static FName MakeName(const UDEPRECATED_DataLayer* DeprecatedDataLayer);
