LevelInstance:
- Fix EnableStreaming Level Instances not being loaded when edited in place
- Fix EnableStreaming Level Instances from not loading when placed in !WP World

#rnx

diff --git a/Engine/Source/Runtime/Engine/Private/LevelInstance/LevelInstanceEditorLevelStreaming.cpp b/Engine/Source/Runtime/Engine/Private/LevelInstance/LevelInstanceEditorLevelStreaming.cpp
--- a/Engine/Source/Runtime/Engine/Private/LevelInstance/LevelInstanceEditorLevelStreaming.cpp
+++ b/Engine/Source/Runtime/Engine/Private/LevelInstance/LevelInstanceEditorLevelStreaming.cpp
@@ -17,6 +17,7 @@
 #include "Engine/World.h"
 #include "GameFramework/WorldSettings.h"
 #include "PackageTools.h"
+#include "WorldPartition/WorldPartition.h"
 #endif
 
 #if WITH_EDITOR
@@ -168,6 +169,14 @@ void ULevelStreamingLevelInstanceEditor::OnLevelLoadedChanged(ULevel* InLevel)
 
 		if (ULevelInstanceSubsystem* LevelInstanceSubsystem = GetWorld()->GetSubsystem<ULevelInstanceSubsystem>())
 		{
+#if WITH_EDITOR
+			if (UWorldPartition* OuterWorldPartition = NewLoadedLevel->GetWorldPartition())
+			{
+				check(!OuterWorldPartition->IsInitialized());
+				OuterWorldPartition->bOverrideEnableStreamingInEditor = false;
+			}
+#endif
+
 			LevelInstanceSubsystem->RegisterLoadedLevelStreamingLevelInstanceEditor(this);
 		}
 	}
diff --git a/Engine/Source/Runtime/Engine/Private/LevelInstance/LevelInstanceLevelStreaming.cpp b/Engine/Source/Runtime/Engine/Private/LevelInstance/LevelInstanceLevelStreaming.cpp
--- a/Engine/Source/Runtime/Engine/Private/LevelInstance/LevelInstanceLevelStreaming.cpp
+++ b/Engine/Source/Runtime/Engine/Private/LevelInstance/LevelInstanceLevelStreaming.cpp
@@ -368,13 +368,24 @@ void ULevelStreamingLevelInstance::OnLevelLoadedChanged(ULevel* InLevel)
 			LevelInstanceSubsystem->RegisterLoadedLevelStreamingLevelInstance(this);
 
 #if WITH_EDITOR
-			if (UWorldPartition* OwningWorldPartition = GetWorld()->GetWorldPartition(); OwningWorldPartition && OwningWorldPartition->IsStreamingEnabled() && !GDisableLevelInstanceEditorPartialLoading)
+			if (UWorldPartition* OuterWorldPartition = NewLoadedLevel->GetWorldPartition())
 			{
-				if (UWorldPartition* OuterWorldPartition = NewLoadedLevel->GetWorldPartition())
+				check(!OuterWorldPartition->IsInitialized());
+				if (UWorldPartition* OwningWorldPartition = GetWorld()->GetWorldPartition(); OwningWorldPartition && OwningWorldPartition->IsStreamingEnabled())
 				{
-					check(!OuterWorldPartition->IsInitialized());
-					ILevelInstanceInterface* LevelInstance = LevelInstanceSubsystem->GetLevelInstance(LevelInstanceID);
-					OuterWorldPartition->bForceEnableStreamingInEditor = LevelInstance ? LevelInstance->SupportsPartialEditorLoading() : false;
+					if (GDisableLevelInstanceEditorPartialLoading)
+					{
+						OuterWorldPartition->bOverrideEnableStreamingInEditor = false;
+					}
+					else if (ILevelInstanceInterface* LevelInstance = LevelInstanceSubsystem->GetLevelInstance(LevelInstanceID); LevelInstance && LevelInstance->SupportsPartialEditorLoading())
+					{
+						OuterWorldPartition->bOverrideEnableStreamingInEditor = true;
+					}
+				}
+				else
+				{
+					// Do not enable Streaming in Editor if Level Instance is not part of a World Partition/Streaming world
+					OuterWorldPartition->bOverrideEnableStreamingInEditor = false;
 				}
 			}
 #endif
diff --git a/Engine/Source/Runtime/Engine/Private/WorldPartition/WorldPartition.cpp b/Engine/Source/Runtime/Engine/Private/WorldPartition/WorldPartition.cpp
--- a/Engine/Source/Runtime/Engine/Private/WorldPartition/WorldPartition.cpp
+++ b/Engine/Source/Runtime/Engine/Private/WorldPartition/WorldPartition.cpp
@@ -275,7 +275,6 @@ UWorldPartition::UWorldPartition(const FObjectInitializer& ObjectInitializer)
 	, bEnablingStreamingJustified(false)
 	, bIsPIE(false)
 	, NumUserCreatedLoadedRegions(0)
-	, bForceEnableStreamingInEditor(false)
 #endif
 	, InitState(EWorldPartitionInitState::Uninitialized)
 	, bStreamingInEnabled(true)
@@ -587,7 +586,7 @@ void UWorldPartition::Initialize(UWorld* InWorld, const FTransform& InTransform)
 
 		// Follow the world's streaming enabled value most of the times, except:
 		//	- World is instanced and from a Level Instance that supports partial loading
-		const bool bIsStreamingEnabled = bForceEnableStreamingInEditor || IsStreamingEnabled();
+		const bool bIsStreamingEnabled = bOverrideEnableStreamingInEditor.IsSet() ? *bOverrideEnableStreamingInEditor : IsStreamingEnabled();
 
 		if (bIsInstanced)
 		{
diff --git a/Engine/Source/Runtime/Engine/Public/WorldPartition/WorldPartition.h b/Engine/Source/Runtime/Engine/Public/WorldPartition/WorldPartition.h
--- a/Engine/Source/Runtime/Engine/Public/WorldPartition/WorldPartition.h
+++ b/Engine/Source/Runtime/Engine/Public/WorldPartition/WorldPartition.h
@@ -439,7 +439,7 @@ private:
 	TSet<FString> GeneratedStreamingPackageNames;
 
 public:
-	bool bForceEnableStreamingInEditor;
+	TOptional<bool> bOverrideEnableStreamingInEditor;
 
 private:
 #endif
