Add an option on static mesh for WPO not to write velocity.
We default to writing velocity, but the option can be useful if we know that we have "static" WPO for things like camera facing impostors where we don't want that to generate motion vectors.

diff --git a/Engine/Source/Runtime/Engine/Classes/Components/StaticMeshComponent.h b/Engine/Source/Runtime/Engine/Classes/Components/StaticMeshComponent.h
--- a/Engine/Source/Runtime/Engine/Classes/Components/StaticMeshComponent.h
+++ b/Engine/Source/Runtime/Engine/Classes/Components/StaticMeshComponent.h
@@ -162,6 +162,14 @@ public:
 	UPROPERTY(EditAnywhere, AdvancedDisplay, BlueprintReadWrite, Category = Rendering)
 	uint8 bEvaluateWorldPositionOffset : 1;
 
+	/** 
+	 * Whether world position offset turns on velocity writes.
+	 * If the WPO isn't static then setting false may give incorrect motion vectors.
+	 * But if we know that the WPO is static then setting false may save performance.
+	 */
+	UPROPERTY(EditAnywhere, AdvancedDisplay, BlueprintReadWrite, Category = Rendering)
+	uint8 bWorldPositionOffsetWritesVelocity : 1;
+
 	/** 
 	 * Whether to evaluate World Position Offset for ray tracing. 
 	 * This is only used when running with r.RayTracing.Geometry.StaticMeshes.WPO=1 
@@ -547,6 +555,7 @@ public:
 	ENGINE_API virtual bool ComponentIsTouchingSelectionFrustum(const FConvexVolume& InSelBBox, const bool bConsiderOnlyBSP, const bool bMustEncompassEntireComponent) const override;
 #endif
 	virtual float GetStreamingScale() const override { return GetComponentTransform().GetMaximumAxisScale(); }
+	virtual bool SupportsWorldPositionOffsetVelocity() const override { return bWorldPositionOffsetWritesVelocity; }
 	//~ End UPrimitiveComponent Interface.
 
 	//~ Begin INavRelevantInterface Interface.
diff --git a/Engine/Source/Runtime/Engine/Private/Components/StaticMeshComponent.cpp b/Engine/Source/Runtime/Engine/Private/Components/StaticMeshComponent.cpp
--- a/Engine/Source/Runtime/Engine/Private/Components/StaticMeshComponent.cpp
+++ b/Engine/Source/Runtime/Engine/Private/Components/StaticMeshComponent.cpp
@@ -203,6 +203,7 @@ UStaticMeshComponent::UStaticMeshComponent(const FObjectInitializer& ObjectIniti
 	bDisallowNanite = false;
 	bForceDisableNanite = false;
 	bEvaluateWorldPositionOffset = true;
+	bWorldPositionOffsetWritesVelocity = true;
 	bEvaluateWorldPositionOffsetInRayTracing = false;
 	bInitialEvaluateWorldPositionOffset = false;
 	DistanceFieldIndirectShadowMinVisibility = .1f;
diff --git a/Engine/Source/Runtime/Foliage/Private/GrassInstancedStaticMeshComponent.cpp b/Engine/Source/Runtime/Foliage/Private/GrassInstancedStaticMeshComponent.cpp
--- a/Engine/Source/Runtime/Foliage/Private/GrassInstancedStaticMeshComponent.cpp
+++ b/Engine/Source/Runtime/Foliage/Private/GrassInstancedStaticMeshComponent.cpp
@@ -14,5 +14,5 @@ static TAutoConsoleVariable<bool> CVarEnableGrassInstancedWPOVelocity(
 
 bool UGrassInstancedStaticMeshComponent::SupportsWorldPositionOffsetVelocity() const 
 { 
-	return CVarEnableGrassInstancedWPOVelocity.GetValueOnAnyThread();
+	return CVarEnableGrassInstancedWPOVelocity.GetValueOnAnyThread() && Super::SupportsWorldPositionOffsetVelocity();
 }
