Remove render thread access of parent UTexture from FVirtualTexture2DResource.
Fixes race condition check firing when Virtual Texture sRGB flag is flipped.

diff --git a/Engine/Source/Runtime/Engine/Private/Texture2D.cpp b/Engine/Source/Runtime/Engine/Private/Texture2D.cpp
--- a/Engine/Source/Runtime/Engine/Private/Texture2D.cpp
+++ b/Engine/Source/Runtime/Engine/Private/Texture2D.cpp
@@ -1279,11 +1279,6 @@ bool UTexture2D::IsCurrentlyVirtualTextured() const
 }
 
 FVirtualTexture2DResource::FVirtualTexture2DResource()
-	: AllocatedVT(nullptr)
-	, VTData(nullptr)
-	, TextureOwner(nullptr)
-	, ProducerHandle()
-	, FirstMipToUse(0)
 {
 	//NOTE: Empty constructor for use with media textures (which do not derive from UTexture2D).
 
@@ -1292,21 +1287,34 @@ FVirtualTexture2DResource::FVirtualTexture2DResource()
 }
 
 FVirtualTexture2DResource::FVirtualTexture2DResource(const UTexture2D* InOwner, FVirtualTextureBuiltData* InVTData, int32 InFirstMipToUse)
-	: AllocatedVT(nullptr)
-	, VTData(InVTData)
-	, TextureOwner(InOwner)
 {
 	check(InOwner);
+	bSRGB = InOwner->SRGB;
+	bGreyScaleFormat = UE::TextureDefines::ShouldUseGreyScaleEditorVisualization(InOwner->CompressionSettings);
+	TextureReferenceRHI = InOwner->TextureReference.TextureReferenceRHI;
+
+	TextureName = InOwner->GetFName();
+	PackageName = InOwner->GetOutermost()->GetFName();
+
+	Filter = (ESamplerFilter)UDeviceProfileManager::Get().GetActiveProfile()->GetTextureLODSettings()->GetSamplerFilter(InOwner);
+	AddressU = InOwner->AddressX == TA_Wrap ? AM_Wrap : (InOwner->AddressX == TA_Clamp ? AM_Clamp : AM_Mirror);
+	AddressV = InOwner->AddressY == TA_Wrap ? AM_Wrap : (InOwner->AddressY == TA_Clamp ? AM_Clamp : AM_Mirror);
+
+	TexCreateFlags = InOwner->SRGB ? ETextureCreateFlags::SRGB : ETextureCreateFlags::None;
+	TexCreateFlags |= InOwner->bNotOfflineProcessed ? ETextureCreateFlags::None : ETextureCreateFlags::OfflineProcessed;
+	TexCreateFlags |= InOwner->bNoTiling ? ETextureCreateFlags::NoTiling : ETextureCreateFlags::None;
+
+	bContinuousUpdate = InOwner->IsVirtualTexturedWithContinuousUpdate();
+	bSinglePhysicalSpace = InOwner->IsVirtualTexturedWithSinglePhysicalSpace();
+
 	check(InVTData);
+	VTData = InVTData;
 
 	// Don't allow input mip bias to drop size below a single tile
 	const uint32 SizeInTiles = FMath::Max(VTData->GetWidthInTiles(), VTData->GetHeightInTiles());
 	const uint32 MaxMip = FMath::CeilLogTwo(SizeInTiles);
 	FirstMipToUse = FMath::Min((int32)MaxMip, InFirstMipToUse);
 
-	bSRGB = InOwner->SRGB;
-	bGreyScaleFormat = UE::TextureDefines::ShouldUseGreyScaleEditorVisualization( InOwner->CompressionSettings );
-
 	// Initialize this resource FeatureLevel, so it gets re-created on FeatureLevel changes
 	SetFeatureLevel(GMaxRHIFeatureLevel);
 }
@@ -1317,10 +1325,8 @@ FVirtualTexture2DResource::~FVirtualTexture2DResource()
 
 void FVirtualTexture2DResource::InitRHI()
 {
-	check(TextureOwner);
-	
 	TRACE_CPUPROFILER_EVENT_SCOPE(FVirtualTexture2DResource::InitRHI);
-	LLM_SCOPED_TAG_WITH_OBJECT_IN_SET(TextureOwner->GetOutermost(), ELLMTagSet::Assets);
+	LLM_SCOPED_TAG_WITH_STAT_NAME_IN_SET(FLowLevelMemTracker::Get().IsTagSetActive(ELLMTagSet::Assets) ? PackageName : NAME_None, ELLMTagSet::Assets, ELLMTracker::Default);
 
 	uint32 MaxAnisotropy = 0;
 	if (VirtualTextureScalability::IsAnisotropicFilteringEnabled())
@@ -1334,11 +1340,11 @@ void FVirtualTexture2DResource::InitRHI()
 	FSamplerStateInitializerRHI SamplerStateInitializer
 	(
 		// This will ensure nearest/linear/trilinear which does matter when sampling both the cache and the miptail
-		(ESamplerFilter)UDeviceProfileManager::Get().GetActiveProfile()->GetTextureLODSettings()->GetSamplerFilter(TextureOwner),
+		Filter,
 
 		// This doesn't really matter when sampling the cache texture but it does when sampling the miptail texture
-		TextureOwner->AddressX == TA_Wrap ? AM_Wrap : (TextureOwner->AddressX == TA_Clamp ? AM_Clamp : AM_Mirror),
-		TextureOwner->AddressY == TA_Wrap ? AM_Wrap : (TextureOwner->AddressY == TA_Clamp ? AM_Clamp : AM_Mirror),
+		AddressU,
+		AddressV,
 		AM_Wrap,
 
 		// This doesn't really matter when sampling the cache texture (as it only has a level 0, so whatever the bias that is sampled) but it does when we sample miptail texture
@@ -1359,12 +1365,9 @@ void FVirtualTexture2DResource::InitRHI()
 	const int32 MaxLevel = VTData->GetNumMips() - FirstMipToUse - 1;
 	check(MaxLevel >= 0);
 
-	const bool bContinuousUpdate = TextureOwner->IsVirtualTexturedWithContinuousUpdate();
-	const bool bSinglePhysicalSpace = TextureOwner->IsVirtualTexturedWithSinglePhysicalSpace();
-
 	FVTProducerDescription ProducerDesc;
-	ProducerDesc.Name = TextureOwner->GetFName();
-	ProducerDesc.FullNameHash = GetTypeHash(TextureOwner->GetName());
+	ProducerDesc.Name = TextureName;
+	ProducerDesc.FullNameHash = GetTypeHash(TextureName);
 	ProducerDesc.bContinuousUpdate = bContinuousUpdate;
 	ProducerDesc.Dimensions = 2;
 	ProducerDesc.TileSize = VTData->TileSize;
@@ -1382,7 +1385,7 @@ void FVirtualTexture2DResource::InitRHI()
 		ProducerDesc.LayerFormat[LayerIndex] = VTData->LayerTypes[LayerIndex];
 		ProducerDesc.LayerFallbackColor[LayerIndex] = VTData->LayerFallbackColors[LayerIndex];
 		ProducerDesc.PhysicalGroupIndex[LayerIndex] = bSinglePhysicalSpace ? 0 : LayerIndex;
-		ProducerDesc.bIsLayerSRGB[LayerIndex] = TextureOwner->SRGB;
+		ProducerDesc.bIsLayerSRGB[LayerIndex] = bSRGB;
 	}
 
 	FUploadingVirtualTexture* VirtualTexture = new FUploadingVirtualTexture(ProducerDesc.Name, VTData, FirstMipToUse);
@@ -1455,23 +1458,8 @@ void FVirtualTexture2DResource::InitializeEditorResources(IVirtualTexture* InVir
 			}
 		}
 
-		FString Name = TextureOwner->GetName();
-
-		FRHITextureCreateDesc Desc =
-			FRHITextureCreateDesc::Create2D(*Name, MipWidthInTiles * TileSizeInPixels, MipHeightInTiles * TileSizeInPixels, PixelFormat);
-
-		if (TextureOwner->SRGB)
-		{
-			Desc.AddFlags(ETextureCreateFlags::SRGB);
-		}
-		if (!TextureOwner->bNotOfflineProcessed)
-		{
-			Desc.AddFlags(ETextureCreateFlags::OfflineProcessed);
-		}
-		if (TextureOwner->bNoTiling)
-		{
-			Desc.AddFlags(ETextureCreateFlags::NoTiling);
-		}
+		FRHITextureCreateDesc Desc = FRHITextureCreateDesc::Create2D(*TextureName.ToString(), MipWidthInTiles * TileSizeInPixels, MipHeightInTiles * TileSizeInPixels, PixelFormat);
+		Desc.AddFlags(TexCreateFlags);
 
 		FTexture2DRHIRef Texture2DRHI = RHICreateTexture(Desc);
 
@@ -1519,7 +1507,7 @@ void FVirtualTexture2DResource::InitializeEditorResources(IVirtualTexture* InVir
 			check(MipHeight <= MipHeightInTiles * TileSizeInPixels);
 
 			const FRHITextureCreateDesc ResizedDesc =
-				FRHITextureCreateDesc::Create2D(*Name, MipWidth, MipHeight, PixelFormat)
+				FRHITextureCreateDesc::Create2D(*TextureName.ToString(), MipWidth, MipHeight, PixelFormat)
 				.SetFlags(Desc.Flags)
 				.SetInitialState(ERHIAccess::CopyDest);
 
@@ -1538,15 +1526,11 @@ void FVirtualTexture2DResource::InitializeEditorResources(IVirtualTexture* InVir
 		}
 
 		TextureRHI = Texture2DRHI;
-		TextureRHI->SetName(TextureOwner->GetFName());
-		RHIBindDebugLabelName(TextureRHI, *Name);
-		RHIUpdateTextureReference(TextureOwner->TextureReference.TextureReferenceRHI, TextureRHI);
-
-		bIgnoreGammaConversions = !TextureOwner->SRGB && !IsHDR(PixelFormat);
+		TextureRHI->SetName(TextureName);
+		RHIBindDebugLabelName(TextureRHI, *TextureName.ToString());
+		RHIUpdateTextureReference(TextureReferenceRHI, TextureRHI);
 
-		// re factored to ensure this is set earlier...make sure it's correct
-		ensure(bSRGB == TextureOwner->SRGB);
-		//bSRGB = TextureOwner->SRGB;
+		bIgnoreGammaConversions = !bSRGB && !IsHDR(PixelFormat);
 	}
 }
 #endif // WITH_EDITOR
@@ -1569,7 +1553,7 @@ class IAllocatedVirtualTexture* FVirtualTexture2DResource::AcquireAllocatedVT()
 		VTDesc.TileSize = VTData->TileSize;
 		VTDesc.TileBorderSize = VTData->TileBorderSize;
 		VTDesc.NumTextureLayers = VTData->GetNumLayers();
-		VTDesc.bShareDuplicateLayers = TextureOwner->IsVirtualTexturedWithSinglePhysicalSpace();
+		VTDesc.bShareDuplicateLayers = bSinglePhysicalSpace;
 
 		for (uint32 LayerIndex = 0u; LayerIndex < VTDesc.NumTextureLayers; ++LayerIndex)
 		{
diff --git a/Engine/Source/Runtime/Engine/Public/TextureResource.h b/Engine/Source/Runtime/Engine/Public/TextureResource.h
--- a/Engine/Source/Runtime/Engine/Public/TextureResource.h
+++ b/Engine/Source/Runtime/Engine/Public/TextureResource.h
@@ -211,11 +211,26 @@ public:
 	ENGINE_API virtual FIntPoint GetPhysicalTextureSize(uint32 LayerIndex) const;
 
 protected:
-	class IAllocatedVirtualTexture* AllocatedVT;
-	struct FVirtualTextureBuiltData* VTData;
-	const UTexture2D* TextureOwner;
+	/** The FName of the texture asset */
+	FName TextureName;
+	/** The FName of the texture package for stats */
+	FName PackageName;
+	/** Cached sampler config */
+	TEnumAsByte<ESamplerFilter> Filter = SF_Point;
+	TEnumAsByte<ESamplerAddressMode> AddressU = AM_Wrap;
+	TEnumAsByte<ESamplerAddressMode> AddressV = AM_Wrap;
+	/** Cached flags for texture creation. */
+	ETextureCreateFlags TexCreateFlags = ETextureCreateFlags::None;
+	/** Cached runtime virtual texture settings */
+	bool bContinuousUpdate = false;
+	bool bSinglePhysicalSpace = false;
+	/** Mip offset */
+	int32 FirstMipToUse = 0;
+	/** Built data owned by texture asset */
+	struct FVirtualTextureBuiltData* VTData = nullptr;
+	/** Local allocated VT objects used for editor views etc. */
+	class IAllocatedVirtualTexture* AllocatedVT = nullptr;
 	FVirtualTextureProducerHandle ProducerHandle;
-	int32 FirstMipToUse;
 };
 
 /** A dynamic 2D texture resource. */
