Fix for calling SVT::ReadVoxel() with negative TileIndex which caused an out of bounds memory access. Added missing checks and validation when reading/writing in SVT::FTextureData. Added missing check against the maximum page table memory size. The old check against the maximum tile data size is now not an error anymore because the streaming system can gracefully handle this by not streaming in higher resolution mips. Simplified/optimized page table generation when converting from VDB to SVT.

#rnx

diff --git a/Engine/Source/Runtime/Engine/Private/SparseVolumeTexture/SparseVolumeTextureData.cpp b/Engine/Source/Runtime/Engine/Private/SparseVolumeTexture/SparseVolumeTextureData.cpp
--- a/Engine/Source/Runtime/Engine/Private/SparseVolumeTexture/SparseVolumeTextureData.cpp
+++ b/Engine/Source/Runtime/Engine/Private/SparseVolumeTexture/SparseVolumeTextureData.cpp
@@ -40,6 +40,21 @@ namespace UE
 namespace SVT
 {
 
+static bool IsInBounds(const FIntVector3& Point, const FIntVector3& Min, const FIntVector3& Max)
+{
+	return Point.X >= Min.X && Point.Y >= Min.Y && Point.Z >= Min.Z
+		&& Point.X < Max.X && Point.Y < Max.Y && Point.Z < Max.Z;
+}
+
+static FIntVector3 ShiftRightAndMax(const FIntVector3& Value, uint32 ShiftBy, int32 MinValue)
+{
+	FIntVector3 Result = FIntVector3(
+		FMath::Max(Value.X >> ShiftBy, MinValue),
+		FMath::Max(Value.Y >> ShiftBy, MinValue),
+		FMath::Max(Value.Z >> ShiftBy, MinValue));
+	return Result;
+}
+
 bool FTextureData::Create(const ITextureDataProvider& DataProvider)
 {
 	const FTextureDataCreateInfo CreateInfo = DataProvider.GetCreateInfo();
@@ -56,17 +71,34 @@ bool FTextureData::Create(const ITextureDataProvider& DataProvider)
 		return false;
 	}
 
+	const int64 PageTableSizeBytes = (int64)Header.PageTableVolumeResolution.X * (int64)Header.PageTableVolumeResolution.Y * (int64)Header.PageTableVolumeResolution.Z * sizeof(uint32);
+	if (PageTableSizeBytes > SVT::MaxResourceSize)
+	{
+		UE_LOG(LogSparseVolumeTextureData, Warning, TEXT("SparseVolumeTexture page table texture memory size (%ll) exceeds the 2048MB GPU resource limit!"), (long long)PageTableSizeBytes);
+		return false;
+	}
+
+	const int32 FormatSize[] = { GPixelFormats[Header.AttributesFormats[0]].BlockBytes, GPixelFormats[Header.AttributesFormats[1]].BlockBytes };
+	uint8 NullTileValuesU8[2][sizeof(float) * 4] = {};
+	for (int32 i = 0; i < 2; ++i)
+	{
+		if (Header.AttributesFormats[i] != PF_Unknown)
+		{
+			SVT::WriteVoxel(0, NullTileValuesU8[i], Header.AttributesFormats[i], Header.FallbackValues[i]);
+			FallbackValuesQuantized[i] = SVT::ReadVoxel(0, NullTileValuesU8[i], Header.AttributesFormats[i]);
+		}
+	}
+
 	MipMaps.SetNum(1);
 
 	// Tag all pages with valid data
 	TBitArray PagesWithData(false, Header.PageTableVolumeResolution.Z * Header.PageTableVolumeResolution.Y * Header.PageTableVolumeResolution.X);
 	DataProvider.IteratePhysicalSource([&](const FIntVector3& Coord, int32 AttributesIdx, int32 ComponentIdx, float VoxelValue)
 		{
-			const bool bIsFallbackValue = (VoxelValue == CreateInfo.FallbackValues[AttributesIdx][ComponentIdx]);
+			const bool bIsFallbackValue = (VoxelValue == Header.FallbackValues[AttributesIdx][ComponentIdx]);
 			if (!bIsFallbackValue)
 			{
-				check(Coord.X >= Header.VirtualVolumeAABBMin.X && Coord.Y >= Header.VirtualVolumeAABBMin.Y && Coord.Z >= Header.VirtualVolumeAABBMin.Z);
-				check(Coord.X < Header.VirtualVolumeAABBMax.X && Coord.Y < Header.VirtualVolumeAABBMax.Y && Coord.Z < Header.VirtualVolumeAABBMax.Z);
+				check(IsInBounds(Coord, Header.VirtualVolumeAABBMin, Header.VirtualVolumeAABBMax));
 
 #if SVT_CORRECT_TILE_ALLOCATION_FOR_LINEAR_FILTERING
 				// Tag all pages touching the 3x3 area around this voxel in order to properly support bilinear/border voxels
@@ -77,17 +109,10 @@ bool FTextureData::Create(const ITextureDataProvider& DataProvider)
 						for (int32 OffsetX = -1; OffsetX < 2; ++OffsetX)
 						{
 							const FIntVector3 GridCoord = Coord + FIntVector3(OffsetX, OffsetY, OffsetZ);
-							if (GridCoord.X >= Header.VirtualVolumeAABBMin.X
-								&& GridCoord.Y >= Header.VirtualVolumeAABBMin.Y 
-								&& GridCoord.Z >= Header.VirtualVolumeAABBMin.Z
-								&& GridCoord.X < Header.VirtualVolumeAABBMax.X
-								&& GridCoord.Y < Header.VirtualVolumeAABBMax.Y
-								&& GridCoord.Z < Header.VirtualVolumeAABBMax.Z)
+							if (IsInBounds(GridCoord, Header.VirtualVolumeAABBMin, Header.VirtualVolumeAABBMax))
 							{
 								const FIntVector3 PageCoord = (GridCoord / SPARSE_VOLUME_TILE_RES) - Header.PageTableVolumeAABBMin;
-								check(PageCoord.X >= 0 && PageCoord.Y >= 0 && PageCoord.Z >= 0);
-								check(PageCoord.X < Header.PageTableVolumeResolution.X && PageCoord.Y < Header.PageTableVolumeResolution.Y && PageCoord.Z < Header.PageTableVolumeResolution.Z);
-
+								check(IsInBounds(PageCoord, FIntVector3::ZeroValue, Header.PageTableVolumeResolution));
 								const int32 PageIndex = PageCoord.Z * (Header.PageTableVolumeResolution.Y * Header.PageTableVolumeResolution.X) + PageCoord.Y * Header.PageTableVolumeResolution.X + PageCoord.X;
 								PagesWithData[PageIndex] = true;
 							}
@@ -97,102 +122,61 @@ bool FTextureData::Create(const ITextureDataProvider& DataProvider)
 #else // SVT_CORRECT_TILE_ALLOCATION_FOR_LINEAR_FILTERING
 				const FIntVector3 GridCoord = Coord;
 				const FIntVector3 PageCoord = (GridCoord / SPARSE_VOLUME_TILE_RES) - Header.PageTableVolumeAABBMin;
-				check(PageCoord.X >= 0 && PageCoord.Y >= 0 && PageCoord.Z >= 0);
-				check(PageCoord.X < Header.PageTableVolumeResolution.X && PageCoord.Y < Header.PageTableVolumeResolution.Y && PageCoord.Z < Header.PageTableVolumeResolution.Z);
-
+				check(IsInBounds(PageCoord, FIntVector3::ZeroValue, Header.PageTableVolumeResolution));
 				const int32 PageIndex = PageCoord.Z * (Header.PageTableVolumeResolution.Y * Header.PageTableVolumeResolution.X) + PageCoord.Y * Header.PageTableVolumeResolution.X + PageCoord.X;
 				PagesWithData[PageIndex] = true;
 #endif // SVT_CORRECT_TILE_ALLOCATION_FOR_LINEAR_FILTERING
 			}
 		});
 
-	// Allocate some memory for temp data (worst case)
-	TArray<FIntVector3> LinearAllocatedPages;
-	LinearAllocatedPages.SetNum(Header.PageTableVolumeResolution.X * Header.PageTableVolumeResolution.Y * Header.PageTableVolumeResolution.Z);
+	// Initialize the page table and tile data arrays.
+	MipMaps[0].PageTable.SetNumZeroed(Header.PageTableVolumeResolution.X * Header.PageTableVolumeResolution.Y * Header.PageTableVolumeResolution.Z);
 
-	// Go over each potential page from the source data and push allocate it if it has any data.
-	// Otherwise point to the default empty page.
 	int32 NumAllocatedPages = 0;
-	for (int32 PageZ = 0; PageZ < Header.PageTableVolumeResolution.Z; ++PageZ)
+	for (TConstSetBitIterator It(PagesWithData); It; ++It)
 	{
-		for (int32 PageY = 0; PageY < Header.PageTableVolumeResolution.Y; ++PageY)
-		{
-			for (int32 PageX = 0; PageX < Header.PageTableVolumeResolution.X; ++PageX)
-			{
-				const int32 PageIndex = PageZ * (Header.PageTableVolumeResolution.Y * Header.PageTableVolumeResolution.X) + PageY * Header.PageTableVolumeResolution.X + PageX;
-				const bool bHasAnyData = PagesWithData[PageIndex];
-				if (bHasAnyData)
-				{
-					LinearAllocatedPages[NumAllocatedPages] = FIntVector3(PageX, PageY, PageZ);
-					NumAllocatedPages++;
-				}
-			}
-		}
+		const int32 PageIndex = It.GetIndex();
+		MipMaps[0].PageTable[PageIndex] = (uint32)++NumAllocatedPages;
 	}
-	LinearAllocatedPages.SetNum(NumAllocatedPages);
 
-	// Initialize the SparseVolumeTexture page and tile.
-	const int32 FormatSize[] = { GPixelFormats[Header.AttributesFormats[0]].BlockBytes, GPixelFormats[Header.AttributesFormats[1]].BlockBytes };
-	MipMaps[0].PageTable.SetNumZeroed(Header.PageTableVolumeResolution.X * Header.PageTableVolumeResolution.Y * Header.PageTableVolumeResolution.Z);
 	MipMaps[0].PhysicalTileDataA.SetNum((int64)NumAllocatedPages * (int64)SVT::NumVoxelsPerPaddedTile * (int64)FormatSize[0]);
 	MipMaps[0].PhysicalTileDataB.SetNum((int64)NumAllocatedPages * (int64)SVT::NumVoxelsPerPaddedTile * (int64)FormatSize[1]);
 	MipMaps[0].NumPhysicalTiles = NumAllocatedPages;
 
-	// Generate page table and tile volume data by splatting the data
-	{
-		// Fill all tiles with null tile data. See below for an explanation of why this is necessary.
+	// Fill all tiles with fallback values/null tile values. We can't guarantee that all voxels of a page will be written to in the next step, so if we don't do this,
+	// we might end up with physical tiles that have only been partially filled with valid voxels.
+	ParallelFor(NumAllocatedPages, [&](int32 TileIndex)
 		{
-			// Fill all tiles with fallback values/null tile values. We can't guarantee that all voxels of a page will be written to in the next step, so if we don't do this,
-			// we might end up with physical tiles that have only been partially filled with valid voxels.
-			ParallelFor(NumAllocatedPages, [&](int32 TileIndex)
+			uint8* TileData[2] = { MipMaps[0].PhysicalTileDataA.GetData(), MipMaps[0].PhysicalTileDataB.GetData() };
+			for (int32 i = 0; i < 2; ++i)
+			{
+				if (FormatSize[i])
 				{
-					for (int32 Z = 0; Z < SPARSE_VOLUME_TILE_RES_PADDED; ++Z)
+					for (int32 VoxelIndex = 0; VoxelIndex < SVT::NumVoxelsPerPaddedTile; ++VoxelIndex)
 					{
-						for (int32 Y = 0; Y < SPARSE_VOLUME_TILE_RES_PADDED; ++Y)
-						{
-							for (int32 X = 0; X < SPARSE_VOLUME_TILE_RES_PADDED; ++X)
-							{
-								const FIntVector3 WriteCoord = FIntVector3(X - 1, Y - 1, Z - 1);
-								WriteTileDataVoxel(TileIndex, WriteCoord, 0 /*MipLevel*/, 0 /*AttributesIndex*/, Header.FallbackValues[0]);
-								WriteTileDataVoxel(TileIndex, WriteCoord, 0 /*MipLevel*/, 1 /*AttributesIndex*/, Header.FallbackValues[1]);
-							}
-						}
+						const int64 VoxelIndexGlobal = (int64)TileIndex * (int64)SVT::NumVoxelsPerPaddedTile + (int64)VoxelIndex;
+						FMemory::Memcpy(TileData[i] + VoxelIndexGlobal * FormatSize[i], NullTileValuesU8[i], FormatSize[i]);
 					}
-				});
-
-			FallbackValuesQuantized[0] = ReadTileDataVoxel(0, FIntVector::ZeroValue, 0 /*MipLevel*/, 0 /*AttributesIndex*/);
-			FallbackValuesQuantized[1] = ReadTileDataVoxel(0, FIntVector::ZeroValue, 0 /*MipLevel*/, 1 /*AttributesIndex*/);
-		}
-
-		// Set up the page table entries
-		for (int32 i = 0; i < NumAllocatedPages; ++i)
-		{
-			const FIntVector3 PageCoordToSplat = LinearAllocatedPages[i];
-			const int32 TileIndex = (PageCoordToSplat.Z * Header.PageTableVolumeResolution.X * Header.PageTableVolumeResolution.Y) + (PageCoordToSplat.Y * Header.PageTableVolumeResolution.X) + PageCoordToSplat.X;
-			MipMaps[0].PageTable[TileIndex] = (uint32)(i + 1); // +1 because 0 is reserved for the null tile which we don't actually store in memory
-		}
-	}
+				}
+			}
+		});
 
 	// Write physical tile data
 	DataProvider.IteratePhysicalSource([&](const FIntVector3& Coord, int32 AttributesIdx, int32 ComponentIdx, float VoxelValue)
 		{
 			const FIntVector3 GridCoord = Coord;
-			check(GridCoord.X >= 0 && GridCoord.Y >= 0 && GridCoord.Z >= 0);
-			check(GridCoord.X < Header.VirtualVolumeAABBMax.X && GridCoord.Y < Header.VirtualVolumeAABBMax.Y && GridCoord.Z < Header.VirtualVolumeAABBMax.Z);
+			check(IsInBounds(GridCoord, FIntVector3::ZeroValue, Header.VirtualVolumeAABBMax));
 			const FIntVector3 PageCoord = (GridCoord / SPARSE_VOLUME_TILE_RES) - Header.PageTableVolumeAABBMin;
-			check(PageCoord.X >= 0 && PageCoord.Y >= 0 && PageCoord.Z >= 0);
-			check(PageCoord.X < Header.PageTableVolumeResolution.X && PageCoord.Y < Header.PageTableVolumeResolution.Y && PageCoord.Z < Header.PageTableVolumeResolution.Z);
-
+			check(IsInBounds(PageCoord, FIntVector3::ZeroValue, Header.PageTableVolumeResolution));
 			const int32 PageIndex = PageCoord.Z * (Header.PageTableVolumeResolution.Y * Header.PageTableVolumeResolution.X) + PageCoord.Y * Header.PageTableVolumeResolution.X + PageCoord.X;
 
-			if (!PagesWithData[PageIndex])
+			const uint32 PageEntry = MipMaps[0].PageTable[PageIndex];
+			if (PageEntry != 0)
 			{
-				return;
+				const int32 TileIndex = (int32)PageEntry - 1; // -1 because 0 is reserved for the null tile which we don't actually store in memory
+				const FIntVector3 TileLocalCoord = GridCoord % SPARSE_VOLUME_TILE_RES;
+				WriteTileDataVoxel(TileIndex, TileLocalCoord, 0 /*MipLevel*/, AttributesIdx, FVector4f(VoxelValue, VoxelValue, VoxelValue, VoxelValue), ComponentIdx);
 			}
-
-			const int32 TileIndex = (int32)MipMaps[0].PageTable[PageIndex];
-			const FIntVector3 TileLocalCoord = GridCoord % SPARSE_VOLUME_TILE_RES;
-			WriteTileDataVoxel(TileIndex - 1, TileLocalCoord, 0 /*MipLevel*/, AttributesIdx, FVector4f(VoxelValue, VoxelValue, VoxelValue, VoxelValue), ComponentIdx);  // -1 because 0 is reserved for the null tile which we don't actually store in memory
 		});
 
 	return true;
@@ -212,15 +196,22 @@ bool FTextureData::CreateFromDense(const FTextureDataCreateInfo& CreateInfo, con
 		return false;
 	}
 
+	const int64 PageTableSizeBytes = (int64)Header.PageTableVolumeResolution.X * (int64)Header.PageTableVolumeResolution.Y * (int64)Header.PageTableVolumeResolution.Z * sizeof(uint32);
+	if (PageTableSizeBytes > SVT::MaxResourceSize)
+	{
+		UE_LOG(LogSparseVolumeTextureData, Warning, TEXT("SparseVolumeTexture page table texture memory size (%ll) exceeds the 2048MB GPU resource limit!"), (long long)PageTableSizeBytes);
+		return false;
+	}
+
 	const int64 FormatSize[] = { GPixelFormats[Header.AttributesFormats[0]].BlockBytes, GPixelFormats[Header.AttributesFormats[1]].BlockBytes };
 	const FIntVector3 FractionalPageOffset = Header.VirtualVolumeAABBMin % SPARSE_VOLUME_TILE_RES; // Offset from page origin to source data volume origin
 	uint8 NullTileValuesU8[2][sizeof(float) * 4] = {};
 	for (int32 i = 0; i < 2; ++i)
 	{
-		if (CreateInfo.AttributesFormats[i] != PF_Unknown)
+		if (Header.AttributesFormats[i] != PF_Unknown)
 		{
-			SVT::WriteVoxel(0, NullTileValuesU8[i], CreateInfo.AttributesFormats[i], CreateInfo.FallbackValues[i]);
-			FallbackValuesQuantized[i] = SVT::ReadVoxel(0, NullTileValuesU8[i], CreateInfo.AttributesFormats[i]);
+			SVT::WriteVoxel(0, NullTileValuesU8[i], Header.AttributesFormats[i], Header.FallbackValues[i]);
+			FallbackValuesQuantized[i] = SVT::ReadVoxel(0, NullTileValuesU8[i], Header.AttributesFormats[i]);
 		}
 	}
 
@@ -253,12 +244,9 @@ bool FTextureData::CreateFromDense(const FTextureDataCreateInfo& CreateInfo, con
 				for (int64 PageX = PageBegin.X; PageX < PageEnd.X; ++PageX)
 				{
 					const FIntVector3 GridCoord = FIntVector3(PageX * SPARSE_VOLUME_TILE_RES + FractionalPageOffset.X, Y, Z); // Global/absolute voxel coordinate
-					check(GridCoord.X >= 0 && GridCoord.Y >= 0 && GridCoord.Z >= 0);
-					check(GridCoord.X < Header.VirtualVolumeAABBMax.X && GridCoord.Y < Header.VirtualVolumeAABBMax.Y && GridCoord.Z < Header.VirtualVolumeAABBMax.Z);
+					check(IsInBounds(GridCoord, FIntVector3::ZeroValue, Header.VirtualVolumeAABBMax));
 					const FIntVector3 PageCoord = (GridCoord / SPARSE_VOLUME_TILE_RES) - Header.PageTableVolumeAABBMin;
-					check(PageCoord.X >= 0 && PageCoord.Y >= 0 && PageCoord.Z >= 0);
-					check(PageCoord.X < Header.PageTableVolumeResolution.X && PageCoord.Y < Header.PageTableVolumeResolution.Y && PageCoord.Z < Header.PageTableVolumeResolution.Z);
-
+					check(IsInBounds(PageCoord, FIntVector3::ZeroValue, Header.PageTableVolumeResolution));
 					const int32 PageIndex = PageCoord.Z * (Header.PageTableVolumeResolution.Y * Header.PageTableVolumeResolution.X) + PageCoord.Y * Header.PageTableVolumeResolution.X + PageCoord.X;
 					const uint32 TileIndex = MipMaps[0].PageTable[PageIndex];
 
@@ -346,16 +334,16 @@ bool FTextureData::CreateFromDense(const FTextureDataCreateInfo& CreateInfo, con
 		// Iterate over physical tiles (which at this point equals the number of allocated pages); NOT logical pages!
 		ParallelFor(NumAllocatedPages, [&](int32 TileIndex)
 			{
-				for (int32 VoxelIndex = 0; VoxelIndex < SVT::NumVoxelsPerPaddedTile; ++VoxelIndex)
+				uint8* TileData[2] = { MipMaps[0].PhysicalTileDataA.GetData(), MipMaps[0].PhysicalTileDataB.GetData() };
+				for (int32 i = 0; i < 2; ++i)
 				{
-					const int64 VoxelIndexGlobal = (int64)TileIndex * (int64)SVT::NumVoxelsPerPaddedTile + (int64)VoxelIndex;
-					if (FormatSize[0])
+					if (FormatSize[i])
 					{
-						FMemory::Memcpy(MipMaps[0].PhysicalTileDataA.GetData() + VoxelIndexGlobal * FormatSize[0], NullTileValuesU8[0], FormatSize[0]);
-					}
-					if (FormatSize[1])
-					{
-						FMemory::Memcpy(MipMaps[0].PhysicalTileDataB.GetData() + VoxelIndexGlobal * FormatSize[1], NullTileValuesU8[1], FormatSize[1]);
+						for (int32 VoxelIndex = 0; VoxelIndex < SVT::NumVoxelsPerPaddedTile; ++VoxelIndex)
+						{
+							const int64 VoxelIndexGlobal = (int64)TileIndex * (int64)SVT::NumVoxelsPerPaddedTile + (int64)VoxelIndex;
+							FMemory::Memcpy(TileData[i] + VoxelIndexGlobal * FormatSize[i], NullTileValuesU8[i], FormatSize[i]);
+						}
 					}
 				}
 			});
@@ -385,19 +373,17 @@ uint32 FTextureData::ReadPageTable(const FIntVector3& PageTableCoord, int32 MipL
 	{
 		return INDEX_NONE;
 	}
-	const FIntVector3 PageTableResolution = FIntVector3(
-		FMath::Max(1, Header.PageTableVolumeResolution.X >> MipLevel),
-		FMath::Max(1, Header.PageTableVolumeResolution.Y >> MipLevel),
-		FMath::Max(1, Header.PageTableVolumeResolution.Z >> MipLevel));
-	if (PageTableCoord.X < 0 || PageTableCoord.Y < 0 || PageTableCoord.Z < 0
-		|| PageTableCoord.X >= PageTableResolution.X
-		|| PageTableCoord.Y >= PageTableResolution.Y
-		|| PageTableCoord.Z >= PageTableResolution.Z)
+	const FIntVector3 PageTableAABBMipMin = Header.PageTableVolumeAABBMin >> MipLevel;
+	const FIntVector3 PageTableAABBMipMax = ShiftRightAndMax(Header.PageTableVolumeAABBMax, MipLevel, 1);
+	if (!IsInBounds(PageTableCoord, PageTableAABBMipMin, PageTableAABBMipMax))
 	{
 		return INDEX_NONE;
 	}
 
-	const int32 PageIndex = PageTableCoord.Z * (PageTableResolution.Y * PageTableResolution.X) + PageTableCoord.Y * PageTableResolution.X + PageTableCoord.X;
+	const FIntVector3 PageTableResolution = PageTableAABBMipMax - PageTableAABBMipMin;
+	check(PageTableResolution == ShiftRightAndMax(Header.PageTableVolumeResolution, MipLevel, 1));
+	const FIntVector3 LocalPageTableCoord = PageTableCoord - PageTableAABBMipMin; // Local to the AABB of the page table
+	const int32 PageIndex = LocalPageTableCoord.Z * (PageTableResolution.Y * PageTableResolution.X) + LocalPageTableCoord.Y * PageTableResolution.X + LocalPageTableCoord.X;
 	if (MipMaps[MipLevel].PageTable.IsValidIndex(PageIndex))
 	{
 		return MipMaps[MipLevel].PageTable[PageIndex];
@@ -408,30 +394,34 @@ uint32 FTextureData::ReadPageTable(const FIntVector3& PageTableCoord, int32 MipL
 FVector4f FTextureData::ReadTileDataVoxel(int32 TileIndex, const FIntVector3& TileDataCoord, int32 MipLevel, int32 AttributesIdx) const
 {
 	check(AttributesIdx >= 0 && AttributesIdx <= 1);
-	if (TileIndex == INDEX_NONE || (AttributesIdx == 0 && MipMaps[MipLevel].PhysicalTileDataA.IsEmpty()) || (AttributesIdx == 1 && MipMaps[MipLevel].PhysicalTileDataB.IsEmpty()))
+
+	if (!MipMaps.IsValidIndex(MipLevel)
+		|| !IsInBounds(TileDataCoord, FIntVector3::ZeroValue, FIntVector3(SPARSE_VOLUME_TILE_RES))
+		|| TileIndex < 0 || TileIndex >= MipMaps[MipLevel].NumPhysicalTiles
+		|| (AttributesIdx == 0 && MipMaps[MipLevel].PhysicalTileDataA.IsEmpty()) 
+		|| (AttributesIdx == 1 && MipMaps[MipLevel].PhysicalTileDataB.IsEmpty()))
 	{
-		return FVector4f();
+		return FallbackValuesQuantized[AttributesIdx];
 	}
 
-	// SV_TODO check for TileIndex oob
+	const EPixelFormat Format = Header.AttributesFormats[AttributesIdx];
 
-	if (TileDataCoord.X < 0 || TileDataCoord.Y < 0 || TileDataCoord.Z < 0
-		|| TileDataCoord.X >= SPARSE_VOLUME_TILE_RES
-		|| TileDataCoord.Y >= SPARSE_VOLUME_TILE_RES
-		|| TileDataCoord.Z >= SPARSE_VOLUME_TILE_RES)
-	{
-		return FVector4f();
-	}
+	// Sanity check to make sure that the tile data offset is within bounds
+	const int64 TileByteOffset = (int64)TileIndex * (int64)SVT::NumVoxelsPerPaddedTile * (int64)GPixelFormats[Format].BlockBytes;
+	const int64 NumTileBytes = AttributesIdx == 0 ? MipMaps[MipLevel].PhysicalTileDataA.Num() : MipMaps[MipLevel].PhysicalTileDataB.Num();
+	check(TileByteOffset < NumTileBytes);
 
 	const FInt64Vector3 CoordPadded = FInt64Vector3(TileDataCoord) + FInt64Vector3(SPARSE_VOLUME_TILE_BORDER);
 	const int64 VoxelIndex = (int64)TileIndex * (int64)SVT::NumVoxelsPerPaddedTile + CoordPadded.Z * (SPARSE_VOLUME_TILE_RES_PADDED * SPARSE_VOLUME_TILE_RES_PADDED) + CoordPadded.Y * SPARSE_VOLUME_TILE_RES_PADDED + CoordPadded.X;
 	const uint8* TileData = AttributesIdx == 0 ? MipMaps[MipLevel].PhysicalTileDataA.GetData() : MipMaps[MipLevel].PhysicalTileDataB.GetData();
-	const EPixelFormat Format = Header.AttributesFormats[AttributesIdx];
 	return SVT::ReadVoxel(VoxelIndex, TileData, Format);
 }
 
 FVector4f FTextureData::Load(const FIntVector3& VolumeCoord, int32 MipLevel, int32 AttributesIdx, const FTextureDataAddressingInfo& AddressingInfo) const
 {
+	// Out of bounds VolumeCoord or MipLevel is caught by ReadPageTable() and handled by returning INDEX_NONE
+	check(AttributesIdx >= 0 && AttributesIdx <= 1);
+
 	auto ApplyAddressMode = [](int32 x, int32 Width, TextureAddress Mode)
 	{
 		switch (Mode)
@@ -447,43 +437,34 @@ FVector4f FTextureData::Load(const FIntVector3& VolumeCoord, int32 MipLevel, int
 		}
 		return x;
 	};
-	if (!MipMaps.IsValidIndex(MipLevel))
-	{
-		return FVector4f();
-	}
 	const FIntVector3 AddressModeVolumeCoord = FIntVector3(
 		ApplyAddressMode(VolumeCoord.X, FMath::Max(1, AddressingInfo.VolumeResolution.X >> MipLevel), AddressingInfo.AddressX),
-		ApplyAddressMode(VolumeCoord.Y, FMath::Max(1, AddressingInfo.VolumeResolution.Y >> MipLevel), AddressingInfo.AddressY), 
+		ApplyAddressMode(VolumeCoord.Y, FMath::Max(1, AddressingInfo.VolumeResolution.Y >> MipLevel), AddressingInfo.AddressY),
 		ApplyAddressMode(VolumeCoord.Z, FMath::Max(1, AddressingInfo.VolumeResolution.Z >> MipLevel), AddressingInfo.AddressZ));
-	const FIntVector3 PageTableCoord = (AddressModeVolumeCoord / SPARSE_VOLUME_TILE_RES) - (Header.PageTableVolumeAABBMin >> MipLevel);
-	const uint32 TileIndex = ReadPageTable(PageTableCoord, MipLevel);
-	if (TileIndex == 0)
+	const FIntVector3 PageTableCoord = AddressModeVolumeCoord / SPARSE_VOLUME_TILE_RES;
+	const uint32 PageTableEntry = ReadPageTable(PageTableCoord, MipLevel);
+	if (PageTableEntry == INDEX_NONE || PageTableEntry == 0)
 	{
 		return FallbackValuesQuantized[AttributesIdx];
 	}
-	const FIntVector3 VoxelCoord = AddressModeVolumeCoord % SPARSE_VOLUME_TILE_RES;
-	const FVector4f Sample = ReadTileDataVoxel(TileIndex - 1, VoxelCoord, MipLevel, AttributesIdx); // -1 because we don't store the null tile in memory.
+	check(PageTableEntry > 0);
+	const int32 TileIndex = (int32)PageTableEntry - 1; // -1 because we don't store the null tile in memory.
+	const FIntVector3 TileLocalCoord = AddressModeVolumeCoord % SPARSE_VOLUME_TILE_RES;
+	const FVector4f Sample = ReadTileDataVoxel(TileIndex, TileLocalCoord, MipLevel, AttributesIdx);
 	return Sample;
 }
 
 void FTextureData::WriteTileDataVoxel(int32 TileIndex, const FIntVector3& TileDataCoord, int32 MipLevel, int32 AttributesIdx, const FVector4f& Value, int32 DstComponent)
 {
 	check(AttributesIdx >= 0 && AttributesIdx <= 1);
-	if ((AttributesIdx == 0 && MipMaps[MipLevel].PhysicalTileDataA.IsEmpty()) || (AttributesIdx == 1 && MipMaps[MipLevel].PhysicalTileDataB.IsEmpty()))
-	{
-		return;
-	}
-
-	// SV_TODO check for TileIndex oob
-
+	
 	// Allow TileDataCoord to extend past SPARSE_VOLUME_TILE_RES by SPARSE_VOLUME_TILE_BORDER.
 	// This allows us to write the border voxels without changing how TileDataCoord needs to be computed for other use cases.
-	if (TileDataCoord.X < -SPARSE_VOLUME_TILE_BORDER
-		|| TileDataCoord.Y < -SPARSE_VOLUME_TILE_BORDER
-		|| TileDataCoord.Z < -SPARSE_VOLUME_TILE_BORDER
-		|| TileDataCoord.X >= (SPARSE_VOLUME_TILE_RES + SPARSE_VOLUME_TILE_BORDER)
-		|| TileDataCoord.Y >= (SPARSE_VOLUME_TILE_RES + SPARSE_VOLUME_TILE_BORDER)
-		|| TileDataCoord.Z >= (SPARSE_VOLUME_TILE_RES + SPARSE_VOLUME_TILE_BORDER))
+	if (!IsInBounds(TileDataCoord, FIntVector3(-SPARSE_VOLUME_TILE_BORDER), FIntVector3(SPARSE_VOLUME_TILE_RES + SPARSE_VOLUME_TILE_BORDER))
+		|| !MipMaps.IsValidIndex(MipLevel)
+		|| TileIndex < 0 || TileIndex >= MipMaps[MipLevel].NumPhysicalTiles
+		|| (AttributesIdx == 0 && MipMaps[MipLevel].PhysicalTileDataA.IsEmpty())
+		|| (AttributesIdx == 1 && MipMaps[MipLevel].PhysicalTileDataB.IsEmpty()))
 	{
 		return;
 	}
@@ -528,12 +509,7 @@ bool FTextureData::GenerateMipMaps(const FTextureDataAddressingInfo& AddressingI
 		MipMaps.AddDefaulted();
 
 		const FIntVector3 PageTableVolumeAABBMin = (Header.VirtualVolumeAABBMin / SPARSE_VOLUME_TILE_RES) >> MipLevel;
-		const FIntVector3 ParentPageTableVolumeAABBMin = (Header.VirtualVolumeAABBMin / SPARSE_VOLUME_TILE_RES) >> (MipLevel - 1);
-		const FIntVector3 PageTableVolumeResolution = FIntVector3(
-			FMath::Max(1, Header.PageTableVolumeResolution.X >> MipLevel), 
-			FMath::Max(1, Header.PageTableVolumeResolution.Y >> MipLevel), 
-			FMath::Max(1, Header.PageTableVolumeResolution.Z >> MipLevel));
-		
+		const FIntVector3 PageTableVolumeResolution = ShiftRightAndMax(Header.PageTableVolumeResolution, MipLevel, 1);
 		check(FMath::IsPowerOfTwo(PageTableVolumeResolution.X)
 			&& FMath::IsPowerOfTwo(PageTableVolumeResolution.Y)
 			&& FMath::IsPowerOfTwo(PageTableVolumeResolution.Z));
@@ -556,7 +532,7 @@ bool FTextureData::GenerateMipMaps(const FTextureDataAddressingInfo& AddressingI
 					for (int32 OffsetIdx = 0; OffsetIdx < 8; ++OffsetIdx)
 					{
 						const FIntVector3 Offset = FIntVector3(OffsetIdx, OffsetIdx >> 1, OffsetIdx >> 2) & 1;
-						const FIntVector3 ParentPageTableCoord = (PageTableVolumeAABBMin + PageCoord) * 2 + Offset - ParentPageTableVolumeAABBMin;
+						const FIntVector3 ParentPageTableCoord = (PageTableVolumeAABBMin + PageCoord) * 2 + Offset;
 						const uint32 PageSample = ReadPageTable(ParentPageTableCoord, MipLevel - 1);
 
 						if (PageSample != INDEX_NONE && PageSample != 0)
@@ -640,10 +616,7 @@ bool FTextureData::GenerateMipMaps(const FTextureDataAddressingInfo& AddressingI
 
 bool FTextureData::GenerateBorderVoxels(const FTextureDataAddressingInfo& AddressingInfo, int32 MipLevel, const TArray<FIntVector3>& PageCoords)
 {
-	const FIntVector3 PageTableVolumeResolution = FIntVector3(
-		FMath::Max(1, Header.PageTableVolumeResolution.X >> MipLevel), 
-		FMath::Max(1, Header.PageTableVolumeResolution.Y >> MipLevel), 
-		FMath::Max(1, Header.PageTableVolumeResolution.Z >> MipLevel));
+	const FIntVector3 PageTableVolumeResolution = ShiftRightAndMax(Header.PageTableVolumeResolution, MipLevel, 1);
 
 	ParallelFor(PageCoords.Num(), [this, MipLevel, &PageCoords, &PageTableVolumeResolution, &AddressingInfo](int32 PageIndex)
 		{
@@ -657,7 +630,7 @@ bool FTextureData::GenerateBorderVoxels(const FTextureDataAddressingInfo& Addres
 				{
 					for (int32 X = -SPARSE_VOLUME_TILE_BORDER; X < (SPARSE_VOLUME_TILE_RES + SPARSE_VOLUME_TILE_BORDER); ++X)
 					{
-						if (X < 0 || Y < 0 || Z < 0 || X >= SPARSE_VOLUME_TILE_RES || Y >= SPARSE_VOLUME_TILE_RES || Z >= SPARSE_VOLUME_TILE_RES)
+						if (!IsInBounds(FIntVector3(X, Y, Z), FIntVector3::ZeroValue, FIntVector3(SPARSE_VOLUME_TILE_RES)))
 						{
 							const FIntVector3 LocalCoord = FIntVector3(X, Y, Z);
 							const FIntVector3 VolumeCoord = (PageTableOffset + PageCoord) * SPARSE_VOLUME_TILE_RES + LocalCoord;
@@ -750,12 +723,11 @@ bool FTextureData::DeduplicateTiles()
 			}
 		}
 
-		const uint64 DeduplicatedTileSizeBytes[] = { (uint64)TileWritePos * (uint64)TileByteSize[0], (uint64)TileWritePos * (uint64)TileByteSize[1] };
-		if (DeduplicatedTileSizeBytes[0] > (uint64)SVT::MaxResourceSize || DeduplicatedTileSizeBytes[1] > (uint64)SVT::MaxResourceSize)
+		const int64 DeduplicatedTileSizeBytes[] = { (int64)TileWritePos * (int64)TileByteSize[0], (int64)TileWritePos * (int64)TileByteSize[1] };
+		if (DeduplicatedTileSizeBytes[0] > SVT::MaxResourceSize || DeduplicatedTileSizeBytes[1] > SVT::MaxResourceSize)
 		{
-			UE_LOG(LogSparseVolumeTextureData, Error, TEXT("SparseVolumeTexture still requires too much memory after tile deduplication! Physical tile data A: %llu bytes. Physical tile data B: %llu bytes"),
-				(unsigned long long)DeduplicatedTileSizeBytes[0], (unsigned long long)DeduplicatedTileSizeBytes[1]);
-			return false;
+			UE_LOG(LogSparseVolumeTextureData, Warning, TEXT("SparseVolumeTexture still requires too much memory (> 2048MB) after tile deduplication! The full mip chain will not fit in memory, so only lower resolution mip levels can be streamed in. Physical tile data A: %ll bytes. Physical tile data B: %ll bytes"),
+				(long long)DeduplicatedTileSizeBytes[0], (long long)DeduplicatedTileSizeBytes[1]);
 		}
 
 		if (!MipMaps[MipLevel].PhysicalTileDataA.IsEmpty())
