Add another perf scope to network physics component

diff --git a/Engine/Source/Runtime/Engine/Private/PhysicsEngine/NetworkPhysicsComponent.cpp b/Engine/Source/Runtime/Engine/Private/PhysicsEngine/NetworkPhysicsComponent.cpp
--- a/Engine/Source/Runtime/Engine/Private/PhysicsEngine/NetworkPhysicsComponent.cpp
+++ b/Engine/Source/Runtime/Engine/Private/PhysicsEngine/NetworkPhysicsComponent.cpp
@@ -82,6 +82,8 @@ float DeQuantizeTimeDilation(int8 i)
 // after presimulate internal (asyncinput internal simulation done and the output created)
 void FNetworkPhysicsCallback::ApplyCallbacks_Internal(int32 PhysicsStep, const TArray<Chaos::ISimCallbackObject*>& SimCallbackObjects)
 {
+	QUICK_SCOPE_CYCLE_COUNTER(NetworkPhysicsComponent_ApplyCallbacks_Internal);
+
 	if(World->IsNetMode(NM_ListenServer) || World->IsNetMode(NM_DedicatedServer))
 	{
 		if (FPhysScene_Chaos* Scene = static_cast<FPhysScene_Chaos*>(World->GetPhysicsScene()))
