Merge plugin and non-plugin actors SupportPinning Logic
Remove reference to GetContentBundleGuid

diff --git a/Engine/Source/Runtime/Engine/Private/WorldPartition/LoaderAdapter/LoaderAdapterPinnedActors.cpp b/Engine/Source/Runtime/Engine/Private/WorldPartition/LoaderAdapter/LoaderAdapterPinnedActors.cpp
--- a/Engine/Source/Runtime/Engine/Private/WorldPartition/LoaderAdapter/LoaderAdapterPinnedActors.cpp
+++ b/Engine/Source/Runtime/Engine/Private/WorldPartition/LoaderAdapter/LoaderAdapterPinnedActors.cpp
@@ -37,18 +37,8 @@ bool FLoaderAdapterPinnedActors::SupportsPinning(FWorldPartitionActorDesc* InAct
 
 	if (UActorDescContainer* Container = InActorDesc->GetContainer())
 	{
-		if (Container->IsMainPartitionContainer())
-		{
-			return true;
-		}
-		else if (InActorDesc->GetContentBundleGuid().IsValid())
-		{
-			const UWorldPartition* ContainerWorldPartition = Container->GetWorldPartition();
-			if (ContainerWorldPartition && ContainerWorldPartition->IsMainWorldPartition())
-			{
-				return InActorDesc->GetActorSoftPath().GetAssetPath().GetPackageName() == ContainerWorldPartition->GetTypedOuter<UWorld>()->GetPackage()->GetFName();
-			}
-		}
+		const UWorldPartition* ContainerWorldPartition = Container->GetWorldPartition();
+		return ContainerWorldPartition && ContainerWorldPartition->IsMainWorldPartition();
 	}
 
 	return false;
