Remove reduntant checks
Remove references to GetContentBundleGuid

diff --git a/Engine/Source/Runtime/Engine/Private/WorldPartition/HLOD/HLODSourceActorsFromCell.cpp b/Engine/Source/Runtime/Engine/Private/WorldPartition/HLOD/HLODSourceActorsFromCell.cpp
--- a/Engine/Source/Runtime/Engine/Private/WorldPartition/HLOD/HLODSourceActorsFromCell.cpp
+++ b/Engine/Source/Runtime/Engine/Private/WorldPartition/HLOD/HLODSourceActorsFromCell.cpp
@@ -53,7 +53,6 @@ ULevelStreaming* UWorldPartitionHLODSourceActorsFromCell::LoadSourceActors(bool&
 			SubActor.ContainerTransform,
 			SubActor.ContainerPackage,
 			World->GetPackage()->GetFName(),
-			HLODActor->GetContentBundleGuid(),
 			SubActor.ActorGuid,
 			false
 		);
diff --git a/Engine/Source/Runtime/Engine/Private/WorldPartition/WorldPartitionLevelStreamingDynamic.cpp b/Engine/Source/Runtime/Engine/Private/WorldPartition/WorldPartitionLevelStreamingDynamic.cpp
--- a/Engine/Source/Runtime/Engine/Private/WorldPartition/WorldPartitionLevelStreamingDynamic.cpp
+++ b/Engine/Source/Runtime/Engine/Private/WorldPartition/WorldPartitionLevelStreamingDynamic.cpp
@@ -343,25 +343,10 @@ bool UWorldPartitionLevelStreamingDynamic::IssueLoadRequests()
 		{
 			if (ContentBundlePaths::IsAContentBundlePath(CellObjectMapping.ContainerPackage.ToString()))
 			{
-				check(CellObjectMapping.ContentBundleGuid.IsValid());
 				check(CellObjectMapping.ContainerPackage != CellObjectMapping.WorldPackage);
-
 				bool bIsContainerPackageAlreadyRemapped = OutLinkInstancingContext.RemapPackage(CellObjectMapping.ContainerPackage) != CellObjectMapping.ContainerPackage;
 				if (!bIsContainerPackageAlreadyRemapped)
 				{
-					// Get Content Bundle World
-					FString ActorDescContainerPackage = CellObjectMapping.ContainerPackage.ToString();
-					FStringView WorldFromContainerPackage = ContentBundlePaths::GetRelativePath(ActorDescContainerPackage);
-
-					// Get ChildPackage World
-					FStringView RelativeWorldPackage;
-					const bool bWithoutSlashes = false;
-					FString WorldPackage = CellObjectMapping.WorldPackage.ToString();
-					FName WorldPackageMountPoint = FPackageName::GetPackageMountPoint(WorldPackage, bWithoutSlashes);
-					ensure(FPathViews::TryMakeChildPathRelativeTo(MakeStringView(WorldPackage), MakeStringView(WorldPackageMountPoint.ToString()), RelativeWorldPackage));
-
-					check(WorldFromContainerPackage.Equals(RelativeWorldPackage));
-					
 					OutLinkInstancingContext.AddPackageMapping(CellObjectMapping.ContainerPackage, RuntimePackage->GetFName());
 					
 				}
diff --git a/Engine/Source/Runtime/Engine/Private/WorldPartition/WorldPartitionRuntimeLevelStreamingCell.cpp b/Engine/Source/Runtime/Engine/Private/WorldPartition/WorldPartitionRuntimeLevelStreamingCell.cpp
--- a/Engine/Source/Runtime/Engine/Private/WorldPartition/WorldPartitionRuntimeLevelStreamingCell.cpp
+++ b/Engine/Source/Runtime/Engine/Private/WorldPartition/WorldPartitionRuntimeLevelStreamingCell.cpp
@@ -239,7 +239,6 @@ void UWorldPartitionRuntimeLevelStreamingCell::AddActorToCell(const FWorldPartit
 				InContainerTransform, 
 				InContainer->GetContainerPackage(), 
 				GetWorld()->GetPackage()->GetFName(), 
-				ReferenceActorDesc.GetContentBundleGuid(),
 				InContainerID.GetActorGuid(ReferenceActorDesc.GetGuid()),
 				true
 			);
@@ -253,7 +252,6 @@ void UWorldPartitionRuntimeLevelStreamingCell::AddActorToCell(const FWorldPartit
 		InContainerTransform, 
 		InContainer->GetContainerPackage(), 
 		GetWorld()->GetPackage()->GetFName(), 
-		ActorDescView.GetContentBundleGuid(),
 		InContainerID.GetActorGuid(ActorDescView.GetGuid()),
 		false
 	);
diff --git a/Engine/Source/Runtime/Engine/Public/WorldPartition/WorldPartitionRuntimeCell.h b/Engine/Source/Runtime/Engine/Public/WorldPartition/WorldPartitionRuntimeCell.h
--- a/Engine/Source/Runtime/Engine/Public/WorldPartition/WorldPartitionRuntimeCell.h
+++ b/Engine/Source/Runtime/Engine/Public/WorldPartition/WorldPartitionRuntimeCell.h
@@ -12,6 +12,7 @@
 #include "ProfilingDebugging/ProfilingHelpers.h"
 #include "Misc/HierarchicalLogArchive.h"
 #include "Algo/AnyOf.h"
+#include "UObject/ObjectMacros.h"
 #include "WorldPartitionRuntimeCell.generated.h"
 
 class UActorContainer;
@@ -43,7 +44,7 @@ struct FWorldPartitionRuntimeCellObjectMapping
 #endif
 	{}
 
-	FWorldPartitionRuntimeCellObjectMapping(FName InPackage, FName InPath, const FActorContainerID& InContainerID, const FTransform& InContainerTransform, FName InContainerPackage, FName InWorldPackage, const FGuid& InContentBundleGuid, const FGuid& InActorInstanceGuid, bool bInIsEditorOnly)
+	FWorldPartitionRuntimeCellObjectMapping(FName InPackage, FName InPath, const FActorContainerID& InContainerID, const FTransform& InContainerTransform, FName InContainerPackage, FName InWorldPackage, const FGuid& InActorInstanceGuid, bool bInIsEditorOnly)
 #if WITH_EDITORONLY_DATA
 		: Package(InPackage)
 		, Path(InPath)
@@ -51,7 +52,6 @@ struct FWorldPartitionRuntimeCellObjectMapping
 		, ContainerTransform(InContainerTransform)
 		, ContainerPackage(InContainerPackage)
 		, WorldPackage(InWorldPackage)
-		, ContentBundleGuid(InContentBundleGuid)
 		, ActorInstanceGuid(InActorInstanceGuid)
 		, LoadedPath(InPath)
 		, bIsEditorOnly(bInIsEditorOnly)
@@ -95,12 +95,6 @@ struct FWorldPartitionRuntimeCellObjectMapping
 	UPROPERTY()
 	FName WorldPackage;
 
-	/**
-	 * Content Bundle Id
-	 */
-	UPROPERTY()
-	FGuid ContentBundleGuid;
-
 	/**
 	 * GUID of the actor instance
 	 */
