[Quartz] Fix for delegate array leak (w/ cvar to revert the fix).

diff --git a/Engine/Source/Runtime/Engine/Classes/Sound/QuartzSubscription.cpp b/Engine/Source/Runtime/Engine/Classes/Sound/QuartzSubscription.cpp
--- a/Engine/Source/Runtime/Engine/Classes/Sound/QuartzSubscription.cpp
+++ b/Engine/Source/Runtime/Engine/Classes/Sound/QuartzSubscription.cpp
@@ -3,6 +3,15 @@
 #include "Sound/QuartzSubscription.h"
 #include "Quartz/QuartzSubsystem.h"
 
+
+static int32 DecrementSlotIndexOnStartedCvar = 1;
+FAutoConsoleVariableRef CVarDecrementSlotIndexOnStarted(
+	TEXT("au.Quartz.DecrementSlotIndexOnStarted"),
+	DecrementSlotIndexOnStartedCvar,
+	TEXT("Defaults to 1 to enable the delegate leak fix.  Set to 0 to revert to pre-fix behavior.\n")
+	TEXT("1: New Behavior, 0: Old Behavior"),
+	ECVF_Default);
+
 namespace Audio
 {
 	FQuartzQueueCommandData::FQuartzQueueCommandData(const FAudioComponentCommandInfo& InAudioComponentCommandInfo, FName InClockName)
@@ -105,13 +114,14 @@ void FQuartzTickableObject::ExecCommand(const Audio::FQuartzQuantizedCommandDele
 		}
 
 		// (end of a command)
-		if (Data.DelegateSubType == EQuartzCommandDelegateSubType::CommandOnCanceled)
+		bool bShouldDecrement = Data.DelegateSubType == EQuartzCommandDelegateSubType::CommandOnCanceled;
+		bShouldDecrement |= (DecrementSlotIndexOnStartedCvar && Data.DelegateSubType == EQuartzCommandDelegateSubType::CommandOnStarted);
+			
+		// are all the commands for this delegate done?
+		if (bShouldDecrement && (GameThreadEntry.RefCount.Decrement() == 0))
 		{
-			// are all the commands done?
-			if (GameThreadEntry.RefCount.Decrement() == 0)
-			{
-				GameThreadEntry.MulticastDelegate.Clear();
-			}
+			// free up the slot for new subscriptions on this clock handle
+			GameThreadEntry.MulticastDelegate.Clear();
 		}
 	}
 
