WorldPartition - Take cell size into account when computing a cell guid
* Will prevent HLODs from being displayed if they were built with a different cell size

diff --git a/Engine/Source/Runtime/Core/Public/UObject/FortniteReleaseBranchCustomObjectVersion.h b/Engine/Source/Runtime/Core/Public/UObject/FortniteReleaseBranchCustomObjectVersion.h
--- a/Engine/Source/Runtime/Core/Public/UObject/FortniteReleaseBranchCustomObjectVersion.h
+++ b/Engine/Source/Runtime/Core/Public/UObject/FortniteReleaseBranchCustomObjectVersion.h
@@ -39,6 +39,9 @@ struct CORE_API FFortniteReleaseBranchCustomObjectVersion
 		// Fixup and synchronize shared properties modified before the synchronicity enforcement
 		LandscapeSharedPropertiesEnforcement,
 
+		// Include the cell size when computing the cell guid
+		WorldPartitionRuntimeCellGuidWithCellSize,
+
 		// -----<new versions can be added above this line>-------------------------------------------------
 		VersionPlusOne,
 		LatestVersion = VersionPlusOne - 1
diff --git a/Engine/Source/Runtime/Engine/Private/WorldPartition/HLOD/HLODActor.cpp b/Engine/Source/Runtime/Engine/Private/WorldPartition/HLOD/HLODActor.cpp
--- a/Engine/Source/Runtime/Engine/Private/WorldPartition/HLOD/HLODActor.cpp
+++ b/Engine/Source/Runtime/Engine/Private/WorldPartition/HLOD/HLODActor.cpp
@@ -8,6 +8,7 @@
 #include "UObject/UE5MainStreamObjectVersion.h"
 #include "UObject/UE5PrivateFrostyStreamObjectVersion.h"
 #include "UObject/FortniteMainBranchObjectVersion.h"
+#include "UObject/FortniteReleaseBranchCustomObjectVersion.h"
 
 #include UE_INLINE_GENERATED_CPP_BY_NAME(HLODActor)
 
@@ -107,6 +108,7 @@ void AWorldPartitionHLOD::Serialize(FArchive& Ar)
 	Ar.UsingCustomVersion(FUE5MainStreamObjectVersion::GUID);
 	Ar.UsingCustomVersion(FUE5PrivateFrostyStreamObjectVersion::GUID);
 	Ar.UsingCustomVersion(FFortniteMainBranchObjectVersion::GUID);
+	Ar.UsingCustomVersion(FFortniteReleaseBranchCustomObjectVersion::GUID);
 
 	Super::Serialize(Ar);
 
@@ -227,6 +229,20 @@ void AWorldPartitionHLOD::PostLoad()
 		check(SourceCellGuid.IsValid());
 	}
 
+	// CellGuid taking the cell size into account
+	if (GetLinkerCustomVersion(FFortniteReleaseBranchCustomObjectVersion::GUID) < FFortniteReleaseBranchCustomObjectVersion::WorldPartitionRuntimeCellGuidWithCellSize)
+	{
+		check(SourceCellGuid.IsValid());
+
+		int32 CellSize = (int32)FMath::RoundToInt(HLODBounds.GetSize().X);
+
+		FArchiveMD5 ArMD5;
+		ArMD5 << SourceCellGuid << CellSize;
+
+		SourceCellGuid = ArMD5.GetGuidFromHash();
+		check(SourceCellGuid.IsValid());
+	}
+
 	// Update the disk size stat on load, as we can't really know it when saving
 	HLODStats.Add(FWorldPartitionHLODStats::MemoryDiskSizeBytes, FHLODActorDesc::GetPackageSize(this));
 #endif
diff --git a/Engine/Source/Runtime/Engine/Private/WorldPartition/RuntimeSpatialHash/RuntimeSpatialHashHLOD.cpp b/Engine/Source/Runtime/Engine/Private/WorldPartition/RuntimeSpatialHash/RuntimeSpatialHashHLOD.cpp
--- a/Engine/Source/Runtime/Engine/Private/WorldPartition/RuntimeSpatialHash/RuntimeSpatialHashHLOD.cpp
+++ b/Engine/Source/Runtime/Engine/Private/WorldPartition/RuntimeSpatialHash/RuntimeSpatialHashHLOD.cpp
@@ -177,7 +177,7 @@ static TArray<FGuid> GenerateHLODsForGrid(UWorldPartition* WorldPartition, const
 			{
 				SlowTask.EnterProgressFrame(1);
 
-				const FGuid CellGuid = UWorldPartitionRuntimeSpatialHash::GetCellGuid(RuntimeGrid.GridName, CellGlobalCoord, GridCellDataChunk.GetDataLayersID(), GridCellDataChunk.GetContentBundleID());
+				const FGuid CellGuid = UWorldPartitionRuntimeSpatialHash::GetCellGuid(RuntimeGrid.GridName, RuntimeGrid.CellSize, CellGlobalCoord, GridCellDataChunk.GetDataLayersID(), GridCellDataChunk.GetContentBundleID());
 				const FString CellName = UWorldPartitionRuntimeSpatialHash::GetCellNameString(OuterWorld, RuntimeGrid.GridName, CellGlobalCoord, GridCellDataChunk.GetDataLayersID(), GridCellDataChunk.GetContentBundleID());
 
 				UE_LOG(LogWorldPartitionRuntimeSpatialHashHLOD, Verbose, TEXT("Creating HLOD for cell %s at %s..."), *CellName, *CellCoord.ToString());
diff --git a/Engine/Source/Runtime/Engine/Private/WorldPartition/WorldPartitionRuntimeSpatialHash.cpp b/Engine/Source/Runtime/Engine/Private/WorldPartition/WorldPartitionRuntimeSpatialHash.cpp
--- a/Engine/Source/Runtime/Engine/Private/WorldPartition/WorldPartitionRuntimeSpatialHash.cpp
+++ b/Engine/Source/Runtime/Engine/Private/WorldPartition/WorldPartitionRuntimeSpatialHash.cpp
@@ -1346,17 +1346,30 @@ FString UWorldPartitionRuntimeSpatialHash::GetCellNameString(UWorld* InOuterWorl
 	return CellName;
 }
 
-FGuid UWorldPartitionRuntimeSpatialHash::GetCellGuid(FName InGridName, const FGridCellCoord& InCellGlobalCoord, const FDataLayersID& InDataLayerID, const FGuid& InContentBundleID)
+FGuid UWorldPartitionRuntimeSpatialHash::GetCellGuid(FName InGridName, int32 InCellSize, const FGridCellCoord& InCellGlobalCoord, const FDataLayersID& InDataLayerID, const FGuid& InContentBundleID)
 {
 	FString GridName = InGridName.ToString().ToLower();
 	FGridCellCoord CellGlobalCoord = InCellGlobalCoord;
 	uint32 DataLayerID = InDataLayerID.GetHash();
 	uint32 ContentBundleID = InContentBundleID.IsValid() ? GetTypeHash(InContentBundleID) : 0;
 
-	FArchiveMD5 ArMD5;
-	ArMD5 << GridName << CellGlobalCoord << DataLayerID << ContentBundleID;
-	FGuid CellGuid = ArMD5.GetGuidFromHash();
-	check(CellGuid.IsValid());
+	FGuid CellGuid;
+
+	{
+		FArchiveMD5 ArMD5;
+		ArMD5 << GridName << CellGlobalCoord << DataLayerID << ContentBundleID;
+		CellGuid = ArMD5.GetGuidFromHash();
+		check(CellGuid.IsValid());
+	}
+
+	// FFortniteReleaseBranchCustomObjectVersion::WorldPartitionRuntimeCellGuidWithCellSize
+	// CellGuid taking the cell size into account
+	{
+		FArchiveMD5 ArMD5;
+		ArMD5 << CellGuid << InCellSize;
+		CellGuid = ArMD5.GetGuidFromHash();
+		check(CellGuid.IsValid());
+	}
 
 	return CellGuid;
 }
@@ -1439,7 +1452,7 @@ bool UWorldPartitionRuntimeSpatialHash::CreateStreamingGrid(const FSpatialHashRu
 					FString WorldInstanceSuffix;
 					verify(PartionedActors.GetCellGlobalCoords(FGridCellCoord(CellCoordX, CellCoordY, Level), CellGlobalCoords));
 					const FString CellName = GetCellNameString(OuterWorld, CurrentStreamingGrid.GridName, CellGlobalCoords, GridCellDataChunk.GetDataLayersID(), GridCellDataChunk.GetContentBundleID(), &WorldInstanceSuffix);
-					const FGuid CellGuid = GetCellGuid(CurrentStreamingGrid.GridName, CellGlobalCoords, GridCellDataChunk.GetDataLayersID(), GridCellDataChunk.GetContentBundleID());
+					const FGuid CellGuid = GetCellGuid(CurrentStreamingGrid.GridName, CurrentStreamingGrid.CellSize, CellGlobalCoords, GridCellDataChunk.GetDataLayersID(), GridCellDataChunk.GetContentBundleID());
 
 					UWorldPartitionRuntimeCell* StreamingCell = CreateRuntimeCell(StreamingPolicy->GetRuntimeCellClass(), UWorldPartitionRuntimeCellDataSpatialHash::StaticClass(), CellName, WorldInstanceSuffix);
 					UWorldPartitionRuntimeCellDataSpatialHash* CellDataSpatialHash = CastChecked<UWorldPartitionRuntimeCellDataSpatialHash>(StreamingCell->RuntimeCellData);
diff --git a/Engine/Source/Runtime/Engine/Public/WorldPartition/WorldPartitionRuntimeSpatialHash.h b/Engine/Source/Runtime/Engine/Public/WorldPartition/WorldPartitionRuntimeSpatialHash.h
--- a/Engine/Source/Runtime/Engine/Public/WorldPartition/WorldPartitionRuntimeSpatialHash.h
+++ b/Engine/Source/Runtime/Engine/Public/WorldPartition/WorldPartitionRuntimeSpatialHash.h
@@ -283,7 +283,7 @@ public:
 	ENGINE_API virtual URuntimeHashExternalStreamingObjectBase* StoreToExternalStreamingObject(UObject* StreamingObjectOuter, FName StreamingObjectName) override;
 
 	static ENGINE_API FString GetCellNameString(UWorld* InOuterWorld, FName InGridName, const FGridCellCoord& InCellGlobalCoord, const FDataLayersID& InDataLayerID, const FGuid& InContentBundleID, FString* OutInstanceSuffix = nullptr);
-	static ENGINE_API FGuid GetCellGuid(FName InGridName, const FGridCellCoord& InCellGlobalCoord, const FDataLayersID& InDataLayerID, const FGuid& InContentBundleID);
+	static ENGINE_API FGuid GetCellGuid(FName InGridName, int32 InCellSize, const FGridCellCoord& InCellGlobalCoord, const FDataLayersID& InDataLayerID, const FGuid& InContentBundleID);
 
 	ENGINE_API bool GetPreviewGrids() const;
 	ENGINE_API void SetPreviewGrids(bool bInPreviewGrids);
