World Partition - HLOD: Display a warning on screen if outdated HLODs are detected
* Will check that loaded HLOD actors are mapping to a valid WP cell.
* For now, doesn't validate the HLOD content vs the cell content (may be too expensive to compute at runtime)

diff --git a/Engine/Source/Runtime/Engine/Private/UnrealEngine.cpp b/Engine/Source/Runtime/Engine/Private/UnrealEngine.cpp
--- a/Engine/Source/Runtime/Engine/Private/UnrealEngine.cpp
+++ b/Engine/Source/Runtime/Engine/Private/UnrealEngine.cpp
@@ -236,6 +236,10 @@ UnrealEngine.cpp: Implements the UEngine class and helpers.
 #include "RenderGraphBuilder.h"
 #endif
 
+#if !(UE_BUILD_SHIPPING || UE_BUILD_TEST)
+#include "WorldPartition/HLOD/HLODSubsystem.h"
+#endif
+
 #if WITH_DUMPGPU
 
 static TAutoConsoleVariable<int32> GDumpGPUScreenshot(
@@ -12036,6 +12040,18 @@ float DrawMapWarnings(UWorld* World, FViewport* Viewport, FCanvas* Canvas, UCanv
 			MessageY += FontSizeY;
 		}
 	}
+
+	if (UHLODSubsystem* HLODSubsystem = World->GetSubsystem<UHLODSubsystem>())
+	{
+		uint32 NumOutdatedHLODActors = HLODSubsystem->GetNumOutdatedHLODActors();
+		if (NumOutdatedHLODActors)
+		{
+			SmallTextItem.SetColor(FLinearColor::Yellow);
+			SmallTextItem.Text = FText::Format(LOCTEXT("OUTDATED_HLOD_ACTORS_FMT", "Outdated HLOD actors detected, HLODs need to be rebuilt ({0} outdated {0}|plural(one=actor,other=actors))"), NumOutdatedHLODActors);
+			Canvas->DrawItem(SmallTextItem, FVector2D(MessageX, MessageY));
+			MessageY += FontSizeY;
+		}
+	}
 #endif
 
 #if WITH_EDITOR && !(UE_BUILD_SHIPPING || UE_BUILD_TEST)
diff --git a/Engine/Source/Runtime/Engine/Private/WorldPartition/HLOD/HLODSubsystem.cpp b/Engine/Source/Runtime/Engine/Private/WorldPartition/HLOD/HLODSubsystem.cpp
--- a/Engine/Source/Runtime/Engine/Private/WorldPartition/HLOD/HLODSubsystem.cpp
+++ b/Engine/Source/Runtime/Engine/Private/WorldPartition/HLOD/HLODSubsystem.cpp
@@ -456,6 +456,10 @@ void UHLODSubsystem::RegisterHLODActor(AWorldPartitionHLOD* InWorldPartitionHLOD
 	{
 		UE_LOG(LogHLODSubsystem, Verbose, TEXT("Found HLOD %s referencing nonexistent cell '%s'"), *InWorldPartitionHLOD->GetActorNameOrLabel(), *InWorldPartitionHLOD->GetSourceCellGuid().ToString());
 		InWorldPartitionHLOD->SetVisibility(false);
+
+#if WITH_EDITOR
+		OutdatedHLODActors.Add(InWorldPartitionHLOD);
+#endif
 	}
 
 	HLODActorsToWarmup.Remove(InWorldPartitionHLOD);
@@ -470,8 +474,13 @@ void UHLODSubsystem::UnregisterHLODActor(AWorldPartitionHLOD* InWorldPartitionHL
 	{
 		UE_LOG(LogHLODSubsystem, Verbose, TEXT("Unregistering HLOD %s for cell %s"), *InWorldPartitionHLOD->GetActorNameOrLabel(), *InWorldPartitionHLOD->GetSourceCellGuid().ToString());
 
-		int32 NumRemoved = CellData->LoadedHLODs.Remove(InWorldPartitionHLOD);
-		check(NumRemoved == 1);
+		verify(CellData->LoadedHLODs.Remove(InWorldPartitionHLOD));
+	}
+	else
+	{
+#if WITH_EDITOR
+		verify(OutdatedHLODActors.Remove(InWorldPartitionHLOD));
+#endif
 	}
 
 	HLODActorUnregisteredEvent.Broadcast(InWorldPartitionHLOD);
diff --git a/Engine/Source/Runtime/Engine/Public/WorldPartition/HLOD/HLODSubsystem.h b/Engine/Source/Runtime/Engine/Public/WorldPartition/HLOD/HLODSubsystem.h
--- a/Engine/Source/Runtime/Engine/Public/WorldPartition/HLOD/HLODSubsystem.h
+++ b/Engine/Source/Runtime/Engine/Public/WorldPartition/HLOD/HLODSubsystem.h
@@ -69,6 +69,7 @@ public:
 	ENGINE_API void OnCVarsChanged();
 
 #ifdef WITH_EDITOR
+	uint32 GetNumOutdatedHLODActors() const { return OutdatedHLODActors.Num(); }
 	static ENGINE_API bool WriteHLODStatsCSV(UWorld* InWorld, const FString& InFilename);
 #endif
 	
@@ -130,5 +131,9 @@ private:
 
 	FWorldPartitionHLODActorRegisteredEvent		HLODActorRegisteredEvent;
 	FWorldPartitionHLODActorUnregisteredEvent	HLODActorUnregisteredEvent;
+
+#if WITH_EDITOR
+	TSet<AWorldPartitionHLOD*> OutdatedHLODActors;
+#endif
 };
 
