Fix bad #ifdef WITH_EDITOR
#buildfix
#rnx

diff --git a/Engine/Source/Runtime/Engine/Public/WorldPartition/HLOD/HLODSubsystem.h b/Engine/Source/Runtime/Engine/Public/WorldPartition/HLOD/HLODSubsystem.h
--- a/Engine/Source/Runtime/Engine/Public/WorldPartition/HLOD/HLODSubsystem.h
+++ b/Engine/Source/Runtime/Engine/Public/WorldPartition/HLOD/HLODSubsystem.h
@@ -68,7 +68,7 @@ public:
 
 	ENGINE_API void OnCVarsChanged();
 
-#ifdef WITH_EDITOR
+#if WITH_EDITOR
 	uint32 GetNumOutdatedHLODActors() const { return OutdatedHLODActors.Num(); }
 	static ENGINE_API bool WriteHLODStatsCSV(UWorld* InWorld, const FString& InFilename);
 #endif
