Buildfix
#buildfix
#rnx

diff --git a/Engine/Source/Runtime/Engine/Private/UnrealEngine.cpp b/Engine/Source/Runtime/Engine/Private/UnrealEngine.cpp
--- a/Engine/Source/Runtime/Engine/Private/UnrealEngine.cpp
+++ b/Engine/Source/Runtime/Engine/Private/UnrealEngine.cpp
@@ -12041,6 +12041,7 @@ float DrawMapWarnings(UWorld* World, FViewport* Viewport, FCanvas* Canvas, UCanv
 		}
 	}
 
+#if WITH_EDITOR
 	if (UHLODSubsystem* HLODSubsystem = World->GetSubsystem<UHLODSubsystem>())
 	{
 		uint32 NumOutdatedHLODActors = HLODSubsystem->GetNumOutdatedHLODActors();
@@ -12054,6 +12055,8 @@ float DrawMapWarnings(UWorld* World, FViewport* Viewport, FCanvas* Canvas, UCanv
 	}
 #endif
 
+#endif
+
 #if WITH_EDITOR && !(UE_BUILD_SHIPPING || UE_BUILD_TEST)
 	if (ULandscapeSubsystem* LandscapeSubsystem = World->GetSubsystem<ULandscapeSubsystem>())
 	{
