CIS Fix - Some AWorldPartitionHLOD actors can be created without GUIDs
#rnx

diff --git a/Engine/Source/Runtime/Engine/Private/WorldPartition/HLOD/HLODActor.cpp b/Engine/Source/Runtime/Engine/Private/WorldPartition/HLOD/HLODActor.cpp
--- a/Engine/Source/Runtime/Engine/Private/WorldPartition/HLOD/HLODActor.cpp
+++ b/Engine/Source/Runtime/Engine/Private/WorldPartition/HLOD/HLODActor.cpp
@@ -232,15 +232,16 @@ void AWorldPartitionHLOD::PostLoad()
 	// CellGuid taking the cell size into account
 	if (GetLinkerCustomVersion(FFortniteReleaseBranchCustomObjectVersion::GUID) < FFortniteReleaseBranchCustomObjectVersion::WorldPartitionRuntimeCellGuidWithCellSize)
 	{
-		check(SourceCellGuid.IsValid());
-
-		int32 CellSize = (int32)FMath::RoundToInt(HLODBounds.GetSize().X);
+		if (SourceCellGuid.IsValid())
+		{
+			int32 CellSize = (int32)FMath::RoundToInt(HLODBounds.GetSize().X);
 
-		FArchiveMD5 ArMD5;
-		ArMD5 << SourceCellGuid << CellSize;
+			FArchiveMD5 ArMD5;
+			ArMD5 << SourceCellGuid << CellSize;
 
-		SourceCellGuid = ArMD5.GetGuidFromHash();
-		check(SourceCellGuid.IsValid());
+			SourceCellGuid = ArMD5.GetGuidFromHash();
+			check(SourceCellGuid.IsValid());
+		}
 	}
 
 	// Update the disk size stat on load, as we can't really know it when saving
