Exposed blocking SparseVolumeTexture streaming requests to blueprint and the SVTViewer component. SVTViewer now uses an enum instead of an integer to select the channel/attribute/component of the SVT to preview. This makes for a nicer UI.

#rnx

diff --git a/Engine/Source/Runtime/Engine/Classes/SparseVolumeTexture/SparseVolumeTexture.h b/Engine/Source/Runtime/Engine/Classes/SparseVolumeTexture/SparseVolumeTexture.h
--- a/Engine/Source/Runtime/Engine/Classes/SparseVolumeTexture/SparseVolumeTexture.h
+++ b/Engine/Source/Runtime/Engine/Classes/SparseVolumeTexture/SparseVolumeTexture.h
@@ -256,7 +256,7 @@ public:
 	// MipLevel is the lowest mip level that the caller intends to use but does not guarantee that the mip is actually resident.
 	// If bBlocking is true, DDC streaming requests will block on completion, guaranteeing that the requested frame will have been streamed in after the next streaming system update.
 	// If streaming cooked data from disk, the highest priority will be used, but no guarantee is given.
-	static ENGINE_API USparseVolumeTextureFrame* GetFrameAndIssueStreamingRequest(USparseVolumeTexture* SparseVolumeTexture, float FrameIndex, int32 MipLevel, bool bBlocking = false);
+	static ENGINE_API USparseVolumeTextureFrame* GetFrameAndIssueStreamingRequest(USparseVolumeTexture* SparseVolumeTexture, float FrameIndex, int32 MipLevel, bool bBlocking);
 
 	ENGINE_API bool Initialize(USparseVolumeTexture* InOwner, int32 InFrameIndex, UE::SVT::FTextureData& UncookedFrame);
 	int32 GetFrameIndex() const { return FrameIndex; }
@@ -497,6 +497,9 @@ public:
 	UPROPERTY(BlueprintReadWrite, Category = "Rendering")
 	int32 MipLevel = 0;
 
+	UPROPERTY(BlueprintReadWrite, Category = "Rendering")
+	bool bBlockingStreamingRequests = false;
+
 	ENGINE_API UAnimatedSparseVolumeTextureController();
 	virtual ~UAnimatedSparseVolumeTextureController() = default;
 
diff --git a/Engine/Source/Runtime/Engine/Private/SparseVolumeTexture/SparseVolumeTexture.cpp b/Engine/Source/Runtime/Engine/Private/SparseVolumeTexture/SparseVolumeTexture.cpp
--- a/Engine/Source/Runtime/Engine/Private/SparseVolumeTexture/SparseVolumeTexture.cpp
+++ b/Engine/Source/Runtime/Engine/Private/SparseVolumeTexture/SparseVolumeTexture.cpp
@@ -1294,7 +1294,7 @@ USparseVolumeTextureFrame* UAnimatedSparseVolumeTextureController::GetFrameByInd
 		return nullptr;
 	}
 
-	return USparseVolumeTextureFrame::GetFrameAndIssueStreamingRequest(SparseVolumeTexture, FrameIndex, MipLevel);
+	return USparseVolumeTextureFrame::GetFrameAndIssueStreamingRequest(SparseVolumeTexture, FrameIndex, MipLevel, bBlockingStreamingRequests);
 }
 
 USparseVolumeTextureFrame* UAnimatedSparseVolumeTextureController::GetCurrentFrame()
@@ -1307,7 +1307,7 @@ USparseVolumeTextureFrame* UAnimatedSparseVolumeTextureController::GetCurrentFra
 	// Compute (fractional) index of frame to sample
 	const float FrameIndexF = GetFractionalFrameIndex();
 
-	return USparseVolumeTextureFrame::GetFrameAndIssueStreamingRequest(SparseVolumeTexture, FrameIndexF, MipLevel);
+	return USparseVolumeTextureFrame::GetFrameAndIssueStreamingRequest(SparseVolumeTexture, FrameIndexF, MipLevel, bBlockingStreamingRequests);
 }
 
 void UAnimatedSparseVolumeTextureController::GetCurrentFramesForInterpolation(USparseVolumeTextureFrame*& Frame0, USparseVolumeTextureFrame*& Frame1, float& LerpAlpha)
@@ -1322,8 +1322,8 @@ void UAnimatedSparseVolumeTextureController::GetCurrentFramesForInterpolation(US
 	const int32 FrameIndex = (int32)FrameIndexF;
 	LerpAlpha = FMath::Frac(FrameIndexF);
 
-	Frame0 = USparseVolumeTextureFrame::GetFrameAndIssueStreamingRequest(SparseVolumeTexture, FrameIndexF, MipLevel);
-	Frame1 = USparseVolumeTextureFrame::GetFrameAndIssueStreamingRequest(SparseVolumeTexture, (FrameIndex + 1) % SparseVolumeTexture->GetNumFrames(), MipLevel);
+	Frame0 = USparseVolumeTextureFrame::GetFrameAndIssueStreamingRequest(SparseVolumeTexture, FrameIndexF, MipLevel, bBlockingStreamingRequests);
+	Frame1 = USparseVolumeTextureFrame::GetFrameAndIssueStreamingRequest(SparseVolumeTexture, (FrameIndex + 1) % SparseVolumeTexture->GetNumFrames(), MipLevel, bBlockingStreamingRequests);
 }
 
 float UAnimatedSparseVolumeTextureController::GetDuration()
diff --git a/Engine/Source/Runtime/Renderer/Private/SparseVolumeTexture/SparseVolumeTextureViewerComponent.cpp b/Engine/Source/Runtime/Renderer/Private/SparseVolumeTexture/SparseVolumeTextureViewerComponent.cpp
--- a/Engine/Source/Runtime/Renderer/Private/SparseVolumeTexture/SparseVolumeTextureViewerComponent.cpp
+++ b/Engine/Source/Runtime/Renderer/Private/SparseVolumeTexture/SparseVolumeTextureViewerComponent.cpp
@@ -31,7 +31,7 @@ USparseVolumeTextureViewerComponent::USparseVolumeTextureViewerComponent(const F
 	, SparseVolumeTexturePreview(nullptr)
 	, bAnimate(false)
 	, AnimationFrame(0.0f)
-	, ComponentToVisualize(0)
+	, PreviewAttribute(ESVTPA_AttributesA_R)
 	, SparseVolumeTextureViewerSceneProxy(nullptr)
 {
 	PrimaryComponentTick.bCanEverTick = true;
@@ -175,7 +175,7 @@ void USparseVolumeTextureViewerComponent::SendRenderTransformCommand()
 
 		FSparseVolumeTextureViewerSceneProxy* SVTViewerSceneProxy = SparseVolumeTextureViewerSceneProxy;
 		ENQUEUE_RENDER_COMMAND(FUpdateSparseVolumeTextureViewerProxyTransformCommand)(
-			[SVTViewerSceneProxy, ToLocalMat, ToLocalMatNoScale, CompIdx = (uint32)ComponentToVisualize, Ext = Extinction, Mip = MipLevel, VolumeRes3f]
+			[SVTViewerSceneProxy, ToLocalMat, ToLocalMatNoScale, CompIdx = (uint32)PreviewAttribute, Ext = Extinction, Mip = MipLevel, VolumeRes3f]
 		(FRHICommandList& RHICmdList)
 			{
 				SVTViewerSceneProxy->WorldToLocal = ToLocalMat;
@@ -206,7 +206,7 @@ void USparseVolumeTextureViewerComponent::TickComponent(float DeltaTime, enum EL
 		}
 		FrameIndexF = FMath::Clamp(FrameIndexF, 0, static_cast<float>(NumFrames - 1));
 		FrameIndex = FrameIndexF;
-		SparseVolumeTextureFrame = USparseVolumeTextureFrame::GetFrameAndIssueStreamingRequest(SparseVolumeTexturePreview, FrameIndexF, MipLevel);
+		SparseVolumeTextureFrame = USparseVolumeTextureFrame::GetFrameAndIssueStreamingRequest(SparseVolumeTexturePreview, FrameIndexF, MipLevel, bBlockingStreamingRequests);
 	}
 	else
 	{
diff --git a/Engine/Source/Runtime/Renderer/Private/SparseVolumeTexture/SparseVolumeTextureViewerComponent.h b/Engine/Source/Runtime/Renderer/Private/SparseVolumeTexture/SparseVolumeTextureViewerComponent.h
--- a/Engine/Source/Runtime/Renderer/Private/SparseVolumeTexture/SparseVolumeTextureViewerComponent.h
+++ b/Engine/Source/Runtime/Renderer/Private/SparseVolumeTexture/SparseVolumeTextureViewerComponent.h
@@ -15,6 +15,18 @@
 
 class FSparseVolumeTextureViewerSceneProxy;
 
+UENUM(BlueprintType)
+enum ESparseVolumeTexturePreviewAttribute : uint8
+{
+	ESVTPA_AttributesA_R UMETA(DisplayName = "AttributesA Red"),
+	ESVTPA_AttributesA_G UMETA(DisplayName = "AttributesA Green"),
+	ESVTPA_AttributesA_B UMETA(DisplayName = "AttributesA Blue"),
+	ESVTPA_AttributesA_A UMETA(DisplayName = "AttributesA Alpha"),
+	ESVTPA_AttributesB_R UMETA(DisplayName = "AttributesB Red"),
+	ESVTPA_AttributesB_G UMETA(DisplayName = "AttributesB Green"),
+	ESVTPA_AttributesB_B UMETA(DisplayName = "AttributesB Blue"),
+	ESVTPA_AttributesB_A UMETA(DisplayName = "AttributesB Alpha"),
+};
 
 /**
  * A component used to inspect sparse volume textures.
@@ -35,6 +47,9 @@ class USparseVolumeTextureViewerComponent : public UPrimitiveComponent
 	UPROPERTY(EditAnywhere, BlueprintReadOnly, Category = "Asset Preview")
 	uint32 bReversePlayback : 1;
 
+	UPROPERTY(EditAnywhere, BlueprintReadOnly, Category = "Asset Preview")
+	uint32 bBlockingStreamingRequests : 1;
+
 	UPROPERTY(EditAnywhere, BlueprintReadWrite, Category = "Asset Preview", meta = (UIMin = 0.0, UIMax = 1.0, ClampMin = 0.0, ClampMax = 1.0, EditCondition = "!bAnimate"))
 	float AnimationFrame;
 
@@ -44,8 +59,8 @@ class USparseVolumeTextureViewerComponent : public UPrimitiveComponent
 	UPROPERTY(VisibleAnywhere, Category = "Asset Preview", meta = (UIMin = 0.0, UIMax = 60.0))
 	float AnimationTime;
 
-	UPROPERTY(EditAnywhere, BlueprintReadWrite, Category = "Asset Preview", meta = (UIMin = 0, UIMax = 7, ClampMin = 0, ClampMax = 7))
-	int32 ComponentToVisualize;
+	UPROPERTY(EditAnywhere, BlueprintReadWrite, Category = "Asset Preview")
+	TEnumAsByte<ESparseVolumeTexturePreviewAttribute> PreviewAttribute;
 
 	UPROPERTY(EditAnywhere, BlueprintReadWrite, Category = "Asset Preview", meta = (UIMin = 0, UIMax = 11, ClampMin = 0, ClampMax = 11))
 	int32 MipLevel;
