[CIS, Quartz] Non-unity fix for new cvar.
#pf none

diff --git a/Engine/Source/Runtime/Engine/Classes/Sound/QuartzSubscription.cpp b/Engine/Source/Runtime/Engine/Classes/Sound/QuartzSubscription.cpp
--- a/Engine/Source/Runtime/Engine/Classes/Sound/QuartzSubscription.cpp
+++ b/Engine/Source/Runtime/Engine/Classes/Sound/QuartzSubscription.cpp
@@ -2,6 +2,7 @@
 
 #include "Sound/QuartzSubscription.h"
 #include "Quartz/QuartzSubsystem.h"
+#include "HAL/IConsoleManager.h"
 
 
 static int32 DecrementSlotIndexOnStartedCvar = 1;
