Remove extraneous debug draw.

diff --git a/Engine/Source/Runtime/Engine/Private/PhysicsEngine/PhysicsReplication.cpp b/Engine/Source/Runtime/Engine/Private/PhysicsEngine/PhysicsReplication.cpp
--- a/Engine/Source/Runtime/Engine/Private/PhysicsEngine/PhysicsReplication.cpp
+++ b/Engine/Source/Runtime/Engine/Private/PhysicsEngine/PhysicsReplication.cpp
@@ -1183,8 +1183,6 @@ bool FPhysicsReplicationAsync::PredictiveInterpolation(Chaos::FPBDRigidParticleH
 	bool bClearTarget = bHardSnap;
 	if (bHardSnap)
 	{
-		Chaos::FDebugDrawQueue::GetInstance().DrawDebugPoint(Target.PrevPosTarget, FColor::Red, false, -1, -1, 4.f);
-
 		// Too much error so just snap state here and be done with it
 		Target.AccumulatedErrorSeconds = 0.0f;
 		Handle->SetX(Target.PrevPosTarget);
