Fix Trace.Screenshot not working when not running PIE or simulation.

diff --git a/Engine/Source/Runtime/Engine/Private/ProfilingDebugging/TraceScreenshot.cpp b/Engine/Source/Runtime/Engine/Private/ProfilingDebugging/TraceScreenshot.cpp
--- a/Engine/Source/Runtime/Engine/Private/ProfilingDebugging/TraceScreenshot.cpp
+++ b/Engine/Source/Runtime/Engine/Private/ProfilingDebugging/TraceScreenshot.cpp
@@ -9,6 +9,10 @@
 #include "Tasks/Task.h"
 #include "UnrealClient.h"
 
+#if WITH_EDITOR
+#include "Editor.h"
+#endif
+
 static void TraceScreenshotCommandCallback(const TArray<FString>& Args)
 {
 	FString Name;
@@ -80,6 +84,14 @@ void FTraceScreenshot::RequestScreenshot(FString Name, bool bShowUI, const FLogC
 	}
 	constexpr bool bAddUniqueSuffix = false;
 	FScreenshotRequest::RequestScreenshot(Name, bShowUI, bAddUniqueSuffix);
+
+#if WITH_EDITOR
+	// If not running PIE or Simulation we have to force a redraw.
+	if (GEditor && !GEditor->IsPlayingSessionInEditor())
+	{
+		GEditor->RedrawAllViewports(/*bInvalidateHitProxies =*/false);
+	}
+#endif
 }
 
 void ImageUtilsImageResize(int32 SrcWidth, int32 SrcHeight, const TArray<FColor>& SrcData, int32 DstWidth, int32 DstHeight, TArray<FColor>& DstData, bool bLinearSpace)
