Improved inertialization missing node warning to use tokenized messages.
Motion matching now warns on missing inertialization node when requesting an inertial blend.
#rnx

diff --git a/Engine/Plugins/Experimental/Animation/PoseSearch/Source/Runtime/Private/PoseSearchLibrary.cpp b/Engine/Plugins/Experimental/Animation/PoseSearch/Source/Runtime/Private/PoseSearchLibrary.cpp
--- a/Engine/Plugins/Experimental/Animation/PoseSearch/Source/Runtime/Private/PoseSearchLibrary.cpp
+++ b/Engine/Plugins/Experimental/Animation/PoseSearch/Source/Runtime/Private/PoseSearchLibrary.cpp
@@ -74,6 +74,10 @@ static void RequestInertialBlend(const FAnimationUpdateContext& Context, float B
 
 			InertializationRequester->RequestInertialization(Request);
 		}
+		else
+		{
+			FAnimNode_Inertialization::LogRequestError(Context, Context.GetCurrentNodeId());
+		}
 	}
 }
 
diff --git a/Engine/Source/Runtime/Engine/Classes/Animation/AnimNode_Inertialization.h b/Engine/Source/Runtime/Engine/Classes/Animation/AnimNode_Inertialization.h
--- a/Engine/Source/Runtime/Engine/Classes/Animation/AnimNode_Inertialization.h
+++ b/Engine/Source/Runtime/Engine/Classes/Animation/AnimNode_Inertialization.h
@@ -376,6 +376,7 @@ public: // FAnimNode_Inertialization
 
 	// Log an error when a node wants to inertialize but no inertialization ancestor node exists
 	//
+	static ENGINE_API void LogRequestError(const FAnimationUpdateContext& Context, const int32 NodePropertyIndex);
 	static ENGINE_API void LogRequestError(const FAnimationUpdateContext& Context, const FPoseLinkBase& RequesterPoseLink);
 
 public: // FAnimNode_Base
diff --git a/Engine/Source/Runtime/Engine/Private/Animation/AnimNode_Inertialization.cpp b/Engine/Source/Runtime/Engine/Private/Animation/AnimNode_Inertialization.cpp
--- a/Engine/Source/Runtime/Engine/Private/Animation/AnimNode_Inertialization.cpp
+++ b/Engine/Source/Runtime/Engine/Private/Animation/AnimNode_Inertialization.cpp
@@ -9,6 +9,7 @@
 #include "Animation/AnimBlueprintGeneratedClass.h"
 #include "Logging/TokenizedMessage.h"
 #include "Animation/AnimCurveUtils.h"
+#include "Misc/UObjectToken.h"
 
 #include UE_INLINE_GENERATED_CPP_BY_NAME(AnimNode_Inertialization)
 
@@ -117,16 +118,24 @@ void FAnimNode_Inertialization::RequestInertialization(const FInertializationReq
 	}
 }
 
-/*static*/ void FAnimNode_Inertialization::LogRequestError(const FAnimationUpdateContext& Context, const FPoseLinkBase& RequesterPoseLink)
+/*static*/ void FAnimNode_Inertialization::LogRequestError(const FAnimationUpdateContext& Context, const int32 NodePropertyIndex)
 {
 #if WITH_EDITORONLY_DATA	
 	UAnimBlueprint* AnimBlueprint = Context.AnimInstanceProxy->GetAnimBlueprint();
 	UAnimBlueprintGeneratedClass* AnimClass = AnimBlueprint ? AnimBlueprint->GetAnimBlueprintGeneratedClass() : nullptr;
-	const UObject* RequesterNode = AnimClass ? AnimClass->GetVisualNodeFromNodePropertyIndex(RequesterPoseLink.SourceLinkID) : nullptr;
+	const UObject* RequesterNode = AnimClass ? AnimClass->GetVisualNodeFromNodePropertyIndex(NodePropertyIndex) : nullptr;
+
+	Context.LogMessage(FTokenizedMessage::Create(EMessageSeverity::Error)
+		->AddToken(FTextToken::Create(LOCTEXT("InertializationRequestError_1", "No Inertialization node found for request from ")))
+		->AddToken(FUObjectToken::Create(RequesterNode))
+		->AddToken(FTextToken::Create(LOCTEXT("InertializationRequestError_2", ". Add an Inertialization node after this request."))));
+#endif
+}
 
-	FText Message = FText::Format(LOCTEXT("InertializationRequestError", "No Inertialization node found for request from '{0}'. Add an Inertialization node after this request."),
-		FText::FromString(GetPathNameSafe(RequesterNode)));
-	Context.LogMessage(EMessageSeverity::Error, Message);
+/*static*/ void FAnimNode_Inertialization::LogRequestError(const FAnimationUpdateContext& Context, const FPoseLinkBase& RequesterPoseLink)
+{
+#if WITH_EDITORONLY_DATA	
+	LogRequestError(Context, RequesterPoseLink.SourceLinkID);
 #endif
 }
 
