[Quartz] Adding ms offset to NotifyOnQUantizationBoundary()

diff --git a/Engine/Source/Runtime/AudioMixer/Private/AudioMixerSourceManager.cpp b/Engine/Source/Runtime/AudioMixer/Private/AudioMixerSourceManager.cpp
--- a/Engine/Source/Runtime/AudioMixer/Private/AudioMixerSourceManager.cpp
+++ b/Engine/Source/Runtime/AudioMixer/Private/AudioMixerSourceManager.cpp
@@ -954,7 +954,7 @@ namespace Audio
 				}
 				else
 				{
-					FQuartzQuantizedCommandInitInfo QuantCommandInitInfo(QuantData, SourceId);
+					FQuartzQuantizedCommandInitInfo QuantCommandInitInfo(QuantData, MixerDevice->GetSampleRate(), SourceId);
 					SourceInfo.QuantizedCommandHandle = MixerDevice->QuantizedEventClockManager.AddCommandToClock(QuantCommandInitInfo);
 				}
 			}
diff --git a/Engine/Source/Runtime/AudioMixer/Private/Quartz/AudioMixerClock.cpp b/Engine/Source/Runtime/AudioMixer/Private/Quartz/AudioMixerClock.cpp
--- a/Engine/Source/Runtime/AudioMixer/Private/Quartz/AudioMixerClock.cpp
+++ b/Engine/Source/Runtime/AudioMixer/Private/Quartz/AudioMixerClock.cpp
@@ -443,6 +443,7 @@ namespace Audio
 		Metronome.UnsubscribeFromAllTimeDivisions(InSubscriber);
 	}
 
+
 	void FQuartzClock::AddQuantizedCommand(FQuartzQuantizationBoundary InQuantizationBondary, TSharedPtr<IQuartzQuantizedCommand> InNewEvent)
 	{
 		if (!ensure(InNewEvent.IsValid()))
@@ -610,6 +611,18 @@ namespace Audio
 		return nullptr;
 	}
 
+	void FQuartzClock::AddQuantizedCommand(FQuartzQuantizedRequestData& InQuantizedRequestData)
+	{
+		float SampleRate = HeadlessClockSampleRateCvar;
+		if (FMixerDevice* MixerDevice = GetMixerDevice())
+		{
+			SampleRate = MixerDevice->GetSampleRate();
+		}
+
+		FQuartzQuantizedCommandInitInfo Info(InQuantizedRequestData, SampleRate);
+		AddQuantizedCommand(Info);
+	}
+
 	void FQuartzClock::AddQuantizedCommand(FQuartzQuantizedCommandInitInfo& InQuantizationCommandInitInfo)
 	{
 		if (!ensure(InQuantizationCommandInitInfo.QuantizedCommandPtr))
diff --git a/Engine/Source/Runtime/AudioMixer/Private/Quartz/AudioMixerClockHandle.cpp b/Engine/Source/Runtime/AudioMixer/Private/Quartz/AudioMixerClockHandle.cpp
--- a/Engine/Source/Runtime/AudioMixer/Private/Quartz/AudioMixerClockHandle.cpp
+++ b/Engine/Source/Runtime/AudioMixer/Private/Quartz/AudioMixerClockHandle.cpp
@@ -62,7 +62,7 @@ void UQuartzClockHandle::QueueQuantizedSound(const UObject* WorldContextObject,
 	QueueCommandPtr->SetQueueCommand(InAudioComponentData);
 
 	//Set up initial command info
-	Audio::FQuartzQuantizedCommandInitInfo CommandInitInfo = UQuartzSubsystem::CreateRequestDataForSchedulePlaySound(InClockHandle, InDelegate, InTargetBoundary);
+	Audio::FQuartzQuantizedRequestData CommandInitInfo = UQuartzSubsystem::CreateRequestDataForSchedulePlaySound(InClockHandle, InDelegate, InTargetBoundary);
 
 	//(Queue's setup is identical to PlaySound except for the command ptr, so fix that here)
 	CommandInitInfo.QuantizedCommandPtr.Reset();
@@ -74,14 +74,14 @@ void UQuartzClockHandle::QueueQuantizedSound(const UObject* WorldContextObject,
 // deprecated: use ResetTransportQuantized
 void UQuartzClockHandle::ResetTransport(const UObject* WorldContextObject, const FOnQuartzCommandEventBP& InDelegate)
 {
-	Audio::FQuartzQuantizedCommandInitInfo Data(UQuartzSubsystem::CreateRequestDataForTransportReset(this, FQuartzQuantizationBoundary(EQuartzCommandQuantization::Bar), InDelegate));
+	Audio::FQuartzQuantizedRequestData Data(UQuartzSubsystem::CreateRequestDataForTransportReset(this, FQuartzQuantizationBoundary(EQuartzCommandQuantization::Bar), InDelegate));
 	RawHandle.SendCommandToClock([Data](Audio::FQuartzClock* InClock) mutable { InClock->AddQuantizedCommand(Data); });
 }
 
 void UQuartzClockHandle::ResetTransportQuantized(const UObject* WorldContextObject, FQuartzQuantizationBoundary InQuantizationBoundary, const FOnQuartzCommandEventBP& InDelegate, UQuartzClockHandle*& ClockHandle)
 {
 	ClockHandle = this;
-	Audio::FQuartzQuantizedCommandInitInfo Data(UQuartzSubsystem::CreateRequestDataForTransportReset(this, InQuantizationBoundary, InDelegate));
+	Audio::FQuartzQuantizedRequestData Data(UQuartzSubsystem::CreateRequestDataForTransportReset(this, InQuantizationBoundary, InDelegate));
 	RawHandle.SendCommandToClock([Data](Audio::FQuartzClock* InClock) mutable { InClock->AddQuantizedCommand(Data); });
 }
 
@@ -92,9 +92,9 @@ bool UQuartzClockHandle::IsClockRunning(const UObject* WorldContextObject)
 	return RawHandle.IsClockRunning();
 }
 
-void UQuartzClockHandle::NotifyOnQuantizationBoundary(const UObject* WorldContextObject, FQuartzQuantizationBoundary InQuantizationBoundary, const FOnQuartzCommandEventBP& InDelegate)
+void UQuartzClockHandle::NotifyOnQuantizationBoundary(const UObject* WorldContextObject, FQuartzQuantizationBoundary InQuantizationBoundary, const FOnQuartzCommandEventBP& InDelegate, float OffsetInMilliseconds)
 {
-	Audio::FQuartzQuantizedCommandInitInfo Data(UQuartzSubsystem::CreateRequestDataForQuantizedNotify(this, InQuantizationBoundary, InDelegate));
+	Audio::FQuartzQuantizedRequestData Data(UQuartzSubsystem::CreateRequestDataForQuantizedNotify(this, InQuantizationBoundary, InDelegate, OffsetInMilliseconds));
 	RawHandle.SendCommandToClock([Data](Audio::FQuartzClock* InClock) mutable { InClock->AddQuantizedCommand(Data); });
 }
 
@@ -121,7 +121,7 @@ void UQuartzClockHandle::StartOtherClock(const UObject* WorldContextObject, FNam
 		return;
 	}
 
-	Audio::FQuartzQuantizedCommandInitInfo Data(UQuartzSubsystem::CreateRequestDataForStartOtherClock(this, OtherClockName, InQuantizationBoundary, InDelegate));
+	Audio::FQuartzQuantizedRequestData Data(UQuartzSubsystem::CreateRequestDataForStartOtherClock(this, OtherClockName, InQuantizationBoundary, InDelegate));
 	RawHandle.SendCommandToClock([Data](Audio::FQuartzClock* InClock) mutable { InClock->AddQuantizedCommand(Data); });
 }
 
@@ -244,7 +244,7 @@ void UQuartzClockHandle::SetBeatsPerMinute(const UObject* WorldContextObject, co
 
 void UQuartzClockHandle::SetTickRateInternal(const FQuartzQuantizationBoundary& InQuantizationBoundary, const FOnQuartzCommandEventBP& InDelegate, const Audio::FQuartzClockTickRate& NewTickRate)
 {
-	Audio::FQuartzQuantizedCommandInitInfo Data(UQuartzSubsystem::CreateRequestDataForTickRateChange(this, InDelegate, NewTickRate, InQuantizationBoundary));
+	Audio::FQuartzQuantizedRequestData Data(UQuartzSubsystem::CreateRequestDataForTickRateChange(this, InDelegate, NewTickRate, InQuantizationBoundary));
 	RawHandle.SendCommandToClock([Data](Audio::FQuartzClock* InClock) mutable { InClock->AddQuantizedCommand(Data); });
 }
 
diff --git a/Engine/Source/Runtime/AudioMixer/Private/Quartz/AudioMixerQuantizedCommands.cpp b/Engine/Source/Runtime/AudioMixer/Private/Quartz/AudioMixerQuantizedCommands.cpp
--- a/Engine/Source/Runtime/AudioMixer/Private/Quartz/AudioMixerQuantizedCommands.cpp
+++ b/Engine/Source/Runtime/AudioMixer/Private/Quartz/AudioMixerQuantizedCommands.cpp
@@ -248,6 +248,20 @@ namespace Audio
 	}
 
 
+	FQuantizedNotify::FQuantizedNotify(float InMsOffset)  : OffsetInMs(InMsOffset)
+	{
+	}
+
+	int32 FQuantizedNotify::OverrideFramesUntilExec(int32 NumFramesUntilExec)
+	{
+		constexpr float MsToSec = 1000.f;
+		return FMath::Max(0, NumFramesUntilExec + (OffsetInMs / MsToSec) * SampleRate);
+	}
+
+	void FQuantizedNotify::OnQueuedCustom(const FQuartzQuantizedCommandInitInfo& InCommandInitInfo)
+	{
+		SampleRate = InCommandInitInfo.SampleRate;
+	}
 
 	TSharedPtr<IQuartzQuantizedCommand> FQuantizedNotify::GetDeepCopyOfDerivedObject() const
 	{
diff --git a/Engine/Source/Runtime/AudioMixer/Private/Quartz/QuartzSubsystem.cpp b/Engine/Source/Runtime/AudioMixer/Private/Quartz/QuartzSubsystem.cpp
--- a/Engine/Source/Runtime/AudioMixer/Private/Quartz/QuartzSubsystem.cpp
+++ b/Engine/Source/Runtime/AudioMixer/Private/Quartz/QuartzSubsystem.cpp
@@ -389,14 +389,14 @@ Audio::FQuartzQuantizedRequestData UQuartzSubsystem::CreateRequestDataForStartOt
 	return CommandInitInfo;
 }
 
-Audio::FQuartzQuantizedRequestData UQuartzSubsystem::CreateRequestDataForQuantizedNotify(UQuartzClockHandle* InClockHandle, const FQuartzQuantizationBoundary& InQuantizationBoundary, const FOnQuartzCommandEventBP& InDelegate)
+Audio::FQuartzQuantizedRequestData UQuartzSubsystem::CreateRequestDataForQuantizedNotify(UQuartzClockHandle* InClockHandle, const FQuartzQuantizationBoundary& InQuantizationBoundary, const FOnQuartzCommandEventBP& InDelegate, float InMsOffset)
 {
 	if (!ensure(InClockHandle))
 	{
 		return { };
 	}
 
-	const TSharedPtr<Audio::FQuantizedNotify> NotifyCommandPtr = MakeShared<Audio::FQuantizedNotify>();
+	const TSharedPtr<Audio::FQuantizedNotify> NotifyCommandPtr = MakeShared<Audio::FQuantizedNotify>(InMsOffset);
 
 	Audio::FQuartzQuantizedRequestData CommandInitInfo;
 
diff --git a/Engine/Source/Runtime/AudioMixer/Public/Quartz/AudioMixerClock.h b/Engine/Source/Runtime/AudioMixer/Public/Quartz/AudioMixerClock.h
--- a/Engine/Source/Runtime/AudioMixer/Public/Quartz/AudioMixerClock.h
+++ b/Engine/Source/Runtime/AudioMixer/Public/Quartz/AudioMixerClock.h
@@ -181,6 +181,7 @@ namespace Audio
 		void UnsubscribeFromAllTimeDivisions(FQuartzGameThreadSubscriber InSubscriber);
 
 		// Quantized Command Management:
+		void AddQuantizedCommand(FQuartzQuantizedRequestData& InQuantizedRequestData);
 		void AddQuantizedCommand(FQuartzQuantizedCommandInitInfo& InQuantizationCommandInitInfo);
 
 		void AddQuantizedCommand(FQuartzQuantizationBoundary InQuantizationBondary, TSharedPtr<IQuartzQuantizedCommand> InNewEvent);
diff --git a/Engine/Source/Runtime/AudioMixer/Public/Quartz/AudioMixerClockHandle.h b/Engine/Source/Runtime/AudioMixer/Public/Quartz/AudioMixerClockHandle.h
--- a/Engine/Source/Runtime/AudioMixer/Public/Quartz/AudioMixerClockHandle.h
+++ b/Engine/Source/Runtime/AudioMixer/Public/Quartz/AudioMixerClockHandle.h
@@ -64,7 +64,7 @@ public:
 	bool IsClockRunning(const UObject* WorldContextObject);
 
 	UFUNCTION(BlueprintCallable, Category = "Quartz Clock", meta = (WorldContext = "WorldContextObject", AutoCreateRefTerm = "InDelegate", Keywords = "Transport, Counter"))
-	void NotifyOnQuantizationBoundary(const UObject* WorldContextObject, FQuartzQuantizationBoundary InQuantizationBoundary, const FOnQuartzCommandEventBP& InDelegate);
+	void NotifyOnQuantizationBoundary(const UObject* WorldContextObject, FQuartzQuantizationBoundary InQuantizationBoundary, const FOnQuartzCommandEventBP& InDelegate, float InMsOffset = 0.f);
 
 	/** Returns the duration in seconds of the given Quantization Type
 	 *
diff --git a/Engine/Source/Runtime/AudioMixer/Public/Quartz/AudioMixerQuantizedCommands.h b/Engine/Source/Runtime/AudioMixer/Public/Quartz/AudioMixerQuantizedCommands.h
--- a/Engine/Source/Runtime/AudioMixer/Public/Quartz/AudioMixerQuantizedCommands.h
+++ b/Engine/Source/Runtime/AudioMixer/Public/Quartz/AudioMixerQuantizedCommands.h
@@ -142,7 +142,7 @@ namespace Audio
 	{
 	public:
 		// ctor
-		FQuantizedNotify() = default;
+		FQuantizedNotify(float InMsOffset = 0.f);
 
 		// dtor
 		virtual ~FQuantizedNotify() override = default;
@@ -155,11 +155,14 @@ namespace Audio
 
 		virtual EQuartzCommandType GetCommandType() const override { return EQuartzCommandType::Notify; };
 
-	protected:
-		TSharedPtr<FQuartzClock> OwningClockPtr{ nullptr };
+		virtual void OnQueuedCustom(const FQuartzQuantizedCommandInitInfo& InCommandInitInfo) override;
 
-		int32 SourceID{ -1 };
+		virtual int32 OverrideFramesUntilExec(int32 NumFramesUntilExec) override;
 
+	protected:
+		TSharedPtr<FQuartzClock> OwningClockPtr{ nullptr };
+		float OffsetInMs = 0.f;
+		float SampleRate = 0.f;
 		bool bIsCanceled = false;
 
 	}; // class FQuantizedNotify
diff --git a/Engine/Source/Runtime/AudioMixer/Public/Quartz/QuartzSubsystem.h b/Engine/Source/Runtime/AudioMixer/Public/Quartz/QuartzSubsystem.h
--- a/Engine/Source/Runtime/AudioMixer/Public/Quartz/QuartzSubsystem.h
+++ b/Engine/Source/Runtime/AudioMixer/Public/Quartz/QuartzSubsystem.h
@@ -86,7 +86,7 @@ public:
 	static Audio::FQuartzQuantizedRequestData CreateRequestDataForTransportReset(UQuartzClockHandle* InClockHandle, const FQuartzQuantizationBoundary& InQuantizationBoundary, const FOnQuartzCommandEventBP& InDelegate);
 	static Audio::FQuartzQuantizedRequestData CreateRequestDataForStartOtherClock(UQuartzClockHandle* InClockHandle, FName InClockToStart, const FQuartzQuantizationBoundary& InQuantizationBoundary, const FOnQuartzCommandEventBP& InDelegate);
 	static Audio::FQuartzQuantizedRequestData CreateRequestDataForSchedulePlaySound(UQuartzClockHandle* InClockHandle, const FOnQuartzCommandEventBP& InDelegate, const FQuartzQuantizationBoundary& InQuantizationBoundary);
-	static Audio::FQuartzQuantizedRequestData CreateRequestDataForQuantizedNotify(UQuartzClockHandle* InClockHandle, const FQuartzQuantizationBoundary& InQuantizationBoundary, const FOnQuartzCommandEventBP& InDelegate);
+	static Audio::FQuartzQuantizedRequestData CreateRequestDataForQuantizedNotify(UQuartzClockHandle* InClockHandle, const FQuartzQuantizationBoundary& InQuantizationBoundary, const FOnQuartzCommandEventBP& InDelegate, float InMsOffset = 0.f);
 
 	// DEPRECATED HELPERS: non-static versions of the above CreateDataFor...() functions
 	UE_DEPRECATED(5.1, "Use the static (CreateRequestDataFor) version of this function instead")
diff --git a/Engine/Source/Runtime/Engine/Classes/Sound/QuartzQuantizationUtilities.cpp b/Engine/Source/Runtime/Engine/Classes/Sound/QuartzQuantizationUtilities.cpp
--- a/Engine/Source/Runtime/Engine/Classes/Sound/QuartzQuantizationUtilities.cpp
+++ b/Engine/Source/Runtime/Engine/Classes/Sound/QuartzQuantizationUtilities.cpp
@@ -456,6 +456,7 @@ namespace Audio
 
 	FQuartzQuantizedCommandInitInfo::FQuartzQuantizedCommandInitInfo(
 		const FQuartzQuantizedRequestData& RHS
+		, float InSampleRate
 		, int32 InSourceID
 	)
 		: ClockName(RHS.ClockName)
@@ -465,6 +466,7 @@ namespace Audio
 		, GameThreadSubscribers(RHS.GameThreadSubscribers)
 		, GameThreadDelegateID(RHS.GameThreadDelegateID)
 		, OwningClockPointer(nullptr)
+		, SampleRate(InSampleRate)
 		, SourceID(InSourceID)
 	{
 	}
diff --git a/Engine/Source/Runtime/Engine/Classes/Sound/QuartzQuantizationUtilities.h b/Engine/Source/Runtime/Engine/Classes/Sound/QuartzQuantizationUtilities.h
--- a/Engine/Source/Runtime/Engine/Classes/Sound/QuartzQuantizationUtilities.h
+++ b/Engine/Source/Runtime/Engine/Classes/Sound/QuartzQuantizationUtilities.h
@@ -536,7 +536,7 @@ namespace Audio
 {
 	// data that is gathered by the AudioThread to get passed from FActiveSound->FMixerSourceVoice
 	// eventually converted to IQuartzQuantizedCommand for the Quantized Command itself
-	struct FQuartzQuantizedRequestData
+	struct ENGINE_API FQuartzQuantizedRequestData
 	{
 		// shared with FQuartzQuantizedCommandInitInfo:
 		FName ClockName;
@@ -550,12 +550,12 @@ namespace Audio
 
 	// data that is passed into IQuartzQuantizedCommand::OnQueued
 	// info that derived classes need can be added here
-	struct FQuartzQuantizedCommandInitInfo
+	struct ENGINE_API FQuartzQuantizedCommandInitInfo
 	{
 		FQuartzQuantizedCommandInitInfo() {}
 
 		// conversion ctor from FQuartzQuantizedRequestData
-		ENGINE_API FQuartzQuantizedCommandInitInfo(const FQuartzQuantizedRequestData& RHS, int32 InSourceID = INDEX_NONE);
+		FQuartzQuantizedCommandInitInfo(const FQuartzQuantizedRequestData& RHS, float InSampleRate, int32 InSourceID = INDEX_NONE);
 
 		void SetOwningClockPtr(TSharedPtr<Audio::FQuartzClock> InClockPointer)
 		{
@@ -573,6 +573,7 @@ namespace Audio
 
 		// Audio Render thread-specific data:
 		TSharedPtr<Audio::FQuartzClock> OwningClockPointer{ nullptr };
+		float SampleRate;
 		int32 SourceID{ -1 };
 
 		// Number of frames used for any FramesTilExec overrides
@@ -580,7 +581,7 @@ namespace Audio
 	};
 
 	// base class for quantized commands. Virtual methods called by owning clock.
-	class IQuartzQuantizedCommand : public FQuartzCrossThreadMessage
+	class ENGINE_API IQuartzQuantizedCommand : public FQuartzCrossThreadMessage
 	{
 	public:
 
@@ -591,38 +592,38 @@ namespace Audio
 		virtual ~IQuartzQuantizedCommand() {};
 
 		// allocate a copy of the derived class
-		ENGINE_API virtual TSharedPtr<IQuartzQuantizedCommand> GetDeepCopyOfDerivedObject() const;
+		virtual TSharedPtr<IQuartzQuantizedCommand> GetDeepCopyOfDerivedObject() const;
 
-		ENGINE_API void AddSubscriber(FQuartzGameThreadSubscriber InSubscriber);
+		void AddSubscriber(FQuartzGameThreadSubscriber InSubscriber);
 
 		// Command has reached the AudioRenderThread
-		ENGINE_API void OnQueued(const FQuartzQuantizedCommandInitInfo& InCommandInitInfo);
+		void OnQueued(const FQuartzQuantizedCommandInitInfo& InCommandInitInfo);
 
 		// scheduled (finalize subscriber offsets) - called by FQuartzClock
-		ENGINE_API void OnScheduled(const FQuartzClockTickRate& InTickRate);
+		void OnScheduled(const FQuartzClockTickRate& InTickRate);
 
 		// called during FQuartzClock::Tick() to let us call AboutToStart
 		// at different times for different subscribers
-		ENGINE_API void Update(int32 NumFramesUntilDeadline);
+		void Update(int32 NumFramesUntilDeadline);
 
 		// Perhaps the associated sound failed concurrency and will not be playing
-		ENGINE_API void FailedToQueue(FQuartzQuantizedRequestData& InGameThreadData);
+		void FailedToQueue(FQuartzQuantizedRequestData& InGameThreadData);
 
 		// Called 2x Assumed thread latency before OnFinalCallback()
-		ENGINE_API void AboutToStart();
+		void AboutToStart();
 
 		// Called on the final callback of this event boundary.
 		// InNumFramesLeft is the number of frames into the callback the exact quantized event should take place
-		ENGINE_API void OnFinalCallback(int32 InNumFramesLeft);
+		void OnFinalCallback(int32 InNumFramesLeft);
 
 		// Called if the owning clock gets stopped
-		ENGINE_API void OnClockPaused();
+		void OnClockPaused();
 
 		// Called if the owning clock gets started
-		ENGINE_API void OnClockStarted();
+		void OnClockStarted();
 
 		// Called if the event is cancelled before OnFinalCallback() is called
-		ENGINE_API void Cancel();
+		void Cancel();
 
 
 		//Called if the event type uses an altered amount of frames
@@ -654,21 +655,21 @@ namespace Audio
 
 	// Audio Render Thread Handle to a queued command
 	// Used by AudioMixerSourceVoices to access a pending associated command
-	struct FQuartzQuantizedCommandHandle
+	struct ENGINE_API FQuartzQuantizedCommandHandle
 	{
 		FName OwningClockName;
 		TSharedPtr<IQuartzQuantizedCommand> CommandPtr{ nullptr };
 		FMixerDevice* MixerDevice{ nullptr };
 
 		// Attempts to cancel the command. Returns true if the cancellation was successful.
-		ENGINE_API bool Cancel();
+		bool Cancel();
 
 		// Resets the handle to initial state.
-		ENGINE_API void Reset();
+		void Reset();
 	};
 } // namespace Audio
 
-struct FAudioComponentCommandInfo
+struct ENGINE_API FAudioComponentCommandInfo
 {
 	FAudioComponentCommandInfo() {}
 
