Default value for r.ProxyLODMeshReductionModule is set to a name that does not exist.

diff --git a/Engine/Source/Runtime/Engine/Private/ProxyLODMeshSimplificationSettings.cpp b/Engine/Source/Runtime/Engine/Private/ProxyLODMeshSimplificationSettings.cpp
--- a/Engine/Source/Runtime/Engine/Private/ProxyLODMeshSimplificationSettings.cpp
+++ b/Engine/Source/Runtime/Engine/Private/ProxyLODMeshSimplificationSettings.cpp
@@ -8,7 +8,7 @@
 
 static FAutoConsoleVariable CVarProxyLODMeshReductionModule(
 	TEXT("r.ProxyLODMeshReductionModule"),
-	TEXT("QuadricMeshProxyLODReduction"),
+	TEXT("ProxyLODMeshReduction"),
 	TEXT("Name of the Proxy LOD reduction module to choose. If blank it chooses any that exist.\n"),
 	ECVF_ReadOnly);
 
