World Partition Class Descriptors: properly handle the case where a class descriptor references a BP class through a redirector, but the class descriptor registry hasn't registered the redirector yet (depending on the load order), resulting in a missing class descriptor.

diff --git a/Engine/Source/Runtime/Engine/Private/WorldPartition/WorldPartitionActorDescArchive.cpp b/Engine/Source/Runtime/Engine/Private/WorldPartition/WorldPartitionActorDescArchive.cpp
--- a/Engine/Source/Runtime/Engine/Private/WorldPartition/WorldPartitionActorDescArchive.cpp
+++ b/Engine/Source/Runtime/Engine/Private/WorldPartition/WorldPartitionActorDescArchive.cpp
@@ -88,12 +88,19 @@ PRAGMA_ENABLE_DEPRECATION_WARNINGS
 	{
 		if (IsLoading())
 		{
-			bIsMissingClassDesc = true;
+			// We need to prefetch the class path here as the class registry might not have registered its redirector from the asset registry class data.
+			ClassDescRegistry.PrefetchClassDescs({ ClassPath });
+			ClassDesc = InActorDesc->bIsDefaultActorDesc ? ClassDescRegistry.GetClassDescDefaultForClass(ClassPath) : ClassDescRegistry.GetClassDescDefaultForActor(ClassPath);
 
-			ClassDesc = ClassDescRegistry.GetClassDescDefault(FTopLevelAssetPath(TEXT("/Script/Engine.Actor")));
-			check(ClassDesc);			
+			if (!ClassDesc)
+			{
+				bIsMissingClassDesc = true;
+	
+				ClassDesc = ClassDescRegistry.GetClassDescDefault(FTopLevelAssetPath(TEXT("/Script/Engine.Actor")));
+				check(ClassDesc);
 
-			UE_LOG(LogWorldPartition, Log, TEXT("Can't find class descriptor '%s' for loading '%s', using '%s'"), *ClassPath.ToString(), *InActorDesc->GetActorSoftPath().ToString(), *ClassDesc->GetActorSoftPath().ToString());
+				UE_LOG(LogWorldPartition, Log, TEXT("Can't find class descriptor '%s' for loading '%s', using '%s'"), *ClassPath.ToString(), *InActorDesc->GetActorSoftPath().ToString(), *ClassDesc->GetActorSoftPath().ToString());
+			}
 		}
 		else
 		{
