Optimized CompressTile lambda function which is invoked during SparseVolumeTexture derived data generation. The bitmask based voxel compression is now split up into multiple steps, which allows some of these steps to be parallelized.

#rnx

diff --git a/Engine/Source/Runtime/Engine/Private/SparseVolumeTexture/SparseVolumeTexture.cpp b/Engine/Source/Runtime/Engine/Private/SparseVolumeTexture/SparseVolumeTexture.cpp
--- a/Engine/Source/Runtime/Engine/Private/SparseVolumeTexture/SparseVolumeTexture.cpp
+++ b/Engine/Source/Runtime/Engine/Private/SparseVolumeTexture/SparseVolumeTexture.cpp
@@ -26,6 +26,7 @@
 #include "DerivedDataRequestOwner.h"
 #endif
 #include "EditorFramework/AssetImportData.h"
+#include "Async/ParallelFor.h"
 
 #include "Serialization/LargeMemoryReader.h"
 #include "Serialization/LargeMemoryWriter.h"
@@ -240,6 +241,8 @@ bool FResources::RebuildBulkDataFromCacheAsync(const UObject* Owner, bool& bFail
 
 bool FResources::Build(USparseVolumeTextureFrame* Owner, UE::Serialization::FEditorBulkData& SourceData)
 {
+	TRACE_CPUPROFILER_EVENT_SCOPE(SVT::FResources::Build);
+
 	// Check if the virtualized bulk data payload is available
 	if (SourceData.HasPayloadData())
 	{
@@ -279,6 +282,8 @@ bool FResources::Build(USparseVolumeTextureFrame* Owner, UE::Serialization::FEdi
 		// Returns number of written/non-zero page table entries
 		auto CompressPageTable = [](const TArray<uint32>& PageTable, const FIntVector3& Resolution, TArray<uint8>& BulkData) -> int32
 		{
+			TRACE_CPUPROFILER_EVENT_SCOPE(SVT::FResources::Build::CompressPageTable);
+
 			int32 NumNonZeroEntries = 0;
 			for (uint32 Entry : PageTable)
 			{
@@ -318,10 +323,11 @@ bool FResources::Build(USparseVolumeTextureFrame* Owner, UE::Serialization::FEdi
 
 		auto CompressTiles = [](int32 NumTiles, const TArray64<uint8>& PhysicalTileDataA, const TArray64<uint8>& PhysicalTileDataB, const TArrayView<EPixelFormat>& Formats, const TArrayView<FVector4f>& FallbackValues, TArray<uint8>& BulkData, FMipLevelStreamingInfo& MipStreamingInfo)
 		{
+			TRACE_CPUPROFILER_EVENT_SCOPE(SVT::FResources::Build::CompressTiles);
+
 			const int64 FormatSize[] = { GPixelFormats[Formats[0]].BlockBytes, GPixelFormats[Formats[1]].BlockBytes };
 			uint8 NullTileValuesU8[2][sizeof(float) * 4] = {};
 			int32 NumTextures = 0;
-			int32 EstimatedTileDataSize = NumTiles * SVT::NumVoxelsPerPaddedTile * (FormatSize[0] + FormatSize[1]);
 			for (int32 i = 0; i < 2; ++i)
 			{
 				if (Formats[i] != PF_Unknown)
@@ -334,7 +340,7 @@ bool FResources::Build(USparseVolumeTextureFrame* Owner, UE::Serialization::FEdi
 			const int32 OccupancySizePerTexture = NumTiles * SVT::NumOccupancyWordsPerPaddedTile * sizeof(uint32);
 			const int32 OccupancySize = NumTextures * OccupancySizePerTexture;
 			const int32 TileDataOffsetsSize = NumTextures * NumTiles * sizeof(uint32);
-			BulkData.Reserve(BulkData.Num() + OccupancySize + TileDataOffsetsSize + EstimatedTileDataSize);
+			BulkData.Reserve(BulkData.Num() + OccupancySize + TileDataOffsetsSize);
 
 			MipStreamingInfo.OccupancyBitsOffset[0] = BulkData.Num() - MipStreamingInfo.BulkOffset;
 			MipStreamingInfo.OccupancyBitsSize[0] = Formats[0] != PF_Unknown ? OccupancySizePerTexture : 0;
@@ -352,38 +358,94 @@ bool FResources::Build(USparseVolumeTextureFrame* Owner, UE::Serialization::FEdi
 			MipStreamingInfo.TileDataOffsetsSize[1] = Formats[1] != PF_Unknown ? NumTiles * sizeof(uint32) : 0;
 			BulkData.SetNum(BulkData.Num() + MipStreamingInfo.TileDataOffsetsSize[1]);
 
+			// Zero the bitmasks
 			FMemory::Memzero(BulkData.GetData() + MipStreamingInfo.BulkOffset + MipStreamingInfo.OccupancyBitsOffset[0], MipStreamingInfo.OccupancyBitsSize[0] + MipStreamingInfo.OccupancyBitsSize[1]);
 
+			uint32* OccupancyBits[2];
+			OccupancyBits[0] = reinterpret_cast<uint32*>(BulkData.GetData() + MipStreamingInfo.BulkOffset + MipStreamingInfo.OccupancyBitsOffset[0]);
+			OccupancyBits[1] = reinterpret_cast<uint32*>(BulkData.GetData() + MipStreamingInfo.BulkOffset + MipStreamingInfo.OccupancyBitsOffset[1]);
+
+			uint32* PrefixSums[2];
+			PrefixSums[0] = reinterpret_cast<uint32*>(BulkData.GetData() + MipStreamingInfo.BulkOffset + MipStreamingInfo.TileDataOffsetsOffset[0]);
+			PrefixSums[1] = reinterpret_cast<uint32*>(BulkData.GetData() + MipStreamingInfo.BulkOffset + MipStreamingInfo.TileDataOffsetsOffset[1]);
+
+			// Compute occupancy bitmasks and count number of non-fallback voxels per tile
+			ParallelFor(NumTiles, [&](int32 TileIndex)
+				{
+					for (int32 AttributesIdx = 0; AttributesIdx < 2; ++AttributesIdx)
+					{
+						PrefixSums[AttributesIdx][TileIndex] = 0;
+						if (Formats[AttributesIdx] != PF_Unknown)
+						{
+							for (int64 VoxelIndex = 0; VoxelIndex < SVT::NumVoxelsPerPaddedTile; ++VoxelIndex)
+							{
+								const uint8* Src = (AttributesIdx == 0 ? PhysicalTileDataA.GetData() : PhysicalTileDataB.GetData()) + FormatSize[AttributesIdx] * (TileIndex * SVT::NumVoxelsPerPaddedTile + VoxelIndex);
+								bool bIsFallbackValue = FMemory::Memcmp(Src, NullTileValuesU8[AttributesIdx], FormatSize[AttributesIdx]) == 0;
+								if (!bIsFallbackValue)
+								{
+									const int64 WordIndex = TileIndex * SVT::NumOccupancyWordsPerPaddedTile + (VoxelIndex / 32);
+									OccupancyBits[AttributesIdx][WordIndex] |= 1u << (static_cast<uint32>(VoxelIndex) % 32u);
+									PrefixSums[AttributesIdx][TileIndex]++;
+								}
+							}
+						}
+					}
+				});
+
+			// Compute actual per-tile voxel data offsets (prefix sum)
+			uint32 Sums[2] = {};
 			for (int32 AttributesIdx = 0; AttributesIdx < 2; ++AttributesIdx)
 			{
-				MipStreamingInfo.TileDataOffset[AttributesIdx] = BulkData.Num() - MipStreamingInfo.BulkOffset;
-
 				if (Formats[AttributesIdx] != PF_Unknown)
 				{
-					uint32 NumValidVoxelsTotal = 0;
-					for (int64 TileIndex = 0; TileIndex < NumTiles; ++TileIndex)
+					for (int32 TileIndex = 0; TileIndex < NumTiles; ++TileIndex)
 					{
-						uint32* TileDataOffsets = reinterpret_cast<uint32*>(BulkData.GetData() + MipStreamingInfo.BulkOffset + MipStreamingInfo.TileDataOffsetsOffset[AttributesIdx]);
-						TileDataOffsets[TileIndex] = NumValidVoxelsTotal;
+						uint32 Tmp = PrefixSums[AttributesIdx][TileIndex];
+						PrefixSums[AttributesIdx][TileIndex] = Sums[AttributesIdx];
+						Sums[AttributesIdx] += Tmp;
+					}
+				}
+			}
+
+			// Allocate compacted tile data memory
+			MipStreamingInfo.TileDataOffset[0] = BulkData.Num() - MipStreamingInfo.BulkOffset;
+			MipStreamingInfo.TileDataSize[0] = Sums[0] * FormatSize[0];
+
+			MipStreamingInfo.TileDataOffset[1] = MipStreamingInfo.TileDataOffset[0] + MipStreamingInfo.TileDataSize[0];
+			MipStreamingInfo.TileDataSize[1] = Sums[1] * FormatSize[1];
+
+			BulkData.SetNum(BulkData.Num() + MipStreamingInfo.TileDataSize[0] + MipStreamingInfo.TileDataSize[1]);
+
+			// All above pointers into the BulkData are now invalid!
+
+			uint8* RelativeBulkDataPtr = BulkData.GetData() + MipStreamingInfo.BulkOffset;
+
+			uint8* DstTileData[2];
+			DstTileData[0] = RelativeBulkDataPtr + MipStreamingInfo.TileDataOffset[0];
+			DstTileData[1] = RelativeBulkDataPtr + MipStreamingInfo.TileDataOffset[1];
+
+			// Copy voxels to compacted locations
+			ParallelFor(NumTiles, [&](int32 TileIndex)
+				{
+					for (int32 AttributesIdx = 0; AttributesIdx < 2; ++AttributesIdx)
+					{
+						uint32 VoxelDataWriteOffset = reinterpret_cast<uint32*>(RelativeBulkDataPtr + MipStreamingInfo.TileDataOffsetsOffset[AttributesIdx])[TileIndex];
+						const uint32* TileOccupancyBits = reinterpret_cast<uint32*>(RelativeBulkDataPtr + MipStreamingInfo.OccupancyBitsOffset[AttributesIdx] + TileIndex * SVT::NumOccupancyWordsPerPaddedTile * sizeof(uint32));
 
 						for (int64 VoxelIndex = 0; VoxelIndex < SVT::NumVoxelsPerPaddedTile; ++VoxelIndex)
 						{
-							const uint8* Src = (AttributesIdx == 0 ? PhysicalTileDataA.GetData() : PhysicalTileDataB.GetData()) + FormatSize[AttributesIdx] * (TileIndex * SVT::NumVoxelsPerPaddedTile + VoxelIndex);
-							bool bIsFallbackValue = FMemory::Memcmp(Src, NullTileValuesU8[AttributesIdx], FormatSize[AttributesIdx]) == 0;
-							if (!bIsFallbackValue)
+							const int64 WordIndex = (VoxelIndex / 32);
+
+							if (TileOccupancyBits[WordIndex] & (1u << (static_cast<uint32>(VoxelIndex) % 32u)))
 							{
-								const int64 WordIndex = TileIndex * SVT::NumOccupancyWordsPerPaddedTile + (VoxelIndex / 32);
-								uint32* OccupancyBits = reinterpret_cast<uint32*>(BulkData.GetData() + MipStreamingInfo.BulkOffset + MipStreamingInfo.OccupancyBitsOffset[AttributesIdx]);
-								OccupancyBits[WordIndex] |= 1u << (static_cast<uint32>(VoxelIndex) % 32u);
-								NumValidVoxelsTotal++;
-								BulkData.Append(Src, FormatSize[AttributesIdx]);
+								uint8* Dst = DstTileData[AttributesIdx] + VoxelDataWriteOffset * FormatSize[AttributesIdx];
+								const uint8* Src = (AttributesIdx == 0 ? PhysicalTileDataA.GetData() : PhysicalTileDataB.GetData()) + FormatSize[AttributesIdx] * (TileIndex * SVT::NumVoxelsPerPaddedTile + VoxelIndex);
+								FMemory::Memcpy(Dst, Src, FormatSize[AttributesIdx]);
+								++VoxelDataWriteOffset;
 							}
 						}
 					}
-				}
-
-				MipStreamingInfo.TileDataSize[AttributesIdx] = BulkData.Num() - MipStreamingInfo.BulkOffset - MipStreamingInfo.TileDataOffset[AttributesIdx];
-			}
+				});
 		};
 
 		TArray<uint8> StreamableBulkData;
@@ -411,6 +473,8 @@ bool FResources::Build(USparseVolumeTextureFrame* Owner, UE::Serialization::FEdi
 
 		// Store StreamableMipLevels
 		{
+			TRACE_CPUPROFILER_EVENT_SCOPE(SVT::FResources::Build::StoreStreamableMipLevels);
+
 			StreamableMipLevels.Lock(LOCK_READ_WRITE);
 			uint8* Ptr = (uint8*)StreamableMipLevels.Realloc(StreamableBulkData.Num());
 			FMemory::Memcpy(Ptr, StreamableBulkData.GetData(), StreamableBulkData.Num());
@@ -431,6 +495,8 @@ void FResources::Cache(USparseVolumeTextureFrame* Owner, UE::Serialization::FEdi
 		return;
 	}
 
+	TRACE_CPUPROFILER_EVENT_SCOPE(SVT::FResources::Cache);
+
 	using namespace UE::DerivedData;
 
 	static const FValueId SVTDataId = FValueId::FromName("SparseVolumeTextureData");
@@ -455,6 +521,8 @@ void FResources::Cache(USparseVolumeTextureFrame* Owner, UE::Serialization::FEdi
 	FSharedBuffer ResourcesDataBuffer;
 	FIoHash SVTStreamingDataHash;
 	{
+		TRACE_CPUPROFILER_EVENT_SCOPE(SVT::FResources::Cache::CheckDDC);
+
 		FCacheRecordPolicyBuilder PolicyBuilder(DefaultCachePolicy | ECachePolicy::KeepAlive);
 		PolicyBuilder.AddValuePolicy(SVTStreamingDataId, DefaultCachePolicy | ECachePolicy::SkipData);
 
@@ -480,6 +548,8 @@ void FResources::Cache(USparseVolumeTextureFrame* Owner, UE::Serialization::FEdi
 
 	if (!ResourcesDataBuffer.IsNull())
 	{
+		TRACE_CPUPROFILER_EVENT_SCOPE(SVT::FResources::Cache::SerializeFromDDC);
+
 		// Found it!
 		// We can serialize the data from the DDC buffer and are done.
 		FMemoryReaderView Ar(ResourcesDataBuffer.GetView(), /*bIsPersistent=*/ true);
@@ -520,6 +590,8 @@ void FResources::Cache(USparseVolumeTextureFrame* Owner, UE::Serialization::FEdi
 
 		bool bSavedToDDC = false;
 		{
+			TRACE_CPUPROFILER_EVENT_SCOPE(SVT::FResources::Cache::SaveToDDC);
+
 			FValue Value = FValue::Compress(FSharedBuffer::MakeView(Ar.GetData(), Ar.TotalSize()));
 			RecordBuilder.AddValue(SVTDataId, Value);
 
