World Partition Class Descriptors:
- Always prefetch class descriptors when saving objects during cooks, fixes missing class descriptors when saving cooked BPs.
- Don't try to create class descriptors for BPs saved prior to class descriptors object version, fixes warnings for missing classes during cooks.

diff --git a/Engine/Source/Runtime/Engine/Private/WorldPartition/WorldPartitionClassDescRegistry.cpp b/Engine/Source/Runtime/Engine/Private/WorldPartition/WorldPartitionClassDescRegistry.cpp
--- a/Engine/Source/Runtime/Engine/Private/WorldPartition/WorldPartitionClassDescRegistry.cpp
+++ b/Engine/Source/Runtime/Engine/Private/WorldPartition/WorldPartitionClassDescRegistry.cpp
@@ -10,6 +10,7 @@
 #include "HAL/FileManager.h"
 #include "Misc/PackageName.h"
 #include "UObject/CoreRedirects.h"
+#include "UObject/FortniteMainBranchObjectVersion.h"
 #include "UObject/ObjectSaveContext.h"
 #include "UObject/UObjectIterator.h"
 #include "AssetRegistry/AssetData.h"
@@ -232,7 +233,7 @@ void FWorldPartitionClassDescRegistry::PrefetchClassDescs(const TArray<FTopLevel
 			for (int32 AssetIndex = 1; AssetIndex < Pair.Value.Num(); AssetIndex++)	{ *Pair.Value[AssetIndex] = FAssetData(); }
 		});
 		Assets.SetNum(Algo::RemoveIf(Assets, [](const FAssetData& AssetData) { return !AssetData.IsValid(); }));
-	}	
+	}
 
 	TMap<FTopLevelAssetPath, FTopLevelAssetPath> ParentClassPathsMap;
 	Algo::ForEach(Assets, [this, &AssetRegistry, &ParentClassPathsMap, &SortBlueprintAssetsByNameLength, &GetBlueprintAssets](FAssetData& AssetData)
@@ -303,7 +304,7 @@ void FWorldPartitionClassDescRegistry::PrefetchClassDescs(const TArray<FTopLevel
 	}
 
 	// Register current class descriptors
-	Algo::ForEach(Assets, [this](FAssetData& AssetData)
+	Algo::ForEach(Assets, [this, &AssetRegistry](FAssetData& AssetData)
 	{
 		const FTopLevelAssetPath AssetDataClassName(GetAssetDataClassName(AssetData));
 		AssetData.AssetName = *AssetDataClassName.GetAssetName().ToString();
@@ -311,7 +312,28 @@ void FWorldPartitionClassDescRegistry::PrefetchClassDescs(const TArray<FTopLevel
 		// Lookup for an already registered class
 		if (!ClassByPath.Contains(FTopLevelAssetPath(AssetData.ToSoftObjectPath().ToString())))
 		{
-			if (TUniquePtr<FWorldPartitionActorDesc> ClassDesc = FWorldPartitionActorDescUtils::GetActorDescriptorFromAssetData(AssetData); ClassDesc.IsValid())
+			bool bOldAsset = false;
+			FAssetPackageData PackageData;
+			if (AssetRegistry.TryGetAssetPackageData(AssetData.PackageName, PackageData) == UE::AssetRegistry::EExists::Exists)
+			{
+				for (const UE::AssetRegistry::FPackageCustomVersion& CustomVersion : PackageData.GetCustomVersions())
+				{
+					if(CustomVersion.Key == FFortniteMainBranchObjectVersion::GUID)
+					{
+						if (CustomVersion.Version < FFortniteMainBranchObjectVersion::WorldPartitionActorClassDescSerialize)
+						{
+							bOldAsset = true;
+						}
+						break;
+					}
+				}
+			}
+			else
+			{
+				bOldAsset = true;
+			}
+
+			if (TUniquePtr<FWorldPartitionActorDesc> ClassDesc = !bOldAsset ? FWorldPartitionActorDescUtils::GetActorDescriptorFromAssetData(AssetData) : nullptr; ClassDesc.IsValid())
 			{
 				FWorldPartitionActorDesc* Result = ClassDesc.Release();
 				RegisterClassDescriptor(Result);
@@ -436,18 +458,15 @@ void FWorldPartitionClassDescRegistry::RegisterClassDescriptorFromActorClass(con
 
 void FWorldPartitionClassDescRegistry::OnObjectPreSave(UObject* InObject, FObjectPreSaveContext InSaveContext)
 {
-	if (!InSaveContext.IsProceduralSave() && !(InSaveContext.GetSaveFlags() & SAVE_FromAutosave))
+	if (UBlueprint* Blueprint = Cast<UBlueprint>(InObject))
 	{
-		if (UBlueprint* Blueprint = Cast<UBlueprint>(InObject))
+		// We are saving a blueprint
+		if (UBlueprintGeneratedClass* BlueprintGeneratedClass = Cast<UBlueprintGeneratedClass>(Blueprint->GeneratedClass))
 		{
-			// We are saving a blueprint
-			if (UBlueprintGeneratedClass* BlueprintGeneratedClass = Cast<UBlueprintGeneratedClass>(Blueprint->GeneratedClass))
+			if (BlueprintGeneratedClass->IsChildOf<AActor>())
 			{
-				if (BlueprintGeneratedClass->IsChildOf<AActor>())
-				{
-					PrefetchClassDesc(BlueprintGeneratedClass);
-					UpdateClassDescriptor(Blueprint, false);
-				}
+				PrefetchClassDesc(BlueprintGeneratedClass);
+				UpdateClassDescriptor(Blueprint, false);
 			}
 		}
 	}
@@ -650,4 +669,4 @@ FTopLevelAssetPath FWorldPartitionClassDescRegistry::RedirectClassPath(const FTo
 
 	return ResolvedClassPath;
 }
-#endif
\ No newline at end of file
+#endif
