Curve Keys appears to be white no matter how many keys are added on all Curves
#fix match old behaviour: generate a unique colour for newly added curve(s)

diff --git a/Engine/Plugins/Animation/AnimationData/Source/AnimationData/Private/AnimSequencerController.cpp b/Engine/Plugins/Animation/AnimationData/Source/AnimationData/Private/AnimSequencerController.cpp
--- a/Engine/Plugins/Animation/AnimationData/Source/AnimationData/Private/AnimSequencerController.cpp
+++ b/Engine/Plugins/Animation/AnimationData/Source/AnimationData/Private/AnimSequencerController.cpp
@@ -495,6 +495,7 @@ bool UAnimSequencerController::AddCurve(const FAnimationCurveIdentifier& CurveId
 					{
 						ReportError(LOCTEXT("FailedtoAddCurveControl", "Failed to add curve control"));
 					}
+					SetCurveColor(CurveId, FAnimCurveBase::MakeColor(CurveId.CurveName), bShouldTransact);
 					break;
 				default:
 					{
diff --git a/Engine/Source/Runtime/Engine/Private/Animation/AnimCurveTypes.cpp b/Engine/Source/Runtime/Engine/Private/Animation/AnimCurveTypes.cpp
--- a/Engine/Source/Runtime/Engine/Private/Animation/AnimCurveTypes.cpp
+++ b/Engine/Source/Runtime/Engine/Private/Animation/AnimCurveTypes.cpp
@@ -74,7 +74,7 @@ void FAnimCurveBase::PostSerialize(const FArchive& Ar)
 		{
 			// Need to set the curve name before we generate a new color
 			CurveName = Name_DEPRECATED.DisplayName;
-			Color = MakeColor();
+			Color = MakeColor(CurveName);
 		}
 
 		if(Ar.CustomVer(FUE5MainStreamObjectVersion::GUID) < FUE5MainStreamObjectVersion::AnimationRemoveSmartNames)
@@ -120,7 +120,7 @@ int32 FAnimCurveBase::GetCurveTypeFlags() const
 }
 
 #if WITH_EDITORONLY_DATA
-FLinearColor FAnimCurveBase::MakeColor()
+FLinearColor FAnimCurveBase::MakeColor(const FName& CurveName)
 {
 	// Create a color based on the hash of the name
 	FRandomStream Stream(GetTypeHash(CurveName));
diff --git a/Engine/Source/Runtime/Engine/Public/Animation/AnimCurveTypes.h b/Engine/Source/Runtime/Engine/Public/Animation/AnimCurveTypes.h
--- a/Engine/Source/Runtime/Engine/Public/Animation/AnimCurveTypes.h
+++ b/Engine/Source/Runtime/Engine/Public/Animation/AnimCurveTypes.h
@@ -125,7 +125,7 @@ public:
 		: CurveTypeFlags(0)
 	{
 #if WITH_EDITORONLY_DATA
-		Color = MakeColor();
+		Color = MakeColor(CurveName);
 #endif
 	}
 
@@ -134,7 +134,7 @@ public:
 		, CurveTypeFlags(InCurveTypeFlags)
 	{
 #if WITH_EDITORONLY_DATA
-		Color = MakeColor();
+		Color = MakeColor(CurveName);
 #endif
 	}
 
@@ -144,7 +144,7 @@ public:
 		, CurveTypeFlags(InCurveTypeFlags)
 	{
 #if WITH_EDITORONLY_DATA
-		Color = MakeColor();
+		Color = MakeColor(CurveName);
 #endif
 	}
 
@@ -195,9 +195,8 @@ public:
 	/** Get the color used to display this curve in the editor */
 	FLinearColor GetColor() const { return Color; }
 
-private:
 	/** Make an initial color */
-	ENGINE_API FLinearColor MakeColor();
+	static ENGINE_API FLinearColor MakeColor(const FName& CurveName);
 #endif
 };
 
