Added wp.Runtime.UpdateStreamingStateTimeLimit to allow incremental updates of registered world partitions (disabled by default)
This is only useful when there's a lot of registered partitioned worlds.

diff --git a/Engine/Source/Runtime/Engine/Private/WorldPartition/WorldPartition.cpp b/Engine/Source/Runtime/Engine/Private/WorldPartition/WorldPartition.cpp
--- a/Engine/Source/Runtime/Engine/Private/WorldPartition/WorldPartition.cpp
+++ b/Engine/Source/Runtime/Engine/Private/WorldPartition/WorldPartition.cpp
@@ -728,7 +728,7 @@ void UWorldPartition::Uninitialize()
 		// Unload all loaded cells
 		if (World->IsGameWorld())
 		{
-			UWorldPartitionSubsystem::UpdateStreamingStateInternal(World, { this });
+			UWorldPartitionSubsystem::UpdateStreamingStateInternal(World, this);
 		}
 
 #if WITH_EDITOR
diff --git a/Engine/Source/Runtime/Engine/Private/WorldPartition/WorldPartitionSubsystem.cpp b/Engine/Source/Runtime/Engine/Private/WorldPartition/WorldPartitionSubsystem.cpp
--- a/Engine/Source/Runtime/Engine/Private/WorldPartition/WorldPartitionSubsystem.cpp
+++ b/Engine/Source/Runtime/Engine/Private/WorldPartition/WorldPartitionSubsystem.cpp
@@ -135,6 +135,14 @@ static FAutoConsoleVariableRef CMaxLoadingStreamingCells(
 	GMaxLoadingStreamingCells,
 	TEXT("Used to limit the number of concurrent loading world partition streaming cells."));
 
+static float GUpdateStreamingStateTimeLimit = 0.f;
+static FAutoConsoleVariableRef CVarUdateStreamingStateTimeLimit(
+	TEXT("wp.Runtime.UpdateStreamingStateTimeLimit"),
+	GUpdateStreamingStateTimeLimit,
+	TEXT("Maximum amount of time to spend doing World Partition UpdateStreamingState (ms per frame)."),
+	ECVF_Default
+);
+
 UWorldPartitionSubsystem::UWorldPartitionSubsystem()
 : StreamingSourcesHash(0)
 , NumWorldPartitionServerStreamingEnabled(0)
@@ -506,6 +514,7 @@ void UWorldPartitionSubsystem::OnWorldPartitionInitialized(UWorldPartition* InWo
 
 	check(!RegisteredWorldPartitions.Contains(InWorldPartition));
 	RegisteredWorldPartitions.Add(InWorldPartition);
+	IncrementalUpdateWorldPartitionsPendingAdd.Add(InWorldPartition);
 	NumWorldPartitionServerStreamingEnabled += InWorldPartition->IsServerStreamingEnabled() ? 1 : 0;
 	check(NumWorldPartitionServerStreamingEnabled >= 0);
 }
@@ -514,6 +523,8 @@ void UWorldPartitionSubsystem::OnWorldPartitionUninitialized(UWorldPartition* In
 {
 	check(RegisteredWorldPartitions.Contains(InWorldPartition));
 	RegisteredWorldPartitions.Remove(InWorldPartition);
+	IncrementalUpdateWorldPartitionsPendingAdd.Remove(InWorldPartition);
+	IncrementalUpdateWorldPartitions.Remove(InWorldPartition);
 	const UWorld* OwningWorld = GetWorld();
 	if (OwningWorld->IsGameWorld())
 	{
@@ -696,7 +707,7 @@ void UWorldPartitionSubsystem::OnLevelStreamingTargetStateChanged(UWorld* InWorl
 		if (WorldPartition && WorldPartition->IsInitialized() && !WorldPartition->CanStream())
 		{
 			TRACE_CPUPROFILER_EVENT_SCOPE(UWorldPartitionSubsystem::OnLevelStreamingTargetStateChanged);
-			UWorldPartitionSubsystem::UpdateStreamingStateInternal(InWorld, { WorldPartition });
+			UWorldPartitionSubsystem::UpdateStreamingStateInternal(InWorld, WorldPartition);
 		}
 	}
 }
@@ -991,10 +1002,47 @@ void UWorldPartitionSubsystem::UpdateStreamingState()
 {
 	SCOPE_CYCLE_COUNTER(STAT_WorldPartitionUpdateStreaming);
 
-	UWorldPartitionSubsystem::UpdateStreamingStateInternal(GetWorld(), RegisteredWorldPartitions);
+	UWorldPartitionSubsystem::UpdateStreamingStateInternal(GetWorld());
+}
+
+bool UWorldPartitionSubsystem::IncrementalUpdateStreamingState()
+{
+	TRACE_CPUPROFILER_EVENT_SCOPE(UWorldPartitionSubsystem::IncrementalUpdateStreamingState);
+	check(GUpdateStreamingStateTimeLimit > 0.f);
+
+	// Make a snapshot of registered World Partitions to update
+	if (IncrementalUpdateWorldPartitions.IsEmpty())
+	{
+		IncrementalUpdateWorldPartitions.Append(RegisteredWorldPartitions);
+		IncrementalUpdateWorldPartitionsPendingAdd.Reset();
+	}
+	// Append all World Partitions added since last incremental update
+	else if (!IncrementalUpdateWorldPartitionsPendingAdd.IsEmpty())
+	{
+		IncrementalUpdateWorldPartitions.Append(IncrementalUpdateWorldPartitionsPendingAdd);
+		IncrementalUpdateWorldPartitionsPendingAdd.Reset();
+	}
+
+	const double StartTime = FPlatformTime::Seconds();
+	for (auto It = IncrementalUpdateWorldPartitions.CreateIterator(); It; ++It)
+	{
+		UWorldPartition* RegisteredWorldPartition = *It;
+		if (RegisteredWorldPartition->StreamingPolicy)
+		{
+			RegisteredWorldPartition->StreamingPolicy->UpdateStreamingState();
+		}
+		It.RemoveCurrent();
+
+		const double DeltaTime = (FPlatformTime::Seconds() - StartTime) * 1000;
+		if (DeltaTime > GUpdateStreamingStateTimeLimit)
+		{
+			break;
+		}
+	}
+	return IncrementalUpdateWorldPartitions.IsEmpty();
 }
 
-void UWorldPartitionSubsystem::UpdateStreamingStateInternal(const UWorld* InWorld, const TArray<TObjectPtr<UWorldPartition>>& InWorldPartitions)
+void UWorldPartitionSubsystem::UpdateStreamingStateInternal(const UWorld* InWorld, UWorldPartition* InWorldPartition)
 {
 	TRACE_CPUPROFILER_EVENT_SCOPE(UWorldPartitionSubsystem::UpdateStreamingState);
 
@@ -1005,27 +1053,67 @@ void UWorldPartitionSubsystem::UpdateStreamingStateInternal(const UWorld* InWorl
 	}
 
 	// Subsystem can be null during EndPlayMap. WorldPartition::Uninitialize will still call UpdateStreamingStateInternal to cleanup it's streaming levels
-	if (UWorldPartitionSubsystem* WorldPartitionSubsystem = UWorld::GetSubsystem<UWorldPartitionSubsystem>(World))
+	UWorldPartitionSubsystem* WorldPartitionSubsystem = UWorld::GetSubsystem<UWorldPartitionSubsystem>(World);
+	check(WorldPartitionSubsystem || InWorldPartition);
+	if (!InWorldPartition && WorldPartitionSubsystem->RegisteredWorldPartitions.IsEmpty())
+	{
+		return;
+	}
+
+	TArray<TObjectPtr<UWorldPartition>> RegisteredWorldPartitionsCopy;
+	auto GetRegisteredWorldPartitionsCopy = [&RegisteredWorldPartitionsCopy, InWorldPartition, WorldPartitionSubsystem]()
+	{
+		if (RegisteredWorldPartitionsCopy.IsEmpty())
+		{
+			// Make temp copy of array as UpdateStreamingState may FlushAsyncLoading, which may add a new world partition to RegisteredWorldPartitions while iterating
+			RegisteredWorldPartitionsCopy = InWorldPartition ? TArray<TObjectPtr<UWorldPartition>>({ InWorldPartition }) : WorldPartitionSubsystem->RegisteredWorldPartitions;
+		}
+		return RegisteredWorldPartitionsCopy;
+	};
+
+	if (WorldPartitionSubsystem)
 	{
 		// Update streaming sources
 		WorldPartitionSubsystem->UpdateStreamingSources();
 		// Update server's clients visible levels
 		WorldPartitionSubsystem->UpdateServerClientsVisibleLevelNames();
 	}
+
+	const bool bIsServer = IsServer(InWorld);
+	const bool bServerStreamingEnabled = bIsServer && WorldPartitionSubsystem && WorldPartitionSubsystem->HasAnyWorldPartitionServerStreamingEnabled();
+	const int32 WorldPartitionUpdateCount = InWorldPartition ? 1 : WorldPartitionSubsystem->RegisteredWorldPartitions.Num();
+	const bool bIncrementalUpdate = (GUpdateStreamingStateTimeLimit > 0.f) &&
+									(WorldPartitionUpdateCount > 1) &&
+									(!bIsServer || bServerStreamingEnabled) &&
+									!InWorld->GetIsInBlockTillLevelStreamingCompleted();
 	
-	// Make temp copy of array as UpdateStreamingState may FlushAsyncLoading, which may add a new world partition to RegisteredWorldPartitions while iterating
-	const TArray<UWorldPartition*> RegisteredWorldPartitionsCopy = InWorldPartitions;
+	// Update streaming state of all registered world partitions
+	if (bIncrementalUpdate)
+	{
+		// Early out if incremental update doesn't complete
+		if (!WorldPartitionSubsystem->IncrementalUpdateStreamingState())
+		{
+			return;
+		}
+	}
+	else
+	{
+		for (UWorldPartition* RegisteredWorldPartition : GetRegisteredWorldPartitionsCopy())
+		{
+			if (RegisteredWorldPartition->StreamingPolicy)
+			{
+				RegisteredWorldPartition->StreamingPolicy->UpdateStreamingState();
+			}
+		}
+	}
 
-	TArray<const UWorldPartitionRuntimeCell*> ToActivateCells;
+	// Cumulate cells to load and to activate
 	TArray<const UWorldPartitionRuntimeCell*> ToLoadCells;
-
-	// Update streaming state of all registered world partitions and cumulate cells to activate and load
-	TMap<const UWorldPartitionRuntimeCell*, UWorldPartition*> CellToWorldPartition;
-	for (UWorldPartition* RegisteredWorldPartition : RegisteredWorldPartitionsCopy)
+	TArray<const UWorldPartitionRuntimeCell*> ToActivateCells;
+	for (UWorldPartition* RegisteredWorldPartition : GetRegisteredWorldPartitionsCopy())
 	{
 		if (RegisteredWorldPartition->StreamingPolicy)
 		{
-			RegisteredWorldPartition->StreamingPolicy->UpdateStreamingState();
 			RegisteredWorldPartition->StreamingPolicy->GetCellsToUpdate(ToLoadCells, ToActivateCells);
 		}
 	}
@@ -1065,7 +1153,7 @@ void UWorldPartitionSubsystem::UpdateStreamingStateInternal(const UWorld* InWorl
 		TRACE_CPUPROFILER_EVENT_SCOPE(BuildPendingCells);
 		PendingToLoadCells.Reserve(ToLoadCells.Num());
 		PendingToActivateCells.Reserve(ToActivateCells.Num());
-		for (UWorldPartition* RegisteredWorldPartition : RegisteredWorldPartitionsCopy)
+		for (UWorldPartition* RegisteredWorldPartition : GetRegisteredWorldPartitionsCopy())
 		{
 			if (RegisteredWorldPartition->StreamingPolicy)
 			{
diff --git a/Engine/Source/Runtime/Engine/Public/WorldPartition/WorldPartitionSubsystem.h b/Engine/Source/Runtime/Engine/Public/WorldPartition/WorldPartitionSubsystem.h
--- a/Engine/Source/Runtime/Engine/Public/WorldPartition/WorldPartitionSubsystem.h
+++ b/Engine/Source/Runtime/Engine/Public/WorldPartition/WorldPartitionSubsystem.h
@@ -170,12 +170,13 @@ private:
 	void OnLevelBeginMakingVisible(UWorld* InWorld, const ULevelStreaming* InStreamingLevel, ULevel* InLoadedLevel);
 	void OnLevelBeginMakingInvisible(UWorld* InWorld, const ULevelStreaming* InStreamingLevel, ULevel* InLoadedLevel);
 	void UpdateLoadingAndPendingLoadStreamingLevels(const ULevelStreaming* InStreamingLevel);
+	bool IncrementalUpdateStreamingState();
 	
 	// Server information
 	uint32 GetServerClientsVisibleLevelsHash() const { return ServerClientsVisibleLevelsHash; }
 	void UpdateServerClientsVisibleLevelNames();
 
-	static ENGINE_API void UpdateStreamingStateInternal(const UWorld* InWorld, const TArray<TObjectPtr<UWorldPartition>>& InWorldPartitions);
+	static ENGINE_API void UpdateStreamingStateInternal(const UWorld* InWorld, UWorldPartition* InWorldPartition = nullptr);
 	static int32 GetMaxCellsToLoad(const UWorld* InWorld);
 	static bool IsServer(const UWorld* InWorld);
 
@@ -190,13 +191,14 @@ private:
 	friend class UWorldPartition;
 	friend class UWorldPartitionStreamingPolicy;
 
+	// Registered world partitions & incremental update
 	TArray<TObjectPtr<UWorldPartition>> RegisteredWorldPartitions;
+	TSet<TObjectPtr<UWorldPartition>> IncrementalUpdateWorldPartitions;
+	TSet<TObjectPtr<UWorldPartition>> IncrementalUpdateWorldPartitionsPendingAdd;
 
+	// Streaming Sources
 	TSet<IWorldPartitionStreamingSourceProvider*> StreamingSourceProviders;
-
 	FWorldPartitionStreamingSourceProviderFilter IsStreamingSourceProviderFiltered;
-
-	// Streaming Sources
 	TArray<FWorldPartitionStreamingSource> StreamingSources;
 	TMap<FName, FStreamingSourceVelocity> StreamingSourcesVelocity;
 	uint32 StreamingSourcesHash;
@@ -206,6 +208,7 @@ private:
 	TSet<FName> ServerClientsVisibleLevelNames;
 	uint32 ServerClientsVisibleLevelsHash;
 
+	// Debug draw
 	TArray<FWorldPartitionDraw2DContext> WorldPartitionsDraw2DContext;
 	FDelegateHandle	DrawHandle;
 
