Server deactivated cells (activated -> loaded) are now processed immediately so that UNetConnection::UpdateLevelVisibilityInternal properly handles client TryMakeVisible level visility requests.

diff --git a/Engine/Source/Runtime/Engine/Private/WorldPartition/WorldPartitionStreamingPolicy.cpp b/Engine/Source/Runtime/Engine/Private/WorldPartition/WorldPartitionStreamingPolicy.cpp
--- a/Engine/Source/Runtime/Engine/Private/WorldPartition/WorldPartitionStreamingPolicy.cpp
+++ b/Engine/Source/Runtime/Engine/Private/WorldPartition/WorldPartitionStreamingPolicy.cpp
@@ -475,7 +475,7 @@ void UWorldPartitionStreamingPolicy::UpdateStreamingState()
 	// Activation superseeds Loading
 	FrameLoadCells = FrameLoadCells.Difference(FrameActivateCells);
 
-	// Determine cells to activate (sorted by importance)
+	// Determine cells to activate
 	if (bCanStream)
 	{
 		TRACE_CPUPROFILER_EVENT_SCOPE(UWorldPartitionStreamingPolicy::UpdateStreamingState_ToActivateCells);
@@ -497,7 +497,8 @@ void UWorldPartitionStreamingPolicy::UpdateStreamingState()
 		}
 	}
 
-	// Determine cells to load (sorted by importance)
+	// Determine cells to load and server cells to deactivate
+	TArray<const UWorldPartitionRuntimeCell*> ServerToDeactivateCells;
 	if (bCanStream)
 	{
 		TRACE_CPUPROFILER_EVENT_SCOPE(UWorldPartitionStreamingPolicy::UpdateStreamingState_ToLoadCells);
@@ -513,13 +514,22 @@ void UWorldPartitionStreamingPolicy::UpdateStreamingState()
 			}
 			else
 			{
-				// Only deactivated server cells need to call ShouldWaitForClientVisibility (those part of ActivatedCells)
-				const bool bIsServerCellToDeactivate = bIsServer && ActivatedCells.Contains(Cell);
-				const bool bShoudSkipServerCell = bIsServerCellToDeactivate && ShouldWaitForClientVisibility(Cell);
-				if (!bShoudSkipServerCell && !ShouldSkipCellForPerformance(Cell) && !ShouldSkipDisabledHLODCell(Cell))
+				if (!ShouldSkipCellForPerformance(Cell) && !ShouldSkipDisabledHLODCell(Cell))
 				{
-					Cell->MergeStreamingSourceInfo();
-					ToLoadCells.Add(Cell);
+					// Server deactivated cells are processed right away (see below for details)
+					if (const bool bIsServerCellToDeactivate = bIsServer && ActivatedCells.Contains(Cell))
+					{
+						// Only deactivated server cells need to call ShouldWaitForClientVisibility (those part of ActivatedCells)
+						if (!ShouldWaitForClientVisibility(Cell))
+						{
+							ServerToDeactivateCells.Add(Cell);
+						}
+					}
+					else
+					{
+						Cell->MergeStreamingSourceInfo();
+						ToLoadCells.Add(Cell);
+					}
 				}
 			}
 		}
@@ -560,10 +570,23 @@ void UWorldPartitionStreamingPolicy::UpdateStreamingState()
 	UpdateDebugCellsStreamingPriority(FrameActivateCells, FrameLoadCells);
 #endif
 
+	// Unloaded cells
 	if (ToUnloadCells.Num() > 0)
 	{
 		SetCellsStateToUnloaded(ToUnloadCells);
 	}
+	
+	// Server deactivated cells (activated -> loaded)
+	// 
+	// Server deactivation is handle right away to ensure that even if WorldPartitionSubsystem::UpdateStreamingState
+	// is running in incremental mode, server deactivated cells will make their streaming level ShouldBeVisible() 
+	// return false. This way, UNetConnection::UpdateLevelVisibilityInternal will not allow clients to make their 
+	// streaming level visible (see LevelVisibility.bTryMakeVisible).
+	for (const UWorldPartitionRuntimeCell* ServerCellToDeactivate : ServerToDeactivateCells)
+	{
+		int32 DummyMaxCellToLoad = 0; // Deactivating is not concerned by MaxCellsToLoad
+		SetCellStateToLoaded(ServerCellToDeactivate, DummyMaxCellToLoad);
+	}
 
 	// Evaluate streaming performance based on cells that should be activated
 	UpdateStreamingPerformance(FrameActivateCells);
