Output additional information about invalid compressed data when trying to finish async compression task

diff --git a/Engine/Source/Runtime/Engine/Private/Animation/AnimSequence.cpp b/Engine/Source/Runtime/Engine/Private/Animation/AnimSequence.cpp
--- a/Engine/Source/Runtime/Engine/Private/Animation/AnimSequence.cpp
+++ b/Engine/Source/Runtime/Engine/Private/Animation/AnimSequence.cpp
@@ -5478,7 +5478,7 @@ void UAnimSequence::FinishAsyncTasks()
 			}
 			else
 			{
-				if (CompressedAnimData->IsValid(this))
+				if (CompressedAnimData->IsValid(this, true))
 				{
 #if WITH_EDITOR
 					//This is only safe during sync anim compression
