Networked Physics: Optimization to FPhysScene_Chaos::PopulateReplicationCache() which locked GT for an extended period of time.

Now only iterate over and cache state of physics-components on actors that actively replicate their movement.

diff --git a/Engine/Source/Runtime/Engine/Private/ActorReplication.cpp b/Engine/Source/Runtime/Engine/Private/ActorReplication.cpp
--- a/Engine/Source/Runtime/Engine/Private/ActorReplication.cpp
+++ b/Engine/Source/Runtime/Engine/Private/ActorReplication.cpp
@@ -380,11 +380,12 @@ void AActor::GatherCurrentMovement()
 			bool bFoundInCache = false;
 
 			UWorld* World = GetWorld();
+			int ServerFrame = 0;
 			if (FPhysScene_Chaos* Scene = static_cast<FPhysScene_Chaos*>(World->GetPhysicsScene()))
 			{
-				if (FRigidBodyState* FoundState = Scene->ReplicationCache.Map.Find(FObjectKey(RootPrimComp)))
+				if (const FRigidBodyState* FoundState = Scene->GetStateFromReplicationCache(RootPrimComp, ServerFrame))
 				{
-					ReplicatedMovement.FillFrom(*FoundState, this, Scene->ReplicationCache.ServerFrame);
+					ReplicatedMovement.FillFrom(*FoundState, this, ServerFrame);
 					bFoundInCache = true;
 				}
 			}
@@ -394,7 +395,7 @@ void AActor::GatherCurrentMovement()
 				// fallback to GT data
 				FRigidBodyState RBState;
 				RootPrimComp->GetRigidBodyState(RBState);
-				ReplicatedMovement.FillFrom(RBState, this, 0);
+				ReplicatedMovement.FillFrom(RBState, this, ServerFrame);
 			}
 
 			// Don't replicate movement if we're welded to another parent actor.
diff --git a/Engine/Source/Runtime/Engine/Private/PhysicsEngine/Experimental/PhysScene_Chaos.cpp b/Engine/Source/Runtime/Engine/Private/PhysicsEngine/Experimental/PhysScene_Chaos.cpp
--- a/Engine/Source/Runtime/Engine/Private/PhysicsEngine/Experimental/PhysScene_Chaos.cpp
+++ b/Engine/Source/Runtime/Engine/Private/PhysicsEngine/Experimental/PhysScene_Chaos.cpp
@@ -57,6 +57,9 @@ FAutoConsoleVariableRef CVar_KinematicDeferralUpdateExternalAccelerationStructur
 bool GKinematicDeferralLogInvalidBodies = false;
 FAutoConsoleVariableRef CVar_KinematicDeferralLogInvalidBodies(TEXT("p.KinematicDeferralLogInvalidBodies"), GKinematicDeferralLogInvalidBodies, TEXT("If true and p.KinematicDeferralCheckValidBodies is true, log when an invalid body is found on kinematic update."));
 
+float GReplicationCacheLingerForNSeconds = 3.f;
+FAutoConsoleVariableRef CVar_ReplicationCacheLingerForNSeconds(TEXT("np2.ReplicationCache.LingerForNSeconds"), GReplicationCacheLingerForNSeconds, TEXT("How long to keep data in the replication cache without the actor accessing it, after this we stop caching the actors state until it tries to access it again."));
+
 DECLARE_CYCLE_STAT(TEXT("Update Kinematics On Deferred SkelMeshes"), STAT_UpdateKinematicsOnDeferredSkelMeshesChaos, STATGROUP_Physics);
 
 struct FPendingAsyncPhysicsCommand
@@ -530,6 +533,7 @@ FPhysScene_Chaos::~FPhysScene_Chaos()
 	
 	// Make sure physics replication is cleared before we're fully destructed
 	PhysicsReplication.Reset();
+	ReplicationCache.Reset();
 
 	FPhysicsDelegates::OnPhysSceneTerm.Broadcast(this);
 
@@ -1217,45 +1221,88 @@ void FPhysScene_Chaos::AddToComponentMaps(UPrimitiveComponent* Component, IPhysi
 	}
 }
 
-void FPhysScene_Chaos::PopulateReplicationCache(const int32 PhysicsStep)
+const FRigidBodyState* FPhysScene_Chaos::GetStateFromReplicationCache(UPrimitiveComponent* RootComponent, int& ServerFrame)
 {
-	auto ReplicationCacheHelper = [this](auto& Handle, UPrimitiveComponent* RootComponent)
+	if (!GetSolver()->GetRewindCallback())
 	{
-		FObjectKey Key(RootComponent);
-		FRigidBodyState& ReplicationState = ReplicationCache.Map.FindOrAdd(Key);
+		// We only populate replication cache through the RewindCallback
+		ServerFrame = 0;
+		return nullptr;
+	}
 
-		ReplicationState.Position = Handle->X();
-		ReplicationState.Quaternion = Handle->R();
-		ReplicationState.LinVel = Handle->V();
-		ReplicationState.AngVel = Handle->W();
-		ReplicationState.Flags = Handle->ObjectState() == Chaos::EObjectStateType::Sleeping ? ERigidBodyFlags::Sleeping : 0;
-	};
+	ServerFrame = ReplicationCache.ServerFrame;
 
+	const FObjectKey Key(RootComponent);
+	if (!ReplicationCache.Map.Contains(Key))
+	{
+		RegisterForReplicationCache(RootComponent);
+	}
+
+	FRigidBodyState* ReplicationState = nullptr;
+	if (FReplicationCacheData* ReplicationData = ReplicationCache.Map.Find(Key))
+	{
+		ReplicationData->SetAccessTime(GetSolver()->GetSolverTime());
+		ReplicationState = &ReplicationData->GetState();
+	}
+	return ReplicationState;
+}
+
+void FPhysScene_Chaos::RegisterForReplicationCache(UPrimitiveComponent* RootComponent)
+{
+	const FObjectKey Key(RootComponent);
+	ReplicationCache.Map.Add(Key, FReplicationCacheData(RootComponent, GetSolver()->GetSolverTime()));
+}
+
+void FPhysScene_Chaos::PopulateReplicationCache(const int32 PhysicsStep)
+{
+	auto ReplicationCacheHelper = [this](auto& Handle, FReplicationCacheData& ReplicationData, bool& StateWasCached)
+	{
+		// If the component reference has lingered in the replication cache for too long without being accessed, remove it and stop caching data.
+		const Chaos::FReal CacheLingerTime = GetSolver()->GetSolverTime() - ReplicationData.GetAccessTime();
+		if (CacheLingerTime > GReplicationCacheLingerForNSeconds)
+		{
+			StateWasCached = false;
+		}
+		else
+		{
+			FRigidBodyState& ReplicationState = ReplicationData.GetState();
+			ReplicationState.Position = Handle->X();
+			ReplicationState.Quaternion = Handle->R();
+			ReplicationState.LinVel = Handle->V();
+			ReplicationState.AngVel = Handle->W();
+			ReplicationState.Flags = Handle->ObjectState() == Chaos::EObjectStateType::Sleeping ? ERigidBodyFlags::Sleeping : 0;
+			StateWasCached = true;
+		}
+	};
+	
 	ReplicationCache.ServerFrame = PhysicsStep;
-	for (TActorIterator<AActor> It(GetOwningWorld()); It; ++It)
+	bool StateWasCached;
+	for (auto It = ReplicationCache.Map.CreateIterator(); It; ++It)
 	{
-		AActor* Actor = *It;
-		if (Actor && Actor->GetIsReplicated() && Actor->IsReplicatingMovement())
+		StateWasCached = false;
+		FReplicationCacheData& ReplicationData = It.Value();
+		TObjectPtr<UPrimitiveComponent> RootComponent = ReplicationData.GetRootComponent();
+		if (RootComponent)
 		{
-			UPrimitiveComponent* RootComponent = Cast<UPrimitiveComponent>(Actor->GetRootComponent());
-			if (RootComponent && RootComponent->IsSimulatingPhysics())
+			if (FBodyInstanceAsyncPhysicsTickHandle BIHandle = RootComponent->GetBodyInstanceAsyncPhysicsTickHandle())
 			{
-				if (FBodyInstanceAsyncPhysicsTickHandle BIHandle = RootComponent->GetBodyInstanceAsyncPhysicsTickHandle())
-				{
-					ReplicationCacheHelper(BIHandle, RootComponent);
-				}
-				else if (Chaos::FPhysicsObjectHandle PhysicsObject = RootComponent->GetPhysicsObjectByName(NAME_None))
+				ReplicationCacheHelper(BIHandle, ReplicationData, StateWasCached);
+			}
+			else if (Chaos::FPhysicsObjectHandle PhysicsObject = RootComponent->GetPhysicsObjectByName(NAME_None))
+			{
+				Chaos::FReadPhysicsObjectInterface_Internal Interface = Chaos::FPhysicsObjectInternalInterface::GetRead();
+				if (Chaos::FPBDRigidParticleHandle* POHandle = Interface.GetRigidParticle(PhysicsObject))
 				{
-					Chaos::FReadPhysicsObjectInterface_Internal Interface = Chaos::FPhysicsObjectInternalInterface::GetRead();
-					Chaos::FPBDRigidParticleHandle* POHandle = Interface.GetRigidParticle(PhysicsObject);
-
-					if (POHandle != nullptr)
-					{
-						ReplicationCacheHelper(POHandle, RootComponent);
-					}
+					ReplicationCacheHelper(POHandle, ReplicationData, StateWasCached);
 				}
 			}
 		}
+
+		if (!StateWasCached)
+		{
+			// Deregister actor from ReplicationCache
+			It.RemoveCurrent();
+		}
 	}
 }
 
diff --git a/Engine/Source/Runtime/Engine/Public/Physics/Experimental/PhysScene_Chaos.h b/Engine/Source/Runtime/Engine/Public/Physics/Experimental/PhysScene_Chaos.h
--- a/Engine/Source/Runtime/Engine/Public/Physics/Experimental/PhysScene_Chaos.h
+++ b/Engine/Source/Runtime/Engine/Public/Physics/Experimental/PhysScene_Chaos.h
@@ -70,7 +70,6 @@ namespace Chaos
 
 }
 
-
 extern int32 GEnableKinematicDeferralStartPhysicsCondition;
 
 struct FConstraintBrokenDelegateWrapper
@@ -247,15 +246,47 @@ public:
 
 	UChaosEventRelay* GetChaosEventRelay() const { return ChaosEventRelay.Get(); }
 
+	/** Get cached state for replication, if no state is cached RegisterForReplicationCache() is called */
+	ENGINE_API const FRigidBodyState* GetStateFromReplicationCache(UPrimitiveComponent* RootComponent, int& ServerFrame);
+	
+	/** Register a component for physics replication state caching, the component will deregister automatically if cache is not accessed within timelimit set by CVar: np2.ReplicationCache.LingerForNSeconds */
+	ENGINE_API void RegisterForReplicationCache(UPrimitiveComponent* RootComponent);
+
 	/** Populate the replication cache from the list of registered components */
 	ENGINE_API void PopulateReplicationCache(const int32 PhysicsStep);
 
+	struct FReplicationCacheData
+	{
+		FReplicationCacheData(UPrimitiveComponent* InRootComponent, Chaos::FReal InAccessTime)
+			: RootComponent(InRootComponent)
+			, AccessTime(InAccessTime)
+		{
+			RootComponent->GetRigidBodyState(State);
+		}
+
+		TObjectPtr<UPrimitiveComponent> GetRootComponent()	{ return RootComponent; }
+		FRigidBodyState& GetState()	{ return State; }
+		void SetAccessTime(Chaos::FReal Time) { AccessTime = Time; }
+		Chaos::FReal GetAccessTime() { return AccessTime; }
+
+	private:
+		TObjectPtr<UPrimitiveComponent> RootComponent;
+		Chaos::FReal AccessTime;
+		FRigidBodyState State;
+	};
+
 	// Storage structure for replication data
 	// probably should just expose read/write API not the structure directly itself like this.
 	struct FPrimitiveComponentReplicationCache
 	{
 		int32 ServerFrame = 0;
-		TMap<FObjectKey, FRigidBodyState>	Map;
+		TMap<FObjectKey, FReplicationCacheData> Map;
+
+		void Reset()
+		{
+			ServerFrame = 0;
+			Map.Reset();
+		}
 	};
 
 	FPrimitiveComponentReplicationCache ReplicationCache;
