Attached Actor Pre/PostSaveRoot transform fix

diff --git a/Engine/Source/Runtime/Engine/Private/LevelUtils.cpp b/Engine/Source/Runtime/Engine/Private/LevelUtils.cpp
--- a/Engine/Source/Runtime/Engine/Private/LevelUtils.cpp
+++ b/Engine/Source/Runtime/Engine/Private/LevelUtils.cpp
@@ -38,6 +38,27 @@ struct FLevelReadOnlyData
 // Map to link the level data with a level
 static TMap<ULevel*, FLevelReadOnlyData> LevelReadOnlyCache;
 
+namespace LevelUtilsInternal
+{
+	static void ApplyEditorTransform(ULevel* LoadedLevel, bool bDoPostEditMove, AActor* Actor, const FTransform& Transform)
+	{
+		if (LoadedLevel)
+		{
+			FLevelUtils::FApplyLevelTransformParams TransformParams(LoadedLevel, Transform);
+
+			while (AActor* AttachParent = Actor ? Actor->GetAttachParentActor() : nullptr)
+			{
+				Actor = AttachParent;
+			}
+
+			TransformParams.bDoPostEditMove = bDoPostEditMove;
+			TransformParams.Actor = Actor;
+
+			FLevelUtils::ApplyLevelTransform(TransformParams);
+		}
+	}
+}
+
 #endif
 
 /////////////////////////////////////////////////////////////////////////////////////////
@@ -356,26 +377,13 @@ void FLevelUtils::SetEditorTransform(ULevelStreaming* StreamingLevel, const FTra
 void FLevelUtils::ApplyEditorTransform(const ULevelStreaming* StreamingLevel, bool bDoPostEditMove, AActor* Actor)
 {
 	check(StreamingLevel);
-	if (ULevel* LoadedLevel = StreamingLevel->GetLoadedLevel())
-	{	
-		FApplyLevelTransformParams TransformParams(LoadedLevel, StreamingLevel->LevelTransform);
-		TransformParams.Actor = Actor;
-		TransformParams.bDoPostEditMove = bDoPostEditMove;
-		ApplyLevelTransform(TransformParams);
-	}
+	LevelUtilsInternal::ApplyEditorTransform(StreamingLevel->GetLoadedLevel(), bDoPostEditMove, Actor, StreamingLevel->LevelTransform);
 }
 
 void FLevelUtils::RemoveEditorTransform(const ULevelStreaming* StreamingLevel, bool bDoPostEditMove, AActor* Actor)
 {
 	check(StreamingLevel);
-	if (ULevel* LoadedLevel = StreamingLevel->GetLoadedLevel())
-	{
-		const FTransform InverseTransform = StreamingLevel->LevelTransform.Inverse();
-		FApplyLevelTransformParams TransformParams(LoadedLevel, InverseTransform);
-		TransformParams.Actor = Actor;
-		TransformParams.bDoPostEditMove = bDoPostEditMove;
-		ApplyLevelTransform(TransformParams);
-	}
+	LevelUtilsInternal::ApplyEditorTransform(StreamingLevel->GetLoadedLevel(), bDoPostEditMove, Actor, StreamingLevel->LevelTransform.Inverse());
 }
 
 void FLevelUtils::ApplyPostEditMove( ULevel* Level )
