nDisplay: Actor Replication System

diff --git a/Engine/Plugins/Runtime/nDisplay/Source/DisplayCluster/DisplayCluster.Build.cs b/Engine/Plugins/Runtime/nDisplay/Source/DisplayCluster/DisplayCluster.Build.cs
--- a/Engine/Plugins/Runtime/nDisplay/Source/DisplayCluster/DisplayCluster.Build.cs
+++ b/Engine/Plugins/Runtime/nDisplay/Source/DisplayCluster/DisplayCluster.Build.cs
@@ -21,6 +21,7 @@ public class DisplayCluster : ModuleRules
 				"DisplayClusterLightCardEditorShaders",
 				"DisplayClusterLightCardExtender",
 				"Engine",
+				"EngineSettings",
 				"EnhancedInput",
 			});
 
diff --git a/Engine/Plugins/Runtime/nDisplay/Source/DisplayCluster/Private/Cluster/DisplayClusterNetDriverHelper.cpp b/Engine/Plugins/Runtime/nDisplay/Source/DisplayCluster/Private/Cluster/DisplayClusterNetDriverHelper.cpp
new file mode 100644
--- /dev/null
+++ b/Engine/Plugins/Runtime/nDisplay/Source/DisplayCluster/Private/Cluster/DisplayClusterNetDriverHelper.cpp
@@ -0,0 +1,93 @@
+// Copyright Epic Games, Inc. All Rights Reserved.
+
+#include "Cluster/DisplayClusterNetDriverHelper.h"
+#include "DisplayClusterConfigurationTypes.h"
+
+#include "Misc/DisplayClusterGlobals.h"
+#include "Misc/DisplayClusterStrings.h"
+
+#include "Network/Service/ClusterEventsBinary/DisplayClusterClusterEventsBinaryClient.h"
+
+bool FDisplayClusterNetDriverHelper::RegisterClusterEventsBinaryClient(uint32 ClusterId, const FString& ClientAddress, uint16 ClientPort)
+{	
+	TSharedPtr<FDisplayClusterClusterEventsBinaryClient>* FoundClient = ClusterEventsBinaryClients.Find(ClusterId);
+
+	if (FoundClient != nullptr)
+	{
+		UE_LOG(LogDisplayClusterNetwork, Error, TEXT("Couldn't not register event binary client for %s:%d. It already exists."), *ClientAddress, ClientPort);
+		return false;
+	}
+	
+	TSharedPtr<FDisplayClusterClusterEventsBinaryClient> Client = MakeShared<FDisplayClusterClusterEventsBinaryClient>();
+	
+	if (!Client->Connect(ClientAddress, ClientPort, 1, 0.f))
+	{
+		UE_LOG(LogDisplayClusterNetwork, Error, TEXT("Couldn't not connect to %s:%d"), *ClientAddress, ClientPort);
+
+		return false;
+	}
+
+	ClusterEventsBinaryClients.Add(ClusterId, Client);
+
+	return true;
+}
+
+bool FDisplayClusterNetDriverHelper::RemoveClusterEventsBinaryClient(uint32 СlusterId)
+{
+	TSharedPtr<FDisplayClusterClusterEventsBinaryClient>* FoundClient = ClusterEventsBinaryClients.Find(СlusterId);
+
+	if (FoundClient == nullptr)
+	{
+		return false;
+	}
+
+	const int32 RemovedEntries = ClusterEventsBinaryClients.Remove(СlusterId);
+
+	return RemovedEntries > 0;
+}
+
+bool FDisplayClusterNetDriverHelper::HasClient(uint32 ClusterId)
+{
+	return ClusterEventsBinaryClients.Contains(ClusterId);
+}
+
+bool FDisplayClusterNetDriverHelper::GetRequiredArguments(const FURL& URL, const TCHAR*& OutClusterId, const TCHAR*& OutPrimaryNodeId, const TCHAR*& OutPrimaryNodePort)
+{
+	if (!URL.HasOption(DisplayClusterStrings::uri_args::ClusterId) ||
+		!URL.HasOption(DisplayClusterStrings::uri_args::PrimaryNodeId) ||
+		!URL.HasOption(DisplayClusterStrings::uri_args::PrimaryNodePort))
+	{
+		return false;
+	}
+
+	OutClusterId = URL.GetOption(DisplayClusterStrings::uri_args::ClusterId, nullptr);
+	OutPrimaryNodeId = URL.GetOption(DisplayClusterStrings::uri_args::PrimaryNodeId, nullptr);
+	OutPrimaryNodePort = URL.GetOption(DisplayClusterStrings::uri_args::PrimaryNodePort, nullptr);
+
+	return true;
+}
+
+bool FDisplayClusterNetDriverHelper::SendCommandToCluster(uint32 ClusterId, const FDisplayClusterClusterEventBinary& Event)
+{
+	TSharedPtr<FDisplayClusterClusterEventsBinaryClient>* ClientFound = ClusterEventsBinaryClients.Find(ClusterId);
+
+	if (ClientFound == nullptr)
+	{
+		UE_LOG(LogDisplayClusterNetwork, Error, TEXT("Event binary client not found for ClusterId %d"), ClusterId);
+		return false;
+	}
+
+	TSharedPtr<FDisplayClusterClusterEventsBinaryClient> Client = *ClientFound;
+
+	const EDisplayClusterCommResult CommResult = Client->EmitClusterEventBinary(Event);
+
+	return (CommResult == EDisplayClusterCommResult::Ok);
+}
+
+void FDisplayClusterNetDriverHelper::SendCommandToAllClusters(const FDisplayClusterClusterEventBinary& Event)
+{
+	for (const TPair<uint32, TSharedPtr<FDisplayClusterClusterEventsBinaryClient>>& Client : ClusterEventsBinaryClients)
+	{
+		Client.Value->EmitClusterEventBinary(Event);
+	}
+}
diff --git a/Engine/Plugins/Runtime/nDisplay/Source/DisplayCluster/Private/Game/EngineClasses/Basics/DisplayClusterGameEngine.cpp b/Engine/Plugins/Runtime/nDisplay/Source/DisplayCluster/Private/Game/EngineClasses/Basics/DisplayClusterGameEngine.cpp
--- a/Engine/Plugins/Runtime/nDisplay/Source/DisplayCluster/Private/Game/EngineClasses/Basics/DisplayClusterGameEngine.cpp
+++ b/Engine/Plugins/Runtime/nDisplay/Source/DisplayCluster/Private/Game/EngineClasses/Basics/DisplayClusterGameEngine.cpp
@@ -15,7 +15,13 @@
 
 #include "Misc/App.h"
 #include "Misc/CommandLine.h"
+#include "Misc/ConfigCacheIni.h"
+#include "Misc/CoreDelegates.h"
 #include "Misc/DisplayClusterAppExit.h"
+#include "Misc/DisplayClusterGlobals.h"
+#include "Misc/DisplayClusterHelpers.h"
+#include "Misc/DisplayClusterLog.h"
+#include "Misc/DisplayClusterStrings.h"
 #include "Misc/Parse.h"
 
 #include "Misc/DisplayClusterGlobals.h"
@@ -26,6 +32,14 @@
 #include "SocketSubsystem.h"
 #include "Interfaces/IPv4/IPv4Endpoint.h"
 
+#include "EngineGlobals.h"
+#include "Engine/NetDriver.h"
+#include "Engine/GameInstance.h"
+
+#include "UObject/Linker.h"
+#include "GameMapsSettings.h"
+#include "GameFramework/GameModeBase.h"
+
 #include "Stats/Stats.h"
 
 #include "GameDelegates.h"
@@ -143,6 +157,9 @@ void UDisplayClusterGameEngine::Init(class IEngineLoop* InEngineLoop)
 
 	// Initialize base stuff.
 	UGameEngine::Init(InEngineLoop);
+
+	OnOverrideBrowseURL.BindUObject(this, &UDisplayClusterGameEngine::BrowseLoadMap);
+	OnOverridePendingNetGameUpdate.BindUObject(this, &UDisplayClusterGameEngine::PendingLevelUpdate);
 }
 
 EDisplayClusterOperationMode UDisplayClusterGameEngine::DetectOperationMode() const
@@ -463,3 +480,153 @@ void UDisplayClusterGameEngine::GameSyncChange(const FDisplayClusterClusterEvent
 		}
 	}
 }
+
+EBrowseReturnVal::Type UDisplayClusterGameEngine::BrowseLoadMap(FWorldContext& WorldContext, FURL URL, FString& Error)
+{
+	FString DisplayClusterServerType;
+	FParse::Value(FCommandLine::Get(), TEXT("dc_replicationserver_type"), DisplayClusterServerType);
+
+	const bool IsDisplayCluster = (IDisplayCluster::Get().GetOperationMode() == EDisplayClusterOperationMode::Cluster);
+	const bool IsDisplayClusterListenServer = IsDisplayCluster && DisplayClusterServerType.Equals(TEXT("listen"));
+
+	if (URL.IsLocalInternal() && !IsDisplayClusterListenServer)
+	{
+		// Local map file.
+		return LoadMap(WorldContext, URL, NULL, Error) ? EBrowseReturnVal::Success : EBrowseReturnVal::Failure;
+	}
+	else if ((URL.IsInternal() && GIsClient) || (URL.IsLocalInternal() && IsDisplayClusterListenServer))
+	{
+		EBrowseReturnVal::Type BrowseResult = EBrowseReturnVal::Failure;
+
+		if (IsDisplayClusterListenServer)
+		{
+			BrowseResult = LoadMap(WorldContext, URL, NULL, Error) ? EBrowseReturnVal::Success : EBrowseReturnVal::Failure;
+		}
+
+		if ((URL.IsInternal() && GIsClient) || (BrowseResult == EBrowseReturnVal::Success))
+		{
+			// Network URL.
+			if (WorldContext.PendingNetGame)
+			{
+				CancelPending(WorldContext);
+			}
+
+			// Clean up the netdriver/socket so that the pending level succeeds
+			if (WorldContext.World() && ShouldShutdownWorldNetDriver())
+			{
+				ShutdownWorldNetDriver(WorldContext.World());
+			}
+
+			WorldContext.PendingNetGame = NewObject<UPendingNetGame>();
+			WorldContext.PendingNetGame->Initialize(URL); //-V595
+			WorldContext.PendingNetGame->InitNetDriver(); //-V595
+
+			if (!WorldContext.PendingNetGame)
+			{
+				// If the inital packet sent in InitNetDriver results in a socket error, HandleDisconnect() and CancelPending() may be called, which will null the PendingNetGame.
+				Error = NSLOCTEXT("Engine", "PendingNetGameInitFailure", "Error initializing the network driver.").ToString();
+				BroadcastTravelFailure(WorldContext.World(), ETravelFailure::PendingNetGameCreateFailure, Error);
+				return EBrowseReturnVal::Failure;
+			}
+
+			if (!WorldContext.PendingNetGame->NetDriver)
+			{
+				// UPendingNetGame will set the appropriate error code and connection lost type, so
+				// we just have to propagate that message to the game.
+				BroadcastTravelFailure(WorldContext.World(), ETravelFailure::PendingNetGameCreateFailure, WorldContext.PendingNetGame->ConnectionError);
+				WorldContext.PendingNetGame = NULL;
+				return EBrowseReturnVal::Failure;
+			}
+			return EBrowseReturnVal::Pending;
+		}
+	}
+	else if (URL.IsInternal())
+	{
+		// Invalid.
+		Error = NSLOCTEXT("Engine", "ServerOpen", "Servers can't open network URLs").ToString();
+		return EBrowseReturnVal::Failure;
+	}
+
+	return EBrowseReturnVal::Failure;
+}
+
+void UDisplayClusterGameEngine::PendingLevelUpdate(FWorldContext& Context, float DeltaSeconds)
+{
+	// Update the pending level.
+	if (Context.PendingNetGame)
+	{
+		Context.PendingNetGame->Tick(DeltaSeconds);
+		if (Context.PendingNetGame && Context.PendingNetGame->ConnectionError.Len() > 0)
+		{
+			BroadcastNetworkFailure(NULL, Context.PendingNetGame->NetDriver, ENetworkFailure::PendingConnectionFailure, Context.PendingNetGame->ConnectionError);
+			CancelPending(Context);
+		}
+		else if (Context.PendingNetGame && Context.PendingNetGame->bSuccessfullyConnected && !Context.PendingNetGame->bSentJoinRequest && !Context.PendingNetGame->bLoadedMapSuccessfully && (Context.OwningGameInstance == NULL || !Context.OwningGameInstance->DelayPendingNetGameTravel()))
+		{
+			if (Context.PendingNetGame->HasFailedTravel())
+			{
+				BrowseToDefaultMap(Context);
+				BroadcastTravelFailure(Context.World(), ETravelFailure::TravelFailure, TEXT("Travel failed for unknown reason"));
+			}
+			else if (!MakeSureMapNameIsValid(Context.PendingNetGame->URL.Map))
+			{
+				BrowseToDefaultMap(Context);
+				BroadcastTravelFailure(Context.World(), ETravelFailure::PackageMissing, Context.PendingNetGame->URL.Map);
+			}
+			else if (!Context.PendingNetGame->bLoadedMapSuccessfully)
+			{
+				// Attempt to load the map.
+				FString Error;
+				FString DisplayClusterServerType;
+				FParse::Value(FCommandLine::Get(), TEXT("dc_replicationserver_type"), DisplayClusterServerType);
+
+				const bool IsDisplayCluster = FParse::Param(FCommandLine::Get(), TEXT("dc_cluster"));
+				const bool IsDisplayClusterListenServer = IsDisplayCluster && DisplayClusterServerType.Equals(TEXT("listen"));
+
+				bool bLoadedMapSuccessfully = true;
+
+				if (IsDisplayClusterListenServer)
+				{
+					MovePendingLevel(Context);
+				}
+				else
+				{
+					bLoadedMapSuccessfully = LoadMap(Context, Context.PendingNetGame->URL, Context.PendingNetGame, Error);
+				}
+
+				if (Context.PendingNetGame != nullptr)
+				{
+					if (!Context.PendingNetGame->LoadMapCompleted(this, Context, bLoadedMapSuccessfully, Error))
+					{
+						BrowseToDefaultMap(Context);
+						BroadcastTravelFailure(Context.World(), ETravelFailure::LoadMapFailure, Error);
+					}
+				}
+				else
+				{
+					BrowseToDefaultMap(Context);
+					BroadcastTravelFailure(Context.World(), ETravelFailure::TravelFailure, Error);
+				}
+			}
+		}
+
+		if (Context.PendingNetGame && Context.PendingNetGame->bLoadedMapSuccessfully && (Context.OwningGameInstance == NULL || !Context.OwningGameInstance->DelayCompletionOfPendingNetGameTravel()))
+		{
+			if (!Context.PendingNetGame->HasFailedTravel())
+			{
+				Context.PendingNetGame->TravelCompleted(this, Context);
+				Context.PendingNetGame = nullptr;
+			}
+			else
+			{
+				CancelPending(Context);
+				BrowseToDefaultMap(Context);
+				BroadcastTravelFailure(Context.World(), ETravelFailure::LoadMapFailure, TEXT("Travel failed for unknown reason"));
+			}
+		}
+	}
+	else if (TransitionType == ETransitionType::WaitingToConnect)
+	{
+		TransitionType = ETransitionType::None;
+	}
+}
diff --git a/Engine/Plugins/Runtime/nDisplay/Source/DisplayCluster/Public/Cluster/DisplayClusterNetDriverHelper.h b/Engine/Plugins/Runtime/nDisplay/Source/DisplayCluster/Public/Cluster/DisplayClusterNetDriverHelper.h
new file mode 100644
--- /dev/null
+++ b/Engine/Plugins/Runtime/nDisplay/Source/DisplayCluster/Public/Cluster/DisplayClusterNetDriverHelper.h
@@ -0,0 +1,59 @@
+// Copyright Epic Games, Inc. All Rights Reserved.
+
+#pragma once
+
+#include "CoreMinimal.h"
+
+#include "Cluster/DisplayClusterClusterEvent.h"
+#include "Engine/EngineBaseTypes.h"
+
+#include "DisplayClusterNetDriverHelper.generated.h"
+
+class FDisplayClusterClusterEventsBinaryClient;
+
+/** Helper for DisplayCluster synchronous replication
+ */
+
+USTRUCT()
+struct DISPLAYCLUSTER_API FDisplayClusterNetDriverHelper
+{
+	GENERATED_BODY()
+
+public:
+	/** Creates and stores nDisplay cluster events client using cluster id and primary node parameters
+	 *  @param СlusterId unique identifier of cluster
+	 *  @param ClientAddress primary node IP address
+	 *  @param ClientPort primary node port
+	 */
+	bool RegisterClusterEventsBinaryClient(uint32 СlusterId, const FString& ClientAddress, uint16 ClientPort);
+
+	/** Removes cluster event client for ClusterId
+	 *  @param СlusterId unique identifier of cluster
+	 */
+	bool RemoveClusterEventsBinaryClient(uint32 СlusterId);
+
+	/** Checks if nDisplay cluster events client registered for specified cluster
+	 *  @param СlusterId unique identifier of cluster
+	 */
+	bool HasClient(uint32 ClusterId);
+
+	/** Checks if URL has all required arguments for DisplayClusterNetDriver
+	 *  @param URL Url to check
+	 */
+	static bool GetRequiredArguments(const FURL& URL, const TCHAR*& OutClusterId, const TCHAR*& OutPrimaryNodeId, const TCHAR*& OutPrimaryNodePort);
+
+	/** Sends command to nDisplay cluster specified by id
+	 *  @param СlusterId identifier of cluster
+	 *  @param NetworkDriverSyncEvent cluster event binary data
+	 */
+	bool SendCommandToCluster(uint32 СlusterId, const FDisplayClusterClusterEventBinary& NetworkDriverSyncEvent);
+
+	/** Sends command to all nDisplay clusters stored by RegisterClusterEventsBinaryClient 
+	 *  @param NetworkDriverSyncEvent cluster event binary data
+	 */
+	void SendCommandToAllClusters(const FDisplayClusterClusterEventBinary& NetworkDriverSyncEvent);
+
+private:
+	// nDisplay primary node events clients for sending control messages to clusters
+	TMap<uint32, TSharedPtr<FDisplayClusterClusterEventsBinaryClient>> ClusterEventsBinaryClients;
+};
diff --git a/Engine/Plugins/Runtime/nDisplay/Source/DisplayCluster/Public/DisplayClusterGameEngine.h b/Engine/Plugins/Runtime/nDisplay/Source/DisplayCluster/Public/DisplayClusterGameEngine.h
--- a/Engine/Plugins/Runtime/nDisplay/Source/DisplayCluster/Public/DisplayClusterGameEngine.h
+++ b/Engine/Plugins/Runtime/nDisplay/Source/DisplayCluster/Public/DisplayClusterGameEngine.h
@@ -16,6 +16,8 @@ class IDisplayClusterClusterNodeController;
 class IDisplayClusterClusterSyncObject;
 class UDisplayClusterConfigurationData;
 
+DECLARE_DELEGATE_RetVal_ThreeParams(EBrowseReturnVal::Type, FOnBrowseLoadMap, FWorldContext& WorldContext, FURL URL, FString& Error);
+DECLARE_DELEGATE_TwoParams(FOnPengineLevelUpdate, FWorldContext& Context, float DeltaSeconds);
 
 /**
  * Extended game engine
@@ -61,7 +63,14 @@ private:
 	bool BarrierAvoidanceOn() const;
 
 	void GameSyncChange(const FDisplayClusterClusterEventJson& InEvent);
-
+	
+	/** Implements custom loading logic to support listening server
+	 */
+	EBrowseReturnVal::Type BrowseLoadMap(FWorldContext& WorldContext, FURL URL, FString& Error);
+
+	/** Custom logic to update PendingLevel in order to support listening server
+	 */
+	void PendingLevelUpdate(FWorldContext& Context, float DeltaSeconds);
 private:
 	IPDisplayClusterClusterManager* ClusterMgr = nullptr;
 
diff --git a/Engine/Plugins/Runtime/nDisplay/Source/DisplayCluster/Public/Misc/DisplayClusterStrings.h b/Engine/Plugins/Runtime/nDisplay/Source/DisplayCluster/Public/Misc/DisplayClusterStrings.h
--- a/Engine/Plugins/Runtime/nDisplay/Source/DisplayCluster/Public/Misc/DisplayClusterStrings.h
+++ b/Engine/Plugins/Runtime/nDisplay/Source/DisplayCluster/Public/Misc/DisplayClusterStrings.h
@@ -59,4 +59,11 @@ namespace DisplayClusterStrings
 		// Specific events
 		static constexpr const TCHAR* EvtQuitName   = TEXT("quit");
 	}
+
+	namespace uri_args
+	{
+		static constexpr const TCHAR* ClusterId = TEXT("ClusterId=");
+		static constexpr const TCHAR* PrimaryNodeId = TEXT("PrimaryNodeId=");
+		static constexpr const TCHAR* PrimaryNodePort = TEXT("PrimaryNodePort=");
+	}
 };
diff --git a/Engine/Plugins/Runtime/nDisplay/Source/DisplayClusterConfiguration/Private/DisplayClusterConfigurationTypes.cpp b/Engine/Plugins/Runtime/nDisplay/Source/DisplayClusterConfiguration/Private/DisplayClusterConfigurationTypes.cpp
--- a/Engine/Plugins/Runtime/nDisplay/Source/DisplayClusterConfiguration/Private/DisplayClusterConfigurationTypes.cpp
+++ b/Engine/Plugins/Runtime/nDisplay/Source/DisplayClusterConfiguration/Private/DisplayClusterConfigurationTypes.cpp
@@ -10,8 +10,10 @@
 #include "DisplayClusterConfigurationLog.h"
 #include "DisplayClusterConfigurationStrings.h"
 #include "DisplayClusterProjectionStrings.h"
+#include "Config/IDisplayClusterConfigManager.h"
 
 #include "Engine/StaticMesh.h"
+#include "IDisplayCluster.h"
 #include "UObject/Package.h"
 
 #if WITH_EDITOR
@@ -434,6 +436,35 @@ void UDisplayClusterConfigurationData::GetReferencedMeshNames(TArray<FString>& O
 	}
 }
 
+int32 UDisplayClusterConfigurationData::GetNumberOfClusterNodes() const
+{
+	if (Cluster == nullptr)
+	{
+		return 0;
+	}
+
+	return Cluster->Nodes.Num();
+}
+
+FString UDisplayClusterConfigurationData::GetPrimaryNodeAddress() const
+{
+	if (!IDisplayCluster::Get().GetConfigMgr())
+	{
+		return FString();
+	}
+
+	if (!IDisplayCluster::Get().GetConfigMgr()->GetConfig())
+	{
+		return FString();
+	}
+
+	UDisplayClusterConfigurationData* DisplayClusterConfigurationData = IDisplayCluster::Get().GetConfigMgr()->GetConfig();
+	
+	const FString PrimaryNodeId = DisplayClusterConfigurationData->Cluster->PrimaryNode.Id;
+
+	return DisplayClusterConfigurationData->Cluster->GetNode(PrimaryNodeId)->Host;
+}
+
 UDisplayClusterConfigurationData* UDisplayClusterConfigurationData::CreateNewConfigData(UObject* Owner, EObjectFlags ObjectFlags)
 {
 	UDisplayClusterConfigurationData* NewConfigData = NewObject<UDisplayClusterConfigurationData>(Owner ? Owner : GetTransientPackage(), NAME_None,
diff --git a/Engine/Plugins/Runtime/nDisplay/Source/DisplayClusterConfiguration/Public/DisplayClusterConfigurationTypes.h b/Engine/Plugins/Runtime/nDisplay/Source/DisplayClusterConfiguration/Public/DisplayClusterConfigurationTypes.h
--- a/Engine/Plugins/Runtime/nDisplay/Source/DisplayClusterConfiguration/Public/DisplayClusterConfigurationTypes.h
+++ b/Engine/Plugins/Runtime/nDisplay/Source/DisplayClusterConfiguration/Public/DisplayClusterConfigurationTypes.h
@@ -541,6 +541,13 @@ public:
 	UFUNCTION(BlueprintCallable, Category = "NDisplay|Configuration")
 	void GetReferencedMeshNames(TArray<FString>& OutMeshNames) const;
 
+public:
+	/** Returns number of nodes in current cluster */
+	int32 GetNumberOfClusterNodes() const;
+
+	/** Returns primary node IP address of current cluster */
+	FString GetPrimaryNodeAddress() const;
+
 public:
 	FDisplayClusterConfigurationDataMetaInfo Meta;
 
diff --git a/Engine/Plugins/Runtime/nDisplay/Source/DisplayClusterEditor/Private/Settings/DisplayClusterEditorSettings.cpp b/Engine/Plugins/Runtime/nDisplay/Source/DisplayClusterEditor/Private/Settings/DisplayClusterEditorSettings.cpp
--- a/Engine/Plugins/Runtime/nDisplay/Source/DisplayClusterEditor/Private/Settings/DisplayClusterEditorSettings.cpp
+++ b/Engine/Plugins/Runtime/nDisplay/Source/DisplayClusterEditor/Private/Settings/DisplayClusterEditorSettings.cpp
@@ -3,6 +3,7 @@
 #include "DisplayClusterEditorSettings.h"
 #include "DisplayClusterEditorEngine.h"
 #include "HAL/PlatformFileManager.h"
+#include "Serialization/JsonSerializerMacros.h"
 #include "Misc/ConfigCacheIni.h"
 
 
@@ -15,6 +16,7 @@ UDisplayClusterEditorSettings::UDisplayClusterEditorSettings(const FObjectInitia
 	: Super(ObjectInitializer) 
 {
 	GET_MEMBER_NAME_CHECKED(UDisplayClusterEditorSettings, bEnabled);
+	GET_MEMBER_NAME_CHECKED(UDisplayClusterEditorSettings, bClusterReplicationEnabled);
 }
 
 void UDisplayClusterEditorSettings::PostEditChangeProperty(FPropertyChangedEvent& PropertyChangedEvent)
@@ -82,5 +84,29 @@ void UDisplayClusterEditorSettings::PostEditChangeProperty(FPropertyChangedEvent
 			GConfig->Flush(false, DefaultGamePath);
 			GConfig->Flush(false, DefaultInputPath);
 		}
+		else if (PropertyName == GET_MEMBER_NAME_CHECKED(UDisplayClusterEditorSettings, bClusterReplicationEnabled))
+		{
+			if (FPlatformFileManager::Get().GetPlatformFile().IsReadOnly(*DefaultEnginePath))
+			{
+				FPlatformFileManager::Get().GetPlatformFile().SetReadOnly(*DefaultEnginePath, false);
+			}
+
+			if (bClusterReplicationEnabled)
+			{
+				FJsonSerializableArray NetDriverDefinitions;
+				NetDriverDefinitions.Add(TEXT("(DefName=GameNetDriver,DriverClassName=/Script/DisplayClusterReplication.DisplayClusterNetDriver,DriverClassNameFallback=/Script/OnlineSubsystemUtils.IpNetDriver)"));
+				NetDriverDefinitions.Add(TEXT("(DefName=DemoNetDriver,DriverClassName=/Script/Engine.DemoNetDriver,DriverClassNameFallback=/Script/Engine.DemoNetDriver)"));
+
+				GConfig->SetArray(TEXT("/Script/Engine.GameEngine"), TEXT("+NetDriverDefinitions"), NetDriverDefinitions, DefaultEnginePath);
+				GConfig->SetString(TEXT("/Script/DisplayClusterReplication.DisplayClusterNetDriver"), TEXT("NetConnectionClassName"), TEXT("DisplayClusterReplication.DisplayClusterNetConnection"), DefaultEnginePath);
+			}
+			else
+			{
+				GConfig->RemoveKey(TEXT("/Script/Engine.GameEngine"), TEXT("!NetDriverDefinitions"), DefaultEnginePath);
+				GConfig->RemoveKey(TEXT("/Script/Engine.GameEngine"), TEXT("+NetDriverDefinitions"), DefaultEnginePath);
+				GConfig->RemoveKey(TEXT("/Script/DisplayClusterReplication.DisplayClusterNetDriver"), TEXT("NetConnectionClassName"), DefaultEnginePath);
+			}
+			GConfig->Flush(false, DefaultEnginePath);
+		}
 	}
 }
diff --git a/Engine/Plugins/Runtime/nDisplay/Source/DisplayClusterEditor/Private/Settings/DisplayClusterEditorSettings.h b/Engine/Plugins/Runtime/nDisplay/Source/DisplayClusterEditor/Private/Settings/DisplayClusterEditorSettings.h
--- a/Engine/Plugins/Runtime/nDisplay/Source/DisplayClusterEditor/Private/Settings/DisplayClusterEditorSettings.h
+++ b/Engine/Plugins/Runtime/nDisplay/Source/DisplayClusterEditor/Private/Settings/DisplayClusterEditorSettings.h
@@ -29,4 +29,8 @@ protected:
 protected:
 	UPROPERTY(config, EditAnywhere, Category = Main, meta = (ToolTip = "When enabled, replaces the original GameEngine to DisplayClusterGameEngine", ConfigRestartRequired = true))
 	bool bEnabled;
+
+	/** When enabled, replaces the original NetDriver to DisplayClusterNetDriver, original NetConnection to DisplayClusterNetDriver */
+	UPROPERTY(config, EditAnywhere, Category = Main, meta = (ToolTip = "When enabled, replaces the original NetDriver to DisplayClusterNetDriver, original NetConnection to DisplayClusterNetDriver", ConfigRestartRequired = true))
+	bool bClusterReplicationEnabled;
 };
diff --git a/Engine/Plugins/Runtime/nDisplay/Source/DisplayClusterReplication/DisplayClusterReplication.Build.cs b/Engine/Plugins/Runtime/nDisplay/Source/DisplayClusterReplication/DisplayClusterReplication.Build.cs
new file mode 100644
--- /dev/null
+++ b/Engine/Plugins/Runtime/nDisplay/Source/DisplayClusterReplication/DisplayClusterReplication.Build.cs
@@ -0,0 +1,30 @@
+// Copyright Epic Games, Inc. All Rights Reserved.
+
+using UnrealBuildTool;
+using System.IO;
+
+public class DisplayClusterReplication : ModuleRules
+{
+	public DisplayClusterReplication(ReadOnlyTargetRules ROTargetRules) : base(ROTargetRules)
+	{
+		PublicDependencyModuleNames.AddRange(
+		new string[] {
+			"Core",
+			"CoreUObject",
+			"DisplayCluster",
+			"DisplayClusterConfiguration",
+			"Engine"
+		});
+
+		PrivateDependencyModuleNames.AddRange(
+		new string[] {
+			"Core",
+			"CoreUObject",
+			"Engine",
+			"NetCore",
+			"OnlineSubsystem",
+			"OnlineSubsystemUtils",
+			"Sockets"
+		});
+	}
+}
diff --git a/Engine/Plugins/Runtime/nDisplay/Source/DisplayClusterReplication/Private/DisplayClusterNetConnection.cpp b/Engine/Plugins/Runtime/nDisplay/Source/DisplayClusterReplication/Private/DisplayClusterNetConnection.cpp
new file mode 100644
--- /dev/null
+++ b/Engine/Plugins/Runtime/nDisplay/Source/DisplayClusterReplication/Private/DisplayClusterNetConnection.cpp
@@ -0,0 +1,164 @@
+// Copyright Epic Games, Inc. All Rights Reserved.
+
+#include "DisplayClusterNetConnection.h"
+
+#include "Cluster/DisplayClusterNetDriverHelper.h"
+
+#include "DisplayClusterNetDriver.h"
+#include "DisplayClusterNetDriverLog.h"
+
+#include "Engine/Engine.h"
+#include "Engine/LevelStreaming.h"
+#include "Engine/PackageMapClient.h"
+#include "Engine/PendingNetGame.h"
+
+#include "IPAddress.h"
+
+#include "Misc/ScopeExit.h"
+
+#include "Net/Core/Misc/PacketAudit.h"
+#include "Net/Core/Trace/NetTrace.h"
+#include "Net/DataChannel.h"
+#include "Net/NetworkProfiler.h"
+#include "Net/RPCDoSDetection.h"
+#include "Net/UnrealNetwork.h"
+
+#include "SocketSubsystem.h"
+#include "Sockets.h"
+
+#include "UObject/LinkerLoad.h"
+#include <Kismet/GameplayStatics.h>
+
+UDisplayClusterNetConnection::UDisplayClusterNetConnection(const FObjectInitializer& ObjectInitializer) :
+	Super(ObjectInitializer),
+	ClientId(0),
+	ClusterId(0),
+	NodePort(0),
+	bNodeIsPrimary(false),
+	bIsClusterConnection(false),
+	bSynchronousMode(false)
+{
+}
+
+void UDisplayClusterNetConnection::SetClientLoginState(const EClientLoginState::Type NewState)
+{
+	Super::SetClientLoginState(NewState);
+
+	const FURL ClientURL(NULL, *RequestURL, TRAVEL_Absolute);
+
+	const TCHAR* NodeOption = TEXT("node=");
+	if (!ClientURL.HasOption(NodeOption))
+	{
+		return;
+	}
+
+	NodeName = ClientURL.GetOption(NodeOption, nullptr);
+
+	UDisplayClusterNetDriver* DisplayClusterNetDriver = Cast<UDisplayClusterNetDriver>(Driver.Get());
+
+	if (DisplayClusterNetDriver == nullptr)
+	{
+		return;
+	}
+
+	if (ClientLoginState == EClientLoginState::Type::ReceivedJoin)
+	{
+		const TCHAR* ClusterIdString = nullptr;
+		const TCHAR* PrimaryNodeId = nullptr;
+		const TCHAR* PrimaryNodePortString = nullptr;
+
+		if (!FDisplayClusterNetDriverHelper::GetRequiredArguments(ClientURL, ClusterIdString, PrimaryNodeId, PrimaryNodePortString))
+		{
+			UE_LOG(LogDisplayClusterNetDriver, Verbose, TEXT("Cluster connection Join has failed: URI arguments are invalid"));
+			return;
+		}
+
+		ClientId = TextKeyUtil::HashString(Challenge);
+		ClusterId = TextKeyUtil::HashString(ClusterIdString);
+
+		NodeAddress = GetRemoteAddr()->ToString(false);
+		NodePort = FCString::Atoi(PrimaryNodePortString);
+
+		bIsClusterConnection = true;
+
+		if (PrimaryNodeId == NodeName)
+		{
+			bNodeIsPrimary = true;
+
+			UE_LOG(LogDisplayClusterNetDriver, Verbose, TEXT("Cluster connection joined: Client [%u]; Primary node [%s]; RemoteAddr: %s:%i"), ClientId, *NodeName, *NodeAddress, NodePort);
+		}
+		else
+		{
+			UE_LOG(LogDisplayClusterNetDriver, Verbose, TEXT("Cluster connection joined: Client [%u]; Node [%s]; RemoteAddr: %s"), ClientId, *NodeName, *GetRemoteAddr()->ToString(false));
+		}
+
+		DisplayClusterNetDriver->AddNodeConnection(this);
+	}
+	else if (ClientLoginState == EClientLoginState::Type::CleanedUp)
+	{
+		DisplayClusterNetDriver->RemoveNodeConnection(this);
+	}
+}
+
+void UDisplayClusterNetConnection::ReceivedPacket(FBitReader& Reader, bool bIsReinjectedPacket, bool bDispatchPacket)
+{
+	UDisplayClusterNetDriver* DisplayClusterNetDriver = Cast<UDisplayClusterNetDriver>(Driver.Get());
+
+	if (DisplayClusterNetDriver == nullptr)
+	{
+		return;
+	}
+
+	if (bSynchronousMode && !DisplayClusterNetDriver->IsServer())
+	{
+		// This is client and it's ready for replication: receive and accumulate packet
+		Super::ReceivedPacket(Reader, bIsReinjectedPacket, false);
+		InPackets.Add(InPacketId, Reader);
+
+		UE_LOG(LogDisplayClusterNetDriver, Verbose, TEXT("Packet %i received and accumulated"), InPacketId);
+	}
+	else
+	{
+		// This is server or client that is not ready to replication
+		Super::ReceivedPacket(Reader, bIsReinjectedPacket, bDispatchPacket);
+
+		UE_LOG(LogDisplayClusterNetDriver, Verbose, TEXT("Packet %i received and processed"), InPacketId);
+	}
+}
+
+void UDisplayClusterNetConnection::ProcessPacket(int32 PacketId)
+{
+	UDisplayClusterNetDriver* DisplayClusterNetDriver = Cast<UDisplayClusterNetDriver>(Driver.Get());
+
+	if (DisplayClusterNetDriver == nullptr)
+	{
+		UE_LOG(LogDisplayClusterNetDriver, Error, TEXT("ProcessPacket: Can't cast NetDriver to DisplayClusterNetDriver"));
+		return;
+	}
+
+	if (!InPackets.Contains(PacketId))
+	{
+		UE_LOG(LogDisplayClusterNetDriver, Verbose, TEXT("Packet %i skipped!"), PacketId);
+
+		return;
+	}
+
+	PacketIDs.Reset();
+	InPackets.GetKeys(PacketIDs);
+
+	for (int32 CurrentPacketId : PacketIDs)
+	{
+		if (CurrentPacketId > PacketId)
+		{
+			break;
+		}
+
+		UE_LOG(LogDisplayClusterNetDriver, Verbose, TEXT("Packet %i processed"), CurrentPacketId);
+
+		bool bSkipAck = false;
+		bool bHasBunchErrors = false;
+		DispatchPacket(InPackets[CurrentPacketId], PacketId, bSkipAck, bHasBunchErrors);
+
+		InPackets.Remove(CurrentPacketId);
+	}
+}
diff --git a/Engine/Plugins/Runtime/nDisplay/Source/DisplayClusterReplication/Private/DisplayClusterNetDriver.cpp b/Engine/Plugins/Runtime/nDisplay/Source/DisplayClusterReplication/Private/DisplayClusterNetDriver.cpp
new file mode 100644
--- /dev/null
+++ b/Engine/Plugins/Runtime/nDisplay/Source/DisplayClusterReplication/Private/DisplayClusterNetDriver.cpp
@@ -0,0 +1,582 @@
+// Copyright Epic Games, Inc. All Rights Reserved.
+
+#include "DisplayClusterNetDriver.h"
+
+#include "Config/IDisplayClusterConfigManager.h"
+
+#include "DisplayClusterNetDriver.h"
+#include "DisplayClusterNetDriverLog.h"
+
+#include "Engine/ActorChannel.h"
+#include "Engine/ChildConnection.h"
+#include "Engine/Engine.h"
+#include "Engine/Level.h"
+#include "Engine/LocalPlayer.h"
+#include "Engine/NetConnection.h"
+#include "Engine/NetworkObjectList.h"
+#include "Engine/PackageMapClient.h"
+#include "Engine/World.h"
+
+#include "GameFramework/Actor.h"
+#include "GameFramework/GameNetworkManager.h"
+#include "GameFramework/Pawn.h"
+#include "GameFramework/PlayerController.h"
+
+#include "IDisplayCluster.h"
+#include "IDisplayClusterCallbacks.h"
+
+#include "Misc/CommandLine.h"
+
+#include "NetworkingDistanceConstants.h"
+
+#include "OnlineSubsystem.h"
+#include "OnlineSubsystemNames.h"
+
+#include "SocketSubsystem.h"
+#include "Sockets.h"
+
+#include "UObject/Package.h"
+
+UDisplayClusterNetDriver::UDisplayClusterNetDriver(const FObjectInitializer& ObjectInitializer) :
+	Super(ObjectInitializer),
+	bLastBunchWasAcked(false),
+	bConnectionViewersAreReady(false)
+{
+	// Preallocate elements to prevent allocs at runtime
+	const int NumReservedElements = 32;
+	ReadyOutPackets.Reserve(NumReservedElements);
+	PacketsParams.Reserve(NumReservedElements);
+
+	// Register event binary listener
+	if (!HasAnyFlags(RF_ArchetypeObject) &&  IDisplayCluster::IsAvailable() && (IDisplayCluster::Get().GetOperationMode() == EDisplayClusterOperationMode::Cluster))
+	{
+		IDisplayClusterClusterManager* ClusterManager = IDisplayCluster::Get().GetClusterMgr();
+	
+		if (ClusterManager)
+		{
+			EventBinaryListener = FOnClusterEventBinaryListener::CreateUObject(this, &UDisplayClusterNetDriver::HandleEvent);
+			ClusterManager->AddClusterEventBinaryListener(EventBinaryListener);
+		}
+	}
+
+	// Delegates to modify ConsiderList processing in UNetDriver::ServerReplicateActors method
+#if WITH_SERVER_CODE
+	OnPreConsiderListUpdateOverride.BindUObject(this, &UDisplayClusterNetDriver::PreListUpdate);
+	OnPostConsiderListUpdateOverride.BindUObject(this, &UDisplayClusterNetDriver::PostListUpdate);
+	OnProcessConsiderListOverride.BindUObject(this, &UDisplayClusterNetDriver::ListUpdate);
+#endif
+}
+
+UDisplayClusterNetDriver::~UDisplayClusterNetDriver()
+{
+	// Unregister event binary listener
+	if (!HasAnyFlags(RF_ArchetypeObject) && IDisplayCluster::IsAvailable() && (IDisplayCluster::Get().GetOperationMode() == EDisplayClusterOperationMode::Cluster))
+	{
+		IDisplayClusterClusterManager* ClusterManager = IDisplayCluster::Get().GetClusterMgr();
+
+		if (ClusterManager)
+		{
+			ClusterManager->RemoveClusterEventBinaryListener(EventBinaryListener);
+		}
+	}
+}
+
+bool UDisplayClusterNetDriver::InitListen(FNetworkNotify* InNotify, FURL& ListenURL, bool bReuseAddressAndPort, FString& Error)
+{
+	ClusterNetworkDriverHelper = MakeUnique<FDisplayClusterNetDriverHelper>();
+
+	if (World->GetNetMode() == NM_ListenServer)
+	{
+		const TCHAR* ClusterIdString = nullptr;
+		const TCHAR* NodeId = nullptr;
+		const TCHAR* NodePortString = nullptr;
+
+		if (!FDisplayClusterNetDriverHelper::GetRequiredArguments(ListenURL, ClusterIdString, NodeId, NodePortString))
+		{
+			UE_LOG(LogDisplayClusterNetDriver, Warning, TEXT("InitListen:: Can't get required URL arguments"));
+
+			return Super::InitListen(InNotify, ListenURL, bReuseAddressAndPort, Error);
+		}
+
+		const UDisplayClusterConfigurationData* ConfigData = IDisplayCluster::Get().GetConfigMgr()->GetConfig();
+
+		const uint32 ClusterId = TextKeyUtil::HashString(ClusterIdString);
+		const uint16 NodePort = FCString::Atoi(NodePortString);
+		const FString NodeAddress = ConfigData->GetPrimaryNodeAddress();
+
+		if (ClusterNetworkDriverHelper->RegisterClusterEventsBinaryClient(ClusterId, NodeAddress, NodePort))
+		{
+			UE_LOG(LogDisplayClusterNetDriver, Verbose, TEXT("Registered primary node for cluster [%u] node [%s:%u]"), ClusterId, *NodeAddress, NodePort);
+		}
+	}
+
+	return Super::InitListen(InNotify, ListenURL, bReuseAddressAndPort, Error);
+}
+
+void UDisplayClusterNetDriver::TickDispatch(float DeltaTime)
+{
+	Super::TickDispatch(DeltaTime);
+
+	// Check if driver is in sync mode
+	if (!IsServer())
+	{
+		return;
+	}
+
+	// Check nodes that are ready for synchronous packets processing
+	for (UDisplayClusterNetConnection* NodeConnectionIt : NodeConnections)
+	{
+		// NodeConnection is already in sync mode - skip
+		if (NodeConnectionIt->bSynchronousMode)
+		{
+			continue;
+		}
+
+		// NodeConnection has not joined server - skip
+		if ((NodeConnectionIt->ClientLoginState != EClientLoginState::Type::ReceivedJoin) || (NodeConnectionIt->OwningActor == nullptr))
+		{
+			continue;
+		}
+		
+		// Event listener for cluster is not registered - skip
+		if (!ClusterNetworkDriverHelper->HasClient(NodeConnectionIt->ClusterId))
+		{
+			continue;
+		}
+
+		// Notify primary node that current node connection is ready
+		FDisplayClusterClusterEventBinary NetworkDriverSyncEvent;
+		GenerateClusterCommandsEvent(NetworkDriverSyncEvent, NodeSyncEvent);
+
+		const bool bEventWasSent = ClusterNetworkDriverHelper->SendCommandToCluster(NodeConnectionIt->ClusterId, NetworkDriverSyncEvent);
+
+		// If event was sent to primary cluster node
+		if (bEventWasSent)
+		{
+			// Mark all node connections related to current ClusterId as ready
+			for (UDisplayClusterNetConnection* ClusterNodeConnectionIt : NodeConnections)
+			{
+				if (NodeConnectionIt->bSynchronousMode)
+				{
+					continue;
+				}
+
+				if (ClusterNodeConnectionIt->ClusterId == NodeConnectionIt->ClusterId)
+				{
+					NodeConnectionIt->bSynchronousMode = true;
+					SyncConnections.Add(NodeConnectionIt);
+				}
+			}
+
+			UE_LOG(LogDisplayClusterNetDriver, Verbose, TEXT("Setting synchronous mode in cluster [%u]"), NodeConnectionIt->ClusterId);
+		}
+	}
+}
+
+void UDisplayClusterNetDriver::AddNodeConnection(UDisplayClusterNetConnection* NetConnection)
+{
+	NodeConnections.Add(NetConnection);
+
+	if (NetConnection->bNodeIsPrimary)
+	{
+		PrimaryNodeConnections.Add(NetConnection);
+
+		if (ClusterNetworkDriverHelper->RegisterClusterEventsBinaryClient(NetConnection->ClusterId, NetConnection->NodeAddress, NetConnection->NodePort))
+		{
+			UE_LOG(LogDisplayClusterNetDriver, Verbose, TEXT("Registered primary node for cluster [%u] node [%s]"), NetConnection->ClusterId, *NetConnection->NodeName);
+		}
+	}
+}
+
+void UDisplayClusterNetDriver::RemoveNodeConnection(UDisplayClusterNetConnection* NetConnection)
+{
+	if (NetConnection->bNodeIsPrimary)
+	{
+		if (ClusterNetworkDriverHelper->RemoveClusterEventsBinaryClient(NetConnection->ClusterId))
+		{
+			UE_LOG(LogDisplayClusterNetDriver, Verbose, TEXT("Removed primary node for cluster [%u] node [%s]"), NetConnection->ClusterId, *NetConnection->NodeName);
+		}
+
+		PrimaryNodeConnections.Remove(NetConnection);
+	}
+
+	SyncConnections.Remove(NetConnection);
+	NodeConnections.Remove(NetConnection);
+}
+
+void UDisplayClusterNetDriver::HandleEvent(FDisplayClusterClusterEventBinary const& InEvent)
+{
+	if ((InEvent.EventId != NodeSyncEvent) && (InEvent.EventId != PacketSyncEvent))
+	{
+		return;
+	}
+
+	if (ServerConnection == nullptr)
+	{
+		return;
+	}
+
+	// This is a client
+	UDisplayClusterNetConnection* ClusterServerConnection = Cast<UDisplayClusterNetConnection>(ServerConnection.Get());
+
+	if (ClusterServerConnection == nullptr)
+	{
+		UE_LOG(LogDisplayClusterNetDriver, Error, TEXT("HandleEvent: Can't cast NetConnection to DriverClusterNetConnection"));
+		return;
+	}
+
+	if (InEvent.EventId == NodeSyncEvent)
+	{
+		UE_LOG(LogDisplayClusterNetDriver, Verbose, TEXT("[StartSynchronousMode] command received"));
+
+		// Client received command to start synchrosonus packets processing
+		ClusterServerConnection->bSynchronousMode = true;
+		ClusterServerConnection->ClientId = TextKeyUtil::HashString(ClusterServerConnection->Challenge);
+
+		UE_LOG(LogDisplayClusterNetDriver, Verbose, TEXT("Client node self-introducing: %u"), ClusterServerConnection->ClientId);
+	}
+
+	if (InEvent.EventId == PacketSyncEvent)
+	{
+		// Deserializing parameters from binary event data
+		const uint8* RawData = InEvent.EventData.GetData();
+
+		// Event data contains packets for all connected clients.
+		// We do iterate on event data and search for own packet by ClientId.
+		// If found - break
+		int ReadOffset = 0;
+		while (ReadOffset < InEvent.EventData.Num())
+		{
+			uint32 CommandClientId = 0;
+			int32 CommandPacketId = 0;
+
+			FMemory::Memcpy(&CommandClientId, RawData + ReadOffset, sizeof(CommandClientId));
+
+			if (CommandClientId == ClusterServerConnection->ClientId)
+			{
+				// Client has received a command to process packet with given ID
+				FMemory::Memcpy(&CommandPacketId, RawData + ReadOffset + sizeof(CommandClientId), sizeof(CommandPacketId));
+
+				UE_LOG(LogDisplayClusterNetDriver, Verbose, TEXT("[ProcessPacket %i] command received"), CommandPacketId);
+
+				ClusterServerConnection->ProcessPacket(CommandPacketId);
+
+				break;
+			}
+
+			ReadOffset += sizeof(CommandClientId) + sizeof(CommandPacketId);
+		}
+	}
+}
+
+void UDisplayClusterNetDriver::TickFlush(float DeltaSeconds)
+{
+	Super::TickFlush(DeltaSeconds);
+
+	// Check if we the server or if any synchronous connections exists for processing
+	if (!IsServer() || SyncConnections.IsEmpty())
+	{
+		// This is a client
+		return;
+	}
+
+	const UDisplayClusterConfigurationData* ConfigData = IDisplayCluster::Get().GetConfigMgr()->GetConfig();
+
+	if ((World->GetNetMode() == NM_ListenServer) && (SyncConnections.Num() < (ConfigData->GetNumberOfClusterNodes() - 1)))
+	{
+		return;
+	}
+
+	bLastBunchWasAcked = false;
+	
+	int32 NumPacketsInQueue = 0;
+
+	// Fill queues of last send packets for each connection, update last acked packet for each connection
+	for (UDisplayClusterNetConnection* SyncConnectionIt : SyncConnections)
+	{
+		int32 ConnectionUniqueId = SyncConnectionIt->GetUniqueID();
+		int32 LastOutPacketId = SyncConnectionIt->LastOut.PacketId;
+
+		// Add last send packet of current Connection to queue
+		TDeque<int32>& ChannelPacketQueue = OutPacketsQueues.FindOrAdd(ConnectionUniqueId);
+		ChannelPacketQueue.PushLast(LastOutPacketId);
+
+		// Get size of queue. Size will be the same for each Connection
+		NumPacketsInQueue = ChannelPacketQueue.Num();
+
+		// Fill LastAckedPacket map, which contains last acked packet for each Connection
+		int32& LastAckedPacketIdRef = LastAckedPacket.FindOrAdd(ConnectionUniqueId, 0);
+		int32 ConnectionAckId = SyncConnectionIt->GetLastNotifiedPacket();
+
+		if (ConnectionAckId > LastAckedPacketIdRef)
+		{
+			LastAckedPacketIdRef = ConnectionAckId;
+		}
+
+		// Driver waits until last sent bunch of any sync connection will be acked
+		bLastBunchWasAcked |= (LastOutPacketId - ConnectionAckId > 1);
+
+		UE_LOG(LogDisplayClusterNetDriver, Verbose, TEXT("Client %u; Packet %i sent, %i acked"), SyncConnectionIt->ClientId, ConnectionAckId, LastAckedPacketIdRef);
+	}
+
+	// Find slice of packets which are ready for processing on clients
+	// Slice corresponds to common denominator of acked packets
+	int32 PacketIdx = INDEX_NONE;
+
+	bool bPacketsReadyForSync = false;
+
+	for (PacketIdx = NumPacketsInQueue - 1; PacketIdx >= 0; PacketIdx--)
+	{
+		if (bPacketsReadyForSync)
+		{
+			break;
+		}
+
+		ReadyOutPackets.Reset();
+
+		for (UDisplayClusterNetConnection* SyncConnectionIt : SyncConnections)
+		{
+			int32 CurrentConnectionId = SyncConnectionIt->GetUniqueID();
+			int32 CurrentPacketId = OutPacketsQueues[CurrentConnectionId][PacketIdx];
+
+			if ((CurrentPacketId > 0) && (LastAckedPacket[CurrentConnectionId] >= CurrentPacketId))
+			{
+				ReadyOutPackets.Add(CurrentConnectionId, CurrentPacketId);
+			}
+
+			if (ReadyOutPackets.Num() == SyncConnections.Num())
+			{
+				bPacketsReadyForSync = true;
+				break;
+			}
+		}
+	}
+
+	// Slice of acked packets is found - form PacketParams and send cluster event to primary nodes
+	if (bPacketsReadyForSync)
+	{
+		PacketsParams.Reset();
+
+		for (UDisplayClusterNetConnection* SyncConnection : SyncConnections)
+		{
+			int32 UniqueID = SyncConnection->GetUniqueID();
+			int32 ReadyForProcessingPacket = ReadyOutPackets[UniqueID];
+
+			// Fill parameters for binary event (key - cliend id, value - packet id)
+			PacketsParams.Add(SyncConnection->ClientId, ReadyForProcessingPacket);
+
+			// Slice queue by removing all packets which are behind ReadyForProcessingPacket
+			TDeque<int32>& ConnectionPacketQueue = OutPacketsQueues[UniqueID];
+			for (int j = 0; j < PacketIdx; j++)
+			{
+				ConnectionPacketQueue.PopFirst();
+			}
+
+			UE_LOG(LogDisplayClusterNetDriver, Verbose, TEXT("Client %u; [ProcessPacket %i] command sent"), SyncConnection->ClientId, ReadyForProcessingPacket);
+		}
+
+		// Generate cluster event
+		FDisplayClusterClusterEventBinary NetworkDriverSyncEvent;
+		GenerateClusterCommandsEvent(NetworkDriverSyncEvent, PacketSyncEvent, PacketsParams);
+
+		// Send cluster event to all primary nodes
+		ClusterNetworkDriverHelper->SendCommandToAllClusters(NetworkDriverSyncEvent);
+	}
+}
+
+#if WITH_SERVER_CODE
+void UDisplayClusterNetDriver::PreListUpdate(ConsiderListUpdateParams const& UpdateParams, int& OutUpdated, const TArray<FNetworkObjectInfo*>& ConsiderList)
+{
+	bConnectionViewersAreReady = false;
+
+	if (bLastBunchWasAcked || NodeConnections.IsEmpty())
+	{
+		return;
+	}
+
+	const float DeltaSeconds = UpdateParams.DeltaSeconds;
+	const bool bCPUSaturated = UpdateParams.bCPUSaturated;
+	UNetConnection* Connection = UpdateParams.Connection;
+
+	ClusterReplicationState = FDisplayClusterReplicationState();
+	PivotNodeConnectionViewers.Reset();
+
+	UDisplayClusterNetConnection* PivotNodeConnection = nullptr;
+
+	// Find first valid primary node connection and choose it as pivot
+	for (UDisplayClusterNetConnection* PrimaryNodeConnection : PrimaryNodeConnections)
+	{
+		if (PrimaryNodeConnection->ViewTarget == nullptr)
+		{
+			continue;
+		}
+
+		PivotNodeConnection = PrimaryNodeConnection;
+		break;
+	}
+	
+	// Primary node connection doesn't exist on listen server - picking the first node connection instead
+	if (World->GetNetMode() == NM_ListenServer)
+	{
+		for (UDisplayClusterNetConnection* ClusterConnection : NodeConnections)
+		{
+			PivotNodeConnection = ClusterConnection;
+			break;
+		}
+	}
+
+	if (PivotNodeConnection == nullptr)
+	{
+		return;
+	}
+
+	AWorldSettings* WorldSettings = World->GetWorldSettings();
+	TArray<FNetViewer>& ConnectionViewers = WorldSettings->ReplicationViewers;
+
+	ConnectionViewers.Reset();
+
+	new(ConnectionViewers)FNetViewer(PivotNodeConnection, DeltaSeconds);
+	for (int32 ViewerIndex = 0; ViewerIndex < PivotNodeConnection->Children.Num(); ViewerIndex++)
+	{
+		if (PivotNodeConnection->Children[ViewerIndex]->ViewTarget != NULL)
+		{
+			new(ConnectionViewers)FNetViewer(PivotNodeConnection->Children[ViewerIndex], DeltaSeconds);
+		}
+	}
+
+	PivotNodeConnectionViewers = ConnectionViewers;
+
+	ClusterReplicationState.FinalSortedCount = ServerReplicateActors_PrioritizeActors(PivotNodeConnection, PivotNodeConnectionViewers,
+		ConsiderList, bCPUSaturated, ClusterReplicationState.PriorityList, ClusterReplicationState.PriorityActors);
+
+	bConnectionViewersAreReady = true;
+}
+
+void UDisplayClusterNetDriver::PostListUpdate(ConsiderListUpdateParams const& UpdateParams, int& OutUpdated, const TArray<FNetworkObjectInfo*>& ConsiderList)
+{
+	// Additional replication for actors number equalization
+	if (!bConnectionViewersAreReady)
+	{
+		return;
+	}
+
+	for (UDisplayClusterNetConnection* DisplayClusterConnection : NodeConnections)
+	{
+		if (!DisplayClusterConnection->bIsClusterConnection || DisplayClusterConnection->ViewTarget == nullptr)
+		{
+			continue;
+		}
+
+		for (int j = 0; j < ClusterReplicationState.FinalSortedCount; j++)
+		{
+			ClusterReplicationState.PriorityList[j].Channel = DisplayClusterConnection->FindActorChannelRef(ClusterReplicationState.PriorityList[j].ActorInfo->WeakActor);
+		}
+
+		const int32* FoundLastProcessedActors = ClusterReplicationState.LastProcessedActors.Find(DisplayClusterConnection->GetUniqueID());
+
+		if (FoundLastProcessedActors == nullptr)
+		{
+			continue;
+		}
+
+		const int32 LastProcessedActors = *FoundLastProcessedActors;
+
+		if (LastProcessedActors < ClusterReplicationState.MaxLastProcessedActor && ClusterReplicationState.MaxLastProcessedActor > 0)
+		{
+			// We do use generated ConnectionViewers above, but than use stored ConnectionViewers from primary nodes?
+			const int32 ActorStartIndex = LastProcessedActors + 1;
+			const int32 ActorsToProcess = ClusterReplicationState.MaxLastProcessedActor - ActorStartIndex + 1;
+
+			constexpr bool bForceUpdate = true;
+			TInterval<int32> ActorsIndexRange(ActorStartIndex, ActorsToProcess);
+			ServerReplicateActors_ProcessPrioritizedActorsRange(DisplayClusterConnection, PivotNodeConnectionViewers, ClusterReplicationState.PriorityActors, ActorsIndexRange,
+				OutUpdated, bForceUpdate);
+		}
+
+		const int32 StartMarkedActor = ClusterReplicationState.MaxLastProcessedActor > 0 ? ClusterReplicationState.MaxLastProcessedActor + 1 : 0;
+
+		ServerReplicateActors_MarkRelevantActors(DisplayClusterConnection, PivotNodeConnectionViewers, StartMarkedActor, ClusterReplicationState.FinalSortedCount, ClusterReplicationState.PriorityActors);
+	}
+}
+
+void UDisplayClusterNetDriver::ListUpdate(ConsiderListUpdateParams const& UpdateParams, int& OutUpdated, const TArray<FNetworkObjectInfo*>& ConsiderList)
+{
+	UDisplayClusterNetConnection* DisplayClusterConnection = Cast<UDisplayClusterNetConnection>(UpdateParams.Connection);
+	check(DisplayClusterConnection);
+
+	const float DeltaSeconds = UpdateParams.DeltaSeconds;
+	const bool bCPUSaturated = UpdateParams.bCPUSaturated;
+	UNetConnection* Connection = UpdateParams.Connection;
+
+	// If not ClusterConnection - fallback to default prioritization
+	if (!DisplayClusterConnection->bIsClusterConnection)
+	{
+		AWorldSettings* WorldSettings = World->GetWorldSettings();
+		TArray<FNetViewer>& ConnectionViewers = WorldSettings->ReplicationViewers;
+
+		FActorPriority* PriorityList = nullptr;
+		FActorPriority** PriorityActors = nullptr;
+
+		// Get a sorted list of actors for this connection
+		const int32 FinalSortedCount = ServerReplicateActors_PrioritizeActors(Connection, ConnectionViewers, ConsiderList, bCPUSaturated, PriorityList, PriorityActors);
+
+		// Process the sorted list of actors for this connection
+		TInterval<int32> ActorsIndexRange(0, FinalSortedCount);
+		const int32 LastProcessedActor = ServerReplicateActors_ProcessPrioritizedActorsRange(Connection, ConnectionViewers, PriorityActors, ActorsIndexRange, OutUpdated);
+
+		ServerReplicateActors_MarkRelevantActors(Connection, ConnectionViewers, LastProcessedActor, FinalSortedCount, PriorityActors);
+	}
+
+	if (DisplayClusterConnection->bIsClusterConnection && bConnectionViewersAreReady)
+	{
+		// Update channels for each actor
+		for (int j = 0; j < ClusterReplicationState.FinalSortedCount; j++)
+		{
+			ClusterReplicationState.PriorityList[j].Channel = Connection->FindActorChannelRef(ClusterReplicationState.PriorityList[j].ActorInfo->WeakActor);
+		}
+
+		// Use single priority list per cluster which was calculated for Primary node
+		// Process the sorted list of actors for this connection
+		TInterval<int32> ActorsIndexRange(0, ClusterReplicationState.FinalSortedCount);
+		const int32 LastProcessedActor = ServerReplicateActors_ProcessPrioritizedActorsRange(Connection, PivotNodeConnectionViewers, ClusterReplicationState.PriorityActors,
+			ActorsIndexRange, OutUpdated);
+
+		// Cache last processed actors for equalization on PostListUpdate step
+		// MaxLastProcessed will be used for equalization
+		ClusterReplicationState.LastProcessedActors.Add(Connection->GetUniqueID(), LastProcessedActor);
+		ClusterReplicationState.MaxLastProcessedActor = FMath::Max<int32>(ClusterReplicationState.MaxLastProcessedActor, LastProcessedActor);
+	}
+}
+#endif
+
+void UDisplayClusterNetDriver::GenerateClusterCommandsEvent(FDisplayClusterClusterEventBinary& NetworkDriverSyncEvent, 
+	int32 EventId)
+{
+	NetworkDriverSyncEvent.EventId = EventId;
+	NetworkDriverSyncEvent.bIsSystemEvent = false;
+}
+
+void UDisplayClusterNetDriver::GenerateClusterCommandsEvent(FDisplayClusterClusterEventBinary& NetworkDriverSyncEvent, 
+	int32 EventId, const TMap<uint32, int32>& Parameters)
+{
+	GenerateClusterCommandsEvent(NetworkDriverSyncEvent, EventId);
+
+	const int KeySize = sizeof(uint32);
+	const int ValueSize = sizeof(int32);
+	const int RecordSize = KeySize + ValueSize;
+
+	ClusterEventData.SetNumUninitialized(Parameters.Num() * RecordSize, true);
+
+	uint8* RawData = ClusterEventData.GetData();
+
+	int32 WriteOffset = 0;
+	for (const TPair<uint32, int32>& Pair : Parameters)
+	{
+		FMemory::Memcpy(RawData + WriteOffset, (uint8*)(&Pair.Key), KeySize);
+		FMemory::Memcpy(RawData + WriteOffset + KeySize, (uint8*)(&Pair.Value), ValueSize);
+
+		WriteOffset += RecordSize;
+	}
+
+	NetworkDriverSyncEvent.EventData = ClusterEventData;
+}
diff --git a/Engine/Plugins/Runtime/nDisplay/Source/DisplayClusterReplication/Private/DisplayClusterNetDriverLog.cpp b/Engine/Plugins/Runtime/nDisplay/Source/DisplayClusterReplication/Private/DisplayClusterNetDriverLog.cpp
new file mode 100644
--- /dev/null
+++ b/Engine/Plugins/Runtime/nDisplay/Source/DisplayClusterReplication/Private/DisplayClusterNetDriverLog.cpp
@@ -0,0 +1,7 @@
+// Copyright Epic Games, Inc. All Rights Reserved.
+
+#include "DisplayClusterNetDriverLog.h"
+
+
+// Module-wide log categories
+DEFINE_LOG_CATEGORY(LogDisplayClusterNetDriver);
diff --git a/Engine/Plugins/Runtime/nDisplay/Source/DisplayClusterReplication/Private/DisplayClusterNetDriverLog.h b/Engine/Plugins/Runtime/nDisplay/Source/DisplayClusterReplication/Private/DisplayClusterNetDriverLog.h
new file mode 100644
--- /dev/null
+++ b/Engine/Plugins/Runtime/nDisplay/Source/DisplayClusterReplication/Private/DisplayClusterNetDriverLog.h
@@ -0,0 +1,9 @@
+// Copyright Epic Games, Inc. All Rights Reserved.
+
+#pragma once
+
+#include "CoreMinimal.h"
+
+
+// Module-wide log categories
+DECLARE_LOG_CATEGORY_EXTERN(LogDisplayClusterNetDriver, Log, All);
diff --git a/Engine/Plugins/Runtime/nDisplay/Source/DisplayClusterReplication/Private/DisplayClusterReplicationModule.cpp b/Engine/Plugins/Runtime/nDisplay/Source/DisplayClusterReplication/Private/DisplayClusterReplicationModule.cpp
new file mode 100644
--- /dev/null
+++ b/Engine/Plugins/Runtime/nDisplay/Source/DisplayClusterReplication/Private/DisplayClusterReplicationModule.cpp
@@ -0,0 +1,3 @@
+#include "Modules/ModuleManager.h"
+
+IMPLEMENT_MODULE(FDefaultModuleImpl, DisplayClusterReplication)
\ No newline at end of file
diff --git a/Engine/Plugins/Runtime/nDisplay/Source/DisplayClusterReplication/Public/DisplayClusterNetConnection.h b/Engine/Plugins/Runtime/nDisplay/Source/DisplayClusterReplication/Public/DisplayClusterNetConnection.h
new file mode 100644
--- /dev/null
+++ b/Engine/Plugins/Runtime/nDisplay/Source/DisplayClusterReplication/Public/DisplayClusterNetConnection.h
@@ -0,0 +1,67 @@
+// Copyright Epic Games, Inc. All Rights Reserved.
+
+#pragma once
+
+#include "Async/TaskGraphInterfaces.h"
+#include "CoreMinimal.h"
+#include "DisplayClusterConfigurationTypes.h"
+#include "Engine/NetConnection.h"
+#include "IpConnection.h"
+#include "Net/NetPacketNotify.h"
+#include "SocketTypes.h"
+#include "UObject/ObjectMacros.h"
+
+#include "DisplayClusterNetConnection.generated.h"
+
+class UDisplayClusterNetDriver;
+
+UCLASS(transient, config = Engine)
+class UDisplayClusterNetConnection : public UIpConnection
+{
+	GENERATED_UCLASS_BODY()
+	
+	//~ Begin UIpConnection Interface
+	virtual void ReceivedPacket(FBitReader& Reader, bool bIsReinjectedPacket = false, bool bDispatchPacket = true) override;
+	virtual void SetClientLoginState(const EClientLoginState::Type NewState) override;
+	//~ End UIpConnection Interface
+
+public:
+	/** The name of the node, parsed from node URL address */
+	FString NodeName;
+
+	/** The ip address of the node, converted to FString from current NetConnection */
+	FString NodeAddress;
+
+	/** Cluster client unique identifier, HashString from NetConnection Challenge */
+	uint32 ClientId;
+
+	/** Cluster unique identifier, HashString from node config file path */
+	uint32 ClusterId;
+
+	/** Cluster node port for binary cluster events */
+	uint16 NodePort;
+
+	/** Whether current connection belongs to primary cluster node */
+	bool bNodeIsPrimary;
+
+	/** Whether current connection belongs to cluster node */
+	bool bIsClusterConnection;
+
+	/** Whether current connection works in synchronous mode */
+	bool bSynchronousMode;
+
+	/** Process accumulated packets in the packet queue until PacketId
+	 *  @param PacketId identified of packet
+	 */
+	void ProcessPacket(int32 PacketId);
+
+	/** Return last notified packet id */
+	int32 GetLastNotifiedPacket() const { return GetLastNotifiedPacketId(); }
+
+protected:
+	// Storage for InPackets map keys, prevents reallocs 
+	TArray<int32> PacketIDs;
+
+	// Data associated with each packet id
+	TSortedMap<int32, FBitReader> InPackets;
+};
diff --git a/Engine/Plugins/Runtime/nDisplay/Source/DisplayClusterReplication/Public/DisplayClusterNetDriver.h b/Engine/Plugins/Runtime/nDisplay/Source/DisplayClusterReplication/Public/DisplayClusterNetDriver.h
new file mode 100644
--- /dev/null
+++ b/Engine/Plugins/Runtime/nDisplay/Source/DisplayClusterReplication/Public/DisplayClusterNetDriver.h
@@ -0,0 +1,138 @@
+// Copyright Epic Games, Inc. All Rights Reserved.
+
+#pragma once
+
+#include "Delegates/DelegateCombinations.h"
+
+#include "Cluster/DisplayClusterClusterEvent.h"
+#include "Cluster/DisplayClusterNetDriverHelper.h"
+#include "Cluster/IDisplayClusterClusterManager.h"
+
+#include "Containers/Deque.h"
+#include "CoreMinimal.h"
+
+#include "DisplayClusterNetConnection.h"
+
+#include "GameFramework/WorldSettings.h"
+#include "IDisplayCluster.h"
+
+#include "IpNetDriver.h"
+#include "UObject/ObjectMacros.h"
+
+#include "DisplayClusterNetDriver.generated.h"
+
+/**
+ * Helper structure to cache PriorityActor list for additional replication
+ * in order to equalize number of replicated actors among all synced connections
+ */
+struct FDisplayClusterReplicationState
+{
+	// Number of actors that were prioritized
+	int32 FinalSortedCount = 0;
+	
+	// Max processed actor among all synced connections
+	int32 MaxLastProcessedActor = -1;
+
+	// Cached priority lists
+	FActorPriority* PriorityList = nullptr;
+	FActorPriority** PriorityActors = nullptr;
+
+	// Last processed actor for each sync connection
+	TMap<uint32, int32> LastProcessedActors;
+};
+
+/**
+ * Custom NetDriver for DisplayCluster.
+ * Responsible for synchronous application of network packets across display cluster nodes.
+ * Results in seamless image when used with Actor replication system.
+ */
+UCLASS(transient, config = Engine)
+class UDisplayClusterNetDriver : public UIpNetDriver
+{
+	GENERATED_UCLASS_BODY()
+
+	virtual ~UDisplayClusterNetDriver();
+
+	//~ Begin UIpNetDriver Interface
+	virtual bool InitListen(FNetworkNotify* InNotify, FURL& ListenURL, bool bReuseAddressAndPort, FString& Error) override;
+	virtual void TickDispatch(float DeltaTime) override;
+	virtual void TickFlush(float DeltaSeconds) override;
+	//~ End UIpNetDriver Interface
+public:
+	/** Performs Cluster Event handling: sync mode start and packet processing on client side
+	 *
+	 * @param Binary event to handle
+	 */
+	void HandleEvent(FDisplayClusterClusterEventBinary const& InEvent);
+
+	/** Adds connection to nDisplay connection TSet 
+	 *
+	 * @param NetConnection connection to add
+	 */
+	void AddNodeConnection(UDisplayClusterNetConnection* NetConnection);
+
+	/** Removes connection from nDisplay connection TSet and from synchronized connections 
+	 *
+	 * @param NetConnection connection to remove
+	 */
+	void RemoveNodeConnection(UDisplayClusterNetConnection* NetConnection);
+
+protected:
+	// Contains last max acked packet for specific connection (key - connectionID, value - packetID)
+	TMap<int32, int32> LastAckedPacket;
+
+	// Contains queue of non-processed packets for specific connection (key - connectionID, value - packetID)
+	TMap<int32, TDeque<int32>> OutPacketsQueues;
+
+	// Holds packets that are ready for replication (key - connectionID, value - packetID)
+	TMap<int32, int32> ReadyOutPackets;
+
+	// Holds ready packets for each client. Member of the class to prevent runtime reallocations.
+	TMap<uint32, int32> PacketsParams;
+
+	// Binary event data. Member of the class to prevent runtime reallocations
+	TArray<uint8> ClusterEventData;
+
+	// nDisplay node connections
+	TSet<UDisplayClusterNetConnection*> NodeConnections;
+
+	// nDisplay Primary node connections
+	TSet<UDisplayClusterNetConnection*> PrimaryNodeConnections;
+
+	// nDisplay node connections that participate in synchronous replication
+	TSet<UDisplayClusterNetConnection*> SyncConnections;
+
+	// ConnectionViewers for Pivot node (primary node connection)
+	TArray<FNetViewer> PivotNodeConnectionViewers;
+
+	// NetDriver helper for Display Cluster Events
+	TUniquePtr<FDisplayClusterNetDriverHelper> ClusterNetworkDriverHelper;
+
+	// Stores replication state to produce additional replication in order to equalize number of actors for each sync connection
+	FDisplayClusterReplicationState ClusterReplicationState;
+
+	// Binary event listener
+	FOnClusterEventBinaryListener EventBinaryListener;
+
+	// Used to control state of acked packets queue
+	bool bLastBunchWasAcked;
+
+	// Checks if connections viewers for Primary node were formed
+	bool bConnectionViewersAreReady;
+
+	// Cluster event id, used to start synchrosonus packets processing
+	inline static const int NodeSyncEvent = GetTypeHash(FStringView(TEXT("nDCRNodeSyncEvent")));
+
+	// Cluster event id, used to identify packet data blob for sync
+	inline static const int PacketSyncEvent = GetTypeHash(FStringView(TEXT("nDCRPacketSyncEvent")));
+
+protected:
+#if WITH_SERVER_CODE
+	void PreListUpdate(ConsiderListUpdateParams const& UpdateParams, int& OutUpdated, const TArray<FNetworkObjectInfo*>& ConsiderList);
+	void PostListUpdate(ConsiderListUpdateParams const& UpdateParams, int& OutUpdated, const TArray<FNetworkObjectInfo*>& ConsiderList);
+	void ListUpdate(ConsiderListUpdateParams const& UpdateParams, int& OutUpdated, const TArray<FNetworkObjectInfo*>& ConsiderList);
+#endif
+	// Helper functions for parameters serialization into binary event data
+	void GenerateClusterCommandsEvent(FDisplayClusterClusterEventBinary& NetworkDriverSyncEvent, int32 EventId);
+	void GenerateClusterCommandsEvent(FDisplayClusterClusterEventBinary& NetworkDriverSyncEvent, int32 EventId, const TMap<uint32, int32>& Parameters);
+};
diff --git a/Engine/Plugins/Runtime/nDisplay/nDisplay.uplugin b/Engine/Plugins/Runtime/nDisplay/nDisplay.uplugin
--- a/Engine/Plugins/Runtime/nDisplay/nDisplay.uplugin
+++ b/Engine/Plugins/Runtime/nDisplay/nDisplay.uplugin
@@ -213,6 +213,15 @@
 				"Win64",
 				"Linux"
 			]
+		},
+		{
+			"Name": "DisplayClusterReplication",
+			"Type": "Runtime",
+			"LoadingPhase": "Default",
+			"PlatformAllowList": [
+				"Win64",
+				"Linux"
+			]
 		}
 	],
 	"Plugins": [
@@ -259,6 +268,14 @@
 		{
 			"Name": "MovieRenderPipeline",
 			"Enabled": true
+		},
+		{
+			"Name": "OnlineSubsystem",
+			"Enabled": true
+		},
+		{
+			"Name": "OnlineSubsystemUtils",
+			"Enabled": true
 		}
 	]
 }
diff --git a/Engine/Plugins/VirtualProduction/Switchboard/Source/Switchboard/switchboard/devices/mpserver/plugin_mpserver.py b/Engine/Plugins/VirtualProduction/Switchboard/Source/Switchboard/switchboard/devices/mpserver/plugin_mpserver.py
new file mode 100644
--- /dev/null
+++ b/Engine/Plugins/VirtualProduction/Switchboard/Source/Switchboard/switchboard/devices/mpserver/plugin_mpserver.py
@@ -0,0 +1,72 @@
+# Copyright Epic Games, Inc. All Rights Reserved.
+
+from switchboard.devices.unreal.plugin_unreal import DeviceUnreal, DeviceWidgetUnreal
+from switchboard.config import (CONFIG, FilePathSetting, StringSetting, IntSetting)
+from switchboard.switchboard_utils import get_game_launch_level_path
+
+
+class DeviceMultiplayerServer(DeviceUnreal):
+    settings = {
+        'executable_path': FilePathSetting(
+            attr_name="executable_path",
+            nice_name="Executable Path",
+            value=""
+        ),
+        'level_path': StringSetting(
+            attr_name="level_path",
+            nice_name="Level Path",
+            value=""
+        ),
+        'multiplayer_port': IntSetting(
+            attr_name="multiplayer_port",
+            nice_name="Multiplayer Port",
+            value=7777,
+            tool_tip="Port to use for multiplayer connections"
+        )
+    }
+    def __init__(self, name, address, **kwargs):
+
+        super().__init__(name, address, **kwargs)
+    
+    def generate_unreal_exe_path(self):
+        exe_path = self.settings['executable_path'].get_value()
+        if exe_path == '':
+            return super().generate_unreal_exe_path()
+        
+        return exe_path
+
+    def generate_unreal_command_line_args(self, map_name):
+        command_line_args = ''
+        uproject_path = CONFIG.UPROJECT_PATH.get_value()
+        
+        custom_level_path = self.settings['level_path'].get_value().strip()
+        level_path = custom_level_path if custom_level_path != '' else get_game_launch_level_path()
+            
+        if self.settings['executable_path'].get_value() == '':
+            command_line_args = f'{uproject_path} {level_path} -server'
+        else:
+            command_line_args = level_path
+
+        command_line_args += f' -log -port={self.settings["multiplayer_port"].get_value()}'
+
+        return command_line_args
+
+    @classmethod
+    def plugin_settings(cls):
+        return list(cls.settings.values()) + [DeviceUnreal.csettings['port']]
+
+    def device_settings(self):
+        return [self.setting_address]
+
+    def setting_overrides(self):
+        return []
+
+
+class DeviceWidgetMultiplayerServer(DeviceWidgetUnreal):
+    def __init__(self, name, device_type, device_hash, address, parent=None):
+
+        super().__init__(name, device_type, device_hash, address, parent=parent)
+        
+        self.autojoin_mu._button.setVisible(False)
+        self.build_button.setDisabled(True)
+        self.build_button.setVisible(False)
diff --git a/Engine/Plugins/VirtualProduction/Switchboard/Source/Switchboard/switchboard/devices/ndisplay/plugin_ndisplay.py b/Engine/Plugins/VirtualProduction/Switchboard/Source/Switchboard/switchboard/devices/ndisplay/plugin_ndisplay.py
--- a/Engine/Plugins/VirtualProduction/Switchboard/Source/Switchboard/switchboard/devices/ndisplay/plugin_ndisplay.py
+++ b/Engine/Plugins/VirtualProduction/Switchboard/Source/Switchboard/switchboard/devices/ndisplay/plugin_ndisplay.py
@@ -19,7 +19,7 @@ from switchboard import switchboard_widgets as sb_widgets
 from switchboard import switchboard_dialog as sb_dialog
 from switchboard.config import CONFIG, BoolSetting, IntSetting, FilePathSetting, \
     LoggingSetting, OptionSetting, Setting, StringSetting, SETTINGS, \
-    StringListSetting, migrate_comma_separated_string_to_list
+    StringListSetting, AddressSetting, migrate_comma_separated_string_to_list
 from switchboard.devices.device_widget_base import AddDeviceDialog
 from switchboard.devices.unreal.plugin_unreal import DeviceUnreal, \
     DeviceWidgetUnreal, LiveLinkPresetSetting, MediaProfileSetting
@@ -28,6 +28,7 @@ from switchboard.switchboard_logging import LOGGER
 
 from .ndisplay_monitor_ui import nDisplayMonitorUI
 from .ndisplay_monitor import nDisplayMonitor
+from ...switchboard_utils import map_name_is_valid, get_game_launch_level_path
 
 
 class AddnDisplayDialog(AddDeviceDialog):
@@ -321,6 +322,25 @@ class DevicenDisplay(DeviceUnreal):
                 "vulkan"
             ],
         ),
+        'multiplayer_mode': OptionSetting(
+            attr_name="multiplayer_mode",
+            nice_name="Multiplayer Server Mode",
+            value='None',
+            possible_values=['None', 'Listen server', 'Dedicated server']
+        ),
+        'dedicated_server_address': AddressSetting(
+            attr_name="dedicated_server_address",
+            nice_name="Dedicated Server Address",
+            value='127.0.0.1',
+            tool_tip='Server address to connect to. Not used if Multiplayer Server Mode '
+                     'is set to "Listen Server" or auto start of dedicated server is enabled'
+        ),
+        'dedicated_server_port': IntSetting(
+            attr_name="dedicated_server_port",
+            nice_name="Dedicated Server Port",
+            value='7777',
+            tool_tip='Server port to connect to. Not used if Multiplayer Server Mode is not Dedicated Server'
+        ),
         'render_mode': OptionSetting(
             attr_name="render_mode",
             nice_name="Render Mode",
@@ -562,8 +582,6 @@ class DevicenDisplay(DeviceUnreal):
         self.path_to_config_on_host = DevicenDisplay.csettings[
             'ndisplay_config_file'].get_value()
 
-        if len(self.settings['ue_command_line'].get_value()) == 0:
-            self.generate_unreal_command_line()
 
         CONFIG.ENGINE_DIR.signal_setting_overridden.connect(
             self.on_cmdline_affecting_override)
@@ -592,6 +610,9 @@ class DevicenDisplay(DeviceUnreal):
         # node configuration (updated from config file)
         self.nodeconfig = {}
 
+        if len(self.settings['ue_command_line'].get_value()) == 0:
+            self.generate_unreal_command_line()
+
         try:
             cfg_file = DevicenDisplay.csettings[
                 'ndisplay_config_file'].get_value(self.name)
@@ -678,8 +699,74 @@ class DevicenDisplay(DeviceUnreal):
         return DevicenDisplay.csettings[
             'udpmessaging_extra_static_endpoints'].get_value(self.name)
 
+    def generate_multiplayer_map_name_args(self, map_name):
+
+        primary_device = self.get_primary_node_device()
+        primary_device_name = 'Node_0' if primary_device is None else primary_device.name
+        primary_device_address = '127.0.0.1' if primary_device is None else primary_device.address
+
+        primary_node_binary_event_port = 41004
+
+        # set primary node binary event port
+        if self.nodeconfig and 'port_ce_binary' in self.nodeconfig:
+            primary_node_binary_event_port = self.nodeconfig['port_ce_binary']
+
+        def get_common_args():
+            """ 
+            get common args for client and server 
+            """
+            return '?'.join([
+                    f'ClusterId={os.getpid()}',
+                    f'PrimaryNodeId={primary_device_name}',
+                    f'PrimaryNodePort={primary_node_binary_event_port}'
+                ])
+
+        def get_client_args(server_address):
+            return '?'.join([
+                    server_address,
+                    f'node={self.name}',
+                    get_common_args()
+                ])
+
+        multiplayer_mode_name = DevicenDisplay.csettings["multiplayer_mode"].get_value()
+
+        # single player
+        if multiplayer_mode_name == 'None':
+            return map_name
+        
+        def is_local_address(address):
+            return address == SETTINGS.ADDRESS.get_value() or address == '127.0.0.1'
+
+        # dedicated server
+        if multiplayer_mode_name == 'Dedicated server':
+            dedicated_server_address = DevicenDisplay.csettings["dedicated_server_address"].get_value()
+
+            if dedicated_server_address == self.address or is_local_address(self.address):
+                dedicated_server_address = '127.0.0.1'
+
+            return get_client_args(
+                f'{dedicated_server_address}:{DevicenDisplay.csettings["dedicated_server_port"].get_value()}')
+
+        # listen server
+        if not self.is_primary():
+            return get_client_args(primary_device_address)
+
+        if map_name_is_valid(map_name):
+            return '?'.join([map_name,'Listen', get_common_args()])
+         
+        current_map_name = get_game_launch_level_path()
+        if current_map_name.strip() == '':
+            return map_name
+
+        return '?'.join([current_map_name, 'Listen', get_common_args()])
+
     def generate_unreal_command_line(self, map_name=""):
-        uproject = os.path.normpath(CONFIG.UPROJECT_PATH.get_value(self.name))
+        uproject = CONFIG.UPROJECT_PATH.get_value(self.name)
+        
+        # normalize path if not empty
+        if uproject != "":
+            uproject = os.path.normpath(uproject)
+            uproject = f'"{uproject}"'
 
         # Extra arguments specified in settings
         additional_args = self.extra_cmdline_args_setting
@@ -799,9 +886,12 @@ class DevicenDisplay(DeviceUnreal):
             if DevicenDisplay.csettings['disable_all_screen_messages'].get_value()
             else '')
 
+        # Modify map name arg for multiplayer mode
+        map_name = self.generate_multiplayer_map_name_args(map_name)
+
         # fill in fixed arguments
         args = [
-            f'"{uproject}"',
+            f'{uproject}',
             "-game",                      # render nodes run in -game
             f'{map_name}',                # map to open
             "-messaging",                 # enables messaging, needed for MultiUser
@@ -835,6 +925,13 @@ class DevicenDisplay(DeviceUnreal):
             f'{udpm_transport_static}',   # -UDPMESSAGING_TRANSPORT_STATIC=
         ]
 
+        # Specify multiplayer with listen server mode by arg for clients if enabled
+        if DevicenDisplay.csettings["multiplayer_mode"].get_value() == 'Listen server'and not self.is_primary():
+            args.extend([
+                '-dc_replicationserver_type listen'
+            ])
+
+
         # fill in ExecCmds
         exec_cmds = self.csettings["ndisplay_exec_cmds"].get_value(self.name).copy()
 
@@ -1031,10 +1128,16 @@ class DevicenDisplay(DeviceUnreal):
 
         # Update the state of the button depending on whether this device is
         # the primary or not.
-        is_primary = (
-            self.name ==
-            DevicenDisplay.csettings['primary_device_name'].get_value())
-        device_widget.primary_button.setChecked(is_primary)
+        device_widget.primary_button.setChecked(self.is_primary())
+
+    def get_primary_node_device(self):
+        for device in self.active_unreal_devices:
+            if device.device_type == 'nDisplay' and device.is_primary():
+                return device
+
+    def is_primary(self):
+        '''Check if the device is the primary or not'''
+        return self.name == DevicenDisplay.csettings['primary_device_name'].get_value()
 
     def select_as_primary(self):
         ''' Selects this node as the primary node in the nDisplay cluster '''
@@ -1207,6 +1310,7 @@ class DevicenDisplay(DeviceUnreal):
                 "address": cnode['host'],
                 "primary": primary,
                 "port_ce": int(primaryNode['ports']['ClusterEventsJson']),
+                "port_ce_binary": int(primaryNode['ports']['ClusterEventsBinary']),
                 "kwargs": kwargs,
             })
 
diff --git a/Engine/Plugins/VirtualProduction/Switchboard/Source/Switchboard/switchboard/switchboard_utils.py b/Engine/Plugins/VirtualProduction/Switchboard/Source/Switchboard/switchboard/switchboard_utils.py
--- a/Engine/Plugins/VirtualProduction/Switchboard/Source/Switchboard/switchboard/switchboard_utils.py
+++ b/Engine/Plugins/VirtualProduction/Switchboard/Source/Switchboard/switchboard/switchboard_utils.py
@@ -9,6 +9,12 @@ import subprocess
 import threading
 from typing import Optional
 
+from pathlib import Path
+from .config import CONFIG, DEFAULT_MAP_TEXT
+from .switchboard_logging import LOGGER
+from re import split
+
+from switchboard import ugs_utils
 
 class PriorityModifier(Enum):
     ''' Corresponds to `int32 PriorityModifier` parameter to `FPlatformProcess::CreateProc`. '''
@@ -193,3 +199,42 @@ def expand_endpoint(
         port_str = str(default_port)
 
     return f'{addr_str}:{port_str}'
+
+
+def map_name_is_valid(map_name: str) -> bool:
+    """
+    Check if map is specified
+    """
+    return map_name != '' and map_name is not None and map_name != DEFAULT_MAP_TEXT \
+        and map_name != 'Default level' and map_name != 'None'
+
+
+def get_game_launch_level_path() -> str:
+    level = CONFIG.CURRENT_LEVEL
+
+    #check if map name is valid. If not - get map name from DefaultEngine.ini
+    if not map_name_is_valid(level):
+        level = ''
+
+        #check if DefaultEngine.ini exists
+        engine_ini = os.path.join(Path(CONFIG.UPROJECT_PATH.get_value()).parent, 'Config', 'DefaultEngine.ini')
+        if not os.path.exists(engine_ini):
+            LOGGER.warning("DefaultEngine.ini not found")
+            return ''
+
+        parser = ugs_utils.IniParser()
+
+        try:
+            engine_ini_file = open(engine_ini, "r", encoding="utf-8")
+            parser.read_file(engine_ini_file)
+
+            levels = parser.try_get('/Script/EngineSettings.GameMapsSettings', 'GameDefaultMap')
+
+            if len(levels) == 0:
+                raise Exception('Default map is not set in DefaultEngine.ini')
+            
+            level = levels[0]
+        except:
+            LOGGER.warning("Map is not set. Please select 'Level' to run dedicated server.")
+
+    return str(level)
\ No newline at end of file
diff --git a/Engine/Source/Runtime/Engine/Classes/Engine/Engine.h b/Engine/Source/Runtime/Engine/Classes/Engine/Engine.h
--- a/Engine/Source/Runtime/Engine/Classes/Engine/Engine.h
+++ b/Engine/Source/Runtime/Engine/Classes/Engine/Engine.h
@@ -693,6 +693,9 @@ DECLARE_MULTICAST_DELEGATE_OneParam(FPreRenderDelegateEx, class FRDGBuilder&);
 DECLARE_MULTICAST_DELEGATE(FPostRenderDelegate);
 DECLARE_MULTICAST_DELEGATE_OneParam(FPostRenderDelegateEx, class FRDGBuilder&);
 
+DECLARE_DELEGATE_RetVal_ThreeParams(EBrowseReturnVal::Type, FBrowseURL, FWorldContext& WorldContext, FURL URL, FString& Error);
+DECLARE_DELEGATE_TwoParams(FPendingLevelUpdate, FWorldContext& Context, float DeltaSeconds);
+
 /**
  * Type of UObject purge type to be performed by the engine
  */
@@ -3438,6 +3441,19 @@ private:
 	bool bIsVanillaProduct;
 
 protected:
+	/**
+	 * Delegate for overriding the method Browse in the part that parses an URL
+	 * and loads specified level or creates PendingNetGame.
+	 * Parameter are the same as those passed to the calling method Browse
+	 */
+	FBrowseURL OnOverrideBrowseURL;
+
+	/**
+	 * Delegate for overriding the method TickWorldTravel in the part that
+	 * controls the state of PendingNetGame
+	 * Parameter are the same as those passed to the calling method Browse
+	 */
+	FPendingLevelUpdate OnOverridePendingNetGameUpdate;
 
 	TIndirectArray<FWorldContext>	WorldList;
 
diff --git a/Engine/Source/Runtime/Engine/Classes/Engine/NetConnection.h b/Engine/Source/Runtime/Engine/Classes/Engine/NetConnection.h
--- a/Engine/Source/Runtime/Engine/Classes/Engine/NetConnection.h
+++ b/Engine/Source/Runtime/Engine/Classes/Engine/NetConnection.h
@@ -1200,8 +1200,16 @@ public:
 	/** 
 	* Handle a packet we just received. 
 	* bIsReinjectedPacket is true if a packet is reprocessed after getting cached 
+	* bDispatchPacket if true the packet will be processed (passed to DispatchPacket)
 	*/
-	void ReceivedPacket( FBitReader& Reader, bool bIsReinjectedPacket=false );
+	ENGINE_API virtual void ReceivedPacket( FBitReader& Reader, bool bIsReinjectedPacket=false, bool bDispatchPacket=true );
+	
+	/**
+	* Disassemble and dispatch packet.
+	* PacketId is an id of the packet to dispatch
+	* bOutSkipAck Flag that marks packet as dropped
+	*/
+	ENGINE_API virtual void DispatchPacket( FBitReader& Reader, int32 PacketId, bool& bOutSkipAck, bool& bOutHasBunchErrors );
 
 	/** Packet was negatively acknowledged. */
 	void ReceivedNak( int32 NakPacketId );
@@ -1243,7 +1251,7 @@ public:
 	 * Wrapper for setting the current client login state, so we can trap for debugging, and verbosity purposes. 
 	 * Only valid on the server
 	 */
-	ENGINE_API void SetClientLoginState( const EClientLoginState::Type NewState );
+	ENGINE_API virtual void SetClientLoginState( const EClientLoginState::Type NewState );
 
 	/** 
 	 * Wrapper for setting the current expected client login msg type. 
@@ -1349,7 +1357,7 @@ public:
 	void CleanupDormantReplicatorsForActor(AActor* Actor);
 
 	/** Trigger a callback on all dormant replicators of every dormant actor we stored */
-	void ExecuteOnAllDormantReplicators(UE::Net::FExecuteForEachDormantReplicator ExecuteFunction);
+	ENGINE_API void ExecuteOnAllDormantReplicators(UE::Net::FExecuteForEachDormantReplicator ExecuteFunction);
 
 	/** Trigger a callback on all dormant replicators owned by a dormant actor */
 	void ExecuteOnAllDormantReplicatorsOfActor(AActor* OwnerActor, UE::Net::FExecuteForEachDormantReplicator ExecuteFunction);
@@ -1525,6 +1533,11 @@ protected:
 	 */
 	ENGINE_API void NotifyConnectionUpdated();
 
+	/**
+	 * Return last notified packet id
+	 */
+	ENGINE_API int32 GetLastNotifiedPacketId() const;
+
 private:
 	/**
 	 * The channels that need ticking. This will be a subset of OpenChannels, only including
diff --git a/Engine/Source/Runtime/Engine/Classes/Engine/NetDriver.h b/Engine/Source/Runtime/Engine/Classes/Engine/NetDriver.h
--- a/Engine/Source/Runtime/Engine/Classes/Engine/NetDriver.h
+++ b/Engine/Source/Runtime/Engine/Classes/Engine/NetDriver.h
@@ -413,6 +413,22 @@ DECLARE_DELEGATE_RetVal(bool, FShouldSkipRepNotifies);
 
 #endif
 
+/**
+ * The structure to pass to the OnConsiderListUpdate delegate 
+ * 
+ * @param DeltaSeconds     Time between the frames
+ * @param Connection       NetConnection to process
+ * @param bCPUSaturated    Not used by the engine at the moment but kept for compatibility
+ */
+struct ENGINE_API ConsiderListUpdateParams
+{
+	float DeltaSeconds = 0;
+	UNetConnection* Connection = nullptr;
+	bool bCPUSaturated = false;
+};
+
+DECLARE_DELEGATE_ThreeParams(FOnConsiderListUpdate, const ConsiderListUpdateParams& UpdateParams, int32& OutUpdated, const TArray<FNetworkObjectInfo*>& ConsiderList);
+
 //
 // Whether to support net lag and packet loss testing.
 //
@@ -626,7 +642,7 @@ struct FNetDriverReplicationSystemConfig
 //
 // Priority sortable list.
 //
-struct FActorPriority
+struct ENGINE_API FActorPriority
 {
 	int32						Priority;	// Update priority, higher = more important.
 	
@@ -1898,8 +1914,38 @@ protected:
 	*/
 	int32 ServerReplicateActors_PrepConnections( const float DeltaSeconds );
 	void ServerReplicateActors_BuildConsiderList( TArray<FNetworkObjectInfo*>& OutConsiderList, const float ServerTickTime );
-	int32 ServerReplicateActors_PrioritizeActors( UNetConnection* Connection, const TArray<FNetViewer>& ConnectionViewers, const TArray<FNetworkObjectInfo*>& ConsiderList, const bool bCPUSaturated, FActorPriority*& OutPriorityList, FActorPriority**& OutPriorityActors );
+
+	// Actor prioritization
+	ENGINE_API int32 ServerReplicateActors_PrioritizeActors( UNetConnection* Connection, const TArray<FNetViewer>& ConnectionViewers, const TArray<FNetworkObjectInfo*>& ConsiderList, const bool bCPUSaturated, FActorPriority*& OutPriorityList, FActorPriority**& OutPriorityActors );
+	
+	UE_DEPRECATED(5.3, "This function has been deprecated. Please use ServerReplicateActors_ProcessPrioritizedActorsRange instead")
 	int32 ServerReplicateActors_ProcessPrioritizedActors( UNetConnection* Connection, const TArray<FNetViewer>& ConnectionViewers, FActorPriority** PriorityActors, const int32 FinalSortedCount, int32& OutUpdated );
+	
+	// Actor relevancy processing within specified range
+	ENGINE_API int32 ServerReplicateActors_ProcessPrioritizedActorsRange( UNetConnection* Connection, const TArray<FNetViewer>& ConnectionViewers, FActorPriority** PriorityActors, const TInterval<int32>& ActorsIndexRange, int32& OutUpdated, bool bIgnoreSaturation = false );
+	
+	// Relevant actors that could not be processed this frame are marked to be considered for next frame
+	ENGINE_API void ServerReplicateActors_MarkRelevantActors( UNetConnection* Connection, const TArray<FNetViewer>& ConnectionViewers, int32 StartActorIndex, int32 EndActorIndex, FActorPriority** PriorityActors );
+	
+	/**
+	* Delegate for overriding the method ServerReplicateActors
+	* in the part that prepares prioritized actors list of
+	* client connections
+	*/
+	FOnConsiderListUpdate OnPreConsiderListUpdateOverride;
+
+	/**
+	* Delegate that complements the method ServerReplicateActors
+	* with the additional replication logic
+	*/
+	FOnConsiderListUpdate OnPostConsiderListUpdateOverride;
+
+	/**
+	* Delegate that allows to implement additional procedures after
+	* main replication logic in the corresponding part of the method
+	* ServerReplicateActors
+	*/
+	FOnConsiderListUpdate OnProcessConsiderListOverride;
 #endif
 
 	/** Used to handle any NetDriver specific cleanup once a level has been removed from the world. */
diff --git a/Engine/Source/Runtime/Engine/Classes/Engine/PendingNetGame.h b/Engine/Source/Runtime/Engine/Classes/Engine/PendingNetGame.h
--- a/Engine/Source/Runtime/Engine/Classes/Engine/PendingNetGame.h
+++ b/Engine/Source/Runtime/Engine/Classes/Engine/PendingNetGame.h
@@ -16,7 +16,7 @@ class UNetDriver;
 struct FWorldContext;
 
 UCLASS(customConstructor, transient)
-class UPendingNetGame :
+class ENGINE_API UPendingNetGame :
 	public UObject,
 	public FNetworkNotify
 {
@@ -61,7 +61,7 @@ public:
 	 *
 	 * @param Response response from the game containing its encryption key or an error message
 	 */
-	ENGINE_API void SetEncryptionKey(const FEncryptionKeyResponse& Response);
+	void SetEncryptionKey(const FEncryptionKeyResponse& Response);
 
 	bool HasFailedTravel() const {return bFailedTravel; }
 	void SetFailedTravel(bool bInFailedTravel) { bFailedTravel = bInFailedTravel; }
@@ -96,7 +96,7 @@ public:
 	/**
 	 * Send the packet for triggering the initial join
 	 */
-	ENGINE_API void SendInitialJoin();
+	void SendInitialJoin();
 
 	//~ Begin FNetworkNotify Interface.
 	virtual EAcceptConnection::Type NotifyAcceptingConnection() override;
diff --git a/Engine/Source/Runtime/Engine/Private/NetConnection.cpp b/Engine/Source/Runtime/Engine/Private/NetConnection.cpp
--- a/Engine/Source/Runtime/Engine/Private/NetConnection.cpp
+++ b/Engine/Source/Runtime/Engine/Private/NetConnection.cpp
@@ -745,6 +745,11 @@ void UNetConnection::NotifyConnectionUpdated()
 	}
 }
 
+int32 UNetConnection::GetLastNotifiedPacketId() const
+{
+	return LastNotifiedPacketId;
+}
+
 void UNetConnection::EnableEncryption(const FEncryptionData& EncryptionData)
 {
 	if (Handler.IsValid())
@@ -2696,7 +2701,7 @@ FNetworkGUID UNetConnection::GetActorGUIDFromOpenBunch(FInBunch& Bunch)
 	return ActorGUID;
 }
 
-void UNetConnection::ReceivedPacket( FBitReader& Reader, bool bIsReinjectedPacket)
+void UNetConnection::ReceivedPacket( FBitReader& Reader, bool bIsReinjectedPacket, bool bDispatchPacket )
 {
 	using namespace UE::Net;
 
@@ -2966,9 +2971,6 @@ void UNetConnection::ReceivedPacket( FBitReader& Reader, bool bIsReinjectedPacke
 
 	bool bHasBunchErrors = false;
 
-	// Track channels that were rejected while processing this packet - used to avoid sending multiple close-channel bunches,
-	// which would cause a disconnect serverside
-	TArray<int32> RejectedChans;
 	const bool bIsServer = Driver->ServerConnection == nullptr;
 	double PostReceiveTime = 0.0;
 
@@ -2988,673 +2990,690 @@ void UNetConnection::ReceivedPacket( FBitReader& Reader, bool bIsReinjectedPacke
 			}
 		};
 
-		// Disassemble and dispatch all bunches in the packet.
-		while( !Reader.AtEnd() && GetConnectionState()!=USOCK_Closed )
+		if (bDispatchPacket)
+		{
+			DispatchPacket(Reader, InPacketId, bSkipAck, bHasBunchErrors);
+		}
+
+		// Close/clean-up channels pending close due to received acks.
+		for (FChannelCloseInfo& Info : ChannelsToClose)
 		{
-			// For demo backwards compatibility, old replays still have this bit
-			if (IsInternalAck() && PacketEngineNetVer < FEngineNetworkCustomVersion::AcksIncludedInHeader)
+			if (UChannel* Channel = Channels[Info.Id])
 			{
-				const bool IsAckDummy = Reader.ReadBit() == 1u;
+				Channel->ConditionalCleanUp(false, Info.CloseReason);
 			}
+		}
 
-			// Parse the bunch.
-			int32 StartPos = Reader.GetPosBits();
-		
-			// Process Received data
+		ValidateSendBuffer();
+	}
+
+	// Acknowledge the packet.
+	if ( !bSkipAck )
+	{
+		LastGoodPacketRealtime = PostReceiveTime;
+	}
+
+	if( !IsInternalAck() )
+	{
+		// We always call AckSequence even if we are explicitly rejecting the packet as this updates the expected InSeq used to drive future acks.
+		if ( bSkipAck )
+		{
+			// Explicit Nak, we treat this packet as dropped but we still report it to the sending side as quickly as possible
+			PacketNotify.NakSeq( InPacketId );
+		}
+		else
+		{
+			PacketNotify.AckSeq( InPacketId );
+
+			// Keep stats happy
+			++OutTotalAcks;
+			++Driver->OutTotalAcks;
+		}
+
+		// We do want to let the other side know about the ack, so even if there are no other outgoing data when we tick the connection we will send an ackpacket.
+		TimeSensitive = 1;
+		++HasDirtyAcks;
+
+		// This is to allow us to recover from hitches were we process more than FNetPacketNotify::SequenceHistoryLength packets in a row withouht sending out any packets.
+		// In most cases this allows us to recover smoothly without reporting any pessimistic naks which will be the case if we overshoot the ack history.
+		// Note: This should only occur if we are running with no timeouts!
+		if (HasDirtyAcks >= FNetPacketNotify::MaxSequenceHistoryLength)
+		{
+			UE_LOG(LogNet, Warning, TEXT("UNetConnection::ReceivedPacket - Too many received packets to ack (%u) since last sent packet. InSeq: %u %s NextOutGoingSeq: %u"), HasDirtyAcks, PacketNotify.GetInSeq().Get(), *Describe(), PacketNotify.GetOutSeq().Get());
+
+			FlushNet();
+			if (HasDirtyAcks) // if acks still are dirty, flush again
 			{
-				// Parse the incoming data.
-				FInBunch Bunch( this );
-				int32 IncomingStartPos		= Reader.GetPosBits();
-				uint8 bControl				= Reader.ReadBit();
-				Bunch.PacketId				= InPacketId;
-				Bunch.bOpen					= bControl ? Reader.ReadBit() : 0;
-				Bunch.bClose				= bControl ? Reader.ReadBit() : 0;
-			
-				if (PacketEngineNetVer < FEngineNetworkCustomVersion::ChannelCloseReason)
-				{
-					const uint8 bDormant = Bunch.bClose ? Reader.ReadBit() : 0;
-					Bunch.CloseReason = bDormant ? EChannelCloseReason::Dormancy : EChannelCloseReason::Destroyed;
-				}
-				else
-				{
-					Bunch.CloseReason = Bunch.bClose ? (EChannelCloseReason)Reader.ReadInt((uint32)EChannelCloseReason::MAX) :
-														EChannelCloseReason::Destroyed;
-				}
+				FlushNet();
+			}
+		}
+	}
 
-				PRAGMA_DISABLE_DEPRECATION_WARNINGS
-				// Removing this will require a new network version for replay compatabilty
-				Bunch.bIsReplicationPaused = Reader.ReadBit();
-				PRAGMA_ENABLE_DEPRECATION_WARNINGS
-				Bunch.bReliable				= Reader.ReadBit();
+	// Flush trace content collector
+#if UE_NET_TRACE_ENABLED
+	if (InTraceCollector)
+	{
+		UE_NET_TRACE_FLUSH_COLLECTOR(InTraceCollector, NetTraceId, GetConnectionId(), ENetTracePacketType::Incoming);
+		UE_NET_TRACE_DESTROY_COLLECTOR(InTraceCollector);
+		InTraceCollector = nullptr;
+	}
+#endif
 
-				if (PacketEngineNetVer < FEngineNetworkCustomVersion::MaxActorChannelsCustomization)
-				{
-					static const int OLD_MAX_ACTOR_CHANNELS = 10240;
-					Bunch.ChIndex = Reader.ReadInt(OLD_MAX_ACTOR_CHANNELS);
-				}
-				else
-				{
-					uint32 ChIndex;
-					Reader.SerializeIntPacked(ChIndex);
+	// Trace end marker of this incoming data packet.
+	UE_NET_TRACE_PACKET_RECV(NetTraceId, GetConnectionId(), InPacketId, Reader.GetNumBits());
 
-					if (ChIndex >= (uint32)Channels.Num())
-					{
-						UE_LOG(LogNet, Warning, TEXT("Bunch channel index exceeds channel limit"));
+	if ((bHasBunchErrors || bSkipAck) && !IsInternalAck())
+	{
+		// Trace packet as lost on receiving end to indicate bSkipAck or errors in traced data
+		UE_NET_TRACE_PACKET_DROPPED(NetTraceId, GetConnectionId(), InPacketId, ENetTracePacketType::Incoming);
+	}
+}
 
-						Close(ENetCloseResult::BunchBadChannelIndex);
+void UNetConnection::DispatchPacket( FBitReader& Reader, int32 PacketId, bool& bOutSkipAck, bool& bOutHasBunchErrors )
+{
+	// Track channels that were rejected while processing this packet - used to avoid sending multiple close-channel bunches,
+	// which would cause a disconnect serverside
+	TArray<int32> RejectedChans;
 
-						return;
-					}
+	const bool bIsServer = Driver->ServerConnection == nullptr;
+	const bool bIgnoreRPCs = Driver->ShouldIgnoreRPCs();
 
-					Bunch.ChIndex = ChIndex;
+	const FEngineNetworkCustomVersion::Type PacketEngineNetVer = static_cast<FEngineNetworkCustomVersion::Type>(Reader.EngineNetVer());
+
+	// Disassemble and dispatch all bunches in the packet.
+	while( !Reader.AtEnd() && GetConnectionState()!=USOCK_Closed )
+	{
+		// For demo backwards compatibility, old replays still have this bit
+		if (IsInternalAck() && PacketEngineNetVer < FEngineNetworkCustomVersion::AcksIncludedInHeader)
+		{
+			const bool IsAckDummy = Reader.ReadBit() == 1u;
+		}
+
+		// Parse the bunch.
+		int32 StartPos = Reader.GetPosBits();
+	
+		// Process Received data
+		{
+			// Parse the incoming data.
+			FInBunch Bunch( this );
+			int32 IncomingStartPos		= Reader.GetPosBits();
+			uint8 bControl				= Reader.ReadBit();
+			Bunch.PacketId				= InPacketId;
+			Bunch.bOpen					= bControl ? Reader.ReadBit() : 0;
+			Bunch.bClose				= bControl ? Reader.ReadBit() : 0;
+		
+			if (PacketEngineNetVer < FEngineNetworkCustomVersion::ChannelCloseReason)
+			{
+				const uint8 bDormant = Bunch.bClose ? Reader.ReadBit() : 0;
+				Bunch.CloseReason = bDormant ? EChannelCloseReason::Dormancy : EChannelCloseReason::Destroyed;
+			}
+			else
+			{
+				Bunch.CloseReason = Bunch.bClose ? (EChannelCloseReason)Reader.ReadInt((uint32)EChannelCloseReason::MAX) :
+													EChannelCloseReason::Destroyed;
+			}
+
+			PRAGMA_DISABLE_DEPRECATION_WARNINGS
+			// Removing this will require a new network version for replay compatabilty
+			Bunch.bIsReplicationPaused = Reader.ReadBit();
+			PRAGMA_ENABLE_DEPRECATION_WARNINGS
+			Bunch.bReliable				= Reader.ReadBit();
+
+			if (PacketEngineNetVer < FEngineNetworkCustomVersion::MaxActorChannelsCustomization)
+			{
+				static const int OLD_MAX_ACTOR_CHANNELS = 10240;
+				Bunch.ChIndex = Reader.ReadInt(OLD_MAX_ACTOR_CHANNELS);
+			}
+			else
+			{
+				uint32 ChIndex;
+				Reader.SerializeIntPacked(ChIndex);
+
+				if (ChIndex >= (uint32)Channels.Num())
+				{
+					UE_LOG(LogNet, Warning, TEXT("Bunch channel index exceeds channel limit"));
+
+					Close(ENetCloseResult::BunchBadChannelIndex);
+
+					return;
 				}
 
-				const int32 CachedChIndex = Bunch.ChIndex;
+				Bunch.ChIndex = ChIndex;
+			}
+
+			const int32 CachedChIndex = Bunch.ChIndex;
 
-				// if flag is set, remap channel index values, we're fast forwarding a replay checkpoint
-				// and there should be no bunches for existing channels
-				if (IsInternalAck() && bAllowExistingChannelIndex && (PacketEngineNetVer >= FEngineNetworkCustomVersion::ReplayDormancy))
+			// if flag is set, remap channel index values, we're fast forwarding a replay checkpoint
+			// and there should be no bunches for existing channels
+			if (IsInternalAck() && bAllowExistingChannelIndex && (PacketEngineNetVer >= FEngineNetworkCustomVersion::ReplayDormancy))
+			{
+				if (ChannelIndexMap.Contains(Bunch.ChIndex))
 				{
-					if (ChannelIndexMap.Contains(Bunch.ChIndex))
-					{
-						Bunch.ChIndex = ChannelIndexMap[Bunch.ChIndex];
-					}
-					else
+					Bunch.ChIndex = ChannelIndexMap[Bunch.ChIndex];
+				}
+				else
+				{
+					if (Channels[Bunch.ChIndex])
 					{
-						if (Channels[Bunch.ChIndex])
+						// this should be an open bunch if its the first time we've seen it
+						if (Bunch.bOpen && (!Bunch.bPartial || Bunch.bPartialInitial))
 						{
-							// this should be an open bunch if its the first time we've seen it
-							if (Bunch.bOpen && (!Bunch.bPartial || Bunch.bPartialInitial))
-							{
-								FBitReaderMark Mark(Reader);
+							FBitReaderMark Mark(Reader);
 
-								Reader.ReadBit(); // bHasPackageMapExports
-								Reader.ReadBit(); // bHasMustBeMappedGUIDs
+							Reader.ReadBit(); // bHasPackageMapExports
+							Reader.ReadBit(); // bHasMustBeMappedGUIDs
 
-								const uint8 bPartial = Reader.ReadBit(); // bPartial
+							const uint8 bPartial = Reader.ReadBit(); // bPartial
 
-								if (bPartial)
-								{
-									Reader.ReadBit(); // bPartialInitial
-									Reader.ReadBit(); // bPartialFinal
-								}
+							if (bPartial)
+							{
+								Reader.ReadBit(); // bPartialInitial
+								Reader.ReadBit(); // bPartialFinal
+							}
 
-								FName ChName;
-								UPackageMap::StaticSerializeName(Reader, ChName);
+							FName ChName;
+							UPackageMap::StaticSerializeName(Reader, ChName);
 
-								Mark.Pop(Reader);
+							Mark.Pop(Reader);
 
-								int32 FreeIndex = GetFreeChannelIndex(ChName);
+							int32 FreeIndex = GetFreeChannelIndex(ChName);
 
-								if (FreeIndex != INDEX_NONE)
-								{
-									UE_LOG(LogNetTraffic, Verbose, TEXT("Adding channel mapping %d to %d"), Bunch.ChIndex, FreeIndex);
-									ChannelIndexMap.Add(Bunch.ChIndex, FreeIndex);
-									Bunch.ChIndex = FreeIndex;
-								}
-								else
-								{
-									UE_LOG(LogNetTraffic, Warning, TEXT("Unable to find free channel index"));
-									continue;
-								}
+							if (FreeIndex != INDEX_NONE)
+							{
+								UE_LOG(LogNetTraffic, Verbose, TEXT("Adding channel mapping %d to %d"), Bunch.ChIndex, FreeIndex);
+								ChannelIndexMap.Add(Bunch.ChIndex, FreeIndex);
+								Bunch.ChIndex = FreeIndex;
+							}
+							else
+							{
+								UE_LOG(LogNetTraffic, Warning, TEXT("Unable to find free channel index"));
+								continue;
 							}
 						}
 					}
 				}
+			}
 
-				Bunch.bHasPackageMapExports	= Reader.ReadBit();
-				Bunch.bHasMustBeMappedGUIDs	= Reader.ReadBit();
-				Bunch.bPartial				= Reader.ReadBit();
+			Bunch.bHasPackageMapExports	= Reader.ReadBit();
+			Bunch.bHasMustBeMappedGUIDs	= Reader.ReadBit();
+			Bunch.bPartial				= Reader.ReadBit();
 
-				if ( Bunch.bReliable )
-				{
-					if ( IsInternalAck() )
-					{
-						// We can derive the sequence for 100% reliable connections
-						Bunch.ChSequence = InReliable[Bunch.ChIndex] + 1;
-					}
-					else
-					{
-						// If this is a reliable bunch, use the last processed reliable sequence to read the new reliable sequence
-						Bunch.ChSequence = MakeRelative( Reader.ReadInt( MAX_CHSEQUENCE ), InReliable[Bunch.ChIndex], MAX_CHSEQUENCE );
-					}
-				} 
-				else if ( Bunch.bPartial )
+			if ( Bunch.bReliable )
+			{
+				if ( IsInternalAck() )
 				{
-					// If this is an unreliable partial bunch, we simply use packet sequence since we already have it
-					Bunch.ChSequence = InPacketId;
+					// We can derive the sequence for 100% reliable connections
+					Bunch.ChSequence = InReliable[Bunch.ChIndex] + 1;
 				}
 				else
 				{
-					Bunch.ChSequence = 0;
+					// If this is a reliable bunch, use the last processed reliable sequence to read the new reliable sequence
+					Bunch.ChSequence = MakeRelative( Reader.ReadInt( MAX_CHSEQUENCE ), InReliable[Bunch.ChIndex], MAX_CHSEQUENCE );
 				}
+			} 
+			else if ( Bunch.bPartial )
+			{
+				// If this is an unreliable partial bunch, we simply use packet sequence since we already have it
+				Bunch.ChSequence = InPacketId;
+			}
+			else
+			{
+				Bunch.ChSequence = 0;
+			}
 
-				Bunch.bPartialInitial = Bunch.bPartial ? Reader.ReadBit() : 0;
-				Bunch.bPartialFinal = Bunch.bPartial ? Reader.ReadBit() : 0;
+			Bunch.bPartialInitial = Bunch.bPartial ? Reader.ReadBit() : 0;
+			Bunch.bPartialFinal = Bunch.bPartial ? Reader.ReadBit() : 0;
 
-				if (PacketEngineNetVer < FEngineNetworkCustomVersion::ChannelNames)
+			if (PacketEngineNetVer < FEngineNetworkCustomVersion::ChannelNames)
+			{
+				uint32 ChType = (Bunch.bReliable || Bunch.bOpen) ? Reader.ReadInt(CHTYPE_MAX) : CHTYPE_None;
+				switch (ChType)
 				{
-					uint32 ChType = (Bunch.bReliable || Bunch.bOpen) ? Reader.ReadInt(CHTYPE_MAX) : CHTYPE_None;
-					switch (ChType)
-					{
-						case CHTYPE_Control:
-							Bunch.ChName = NAME_Control;
-							break;
-						case CHTYPE_Voice:
-							Bunch.ChName = NAME_Voice;
-							break;
-						case CHTYPE_Actor:
-							Bunch.ChName = NAME_Actor;
-							break;
+					case CHTYPE_Control:
+						Bunch.ChName = NAME_Control;
 						break;
-					}
+					case CHTYPE_Voice:
+						Bunch.ChName = NAME_Voice;
+						break;
+					case CHTYPE_Actor:
+						Bunch.ChName = NAME_Actor;
+						break;
+					break;
 				}
-				else
+			}
+			else
+			{
+				if (Bunch.bReliable || Bunch.bOpen)
 				{
-					if (Bunch.bReliable || Bunch.bOpen)
-					{
-						UPackageMap::StaticSerializeName(Reader, Bunch.ChName);
+					UPackageMap::StaticSerializeName(Reader, Bunch.ChName);
 
-						if( Reader.IsError() )
-						{
-							UE_LOG(LogNet, Warning, TEXT("Channel name serialization failed."));
+					if( Reader.IsError() )
+					{
+						UE_LOG(LogNet, Warning, TEXT("Channel name serialization failed."));
 
-							HandleNetResultOrClose(ENetCloseResult::BunchChannelNameFail);
+						HandleNetResultOrClose(ENetCloseResult::BunchChannelNameFail);
 
-							return;
-						}
-					}
-					else
-					{
-						Bunch.ChName = NAME_None;
+						return;
 					}
 				}
+				else
+				{
+					Bunch.ChName = NAME_None;
+				}
+			}
 
-				UChannel* Channel = Channels[Bunch.ChIndex];
+			UChannel* Channel = Channels[Bunch.ChIndex];
 
-				// If there's an existing channel and the bunch specified it's channel type, make sure they match.
-				if (Channel && (Bunch.ChName != NAME_None) && (Bunch.ChName != Channel->ChName))
-				{
-					UE_LOG(LogNet, Error,
-							TEXT("Existing channel at index %d with type \"%s\" differs from the incoming bunch's expected channel type, \"%s\"."),
-							Bunch.ChIndex, *Channel->ChName.ToString(), *Bunch.ChName.ToString());
+			// If there's an existing channel and the bunch specified it's channel type, make sure they match.
+			if (Channel && (Bunch.ChName != NAME_None) && (Bunch.ChName != Channel->ChName))
+			{
+				UE_LOG(LogNet, Error,
+						TEXT("Existing channel at index %d with type \"%s\" differs from the incoming bunch's expected channel type, \"%s\"."),
+						Bunch.ChIndex, *Channel->ChName.ToString(), *Bunch.ChName.ToString());
 
-					Close(ENetCloseResult::BunchWrongChannelType);
+				Close(ENetCloseResult::BunchWrongChannelType);
 
-					return;
-				}
+				return;
+			}
 
-				int32 BunchDataBits  = Reader.ReadInt( UNetConnection::MaxPacket*8 );
+			int32 BunchDataBits  = Reader.ReadInt( UNetConnection::MaxPacket*8 );
 
-				if ((Bunch.bClose || Bunch.bOpen) && UE_LOG_ACTIVE(LogNetDormancy,VeryVerbose) )
-				{
-					UE_LOG(LogNetDormancy, VeryVerbose, TEXT("Received: %s"), *Bunch.ToString());
-				}
+			if ((Bunch.bClose || Bunch.bOpen) && UE_LOG_ACTIVE(LogNetDormancy,VeryVerbose) )
+			{
+				UE_LOG(LogNetDormancy, VeryVerbose, TEXT("Received: %s"), *Bunch.ToString());
+			}
 
-				if (UE_LOG_ACTIVE(LogNetTraffic,VeryVerbose))
-				{
-					UE_LOG(LogNetTraffic, VeryVerbose, TEXT("Received: %s"), *Bunch.ToString());
-				}
+			if (UE_LOG_ACTIVE(LogNetTraffic,VeryVerbose))
+			{
+				UE_LOG(LogNetTraffic, VeryVerbose, TEXT("Received: %s"), *Bunch.ToString());
+			}
 
-				const int32 HeaderPos = Reader.GetPosBits();
+			const int32 HeaderPos = Reader.GetPosBits();
 
-				if( Reader.IsError() )
-				{
-					UE_LOG(LogNet, Warning, TEXT("Bunch header overflowed"));
+			if( Reader.IsError() )
+			{
+				UE_LOG(LogNet, Warning, TEXT("Bunch header overflowed"));
 
-					HandleNetResultOrClose(ENetCloseResult::BunchHeaderOverflow);
+				HandleNetResultOrClose(ENetCloseResult::BunchHeaderOverflow);
 
-					return;
-				}
+				return;
+			}
 
-				// Trace bunch read
-				UE_NET_TRACE_BUNCH_SCOPE(InTraceCollector, Bunch, StartPos, HeaderPos - StartPos);
+			// Trace bunch read
+			UE_NET_TRACE_BUNCH_SCOPE(InTraceCollector, Bunch, StartPos, HeaderPos - StartPos);
 
-				Bunch.SetData( Reader, BunchDataBits );
-				if( Reader.IsError() )
-				{
-					// Bunch claims it's larger than the enclosing packet.
-					UE_LOG(LogNet, Warning, TEXT("Bunch data overflowed (%i %i+%i/%i)"), IncomingStartPos, HeaderPos, BunchDataBits,
-							Reader.GetNumBits());
+			Bunch.SetData( Reader, BunchDataBits );
+			if( Reader.IsError() )
+			{
+				// Bunch claims it's larger than the enclosing packet.
+				UE_LOG(LogNet, Warning, TEXT("Bunch data overflowed (%i %i+%i/%i)"), IncomingStartPos, HeaderPos, BunchDataBits,
+						Reader.GetNumBits());
 
-					HandleNetResultOrClose(ENetCloseResult::BunchDataOverflow);
+				HandleNetResultOrClose(ENetCloseResult::BunchDataOverflow);
 
-					return;
-				}
+				return;
+			}
 
-				if (Bunch.bHasPackageMapExports)
+			if (Bunch.bHasPackageMapExports)
+			{
+				// Clients still send NetGUID.IsDefault()/FExportFlags.bHasPath packets to the server, separate from this check
+				if (Driver->IsServer())
 				{
-					// Clients still send NetGUID.IsDefault()/FExportFlags.bHasPath packets to the server, separate from this check
-					if (Driver->IsServer())
-					{
-						UE_LOG(LogNetTraffic, Error, TEXT("UNetConnection::ReceivedPacket: Server received bHasPackageMapExports packet."));
+					UE_LOG(LogNetTraffic, Error, TEXT("UNetConnection::ReceivedPacket: Server received bHasPackageMapExports packet."));
 
-						Close(ENetCloseResult::BunchServerPackageMapExports);
+					Close(ENetCloseResult::BunchServerPackageMapExports);
 
-						return;
-					}
+					return;
+				}
 
-					Driver->NetGUIDInBytes += (BunchDataBits + (HeaderPos - IncomingStartPos)) >> 3;
+				Driver->NetGUIDInBytes += (BunchDataBits + (HeaderPos - IncomingStartPos)) >> 3;
 
-					if (IsInternalAck())
-					{
-						// NOTE - For replays, we do this even earlier, to try and load this as soon as possible,
-						// in case there's an issue creating the channel. If a replay fails to create a channel, we want to salvage as much as possible
-						Cast<UPackageMapClient>(PackageMap)->ReceiveNetGUIDBunch(Bunch);
+				if (IsInternalAck())
+				{
+					// NOTE - For replays, we do this even earlier, to try and load this as soon as possible,
+					// in case there's an issue creating the channel. If a replay fails to create a channel, we want to salvage as much as possible
+					Cast<UPackageMapClient>(PackageMap)->ReceiveNetGUIDBunch(Bunch);
 
-						if (Bunch.IsError())
-						{
-							UE_LOG(LogNetTraffic, Error, TEXT("UNetConnection::ReceivedPacket: Bunch.IsError() after ReceiveNetGUIDBunch. ChIndex: %i"),
-									Bunch.ChIndex);
-						}
+					if (Bunch.IsError())
+					{
+						UE_LOG(LogNetTraffic, Error, TEXT("UNetConnection::ReceivedPacket: Bunch.IsError() after ReceiveNetGUIDBunch. ChIndex: %i"),
+								Bunch.ChIndex);
 					}
 				}
+			}
 
-				if (Bunch.bReliable)
-				{
-					UE_LOG(LogNetTraffic, Verbose, TEXT("   Reliable Bunch, Channel %i Sequence %i: Size %.1f+%.1f"),
-							Bunch.ChIndex, Bunch.ChSequence, (HeaderPos-IncomingStartPos)/8.f, (Reader.GetPosBits()-HeaderPos)/8.f);
-				}
-				else
-				{
-					UE_LOG(LogNetTraffic, Verbose, TEXT("   Unreliable Bunch, Channel %i: Size %.1f+%.1f"),
-							Bunch.ChIndex, (HeaderPos-IncomingStartPos)/8.f, (Reader.GetPosBits()-HeaderPos)/8.f);
-				}
+			if (Bunch.bReliable)
+			{
+				UE_LOG(LogNetTraffic, Verbose, TEXT("   Reliable Bunch, Channel %i Sequence %i: Size %.1f+%.1f"),
+						Bunch.ChIndex, Bunch.ChSequence, (HeaderPos-IncomingStartPos)/8.f, (Reader.GetPosBits()-HeaderPos)/8.f);
+			}
+			else
+			{
+				UE_LOG(LogNetTraffic, Verbose, TEXT("   Unreliable Bunch, Channel %i: Size %.1f+%.1f"),
+						Bunch.ChIndex, (HeaderPos-IncomingStartPos)/8.f, (Reader.GetPosBits()-HeaderPos)/8.f);
+			}
 
-				if (Bunch.bOpen)
-				{
-					UE_LOG(LogNetTraffic, Verbose, TEXT("   bOpen Bunch, Channel %i Sequence %i: Size %.1f+%.1f"),
-							Bunch.ChIndex, Bunch.ChSequence, (HeaderPos-IncomingStartPos)/8.f, (Reader.GetPosBits()-HeaderPos)/8.f );
-				}
+			if (Bunch.bOpen)
+			{
+				UE_LOG(LogNetTraffic, Verbose, TEXT("   bOpen Bunch, Channel %i Sequence %i: Size %.1f+%.1f"),
+						Bunch.ChIndex, Bunch.ChSequence, (HeaderPos-IncomingStartPos)/8.f, (Reader.GetPosBits()-HeaderPos)/8.f );
+			}
 
-				if ( Channels[Bunch.ChIndex] == NULL && ( Bunch.ChIndex != 0 || Bunch.ChName != NAME_Control ) )
+			if ( Channels[Bunch.ChIndex] == NULL && ( Bunch.ChIndex != 0 || Bunch.ChName != NAME_Control ) )
+			{
+				// Can't handle other channels until control channel exists.
+				if ( Channels[0] == NULL )
 				{
-					// Can't handle other channels until control channel exists.
-					if ( Channels[0] == NULL )
-					{
-						UE_LOG(LogNetTraffic, Log,
-								TEXT("ReceivedPacket: Received non-control bunch before control channel was created. ChIndex: %i, ChName: %s"),
-								Bunch.ChIndex, *Bunch.ChName.ToString());
-
-						Close(ENetCloseResult::BunchPrematureControlChannel);
-
-						return;
-					}
-					// on the server, if we receive bunch data for a channel that doesn't exist while we're still logging in,
-					// it's either a broken client or a new instance of a previous connection,
-					// so reject it
-					else if ( PlayerController == NULL && Driver->ClientConnections.Contains( this ) )
-					{
-						UE_LOG(LogNet, Warning,
-								TEXT("ReceivedPacket: Received non-control bunch before player controller was assigned. ChIndex: %i, ChName: %s" ),
-								Bunch.ChIndex, *Bunch.ChName.ToString());
+					UE_LOG(LogNetTraffic, Log,
+							TEXT("ReceivedPacket: Received non-control bunch before control channel was created. ChIndex: %i, ChName: %s"),
+							Bunch.ChIndex, *Bunch.ChName.ToString());
 
-						Close(ENetCloseResult::BunchPrematureChannel);
+					Close(ENetCloseResult::BunchPrematureControlChannel);
 
-						return;
-					}
+					return;
 				}
-				// ignore control channel close if it hasn't been opened yet
-				if ( Bunch.ChIndex == 0 && Channels[0] == NULL && Bunch.bClose && Bunch.ChName == NAME_Control )
+				// on the server, if we receive bunch data for a channel that doesn't exist while we're still logging in,
+				// it's either a broken client or a new instance of a previous connection,
+				// so reject it
+				else if ( PlayerController == NULL && Driver->ClientConnections.Contains( this ) )
 				{
-					UE_LOG(LogNetTraffic, Log, TEXT("ReceivedPacket: Received control channel close before open"));
+					UE_LOG(LogNet, Warning,
+							TEXT("ReceivedPacket: Received non-control bunch before player controller was assigned. ChIndex: %i, ChName: %s" ),
+							Bunch.ChIndex, *Bunch.ChName.ToString());
 
-					Close(ENetCloseResult::BunchPrematureControlClose);
+					Close(ENetCloseResult::BunchPrematureChannel);
 
 					return;
 				}
+			}
+			// ignore control channel close if it hasn't been opened yet
+			if ( Bunch.ChIndex == 0 && Channels[0] == NULL && Bunch.bClose && Bunch.ChName == NAME_Control )
+			{
+				UE_LOG(LogNetTraffic, Log, TEXT("ReceivedPacket: Received control channel close before open"));
 
-				// Receiving data.
+				Close(ENetCloseResult::BunchPrematureControlClose);
 
-				// We're on a 100% reliable connection and we are rolling back some data.
-				// In that case, we can generally ignore these bunches.
-				if (IsInternalAck() && bAllowExistingChannelIndex)
+				return;
+			}
+
+			// Receiving data.
+
+			// We're on a 100% reliable connection and we are rolling back some data.
+			// In that case, we can generally ignore these bunches.
+			if (IsInternalAck() && bAllowExistingChannelIndex)
+			{
+				if (PacketEngineNetVer < FEngineNetworkCustomVersion::ReplayDormancy)
 				{
-					if (PacketEngineNetVer < FEngineNetworkCustomVersion::ReplayDormancy)
+					if (Channel)
 					{
-						if (Channel)
+						// This was an open bunch for a channel that's already opened.
+						// We can ignore future bunches from this channel.
+						const bool bNewlyOpenedActorChannel = Bunch.bOpen && (Bunch.ChName == NAME_Actor) &&
+																(!Bunch.bPartial || Bunch.bPartialInitial);
+
+						if (bNewlyOpenedActorChannel)
 						{
-							// This was an open bunch for a channel that's already opened.
-							// We can ignore future bunches from this channel.
-							const bool bNewlyOpenedActorChannel = Bunch.bOpen && (Bunch.ChName == NAME_Actor) &&
-																	(!Bunch.bPartial || Bunch.bPartialInitial);
+							FNetworkGUID ActorGUID = GetActorGUIDFromOpenBunch(Bunch);
 
-							if (bNewlyOpenedActorChannel)
+							if (!Bunch.IsError())
 							{
-								FNetworkGUID ActorGUID = GetActorGUIDFromOpenBunch(Bunch);
-
-								if (!Bunch.IsError())
-								{
-									IgnoringChannels.Add(Bunch.ChIndex, ActorGUID);
-								}
-								else
-								{
-									UE_LOG(LogNetTraffic, Error,
-											TEXT("UNetConnection::ReceivedPacket: Unable to read actor GUID for ignored bunch. (Channel %d)"),
-											Bunch.ChIndex);
-								}
+								IgnoringChannels.Add(Bunch.ChIndex, ActorGUID);
+							}
+							else
+							{
+								UE_LOG(LogNetTraffic, Error,
+										TEXT("UNetConnection::ReceivedPacket: Unable to read actor GUID for ignored bunch. (Channel %d)"),
+										Bunch.ChIndex);
 							}
+						}
 
-							if (IgnoringChannels.Contains(Bunch.ChIndex))
+						if (IgnoringChannels.Contains(Bunch.ChIndex))
+						{
+							if (Bunch.bClose && (!Bunch.bPartial || Bunch.bPartialFinal))
 							{
-								if (Bunch.bClose && (!Bunch.bPartial || Bunch.bPartialFinal))
+								FNetworkGUID ActorGUID = IgnoringChannels.FindAndRemoveChecked(Bunch.ChIndex);
+								if (ActorGUID.IsStatic())
 								{
-									FNetworkGUID ActorGUID = IgnoringChannels.FindAndRemoveChecked(Bunch.ChIndex);
-									if (ActorGUID.IsStatic())
+									UObject* FoundObject = Driver->GuidCache->GetObjectFromNetGUID(ActorGUID, false);
+									if (AActor* StaticActor = Cast<AActor>(FoundObject))
 									{
-										UObject* FoundObject = Driver->GuidCache->GetObjectFromNetGUID(ActorGUID, false);
-										if (AActor* StaticActor = Cast<AActor>(FoundObject))
-										{
-											DestroyIgnoredActor(StaticActor);
-										}
-										else
-										{
-											ensure(FoundObject == nullptr);
-
-											UE_LOG(LogNetTraffic, Log,
-													TEXT("UNetConnection::ReceivedPacket: Unable to find static actor to cleanup for ignored bunch. ")
-													TEXT("(Channel %d NetGUID %s)"), Bunch.ChIndex, *ActorGUID.ToString());
-										}
+										DestroyIgnoredActor(StaticActor);
 									}
-								}
+									else
+									{
+										ensure(FoundObject == nullptr);
 
-								UE_LOG(LogNetTraffic, Log, TEXT("Ignoring bunch for already open channel: %i"), Bunch.ChIndex);
-								continue;
+										UE_LOG(LogNetTraffic, Log,
+												TEXT("UNetConnection::ReceivedPacket: Unable to find static actor to cleanup for ignored bunch. ")
+												TEXT("(Channel %d NetGUID %s)"), Bunch.ChIndex, *ActorGUID.ToString());
+									}
+								}
 							}
+
+							UE_LOG(LogNetTraffic, Log, TEXT("Ignoring bunch for already open channel: %i"), Bunch.ChIndex);
+							continue;
 						}
 					}
-					else
-					{
-						const bool bCloseBunch = Bunch.bClose && (!Bunch.bPartial || Bunch.bPartialFinal);
+				}
+				else
+				{
+					const bool bCloseBunch = Bunch.bClose && (!Bunch.bPartial || Bunch.bPartialFinal);
 
-						if (bCloseBunch && ChannelIndexMap.Contains(CachedChIndex))
-						{
-							check(ChannelIndexMap[CachedChIndex] == Bunch.ChIndex);
+					if (bCloseBunch && ChannelIndexMap.Contains(CachedChIndex))
+					{
+						check(ChannelIndexMap[CachedChIndex] == Bunch.ChIndex);
 
-							UE_LOG(LogNetTraffic, Verbose, TEXT("Removing channel mapping %d to %d"), CachedChIndex, Bunch.ChIndex);
-							ChannelIndexMap.Remove(CachedChIndex);
-						}
+						UE_LOG(LogNetTraffic, Verbose, TEXT("Removing channel mapping %d to %d"), CachedChIndex, Bunch.ChIndex);
+						ChannelIndexMap.Remove(CachedChIndex);
 					}
 				}
+			}
 
-				// Ignore if reliable packet has already been processed.
-				if ( Bunch.bReliable && Bunch.ChSequence <= InReliable[Bunch.ChIndex] )
-				{
-					UE_LOG(LogNetTraffic, Log, TEXT("UNetConnection::ReceivedPacket: Received outdated bunch (Channel %d Current Sequence %i)"),
-							Bunch.ChIndex, InReliable[Bunch.ChIndex]);
+			// Ignore if reliable packet has already been processed.
+			if ( Bunch.bReliable && Bunch.ChSequence <= InReliable[Bunch.ChIndex] )
+			{
+				UE_LOG(LogNetTraffic, Log, TEXT("UNetConnection::ReceivedPacket: Received outdated bunch (Channel %d Current Sequence %i)"),
+						Bunch.ChIndex, InReliable[Bunch.ChIndex]);
 
-					check( !IsInternalAck() );		// Should be impossible with 100% reliable connections
-					continue;
-				}
-			
-				// If opening the channel with an unreliable packet, check that it is "bNetTemporary", otherwise discard it
-				if( !Channel && !Bunch.bReliable )
-				{
-					// Unreliable bunches that open channels should be bOpen && (bClose || bPartial)
-					// NetTemporary usually means one bunch that is unreliable (bOpen and bClose):	1(bOpen, bClose)
-					// But if that bunch export NetGUIDs, it will get split into 2 bunches:			1(bOpen, bPartial) - 2(bClose).
-					// (the initial actor bunch itself could also be split into multiple bunches. So bPartial is the right check here)
+				check( !IsInternalAck() );		// Should be impossible with 100% reliable connections
+				continue;
+			}
+		
+			// If opening the channel with an unreliable packet, check that it is "bNetTemporary", otherwise discard it
+			if( !Channel && !Bunch.bReliable )
+			{
+				// Unreliable bunches that open channels should be bOpen && (bClose || bPartial)
+				// NetTemporary usually means one bunch that is unreliable (bOpen and bClose):	1(bOpen, bClose)
+				// But if that bunch export NetGUIDs, it will get split into 2 bunches:			1(bOpen, bPartial) - 2(bClose).
+				// (the initial actor bunch itself could also be split into multiple bunches. So bPartial is the right check here)
 
-					const bool ValidUnreliableOpen = Bunch.bOpen && (Bunch.bClose || Bunch.bPartial);
-					if (!ValidUnreliableOpen)
+				const bool ValidUnreliableOpen = Bunch.bOpen && (Bunch.bClose || Bunch.bPartial);
+				if (!ValidUnreliableOpen)
+				{
+					if ( IsInternalAck() )
 					{
-						if ( IsInternalAck() )
-						{
-							// Should be impossible with 100% reliable connections
-							UE_LOG(LogNetTraffic, Error,
-									TEXT("      Received unreliable bunch before open with reliable connection (Channel %d Current Sequence %i)" ),
-									Bunch.ChIndex, InReliable[Bunch.ChIndex]);
-						}
-						else
-						{
-							// Simply a log (not a warning, since this can happen under normal conditions, like from a re-join, etc)
-							UE_LOG(LogNetTraffic, Log, TEXT("      Received unreliable bunch before open (Channel %d Current Sequence %i)"),
-									Bunch.ChIndex, InReliable[Bunch.ChIndex]);
-						}
-
-						// Since we won't be processing this packet, don't ack it
-						// We don't want the sender to think this bunch was processed when it really wasn't
-						bSkipAck = true;
-						continue;
+						// Should be impossible with 100% reliable connections
+						UE_LOG(LogNetTraffic, Error,
+								TEXT("      Received unreliable bunch before open with reliable connection (Channel %d Current Sequence %i)" ),
+								Bunch.ChIndex, InReliable[Bunch.ChIndex]);
+					}
+					else
+					{
+						// Simply a log (not a warning, since this can happen under normal conditions, like from a re-join, etc)
+						UE_LOG(LogNetTraffic, Log, TEXT("      Received unreliable bunch before open (Channel %d Current Sequence %i)"),
+								Bunch.ChIndex, InReliable[Bunch.ChIndex]);
 					}
+
+					// Since we won't be processing this packet, don't ack it
+					// We don't want the sender to think this bunch was processed when it really wasn't
+					bOutSkipAck = true;
+					continue;
 				}
+			}
 
-				// Create channel if necessary.
-				if (Channel == nullptr)
+			// Create channel if necessary.
+			if (Channel == nullptr)
+			{
+				if (RejectedChans.Contains(Bunch.ChIndex))
 				{
-					if (RejectedChans.Contains(Bunch.ChIndex))
-					{
-						UE_LOG(LogNetTraffic, Log,
-								TEXT("      Ignoring Bunch for ChIndex %i, as the channel was already rejected while processing this packet."),
-								Bunch.ChIndex);
+					UE_LOG(LogNetTraffic, Log,
+							TEXT("      Ignoring Bunch for ChIndex %i, as the channel was already rejected while processing this packet."),
+							Bunch.ChIndex);
 
-						continue;
-					}
+					continue;
+				}
 
-					// Validate channel type.
-					if ( !Driver->IsKnownChannelName( Bunch.ChName ) )
-					{
-						UE_LOG(LogNet, Warning, TEXT("ReceivedPacket: Connection unknown channel type (%s)"), *Bunch.ChName.ToString());
+				// Validate channel type.
+				if ( !Driver->IsKnownChannelName( Bunch.ChName ) )
+				{
+					UE_LOG(LogNet, Warning, TEXT("ReceivedPacket: Connection unknown channel type (%s)"), *Bunch.ChName.ToString());
 
-						Close(ENetCloseResult::UnknownChannelType);
+					Close(ENetCloseResult::UnknownChannelType);
 
-						return;
-					}
+					return;
+				}
 
-					// Ignore incoming data on channel types clients can't create. Can occur if we've incoming data when server is closing a channel
-					if ( Driver->IsServer() && (Driver->ChannelDefinitionMap[Bunch.ChName].bClientOpen == false) )
-					{
-						UE_LOG(LogNetTraffic, Warning, TEXT("      Ignoring Bunch Create received from client since only server is allowed to ")
-								TEXT("create this type of channel: Bunch  %i: ChName %s, ChSequence: %i, bReliable: %i, bPartial: %i, ")
-								TEXT("bPartialInitial: %i, bPartialFinal: %i"), Bunch.ChIndex, *Bunch.ChName.ToString(), Bunch.ChSequence,
-								(int)Bunch.bReliable, (int)Bunch.bPartial, (int)Bunch.bPartialInitial, (int)Bunch.bPartialFinal);
+				// Ignore incoming data on channel types clients can't create. Can occur if we've incoming data when server is closing a channel
+				if ( Driver->IsServer() && (Driver->ChannelDefinitionMap[Bunch.ChName].bClientOpen == false) )
+				{
+					UE_LOG(LogNetTraffic, Warning, TEXT("      Ignoring Bunch Create received from client since only server is allowed to ")
+							TEXT("create this type of channel: Bunch  %i: ChName %s, ChSequence: %i, bReliable: %i, bPartial: %i, ")
+							TEXT("bPartialInitial: %i, bPartialFinal: %i"), Bunch.ChIndex, *Bunch.ChName.ToString(), Bunch.ChSequence,
+							(int)Bunch.bReliable, (int)Bunch.bPartial, (int)Bunch.bPartialInitial, (int)Bunch.bPartialFinal);
 
-						RejectedChans.AddUnique(Bunch.ChIndex);
+					RejectedChans.AddUnique(Bunch.ChIndex);
 
-						continue;
-					}
+					continue;
+				}
 
-					// peek for guid
-					if (IsInternalAck() && bIgnoreActorBunches)
+				// peek for guid
+				if (IsInternalAck() && bIgnoreActorBunches)
+				{
+					if (Bunch.bOpen && (!Bunch.bPartial || Bunch.bPartialInitial) && (Bunch.ChName == NAME_Actor))
 					{
-						if (Bunch.bOpen && (!Bunch.bPartial || Bunch.bPartialInitial) && (Bunch.ChName == NAME_Actor))
-						{
-							FBitReaderMark Mark(Bunch);
-							FNetworkGUID ActorGUID = GetActorGUIDFromOpenBunch(Bunch);
-							Mark.Pop(Bunch);
+						FBitReaderMark Mark(Bunch);
+						FNetworkGUID ActorGUID = GetActorGUIDFromOpenBunch(Bunch);
+						Mark.Pop(Bunch);
 
-							if (ActorGUID.IsValid() && !ActorGUID.IsDefault())
+						if (ActorGUID.IsValid() && !ActorGUID.IsDefault())
+						{
+							if (IgnoredBunchGuids.Contains(ActorGUID))
 							{
-								if (IgnoredBunchGuids.Contains(ActorGUID))
-								{
-									UE_LOG(LogNetTraffic, Verbose, TEXT("Adding Channel: %i to ignore list, ignoring guid: %s"),
-											Bunch.ChIndex, *ActorGUID.ToString());
+								UE_LOG(LogNetTraffic, Verbose, TEXT("Adding Channel: %i to ignore list, ignoring guid: %s"),
+										Bunch.ChIndex, *ActorGUID.ToString());
 
-									IgnoredBunchChannels.Add(Bunch.ChIndex);
+								IgnoredBunchChannels.Add(Bunch.ChIndex);
 
-									continue;
-								}
-								else
-								{
-									if (IgnoredBunchChannels.Remove(Bunch.ChIndex))
-									{
-										UE_LOG(LogNetTraffic, Verbose, TEXT("Removing Channel: %i from ignore list, got new guid: %s"),
-												Bunch.ChIndex, *ActorGUID.ToString());
-									}
-								}
+								continue;
 							}
 							else
 							{
-								UE_LOG(LogNetTraffic, Warning, TEXT("Open bunch with invalid actor guid, Channel: %i"), Bunch.ChIndex);
+								if (IgnoredBunchChannels.Remove(Bunch.ChIndex))
+								{
+									UE_LOG(LogNetTraffic, Verbose, TEXT("Removing Channel: %i from ignore list, got new guid: %s"),
+											Bunch.ChIndex, *ActorGUID.ToString());
+								}
 							}
 						}
 						else
 						{
-							if (IgnoredBunchChannels.Contains(Bunch.ChIndex))
-							{
-								UE_LOG(LogNetTraffic, Verbose, TEXT("Ignoring bunch on channel: %i"), Bunch.ChIndex);
-								continue;
-							}
+							UE_LOG(LogNetTraffic, Warning, TEXT("Open bunch with invalid actor guid, Channel: %i"), Bunch.ChIndex);
 						}
 					}
-
-					// Reliable (either open or later), so create new channel.
-					UE_LOG(LogNetTraffic, Log, TEXT("      Bunch Create %i: ChName %s, ChSequence: %i, bReliable: %i, bPartial: %i, ")
-							TEXT("bPartialInitial: %i, bPartialFinal: %i"), Bunch.ChIndex, *Bunch.ChName.ToString(), Bunch.ChSequence,
-							(int)Bunch.bReliable, (int)Bunch.bPartial, (int)Bunch.bPartialInitial, (int)Bunch.bPartialFinal);
-
-					Channel = CreateChannelByName( Bunch.ChName, EChannelCreateFlags::None, Bunch.ChIndex );
-
-					// Notify the server of the new channel.
-					if( Driver->Notify == nullptr || !Driver->Notify->NotifyAcceptingChannel( Channel ) )
+					else
 					{
-						// Channel refused, so close it, flush it, and delete it.
-						UE_LOG(LogNet, Verbose, TEXT("      NotifyAcceptingChannel Failed! Channel: %s"), *Channel->Describe() );
-
-						RejectedChans.AddUnique(Bunch.ChIndex);
-
-						FOutBunch CloseBunch( Channel, true );
-						check(!CloseBunch.IsError());
-						check(CloseBunch.bClose);
-						CloseBunch.bReliable = 1;
-						Channel->SendBunch( &CloseBunch, false );
-						FlushNet();
-						Channel->ConditionalCleanUp(false, EChannelCloseReason::Destroyed);
-						if( Bunch.ChIndex==0 )
+						if (IgnoredBunchChannels.Contains(Bunch.ChIndex))
 						{
-							UE_LOG(LogNetTraffic, Log, TEXT("Channel 0 create failed") );
-							SetConnectionState(EConnectionState::USOCK_Closed);
+							UE_LOG(LogNetTraffic, Verbose, TEXT("Ignoring bunch on channel: %i"), Bunch.ChIndex);
+							continue;
 						}
-						continue;
 					}
 				}
 
-				Bunch.bIgnoreRPCs = bIgnoreRPCs;
+				// Reliable (either open or later), so create new channel.
+				UE_LOG(LogNetTraffic, Log, TEXT("      Bunch Create %i: ChName %s, ChSequence: %i, bReliable: %i, bPartial: %i, ")
+						TEXT("bPartialInitial: %i, bPartialFinal: %i"), Bunch.ChIndex, *Bunch.ChName.ToString(), Bunch.ChSequence,
+						(int)Bunch.bReliable, (int)Bunch.bPartial, (int)Bunch.bPartialInitial, (int)Bunch.bPartialFinal);
 
-				// Dispatch the raw, unsequenced bunch to the channel.
-				bool bLocalSkipAck = false;
-
-				if (Channel != nullptr)
-				{
-					// Warning: May destroy channel
-					Channel->ReceivedRawBunch(Bunch, bLocalSkipAck);
-				}
-
-				if ( bLocalSkipAck )
-				{
-					bSkipAck = true;
-				}
-				Driver->InBunches++;
-				Driver->InTotalBunches++;
+				Channel = CreateChannelByName( Bunch.ChName, EChannelCreateFlags::None, Bunch.ChIndex );
 
-				if (Bunch.IsCriticalError() || Bunch.IsError())
+				// Notify the server of the new channel.
+				if( Driver->Notify == nullptr || !Driver->Notify->NotifyAcceptingChannel( Channel ) )
 				{
-					bHasBunchErrors = true;
-
-					// Disconnect if we received a corrupted packet from the client (eg server crash attempt).
-					if (bIsServer)
-					{
-						UE_LOG(LogNetTraffic, Error, TEXT("Received corrupted packet data with SequenceId: %d from client %s. Disconnecting."), InPacketId, *LowLevelGetRemoteAddress());
-
-						Close(AddToAndConsumeChainResultPtr(Bunch.ExtendedError, ENetCloseResult::CorruptData));
-
-						return;
-					}
-					else
+					// Channel refused, so close it, flush it, and delete it.
+					UE_LOG(LogNet, Verbose, TEXT("      NotifyAcceptingChannel Failed! Channel: %s"), *Channel->Describe() );
+
+					RejectedChans.AddUnique(Bunch.ChIndex);
+
+					FOutBunch CloseBunch( Channel, true );
+					check(!CloseBunch.IsError());
+					check(CloseBunch.bClose);
+					CloseBunch.bReliable = 1;
+					Channel->SendBunch( &CloseBunch, false );
+					FlushNet();
+					Channel->ConditionalCleanUp(false, EChannelCloseReason::Destroyed);
+					if( Bunch.ChIndex==0 )
 					{
-						UE_LOG(LogNetTraffic, Error, TEXT("Received corrupted packet data with SequenceId: %d from server %s"), InPacketId, *LowLevelGetRemoteAddress());
-#if UE_WITH_IRIS
-						// If Iris reports errors they are unrecoverable.
-						if (Driver->GetReplicationSystem() != nullptr)
-						{
-							ensureAlwaysMsgf(false, TEXT("Received corrupted packet data with SequenceId: %d from server."), InPacketId);
-							Close(AddToAndConsumeChainResultPtr(Bunch.ExtendedError, ENetCloseResult::CorruptData));
-							return;
-						}
-#endif
+						UE_LOG(LogNetTraffic, Log, TEXT("Channel 0 create failed") );
+						SetConnectionState(EConnectionState::USOCK_Closed);
 					}
-				}
-				// In replay, if the bunch generated an error but the channel isn't actually open, clean it up so the channel index remains free
-				if (IsInternalAck() && Bunch.IsError() && Channel && !Channel->OpenedLocally && !Channel->OpenAcked)
-				{
-					UE_LOG(LogNetTraffic, Warning, TEXT("Replay cleaning up channel that couldn't be opened: %s"), *Channel->Describe());
-					Channel->ConditionalCleanUp(true, EChannelCloseReason::Destroyed);
+					continue;
 				}
 			}
-		}
 
-		// Close/clean-up channels pending close due to received acks.
-		for (FChannelCloseInfo& Info : ChannelsToClose)
-		{
-			if (UChannel* Channel = Channels[Info.Id])
-			{
-				Channel->ConditionalCleanUp(false, Info.CloseReason);
-			}
-		}
+			Bunch.bIgnoreRPCs = bIgnoreRPCs;
 
-		ValidateSendBuffer();
-	}
+			// Dispatch the raw, unsequenced bunch to the channel.
+			bool bLocalSkipAck = false;
 
-	// Acknowledge the packet.
-	if ( !bSkipAck )
-	{
-		LastGoodPacketRealtime = PostReceiveTime;
-	}
+			if (Channel != nullptr)
+			{
+				// Warning: May destroy channel
+				Channel->ReceivedRawBunch(Bunch, bLocalSkipAck);
+			}
 
-	if( !IsInternalAck() )
-	{
-		// We always call AckSequence even if we are explicitly rejecting the packet as this updates the expected InSeq used to drive future acks.
-		if ( bSkipAck )
-		{
-			// Explicit Nak, we treat this packet as dropped but we still report it to the sending side as quickly as possible
-			PacketNotify.NakSeq( InPacketId );
-		}
-		else
-		{
-			PacketNotify.AckSeq( InPacketId );
+			if ( bLocalSkipAck )
+			{
+				bOutSkipAck = true;
+			}
+			Driver->InBunches++;
+			Driver->InTotalBunches++;
 
-			// Keep stats happy
-			++OutTotalAcks;
-			++Driver->OutTotalAcks;
-		}
+			if (Bunch.IsCriticalError() || Bunch.IsError())
+			{
+				bOutHasBunchErrors = true;
 
-		// We do want to let the other side know about the ack, so even if there are no other outgoing data when we tick the connection we will send an ackpacket.
-		TimeSensitive = 1;
-		++HasDirtyAcks;
+				// Disconnect if we received a corrupted packet from the client (eg server crash attempt).
+				if (bIsServer)
+				{
+					UE_LOG(LogNetTraffic, Error, TEXT("Received corrupted packet data with SequenceId: %d from client %s. Disconnecting."), InPacketId, *LowLevelGetRemoteAddress());
 
-		// This is to allow us to recover from hitches were we process more than FNetPacketNotify::SequenceHistoryLength packets in a row withouht sending out any packets.
-		// In most cases this allows us to recover smoothly without reporting any pessimistic naks which will be the case if we overshoot the ack history.
-		// Note: This should only occur if we are running with no timeouts!
-		if (HasDirtyAcks >= FNetPacketNotify::MaxSequenceHistoryLength)
-		{
-			UE_LOG(LogNet, Warning, TEXT("UNetConnection::ReceivedPacket - Too many received packets to ack (%u) since last sent packet. InSeq: %u %s NextOutGoingSeq: %u"), HasDirtyAcks, PacketNotify.GetInSeq().Get(), *Describe(), PacketNotify.GetOutSeq().Get());
+					Close(AddToAndConsumeChainResultPtr(Bunch.ExtendedError, ENetCloseResult::CorruptData));
 
-			FlushNet();
-			if (HasDirtyAcks) // if acks still are dirty, flush again
+					return;
+				}
+				else
+				{
+					UE_LOG(LogNetTraffic, Error, TEXT("Received corrupted packet data with SequenceId: %d from server %s"), InPacketId, *LowLevelGetRemoteAddress());
+#if UE_WITH_IRIS
+					// If Iris reports errors they are unrecoverable.
+					if (Driver->GetReplicationSystem() != nullptr)
+					{
+						ensureAlwaysMsgf(false, TEXT("Received corrupted packet data with SequenceId: %d from server."), InPacketId);
+						Close(AddToAndConsumeChainResultPtr(Bunch.ExtendedError, ENetCloseResult::CorruptData));
+						return;
+					}
+#endif
+				}
+			}
+			// In replay, if the bunch generated an error but the channel isn't actually open, clean it up so the channel index remains free
+			if (IsInternalAck() && Bunch.IsError() && Channel && !Channel->OpenedLocally && !Channel->OpenAcked)
 			{
-				FlushNet();
+				UE_LOG(LogNetTraffic, Warning, TEXT("Replay cleaning up channel that couldn't be opened: %s"), *Channel->Describe());
+				Channel->ConditionalCleanUp(true, EChannelCloseReason::Destroyed);
 			}
 		}
 	}
-
-	// Flush trace content collector
-#if UE_NET_TRACE_ENABLED
-	if (InTraceCollector)
-	{
-		UE_NET_TRACE_FLUSH_COLLECTOR(InTraceCollector, NetTraceId, GetConnectionId(), ENetTracePacketType::Incoming);
-		UE_NET_TRACE_DESTROY_COLLECTOR(InTraceCollector);
-		InTraceCollector = nullptr;
-	}
-#endif
-
-	// Trace end marker of this incoming data packet.
-	UE_NET_TRACE_PACKET_RECV(NetTraceId, GetConnectionId(), InPacketId, Reader.GetNumBits());
-
-	if ((bHasBunchErrors || bSkipAck) && !IsInternalAck())
-	{
-		// Trace packet as lost on receiving end to indicate bSkipAck or errors in traced data
-		UE_NET_TRACE_PACKET_DROPPED(NetTraceId, GetConnectionId(), InPacketId, ENetTracePacketType::Incoming);
-	}
 }
 
 void UNetConnection::RestoreRemappedChannel(const int32 ChIndex)
diff --git a/Engine/Source/Runtime/Engine/Private/NetDriver.cpp b/Engine/Source/Runtime/Engine/Private/NetDriver.cpp
--- a/Engine/Source/Runtime/Engine/Private/NetDriver.cpp
+++ b/Engine/Source/Runtime/Engine/Private/NetDriver.cpp
@@ -4691,7 +4691,13 @@ int32 UNetDriver::ServerReplicateActors_PrioritizeActors( UNetConnection* Connec
 	return FinalSortedCount;
 }
 
-int32 UNetDriver::ServerReplicateActors_ProcessPrioritizedActors( UNetConnection* Connection, const TArray<FNetViewer>& ConnectionViewers, FActorPriority** PriorityActors, const int32 FinalSortedCount, int32& OutUpdated )
+int32 UNetDriver::ServerReplicateActors_ProcessPrioritizedActors(UNetConnection* Connection, const TArray<FNetViewer>& ConnectionViewers, FActorPriority** PriorityActors, const int32 FinalSortedCount, int32& OutUpdated )
+{
+	TInterval<int32> ActorsIndexRange(0, FinalSortedCount);
+	return ServerReplicateActors_ProcessPrioritizedActorsRange(Connection, ConnectionViewers, PriorityActors, ActorsIndexRange, OutUpdated);
+}
+
+int32 UNetDriver::ServerReplicateActors_ProcessPrioritizedActorsRange( UNetConnection* Connection, const TArray<FNetViewer>& ConnectionViewers, FActorPriority** PriorityActors, const TInterval<int32>& ActorsIndexRange, int32& OutUpdated, bool bIgnoreSaturation )
 {
 	SCOPE_CYCLE_COUNTER(STAT_NetProcessPrioritizedActorsTime);
 
@@ -4699,14 +4705,14 @@ int32 UNetDriver::ServerReplicateActors_ProcessPrioritizedActors( UNetConnection
 	int32 ActorUpdatesThisConnectionSent	= 0;
 	int32 FinalRelevantCount				= 0;
 
-	if (!Connection->IsNetReady( 0 ))
+	if (!Connection->IsNetReady( 0 ) && !bIgnoreSaturation)
 	{
 		GNumSaturatedConnections++;
 		// Connection saturated, don't process any actors
 		return 0;
 	}
 
-	for ( int32 j = 0; j < FinalSortedCount; j++ )
+	for ( int32 j = ActorsIndexRange.Min; j < ActorsIndexRange.Min + ActorsIndexRange.Max; j++ )
 	{
 		FNetworkObjectInfo*	ActorInfo = PriorityActors[j]->ActorInfo;
 
@@ -4819,7 +4825,7 @@ int32 UNetDriver::ServerReplicateActors_ProcessPrioritizedActors( UNetConnection
 						Channel->RelevantTime = ElapsedTime + 0.5 * UpdateDelayRandomStream.FRand();
 					}
 					// if the channel isn't saturated
-					if ( Channel->IsNetReady( 0 ) )
+					if ( Channel->IsNetReady( 0 ) || bIgnoreSaturation)
 					{
 						// replicate the actor
 						UE_LOG( LogNetTraffic, Log, TEXT( "- Replicate %s. %d" ), *Actor->GetName(), PriorityActors[j]->Priority );
@@ -4871,7 +4877,7 @@ int32 UNetDriver::ServerReplicateActors_ProcessPrioritizedActors( UNetConnection
 						Actor->ForceNetUpdate();
 					}
 					// second check for channel saturation
-					if (!Connection->IsNetReady( 0 ))
+					if (!Connection->IsNetReady( 0 ) && !bIgnoreSaturation)
 					{
 						// We can bail out now since this connection is saturated, we'll return how far we got though
 						GNumSaturatedConnections++;
@@ -4896,8 +4902,49 @@ int32 UNetDriver::ServerReplicateActors_ProcessPrioritizedActors( UNetConnection
 		}
 	}
 
-	return FinalSortedCount;
+	return ActorsIndexRange.Max;
+}
+
+void UNetDriver::ServerReplicateActors_MarkRelevantActors( UNetConnection* Connection, const TArray<FNetViewer>& ConnectionViewers, int32 StartActorIndex, int32 EndActorIndex, FActorPriority** PriorityActors )
+{
+	// relevant actors that could not be processed this frame are marked to be considered for next frame
+	for ( int32 k=StartActorIndex; k<EndActorIndex; k++ )
+	{
+		if (!PriorityActors[k]->ActorInfo)
+		{
+			// A deletion entry, skip it because we dont have anywhere to store a 'better give higher priority next time'
+			continue;
+		}
+	
+		AActor* Actor = PriorityActors[k]->ActorInfo->Actor;
+	
+		UActorChannel* Channel = PriorityActors[k]->Channel;
+		
+		UE_LOG(LogNetTraffic, Verbose, TEXT("Saturated. %s"), *Actor->GetName());
+		if (Channel != NULL && ElapsedTime - Channel->RelevantTime <= 1.0)
+		{
+			UE_LOG(LogNetTraffic, Log, TEXT(" Saturated. Mark %s NetUpdateTime to be checked for next tick"), *Actor->GetName());
+			PriorityActors[k]->ActorInfo->bPendingNetUpdate = true;
+		}
+		else if ( IsActorRelevantToConnection( Actor, ConnectionViewers ) )
+		{
+			// If this actor was relevant but didn't get processed, force another update for next frame
+			UE_LOG( LogNetTraffic, Log, TEXT( " Saturated. Mark %s NetUpdateTime to be checked for next tick" ), *Actor->GetName() );
+			PriorityActors[k]->ActorInfo->bPendingNetUpdate = true;
+			if ( Channel != NULL )
+			{
+				Channel->RelevantTime = ElapsedTime + 0.5 * UpdateDelayRandomStream.FRand();
+			}
+		}
+	
+		// If the actor was forced to relevant and didn't get processed, try again on the next update;
+		if (PriorityActors[k]->ActorInfo->ForceRelevantFrame >= Connection->LastProcessedFrame)
+		{
+			PriorityActors[k]->ActorInfo->ForceRelevantFrame = ReplicationFrame+1;
+		}
+	}
 }
+
 #endif
 
 int64 UNetDriver::SendDestructionInfo(UNetConnection* Connection, FActorDestructionInfo* DestructionInfo)
@@ -5227,6 +5274,11 @@ int32 UNetDriver::ServerReplicateActors(float DeltaSeconds)
 
 	FMemMark Mark( FMemStack::Get() );
 
+	if (OnPreConsiderListUpdateOverride.IsBound())
+	{
+		OnPreConsiderListUpdateOverride.Execute({ DeltaSeconds, nullptr, bCPUSaturated }, Updated, ConsiderList);
+	}
+
 	for ( int32 i=0; i < ClientConnections.Num(); i++ )
 	{
 		UNetConnection* Connection = ClientConnections[i];
@@ -5310,51 +5362,28 @@ int32 UNetDriver::ServerReplicateActors(float DeltaSeconds)
 
 			FMemMark RelevantActorMark(FMemStack::Get());
 
-			FActorPriority* PriorityList	= NULL;
-			FActorPriority** PriorityActors = NULL;
-
-			// Get a sorted list of actors for this connection
-			const int32 FinalSortedCount = ServerReplicateActors_PrioritizeActors( Connection, ConnectionViewers, ConsiderList, bCPUSaturated, PriorityList, PriorityActors );
+			const bool bProcessConsiderListIsBound = OnProcessConsiderListOverride.IsBound();
 
-			// Process the sorted list of actors for this connection
-			const int32 LastProcessedActor = ServerReplicateActors_ProcessPrioritizedActors( Connection, ConnectionViewers, PriorityActors, FinalSortedCount, Updated );
+			if (bProcessConsiderListIsBound)
+			{
+				OnProcessConsiderListOverride.Execute( { DeltaSeconds, Connection, bCPUSaturated }, Updated, ConsiderList );
+			}
 
-			// relevant actors that could not be processed this frame are marked to be considered for next frame
-			for ( int32 k=LastProcessedActor; k<FinalSortedCount; k++ )
+			if (!bProcessConsiderListIsBound)
 			{
-				if (!PriorityActors[k]->ActorInfo)
-				{
-					// A deletion entry, skip it because we dont have anywhere to store a 'better give higher priority next time'
-					continue;
-				}
+				FActorPriority* PriorityList = NULL;
+				FActorPriority** PriorityActors = NULL;
 
-				AActor* Actor = PriorityActors[k]->ActorInfo->Actor;
+				// Get a sorted list of actors for this connection
+				const int32 FinalSortedCount = ServerReplicateActors_PrioritizeActors(Connection, ConnectionViewers, ConsiderList, bCPUSaturated, PriorityList, PriorityActors);
 
-				UActorChannel* Channel = PriorityActors[k]->Channel;
-				
-				UE_LOG(LogNetTraffic, Verbose, TEXT("Saturated. %s"), *Actor->GetName());
-				if (Channel != NULL && ElapsedTime - Channel->RelevantTime <= 1.0)
-				{
-					UE_LOG(LogNetTraffic, Log, TEXT(" Saturated. Mark %s NetUpdateTime to be checked for next tick"), *Actor->GetName());
-					PriorityActors[k]->ActorInfo->bPendingNetUpdate = true;
-				}
-				else if ( IsActorRelevantToConnection( Actor, ConnectionViewers ) )
-				{
-					// If this actor was relevant but didn't get processed, force another update for next frame
-					UE_LOG( LogNetTraffic, Log, TEXT( " Saturated. Mark %s NetUpdateTime to be checked for next tick" ), *Actor->GetName() );
-					PriorityActors[k]->ActorInfo->bPendingNetUpdate = true;
-					if ( Channel != NULL )
-					{
-						Channel->RelevantTime = ElapsedTime + 0.5 * UpdateDelayRandomStream.FRand();
-					}
-					}
+				// Process the sorted list of actors for this connection
+				TInterval<int32> ActorsIndexRange(0, FinalSortedCount);
+				const int32 LastProcessedActor = ServerReplicateActors_ProcessPrioritizedActorsRange(Connection, ConnectionViewers, PriorityActors, ActorsIndexRange, Updated);
 
-				// If the actor was forced to relevant and didn't get processed, try again on the next update;
-				if (PriorityActors[k]->ActorInfo->ForceRelevantFrame >= Connection->LastProcessedFrame)
-				{
-					PriorityActors[k]->ActorInfo->ForceRelevantFrame = ReplicationFrame+1;
-				}
+				ServerReplicateActors_MarkRelevantActors(Connection, ConnectionViewers, LastProcessedActor, FinalSortedCount, PriorityActors);
 			}
+
 			RelevantActorMark.Pop();
 
 			ConnectionViewers.Reset();
@@ -5371,6 +5400,11 @@ int32 UNetDriver::ServerReplicateActors(float DeltaSeconds)
 		}
 	}
 
+	if (OnPostConsiderListUpdateOverride.IsBound())
+	{
+		OnPostConsiderListUpdateOverride.ExecuteIfBound( { DeltaSeconds, nullptr, bCPUSaturated }, Updated, ConsiderList );
+	}
+
 	// shuffle the list of connections if not all connections were ticked
 	if (NumClientsToTick < ClientConnections.Num())
 	{
diff --git a/Engine/Source/Runtime/Engine/Private/UnrealEngine.cpp b/Engine/Source/Runtime/Engine/Private/UnrealEngine.cpp
--- a/Engine/Source/Runtime/Engine/Private/UnrealEngine.cpp
+++ b/Engine/Source/Runtime/Engine/Private/UnrealEngine.cpp
@@ -14357,6 +14357,12 @@ EBrowseReturnVal::Type UEngine::Browse( FWorldContext& WorldContext, FURL URL, F
 		BroadcastTravelFailure(WorldContext.World(), ETravelFailure::CheatCommands, Error);
 		return EBrowseReturnVal::Failure;
 	}
+
+	if (OnOverrideBrowseURL.IsBound())
+	{
+		return OnOverrideBrowseURL.Execute(WorldContext, URL, Error);
+	}
+
 	if( URL.IsLocalInternal() )
 	{
 		// Local map file.
@@ -14575,6 +14581,11 @@ void UEngine::TickWorldTravel(FWorldContext& Context, float DeltaSeconds)
 	}
 
 	// Update the pending level.
+	if ( OnOverridePendingNetGameUpdate.ExecuteIfBound(Context, DeltaSeconds) )
+	{
+		return;
+	}
+
 	if( Context.PendingNetGame )
 	{
 		Context.PendingNetGame->Tick( DeltaSeconds );
