Fix for OOB memory access in SparseVolumeTexture bitmask based voxel compression.

#rnx

diff --git a/Engine/Source/Runtime/Engine/Private/SparseVolumeTexture/SparseVolumeTexture.cpp b/Engine/Source/Runtime/Engine/Private/SparseVolumeTexture/SparseVolumeTexture.cpp
--- a/Engine/Source/Runtime/Engine/Private/SparseVolumeTexture/SparseVolumeTexture.cpp
+++ b/Engine/Source/Runtime/Engine/Private/SparseVolumeTexture/SparseVolumeTexture.cpp
@@ -374,9 +374,9 @@ bool FResources::Build(USparseVolumeTextureFrame* Owner, UE::Serialization::FEdi
 				{
 					for (int32 AttributesIdx = 0; AttributesIdx < 2; ++AttributesIdx)
 					{
-						PrefixSums[AttributesIdx][TileIndex] = 0;
 						if (Formats[AttributesIdx] != PF_Unknown)
 						{
+							PrefixSums[AttributesIdx][TileIndex] = 0;
 							for (int64 VoxelIndex = 0; VoxelIndex < SVT::NumVoxelsPerPaddedTile; ++VoxelIndex)
 							{
 								const uint8* Src = (AttributesIdx == 0 ? PhysicalTileDataA.GetData() : PhysicalTileDataB.GetData()) + FormatSize[AttributesIdx] * (TileIndex * SVT::NumVoxelsPerPaddedTile + VoxelIndex);
@@ -429,19 +429,22 @@ bool FResources::Build(USparseVolumeTextureFrame* Owner, UE::Serialization::FEdi
 				{
 					for (int32 AttributesIdx = 0; AttributesIdx < 2; ++AttributesIdx)
 					{
-						uint32 VoxelDataWriteOffset = reinterpret_cast<uint32*>(RelativeBulkDataPtr + MipStreamingInfo.TileDataOffsetsOffset[AttributesIdx])[TileIndex];
-						const uint32* TileOccupancyBits = reinterpret_cast<uint32*>(RelativeBulkDataPtr + MipStreamingInfo.OccupancyBitsOffset[AttributesIdx] + TileIndex * SVT::NumOccupancyWordsPerPaddedTile * sizeof(uint32));
-
-						for (int64 VoxelIndex = 0; VoxelIndex < SVT::NumVoxelsPerPaddedTile; ++VoxelIndex)
+						if (Formats[AttributesIdx] != PF_Unknown)
 						{
-							const int64 WordIndex = (VoxelIndex / 32);
+							uint32 VoxelDataWriteOffset = reinterpret_cast<uint32*>(RelativeBulkDataPtr + MipStreamingInfo.TileDataOffsetsOffset[AttributesIdx])[TileIndex];
+							const uint32* TileOccupancyBits = reinterpret_cast<uint32*>(RelativeBulkDataPtr + MipStreamingInfo.OccupancyBitsOffset[AttributesIdx] + TileIndex * SVT::NumOccupancyWordsPerPaddedTile * sizeof(uint32));
 
-							if (TileOccupancyBits[WordIndex] & (1u << (static_cast<uint32>(VoxelIndex) % 32u)))
+							for (int64 VoxelIndex = 0; VoxelIndex < SVT::NumVoxelsPerPaddedTile; ++VoxelIndex)
 							{
-								uint8* Dst = DstTileData[AttributesIdx] + VoxelDataWriteOffset * FormatSize[AttributesIdx];
-								const uint8* Src = (AttributesIdx == 0 ? PhysicalTileDataA.GetData() : PhysicalTileDataB.GetData()) + FormatSize[AttributesIdx] * (TileIndex * SVT::NumVoxelsPerPaddedTile + VoxelIndex);
-								FMemory::Memcpy(Dst, Src, FormatSize[AttributesIdx]);
-								++VoxelDataWriteOffset;
+								const int64 WordIndex = (VoxelIndex / 32);
+
+								if (TileOccupancyBits[WordIndex] & (1u << (static_cast<uint32>(VoxelIndex) % 32u)))
+								{
+									uint8* Dst = DstTileData[AttributesIdx] + VoxelDataWriteOffset * FormatSize[AttributesIdx];
+									const uint8* Src = (AttributesIdx == 0 ? PhysicalTileDataA.GetData() : PhysicalTileDataB.GetData()) + FormatSize[AttributesIdx] * (TileIndex * SVT::NumVoxelsPerPaddedTile + VoxelIndex);
+									FMemory::Memcpy(Dst, Src, FormatSize[AttributesIdx]);
+									++VoxelDataWriteOffset;
+								}
 							}
 						}
 					}
