Configure FakeStereoRenderingDevice based on viewport settings.

This changes the default FOV of the emulated stereo view to 90, which matches the default FOV for single view. The width and height of the FakeStereoRenderingDevice, which are used to compute the projection matrix, are now also initialized based on the viewport resolution (instead of default constants of 640 and 480). The main purpose of this change is so that running a project with "-game -windowed -emulatestereo" with twice the resolution width of the non-stereo version will be roughly equivalent in terms of rendering.

diff --git a/Engine/Source/Runtime/Engine/Private/UnrealEngine.cpp b/Engine/Source/Runtime/Engine/Private/UnrealEngine.cpp
--- a/Engine/Source/Runtime/Engine/Private/UnrealEngine.cpp
+++ b/Engine/Source/Runtime/Engine/Private/UnrealEngine.cpp
@@ -3596,16 +3596,16 @@ bool UEngine::UseSound() const
 class FFakeStereoRenderingDevice : public IStereoRendering
 {
 public:
-	FFakeStereoRenderingDevice() 
+	FFakeStereoRenderingDevice(int ViewportWidth = 640, int ViewportHeight = 480, int RequestedNumViews = 2) 
 		: FOVInDegrees(GetDefault<UGeneralProjectSettings>()->FOVForFakeStereoRenderingDevice)
-		, Width(640)
-		, Height(480)
-		, NumViews(2)
-	{
-		static TAutoConsoleVariable<float> CVarEmulateStereoFOV(TEXT("r.StereoEmulationFOV"), 0, TEXT("FOV in degrees, of the imaginable HMD for stereo emulation"));
-		static TAutoConsoleVariable<int32> CVarEmulateStereoWidth(TEXT("r.StereoEmulationWidth"), 0, TEXT("Width of the imaginable HMD for stereo emulation"));
-		static TAutoConsoleVariable<int32> CVarEmulateStereoHeight(TEXT("r.StereoEmulationHeight"), 0, TEXT("Height of the imaginable HMD for stereo emulation"));
-		static TAutoConsoleVariable<int32> CVarEmulateStereoViews(TEXT("r.StereoEmulationViews"), 0, TEXT("Number of views in the imaginable HMD for stereo emulation"));
+		, Width(ViewportWidth)
+		, Height(ViewportHeight)
+		, NumViews(RequestedNumViews)
+	{
+		static TAutoConsoleVariable<float> CVarEmulateStereoFOV(TEXT("r.StereoEmulationFOV"), 0, TEXT("FOV in degrees, of the imaginable HMD for stereo emulation"), ECVF_ReadOnly);
+		static TAutoConsoleVariable<int32> CVarEmulateStereoWidth(TEXT("r.StereoEmulationWidth"), 0, TEXT("Width of the imaginable HMD for stereo emulation"), ECVF_ReadOnly);
+		static TAutoConsoleVariable<int32> CVarEmulateStereoHeight(TEXT("r.StereoEmulationHeight"), 0, TEXT("Height of the imaginable HMD for stereo emulation"), ECVF_ReadOnly);
+		static TAutoConsoleVariable<int32> CVarEmulateStereoViews(TEXT("r.StereoEmulationViews"), 0, TEXT("Number of views in the imaginable HMD for stereo emulation"), ECVF_ReadOnly);
 		float FOV = CVarEmulateStereoFOV.GetValueOnAnyThread();
 		if (FOV != 0)
 		{
@@ -3722,10 +3722,10 @@ bool UEngine::InitializeHMDDevice()
 {
 	if (!IsRunningCommandlet())
 	{
-		static TAutoConsoleVariable<int32> CVarEmulateStereo(TEXT("r.EnableStereoEmulation"), 0, TEXT("Emulate stereo rendering"));
+		static TAutoConsoleVariable<int32> CVarEmulateStereo(TEXT("r.EnableStereoEmulation"), 0, TEXT("Emulate stereo rendering"), ECVF_ReadOnly);
 		if (FParse::Param(FCommandLine::Get(), TEXT("emulatestereo")) || CVarEmulateStereo.GetValueOnAnyThread() != 0)
 		{
-			TSharedPtr<FFakeStereoRenderingDevice, ESPMode::ThreadSafe> FakeStereoDevice(new FFakeStereoRenderingDevice());
+			TSharedPtr<FFakeStereoRenderingDevice, ESPMode::ThreadSafe> FakeStereoDevice(new FFakeStereoRenderingDevice(GSystemResolution.ResX / 2, GSystemResolution.ResY));
 			StereoRenderingDevice = FakeStereoDevice;
 		}
 		// No reason to connect an HMD on a dedicated server.  Also fixes dedicated servers stealing the oculus connection.
diff --git a/Engine/Source/Runtime/EngineSettings/Private/EngineSettingsModule.cpp b/Engine/Source/Runtime/EngineSettings/Private/EngineSettingsModule.cpp
--- a/Engine/Source/Runtime/EngineSettings/Private/EngineSettingsModule.cpp
+++ b/Engine/Source/Runtime/EngineSettings/Private/EngineSettingsModule.cpp
@@ -74,7 +74,7 @@ UGeneralProjectSettings::UGeneralProjectSettings( const FObjectInitializer& Obje
 	, bAllowMaximize(true)
 	, bAllowMinimize(true)
 	, EyeOffsetForFakeStereoRenderingDevice(0.0320000005f)
-	, FOVForFakeStereoRenderingDevice(100)
+	, FOVForFakeStereoRenderingDevice(90)
 { }
 
 
