substrate - added static assert post feedback.

diff --git a/Engine/Shaders/Shared/StrataDefinitions.h b/Engine/Shaders/Shared/StrataDefinitions.h
--- a/Engine/Shaders/Shared/StrataDefinitions.h
+++ b/Engine/Shaders/Shared/StrataDefinitions.h
@@ -10,12 +10,12 @@
 // Change this to force recompilation of all strata dependent shaders (use https://www.random.org/cgi-bin/randbyte?nbytes=4&format=h)
 #define STRATA_SHADER_VERSION 0x56F06029 
 
-// BSDF offsets are packed into 32bits to max. BSDF is < STRATA_MAX_BSDF_COUNT, which could leads to some issue.
+// BSDF offsets are packed into 32bits, each entry using STRATA_BSDF_OFFSET_BIT_COUNT bits
 #define STRATA_MAX_BSDF_COUNT_FOR_BDSFOFFSET	8u
 #define STRATA_BSDF_OFFSET_BIT_COUNT			4u
 #define STRATA_BSDF_OFFSET_BIT_MASK				0xF
 
-// We can only ever use STRATA_MAX_BSDF_COUNT_FOR_BDSFOFFSET BSDFs for Lumen, se we use that as a global BSDF count limit today.
+// We can only ever use STRATA_MAX_BSDF_COUNT_FOR_BDSFOFFSET BSDFs for Lumen, so we use that as a global BSDF count limit today.
 #define STRATA_MAX_BSDF_COUNT				STRATA_MAX_BSDF_COUNT_FOR_BDSFOFFSET
 #define STRATA_MAX_OPERATOR_COUNT			15
 
diff --git a/Engine/Source/Runtime/Engine/Private/Materials/HLSLMaterialTranslator.cpp b/Engine/Source/Runtime/Engine/Private/Materials/HLSLMaterialTranslator.cpp
--- a/Engine/Source/Runtime/Engine/Private/Materials/HLSLMaterialTranslator.cpp
+++ b/Engine/Source/Runtime/Engine/Private/Materials/HLSLMaterialTranslator.cpp
@@ -12002,10 +12002,11 @@ bool FHLSLMaterialTranslator::FStrataCompilationContext::StrataGenerateDerivedMa
 			//  - We do not need to account for the Top Normal texture when evaluating the material byte count for the optimization algorithm.
 			//  - This is because we only need to optimize for the Strata uint material buffer.
 
-
 			// 2. Process the list of BSDFs for worst case memory usage and count operators.
+			static_assert(STRATA_MAX_BSDF_COUNT_FOR_BDSFOFFSET	== (32u / STRATA_BSDF_OFFSET_BIT_COUNT));
+			static_assert(STRATA_MAX_BSDF_COUNT					<= (1u << STRATA_BSDF_OFFSET_BIT_COUNT));
+			const uint32 BSDFMaxByteCountForOffset = uint32(1u << STRATA_BSDF_OFFSET_BIT_COUNT) * sizeof(uint32);
 			uint32 OperatorCount = 0;
-			const uint32 BSDFMaxByteCount = uint32(FMath::Pow(2.0f, float(STRATA_BSDF_OFFSET_BIT_COUNT))) * sizeof(uint32);
 			for (auto& It : StrataMaterialExpressionRegisteredOperators)
 			{
 				if (It.IsDiscarded())
@@ -12137,9 +12138,9 @@ bool FHLSLMaterialTranslator::FStrataCompilationContext::StrataGenerateDerivedMa
 				} // switch (It.OperatorType)
 
 				const uint32 BSDFRequestedSizeByte = StrataMaterialRequestedSizeByte - PreStrataMaterialRequestedSizeByte;
-				if (BSDFRequestedSizeByte > BSDFMaxByteCount)
+				if (BSDFRequestedSizeByte > BSDFMaxByteCountForOffset)
 				{
-					Compiler->Errorf(TEXT("A BSDF is requesting more bytes than our BSDF offset system can handle (%d/%d bytes). Notify your rendering engineer. Material %s (asset: %s).\r\n"), BSDFRequestedSizeByte, BSDFMaxByteCount, *CompilerMaterial->GetDebugName(), *CompilerMaterial->GetAssetPath().ToString(), OperatorCount, STRATA_MAX_OPERATOR_COUNT);
+					Compiler->Errorf(TEXT("A BSDF is requesting more bytes than our BSDF offset system can handle (%d/%d bytes). Notify your rendering engineer. Material %s (asset: %s).\r\n"), BSDFRequestedSizeByte, BSDFMaxByteCountForOffset, *CompilerMaterial->GetDebugName(), *CompilerMaterial->GetAssetPath().ToString(), OperatorCount, STRATA_MAX_OPERATOR_COUNT);
 					return false;
 				}
 			}
