#fn - Add FScopedMovementUpdate around movement of interpolated component in ProjectileMovementComponent. This avoids moving its child components more than once if it then again simulates during the tick.

Added bInterpolationUseScopedMovement var on PMC to control this per instance (default=true).

diff --git a/Engine/Source/Runtime/Engine/Classes/GameFramework/ProjectileMovementComponent.h b/Engine/Source/Runtime/Engine/Classes/GameFramework/ProjectileMovementComponent.h
--- a/Engine/Source/Runtime/Engine/Classes/GameFramework/ProjectileMovementComponent.h
+++ b/Engine/Source/Runtime/Engine/Classes/GameFramework/ProjectileMovementComponent.h
@@ -119,6 +119,14 @@ protected:
 	uint8 bInterpolationComplete:1;
 
 public:
+
+	/**
+	 * If true, uses FScopedMovementUpdate to avoid moving the attached interpolated object's children more than once during a tick when it would both interpolate and move during projectile simulation.
+	 * This also defers overlap updates for the interpolated object until after the simulation update completes.
+	 */
+	UPROPERTY(EditAnywhere, BlueprintReadWrite, AdvancedDisplay, Category=ProjectileInterpolation)
+	uint8 bInterpolationUseScopedMovement:1;
+
 	/** Saved HitResult Time (0 to 1) from previous simulation step. Equal to 1.0 when there was no impact. */
 	UPROPERTY(VisibleInstanceOnly, BlueprintReadOnly, Category=ProjectileBounces)
 	float PreviousHitTime;
diff --git a/Engine/Source/Runtime/Engine/Private/Components/ProjectileMovementComponent.cpp b/Engine/Source/Runtime/Engine/Private/Components/ProjectileMovementComponent.cpp
--- a/Engine/Source/Runtime/Engine/Private/Components/ProjectileMovementComponent.cpp
+++ b/Engine/Source/Runtime/Engine/Private/Components/ProjectileMovementComponent.cpp
@@ -25,6 +25,7 @@ UProjectileMovementComponent::UProjectileMovementComponent(const FObjectInitiali
 	bInterpMovement = false;
 	bInterpRotation = false;
 	bInterpolationComplete = true;
+	bInterpolationUseScopedMovement = true;
 	InterpLocationTime = 0.100f;
 	InterpRotationTime = 0.050f;
 	InterpLocationMaxLagDistance = 300.0f;
@@ -134,6 +135,11 @@ void UProjectileMovementComponent::TickComponent(float DeltaTime, enum ELevelTic
 	QUICK_SCOPE_CYCLE_COUNTER( STAT_ProjectileMovementComponent_TickComponent );
 	CSV_SCOPED_TIMING_STAT_EXCLUSIVE(ProjectileMovement);
 
+	// Can avoid moving the interpolated object's children until the end of the entire simulation frame.
+	// This only makes sense if simulation is also enabled, which would move the UpdatedComponent and move the attached InterpolatedComponent (and children) again.
+	const bool bUseScopedInterpolatedMove = bInterpolationUseScopedMovement && bSimulationEnabled;
+	const FScopedMovementUpdate ScopedInterpolatedMove(GetInterpolatedComponent(), bUseScopedInterpolatedMove ? EScopedUpdate::DeferredUpdates : EScopedUpdate::ImmediateUpdates);
+
 	// Still need to finish interpolating after we've stopped simulating, so do that first.
 	if (!bInterpolationComplete)
 	{
