Skelmesh: Fix performance issues on meshes with large number of morph targets.

diff --git a/Engine/Source/Runtime/Engine/Private/Components/SkinnedMeshComponent.cpp b/Engine/Source/Runtime/Engine/Private/Components/SkinnedMeshComponent.cpp
--- a/Engine/Source/Runtime/Engine/Private/Components/SkinnedMeshComponent.cpp
+++ b/Engine/Source/Runtime/Engine/Private/Components/SkinnedMeshComponent.cpp
@@ -463,38 +463,54 @@ USkinnedMeshComponent::~USkinnedMeshComponent() = default;
 
 void USkinnedMeshComponent::UpdateMorphMaterialUsageOnProxy()
 {
-	// update morph material usage
-	if (SceneProxy)
+	if (!SceneProxy)
 	{
-		if (ActiveMorphTargets.Num() > 0 && GetSkinnedAsset()->GetMorphTargets().Num() > 0)
+		return;
+	}
+	TArray<UMaterialInterface*> MaterialUsingMorphTarget;
+
+	// Collect all sections on the underlying mesh that are going to be modified by morph targets, 
+	if (!ActiveMorphTargets.IsEmpty())
+	{
+		TSet<int32> UsedMaterialIndices;
+
+		const TIndirectArray<FSkeletalMeshLODRenderData>& LODRenderData = GetSkinnedAsset()->GetResourceForRendering()->LODRenderData;
+		for (const TTuple<const UMorphTarget*, int32> ActiveMorphTargetInfo: ActiveMorphTargets)
 		{
-			TArray<UMaterialInterface*> MaterialUsingMorphTarget;
-			for (UMorphTarget* MorphTarget : GetSkinnedAsset()->GetMorphTargets())
+			const UMorphTarget* MorphTarget = ActiveMorphTargetInfo.Key;
+			if (!MorphTarget)
 			{
-				if (!MorphTarget)
-				{
-					continue;
-				}
-				for (const FMorphTargetLODModel& MorphTargetLODModel : MorphTarget->GetMorphLODModels())
+				continue;
+			}
+
+			const TArray<FMorphTargetLODModel>& MorphLODModels = MorphTarget->GetMorphLODModels();
+			if (ensure(MorphLODModels.Num() == LODRenderData.Num()))
+			{
+				for (int32 LODIndex = 0; LODIndex < MorphLODModels.Num(); LODIndex++)
 				{
-					for (int32 SectionIndex : MorphTargetLODModel.SectionIndices)
+					const FSkeletalMeshLODRenderData& RenderData = LODRenderData[LODIndex];
+					
+					for(const int32 MorphSectionIndex: MorphLODModels[LODIndex].SectionIndices)
 					{
-						for (int32 LodIdx = 0; LodIdx < GetSkinnedAsset()->GetResourceForRendering()->LODRenderData.Num(); LodIdx++)
+						if (ensure(RenderData.RenderSections.IsValidIndex(MorphSectionIndex)))
 						{
-							const FSkeletalMeshLODRenderData& LODModel = GetSkinnedAsset()->GetResourceForRendering()->LODRenderData[LodIdx];
-							if (LODModel.RenderSections.IsValidIndex(SectionIndex))
-							{
-								MaterialUsingMorphTarget.AddUnique(GetMaterial(LODModel.RenderSections[SectionIndex].MaterialIndex));
-							}
+							UsedMaterialIndices.Add(RenderData.RenderSections[MorphSectionIndex].MaterialIndex);
 						}
 					}
 				}
 			}
-			((FSkeletalMeshSceneProxy*)SceneProxy)->UpdateMorphMaterialUsage_GameThread(MaterialUsingMorphTarget);
+		}
+		for (const int32 MaterialIndex: UsedMaterialIndices)
+		{
+			MaterialUsingMorphTarget.Add(GetMaterial(MaterialIndex));
 		}
 	}
+
+	// If no morph targets are active, then this function needs to know that as well.
+	static_cast<FSkeletalMeshSceneProxy*>(SceneProxy)->UpdateMorphMaterialUsage_GameThread(MaterialUsingMorphTarget);
 }
 
+
 void USkinnedMeshComponent::GetResourceSizeEx(FResourceSizeEx& CumulativeResourceSize)
 {
 	Super::GetResourceSizeEx(CumulativeResourceSize);
