Allow cheat extensions to add bugit logging. Allows games to add bugit logging without overriding the cheat manager class.

diff --git a/Engine/Source/Runtime/Engine/Classes/GameFramework/CheatManager.h b/Engine/Source/Runtime/Engine/Classes/GameFramework/CheatManager.h
--- a/Engine/Source/Runtime/Engine/Classes/GameFramework/CheatManager.h
+++ b/Engine/Source/Runtime/Engine/Classes/GameFramework/CheatManager.h
@@ -84,6 +84,8 @@ public:
 
 	UFUNCTION(BlueprintImplementableEvent, meta = (DisplayName = "Removed From Cheat Manager", Keywords = "End Play"))
 	ENGINE_API void RemovedFromCheatManager();
+
+	ENGINE_API virtual void DoExtensionSpecificBugItLog(FOutputDevice& OutputFile);
 };
 
 /** 
@@ -466,7 +468,7 @@ protected:
 
 protected:
 	/** Do game specific bugIt */
-	virtual bool DoGameSpecificBugItLog(FOutputDevice& OutputFile) { return true; }
+	ENGINE_API virtual bool DoGameSpecificBugItLog(FOutputDevice& OutputFile);
 
 	/** Switch controller to debug camera without locking gameplay and with locking local player controller input */
 	UFUNCTION(BlueprintCallable,Category="Cheat Manager")
diff --git a/Engine/Source/Runtime/Engine/Private/CheatManager.cpp b/Engine/Source/Runtime/Engine/Private/CheatManager.cpp
--- a/Engine/Source/Runtime/Engine/Private/CheatManager.cpp
+++ b/Engine/Source/Runtime/Engine/Private/CheatManager.cpp
@@ -55,6 +55,10 @@ APlayerController* UCheatManagerExtension::GetPlayerController() const
 	return GetOuterUCheatManager()->GetPlayerController();
 }
 
+void UCheatManagerExtension::DoExtensionSpecificBugItLog(FOutputDevice& OutputFile)
+{
+}
+
 UCheatManager::UCheatManager(const FObjectInitializer& ObjectInitializer)
 	: Super(ObjectInitializer)
 	, bToggleAILogging(false)
@@ -1263,6 +1267,19 @@ void UCheatManager::LogOutBugItGoToLogFile( const FString& InScreenShotDesc, con
 #endif // ALLOW_DEBUG_FILES
 }
 
+bool UCheatManager::DoGameSpecificBugItLog(FOutputDevice& OutputFile)
+{
+	for (UCheatManagerExtension* Extension : CheatManagerExtensions)
+	{
+		if (Extension)
+		{
+			Extension->DoExtensionSpecificBugItLog(OutputFile);
+		}
+	}
+
+	return true;
+}
+
 AActor* UCheatManager::GetTarget(APlayerController* PlayerController, struct FHitResult& OutHit)
 {
     if ((PlayerController == NULL) || (PlayerController->PlayerCameraManager == NULL))
