Level Instance: Fix Packing on Streaming enabled Level Instances

#rnx

diff --git a/Engine/Source/Runtime/Engine/Private/LevelInstance/LevelInstanceInterface.cpp b/Engine/Source/Runtime/Engine/Private/LevelInstance/LevelInstanceInterface.cpp
--- a/Engine/Source/Runtime/Engine/Private/LevelInstance/LevelInstanceInterface.cpp
+++ b/Engine/Source/Runtime/Engine/Private/LevelInstance/LevelInstanceInterface.cpp
@@ -25,38 +25,34 @@ bool ILevelInstanceInterface::SupportsPartialEditorLoading() const
 	{
 		if (ULevelInstanceSubsystem* LevelInstanceSubsystem = GetLevelInstanceSubsystem())
 		{
-			bool bResult = true;
+			if (LevelInstanceSubsystem->IsEditingLevelInstance(this))
+			{
+				return false;
+			}
+
 			const AActor* Actor = CastChecked<AActor>(this);
 
-			LevelInstanceSubsystem->ForEachLevelInstanceAncestorsAndSelf(Actor, [LevelInstanceSubsystem, &bResult](const ILevelInstanceInterface* LevelInstance)
+			// If the level instance actor (or any parent actor) is not spatially loaded, don't partially load.
+			if (!Actor->GetIsSpatiallyLoaded())
 			{
-				// If the level instance (or any parent level instance) is in edit mode, don't partially load.
-				if (LevelInstanceSubsystem->IsEditingLevelInstance(LevelInstance))
-				{
-					bResult = false;
-					return false;
-				}
+				return false;
+			}
 
-				// If the level instance actor (or any parent actor) is not spatially loaded, don't partially load.
-				const AActor* LevelInstanceActor = CastChecked<AActor>(LevelInstance);
-				if (!LevelInstanceActor->GetIsSpatiallyLoaded())
-				{
-					bResult = false;
-					return false;
-				}
+			// If the level instance actor (or any parent actor) is unsaved, don't partially load.
+			if (Actor->GetPackage()->HasAllPackagesFlags(PKG_NewlyCreated))
+			{
+				return false;
+			}
 
-				// If the level instance actor (or any parent actor) is unsaved, don't partially load.
-				if (LevelInstanceActor->GetPackage()->HasAllPackagesFlags(PKG_NewlyCreated))
+			if (ILevelInstanceInterface* Parent = LevelInstanceSubsystem->GetParentLevelInstance(Actor))
+			{
+				if (!Parent->SupportsPartialEditorLoading())
 				{
-					bResult = false;
 					return false;
 				}
+			}
 
-				// We can partially load
-				return true;
-			});
-
-			return bResult;
+			return true;
 		}
 	}
 
diff --git a/Engine/Source/Runtime/Engine/Private/LevelInstance/LevelInstanceLevelStreaming.cpp b/Engine/Source/Runtime/Engine/Private/LevelInstance/LevelInstanceLevelStreaming.cpp
--- a/Engine/Source/Runtime/Engine/Private/LevelInstance/LevelInstanceLevelStreaming.cpp
+++ b/Engine/Source/Runtime/Engine/Private/LevelInstance/LevelInstanceLevelStreaming.cpp
@@ -377,9 +377,9 @@ void ULevelStreamingLevelInstance::OnLevelLoadedChanged(ULevel* InLevel)
 					{
 						OuterWorldPartition->bOverrideEnableStreamingInEditor = false;
 					}
-					else if (ILevelInstanceInterface* LevelInstance = LevelInstanceSubsystem->GetLevelInstance(LevelInstanceID); LevelInstance && LevelInstance->SupportsPartialEditorLoading())
+					else if (ILevelInstanceInterface* LevelInstance = LevelInstanceSubsystem->GetLevelInstance(LevelInstanceID))
 					{
-						OuterWorldPartition->bOverrideEnableStreamingInEditor = true;
+						OuterWorldPartition->bOverrideEnableStreamingInEditor = LevelInstance->SupportsPartialEditorLoading();
 					}
 				}
 				else
