World Partition: update actor descriptors when changing BPs.

diff --git a/Engine/Source/Runtime/Engine/Private/WorldPartition/ActorDescContainer.cpp b/Engine/Source/Runtime/Engine/Private/WorldPartition/ActorDescContainer.cpp
--- a/Engine/Source/Runtime/Engine/Private/WorldPartition/ActorDescContainer.cpp
+++ b/Engine/Source/Runtime/Engine/Private/WorldPartition/ActorDescContainer.cpp
@@ -6,6 +6,7 @@
 
 #if WITH_EDITOR
 #include "Editor.h"
+#include "Algo/Transform.h"
 #include "Engine/Level.h"
 #include "UObject/ObjectSaveContext.h"
 #include "AssetRegistry/ARFilter.h"
@@ -52,8 +53,8 @@ void UActorDescContainer::Initialize(const FInitializeParams& InitParams)
 		// Do a synchronous scan of the level external actors path.					
 		IAssetRegistry& AssetRegistry = FModuleManager::LoadModuleChecked<FAssetRegistryModule>(TEXT("AssetRegistry")).Get();
 		{
-			TRACE_CPUPROFILER_EVENT_SCOPE(ScanPathsSynchronous);
-			AssetRegistry.ScanPathsSynchronous({ ContainerExternalActorsPath }, /*bForceRescan*/false, /*bIgnoreDenyListScanFilters*/false);
+			TRACE_CPUPROFILER_EVENT_SCOPE(ScanSynchronous);
+			AssetRegistry.ScanSynchronous({ ContainerExternalActorsPath }, TArray<FString>());
 		}
 
 		FARFilter Filter;
@@ -360,6 +361,55 @@ void UActorDescContainer::OnObjectsReplaced(const TMap<UObject*, UObject*>& OldT
 	}
 }
 
+void UActorDescContainer::OnClassDescriptorUpdated(const FWorldPartitionActorDesc* InClassDesc)
+{
+	FWorldPartitionClassDescRegistry& ClassDescRegistry = FWorldPartitionClassDescRegistry::Get();
+
+	TArray<FString> ActorPackages;
+	for (FActorDescList::TIterator<> ActorDescIterator(this); ActorDescIterator; ++ActorDescIterator)
+	{
+		if (ActorDescIterator->GetBaseClass().IsValid())
+		{
+			if (const FWorldPartitionActorDesc* ActorClassDesc = ClassDescRegistry.GetClassDescDefaultForActor(ActorDescIterator->GetBaseClass()))
+			{
+				if (ClassDescRegistry.IsDerivedFrom(ActorClassDesc, InClassDesc))
+				{
+					ActorPackages.Add(ActorDescIterator->GetActorPackage().ToString());
+				}
+			}
+		}
+	}
+
+	if (ActorPackages.Num())
+	{
+		FARFilter Filter;
+		Filter.bIncludeOnlyOnDiskAssets = true;
+		Filter.PackageNames.Reserve(ActorPackages.Num());
+		Algo::Transform(ActorPackages, Filter.PackageNames, [](const FString& ActorPath) { return *ActorPath; });
+
+		IAssetRegistry& AssetRegistry = FModuleManager::LoadModuleChecked<FAssetRegistryModule>(TEXT("AssetRegistry")).Get();
+		AssetRegistry.ScanSynchronous(TArray<FString>(), ActorPackages);
+
+		TArray<FAssetData> Assets;
+		AssetRegistry.GetAssets(Filter, Assets);
+
+		for (const FAssetData& Asset : Assets)
+		{
+			TUniquePtr<FWorldPartitionActorDesc> NewActorDesc = FWorldPartitionActorDescUtils::GetActorDescriptorFromAssetData(Asset);
+								
+			if (NewActorDesc.IsValid() && NewActorDesc->GetNativeClass().IsValid())
+			{
+				if (TUniquePtr<FWorldPartitionActorDesc>* ExistingActorDesc = GetActorDescriptor(NewActorDesc->GetGuid()))
+				{
+					OnActorDescUpdating(ExistingActorDesc->Get());
+					FWorldPartitionActorDescUtils::UpdateActorDescriptorFromActorDescriptor(NewActorDesc, *ExistingActorDesc);
+					OnActorDescUpdated(ExistingActorDesc->Get());
+				}
+			}
+		}
+	}
+}
+
 bool UActorDescContainer::RemoveActor(const FGuid& ActorGuid)
 {
 	if (TUniquePtr<FWorldPartitionActorDesc>* ExistingActorDesc = GetActorDescriptor(ActorGuid))
@@ -402,6 +452,9 @@ void UActorDescContainer::RegisterEditorDelegates()
 		FCoreUObjectDelegates::OnObjectPreSave.AddUObject(this, &UActorDescContainer::OnObjectPreSave);
 		FEditorDelegates::OnPackageDeleted.AddUObject(this, &UActorDescContainer::OnPackageDeleted);
 		FCoreUObjectDelegates::OnObjectsReplaced.AddUObject(this, &UActorDescContainer::OnObjectsReplaced);
+		
+		FWorldPartitionClassDescRegistry& ClassDescRegistry = FWorldPartitionClassDescRegistry::Get();
+		ClassDescRegistry.OnClassDescriptorUpdated().AddUObject(this, &UActorDescContainer::OnClassDescriptorUpdated);
 	}
 }
 
@@ -412,6 +465,9 @@ void UActorDescContainer::UnregisterEditorDelegates()
 		FCoreUObjectDelegates::OnObjectPreSave.RemoveAll(this);
 		FEditorDelegates::OnPackageDeleted.RemoveAll(this);
 		FCoreUObjectDelegates::OnObjectsReplaced.RemoveAll(this);
+
+		FWorldPartitionClassDescRegistry& ClassDescRegistry = FWorldPartitionClassDescRegistry::Get();
+		ClassDescRegistry.OnClassDescriptorUpdated().RemoveAll(this);
 	}
 }
 
diff --git a/Engine/Source/Runtime/Engine/Private/WorldPartition/WorldPartitionClassDescRegistry.cpp b/Engine/Source/Runtime/Engine/Private/WorldPartition/WorldPartitionClassDescRegistry.cpp
--- a/Engine/Source/Runtime/Engine/Private/WorldPartition/WorldPartitionClassDescRegistry.cpp
+++ b/Engine/Source/Runtime/Engine/Private/WorldPartition/WorldPartitionClassDescRegistry.cpp
@@ -215,7 +215,7 @@ void FWorldPartitionClassDescRegistry::PrefetchClassDescs(const TArray<FTopLevel
 		Filter.PackageNames.Reserve(FilePaths.Num());
 		Algo::Transform(FilePaths, Filter.PackageNames, [](const FString& ClassPath) { return *ClassPath; });
 	
-		AssetRegistry.ScanFilesSynchronous(FilePaths, /*bForceRescan*/false);
+		AssetRegistry.ScanSynchronous(TArray<FString>(), FilePaths);
 		AssetRegistry.GetAssets(Filter, Assets);
 	};
 
@@ -356,6 +356,25 @@ bool FWorldPartitionClassDescRegistry::IsRegisteredClass(const FTopLevelAssetPat
 	return ClassByPath.Contains(ClassPath);
 }
 
+bool FWorldPartitionClassDescRegistry::IsDerivedFrom(const FWorldPartitionActorDesc* InClassDesc, const FWorldPartitionActorDesc* InParentClassDesc) const
+{
+	const FTopLevelAssetPath ParentClassPath = FTopLevelAssetPath(InParentClassDesc->GetActorSoftPath().ToString());
+
+	FTopLevelAssetPath ClassPath = FTopLevelAssetPath(InClassDesc->GetActorSoftPath().ToString());
+
+	while (ClassPath.IsValid())
+	{
+		if (ClassPath == ParentClassPath)
+		{
+			return true;
+		}
+
+		ClassPath = ParentClassMap.FindRef(ClassPath);
+	}
+
+	return false;
+}
+
 const FWorldPartitionActorDesc* FWorldPartitionClassDescRegistry::GetClassDescDefault(const FTopLevelAssetPath& InClassPath) const
 {
 	check(IsInitialized());
@@ -637,6 +656,8 @@ void FWorldPartitionClassDescRegistry::UpdateClassDescriptor(UObject* InObject,
 
 			ValidateInternalState();
 		}
+
+		ClassDescriptorUpdatedEvent.Broadcast(ExistingClassDesc->Get());
 	}
 }
 
diff --git a/Engine/Source/Runtime/Engine/Public/WorldPartition/ActorDescContainer.h b/Engine/Source/Runtime/Engine/Public/WorldPartition/ActorDescContainer.h
--- a/Engine/Source/Runtime/Engine/Public/WorldPartition/ActorDescContainer.h
+++ b/Engine/Source/Runtime/Engine/Public/WorldPartition/ActorDescContainer.h
@@ -54,6 +54,7 @@ public:
 	ENGINE_API void OnObjectPreSave(UObject* Object, FObjectPreSaveContext SaveContext);
 	ENGINE_API void OnPackageDeleted(UPackage* Package);
 	ENGINE_API void OnObjectsReplaced(const TMap<UObject*, UObject*>& OldToNewObjectMap);
+	ENGINE_API void OnClassDescriptorUpdated(const FWorldPartitionActorDesc* InClassDesc);
 
 	FName GetContainerPackage() const { return ContainerPackageName; }
 	void SetContainerPackage(const FName& InContainerPackageName) { ContainerPackageName = InContainerPackageName; }
diff --git a/Engine/Source/Runtime/Engine/Public/WorldPartition/WorldPartitionClassDescRegistry.h b/Engine/Source/Runtime/Engine/Public/WorldPartition/WorldPartitionClassDescRegistry.h
--- a/Engine/Source/Runtime/Engine/Public/WorldPartition/WorldPartitionClassDescRegistry.h
+++ b/Engine/Source/Runtime/Engine/Public/WorldPartition/WorldPartitionClassDescRegistry.h
@@ -29,7 +29,10 @@ public:
 
 	ENGINE_API bool IsRegisteredClass(const FTopLevelAssetPath& InClassPath) const;
 
-	const FParentClassMap& GetParentClassMap() const { check(IsInitialized()); return ParentClassMap; }
+	ENGINE_API bool IsDerivedFrom(const FWorldPartitionActorDesc* InClassDesc, const FWorldPartitionActorDesc* InParentClassDesc) const;
+
+	DECLARE_EVENT_OneParam(FWorldPartitionClassDescRegistry, FClassDescriptorUpdated, const FWorldPartitionActorDesc*);
+	FClassDescriptorUpdated& OnClassDescriptorUpdated() { return ClassDescriptorUpdatedEvent; }
 
 private:
 	ENGINE_API void PrefetchClassDesc(UClass* InClass);
@@ -41,10 +44,14 @@ private:
 	ENGINE_API void RegisterClassDescriptorFromActorClass(const UClass* InActorClass);
 
 	friend class FActorDescArchive;
+	friend class UActorDescContainer;
 	ENGINE_API const FWorldPartitionActorDesc* GetClassDescDefault(const FTopLevelAssetPath& InClassPath) const;
 	ENGINE_API const FWorldPartitionActorDesc* GetClassDescDefaultForActor(const FTopLevelAssetPath& InClassPath) const;
 	ENGINE_API const FWorldPartitionActorDesc* GetClassDescDefaultForClass(const FTopLevelAssetPath& InClassPath) const;
 
+	friend class UWorldPartitionResaveActorsBuilder;
+	ENGINE_API const FParentClassMap& GetParentClassMap() const { check(IsInitialized()); return ParentClassMap; }
+
 	ENGINE_API void OnObjectPreSave(UObject* InObject, FObjectPreSaveContext InSaveContext);
 	ENGINE_API void OnObjectPropertyChanged(UObject* InObject, FPropertyChangedEvent& InPropertyChangedEvent);
 
@@ -65,5 +72,6 @@ private:
 	FNameClassDescMap ClassByPath;
 	FParentClassMap ParentClassMap;
 	FRedirectClassMap RedirectClassMap;
+	FClassDescriptorUpdated ClassDescriptorUpdatedEvent;
 };
 #endif
