Moved duplicate logic from FDistanceFieldVolumeData/FCardRepresentationData::CacheDerivedData(...) to BuildSignedDistanceFieldBuildMaterialData(...)

diff --git a/Engine/Source/Runtime/Engine/Private/DistanceFieldAtlas.cpp b/Engine/Source/Runtime/Engine/Private/DistanceFieldAtlas.cpp
--- a/Engine/Source/Runtime/Engine/Private/DistanceFieldAtlas.cpp
+++ b/Engine/Source/Runtime/Engine/Private/DistanceFieldAtlas.cpp
@@ -147,28 +147,24 @@ FString BuildDistanceFieldDerivedDataKey(const FString& InMeshKey)
 
 #if WITH_EDITORONLY_DATA
 
-void FDistanceFieldVolumeData::CacheDerivedData(const FString& InStaticMeshDerivedDataKey, const ITargetPlatform* TargetPlatform, UStaticMesh* Mesh, FStaticMeshRenderData& RenderData, UStaticMesh* GenerateSource, float DistanceFieldResolutionScale, bool bGenerateDistanceFieldAsIfTwoSided)
+void BuildSignedDistanceFieldBuildMaterialData(UStaticMesh* Mesh, uint32 LODIndex, TArray<FSignedDistanceFieldBuildMaterialData>& OutData)
 {
-	TRACE_CPUPROFILER_EVENT_SCOPE(FDistanceFieldVolumeData::CacheDerivedData);
-
 	const TArray<FStaticMaterial>& StaticMaterials = Mesh->GetStaticMaterials();
 
-	TArray<FSignedDistanceFieldBuildMaterialData> BuildMaterialData;
-	BuildMaterialData.SetNum(StaticMaterials.Num());
+	OutData.SetNum(StaticMaterials.Num());
 
 	const FMeshSectionInfoMap& SectionInfoMap = Mesh->GetSectionInfoMap();
-	const uint32 LODIndex = 0;
 
 	for (int32 SectionIndex = 0; SectionIndex < SectionInfoMap.GetSectionNumber(LODIndex); SectionIndex++)
 	{
 		const FMeshSectionInfo& Section = SectionInfoMap.Get(LODIndex, SectionIndex);
 
-		if (!BuildMaterialData.IsValidIndex(Section.MaterialIndex))
+		if (!OutData.IsValidIndex(Section.MaterialIndex))
 		{
 			continue;
 		}
 
-		FSignedDistanceFieldBuildMaterialData& MaterialData = BuildMaterialData[Section.MaterialIndex];
+		FSignedDistanceFieldBuildMaterialData& MaterialData = OutData[Section.MaterialIndex];
 		MaterialData.bAffectDistanceFieldLighting = Section.bAffectDistanceFieldLighting;
 
 		UMaterialInterface* MaterialInterface = StaticMaterials[Section.MaterialIndex].MaterialInterface;
@@ -178,10 +174,20 @@ void FDistanceFieldVolumeData::CacheDerivedData(const FString& InStaticMeshDeriv
 			MaterialData.bTwoSided = MaterialInterface->IsTwoSided();
 		}
 	}
+}
+
+void FDistanceFieldVolumeData::CacheDerivedData(const FString& InStaticMeshDerivedDataKey, const ITargetPlatform* TargetPlatform, UStaticMesh* Mesh, FStaticMeshRenderData& RenderData, UStaticMesh* GenerateSource, float DistanceFieldResolutionScale, bool bGenerateDistanceFieldAsIfTwoSided)
+{
+	TRACE_CPUPROFILER_EVENT_SCOPE(FDistanceFieldVolumeData::CacheDerivedData);
+
+	TArray<FSignedDistanceFieldBuildMaterialData> BuildMaterialData;
+
+	const uint32 LODIndex = 0;
+	BuildSignedDistanceFieldBuildMaterialData(Mesh, LODIndex, BuildMaterialData);
 
 	FString DistanceFieldKey = BuildDistanceFieldDerivedDataKey(InStaticMeshDerivedDataKey);
 
-	for (int32 MaterialIndex = 0; MaterialIndex < Mesh->GetStaticMaterials().Num(); MaterialIndex++)
+	for (int32 MaterialIndex = 0; MaterialIndex < BuildMaterialData.Num(); MaterialIndex++)
 	{
 		DistanceFieldKey += FString::Printf(TEXT("_M%u_%u_%u"), 
 			(uint32)BuildMaterialData[MaterialIndex].BlendMode,
diff --git a/Engine/Source/Runtime/Engine/Private/MeshCardRepresentation.cpp b/Engine/Source/Runtime/Engine/Private/MeshCardRepresentation.cpp
--- a/Engine/Source/Runtime/Engine/Private/MeshCardRepresentation.cpp
+++ b/Engine/Source/Runtime/Engine/Private/MeshCardRepresentation.cpp
@@ -117,6 +117,8 @@ FString BuildCardRepresentationDerivedDataKey(const FString& InMeshKey, int32 Ma
 
 #if WITH_EDITORONLY_DATA
 
+extern void BuildSignedDistanceFieldBuildMaterialData(UStaticMesh* Mesh, uint32 LODIndex, TArray<FSignedDistanceFieldBuildMaterialData>& OutData);
+
 void BeginCacheMeshCardRepresentation(const ITargetPlatform* TargetPlatform, UStaticMesh* StaticMeshAsset, FStaticMeshRenderData& RenderData, const FString& DistanceFieldKey, FSourceMeshDataForDerivedDataTask* OptionalSourceMeshData)
 {
 	static const auto CVarCards = IConsoleManager::Get().FindTConsoleVariableDataInt(TEXT("r.MeshCardRepresentation"));
@@ -173,31 +175,9 @@ void FCardRepresentationData::CacheDerivedData(const FString& InDDCKey, const IT
 		NewTask->GeneratedCardRepresentation = new FCardRepresentationData();
 		NewTask->MaxLumenMeshCards = MaxLumenMeshCards;
 		NewTask->bGenerateDistanceFieldAsIfTwoSided = bGenerateDistanceFieldAsIfTwoSided;
-		NewTask->MaterialBlendModes.SetNum(Mesh->GetStaticMaterials().Num());
 
-		const TArray<FStaticMaterial>& StaticMaterials = Mesh->GetStaticMaterials();
-		const FMeshSectionInfoMap& SectionInfoMap = Mesh->GetSectionInfoMap();
 		const uint32 LODIndex = 0;
-
-		for (int32 SectionIndex = 0; SectionIndex < SectionInfoMap.GetSectionNumber(LODIndex); SectionIndex++)
-		{
-			const FMeshSectionInfo& Section = SectionInfoMap.Get(LODIndex, SectionIndex);
-
-			if (!NewTask->MaterialBlendModes.IsValidIndex(Section.MaterialIndex))
-			{
-				continue;
-			}
-
-			FSignedDistanceFieldBuildMaterialData& MaterialData = NewTask->MaterialBlendModes[Section.MaterialIndex];
-			MaterialData.bAffectDistanceFieldLighting = Section.bAffectDistanceFieldLighting;
-
-			UMaterialInterface* MaterialInterface = StaticMaterials[Section.MaterialIndex].MaterialInterface;
-			if (MaterialInterface)
-			{
-				MaterialData.BlendMode = MaterialInterface->GetBlendMode();
-				MaterialData.bTwoSided = MaterialInterface->IsTwoSided();
-			}
-		}
+		BuildSignedDistanceFieldBuildMaterialData(Mesh, LODIndex, NewTask->MaterialBlendModes);
 
 		// Nanite overrides source static mesh with a coarse representation. Need to load original data before we build the mesh SDF.
 		if (OptionalSourceMeshData)
