WorldPartitionRuntimeHashSet:
- Generate a streaming object for internal streaming data, not just injected content bundles, to minimize code duplication, etc.
- Runtime partitions now supports a loading range and name in order to be on par with existing runtime 2d hash targetgrids/default loading range.

diff --git a/Engine/Source/Runtime/Engine/Private/WorldPartition/RuntimeHashSet/RuntimePartitionLevelStreaming.cpp b/Engine/Source/Runtime/Engine/Private/WorldPartition/RuntimeHashSet/RuntimePartitionLevelStreaming.cpp
--- a/Engine/Source/Runtime/Engine/Private/WorldPartition/RuntimeHashSet/RuntimePartitionLevelStreaming.cpp
+++ b/Engine/Source/Runtime/Engine/Private/WorldPartition/RuntimeHashSet/RuntimePartitionLevelStreaming.cpp
@@ -25,7 +25,7 @@ bool URuntimePartitionLevelStreaming::GenerateStreaming(const TArray<const IStre
 	TArray<IStreamingGenerationContext::FActorInstance> CellActorInstances;
 	if (PopulateCellActorInstances(ActorSetInstances, bIsMainWorldPartition, false, CellActorInstances))
 	{
-		TMap<FName, TPair<FName, TArray<IStreamingGenerationContext::FActorInstance>>> SubLevelsActorInstances;
+		TMap<FName, TArray<IStreamingGenerationContext::FActorInstance>> SubLevelsActorInstances;
 		for (const IStreamingGenerationContext::FActorInstance& ActorInstance : CellActorInstances)
 		{
 			TArray<FName> ActorSetGridNameList;
@@ -54,22 +54,20 @@ bool URuntimePartitionLevelStreaming::GenerateStreaming(const TArray<const IStre
 
 			FName SubLevelName = *StringBuilder;
 
-			TPair<FName, TArray<IStreamingGenerationContext::FActorInstance>>& Pair = SubLevelsActorInstances.FindOrAdd(SubLevelName);
-			Pair.Key = SubLevelName;
-			Pair.Value.Add(ActorInstance);
+			SubLevelsActorInstances.FindOrAdd(SubLevelName).Add(ActorInstance);
 		}
 
 		for (auto& [SubLevelName, SubLevelActorSetInstances] : SubLevelsActorInstances)
 		{
 			FCellDesc& CellDesc = OutRuntimeCellDescs.Emplace_GetRef();
 
-			CellDesc.Name = SubLevelActorSetInstances.Key;
+			CellDesc.Name = SubLevelName;
 			CellDesc.bIsSpatiallyLoaded = true;
-			CellDesc.ContentBundleID = SubLevelActorSetInstances.Value[0].ActorSetInstance->ContentBundleID;
+			CellDesc.ContentBundleID = SubLevelActorSetInstances[0].ActorSetInstance->ContentBundleID;
 			CellDesc.bBlockOnSlowStreaming = bBlockOnSlowStreaming;
 			CellDesc.bClientOnlyVisible = bClientOnlyVisible;
 			CellDesc.Priority = Priority;
-			CellDesc.ActorInstances = SubLevelActorSetInstances.Value;
+			CellDesc.ActorInstances = SubLevelActorSetInstances;
 
 			for (const IStreamingGenerationContext::FActorInstance& ActorInstance : CellDesc.ActorInstances)
 			{
diff --git a/Engine/Source/Runtime/Engine/Private/WorldPartition/RuntimeHashSet/WorldPartitionRuntimeHashSet.cpp b/Engine/Source/Runtime/Engine/Private/WorldPartition/RuntimeHashSet/WorldPartitionRuntimeHashSet.cpp
--- a/Engine/Source/Runtime/Engine/Private/WorldPartition/RuntimeHashSet/WorldPartitionRuntimeHashSet.cpp
+++ b/Engine/Source/Runtime/Engine/Private/WorldPartition/RuntimeHashSet/WorldPartitionRuntimeHashSet.cpp
@@ -14,22 +14,12 @@
 #include "Algo/Transform.h"
 #include "Misc/ArchiveMD5.h"
 
-/**
- *	TODOs:
- *		- Deprecate WorldPartitionStreamingSourceComponent.TargetGrids and make it more generic to handle all sorts of filtering.
- *		- Deprecate Actor.RuntimeGrid and introduce TargetPartition + HLOD setup index.
- *		- Refactor Source.ForEachShape to extract target grid and maybe loading range.
- */
-
 FAutoConsoleCommand WorldPartitionRuntimeHashSetEnable(
 	TEXT("wp.Editor.WorldPartitionRuntimeHashSet.Enable"),
 	TEXT("Enable experimental runtime hash set class."),
 	FConsoleCommandWithArgsDelegate::CreateLambda([](const TArray<FString>& Args)
 	{
-		if (UClass* WorldPartitionRuntimeHashSetClass = FindObject<UClass>(nullptr, TEXT("/Script/Engine.WorldPartitionRuntimeHashSet")))
-		{
-			WorldPartitionRuntimeHashSetClass->ClassFlags &= ~CLASS_HideDropDown;
-		}
+		UWorldPartitionRuntimeHashSet::StaticClass()->ClassFlags &= ~CLASS_HideDropDown;
 	})
 );
 
@@ -75,7 +65,7 @@ void FRuntimePartitionDesc::UpdateHLODPartitionLayers()
 				HLODSetup.PartitionLayer->bIsHLODSetup = true;
 			}
 
-			CurHLODLayer = CurHLODLayer->GetParentLayer();
+			CurHLODLayer = CurHLODLayer->GetParentLayer();	
 		}
 
 		HLODSetups.SetNum(VisitedHLODLayers.Num());
@@ -83,34 +73,52 @@ void FRuntimePartitionDesc::UpdateHLODPartitionLayers()
 }
 #endif
 
-UWorldPartitionRuntimeHashSet::UWorldPartitionRuntimeHashSet(const FObjectInitializer& ObjectInitializer)
-	: Super(ObjectInitializer)
-{}
+void URuntimeHashSetExternalStreamingObject::CreatePartitionsSpatialIndex() const
+{
+	for (const FRuntimePartitionStreamingData& StreamingData : RuntimeStreamingData)
+	{
+		if (!StreamingData.SpatialIndex)
+		{
+			StreamingData.SpatialIndex = MakeUnique<FStaticSpatialIndexType>();
+
+			TArray<TPair<FBox, UWorldPartitionRuntimeCell*>> PartitionsElements;
+			Algo::Transform(StreamingData.RuntimeCells, PartitionsElements, [](UWorldPartitionRuntimeCell* RuntimeCell)
+			{
+				return TPair<FBox, UWorldPartitionRuntimeCell*>(RuntimeCell->GetContentBounds(), RuntimeCell);
+			});
+			StreamingData.SpatialIndex->Init(PartitionsElements);
+		}
+	}
+}
 
-void UWorldPartitionRuntimeHashSet::Serialize(FArchive& Ar)
+void URuntimeHashSetExternalStreamingObject::DestroyPartitionsSpatialIndex() const
 {
-	Super::Serialize(Ar);
+	for (const FRuntimePartitionStreamingData& StreamingData : RuntimeStreamingData)
+	{
+		StreamingData.SpatialIndex.Reset();
+	}
+}
 
+void URuntimeHashSetExternalStreamingObject::AddReferencedObjects(UObject* InThis, FReferenceCollector& Collector)
+{
 #if WITH_EDITOR
-	if (Ar.GetPortFlags() & PPF_DuplicateForPIE)
+	URuntimeHashSetExternalStreamingObject* This = CastChecked<URuntimeHashSetExternalStreamingObject>(InThis);
+	for (const FRuntimePartitionStreamingData& StreamingData : This->RuntimeStreamingData)
 	{
-		bool bSpatialIndexValid = SpatialIndex.IsValid();
-		Ar << bSpatialIndexValid;
-
-		if (bSpatialIndexValid)
+		if (StreamingData.SpatialIndex.IsValid())
 		{
-			if (Ar.IsLoading())
-			{
-				check(!SpatialIndex);
-				SpatialIndex = MakeUnique<FStaticSpatialIndexType>();
-			}
-
-			SpatialIndex->Serialize(Ar);
+			StreamingData.SpatialIndex->AddReferencedObjects(Collector);
 		}
 	}
 #endif
+
+	Super::AddReferencedObjects(InThis, Collector);
 }
 
+UWorldPartitionRuntimeHashSet::UWorldPartitionRuntimeHashSet(const FObjectInitializer& ObjectInitializer)
+	: Super(ObjectInitializer)
+{}
+
 void UWorldPartitionRuntimeHashSet::PostLoad()
 {
 	Super::PostLoad();
@@ -120,21 +128,14 @@ void UWorldPartitionRuntimeHashSet::PostLoad()
 	{
 		UpdateHLODPartitionLayers();
 	}
+	else
 #endif
-}
-
-void UWorldPartitionRuntimeHashSet::AddReferencedObjects(UObject* InThis, FReferenceCollector& Collector)
-{
-#if WITH_EDITOR
-	UWorldPartitionRuntimeHashSet* This = CastChecked<UWorldPartitionRuntimeHashSet>(InThis);
-	if (This->SpatialIndex)
 	{
-		This->SpatialIndex->AddReferencedObjects(Collector);
+		ForEachStreamingObject([](const URuntimeHashSetExternalStreamingObject* StreamingObject)
+		{
+			StreamingObject->CreatePartitionsSpatialIndex();
+		});
 	}
-	Collector.AddStableReferenceSet(&This->InjectedExternalStreamingObjects);
-#endif
-
-	Super::AddReferencedObjects(InThis, Collector);
 }
 
 #if WITH_EDITOR
@@ -171,7 +172,7 @@ bool UWorldPartitionRuntimeHashSet::GenerateStreaming(UWorldPartitionStreamingPo
 	PersistentPartitionDesc.Class = URuntimePartitionPersistent::StaticClass();
 	PersistentPartitionDesc.Name = NAME_PersistentLevel;
 	PersistentPartitionDesc.MainLayer = NewObject<URuntimePartition>(this, URuntimePartitionPersistent::StaticClass(), NAME_None);
-	PersistentPartitionDesc.MainLayer->Name = TEXT("MainPartition");
+	PersistentPartitionDesc.MainLayer->Name = NAME_PersistentLevel;
 
 	//
 	// Split actor sets into their corresponding runtime partition implementation
@@ -197,56 +198,62 @@ bool UWorldPartitionRuntimeHashSet::GenerateStreaming(UWorldPartitionStreamingPo
 		}
 	});
 
-	// Generate per data layer instance of partition cells
 	struct FCellDescInstance : public URuntimePartition::FCellDesc
 	{
-		FBox Bounds;
+		FCellDescInstance(const URuntimePartition::FCellDesc& InCellDesc, URuntimePartition* InSourcePartition, const TArray<const UDataLayerInstance*>& InDataLayers)
+			: URuntimePartition::FCellDesc(InCellDesc)
+			, SourcePartition(InSourcePartition)
+			, DataLayers(InDataLayers)
+		{}
+
+		URuntimePartition* SourcePartition;
 		TArray<const UDataLayerInstance*> DataLayers;
 	};
 
-	TArray<FCellDescInstance> RuntimeCellDescsDataLayers;
+	//
+	// Generate runtime partitions streaming data
+	//
+	TArray<FCellDescInstance> RuntimeCellDescsInstances;
 	{
-		TArray<URuntimePartition::FCellDesc> RuntimeCellDescs;
 		for (auto [RuntimePartition, ActorSetInstances] : RuntimePartitionsToActorSetMap)
 		{
+			// Gather runtime partition cell descs
+			TArray<URuntimePartition::FCellDesc> RuntimeCellDescs;
 			if (!RuntimePartition->GenerateStreaming(ActorSetInstances, RuntimeCellDescs))
 			{
 				return false;
 			}
-		}
-
-		// Split cell descs into data layers	
-		for (const URuntimePartition::FCellDesc& RuntimeCellDesc : RuntimeCellDescs)
-		{
-			TMap<FDataLayersID, FCellDescInstance> RuntimeCellDescsDataLayersSet;
 
-			for (const IStreamingGenerationContext::FActorInstance& ActorInstance : RuntimeCellDesc.ActorInstances)
+			// Split cell descs into data layers
+			for (const URuntimePartition::FCellDesc& RuntimeCellDesc : RuntimeCellDescs)
 			{
-				const FDataLayersID DataLayersID(ActorInstance.ActorSetInstance->DataLayers);
+				TMap<FDataLayersID, FCellDescInstance> RuntimeCellDescsInstancesSet;
 
-				FCellDescInstance* DataLayerCellDesc = RuntimeCellDescsDataLayersSet.Find(DataLayersID);
-			
-				if (!DataLayerCellDesc)
+				for (const IStreamingGenerationContext::FActorInstance& ActorInstance : RuntimeCellDesc.ActorInstances)
 				{
-					DataLayerCellDesc = &RuntimeCellDescsDataLayersSet.Emplace(DataLayersID);
-
-					DataLayerCellDesc->Name = RuntimeCellDesc.Name;
-					DataLayerCellDesc->bIsSpatiallyLoaded = RuntimeCellDesc.bIsSpatiallyLoaded;
-					DataLayerCellDesc->ContentBundleID = RuntimeCellDesc.ContentBundleID;
-					DataLayerCellDesc->bBlockOnSlowStreaming = RuntimeCellDesc.bBlockOnSlowStreaming;
-					DataLayerCellDesc->bClientOnlyVisible = RuntimeCellDesc.bClientOnlyVisible;
-					DataLayerCellDesc->Priority = RuntimeCellDesc.Priority;
-					DataLayerCellDesc->Bounds = RuntimeCellDesc.Bounds;
-					DataLayerCellDesc->DataLayers = ActorInstance.ActorSetInstance->DataLayers;
+					const FDataLayersID DataLayersID(ActorInstance.ActorSetInstance->DataLayers);
+					FCellDescInstance* CellDescInstance = RuntimeCellDescsInstancesSet.Find(DataLayersID);
+
+					if (!CellDescInstance)
+					{
+						CellDescInstance = &RuntimeCellDescsInstancesSet.Emplace(DataLayersID, FCellDescInstance(RuntimeCellDesc, RuntimePartition, ActorInstance.ActorSetInstance->DataLayers));
+						CellDescInstance->ActorInstances.Empty();
+					}
+
+					CellDescInstance->ActorInstances.Add(ActorInstance);
 				}
 
-				DataLayerCellDesc->ActorInstances.Add(ActorInstance);
+				Algo::Transform(RuntimeCellDescsInstancesSet, RuntimeCellDescsInstances, [](const auto& Value) { return Value.Value; });
 			}
-
-			Algo::Transform(RuntimeCellDescsDataLayersSet, RuntimeCellDescsDataLayers, [](const auto& Value) { return Value.Value; });
 		}
 	}
 
+	//
+	// Create and populate streaming object
+	//
+	check(!StreamingObject);
+	StreamingObject = CreateExternalStreamingObject<URuntimeHashSetExternalStreamingObject>(this, MakeUniqueObjectName(this, URuntimeHashExternalStreamingObjectBase::StaticClass()));
+
 	// Generate runtime cells
 	auto CreateRuntimeCellFromCellDesc = [this](const URuntimePartition::FCellDesc& CellDesc, const TArray<const UDataLayerInstance*>& DataLayers, TSubclassOf<UWorldPartitionRuntimeCell> CellClass, TSubclassOf<UWorldPartitionRuntimeCellData> CellDataClass)
 	{
@@ -299,7 +306,7 @@ bool UWorldPartitionRuntimeHashSet::GenerateStreaming(UWorldPartitionStreamingPo
 			check(CellGuid.IsValid());
 		}
 
-		UWorldPartitionRuntimeCell* RuntimeCell = Super::CreateRuntimeCell(CellClass, CellDataClass, CellObjectName, TEXT(""), this);
+		UWorldPartitionRuntimeCell* RuntimeCell = Super::CreateRuntimeCell(CellClass, CellDataClass, CellObjectName, TEXT(""), StreamingObject);
 
 		RuntimeCell->SetIsAlwaysLoaded(!CellDesc.bIsSpatiallyLoaded);
 		RuntimeCell->SetDataLayers(DataLayers);
@@ -314,8 +321,9 @@ bool UWorldPartitionRuntimeHashSet::GenerateStreaming(UWorldPartitionStreamingPo
 		return RuntimeCell;
 	};
 
-	TArray<UWorldPartitionRuntimeCell*> RuntimeCells;	
-	for (const FCellDescInstance& CellDescInstance : RuntimeCellDescsDataLayers)
+	TArray<UWorldPartitionRuntimeCell*> RuntimeCells;
+	TMap<URuntimePartition*, FRuntimePartitionStreamingData> RuntimePartitionsStreamingData;
+	for (const FCellDescInstance& CellDescInstance : RuntimeCellDescsInstances)
 	{
 		UWorldPartitionRuntimeCell* RuntimeCell = RuntimeCells.Emplace_GetRef(CreateRuntimeCellFromCellDesc(CellDescInstance, CellDescInstance.DataLayers, StreamingPolicy->GetRuntimeCellClass(), UWorldPartitionRuntimeCellData::StaticClass()));
 		PopulateRuntimeCell(RuntimeCell, CellDescInstance.ActorInstances, nullptr);
@@ -325,8 +333,32 @@ bool UWorldPartitionRuntimeHashSet::GenerateStreaming(UWorldPartitionStreamingPo
 		{
 			RuntimeCell->RuntimeCellData->ContentBounds = CellDescInstance.Bounds;
 		}
+
+		if (RuntimeCell->IsAlwaysLoaded())
+		{
+			StreamingObject->NonSpatiallyLoadedRuntimeCells.Add(RuntimeCell);
+		}
+		else
+		{
+			FRuntimePartitionStreamingData& StreamingData = RuntimePartitionsStreamingData.FindOrAdd(CellDescInstance.SourcePartition);
+
+			StreamingData.Name = CellDescInstance.SourcePartition->Name;
+			StreamingData.LoadingRange = CellDescInstance.SourcePartition->LoadingRange;
+			StreamingData.RuntimeCells.Add(RuntimeCell);
+		}
+	}
+
+	//
+	// Finalize streaming object
+	//
+	for (auto& [Partition, StreamingData] : RuntimePartitionsStreamingData)
+	{
+		StreamingObject->RuntimeStreamingData.Emplace(MoveTemp(StreamingData));
 	}
 
+	//
+	// Output generated packages
+	//
 	if (OutPackagesToGenerate)
 	{
 		for (UWorldPartitionRuntimeCell* RuntimeCell : RuntimeCells)
@@ -345,32 +377,6 @@ bool UWorldPartitionRuntimeHashSet::GenerateStreaming(UWorldPartitionStreamingPo
 		}
 	}
 
-	// Init spatial index
-	check(!SpatialIndex);
-	SpatialIndex = MakeUnique<FStaticSpatialIndexType>();
-
-	TSet<FGuid> RuntimeCellsGuids;
-	RuntimeCellsGuids.Reserve(RuntimeCells.Num());
-
-	TArray<TPair<FBox, UWorldPartitionRuntimeCell*>> Elements;
-	Algo::ForEach(RuntimeCells, [this, &Elements, &RuntimeCellsGuids](UWorldPartitionRuntimeCell* RuntimeCell)
-	{
-		bool bCellWasAlreadyInSet;
-		RuntimeCellsGuids.Add(RuntimeCell->GetGuid(), &bCellWasAlreadyInSet);
-		check(!bCellWasAlreadyInSet);
-
-		if (RuntimeCell->IsAlwaysLoaded())
-		{
-			NonSpatiallyLoadedRuntimeCells.Add(RuntimeCell);
-		}
-		else
-		{
-			Elements.Add(TPair<FBox, UWorldPartitionRuntimeCell*>(RuntimeCell->GetContentBounds(), RuntimeCell));
-		}
-	});
-
-	SpatialIndex->Init(Elements);
-
 	return true;
 }
 
@@ -383,8 +389,7 @@ void UWorldPartitionRuntimeHashSet::FlushStreaming()
 	PersistentPartitionDesc.Name = NAME_None;
 	PersistentPartitionDesc.MainLayer = nullptr;
 
-	NonSpatiallyLoadedRuntimeCells.Empty();	
-	SpatialIndex.Reset();
+	StreamingObject = nullptr;
 }
 
 bool UWorldPartitionRuntimeHashSet::IsValidGrid(FName GridName) const
@@ -413,6 +418,11 @@ bool UWorldPartitionRuntimeHashSet::IsValidGrid(FName GridName) const
 
 TArray<UWorldPartitionRuntimeCell*> UWorldPartitionRuntimeHashSet::GetAlwaysLoadedCells() const
 {
+	TArray<UWorldPartitionRuntimeCell*> NonSpatiallyLoadedRuntimeCells;
+	ForEachStreamingObject([&NonSpatiallyLoadedRuntimeCells](const URuntimeHashSetExternalStreamingObject* StreamingObject)
+	{
+		NonSpatiallyLoadedRuntimeCells.Append(StreamingObject->NonSpatiallyLoadedRuntimeCells);
+	});
 	return NonSpatiallyLoadedRuntimeCells;
 }
 
@@ -453,43 +463,36 @@ TArray<FName> UWorldPartitionRuntimeHashSet::ParseGridName(FName GridName)
 
 URuntimeHashExternalStreamingObjectBase* UWorldPartitionRuntimeHashSet::StoreToExternalStreamingObject(UObject* StreamingObjectOuter, FName StreamingObjectName)
 {
-	URuntimeHashSetExternalStreamingObject* StreamingObject = CreateExternalStreamingObject<URuntimeHashSetExternalStreamingObject>(StreamingObjectOuter, StreamingObjectName);
-	StreamingObject->NonSpatiallyLoadedRuntimeCells = NonSpatiallyLoadedRuntimeCells;
-	SpatialIndex->ForEachElement([StreamingObject](UWorldPartitionRuntimeCell* RuntimeCell) { StreamingObject->SpatiallyLoadedRuntimeCells.Add(RuntimeCell); });
-	return StreamingObject;
+	check(StreamingObject);
+	StreamingObject->Rename(*StreamingObjectName.ToString(), StreamingObjectOuter, REN_DoNotDirty | REN_ForceNoResetLoaders);
+	URuntimeHashExternalStreamingObjectBase* Result = StreamingObject;
+	StreamingObject = nullptr;
+	return Result;
 }
 #endif
 
 bool UWorldPartitionRuntimeHashSet::InjectExternalStreamingObject(URuntimeHashExternalStreamingObjectBase* ExternalStreamingObject)
 {
-	URuntimeHashSetExternalStreamingObject* HashSetExternalStreamingObject = CastChecked<URuntimeHashSetExternalStreamingObject>(ExternalStreamingObject);
-
-	bool bWasAlreadyInSet;
-	InjectedExternalStreamingObjects.Add(HashSetExternalStreamingObject, &bWasAlreadyInSet);
-	check(!bWasAlreadyInSet);
-
-	check(!HashSetExternalStreamingObject->SpatialIndex);
-	HashSetExternalStreamingObject->SpatialIndex = MakeUnique<FStaticSpatialIndexType>();
-
-	TArray<TPair<FBox, UWorldPartitionRuntimeCell*>> Elements;
-	Algo::Transform(HashSetExternalStreamingObject->SpatiallyLoadedRuntimeCells, Elements, 
-		[](UWorldPartitionRuntimeCell* RuntimeCell) { return TPair<FBox, UWorldPartitionRuntimeCell*>(RuntimeCell->GetContentBounds(), RuntimeCell); }
-	);	
-	HashSetExternalStreamingObject->SpatialIndex->Init(Elements);
+	if (Super::InjectExternalStreamingObject(ExternalStreamingObject))
+	{
+		URuntimeHashSetExternalStreamingObject* HashSetExternalStreamingObject = CastChecked<URuntimeHashSetExternalStreamingObject>(ExternalStreamingObject);
+		HashSetExternalStreamingObject->CreatePartitionsSpatialIndex();
+		return true;
+	}
 
-	return true;
+	return false;
 }
 
 bool UWorldPartitionRuntimeHashSet::RemoveExternalStreamingObject(URuntimeHashExternalStreamingObjectBase* ExternalStreamingObject)
 {
-	URuntimeHashSetExternalStreamingObject* HashSetExternalStreamingObject = CastChecked<URuntimeHashSetExternalStreamingObject>(ExternalStreamingObject);
-
-	check(HashSetExternalStreamingObject->SpatialIndex);
-	HashSetExternalStreamingObject->SpatialIndex.Reset();
-
-	verify(InjectedExternalStreamingObjects.Remove(HashSetExternalStreamingObject));
+	if (Super::RemoveExternalStreamingObject(ExternalStreamingObject))
+	{
+		URuntimeHashSetExternalStreamingObject* HashSetExternalStreamingObject = CastChecked<URuntimeHashSetExternalStreamingObject>(ExternalStreamingObject);
+		HashSetExternalStreamingObject->DestroyPartitionsSpatialIndex();
+		return true;
+	}
 
-	return true;
+	return false;
 }
 
 // Streaming interface
@@ -520,14 +523,15 @@ void UWorldPartitionRuntimeHashSet::ForEachStreamingCells(TFunctionRef<bool(cons
 		}
 	};
 
-	ForEachStreamingCells(SpatialIndex.Get());
-	ForEachNonStreamingCells(NonSpatiallyLoadedRuntimeCells);
-
-	for (URuntimeHashSetExternalStreamingObject* InjectedExternalStreamingObject : InjectedExternalStreamingObjects)
+	ForEachStreamingObject([&ForEachStreamingCells, &ForEachNonStreamingCells](const URuntimeHashSetExternalStreamingObject* StreamingObject)
 	{
-		ForEachStreamingCells(InjectedExternalStreamingObject->SpatialIndex.Get());
-		ForEachNonStreamingCells(InjectedExternalStreamingObject->NonSpatiallyLoadedRuntimeCells);
-	}
+		for (const FRuntimePartitionStreamingData& StreamingData : StreamingObject->RuntimeStreamingData)
+		{
+			ForEachStreamingCells(StreamingData.SpatialIndex.Get());
+		}
+
+		ForEachNonStreamingCells(StreamingObject->NonSpatiallyLoadedRuntimeCells);
+	});
 }
 
 void UWorldPartitionRuntimeHashSet::ForEachStreamingCellsQuery(const FWorldPartitionStreamingQuerySource& QuerySource, TFunctionRef<bool(const UWorldPartitionRuntimeCell*)> Func, FWorldPartitionQueryCache* QueryCache) const
@@ -577,14 +581,15 @@ void UWorldPartitionRuntimeHashSet::ForEachStreamingCellsQuery(const FWorldParti
 		}
 	};
 
-	ForEachStreamingCells(SpatialIndex.Get());
-	ForEachNonStreamingCells(NonSpatiallyLoadedRuntimeCells);
-
-	for (URuntimeHashSetExternalStreamingObject* InjectedExternalStreamingObject : InjectedExternalStreamingObjects)
+	ForEachStreamingObject([&ForEachStreamingCells, &ForEachNonStreamingCells](const URuntimeHashSetExternalStreamingObject* StreamingObject)
 	{
-		ForEachStreamingCells(InjectedExternalStreamingObject->SpatialIndex.Get());
-		ForEachNonStreamingCells(InjectedExternalStreamingObject->NonSpatiallyLoadedRuntimeCells);
-	}
+		for (const FRuntimePartitionStreamingData& StreamingData : StreamingObject->RuntimeStreamingData)
+		{
+			ForEachStreamingCells(StreamingData.SpatialIndex.Get());
+		}
+
+		ForEachNonStreamingCells(StreamingObject->NonSpatiallyLoadedRuntimeCells);
+	});
 }
 
 void UWorldPartitionRuntimeHashSet::ForEachStreamingCellsSources(const TArray<FWorldPartitionStreamingSource>& Sources, TFunctionRef<bool(const UWorldPartitionRuntimeCell*, EStreamingSourceTargetState)> Func) const
@@ -592,17 +597,15 @@ void UWorldPartitionRuntimeHashSet::ForEachStreamingCellsSources(const TArray<FW
 	UWorldPartitionRuntimeHash::FStreamingSourceCells ActivateStreamingSourceCells;
 	UWorldPartitionRuntimeHash::FStreamingSourceCells LoadStreamingSourceCells;
 
-	auto ForEachStreamingCells = [this, &Sources, &Func, &ActivateStreamingSourceCells, &LoadStreamingSourceCells](FStaticSpatialIndexType* InSpatialIndex)
+	auto ForEachStreamingCells = [this, &Sources, &Func, &ActivateStreamingSourceCells, &LoadStreamingSourceCells](FStaticSpatialIndexType* InSpatialIndex, float InLoadingRange, FName InGridName)
 	{
 		if (InSpatialIndex)
 		{
 			for (const FWorldPartitionStreamingSource& Source : Sources)
 			{
 				// @todo_jfd
-				const FName GridName;
 				const FSoftObjectPath HLODLayer;
-
-				Source.ForEachShape(/*LoadingRange*/25600, GridName, HLODLayer, false, [this, &Source, InSpatialIndex, &Func, &ActivateStreamingSourceCells, &LoadStreamingSourceCells](const FSphericalSector& Shape)
+				Source.ForEachShape(InLoadingRange, InGridName, HLODLayer, false, [this, &Source, InSpatialIndex, &Func, &ActivateStreamingSourceCells, &LoadStreamingSourceCells](const FSphericalSector& Shape)
 				{
 					const FSphere ShapeSphere(Shape.GetCenter(), Shape.GetRadius());
 
@@ -650,14 +653,15 @@ void UWorldPartitionRuntimeHashSet::ForEachStreamingCellsSources(const TArray<FW
 		}
 	};
 
-	ForEachStreamingCells(SpatialIndex.Get());
-	ForEachNonStreamingCells(NonSpatiallyLoadedRuntimeCells);
-
-	for (URuntimeHashSetExternalStreamingObject* InjectedExternalStreamingObject : InjectedExternalStreamingObjects)
+	ForEachStreamingObject([&ForEachStreamingCells, &ForEachNonStreamingCells](const URuntimeHashSetExternalStreamingObject* StreamingObject)
 	{
-		ForEachStreamingCells(InjectedExternalStreamingObject->SpatialIndex.Get());
-		ForEachNonStreamingCells(InjectedExternalStreamingObject->NonSpatiallyLoadedRuntimeCells);
-	}
+		for (const FRuntimePartitionStreamingData& StreamingData : StreamingObject->RuntimeStreamingData)
+		{
+			ForEachStreamingCells(StreamingData.SpatialIndex.Get(), StreamingData.LoadingRange, StreamingData.Name);
+		}
+
+		ForEachNonStreamingCells(StreamingObject->NonSpatiallyLoadedRuntimeCells);
+	});
 
 	auto ExecuteFuncOnCells = [Func](const TSet<const UWorldPartitionRuntimeCell*>& Cells, EStreamingSourceTargetState TargetState)
 	{
@@ -695,7 +699,7 @@ void UWorldPartitionRuntimeHashSet::PostEditChangeChainProperty(FPropertyChanged
 		{
 			RuntimePartitionDesc.Name = RuntimePartitionDesc.Class->GetFName();
 			RuntimePartitionDesc.MainLayer = NewObject<URuntimePartition>(this, RuntimePartitionDesc.Class, NAME_None);
-			RuntimePartitionDesc.MainLayer->Name = TEXT("MainPartition");
+			RuntimePartitionDesc.MainLayer->Name = RuntimePartitionDesc.Name;
 
 			RuntimePartitionDesc.UpdateHLODPartitionLayers();
 		}
@@ -725,6 +729,8 @@ void UWorldPartitionRuntimeHashSet::PostEditChangeChainProperty(FPropertyChanged
 				}
 			}
 		}
+
+		RuntimePartitionDesc.MainLayer->Name = RuntimePartitionDesc.Name;
 	}
 	else if (PropertyName == NAME_HLODLayer)
 	{
@@ -742,4 +748,20 @@ void UWorldPartitionRuntimeHashSet::UpdateHLODPartitionLayers()
 		RuntimePartitionDesc.UpdateHLODPartitionLayers();
 	}
 }
-#endif
\ No newline at end of file
+#endif
+
+void UWorldPartitionRuntimeHashSet::ForEachStreamingObject(TFunctionRef<void(const URuntimeHashSetExternalStreamingObject*)> Func) const
+{
+	if (StreamingObject)
+	{
+		Func(StreamingObject.Get());
+	}
+
+	for (const TWeakObjectPtr<URuntimeHashExternalStreamingObjectBase>& InjectedExternalStreamingObject : InjectedExternalStreamingObjects)
+	{
+		if (InjectedExternalStreamingObject.IsValid())
+		{
+			Func(CastChecked<URuntimeHashSetExternalStreamingObject>(InjectedExternalStreamingObject.Get()));
+		}
+	}
+}
\ No newline at end of file
diff --git a/Engine/Source/Runtime/Engine/Private/WorldPartition/WorldPartitionStreamingGenerationContext.cpp b/Engine/Source/Runtime/Engine/Private/WorldPartition/WorldPartitionStreamingGenerationContext.cpp
--- a/Engine/Source/Runtime/Engine/Private/WorldPartition/WorldPartitionStreamingGenerationContext.cpp
+++ b/Engine/Source/Runtime/Engine/Private/WorldPartition/WorldPartitionStreamingGenerationContext.cpp
@@ -23,4 +23,9 @@ const UActorDescContainer* IStreamingGenerationContext::FActorInstance::GetActor
 {
 	return GetActorDescView().GetActorDesc()->GetContainer();
 }
+
+const FBox IStreamingGenerationContext::FActorInstance::GetBounds() const
+{
+	return ActorSetInstance->Bounds;
+}
 #endif
diff --git a/Engine/Source/Runtime/Engine/Public/WorldPartition/RuntimeHashSet/RuntimePartition.h b/Engine/Source/Runtime/Engine/Public/WorldPartition/RuntimeHashSet/RuntimePartition.h
--- a/Engine/Source/Runtime/Engine/Public/WorldPartition/RuntimeHashSet/RuntimePartition.h
+++ b/Engine/Source/Runtime/Engine/Public/WorldPartition/RuntimeHashSet/RuntimePartition.h
@@ -6,7 +6,7 @@
 #include "WorldPartition/WorldPartitionStreamingGenerationContext.h"
 #include "RuntimePartition.generated.h"
 
-UCLASS(Abstract, Within=WorldPartitionRuntimeHashSet, CollapseCategories)
+UCLASS(Abstract, CollapseCategories)
 class URuntimePartition : public UObject
 {
 	GENERATED_BODY()
@@ -32,7 +32,7 @@ public:
 #endif
 
 #if WITH_EDITORONLY_DATA
-	UPROPERTY(VisibleAnywhere, Category = RuntimeSettings, Meta = (EditCondition = "!bIsHLODSetup", EditConditionHides, HideEditConditionToggle))
+	UPROPERTY()
 	FName Name;
 
 	UPROPERTY(EditAnywhere, Category = RuntimeSettings, Meta = (EditCondition = "!bIsHLODSetup", EditConditionHides, HideEditConditionToggle))
@@ -44,6 +44,9 @@ public:
 	UPROPERTY(EditAnywhere, Category = RuntimeSettings, Meta = (EditCondition = "!bIsHLODSetup", EditConditionHides, HideEditConditionToggle))
 	int32 Priority;
 
+	UPROPERTY(EditAnywhere, Category = RuntimeSettings)
+	float LoadingRange;
+
 	UPROPERTY()
 	bool bIsHLODSetup;
 #endif
diff --git a/Engine/Source/Runtime/Engine/Public/WorldPartition/RuntimeHashSet/StaticSpatialIndex.h b/Engine/Source/Runtime/Engine/Public/WorldPartition/RuntimeHashSet/StaticSpatialIndex.h
--- a/Engine/Source/Runtime/Engine/Public/WorldPartition/RuntimeHashSet/StaticSpatialIndex.h
+++ b/Engine/Source/Runtime/Engine/Public/WorldPartition/RuntimeHashSet/StaticSpatialIndex.h
@@ -64,16 +64,6 @@ public:
 		}
 	}
 
-	void Serialize(FArchive& Ar)
-	{
-		Ar << Elements;
-
-		if (Ar.IsLoading())
-		{
-			InitSpatialIndex();
-		}
-	}
-
 	// IStaticSpatialIndexDataInterface interface
 	virtual const FBox& GetBox(uint32 InIndex) const override
 	{
diff --git a/Engine/Source/Runtime/Engine/Public/WorldPartition/RuntimeHashSet/WorldPartitionRuntimeHashSet.h b/Engine/Source/Runtime/Engine/Public/WorldPartition/RuntimeHashSet/WorldPartitionRuntimeHashSet.h
--- a/Engine/Source/Runtime/Engine/Public/WorldPartition/RuntimeHashSet/WorldPartitionRuntimeHashSet.h
+++ b/Engine/Source/Runtime/Engine/Public/WorldPartition/RuntimeHashSet/WorldPartitionRuntimeHashSet.h
@@ -19,8 +19,10 @@ struct FRuntimePartitionHLODSetup
 {
 	GENERATED_USTRUCT_BODY()
 
+#if WITH_EDITORONLY_DATA
 	UPROPERTY(VisibleAnywhere, Category = RuntimeSettings, Meta = (DisplayThumbnail = false))
 	TObjectPtr<const UHLODLayer> HLODLayer;
+#endif
 
 	UPROPERTY(EditAnywhere, Category = RuntimeSettings, Instanced)
 	TObjectPtr<URuntimePartition> PartitionLayer;
@@ -59,20 +61,52 @@ struct FRuntimePartitionDesc
 #endif
 };
 
+USTRUCT()
+struct FRuntimePartitionStreamingData
+{
+	GENERATED_USTRUCT_BODY()
+
+	/** Name of the runtime partition, currently maps to target grids. */
+	UPROPERTY()
+	FName Name;
+
+	UPROPERTY()
+	float LoadingRange;
+
+	UPROPERTY()
+	TArray<TObjectPtr<UWorldPartitionRuntimeCell>> RuntimeCells;
+
+	// Transient
+	mutable TUniquePtr<FStaticSpatialIndexType> SpatialIndex;
+};
+
+template<>
+struct TStructOpsTypeTraits<FRuntimePartitionStreamingData> : public TStructOpsTypeTraitsBase2<FRuntimePartitionStreamingData>
+{
+	enum
+	{
+		WithCopy = false
+	};
+};
+
 UCLASS()
 class URuntimeHashSetExternalStreamingObject : public URuntimeHashExternalStreamingObjectBase
 {
 	GENERATED_BODY()
 
 public:
+	//~ Begin UObject Interface
+	ENGINE_API static void AddReferencedObjects(UObject* InThis, FReferenceCollector& Collector);
+	//~ End UObject Interface
+
+	void CreatePartitionsSpatialIndex() const;
+	void DestroyPartitionsSpatialIndex() const;
+
 	UPROPERTY()
 	TArray<TObjectPtr<UWorldPartitionRuntimeCell>> NonSpatiallyLoadedRuntimeCells;
 
 	UPROPERTY()
-	TArray<TObjectPtr<UWorldPartitionRuntimeCell>> SpatiallyLoadedRuntimeCells;
-
-	// Transient
-	TUniquePtr<FStaticSpatialIndexType> SpatialIndex;
+	TArray<FRuntimePartitionStreamingData> RuntimeStreamingData;
 };
 
 UCLASS(HideDropdown, MinimalAPI)
@@ -81,9 +115,7 @@ class UWorldPartitionRuntimeHashSet : public UWorldPartitionRuntimeHash
 	GENERATED_UCLASS_BODY()
 
 	//~ Begin UObject Interface
-	ENGINE_API virtual void Serialize(FArchive& Ar) override;
 	ENGINE_API virtual void PostLoad() override;
-	static ENGINE_API void AddReferencedObjects(UObject* InThis, FReferenceCollector& Collector);
 #if WITH_EDITOR
 	ENGINE_API virtual void PostEditChangeChainProperty(FPropertyChangedChainEvent& PropertyChangedEvent) override;
 #endif
@@ -122,21 +154,19 @@ private:
 	void UpdateHLODPartitionLayers();
 #endif
 
+	ENGINE_API void ForEachStreamingObject(TFunctionRef<void(const URuntimeHashSetExternalStreamingObject*)> Func) const;
+
 public:
 #if WITH_EDITORONLY_DATA
 	/** Persistent partition */
 	UPROPERTY()
 	FRuntimePartitionDesc PersistentPartitionDesc;
+#endif
 
 	/** Array of runtime partition descriptors */
 	UPROPERTY(EditAnywhere, Category = RuntimeSettings, Meta = (TitleProperty = "Name"))
 	TArray<FRuntimePartitionDesc> RuntimePartitions;
-#endif
 
 	UPROPERTY()
-	TArray<TObjectPtr<UWorldPartitionRuntimeCell>> NonSpatiallyLoadedRuntimeCells;
-
-	TUniquePtr<FStaticSpatialIndexType> SpatialIndex;
-
-	TSet<URuntimeHashSetExternalStreamingObject*> InjectedExternalStreamingObjects;
-};
+	TObjectPtr<URuntimeHashSetExternalStreamingObject> StreamingObject;
+};
\ No newline at end of file
diff --git a/Engine/Source/Runtime/Engine/Public/WorldPartition/WorldPartitionRuntimeHash.h b/Engine/Source/Runtime/Engine/Public/WorldPartition/WorldPartitionRuntimeHash.h
--- a/Engine/Source/Runtime/Engine/Public/WorldPartition/WorldPartitionRuntimeHash.h
+++ b/Engine/Source/Runtime/Engine/Public/WorldPartition/WorldPartitionRuntimeHash.h
@@ -186,7 +186,7 @@ public:
 	mutable FActorDescList ModifiedActorDescListForPIE;
 #endif
 
-private:
+protected:
 #if WITH_EDITOR
 	ENGINE_API void ForceExternalActorLevelReference(bool bForceExternalActorLevelReferenceForPIE);
 #endif
diff --git a/Engine/Source/Runtime/Engine/Public/WorldPartition/WorldPartitionStreamingGenerationContext.h b/Engine/Source/Runtime/Engine/Public/WorldPartition/WorldPartitionStreamingGenerationContext.h
--- a/Engine/Source/Runtime/Engine/Public/WorldPartition/WorldPartitionStreamingGenerationContext.h
+++ b/Engine/Source/Runtime/Engine/Public/WorldPartition/WorldPartitionStreamingGenerationContext.h
@@ -62,6 +62,7 @@ public:
 		ENGINE_API const FActorContainerID& GetContainerID() const;
 		ENGINE_API const FTransform& GetTransform() const;
 		ENGINE_API const UActorDescContainer* GetActorDescContainer() const;
+		ENGINE_API const FBox GetBounds() const;
 	};
 
 	virtual FBox GetWorldBounds() const = 0;
@@ -69,4 +70,4 @@ public:
 	virtual void ForEachActorSetInstance(TFunctionRef<void(const FActorSetInstance&)> Func) const = 0;
 	virtual void ForEachActorSetContainer(TFunctionRef<void(const FActorSetContainer&)> Func) const = 0;
 };
-#endif
+#endif
\ No newline at end of file
