Adding missing includes. Found when compiling with -IncludeHeaders

diff --git a/Engine/Source/Runtime/Engine/Classes/GameFramework/LightWeightInstanceUtils.h b/Engine/Source/Runtime/Engine/Classes/GameFramework/LightWeightInstanceUtils.h
--- a/Engine/Source/Runtime/Engine/Classes/GameFramework/LightWeightInstanceUtils.h
+++ b/Engine/Source/Runtime/Engine/Classes/GameFramework/LightWeightInstanceUtils.h
@@ -5,6 +5,7 @@
 // HEADER_UNIT_UNSUPPORTED - I don't think this header is used. dll export of templated function below indicates that.
 
 #include "Kismet/BlueprintFunctionLibrary.h"
+#include "Engine/ActorInstanceHandle.h"
 #include "Engine/EngineTypes.h"
 
 #include "LightWeightInstanceBlueprintFunctionLibrary.generated.h"
@@ -15,13 +16,13 @@ namespace LWIUtils
 	 * Returns true if the object this handle represents supports the interface of type U
 	 */
 	template<typename U>
-	static ENGINE_API bool DoesHandleSupportInterface(const FActorInstanceHandle& Handle);
+	static bool DoesHandleSupportInterface(const FActorInstanceHandle& Handle);
 
 	/**
 	 * Returns an object implementing the interface I.
 	 */
 	template<typename I>
-	static ENGINE_API I* FetchInterfaceFromHandle(const FActorInstanceHandle& Handle);
+	static I* FetchInterfaceFromHandle(const FActorInstanceHandle& Handle);
 
 
 	/**
@@ -46,7 +47,7 @@ namespace LWIUtils
 	{
 		if (Handle.IsActorValid())
 		{
-			return Cast<I>(Actor.Get());
+			return Cast<I>(Handle.FetchActor().Get());
 		}
 
 		return FLightWeightInstanceSubsystem::Get().FetchInterfaceObject<I>(Handle);
diff --git a/Engine/Source/Runtime/Engine/Classes/Intrinsic/Model.h b/Engine/Source/Runtime/Engine/Classes/Intrinsic/Model.h
--- a/Engine/Source/Runtime/Engine/Classes/Intrinsic/Model.h
+++ b/Engine/Source/Runtime/Engine/Classes/Intrinsic/Model.h
@@ -4,6 +4,8 @@
 
 // HEADER_UNIT_SKIP
 
+#include "UObject/ObjectMacros.h"
+
 #if !CPP && !UE_BUILD_DOCS      //noexport class
 UCLASS(noexport, Intrinsic)
 class UModel : public UObject
diff --git a/Engine/Source/Runtime/Engine/Classes/Materials/MaterialInstanceUpdateParameterSet.h b/Engine/Source/Runtime/Engine/Classes/Materials/MaterialInstanceUpdateParameterSet.h
--- a/Engine/Source/Runtime/Engine/Classes/Materials/MaterialInstanceUpdateParameterSet.h
+++ b/Engine/Source/Runtime/Engine/Classes/Materials/MaterialInstanceUpdateParameterSet.h
@@ -6,6 +6,8 @@
 
 #include "Materials/Material.h"
 #include "Materials/MaterialInstance.h"
+#include "Materials/MaterialExpressionMaterialAttributeLayers.h"
+#include "Materials/MaterialExpressionMaterialFunctionCall.h"
 #include "Misc/App.h"
 
 #if WITH_EDITOR
diff --git a/Engine/Source/Runtime/RHI/Public/Windows/IDXGISwapchainProvider.h b/Engine/Source/Runtime/RHI/Public/Windows/IDXGISwapchainProvider.h
--- a/Engine/Source/Runtime/RHI/Public/Windows/IDXGISwapchainProvider.h
+++ b/Engine/Source/Runtime/RHI/Public/Windows/IDXGISwapchainProvider.h
@@ -4,6 +4,7 @@
 
 #include "Features/IModularFeature.h"
 #include "RHIDefinitions.h"
+#include "UObject/NameTypes.h"
 
 class IDXGISwapchainProvider : public IModularFeature
 {
diff --git a/Engine/Source/Runtime/RenderCore/Public/VertexStreamComponent.h b/Engine/Source/Runtime/RenderCore/Public/VertexStreamComponent.h
--- a/Engine/Source/Runtime/RenderCore/Public/VertexStreamComponent.h
+++ b/Engine/Source/Runtime/RenderCore/Public/VertexStreamComponent.h
@@ -3,6 +3,7 @@
 #pragma once
 
 #include "RHIDefinitions.h"
+#include "Containers/EnumAsByte.h"
 
 class FVertexBuffer;
 
