Fix use after free of string object.

diff --git a/Engine/Source/Runtime/Engine/Private/Texture2D.cpp b/Engine/Source/Runtime/Engine/Private/Texture2D.cpp
--- a/Engine/Source/Runtime/Engine/Private/Texture2D.cpp
+++ b/Engine/Source/Runtime/Engine/Private/Texture2D.cpp
@@ -1458,7 +1458,8 @@ void FVirtualTexture2DResource::InitializeEditorResources(IVirtualTexture* InVir
 			}
 		}
 
-		FRHITextureCreateDesc Desc = FRHITextureCreateDesc::Create2D(*TextureName.ToString(), MipWidthInTiles * TileSizeInPixels, MipHeightInTiles * TileSizeInPixels, PixelFormat);
+		FString Name = TextureName.ToString();
+		FRHITextureCreateDesc Desc = FRHITextureCreateDesc::Create2D(*Name, MipWidthInTiles * TileSizeInPixels, MipHeightInTiles * TileSizeInPixels, PixelFormat);
 		Desc.AddFlags(TexCreateFlags);
 
 		FTexture2DRHIRef Texture2DRHI = RHICreateTexture(Desc);
@@ -1507,7 +1508,7 @@ void FVirtualTexture2DResource::InitializeEditorResources(IVirtualTexture* InVir
 			check(MipHeight <= MipHeightInTiles * TileSizeInPixels);
 
 			const FRHITextureCreateDesc ResizedDesc =
-				FRHITextureCreateDesc::Create2D(*TextureName.ToString(), MipWidth, MipHeight, PixelFormat)
+				FRHITextureCreateDesc::Create2D(*Name, MipWidth, MipHeight, PixelFormat)
 				.SetFlags(Desc.Flags)
 				.SetInitialState(ERHIAccess::CopyDest);
 
@@ -1527,7 +1528,7 @@ void FVirtualTexture2DResource::InitializeEditorResources(IVirtualTexture* InVir
 
 		TextureRHI = Texture2DRHI;
 		TextureRHI->SetName(TextureName);
-		RHIBindDebugLabelName(TextureRHI, *TextureName.ToString());
+		RHIBindDebugLabelName(TextureRHI, *Name);
 		RHIUpdateTextureReference(TextureReferenceRHI, TextureRHI);
 
 		bIgnoreGammaConversions = !bSRGB && !IsHDR(PixelFormat);
