Fix dlecaration shadowing a field CIS warning on Incremental UnrealEditor Mac.

diff --git a/Engine/Source/Runtime/Engine/Private/WorldPartition/RuntimeHashSet/WorldPartitionRuntimeHashSet.cpp b/Engine/Source/Runtime/Engine/Private/WorldPartition/RuntimeHashSet/WorldPartitionRuntimeHashSet.cpp
--- a/Engine/Source/Runtime/Engine/Private/WorldPartition/RuntimeHashSet/WorldPartitionRuntimeHashSet.cpp
+++ b/Engine/Source/Runtime/Engine/Private/WorldPartition/RuntimeHashSet/WorldPartitionRuntimeHashSet.cpp
@@ -131,9 +131,9 @@ void UWorldPartitionRuntimeHashSet::PostLoad()
 	else
 #endif
 	{
-		ForEachStreamingObject([](const URuntimeHashSetExternalStreamingObject* StreamingObject)
+		ForEachStreamingObject([](const URuntimeHashSetExternalStreamingObject* CurStreamingObject)
 		{
-			StreamingObject->CreatePartitionsSpatialIndex();
+			CurStreamingObject->CreatePartitionsSpatialIndex();
 		});
 	}
 }
@@ -419,9 +419,9 @@ bool UWorldPartitionRuntimeHashSet::IsValidGrid(FName GridName) const
 TArray<UWorldPartitionRuntimeCell*> UWorldPartitionRuntimeHashSet::GetAlwaysLoadedCells() const
 {
 	TArray<UWorldPartitionRuntimeCell*> NonSpatiallyLoadedRuntimeCells;
-	ForEachStreamingObject([&NonSpatiallyLoadedRuntimeCells](const URuntimeHashSetExternalStreamingObject* StreamingObject)
+	ForEachStreamingObject([&NonSpatiallyLoadedRuntimeCells](const URuntimeHashSetExternalStreamingObject* CurStreamingObject)
 	{
-		NonSpatiallyLoadedRuntimeCells.Append(StreamingObject->NonSpatiallyLoadedRuntimeCells);
+		NonSpatiallyLoadedRuntimeCells.Append(CurStreamingObject->NonSpatiallyLoadedRuntimeCells);
 	});
 	return NonSpatiallyLoadedRuntimeCells;
 }
@@ -523,14 +523,14 @@ void UWorldPartitionRuntimeHashSet::ForEachStreamingCells(TFunctionRef<bool(cons
 		}
 	};
 
-	ForEachStreamingObject([&ForEachStreamingCells, &ForEachNonStreamingCells](const URuntimeHashSetExternalStreamingObject* StreamingObject)
+	ForEachStreamingObject([&ForEachStreamingCells, &ForEachNonStreamingCells](const URuntimeHashSetExternalStreamingObject* CurStreamingObject)
 	{
-		for (const FRuntimePartitionStreamingData& StreamingData : StreamingObject->RuntimeStreamingData)
+		for (const FRuntimePartitionStreamingData& StreamingData : CurStreamingObject->RuntimeStreamingData)
 		{
 			ForEachStreamingCells(StreamingData.SpatialIndex.Get());
 		}
 
-		ForEachNonStreamingCells(StreamingObject->NonSpatiallyLoadedRuntimeCells);
+		ForEachNonStreamingCells(CurStreamingObject->NonSpatiallyLoadedRuntimeCells);
 	});
 }
 
@@ -581,14 +581,14 @@ void UWorldPartitionRuntimeHashSet::ForEachStreamingCellsQuery(const FWorldParti
 		}
 	};
 
-	ForEachStreamingObject([&ForEachStreamingCells, &ForEachNonStreamingCells](const URuntimeHashSetExternalStreamingObject* StreamingObject)
+	ForEachStreamingObject([&ForEachStreamingCells, &ForEachNonStreamingCells](const URuntimeHashSetExternalStreamingObject* CurStreamingObject)
 	{
-		for (const FRuntimePartitionStreamingData& StreamingData : StreamingObject->RuntimeStreamingData)
+		for (const FRuntimePartitionStreamingData& StreamingData : CurStreamingObject->RuntimeStreamingData)
 		{
 			ForEachStreamingCells(StreamingData.SpatialIndex.Get());
 		}
 
-		ForEachNonStreamingCells(StreamingObject->NonSpatiallyLoadedRuntimeCells);
+		ForEachNonStreamingCells(CurStreamingObject->NonSpatiallyLoadedRuntimeCells);
 	});
 }
 
@@ -653,14 +653,14 @@ void UWorldPartitionRuntimeHashSet::ForEachStreamingCellsSources(const TArray<FW
 		}
 	};
 
-	ForEachStreamingObject([&ForEachStreamingCells, &ForEachNonStreamingCells](const URuntimeHashSetExternalStreamingObject* StreamingObject)
+	ForEachStreamingObject([&ForEachStreamingCells, &ForEachNonStreamingCells](const URuntimeHashSetExternalStreamingObject* CurStreamingObject)
 	{
-		for (const FRuntimePartitionStreamingData& StreamingData : StreamingObject->RuntimeStreamingData)
+		for (const FRuntimePartitionStreamingData& StreamingData : CurStreamingObject->RuntimeStreamingData)
 		{
 			ForEachStreamingCells(StreamingData.SpatialIndex.Get(), StreamingData.LoadingRange, StreamingData.Name);
 		}
 
-		ForEachNonStreamingCells(StreamingObject->NonSpatiallyLoadedRuntimeCells);
+		ForEachNonStreamingCells(CurStreamingObject->NonSpatiallyLoadedRuntimeCells);
 	});
 
 	auto ExecuteFuncOnCells = [Func](const TSet<const UWorldPartitionRuntimeCell*>& Cells, EStreamingSourceTargetState TargetState)
