DeadBlending: added options for how to handle missing curves, option for linear interpolation of scales, and improved performance.

diff --git a/Engine/Source/Runtime/Engine/Classes/Animation/AnimNode_DeadBlending.h b/Engine/Source/Runtime/Engine/Classes/Animation/AnimNode_DeadBlending.h
--- a/Engine/Source/Runtime/Engine/Classes/Animation/AnimNode_DeadBlending.h
+++ b/Engine/Source/Runtime/Engine/Classes/Animation/AnimNode_DeadBlending.h
@@ -39,7 +39,7 @@ struct FAnimNode_DeadBlending : public FAnimNode_Base, public IBoneReferenceSkel
 
 private:
 
-	// If the default blend settings should always be used rather than those coming from the inertialization request.
+	// When enabled, the default blend settings will always be used rather than those coming from the inertialization request.
 	UPROPERTY(EditAnywhere, Category = Blending)
 	bool bAlwaysUseDefaultBlendSettings = false;
 
@@ -63,6 +63,14 @@ private:
 	UPROPERTY(EditAnywhere, Category = Blending, meta = (Min = "0.0", UIMin = "0.0"))
 	float BlendTimeMultiplier = 1.0f;
 
+	/**
+	 * When enabled, bone scales will be linearly interpolated. This is slightly more performant and consistent with the
+	 * rest of Unreal but visually gives the appearance of the rate of change of scale being affected by the overall
+	 * size of the bone.
+	 */
+	UPROPERTY(EditAnywhere, Category = Blending)
+	bool bLinearlyInterpolateScales = false;
+
 	// List of curves that should not use inertial blending. These curves will change instantly when the animation switches.
 	UPROPERTY(EditAnywhere, Category = Filter)
 	TArray<FName> FilteredCurves;
@@ -95,6 +103,20 @@ private:
 	UPROPERTY(EditAnywhere, Category = Extrapolation, meta = (Min = "0.0", UIMin = "0.0"), DisplayName = "Maximum Extrapolation Half Life")
 	float ExtrapolationHalfLifeMax = 1.0f;
 
+	/**
+	 * When enabled, curves will blend in from zero when they are in the animation being transitioned to but not in the 
+	 * animation being transitioned from.
+	 */
+	UPROPERTY(EditAnywhere, Category = Curves)
+	bool bBlendInMissingCurves = true;
+
+	/**
+	 * When enabled, curves will blend out to zero when they are in the animation being transitioned from but not in the 
+	 * animation being transitioned to.
+	 */
+	UPROPERTY(EditAnywhere, Category = Curves)
+	bool bBlendOutMissingCurves = true;
+
 #if WITH_EDITORONLY_DATA
 	
 	// This setting can be used to show what the extrapolation of the animation looks like.
diff --git a/Engine/Source/Runtime/Engine/Private/Animation/AnimNode_DeadBlending.cpp b/Engine/Source/Runtime/Engine/Private/Animation/AnimNode_DeadBlending.cpp
--- a/Engine/Source/Runtime/Engine/Private/Animation/AnimNode_DeadBlending.cpp
+++ b/Engine/Source/Runtime/Engine/Private/Animation/AnimNode_DeadBlending.cpp
@@ -101,12 +101,19 @@ namespace UE::Anim::DeadBlending::Private
 			FMath::InvExpApprox(V.Z));
 	}
 
-	static inline FVector VectorEerp(const FVector V, const FVector W, const float Alpha)
+	static inline FVector VectorEerp(const FVector V, const FVector W, const float Alpha, const float Epsilon = UE_SMALL_NUMBER)
 	{
-		return FVector(
-			FMath::Pow(V.X, (1.0f - Alpha)) * FMath::Pow(W.X, Alpha),
-			FMath::Pow(V.Y, (1.0f - Alpha)) * FMath::Pow(W.Y, Alpha),
-			FMath::Pow(V.Z, (1.0f - Alpha)) * FMath::Pow(W.Z, Alpha));
+		if (FVector::DistSquared(V, W) < Epsilon)
+		{
+			return FMath::Lerp(V, W, Alpha);
+		}
+		else
+		{
+			return FVector(
+				FMath::Pow(V.X, (1.0f - Alpha)) * FMath::Pow(W.X, Alpha),
+				FMath::Pow(V.Y, (1.0f - Alpha)) * FMath::Pow(W.Y, Alpha),
+				FMath::Pow(V.Z, (1.0f - Alpha)) * FMath::Pow(W.Z, Alpha));
+		}
 	}
 
 	static inline FVector VectorExp(const FVector V)
@@ -132,8 +139,15 @@ namespace UE::Anim::DeadBlending::Private
 		const FVector DecayHalflife,
 		const float Epsilon = UE_SMALL_NUMBER)
 	{
-		const FVector C = VectorDivMax(Ln2, DecayHalflife, Epsilon);
-		return Translation + VectorDivMax(Velocity, C, Epsilon) * (FVector::OneVector - VectorInvExpApprox(C * Time));
+		if (Velocity.SquaredLength() > Epsilon)
+		{
+			const FVector C = VectorDivMax(Ln2, DecayHalflife, Epsilon);
+			return Translation + VectorDivMax(Velocity, C, Epsilon) * (FVector::OneVector - VectorInvExpApprox(C * Time));
+		}
+		else
+		{
+			return Translation;
+		}
 	}
 
 	static inline FQuat ExtrapolateRotation(
@@ -143,8 +157,15 @@ namespace UE::Anim::DeadBlending::Private
 		const FVector DecayHalflife,
 		const float Epsilon = UE_SMALL_NUMBER)
 	{
-		const FVector C = VectorDivMax(Ln2, DecayHalflife, Epsilon);
-		return FQuat::MakeFromRotationVector(VectorDivMax(Velocity, C, Epsilon) * (FVector::OneVector - VectorInvExpApprox(C * Time))) * Rotation;
+		if (Velocity.SquaredLength() > Epsilon)
+		{
+			const FVector C = VectorDivMax(Ln2, DecayHalflife, Epsilon);
+			return FQuat::MakeFromRotationVector(VectorDivMax(Velocity, C, Epsilon) * (FVector::OneVector - VectorInvExpApprox(C * Time))) * Rotation;
+		}
+		else
+		{
+			return Rotation;
+		}
 	}
 
 	static inline FVector ExtrapolateScale(
@@ -154,8 +175,15 @@ namespace UE::Anim::DeadBlending::Private
 		const FVector DecayHalflife,
 		const float Epsilon = UE_SMALL_NUMBER)
 	{
-		const FVector C = VectorDivMax(Ln2, DecayHalflife, Epsilon);
-		return VectorExp(VectorDivMax(Velocity, C, Epsilon) * (FVector::OneVector - VectorInvExpApprox(C * Time))) * Scale;
+		if (Velocity.SquaredLength() > Epsilon)
+		{
+			const FVector C = VectorDivMax(Ln2, DecayHalflife, Epsilon);
+			return VectorExp(VectorDivMax(Velocity, C, Epsilon) * (FVector::OneVector - VectorInvExpApprox(C * Time))) * Scale;
+		}
+		else
+		{
+			return Scale;
+		}
 	}
 
 	static inline float ExtrapolateCurve(
@@ -165,8 +193,15 @@ namespace UE::Anim::DeadBlending::Private
 		const float DecayHalflife,
 		const float Epsilon = UE_SMALL_NUMBER)
 	{
-		const float C = Ln2 / FMath::Max(DecayHalflife, Epsilon);
-		return Curve + FMath::Max(Velocity / C, Epsilon) * (1.0f - FMath::InvExpApprox(C * Time));
+		if (FMath::Square(Velocity) > Epsilon)
+		{
+			const float C = Ln2 / FMath::Max(DecayHalflife, Epsilon);
+			return Curve + FMath::Max(Velocity / C, Epsilon) * (1.0f - FMath::InvExpApprox(C * Time));
+		}
+		else
+		{
+			return Curve;
+		}
 	}
 
 	static inline float ClipMagnitudeToGreaterThanEpsilon(const float X, const float Epsilon = UE_KINDA_SMALL_NUMBER)
@@ -352,6 +387,13 @@ void FAnimNode_DeadBlending::InitFrom(
 		UE::Anim::FNamedValueArrayUtils::Union(CurveData, InCurves,
 			[this](FDeadBlendingCurveElement& OutResultElement, const UE::Anim::FCurveElement& InElement1, UE::Anim::ENamedValueUnionFlags InFlags)
 			{
+				if (bBlendInMissingCurves)
+				{
+					OutResultElement.Valid = true;
+					OutResultElement.Value = 0.0f;
+					OutResultElement.Velocity = 0.0f;
+				}
+
 				if (OutResultElement.Valid)
 				{
 					OutResultElement.HalfLife = UE::Anim::DeadBlending::Private::ComputeDecayHalfLifeFromDiffAndVelocity(
@@ -441,55 +483,56 @@ void FAnimNode_DeadBlending::ApplyTo(FCompactPose& InOutPose, FBlendedCurve& InO
 			InOutPose[BoneIndex].SetTranslation(FMath::Lerp(InOutPose[BoneIndex].GetTranslation(), ExtrapolatedTranslation, Alpha));
 			InOutPose[BoneIndex].SetRotation(FQuat::MakeFromRotationVector(RotationDiff.ToRotationVector() * Alpha) * InOutPose[BoneIndex].GetRotation());
 
-			// Here we use `Eerp` rather than `Lerp` to interpolate scales correctly (see: https://theorangeduck.com/page/scalar-velocity).
-			// This is inconsistent with the rest of Unreal which (mostly) uses `Lerp` on scales. The decision to use `Eerp` here is partially 
-			// due to the fact we are also dealing properly with scalar velocities in this node, and partially to try and not to lock this node 
-			// into having the same incorrect behavior by default. We can add an option to interpolate scales with `Lerp` later down the line if 
-			// users want this, but we will not be able to change the default behavior again, so the decision is to opt for the most correct by 
-			// default this time even if it is somewhat less performant and inconsistent with other parts of Unreal.
-			InOutPose[BoneIndex].SetScale3D(UE::Anim::DeadBlending::Private::VectorEerp(InOutPose[BoneIndex].GetScale3D(), ExtrapolatedScale, Alpha));
+			// Here we use `Eerp` rather than `Lerp` to interpolate scales by default (see: https://theorangeduck.com/page/scalar-velocity).
+			// This default is inconsistent with the rest of Unreal which (mostly) uses `Lerp` on scales. The decision 
+			// to use `Eerp` by default here is partially due to the fact we are also dealing properly with scalar 
+			// velocities in this node, and partially to try and not to lock this node into having the same less 
+			// accurate behavior by default. Users still have the option to interpolate scales with `Lerp` if they want.
+			if (bLinearlyInterpolateScales)
+			{
+				InOutPose[BoneIndex].SetScale3D(FMath::Lerp(InOutPose[BoneIndex].GetScale3D(), ExtrapolatedScale, Alpha));
+			}
+			else
+			{
+				InOutPose[BoneIndex].SetScale3D(UE::Anim::DeadBlending::Private::VectorEerp(InOutPose[BoneIndex].GetScale3D(), ExtrapolatedScale, Alpha));
+			}
 		}
 	}
 
-	if (InOutCurves.Num() > 0)
-	{
-		// Compute Blend Alpha
+	// Compute Blend Alpha
 
-		const float CurveAlpha = 1.0f - FAlphaBlend::AlphaToBlendOption(
-			InertializationTime / FMath::Max(InertializationDuration, UE_SMALL_NUMBER),
-			InertializationBlendMode, InertializationCustomBlendCurve);
+	const float CurveAlpha = 1.0f - FAlphaBlend::AlphaToBlendOption(
+		InertializationTime / FMath::Max(InertializationDuration, UE_SMALL_NUMBER),
+		InertializationBlendMode, InertializationCustomBlendCurve);
 
-		// Blend Curves
+	// Blend Curves
 
-		if (CurveAlpha != 0.0f)
+	UE::Anim::FNamedValueArrayUtils::Union(InOutCurves, CurveData,
+		[CurveAlpha, this](UE::Anim::FCurveElement& OutResultElement, const FDeadBlendingCurveElement& InElement1, UE::Anim::ENamedValueUnionFlags InFlags)
 		{
-			UE::Anim::FNamedValueArrayUtils::Union(InOutCurves, CurveData,
-				[CurveAlpha, this](UE::Anim::FCurveElement& OutResultElement, const FDeadBlendingCurveElement& InElement1, UE::Anim::ENamedValueUnionFlags InFlags)
-				{
-					// Compute Extrapolated Curve Value
+			// Compute Extrapolated Curve Value
 
-					if (InElement1.Valid)
-					{
-						const float ExtrapolatedCurve = UE::Anim::DeadBlending::Private::ExtrapolateCurve(
-							InElement1.Value,
-							InElement1.Velocity,
-							InertializationTime,
-							InElement1.HalfLife);
+			if (InElement1.Valid && (bBlendOutMissingCurves || (bool)(InFlags & UE::Anim::ENamedValueUnionFlags::ValidArg0)))
+			{
+				const float ExtrapolatedCurve = UE::Anim::DeadBlending::Private::ExtrapolateCurve(
+					InElement1.Value,
+					InElement1.Velocity,
+					InertializationTime,
+					InElement1.HalfLife);
 
 #if WITH_EDITORONLY_DATA
-						if (bShowExtrapolations)
-						{
-							OutResultElement.Value = ExtrapolatedCurve;
-							OutResultElement.Flags |= InElement1.Flags;
-							return;
-						}
+				if (bShowExtrapolations)
+				{
+					OutResultElement.Value = ExtrapolatedCurve;
+					OutResultElement.Flags |= InElement1.Flags;
+					return;
+				}
 #endif
-						OutResultElement.Value = FMath::Lerp(OutResultElement.Value, ExtrapolatedCurve, CurveAlpha);
-						OutResultElement.Flags |= InElement1.Flags;
-					}
-				});
-		}
-	}
+
+				OutResultElement.Value = FMath::Lerp(OutResultElement.Value, ExtrapolatedCurve, CurveAlpha);
+				OutResultElement.Flags |= InElement1.Flags;
+			}
+		});
 }
 
 class USkeleton* FAnimNode_DeadBlending::GetSkeleton(bool& bInvalidSkeletonIsError, const IPropertyHandle* PropertyHandle)
