Fix incorrect Distance Field / Card Representation when multiple mesh sections with different settings point to the same material slot
- Different mesh sections can point to the same material slot which was a problem since bAffectDistanceFieldLighting flag can be set per section but we were using an array with an entry per material slot instead of per section.
- To address this now we build array of FSignedDistanceFieldBuildSectionData (one entry per section) so section don't override each other.

#prelight 64921d9ab76bf403cda2f33f

diff --git a/Engine/Source/Developer/MeshUtilities/Private/MeshCardRepresentationUtilities.cpp b/Engine/Source/Developer/MeshUtilities/Private/MeshCardRepresentationUtilities.cpp
--- a/Engine/Source/Developer/MeshUtilities/Private/MeshCardRepresentationUtilities.cpp
+++ b/Engine/Source/Developer/MeshUtilities/Private/MeshCardRepresentationUtilities.cpp
@@ -1054,7 +1054,7 @@ bool FMeshUtilities::GenerateCardRepresentationData(
 	const FSourceMeshDataForDerivedDataTask& SourceMeshData,
 	const FStaticMeshLODResources& LODModel,
 	class FQueuedThreadPool& ThreadPool,
-	const TArray<FSignedDistanceFieldBuildMaterialData>& MaterialBlendModes,
+	const TArray<FSignedDistanceFieldBuildSectionData>& SectionData,
 	const FBoxSphereBounds& Bounds,
 	const FDistanceFieldVolumeData* DistanceFieldVolumeData,
 	int32 MaxLumenMeshCards,
@@ -1072,7 +1072,7 @@ bool FMeshUtilities::GenerateCardRepresentationData(
 		MeshRepresentation::SetupEmbreeScene(MeshName,
 			SourceMeshData,
 			LODModel,
-			MaterialBlendModes,
+			SectionData,
 			bGenerateAsIfTwoSided,
 			EmbreeScene);
 
diff --git a/Engine/Source/Developer/MeshUtilities/Private/MeshDistanceFieldUtilities.cpp b/Engine/Source/Developer/MeshUtilities/Private/MeshDistanceFieldUtilities.cpp
--- a/Engine/Source/Developer/MeshUtilities/Private/MeshDistanceFieldUtilities.cpp
+++ b/Engine/Source/Developer/MeshUtilities/Private/MeshDistanceFieldUtilities.cpp
@@ -252,7 +252,7 @@ void FMeshUtilities::GenerateSignedDistanceFieldVolumeData(
 	const FSourceMeshDataForDerivedDataTask& SourceMeshData,
 	const FStaticMeshLODResources& LODModel,
 	class FQueuedThreadPool& ThreadPool,
-	const TArray<FSignedDistanceFieldBuildMaterialData>& MaterialBlendModes,
+	const TArray<FSignedDistanceFieldBuildSectionData>& SectionData,
 	const FBoxSphereBounds& Bounds,
 	float DistanceFieldResolutionScale,
 	bool bGenerateAsIfTwoSided,
@@ -268,7 +268,7 @@ void FMeshUtilities::GenerateSignedDistanceFieldVolumeData(
 		MeshRepresentation::SetupEmbreeScene(MeshName,
 			SourceMeshData,
 			LODModel,
-			MaterialBlendModes,
+			SectionData,
 			bGenerateAsIfTwoSided,
 			EmbreeScene);
 
@@ -529,7 +529,7 @@ void FMeshUtilities::GenerateSignedDistanceFieldVolumeData(
 	const FSourceMeshDataForDerivedDataTask& SourceMeshData,
 	const FStaticMeshLODResources& LODModel,
 	class FQueuedThreadPool& ThreadPool,
-	const TArray<FSignedDistanceFieldBuildMaterialData>& MaterialBlendModes,
+	const TArray<FSignedDistanceFieldBuildSectionData>& SectionData,
 	const FBoxSphereBounds& Bounds,
 	float DistanceFieldResolutionScale,
 	bool bGenerateAsIfTwoSided,
diff --git a/Engine/Source/Developer/MeshUtilities/Private/MeshRepresentationCommon.cpp b/Engine/Source/Developer/MeshUtilities/Private/MeshRepresentationCommon.cpp
--- a/Engine/Source/Developer/MeshUtilities/Private/MeshRepresentationCommon.cpp
+++ b/Engine/Source/Developer/MeshUtilities/Private/MeshRepresentationCommon.cpp
@@ -126,7 +126,7 @@ void MeshRepresentation::SetupEmbreeScene(
 	FString MeshName,
 	const FSourceMeshDataForDerivedDataTask& SourceMeshData,
 	const FStaticMeshLODResources& LODModel,
-	const TArray<FSignedDistanceFieldBuildMaterialData>& MaterialBlendModes,
+	const TArray<FSignedDistanceFieldBuildSectionData>& SectionData,
 	bool bGenerateAsIfTwoSided,
 	FEmbreeScene& EmbreeScene)
 {
@@ -165,6 +165,11 @@ void MeshRepresentation::SetupEmbreeScene(
 	}
 #endif
 
+	if (LODModel.Sections.Num() > SectionData.Num())
+	{
+		UE_LOG(LogMeshUtilities, Warning, TEXT("Unexpected number of mesh sections when setting up Embree Scene for %s."), *MeshName);
+	}
+
 	TArray<int32> FilteredTriangles;
 	FilteredTriangles.Empty(NumTriangles);
 
@@ -213,9 +218,9 @@ void MeshRepresentation::SetupEmbreeScene(
 
 					if ((uint32)(TriangleIndex * 3) >= Section.FirstIndex && (uint32)(TriangleIndex * 3) < Section.FirstIndex + Section.NumTriangles * 3)
 					{
-						if (MaterialBlendModes.IsValidIndex(Section.MaterialIndex))
+						if (SectionData.IsValidIndex(SectionIndex))
 						{
-							bTriangleIsOpaqueOrMasked = !IsTranslucentBlendMode(MaterialBlendModes[Section.MaterialIndex].BlendMode) && MaterialBlendModes[Section.MaterialIndex].bAffectDistanceFieldLighting;
+							bTriangleIsOpaqueOrMasked = !IsTranslucentBlendMode(SectionData[SectionIndex].BlendMode) && SectionData[SectionIndex].bAffectDistanceFieldLighting;
 						}
 
 						break;
@@ -279,9 +284,9 @@ void MeshRepresentation::SetupEmbreeScene(
 
 			if ((uint32)(TriangleIndex * 3) >= Section.FirstIndex && (uint32)(TriangleIndex * 3) < Section.FirstIndex + Section.NumTriangles * 3)
 			{
-				if (MaterialBlendModes.IsValidIndex(Section.MaterialIndex))
+				if (SectionData.IsValidIndex(SectionIndex))
 				{
-					bTriangleIsTwoSided = MaterialBlendModes[Section.MaterialIndex].bTwoSided;
+					bTriangleIsTwoSided = SectionData[SectionIndex].bTwoSided;
 				}
 
 				break;
@@ -354,11 +359,11 @@ void MeshRepresentation::SetupEmbreeScene(
 		{
 			const FStaticMeshSection& Section = LODModel.Sections[SectionIndex];
 
-			if (MaterialBlendModes.IsValidIndex(Section.MaterialIndex))
+			if (SectionData.IsValidIndex(SectionIndex))
 			{
 				NumTrianglesTotal += Section.NumTriangles;
 
-				if (MaterialBlendModes[Section.MaterialIndex].bTwoSided)
+				if (SectionData[SectionIndex].bTwoSided)
 				{
 					NumTwoSidedTriangles += Section.NumTriangles;
 				}
diff --git a/Engine/Source/Developer/MeshUtilities/Private/MeshRepresentationCommon.h b/Engine/Source/Developer/MeshUtilities/Private/MeshRepresentationCommon.h
--- a/Engine/Source/Developer/MeshUtilities/Private/MeshRepresentationCommon.h
+++ b/Engine/Source/Developer/MeshUtilities/Private/MeshRepresentationCommon.h
@@ -159,7 +159,7 @@ namespace MeshRepresentation
 	void SetupEmbreeScene(FString MeshName,
 		const FSourceMeshDataForDerivedDataTask& SourceMeshData,
 		const FStaticMeshLODResources& LODModel,
-		const TArray<FSignedDistanceFieldBuildMaterialData>& MaterialBlendModes,
+		const TArray<FSignedDistanceFieldBuildSectionData>& SectionData,
 		bool bGenerateAsIfTwoSided,
 		FEmbreeScene& EmbreeScene);
 
diff --git a/Engine/Source/Developer/MeshUtilities/Private/MeshUtilitiesPrivate.h b/Engine/Source/Developer/MeshUtilities/Private/MeshUtilitiesPrivate.h
--- a/Engine/Source/Developer/MeshUtilities/Private/MeshUtilitiesPrivate.h
+++ b/Engine/Source/Developer/MeshUtilities/Private/MeshUtilitiesPrivate.h
@@ -60,7 +60,7 @@ private:
 		const FSourceMeshDataForDerivedDataTask& SourceMeshData,
 		const FStaticMeshLODResources& LODModel,
 		class FQueuedThreadPool& ThreadPool,
-		const TArray<FSignedDistanceFieldBuildMaterialData>& MaterialBlendModes,
+		const TArray<FSignedDistanceFieldBuildSectionData>& SectionData,
 		const FBoxSphereBounds& Bounds,
 		float DistanceFieldResolutionScale,
 		bool bGenerateAsIfTwoSided,
@@ -71,7 +71,7 @@ private:
 		const FSourceMeshDataForDerivedDataTask& SourceMeshData,
 		const FStaticMeshLODResources& LODModel,
 		class FQueuedThreadPool& ThreadPool,
-		const TArray<FSignedDistanceFieldBuildMaterialData>& MaterialBlendModes,
+		const TArray<FSignedDistanceFieldBuildSectionData>& SectionData,
 		const FBoxSphereBounds& Bounds,
 		const FDistanceFieldVolumeData* DistanceFieldVolumeData,
 		int32 MaxLumenMeshCards,
diff --git a/Engine/Source/Developer/MeshUtilities/Public/MeshUtilities.h b/Engine/Source/Developer/MeshUtilities/Public/MeshUtilities.h
--- a/Engine/Source/Developer/MeshUtilities/Public/MeshUtilities.h
+++ b/Engine/Source/Developer/MeshUtilities/Public/MeshUtilities.h
@@ -48,9 +48,8 @@ namespace ETangentOptions
 	};
 };
 
-class FSignedDistanceFieldBuildMaterialData
+struct FSignedDistanceFieldBuildSectionData
 {
-public:
 	EBlendMode BlendMode = BLEND_Opaque;
 	bool bTwoSided = false;
 	bool bAffectDistanceFieldLighting = true;
@@ -130,7 +129,7 @@ public:
 		const FSourceMeshDataForDerivedDataTask& SourceMeshData,
 		const FStaticMeshLODResources& LODModel,
 		class FQueuedThreadPool& ThreadPool,
-		const TArray<FSignedDistanceFieldBuildMaterialData>& MaterialBlendModes,
+		const TArray<FSignedDistanceFieldBuildSectionData>& SectionData,
 		const FBoxSphereBounds& Bounds,
 		float DistanceFieldResolutionScale,
 		bool bGenerateAsIfTwoSided,
@@ -141,7 +140,7 @@ public:
 		const FSourceMeshDataForDerivedDataTask& SourceMeshData,
 		const FStaticMeshLODResources& LODModel,
 		class FQueuedThreadPool& ThreadPool,
-		const TArray<FSignedDistanceFieldBuildMaterialData>& MaterialBlendModes,
+		const TArray<FSignedDistanceFieldBuildSectionData>& SectionData,
 		const FBoxSphereBounds& Bounds,
 		const class FDistanceFieldVolumeData* DistanceFieldVolumeData,
 		int32 MaxLumenMeshCards,
diff --git a/Engine/Source/Runtime/Engine/Private/DistanceFieldAtlas.cpp b/Engine/Source/Runtime/Engine/Private/DistanceFieldAtlas.cpp
--- a/Engine/Source/Runtime/Engine/Private/DistanceFieldAtlas.cpp
+++ b/Engine/Source/Runtime/Engine/Private/DistanceFieldAtlas.cpp
@@ -125,7 +125,7 @@ FDistanceFieldAsyncQueue* GDistanceFieldAsyncQueue = NULL;
 #if WITH_EDITOR
 
 // DDC key for distance field data, must be changed when modifying the generation code or data format
-#define DISTANCEFIELD_DERIVEDDATA_VER TEXT("295895D9-FFCE-48B5-9A4B-131DB318441E")
+#define DISTANCEFIELD_DERIVEDDATA_VER TEXT("C8D7BA5F-5874-4FFB-A264-7B7DE9299F68")
 
 FString BuildDistanceFieldDerivedDataKey(const FString& InMeshKey)
 {
@@ -147,15 +147,14 @@ FString BuildDistanceFieldDerivedDataKey(const FString& InMeshKey)
 
 #if WITH_EDITORONLY_DATA
 
-void BuildSignedDistanceFieldBuildMaterialData(UStaticMesh* Mesh, uint32 LODIndex, TArray<FSignedDistanceFieldBuildMaterialData>& OutData)
+void BuildSignedDistanceFieldBuildSectionData(UStaticMesh* Mesh, uint32 LODIndex, TArray<FSignedDistanceFieldBuildSectionData>& OutData)
 {
 	const TArray<FStaticMaterial>& StaticMaterials = Mesh->GetStaticMaterials();
-
-	OutData.SetNum(StaticMaterials.Num());
-
 	const FMeshSectionInfoMap& SectionInfoMap = Mesh->GetSectionInfoMap();
 
-	for (int32 SectionIndex = 0; SectionIndex < SectionInfoMap.GetSectionNumber(LODIndex); SectionIndex++)
+	OutData.SetNum(SectionInfoMap.GetSectionNumber(LODIndex));
+
+	for (int32 SectionIndex = 0; SectionIndex < OutData.Num(); SectionIndex++)
 	{
 		const FMeshSectionInfo& Section = SectionInfoMap.Get(LODIndex, SectionIndex);
 
@@ -164,14 +163,14 @@ void BuildSignedDistanceFieldBuildMaterialData(UStaticMesh* Mesh, uint32 LODInde
 			continue;
 		}
 
-		FSignedDistanceFieldBuildMaterialData& MaterialData = OutData[Section.MaterialIndex];
-		MaterialData.bAffectDistanceFieldLighting = Section.bAffectDistanceFieldLighting;
+		FSignedDistanceFieldBuildSectionData& SectionData = OutData[SectionIndex];
+		SectionData.bAffectDistanceFieldLighting = Section.bAffectDistanceFieldLighting;
 
 		UMaterialInterface* MaterialInterface = StaticMaterials[Section.MaterialIndex].MaterialInterface;
 		if (MaterialInterface)
 		{
-			MaterialData.BlendMode = MaterialInterface->GetBlendMode();
-			MaterialData.bTwoSided = MaterialInterface->IsTwoSided();
+			SectionData.BlendMode = MaterialInterface->GetBlendMode();
+			SectionData.bTwoSided = MaterialInterface->IsTwoSided();
 		}
 	}
 }
@@ -180,19 +179,19 @@ void FDistanceFieldVolumeData::CacheDerivedData(const FString& InStaticMeshDeriv
 {
 	TRACE_CPUPROFILER_EVENT_SCOPE(FDistanceFieldVolumeData::CacheDerivedData);
 
-	TArray<FSignedDistanceFieldBuildMaterialData> BuildMaterialData;
+	TArray<FSignedDistanceFieldBuildSectionData> BuildSectionData;
 
 	const uint32 LODIndex = 0;
-	BuildSignedDistanceFieldBuildMaterialData(Mesh, LODIndex, BuildMaterialData);
+	BuildSignedDistanceFieldBuildSectionData(Mesh, LODIndex, BuildSectionData);
 
 	FString DistanceFieldKey = BuildDistanceFieldDerivedDataKey(InStaticMeshDerivedDataKey);
 
-	for (int32 MaterialIndex = 0; MaterialIndex < BuildMaterialData.Num(); MaterialIndex++)
+	for (int32 MaterialIndex = 0; MaterialIndex < BuildSectionData.Num(); MaterialIndex++)
 	{
 		DistanceFieldKey += FString::Printf(TEXT("_M%u_%u_%u"), 
-			(uint32)BuildMaterialData[MaterialIndex].BlendMode,
-			BuildMaterialData[MaterialIndex].bTwoSided ? 1 : 0,
-			BuildMaterialData[MaterialIndex].bAffectDistanceFieldLighting ? 1 : 0);
+			(uint32)BuildSectionData[MaterialIndex].BlendMode,
+			BuildSectionData[MaterialIndex].bTwoSided ? 1 : 0,
+			BuildSectionData[MaterialIndex].bAffectDistanceFieldLighting ? 1 : 0);
 	}
 
 	TArray<uint8> DerivedData;
@@ -222,7 +221,7 @@ void FDistanceFieldVolumeData::CacheDerivedData(const FString& InStaticMeshDeriv
 		NewTask->GeneratedVolumeData = new FDistanceFieldVolumeData();
 		NewTask->GeneratedVolumeData->AssetName = Mesh->GetFName();
 		NewTask->GeneratedVolumeData->bAsyncBuilding = true;
-		NewTask->MaterialBlendModes = MoveTemp(BuildMaterialData);
+		NewTask->SectionData = MoveTemp(BuildSectionData);
 
 		// Nanite overrides source static mesh with a coarse representation. Need to load original data before we build the mesh SDF.
 		if (Mesh->IsNaniteEnabled())
@@ -724,7 +723,7 @@ void FDistanceFieldAsyncQueue::Build(FAsyncDistanceFieldTask* Task, FQueuedThrea
 			Task->SourceMeshData,
 			LODModel,
 			BuildThreadPool,
-			Task->MaterialBlendModes,
+			Task->SectionData,
 			Task->GenerateSource->GetRenderData()->Bounds,
 			Task->DistanceFieldResolutionScale,
 			Task->bGenerateDistanceFieldAsIfTwoSided,
diff --git a/Engine/Source/Runtime/Engine/Private/MeshCardRepresentation.cpp b/Engine/Source/Runtime/Engine/Private/MeshCardRepresentation.cpp
--- a/Engine/Source/Runtime/Engine/Private/MeshCardRepresentation.cpp
+++ b/Engine/Source/Runtime/Engine/Private/MeshCardRepresentation.cpp
@@ -98,7 +98,7 @@ FCardRepresentationAsyncQueue* GCardRepresentationAsyncQueue = NULL;
 #if WITH_EDITOR
 
 // DDC key for card representation data, must be changed when modifying the generation code or data format
-#define CARDREPRESENTATION_DERIVEDDATA_VER TEXT("FF1E9B99-1837-4F13-A892-13BD62922D0B")
+#define CARDREPRESENTATION_DERIVEDDATA_VER TEXT("E8B65E19-7C68-4BDC-BA5E-91C9A5855599")
 
 FString BuildCardRepresentationDerivedDataKey(const FString& InMeshKey, int32 MaxLumenMeshCards)
 {
@@ -117,7 +117,7 @@ FString BuildCardRepresentationDerivedDataKey(const FString& InMeshKey, int32 Ma
 
 #if WITH_EDITORONLY_DATA
 
-extern void BuildSignedDistanceFieldBuildMaterialData(UStaticMesh* Mesh, uint32 LODIndex, TArray<FSignedDistanceFieldBuildMaterialData>& OutData);
+extern void BuildSignedDistanceFieldBuildSectionData(UStaticMesh* Mesh, uint32 LODIndex, TArray<FSignedDistanceFieldBuildSectionData>& OutData);
 
 void BeginCacheMeshCardRepresentation(const ITargetPlatform* TargetPlatform, UStaticMesh* StaticMeshAsset, FStaticMeshRenderData& RenderData, const FString& DistanceFieldKey, FSourceMeshDataForDerivedDataTask* OptionalSourceMeshData)
 {
@@ -177,7 +177,7 @@ void FCardRepresentationData::CacheDerivedData(const FString& InDDCKey, const IT
 		NewTask->bGenerateDistanceFieldAsIfTwoSided = bGenerateDistanceFieldAsIfTwoSided;
 
 		const uint32 LODIndex = 0;
-		BuildSignedDistanceFieldBuildMaterialData(Mesh, LODIndex, NewTask->MaterialBlendModes);
+		BuildSignedDistanceFieldBuildSectionData(Mesh, LODIndex, NewTask->SectionData);
 
 		// Nanite overrides source static mesh with a coarse representation. Need to load original data before we build the mesh SDF.
 		if (OptionalSourceMeshData)
@@ -652,7 +652,7 @@ void FCardRepresentationAsyncQueue::Build(FAsyncCardRepresentationTask* Task, FQ
 			Task->SourceMeshData,
 			LODModel,
 			BuildThreadPool,
-			Task->MaterialBlendModes,
+			Task->SectionData,
 			Task->GenerateSource->GetRenderData()->Bounds,
 			Task->GenerateSource->GetRenderData()->LODResources[0].DistanceFieldData,
 			Task->MaxLumenMeshCards,
diff --git a/Engine/Source/Runtime/Engine/Public/DistanceFieldAtlas.h b/Engine/Source/Runtime/Engine/Public/DistanceFieldAtlas.h
--- a/Engine/Source/Runtime/Engine/Public/DistanceFieldAtlas.h
+++ b/Engine/Source/Runtime/Engine/Public/DistanceFieldAtlas.h
@@ -334,7 +334,7 @@ public:
 	FAsyncDistanceFieldTask();
 
 #if WITH_EDITOR
-	TArray<FSignedDistanceFieldBuildMaterialData> MaterialBlendModes;
+	TArray<FSignedDistanceFieldBuildSectionData> SectionData;
 #endif
 	FSourceMeshDataForDerivedDataTask SourceMeshData;
 	UStaticMesh* StaticMesh;
diff --git a/Engine/Source/Runtime/Engine/Public/MeshCardBuild.h b/Engine/Source/Runtime/Engine/Public/MeshCardBuild.h
--- a/Engine/Source/Runtime/Engine/Public/MeshCardBuild.h
+++ b/Engine/Source/Runtime/Engine/Public/MeshCardBuild.h
@@ -12,7 +12,7 @@
 #include <atomic>
 
 class UStaticMesh;
-class FSignedDistanceFieldBuildMaterialData;
+struct FSignedDistanceFieldBuildSectionData;
 
 class FLumenCardBuildData
 {
@@ -190,7 +190,7 @@ public:
 	bool bSuccess = false;
 
 #if WITH_EDITOR
-	TArray<FSignedDistanceFieldBuildMaterialData> MaterialBlendModes;
+	TArray<FSignedDistanceFieldBuildSectionData> SectionData;
 #endif
 
 	FSourceMeshDataForDerivedDataTask SourceMeshData;
