Fix a bug where FSpanAllocator::GetSparselyAllocatedSize was not taking pending free spans into account.
Keep track of current total allocated count so that calling GetSparselyAllocatedSize can be efficient.

diff --git a/Engine/Source/Runtime/Engine/Private/SpanAllocator.cpp b/Engine/Source/Runtime/Engine/Private/SpanAllocator.cpp
--- a/Engine/Source/Runtime/Engine/Private/SpanAllocator.cpp
+++ b/Engine/Source/Runtime/Engine/Private/SpanAllocator.cpp
@@ -4,7 +4,8 @@
 #include "CoreGlobals.h"
 
 FSpanAllocator::FSpanAllocator(bool bInGrowOnly) 
-	: CurrentMaxSize(0)
+	: CurrentSize(0)
+	, CurrentMaxSize(0)
 	, PeakMaxSize(0)
 	, FirstNonEmptySpan(0)
 	, bGrowOnly(bInGrowOnly)
@@ -13,6 +14,7 @@ FSpanAllocator::FSpanAllocator(bool bInGrowOnly)
 
 void FSpanAllocator::Reset()
 {
+	CurrentSize = 0;
 	CurrentMaxSize = 0;
 	PeakMaxSize = 0;
 	FirstNonEmptySpan = 0;
@@ -22,6 +24,7 @@ void FSpanAllocator::Reset()
 
 void FSpanAllocator::Empty()
 {
+	CurrentSize = 0;
 	CurrentMaxSize = 0;
 	PeakMaxSize = 0;
 	FirstNonEmptySpan = 0;
diff --git a/Engine/Source/Runtime/Engine/Public/SpanAllocator.h b/Engine/Source/Runtime/Engine/Public/SpanAllocator.h
--- a/Engine/Source/Runtime/Engine/Public/SpanAllocator.h
+++ b/Engine/Source/Runtime/Engine/Public/SpanAllocator.h
@@ -30,6 +30,8 @@ public:
 			FoundIndex = SearchFreeList(Num, FirstNonEmptySpan);
 		}
 
+		CurrentSize += Num;
+
 		// Use an existing free span if one is found
 		if (FoundIndex != INDEX_NONE)
 		{
@@ -61,23 +63,18 @@ public:
 	FORCEINLINE void Free(int32 BaseOffset, int32 Num = 1)
 	{
 		checkSlow(BaseOffset + Num <= CurrentMaxSize);
+		checkSlow(Num <= CurrentSize);
 		PendingFreeSpans.Add(FLinearAllocation{ BaseOffset, Num });
+		CurrentSize -= Num;
 	}
 
 	ENGINE_API void Reset();
 
 	ENGINE_API void Empty();
 
-	inline int32 GetSparselyAllocatedSize() const
+	FORCEINLINE int32 GetSparselyAllocatedSize() const
 	{
-		int32 AllocatedSize = CurrentMaxSize;
-
-		for (int32 i = 0; i < FreeSpans.Num(); i++)
-		{
-			AllocatedSize -= FreeSpans[i].Num;
-		}
-
-		return AllocatedSize;
+		return CurrentSize;
 	}
 
 	FORCEINLINE int32 GetMaxSize() const
@@ -126,7 +123,9 @@ private:
 			return StartOffset < Other.StartOffset;
 		}
 	};
-
+	
+	// Total of number of items currently allocated
+	int32 CurrentSize;
 	// Size of the linear range used by the allocator
 	int32 CurrentMaxSize;
 	// Peak tracked size since last Reset or Empty
