CPU Path Data Corruption Caused by [REDACTED_PERSON] in PrecomputedVolumetricLightmap

- Fixed Ambient Occlusion forward rendering bug in CompositionLighting.cpp

diff --git a/Engine/Source/Runtime/Engine/Private/PrecomputedVolumetricLightmap.cpp b/Engine/Source/Runtime/Engine/Private/PrecomputedVolumetricLightmap.cpp
--- a/Engine/Source/Runtime/Engine/Private/PrecomputedVolumetricLightmap.cpp
+++ b/Engine/Source/Runtime/Engine/Private/PrecomputedVolumetricLightmap.cpp
@@ -522,9 +522,12 @@ ENGINE_API void FPrecomputedVolumetricLightmapData::AddToSceneData(FPrecomputedV
 				uint8 Value = (uint8)IndexInCPUSubLevelBrickDataList;
 				SceneData->CPUSubLevelBrickDataList[IndexInCPUSubLevelBrickDataList] = this;
 
+				const int32 PaddedBrickSize = BrickSize + 1;
+				const FIntVector BrickLayoutDimensions = BrickDataDimensions / PaddedBrickSize;
+
 				for (int32 BrickIndex = 0; BrickIndex < SubLevelBrickPositions.Num(); BrickIndex++)
 				{
-					const FIntVector BrickLayoutPosition = ComputeBrickLayoutPosition(BrickIndex, BrickDataDimensions);
+					const FIntVector BrickLayoutPosition = ComputeBrickLayoutPosition(BrickIndex, BrickLayoutDimensions);
 
 					const FIntVector IndirectionDestDataCoordinate = SubLevelBrickPositions[BrickIndex];
 					const int32 IndirectionDestDataIndex =
@@ -620,7 +623,7 @@ ENGINE_API void FPrecomputedVolumetricLightmapData::RemoveFromSceneData(FPrecomp
 								*(IndirectionVoxelPtr + 0) = OriginalValue.R;
 								*(IndirectionVoxelPtr + 1) = OriginalValue.G;
 								*(IndirectionVoxelPtr + 2) = OriginalValue.B;
-								*(IndirectionVoxelPtr + 3) = 1;
+								*(IndirectionVoxelPtr + 3) = OriginalValue.A;
 							}
 
 							{
diff --git a/Engine/Source/Runtime/Renderer/Private/CompositionLighting/CompositionLighting.cpp b/Engine/Source/Runtime/Renderer/Private/CompositionLighting/CompositionLighting.cpp
--- a/Engine/Source/Runtime/Renderer/Private/CompositionLighting/CompositionLighting.cpp
+++ b/Engine/Source/Runtime/Renderer/Private/CompositionLighting/CompositionLighting.cpp
@@ -94,7 +94,7 @@ bool ShouldRenderScreenSpaceAmbientOcclusion(const FViewInfo& View, bool bLumenW
 		&& View.Family->EngineShowFlags.Lighting
 		&& View.FinalPostProcessSettings.AmbientOcclusionRadius >= 0.1f
 		&& !View.Family->UseDebugViewPS()
-		&& (FSSAOHelper::IsBasePassAmbientOcclusionRequired(View) || IsAmbientCubemapPassRequired(View) || IsReflectionEnvironmentActive(View) || IsSkylightActive(View) || View.Family->EngineShowFlags.VisualizeBuffer || bLumenWantsSSAO);
+		&& (FSSAOHelper::IsBasePassAmbientOcclusionRequired(View) || IsAmbientCubemapPassRequired(View) || IsReflectionEnvironmentActive(View) || IsSkylightActive(View) || IsForwardShadingEnabled(View.GetShaderPlatform()) || View.Family->EngineShowFlags.VisualizeBuffer || bLumenWantsSSAO);
 #if RHI_RAYTRACING
 	bEnabled &= !ShouldRenderRayTracingAmbientOcclusion(View);
 #endif
