Shader Compiler:  GlobalBeginCompileShader define optimizations, plus a few other micro-optimizations, producing a 3.2x performance improvement.
* Major define optimization involves converting map of defines to use FName keys and variant values rather than strings -- this eliminates most of the cost of string hashing, allocation, and conversion.
* Lower overhead FHashTable used instead of TMap.
* An initial map of defines can optionally be provided globally.  Anything using an initial define can have its map index cached for optimized lookup when reading or writing.

Other micro-optimizations:
* Added Reserve calls for uniform buffer related maps, to eliminate map resizing / rehashing.  Saved around 15% perf (after define optimizations).
* Added a map for UB lookup, instead of iterating through the linked list.  Saved around 10% perf.

Non-optimization:  Sort the order in which uniform buffer variable names are searched in BuildShaderFileToUniformBufferMap, to create determinism in ShaderDebug data for A/B testing (previously order was dependent on global constructor order for UB definitions, which could vary arbitrarily with unrelated changes).

#rnx

diff --git a/Engine/Source/Developer/ShaderCompilerCommon/Private/ShaderConductorContext.cpp b/Engine/Source/Developer/ShaderCompilerCommon/Private/ShaderConductorContext.cpp
--- a/Engine/Source/Developer/ShaderCompilerCommon/Private/ShaderConductorContext.cpp
+++ b/Engine/Source/Developer/ShaderCompilerCommon/Private/ShaderConductorContext.cpp
@@ -56,6 +56,11 @@ namespace CrossCompiler
 		ConvertByteArrayToAnsiString(TCHAR_TO_ANSI(*InString), InString.Len(), OutString);
 	}
 
+	static void ConvertFStringViewToAnsiString(const FStringView& InString, TArray<ANSICHAR>& OutString)
+	{
+		ConvertByteArrayToAnsiString(TCHAR_TO_ANSI(InString.GetData()), InString.Len(), OutString);
+	}
+
 	// Copies the NULL-terminated string 'InString' to 'OutString'. Also copies the '\0' character at the end.
 	static void CopyAnsiString(const ANSICHAR* InString, TArray<ANSICHAR>& OutString)
 	{
@@ -241,14 +246,14 @@ namespace CrossCompiler
 	}
 
 	// Converts a map of string pairs to a C-Style macro defines array
-	static void ConvertStringMapToMacroDefines(const TMap<FString,FString>& InPairs, TArray<TPair<TArray<ANSICHAR>, TArray<ANSICHAR>>>& OutPairs, TArray<ShaderConductor::MacroDefine>& OutPairRefs)
+	static void ConvertDefineMapToMacroDefines(const FShaderCompilerDefinitions& Definitions, TArray<TPair<TArray<ANSICHAR>, TArray<ANSICHAR>>>& OutPairs, TArray<ShaderConductor::MacroDefine>& OutPairRefs)
 	{
 		// Convert map into an array container
 		TArray<ANSICHAR> Name, Value;
-		for (const TPair<FString, FString>& Iter : InPairs)
+		for (FShaderCompilerDefinitions::FConstIterator Iter(Definitions); Iter; ++Iter)
 		{
-			ConvertFStringToAnsiString(Iter.Key, Name);
-			ConvertFStringToAnsiString(Iter.Value, Value);
+			ConvertFStringViewToAnsiString(FStringView(Iter.Key()), Name);
+			ConvertFStringViewToAnsiString(FStringView(Iter.Value()), Value);
 			OutPairs.Emplace(MoveTemp(Name), MoveTemp(Value));
 		}
 
@@ -285,7 +290,7 @@ namespace CrossCompiler
 		}
 
 		// Convert flags map into an array container
-		ConvertStringMapToMacroDefines(InTarget.CompileFlags.GetDefinitionMap(), Intermediates.Flags, Intermediates.FlagRefs);
+		ConvertDefineMapToMacroDefines(InTarget.CompileFlags, Intermediates.Flags, Intermediates.FlagRefs);
 
 		OutTargetDesc.options = Intermediates.FlagRefs.GetData();
 		OutTargetDesc.numOptions = static_cast<uint32>(Intermediates.FlagRefs.Num());
@@ -495,7 +500,7 @@ namespace CrossCompiler
 		// Convert macro definitions map into an array container
 		if (Definitions != nullptr)
 		{
-			ConvertStringMapToMacroDefines(Definitions->GetDefinitionMap(), Intermediates->Defines, Intermediates->DefineRefs);
+			ConvertDefineMapToMacroDefines(*Definitions, Intermediates->Defines, Intermediates->DefineRefs);
 		}
 
 		if (ExtraDxcArgs && ExtraDxcArgs->Num() > 0)
@@ -519,7 +524,7 @@ namespace CrossCompiler
 		// Convert macro definitions map into an array container
 		if (Definitions != nullptr)
 		{
-			ConvertStringMapToMacroDefines(Definitions->GetDefinitionMap(), Intermediates->Defines, Intermediates->DefineRefs);
+			ConvertDefineMapToMacroDefines(*Definitions, Intermediates->Defines, Intermediates->DefineRefs);
 		}
 
 		if (ExtraDxcArgs && ExtraDxcArgs->Num() > 0)
diff --git a/Engine/Source/Developer/ShaderPreprocessor/Private/ShaderPreprocessor.cpp b/Engine/Source/Developer/ShaderPreprocessor/Private/ShaderPreprocessor.cpp
--- a/Engine/Source/Developer/ShaderPreprocessor/Private/ShaderPreprocessor.cpp
+++ b/Engine/Source/Developer/ShaderPreprocessor/Private/ShaderPreprocessor.cpp
@@ -24,22 +24,25 @@ namespace
 }
 
 static void AddStbDefine(stb_arena* MacroArena, macro_definition**& StbDefines, const TCHAR* Name, const TCHAR* Value);
-static void AddStbDefines(stb_arena* MacroArena, macro_definition**& StbDefines, TMap<FString, FString> DefinitionsMap);
+static void AddStbDefines(stb_arena* MacroArena, macro_definition**& StbDefines, const FShaderCompilerDefinitions& Defines);
 
 class FShaderPreprocessorUtilities
 {
 public:
 	static void DumpShaderDefinesAsCommentedCode(const FShaderCompilerEnvironment& Environment, FString* OutDefines)
 	{
-		const TMap<FString, FString>& Definitions = Environment.Definitions.GetDefinitionMap();
-		TArray<FString> Keys;
-		Definitions.GetKeys(/* out */ Keys);
-		Keys.Sort();
+		TArray<FString> DefinesLines;
+		DefinesLines.Reserve(Environment.Definitions.Num());
+		for (FShaderCompilerDefinitions::FConstIterator DefineIt(Environment.Definitions); DefineIt; ++DefineIt)
+		{
+			DefinesLines.Add(FString::Printf(TEXT("// #define %s %s\n"), DefineIt.Key(), DefineIt.Value()));
+		}
+		DefinesLines.Sort();
 
 		FString Defines;
-		for (const FString& Key : Keys)
+		for (const FString& DefineLine : DefinesLines)
 		{
-			Defines += FString::Printf(TEXT("// #define %s %s\n"), *Key, *Definitions[Key]);
+			Defines += DefineLine;
 		}
 
 		*OutDefines += MakeInjectedShaderCodeBlock(TEXT("DumpShaderDefinesAsCommentedCode"), Defines);
@@ -47,8 +50,8 @@ public:
 
 	static void PopulateDefines(const FShaderCompilerEnvironment& Environment, const FShaderCompilerDefinitions& AdditionalDefines, stb_arena* MacroArena, macro_definition**& OutDefines)
 	{
-		AddStbDefines(MacroArena, OutDefines, Environment.Definitions.GetDefinitionMap());
-		AddStbDefines(MacroArena, OutDefines, AdditionalDefines.GetDefinitionMap());
+		AddStbDefines(MacroArena, OutDefines, Environment.Definitions);
+		AddStbDefines(MacroArena, OutDefines, AdditionalDefines);
 	}
 };
 
@@ -299,11 +302,11 @@ static void AddStbDefine(stb_arena* MacroArena, macro_definition**& StbDefines,
 	arrput(StbDefines, pp_define(MacroArena, (ANSICHAR*)ConvertedDefine.Get()));
 }
 
-static void AddStbDefines(stb_arena* MacroArena, macro_definition**& StbDefines, TMap<FString, FString> DefinitionsMap)
+static void AddStbDefines(stb_arena* MacroArena, macro_definition**& StbDefines, const FShaderCompilerDefinitions& Defines)
 {
-	for (TMap<FString, FString>::TConstIterator It(DefinitionsMap); It; ++It)
+	for (FShaderCompilerDefinitions::FConstIterator It(Defines); It; ++It)
 	{
-		AddStbDefine(MacroArena, StbDefines, *It.Key(), *It.Value());
+		AddStbDefine(MacroArena, StbDefines, It.Key(), It.Value());
 	}
 }
 
diff --git a/Engine/Source/Runtime/Core/Public/Containers/HashTable.h b/Engine/Source/Runtime/Core/Public/Containers/HashTable.h
--- a/Engine/Source/Runtime/Core/Public/Containers/HashTable.h
+++ b/Engine/Source/Runtime/Core/Public/Containers/HashTable.h
@@ -216,6 +216,8 @@ public:
 	// Average # of compares per search
 	CORE_API float	AverageSearch() const;
 
+	FHashTable&		operator=(const FHashTable& Other);
+
 protected:
 	// Avoids allocating hash until first add
 	CORE_API static uint32	EmptyHash[1];
@@ -266,6 +268,24 @@ FORCEINLINE FHashTable::FHashTable( const FHashTable& Other )
 	}
 }
 
+FORCEINLINE FHashTable& FHashTable::operator=(const FHashTable& Other)
+{
+	Free();
+	HashSize = Other.HashSize;
+	HashMask = Other.HashMask;
+	IndexSize = Other.IndexSize;
+
+	if (IndexSize)
+	{
+		Hash = new uint32[HashSize];
+		NextIndex = new uint32[IndexSize];
+
+		FMemory::Memcpy(Hash, Other.Hash, HashSize * 4);
+		FMemory::Memcpy(NextIndex, Other.NextIndex, IndexSize * 4);
+	}
+	return *this;
+}
+
 FORCEINLINE FHashTable::~FHashTable()
 {
 	Free();
diff --git a/Engine/Source/Runtime/Engine/Private/Materials/MaterialShared.cpp b/Engine/Source/Runtime/Engine/Private/Materials/MaterialShared.cpp
--- a/Engine/Source/Runtime/Engine/Private/Materials/MaterialShared.cpp
+++ b/Engine/Source/Runtime/Engine/Private/Materials/MaterialShared.cpp
@@ -2382,39 +2382,39 @@ void FMaterial::SetupMaterialEnvironment(
 			// For example IsTranslucentBlendMode doesn't check IsMasked
 			if(!WritesEveryPixel())
 			{
-				OutEnvironment.SetDefine(TEXT("MATERIALBLENDING_MASKED"),TEXT("1"));
+				SET_SHADER_DEFINE(OutEnvironment, MATERIALBLENDING_MASKED, 1);
 			}
 			else
 			{
-				OutEnvironment.SetDefine(TEXT("MATERIALBLENDING_SOLID"),TEXT("1"));
+				SET_SHADER_DEFINE(OutEnvironment, MATERIALBLENDING_SOLID, 1);
 			}
 			break;
 		}
 		case BLEND_AlphaComposite:
 		{
 			// Blend mode will reuse MATERIALBLENDING_TRANSLUCENT
-			OutEnvironment.SetDefineAndCompileArgument(TEXT("MATERIALBLENDING_ALPHACOMPOSITE"), true);
-			OutEnvironment.SetDefineAndCompileArgument(TEXT("MATERIALBLENDING_TRANSLUCENT"), true);
+			SET_SHADER_DEFINE_AND_COMPILE_ARGUMENT(OutEnvironment, MATERIALBLENDING_ALPHACOMPOSITE, true);
+			SET_SHADER_DEFINE_AND_COMPILE_ARGUMENT(OutEnvironment, MATERIALBLENDING_TRANSLUCENT, true);
 			break;
 		}
 		case BLEND_AlphaHoldout:
 		{
 			// Blend mode will reuse MATERIALBLENDING_TRANSLUCENT
-			OutEnvironment.SetDefine(TEXT("MATERIALBLENDING_ALPHAHOLDOUT"), TEXT("1"));
-			OutEnvironment.SetDefineAndCompileArgument(TEXT("MATERIALBLENDING_TRANSLUCENT"), true);
+			SET_SHADER_DEFINE(OutEnvironment, MATERIALBLENDING_ALPHAHOLDOUT, 1);
+			SET_SHADER_DEFINE_AND_COMPILE_ARGUMENT(OutEnvironment, MATERIALBLENDING_TRANSLUCENT, true);
 			break;
 		}
 		case BLEND_TranslucentColoredTransmittance:
 		case BLEND_Translucent:
-			OutEnvironment.SetDefineAndCompileArgument(TEXT("MATERIALBLENDING_TRANSLUCENT"), true); 
+			SET_SHADER_DEFINE_AND_COMPILE_ARGUMENT(OutEnvironment, MATERIALBLENDING_TRANSLUCENT, true);
 			break;
 
-		case BLEND_Additive: OutEnvironment.SetDefineAndCompileArgument(TEXT("MATERIALBLENDING_ADDITIVE"), true); break;
-		case BLEND_Modulate: OutEnvironment.SetDefineAndCompileArgument(TEXT("MATERIALBLENDING_MODULATE"), true); break;
+		case BLEND_Additive: SET_SHADER_DEFINE_AND_COMPILE_ARGUMENT(OutEnvironment, MATERIALBLENDING_ADDITIVE, true); break;
+		case BLEND_Modulate: SET_SHADER_DEFINE_AND_COMPILE_ARGUMENT(OutEnvironment, MATERIALBLENDING_MODULATE, true); break;
 
 		default: 
 			UE_LOG(LogMaterial, Warning, TEXT("Unknown material blend mode: %u  Setting to BLEND_Opaque"),(int32)GetBlendMode());
-			OutEnvironment.SetDefine(TEXT("MATERIALBLENDING_SOLID"),TEXT("1"));
+			SET_SHADER_DEFINE(OutEnvironment, MATERIALBLENDING_SOLID, 1);
 		}
 	}
 	else
@@ -2423,8 +2423,8 @@ void FMaterial::SetupMaterialEnvironment(
 		{
 		case BLEND_Opaque:
 		{
-			OutEnvironment.SetDefine(TEXT("MATERIALBLENDING_SOLID"), TEXT("1"));
-			OutEnvironment.SetDefine(TEXT("STRATA_BLENDING_OPAQUE"), TEXT("1"));
+			SET_SHADER_DEFINE(OutEnvironment, MATERIALBLENDING_SOLID, 1);
+			SET_SHADER_DEFINE(OutEnvironment, STRATA_BLENDING_OPAQUE, 1);
 			break;
 		}
 		case BLEND_Masked:
@@ -2434,56 +2434,56 @@ void FMaterial::SetupMaterialEnvironment(
 			// For example IsTranslucentBlendMode doesn't check IsMasked
 			if(!WritesEveryPixel())
 			{
-				OutEnvironment.SetDefine(TEXT("MATERIALBLENDING_MASKED"),TEXT("1"));
+				SET_SHADER_DEFINE(OutEnvironment, MATERIALBLENDING_MASKED, 1);
 			}
 			else
 			{
-				OutEnvironment.SetDefine(TEXT("MATERIALBLENDING_SOLID"),TEXT("1"));
+				SET_SHADER_DEFINE(OutEnvironment, MATERIALBLENDING_SOLID, 1);
 			}
-			OutEnvironment.SetDefine(TEXT("STRATA_BLENDING_MASKED"), TEXT("1"));
+			SET_SHADER_DEFINE(OutEnvironment, STRATA_BLENDING_MASKED, 1);
 			break;
 		}
 		case BLEND_Additive:
 		{
-			OutEnvironment.SetDefine(TEXT("MATERIALBLENDING_ADDITIVE"),TEXT("1"));
-			OutEnvironment.SetDefine(TEXT("STRATA_BLENDING_TRANSLUCENT_GREYTRANSMITTANCE"), TEXT("1"));
+			SET_SHADER_DEFINE(OutEnvironment, MATERIALBLENDING_ADDITIVE, 1);
+			SET_SHADER_DEFINE(OutEnvironment, STRATA_BLENDING_TRANSLUCENT_GREYTRANSMITTANCE, 1);
 			break;
 		}
 		case BLEND_AlphaComposite:
 		{
-			OutEnvironment.SetDefineAndCompileArgument(TEXT("MATERIALBLENDING_ALPHACOMPOSITE"), true);
-			OutEnvironment.SetDefineAndCompileArgument(TEXT("MATERIALBLENDING_TRANSLUCENT"), true);
-			OutEnvironment.SetDefine(TEXT("STRATA_BLENDING_TRANSLUCENT_GREYTRANSMITTANCE"), TEXT("1"));
+			SET_SHADER_DEFINE_AND_COMPILE_ARGUMENT(OutEnvironment, MATERIALBLENDING_ALPHACOMPOSITE, true);
+			SET_SHADER_DEFINE_AND_COMPILE_ARGUMENT(OutEnvironment, MATERIALBLENDING_TRANSLUCENT, true);
+			SET_SHADER_DEFINE(OutEnvironment, STRATA_BLENDING_TRANSLUCENT_GREYTRANSMITTANCE, 1);
 			break;
 		}
 		case BLEND_TranslucentGreyTransmittance:
 		{
-			OutEnvironment.SetDefineAndCompileArgument(TEXT("MATERIALBLENDING_TRANSLUCENT"), true);
-			OutEnvironment.SetDefine(TEXT("STRATA_BLENDING_TRANSLUCENT_GREYTRANSMITTANCE"), TEXT("1"));
+			SET_SHADER_DEFINE_AND_COMPILE_ARGUMENT(OutEnvironment, MATERIALBLENDING_TRANSLUCENT, true);
+			SET_SHADER_DEFINE(OutEnvironment, STRATA_BLENDING_TRANSLUCENT_GREYTRANSMITTANCE, 1);
 			break;
 		}
 		case BLEND_TranslucentColoredTransmittance:
 		{
-			OutEnvironment.SetDefineAndCompileArgument(TEXT("MATERIALBLENDING_TRANSLUCENT"), true);
-			OutEnvironment.SetDefine(TEXT("STRATA_BLENDING_TRANSLUCENT_COLOREDTRANSMITTANCE"), TEXT("1"));
+			SET_SHADER_DEFINE_AND_COMPILE_ARGUMENT(OutEnvironment, MATERIALBLENDING_TRANSLUCENT, true);
+			SET_SHADER_DEFINE(OutEnvironment, STRATA_BLENDING_TRANSLUCENT_COLOREDTRANSMITTANCE, 1);
 			break;
 		}
 		case BLEND_ColoredTransmittanceOnly:
 		{
-			OutEnvironment.SetDefineAndCompileArgument(TEXT("MATERIALBLENDING_MODULATE"), true);
-			OutEnvironment.SetDefine(TEXT("STRATA_BLENDING_COLOREDTRANSMITTANCEONLY"), TEXT("1"));
+			SET_SHADER_DEFINE_AND_COMPILE_ARGUMENT(OutEnvironment, MATERIALBLENDING_MODULATE, true);
+			SET_SHADER_DEFINE(OutEnvironment, STRATA_BLENDING_COLOREDTRANSMITTANCEONLY, 1);
 			break;
 		}
 		case BLEND_AlphaHoldout:
 		{
-			OutEnvironment.SetDefineAndCompileArgument(TEXT("MATERIALBLENDING_TRANSLUCENT"), true);
-			OutEnvironment.SetDefine(TEXT("MATERIALBLENDING_ALPHAHOLDOUT"), TEXT("1"));
-			OutEnvironment.SetDefine(TEXT("STRATA_BLENDING_ALPHAHOLDOUT"), TEXT("1"));
+			SET_SHADER_DEFINE_AND_COMPILE_ARGUMENT(OutEnvironment, MATERIALBLENDING_TRANSLUCENT, true);
+			SET_SHADER_DEFINE(OutEnvironment, MATERIALBLENDING_ALPHAHOLDOUT, 1);
+			SET_SHADER_DEFINE(OutEnvironment, STRATA_BLENDING_ALPHAHOLDOUT, 1);
 			break;
 		}
 		default:
 			UE_LOG(LogMaterial, Error, TEXT("%s: Unkown Strata material blend mode could not be converted to Starta. (Asset: %s) Setting to BLEND_Opaque"), *GetFriendlyName(), *GetAssetName());
-			OutEnvironment.SetDefine(TEXT("MATERIALBLENDING_SOLID"), TEXT("1"));
+			SET_SHADER_DEFINE(OutEnvironment, MATERIALBLENDING_SOLID, 1);
 		}
 	}
 
@@ -2507,44 +2507,44 @@ void FMaterial::SetupMaterialEnvironment(
 				check(0);
 		}
 
-		OutEnvironment.SetDefine(TEXT("MATERIALDECALRESPONSEMASK"), MaterialDecalResponseMask);
+		SET_SHADER_DEFINE(OutEnvironment, MATERIALDECALRESPONSEMASK, MaterialDecalResponseMask);
 	}
 
 	switch(GetRefractionMode())
 	{
-	case RM_IndexOfRefraction:	OutEnvironment.SetDefine(TEXT("REFRACTION_USE_INDEX_OF_REFRACTION"),	TEXT("1")); break;
-	case RM_PixelNormalOffset:	OutEnvironment.SetDefine(TEXT("REFRACTION_USE_PIXEL_NORMAL_OFFSET"),	TEXT("1")); break;
-	case RM_2DOffset:			OutEnvironment.SetDefine(TEXT("REFRACTION_USE_2D_OFFSET"),				TEXT("1")); break;
-	case RM_None:				OutEnvironment.SetDefine(TEXT("REFRACTION_USE_NONE"),					TEXT("1")); break;
+	case RM_IndexOfRefraction:	SET_SHADER_DEFINE(OutEnvironment, REFRACTION_USE_INDEX_OF_REFRACTION,	1); break;
+	case RM_PixelNormalOffset:	SET_SHADER_DEFINE(OutEnvironment, REFRACTION_USE_PIXEL_NORMAL_OFFSET,	1); break;
+	case RM_2DOffset:			SET_SHADER_DEFINE(OutEnvironment, REFRACTION_USE_2D_OFFSET,				1); break;
+	case RM_None:				SET_SHADER_DEFINE(OutEnvironment, REFRACTION_USE_NONE,					1); break;
 	default: 
 		UE_LOG(LogMaterial, Warning, TEXT("Unknown material refraction mode: %u  Setting to RM_IndexOfRefraction"),(int32)GetRefractionMode());
-		OutEnvironment.SetDefine(TEXT("REFRACTION_USE_INDEX_OF_REFRACTION"),TEXT("1"));
-	}
-	OutEnvironment.SetDefine(TEXT("REFRACTION_ROOT_NODE_OVERRIDES_DEFAULT"), GetRootNodeOverridesDefaultRefraction() ? 1 : 0);
-
-	OutEnvironment.SetDefine(TEXT("USE_DITHERED_LOD_TRANSITION_FROM_MATERIAL"), IsDitheredLODTransition());
-	OutEnvironment.SetDefine(TEXT("MATERIAL_TWOSIDED"), IsTwoSided());
-	OutEnvironment.SetDefine(TEXT("MATERIAL_ISTHINSURFACE"), IsThinSurface());
-	OutEnvironment.SetDefine(TEXT("MATERIAL_TANGENTSPACENORMAL"), IsTangentSpaceNormal());
-	OutEnvironment.SetDefine(TEXT("GENERATE_SPHERICAL_PARTICLE_NORMALS"),ShouldGenerateSphericalParticleNormals());
-	OutEnvironment.SetDefine(TEXT("MATERIAL_USES_SCENE_COLOR_COPY"), RequiresSceneColorCopy_GameThread());
-	OutEnvironment.SetDefine(TEXT("MATERIAL_USE_PREINTEGRATED_GF"), IsUsingPreintegratedGFForSimpleIBL());
-	OutEnvironment.SetDefine(TEXT("MATERIAL_HQ_FORWARD_REFLECTION_CAPTURES"), IsUsingHQForwardReflections());
-	OutEnvironment.SetDefine(TEXT("MATERIAL_FORWARD_BLENDS_SKYLIGHT_CUBEMAPS"), GetForwardBlendsSkyLightCubemaps());
-	OutEnvironment.SetDefine(TEXT("MATERIAL_PLANAR_FORWARD_REFLECTIONS"), IsUsingPlanarForwardReflections());
-	OutEnvironment.SetDefine(TEXT("MATERIAL_NONMETAL"), IsNonmetal());
-	OutEnvironment.SetDefine(TEXT("MATERIAL_USE_LM_DIRECTIONALITY"), UseLmDirectionality());
-	OutEnvironment.SetDefine(TEXT("MATERIAL_SSR"), ShouldDoSSR() && IsTranslucentBlendMode(GetBlendMode()));
-	OutEnvironment.SetDefine(TEXT("MATERIAL_CONTACT_SHADOWS"), ShouldDoContactShadows() && IsTranslucentBlendMode(GetBlendMode()));
-	OutEnvironment.SetDefine(TEXT("MATERIAL_DITHER_OPACITY_MASK"), IsDitherMasked());
-	OutEnvironment.SetDefine(TEXT("MATERIAL_NORMAL_CURVATURE_TO_ROUGHNESS"), UseNormalCurvatureToRoughness() ? TEXT("1") : TEXT("0"));
-	OutEnvironment.SetDefine(TEXT("MATERIAL_ALLOW_NEGATIVE_EMISSIVECOLOR"), AllowNegativeEmissiveColor());
-	OutEnvironment.SetDefine(TEXT("MATERIAL_OUTPUT_OPACITY_AS_ALPHA"), GetBlendableOutputAlpha());
-	OutEnvironment.SetDefine(TEXT("TRANSLUCENT_SHADOW_WITH_MASKED_OPACITY"), GetCastDynamicShadowAsMasked());
-	OutEnvironment.SetDefine(TEXT("TRANSLUCENT_WRITING_VELOCITY"), IsTranslucencyWritingVelocity());
-	OutEnvironment.SetDefine(TEXT("TRANSLUCENT_WRITING_FRONT_LAYER_TRANSPARENCY"), IsTranslucencyWritingFrontLayerTransparency());
-	OutEnvironment.SetDefine(TEXT("MATERIAL_USE_ALPHA_TO_COVERAGE"), IsUsingAlphaToCoverage());
-	OutEnvironment.SetDefine(TEXT("MOBILE_HIGH_QUALITY_BRDF"), IsMobileHighQualityBRDFEnabled());
+		SET_SHADER_DEFINE(OutEnvironment, REFRACTION_USE_INDEX_OF_REFRACTION, 1);
+	}
+	SET_SHADER_DEFINE(OutEnvironment, REFRACTION_ROOT_NODE_OVERRIDES_DEFAULT, GetRootNodeOverridesDefaultRefraction() ? 1 : 0);
+
+	SET_SHADER_DEFINE(OutEnvironment, USE_DITHERED_LOD_TRANSITION_FROM_MATERIAL, IsDitheredLODTransition());
+	SET_SHADER_DEFINE(OutEnvironment, MATERIAL_TWOSIDED, IsTwoSided());
+	SET_SHADER_DEFINE(OutEnvironment, MATERIAL_ISTHINSURFACE, IsThinSurface());
+	SET_SHADER_DEFINE(OutEnvironment, MATERIAL_TANGENTSPACENORMAL, IsTangentSpaceNormal());
+	SET_SHADER_DEFINE(OutEnvironment, GENERATE_SPHERICAL_PARTICLE_NORMALS,ShouldGenerateSphericalParticleNormals());
+	SET_SHADER_DEFINE(OutEnvironment, MATERIAL_USES_SCENE_COLOR_COPY, RequiresSceneColorCopy_GameThread());
+	SET_SHADER_DEFINE(OutEnvironment, MATERIAL_USE_PREINTEGRATED_GF, IsUsingPreintegratedGFForSimpleIBL());
+	SET_SHADER_DEFINE(OutEnvironment, MATERIAL_HQ_FORWARD_REFLECTION_CAPTURES, IsUsingHQForwardReflections());
+	SET_SHADER_DEFINE(OutEnvironment, MATERIAL_FORWARD_BLENDS_SKYLIGHT_CUBEMAPS, GetForwardBlendsSkyLightCubemaps());
+	SET_SHADER_DEFINE(OutEnvironment, MATERIAL_PLANAR_FORWARD_REFLECTIONS, IsUsingPlanarForwardReflections());
+	SET_SHADER_DEFINE(OutEnvironment, MATERIAL_NONMETAL, IsNonmetal());
+	SET_SHADER_DEFINE(OutEnvironment, MATERIAL_USE_LM_DIRECTIONALITY, UseLmDirectionality());
+	SET_SHADER_DEFINE(OutEnvironment, MATERIAL_SSR, ShouldDoSSR() && IsTranslucentBlendMode(GetBlendMode()));
+	SET_SHADER_DEFINE(OutEnvironment, MATERIAL_CONTACT_SHADOWS, ShouldDoContactShadows() && IsTranslucentBlendMode(GetBlendMode()));
+	SET_SHADER_DEFINE(OutEnvironment, MATERIAL_DITHER_OPACITY_MASK, IsDitherMasked());
+	SET_SHADER_DEFINE(OutEnvironment, MATERIAL_NORMAL_CURVATURE_TO_ROUGHNESS, UseNormalCurvatureToRoughness() ? 1 : 0);
+	SET_SHADER_DEFINE(OutEnvironment, MATERIAL_ALLOW_NEGATIVE_EMISSIVECOLOR, AllowNegativeEmissiveColor());
+	SET_SHADER_DEFINE(OutEnvironment, MATERIAL_OUTPUT_OPACITY_AS_ALPHA, GetBlendableOutputAlpha());
+	SET_SHADER_DEFINE(OutEnvironment, TRANSLUCENT_SHADOW_WITH_MASKED_OPACITY, GetCastDynamicShadowAsMasked());
+	SET_SHADER_DEFINE(OutEnvironment, TRANSLUCENT_WRITING_VELOCITY, IsTranslucencyWritingVelocity());
+	SET_SHADER_DEFINE(OutEnvironment, TRANSLUCENT_WRITING_FRONT_LAYER_TRANSPARENCY, IsTranslucencyWritingFrontLayerTransparency());
+	SET_SHADER_DEFINE(OutEnvironment, MATERIAL_USE_ALPHA_TO_COVERAGE, IsUsingAlphaToCoverage());
+	SET_SHADER_DEFINE(OutEnvironment, MOBILE_HIGH_QUALITY_BRDF, IsMobileHighQualityBRDFEnabled());
 
 	bool bFullPrecisionInMaterial = false;
 	bool bFullPrecisionInPS = false;
@@ -2553,50 +2553,50 @@ void FMaterial::SetupMaterialEnvironment(
 
 	if (bFullPrecisionInMaterial)
 	{
-		OutEnvironment.SetDefine(TEXT("FORCE_MATERIAL_FLOAT_FULL_PRECISION"), TEXT("1"));
+		SET_SHADER_DEFINE(OutEnvironment, FORCE_MATERIAL_FLOAT_FULL_PRECISION, 1);
 	}
 
 	OutEnvironment.FullPrecisionInPS |= bFullPrecisionInPS;
 
 	switch(GetMaterialDomain())
 	{
-		case MD_Surface:				OutEnvironment.SetDefine(TEXT("MATERIAL_DOMAIN_SURFACE"),			TEXT("1")); break;
-		case MD_DeferredDecal:			OutEnvironment.SetDefine(TEXT("MATERIAL_DOMAIN_DEFERREDDECAL"),		TEXT("1")); break;
-		case MD_LightFunction:			OutEnvironment.SetDefine(TEXT("MATERIAL_DOMAIN_LIGHTFUNCTION"),		TEXT("1")); break;
-		case MD_Volume:					OutEnvironment.SetDefine(TEXT("MATERIAL_DOMAIN_VOLUME"),			TEXT("1")); break;
-		case MD_PostProcess:			OutEnvironment.SetDefine(TEXT("MATERIAL_DOMAIN_POSTPROCESS"),		TEXT("1")); break;
-		case MD_UI:						OutEnvironment.SetDefine(TEXT("MATERIAL_DOMAIN_UI"),				TEXT("1")); break;
+		case MD_Surface:				SET_SHADER_DEFINE(OutEnvironment, MATERIAL_DOMAIN_SURFACE,			1); break;
+		case MD_DeferredDecal:			SET_SHADER_DEFINE(OutEnvironment, MATERIAL_DOMAIN_DEFERREDDECAL,	1); break;
+		case MD_LightFunction:			SET_SHADER_DEFINE(OutEnvironment, MATERIAL_DOMAIN_LIGHTFUNCTION,	1); break;
+		case MD_Volume:					SET_SHADER_DEFINE(OutEnvironment, MATERIAL_DOMAIN_VOLUME,			1); break;
+		case MD_PostProcess:			SET_SHADER_DEFINE(OutEnvironment, MATERIAL_DOMAIN_POSTPROCESS,		1); break;
+		case MD_UI:						SET_SHADER_DEFINE(OutEnvironment, MATERIAL_DOMAIN_UI,				1); break;
 		default:
 			UE_LOG(LogMaterial, Warning, TEXT("Unknown material domain: %u  Setting to MD_Surface"),(int32)GetMaterialDomain());
-			OutEnvironment.SetDefine(TEXT("MATERIAL_DOMAIN_SURFACE"),TEXT("1"));
+			SET_SHADER_DEFINE(OutEnvironment, MATERIAL_DOMAIN_SURFACE, 1);
 	};
 
 	if (IsTranslucentBlendMode(GetBlendMode()))
 	{
 		switch(GetTranslucencyLightingMode())
 		{
-		case TLM_VolumetricNonDirectional: OutEnvironment.SetDefine(TEXT("TRANSLUCENCY_LIGHTING_VOLUMETRIC_NONDIRECTIONAL"),TEXT("1")); break;
-		case TLM_VolumetricDirectional: OutEnvironment.SetDefine(TEXT("TRANSLUCENCY_LIGHTING_VOLUMETRIC_DIRECTIONAL"),TEXT("1")); break;
-		case TLM_VolumetricPerVertexNonDirectional: OutEnvironment.SetDefine(TEXT("TRANSLUCENCY_LIGHTING_VOLUMETRIC_PERVERTEX_NONDIRECTIONAL"),TEXT("1")); break;
-		case TLM_VolumetricPerVertexDirectional: OutEnvironment.SetDefine(TEXT("TRANSLUCENCY_LIGHTING_VOLUMETRIC_PERVERTEX_DIRECTIONAL"),TEXT("1")); break;
-		case TLM_Surface: OutEnvironment.SetDefine(TEXT("TRANSLUCENCY_LIGHTING_SURFACE_LIGHTINGVOLUME"),TEXT("1")); break;
-		case TLM_SurfacePerPixelLighting: OutEnvironment.SetDefine(TEXT("TRANSLUCENCY_LIGHTING_SURFACE_FORWARDSHADING"),TEXT("1")); break;
+		case TLM_VolumetricNonDirectional: SET_SHADER_DEFINE(OutEnvironment, TRANSLUCENCY_LIGHTING_VOLUMETRIC_NONDIRECTIONAL, 1); break;
+		case TLM_VolumetricDirectional: SET_SHADER_DEFINE(OutEnvironment, TRANSLUCENCY_LIGHTING_VOLUMETRIC_DIRECTIONAL, 1); break;
+		case TLM_VolumetricPerVertexNonDirectional: SET_SHADER_DEFINE(OutEnvironment, TRANSLUCENCY_LIGHTING_VOLUMETRIC_PERVERTEX_NONDIRECTIONAL, 1); break;
+		case TLM_VolumetricPerVertexDirectional: SET_SHADER_DEFINE(OutEnvironment, TRANSLUCENCY_LIGHTING_VOLUMETRIC_PERVERTEX_DIRECTIONAL, 1); break;
+		case TLM_Surface: SET_SHADER_DEFINE(OutEnvironment, TRANSLUCENCY_LIGHTING_SURFACE_LIGHTINGVOLUME, 1); break;
+		case TLM_SurfacePerPixelLighting: SET_SHADER_DEFINE(OutEnvironment, TRANSLUCENCY_LIGHTING_SURFACE_FORWARDSHADING, 1); break;
 
 		default: 
 			UE_LOG(LogMaterial, Warning, TEXT("Unknown lighting mode: %u"),(int32)GetTranslucencyLightingMode());
-			OutEnvironment.SetDefine(TEXT("TRANSLUCENCY_LIGHTING_VOLUMETRIC_NONDIRECTIONAL"),TEXT("1")); break;
+			SET_SHADER_DEFINE(OutEnvironment, TRANSLUCENCY_LIGHTING_VOLUMETRIC_NONDIRECTIONAL, 1); break;
 		};
 	}
 
 	if( IsUsedWithEditorCompositing() )
 	{
-		OutEnvironment.SetDefine(TEXT("EDITOR_PRIMITIVE_MATERIAL"),TEXT("1"));
+		SET_SHADER_DEFINE(OutEnvironment, EDITOR_PRIMITIVE_MATERIAL, 1);
 	}
 
 	if (IsFeatureLevelSupported(Platform, ERHIFeatureLevel::SM5))
 	{	
 		static const auto CVar = IConsoleManager::Get().FindTConsoleVariableDataInt(TEXT("r.StencilForLODDither"));
-		OutEnvironment.SetDefine(TEXT("USE_STENCIL_LOD_DITHER_DEFAULT"), CVar->GetValueOnAnyThread() != 0 ? 1 : 0);
+		SET_SHADER_DEFINE(OutEnvironment, USE_STENCIL_LOD_DITHER_DEFAULT, CVar->GetValueOnAnyThread() != 0 ? 1 : 0);
 	}
 
 	if (GetShadingRate() != MSR_1x1)
@@ -2607,11 +2607,11 @@ void FMaterial::SetupMaterialEnvironment(
 	{
 		switch (GetMaterialDomain())
 		{
-			case MD_Surface:		OutEnvironment.SetDefine(TEXT("MATERIALDOMAIN_SURFACE"), 1u); break;
-			case MD_DeferredDecal:	OutEnvironment.SetDefine(TEXT("MATERIALDOMAIN_DEFERREDDECAL"), 1u); break;
-			case MD_LightFunction:	OutEnvironment.SetDefine(TEXT("MATERIALDOMAIN_LIGHTFUNCTION"), 1u); break;
-			case MD_PostProcess:	OutEnvironment.SetDefine(TEXT("MATERIALDOMAIN_POSTPROCESS"), 1u); break;
-			case MD_UI:				OutEnvironment.SetDefine(TEXT("MATERIALDOMAIN_UI"), 1u); break;
+			case MD_Surface:		SET_SHADER_DEFINE(OutEnvironment, MATERIALDOMAIN_SURFACE, 1); break;
+			case MD_DeferredDecal:	SET_SHADER_DEFINE(OutEnvironment, MATERIALDOMAIN_DEFERREDDECAL, 1); break;
+			case MD_LightFunction:	SET_SHADER_DEFINE(OutEnvironment, MATERIALDOMAIN_LIGHTFUNCTION, 1); break;
+			case MD_PostProcess:	SET_SHADER_DEFINE(OutEnvironment, MATERIALDOMAIN_POSTPROCESS, 1); break;
+			case MD_UI:				SET_SHADER_DEFINE(OutEnvironment, MATERIALDOMAIN_UI, 1); break;
 		}
 	}
 }
diff --git a/Engine/Source/Runtime/Engine/Private/ShaderCompiler/ShaderCompiler.cpp b/Engine/Source/Runtime/Engine/Private/ShaderCompiler/ShaderCompiler.cpp
--- a/Engine/Source/Runtime/Engine/Private/ShaderCompiler/ShaderCompiler.cpp
+++ b/Engine/Source/Runtime/Engine/Private/ShaderCompiler/ShaderCompiler.cpp
@@ -6166,16 +6166,16 @@ void GlobalBeginCompileShader(
 
 	// Add the appropriate definitions for the shader frequency.
 	{
-		Input.Environment.SetDefine(TEXT("PIXELSHADER"), Target.Frequency == SF_Pixel);
-		Input.Environment.SetDefine(TEXT("VERTEXSHADER"), Target.Frequency == SF_Vertex);
-		Input.Environment.SetDefine(TEXT("MESHSHADER"), Target.Frequency == SF_Mesh);
-		Input.Environment.SetDefine(TEXT("AMPLIFICATIONSHADER"), Target.Frequency == SF_Amplification);
-		Input.Environment.SetDefine(TEXT("GEOMETRYSHADER"), Target.Frequency == SF_Geometry);
-		Input.Environment.SetDefine(TEXT("COMPUTESHADER"), Target.Frequency == SF_Compute);
-		Input.Environment.SetDefine(TEXT("RAYCALLABLESHADER"), Target.Frequency == SF_RayCallable);
-		Input.Environment.SetDefine(TEXT("RAYHITGROUPSHADER"), Target.Frequency == SF_RayHitGroup);
-		Input.Environment.SetDefine(TEXT("RAYGENSHADER"), Target.Frequency == SF_RayGen);
-		Input.Environment.SetDefine(TEXT("RAYMISSSHADER"), Target.Frequency == SF_RayMiss);
+		SET_SHADER_DEFINE(Input.Environment, PIXELSHADER,			Target.Frequency == SF_Pixel);
+		SET_SHADER_DEFINE(Input.Environment, VERTEXSHADER,			Target.Frequency == SF_Vertex);
+		SET_SHADER_DEFINE(Input.Environment, MESHSHADER,			Target.Frequency == SF_Mesh);
+		SET_SHADER_DEFINE(Input.Environment, AMPLIFICATIONSHADER,	Target.Frequency == SF_Amplification);
+		SET_SHADER_DEFINE(Input.Environment, GEOMETRYSHADER,		Target.Frequency == SF_Geometry);
+		SET_SHADER_DEFINE(Input.Environment, COMPUTESHADER,			Target.Frequency == SF_Compute);
+		SET_SHADER_DEFINE(Input.Environment, RAYCALLABLESHADER,		Target.Frequency == SF_RayCallable);
+		SET_SHADER_DEFINE(Input.Environment, RAYHITGROUPSHADER,		Target.Frequency == SF_RayHitGroup);
+		SET_SHADER_DEFINE(Input.Environment, RAYGENSHADER,			Target.Frequency == SF_RayGen);
+		SET_SHADER_DEFINE(Input.Environment, RAYMISSSHADER,			Target.Frequency == SF_RayMiss);
 	}
 
 	// Enables HLSL 2021
@@ -6197,29 +6197,29 @@ void GlobalBeginCompileShader(
 	}
 
 	// #defines get stripped out by the preprocessor without this. We can override with this
-	Input.Environment.SetDefine(TEXT("COMPILER_DEFINE"), TEXT("#define"));
+	SET_SHADER_DEFINE(Input.Environment, COMPILER_DEFINE, TEXT("#define"));
 
 	if (FSceneInterface::GetShadingPath(GetMaxSupportedFeatureLevel(ShaderPlatform)) == EShadingPath::Deferred)
 	{
-		Input.Environment.SetDefine(TEXT("SHADING_PATH_DEFERRED"), 1);
+		SET_SHADER_DEFINE(Input.Environment, SHADING_PATH_DEFERRED, 1);
 	}
 
 	const bool bUsingMobileRenderer = FSceneInterface::GetShadingPath(GetMaxSupportedFeatureLevel(ShaderPlatform)) == EShadingPath::Mobile;
 	if (bUsingMobileRenderer)
 	{
-		Input.Environment.SetDefineAndCompileArgument(TEXT("SHADING_PATH_MOBILE"), true);
+		SET_SHADER_DEFINE_AND_COMPILE_ARGUMENT(Input.Environment, SHADING_PATH_MOBILE, true);
 		
 		const bool bMobileDeferredShading = IsMobileDeferredShadingEnabled((EShaderPlatform)Target.Platform);
-		Input.Environment.SetDefineAndCompileArgument(TEXT("MOBILE_DEFERRED_SHADING"), bMobileDeferredShading);
+		SET_SHADER_DEFINE_AND_COMPILE_ARGUMENT(Input.Environment, MOBILE_DEFERRED_SHADING, bMobileDeferredShading);
 
 		if (bMobileDeferredShading)
 		{
 			bool bGLESDeferredShading = Target.Platform == SP_OPENGL_ES3_1_ANDROID;
-			Input.Environment.SetDefine(TEXT("USE_GLES_FBF_DEFERRED"), bGLESDeferredShading ? 1 : 0);
-			Input.Environment.SetDefine(TEXT("MOBILE_EXTENDED_GBUFFER"), MobileUsesExtenedGBuffer((EShaderPlatform)Target.Platform) ? 1 : 0);
+			SET_SHADER_DEFINE(Input.Environment, USE_GLES_FBF_DEFERRED, bGLESDeferredShading ? 1 : 0);
+			SET_SHADER_DEFINE(Input.Environment, MOBILE_EXTENDED_GBUFFER, MobileUsesExtenedGBuffer((EShaderPlatform)Target.Platform) ? 1 : 0);
 		}
 
-		Input.Environment.SetDefine(TEXT("USE_SCENE_DEPTH_AUX"), MobileRequiresSceneDepthAux(ShaderPlatform) ? 1 : 0);
+		SET_SHADER_DEFINE(Input.Environment, USE_SCENE_DEPTH_AUX, MobileRequiresSceneDepthAux(ShaderPlatform) ? 1 : 0);
 
 		static FShaderPlatformCachedIniValue<bool> EnableCullBeforeFetchIniValue(TEXT("r.CullBeforeFetch"));
 		if (EnableCullBeforeFetchIniValue.Get((EShaderPlatform)Target.Platform) == 1)
@@ -6249,9 +6249,9 @@ void GlobalBeginCompileShader(
 	{
 		const UE::StereoRenderUtils::FStereoShaderAspects Aspects(ShaderPlatform);
 
-		Input.Environment.SetDefineAndCompileArgument(TEXT("INSTANCED_STEREO"), Aspects.IsInstancedStereoEnabled());
-		Input.Environment.SetDefineAndCompileArgument(TEXT("MULTI_VIEW"), Aspects.IsInstancedMultiViewportEnabled());
-		Input.Environment.SetDefineAndCompileArgument(TEXT("MOBILE_MULTI_VIEW"), Aspects.IsMobileMultiViewEnabled());
+		SET_SHADER_DEFINE_AND_COMPILE_ARGUMENT(Input.Environment, INSTANCED_STEREO, Aspects.IsInstancedStereoEnabled());
+		SET_SHADER_DEFINE_AND_COMPILE_ARGUMENT(Input.Environment, MULTI_VIEW, Aspects.IsInstancedMultiViewportEnabled());
+		SET_SHADER_DEFINE_AND_COMPILE_ARGUMENT(Input.Environment, MOBILE_MULTI_VIEW, Aspects.IsMobileMultiViewEnabled());
 
 		// Throw a warning if we are silently disabling ISR due to missing platform support (but don't have MMV enabled).
 		static const auto CVarInstancedStereo = IConsoleManager::Get().FindTConsoleVariableDataInt(TEXT("vr.InstancedStereo"));
@@ -6278,11 +6278,16 @@ void GlobalBeginCompileShader(
 	}
 	else
 	{
-		Input.Environment.SetDefineAndCompileArgument(TEXT("INSTANCED_STEREO"), false);
-		Input.Environment.SetDefineAndCompileArgument(TEXT("MULTI_VIEW"), 0);
-		Input.Environment.SetDefineAndCompileArgument(TEXT("MOBILE_MULTI_VIEW"), false);
+		SET_SHADER_DEFINE_AND_COMPILE_ARGUMENT(Input.Environment, INSTANCED_STEREO, false);
+		SET_SHADER_DEFINE_AND_COMPILE_ARGUMENT(Input.Environment, MULTI_VIEW, 0);
+		SET_SHADER_DEFINE_AND_COMPILE_ARGUMENT(Input.Environment, MOBILE_MULTI_VIEW, false);
 	}
 
+	// Reserve space in maps to prevent reallocation and rehashing in AddUniformBufferIncludesToEnvironment
+	const int32 UniformBufferReserveNum = Input.Environment.UniformBufferMap.Num() + ShaderType->GetReferencedUniformBufferNames().Num() + (VFType ? VFType->GetReferencedUniformBufferNames().Num() : 0);
+	Input.Environment.UniformBufferMap.Reserve(UniformBufferReserveNum);
+	Input.Environment.IncludeVirtualPathToExternalContentsMap.Reserve(UniformBufferReserveNum);
+
 	ShaderType->AddUniformBufferIncludesToEnvironment(Input.Environment, ShaderPlatform);
 
 	if (VFType)
@@ -6371,8 +6376,8 @@ void GlobalBeginCompileShader(
 	}
 
 	{
-		Input.Environment.SetDefine(TEXT("DO_CHECK"), GSShaderCheckLevel > 0 ? 1 : 0);
-		Input.Environment.SetDefine(TEXT("DO_GUARD_SLOW"), GSShaderCheckLevel > 1 ? 1 : 0);
+		SET_SHADER_DEFINE(Input.Environment, DO_CHECK, GSShaderCheckLevel > 0 ? 1 : 0);
+		SET_SHADER_DEFINE(Input.Environment, DO_GUARD_SLOW, GSShaderCheckLevel > 1 ? 1 : 0);
 	}
 
 	if (CVarShaderWarningsAsErrors.GetValueOnAnyThread())
@@ -6486,13 +6491,13 @@ void GlobalBeginCompileShader(
 	// Mobile emulation should be defined when a PC platform is using a mobile renderer (limited to feature level ES3_1)...  eg SP_PCD3D_ES3_1,SP_VULKAN_PCES3_1,SP_METAL_MACES3_1
 	if (IsSimulatedPlatform(EShaderPlatform(Target.Platform)) && bAllowDevelopmentShaderCompile)
 	{
-		Input.Environment.SetDefine(TEXT("MOBILE_EMULATION"), 1);
+		SET_SHADER_DEFINE(Input.Environment, MOBILE_EMULATION, 1);
 	}
 
 	// Add compiler flag CFLAG_ForceDXC if DXC is enabled
 	const bool bHlslVersion2021 = Input.Environment.CompilerFlags.Contains(CFLAG_HLSL2021);
 	const bool bIsDxcEnabled = IsDxcEnabledForPlatform((EShaderPlatform)Target.Platform, bHlslVersion2021);
-	Input.Environment.SetDefine(TEXT("COMPILER_DXC"), bIsDxcEnabled);
+	SET_SHADER_DEFINE(Input.Environment, COMPILER_DXC, bIsDxcEnabled);
 	if (bIsDxcEnabled)
 	{
 		Input.Environment.CompilerFlags.Add(CFLAG_ForceDXC);
@@ -6508,37 +6513,37 @@ void GlobalBeginCompileShader(
 		}
 	}
 
-	Input.Environment.SetDefine(TEXT("HAS_INVERTED_Z_BUFFER"), (bool)ERHIZBuffer::IsInverted);
+	SET_SHADER_DEFINE(Input.Environment, HAS_INVERTED_Z_BUFFER, (bool)ERHIZBuffer::IsInverted);
 
 	if (Input.Environment.CompilerFlags.Contains(CFLAG_HLSL2021))
 	{
-		Input.Environment.SetDefine(TEXT("COMPILER_SUPPORTS_HLSL2021"), 1);
+		SET_SHADER_DEFINE(Input.Environment, COMPILER_SUPPORTS_HLSL2021, 1);
 	}
 
 	{
 		static const auto CVar = IConsoleManager::Get().FindTConsoleVariableDataInt(TEXT("r.ClearCoatNormal"));
-		Input.Environment.SetDefine(TEXT("CLEAR_COAT_BOTTOM_NORMAL"), CVar ? (CVar->GetValueOnAnyThread() != 0) && !bIsMobilePlatform : 0);
+		SET_SHADER_DEFINE(Input.Environment, CLEAR_COAT_BOTTOM_NORMAL, CVar ? (CVar->GetValueOnAnyThread() != 0) && !bIsMobilePlatform : 0);
 	}
 
 	{
 		static const auto CVar = IConsoleManager::Get().FindTConsoleVariableDataInt(TEXT("r.IrisNormal"));
-		Input.Environment.SetDefine(TEXT("IRIS_NORMAL"), CVar ? (CVar->GetValueOnAnyThread() != 0) : 0);
+		SET_SHADER_DEFINE(Input.Environment, IRIS_NORMAL, CVar ? (CVar->GetValueOnAnyThread() != 0) : 0);
 	}
 
 	{
 		static const auto CVar = IConsoleManager::Get().FindTConsoleVariableDataInt(TEXT("Compat.UseDXT5NormalMaps"));
-		Input.Environment.SetDefine(TEXT("DXT5_NORMALMAPS"), CVar ? (CVar->GetValueOnAnyThread() != 0) : 0);
+		SET_SHADER_DEFINE(Input.Environment, DXT5_NORMALMAPS, CVar ? (CVar->GetValueOnAnyThread() != 0) : 0);
 	}
 
 	if (bAllowDevelopmentShaderCompile)
 	{
 		static const auto CVar = IConsoleManager::Get().FindTConsoleVariableDataInt(TEXT("r.CompileShadersForDevelopment"));
-		Input.Environment.SetDefine(TEXT("COMPILE_SHADERS_FOR_DEVELOPMENT"), CVar ? (CVar->GetValueOnAnyThread() != 0) : 0);
+		SET_SHADER_DEFINE(Input.Environment, COMPILE_SHADERS_FOR_DEVELOPMENT, CVar ? (CVar->GetValueOnAnyThread() != 0) : 0);
 	}
 
 	{
 		static const auto CVar = IConsoleManager::Get().FindTConsoleVariableDataInt(TEXT("r.AllowStaticLighting"));
-		Input.Environment.SetDefine(TEXT("ALLOW_STATIC_LIGHTING"), CVar ? (CVar->GetValueOnAnyThread() != 0) : 1);
+		SET_SHADER_DEFINE(Input.Environment, ALLOW_STATIC_LIGHTING, CVar ? (CVar->GetValueOnAnyThread() != 0) : 1);
 	}
 
 	{
@@ -6549,30 +6554,30 @@ void GlobalBeginCompileShader(
 			const EGBufferLayout Layout = FShaderCompileUtilities::FetchGBufferLayout(Input.Environment);
 			bGBufferHasVelocity |= (Layout == GBL_ForceVelocity);
 		}
-		Input.Environment.SetDefine(TEXT("GBUFFER_HAS_VELOCITY"), bGBufferHasVelocity ? 1 : 0);
+		SET_SHADER_DEFINE(Input.Environment, GBUFFER_HAS_VELOCITY, bGBufferHasVelocity ? 1 : 0);
 	}
 
 	{
 		static const auto CVar = IConsoleManager::Get().FindTConsoleVariableDataInt(TEXT("r.GBufferDiffuseSampleOcclusion"));
-		Input.Environment.SetDefine(TEXT("GBUFFER_HAS_DIFFUSE_SAMPLE_OCCLUSION"), CVar ? (CVar->GetValueOnAnyThread() != 0) : 1);
+		SET_SHADER_DEFINE(Input.Environment, GBUFFER_HAS_DIFFUSE_SAMPLE_OCCLUSION, CVar ? (CVar->GetValueOnAnyThread() != 0) : 1);
 	}
 
 	{
-		Input.Environment.SetDefine(TEXT("SELECTIVE_BASEPASS_OUTPUTS"), IsUsingSelectiveBasePassOutputs((EShaderPlatform)Target.Platform) ? 1 : 0);
+		SET_SHADER_DEFINE(Input.Environment, SELECTIVE_BASEPASS_OUTPUTS, IsUsingSelectiveBasePassOutputs((EShaderPlatform)Target.Platform) ? 1 : 0);
 	}
 
 	{
-		Input.Environment.SetDefine(TEXT("USE_DBUFFER"), IsUsingDBuffers((EShaderPlatform)Target.Platform) ? 1 : 0);
+		SET_SHADER_DEFINE(Input.Environment, USE_DBUFFER, IsUsingDBuffers((EShaderPlatform)Target.Platform) ? 1 : 0);
 	}
 
 	{
 		static IConsoleVariable* CVar = IConsoleManager::Get().FindConsoleVariable(TEXT("r.AllowGlobalClipPlane"));
-		Input.Environment.SetDefine(TEXT("PROJECT_ALLOW_GLOBAL_CLIP_PLANE"), CVar ? (CVar->GetInt() != 0) : 0);
+		SET_SHADER_DEFINE(Input.Environment, PROJECT_ALLOW_GLOBAL_CLIP_PLANE, CVar ? (CVar->GetInt() != 0) : 0);
 	}
 
 	{
 		const bool bSupportsClipDistance = FDataDrivenShaderPlatformInfo::GetSupportsClipDistance((EShaderPlatform)Target.Platform);
-		Input.Environment.SetDefine(TEXT("PLATFORM_SUPPORTS_CLIP_DISTANCE"), bSupportsClipDistance ? 1u : 0u);
+		SET_SHADER_DEFINE(Input.Environment, PLATFORM_SUPPORTS_CLIP_DISTANCE, bSupportsClipDistance ? 1u : 0u);
 	}
 
 	ITargetPlatform* TargetPlatform = GetTargetPlatformManager()->FindTargetPlatformWithSupport(TEXT("ShaderFormat"), ShaderFormatName);
@@ -6587,28 +6592,28 @@ void GlobalBeginCompileShader(
 			static IConsoleVariable* CVarForwardShading = IConsoleManager::Get().FindConsoleVariable(TEXT("r.ForwardShading"));
 			bForwardShading = CVarForwardShading ? (CVarForwardShading->GetInt() != 0) : false;
 		}
-		Input.Environment.SetDefine(TEXT("FORWARD_SHADING"), bForwardShading);
+		SET_SHADER_DEFINE(Input.Environment, FORWARD_SHADING, bForwardShading);
 	}
 
 	{
 		if (VelocityEncodeDepth((EShaderPlatform)Target.Platform))
 		{
-			Input.Environment.SetDefine(TEXT("VELOCITY_ENCODE_DEPTH"), 1);
+			SET_SHADER_DEFINE(Input.Environment, VELOCITY_ENCODE_DEPTH, 1);
 		}
 		else
 		{
-			Input.Environment.SetDefine(TEXT("VELOCITY_ENCODE_DEPTH"), 0);
+			SET_SHADER_DEFINE(Input.Environment, VELOCITY_ENCODE_DEPTH, 0);
 		}
 	}
 
 	{
 		if (MaskedInEarlyPass((EShaderPlatform)Target.Platform))
 		{
-			Input.Environment.SetDefine(TEXT("EARLY_Z_PASS_ONLY_MATERIAL_MASKING"), 1);
+			SET_SHADER_DEFINE(Input.Environment, EARLY_Z_PASS_ONLY_MATERIAL_MASKING, 1);
 		}
 		else
 		{
-			Input.Environment.SetDefine(TEXT("EARLY_Z_PASS_ONLY_MATERIAL_MASKING"), 0);
+			SET_SHADER_DEFINE(Input.Environment, EARLY_Z_PASS_ONLY_MATERIAL_MASKING, 0);
 		}
 	}
 
@@ -6631,118 +6636,118 @@ void GlobalBeginCompileShader(
 				}
 			}
 		}
-		Input.Environment.SetDefine(TEXT("PROJECT_VERTEX_FOGGING_FOR_OPAQUE"), bVertexFoggingForOpaque);
+		SET_SHADER_DEFINE(Input.Environment, PROJECT_VERTEX_FOGGING_FOR_OPAQUE, bVertexFoggingForOpaque);
 	}
 
 	{
 		static IConsoleVariable* CVar = IConsoleManager::Get().FindConsoleVariable(TEXT("r.Mobile.DisableVertexFog"));
-		Input.Environment.SetDefine(TEXT("PROJECT_MOBILE_DISABLE_VERTEX_FOG"), CVar ? (CVar->GetInt() != 0) : 0);
+		SET_SHADER_DEFINE(Input.Environment, PROJECT_MOBILE_DISABLE_VERTEX_FOG, CVar ? (CVar->GetInt() != 0) : 0);
 	}
 
 	bool bSupportSkyAtmosphere = false;
 	{
 		static IConsoleVariable* CVar = IConsoleManager::Get().FindConsoleVariable(TEXT("r.SupportSkyAtmosphere"));
 		bSupportSkyAtmosphere = CVar && CVar->GetInt() != 0;
-		Input.Environment.SetDefine(TEXT("PROJECT_SUPPORT_SKY_ATMOSPHERE"), bSupportSkyAtmosphere ? 1 : 0);
+		SET_SHADER_DEFINE(Input.Environment, PROJECT_SUPPORT_SKY_ATMOSPHERE, bSupportSkyAtmosphere ? 1 : 0);
 	}
 
 	{
 		static IConsoleVariable* CVar = IConsoleManager::Get().FindConsoleVariable(TEXT("r.SupportCloudShadowOnForwardLitTranslucent"));
 		const bool bSupportCloudShadowOnForwardLitTranslucent = CVar && CVar->GetInt() > 0;
-		Input.Environment.SetDefine(TEXT("SUPPORT_CLOUD_SHADOW_ON_FORWARD_LIT_TRANSLUCENT"), bSupportCloudShadowOnForwardLitTranslucent ? 1 : 0);
+		SET_SHADER_DEFINE(Input.Environment, SUPPORT_CLOUD_SHADOW_ON_FORWARD_LIT_TRANSLUCENT, bSupportCloudShadowOnForwardLitTranslucent ? 1 : 0);
 	}
 
 	{
 		static IConsoleVariable *CVar = IConsoleManager::Get().FindConsoleVariable(TEXT("r.Water.SingleLayerWater.SupportCloudShadow"));
 		const bool bSupportCloudShadowOnSingleLayerWater = CVar && CVar->GetInt() > 0;
-		Input.Environment.SetDefine(TEXT("SUPPORT_CLOUD_SHADOW_ON_SINGLE_LAYER_WATER"), bSupportCloudShadowOnSingleLayerWater ? 1 : 0);
+		SET_SHADER_DEFINE(Input.Environment, SUPPORT_CLOUD_SHADOW_ON_SINGLE_LAYER_WATER, bSupportCloudShadowOnSingleLayerWater ? 1 : 0);
 	}
 
 	{
 		static IConsoleVariable* CVar = IConsoleManager::Get().FindConsoleVariable(TEXT("r.RectLightAtlas.Translucent"));
 		const bool bSupportRectLlightOnTranslucentSurface = CVar && CVar->GetInt() > 0;
-		Input.Environment.SetDefine(TEXT("SUPPORT_RECTLIGHT_ON_FORWARD_LIT_TRANSLUCENT"), bSupportRectLlightOnTranslucentSurface ? 1 : 0);
+		SET_SHADER_DEFINE(Input.Environment, SUPPORT_RECTLIGHT_ON_FORWARD_LIT_TRANSLUCENT, bSupportRectLlightOnTranslucentSurface ? 1 : 0);
 	}
 
 	{
 		static IConsoleVariable* CVar = IConsoleManager::Get().FindConsoleVariable(TEXT("r.Shadow.Virtual.TranslucentQuality"));
 		const bool bHighQualityShadow = CVar && CVar->GetInt() > 0;
-		Input.Environment.SetDefine(TEXT("SUPPORT_VSM_FOWARD_QUALITY"), bHighQualityShadow ? 1 : 0);
+		SET_SHADER_DEFINE(Input.Environment, SUPPORT_VSM_FOWARD_QUALITY, bHighQualityShadow ? 1 : 0);
 	}
 
 	const bool bStrata = Strata::IsStrataEnabled();
 	{
-		Input.Environment.SetDefine(TEXT("STRATA_ENABLED"), bStrata ? 1 : 0);
+		SET_SHADER_DEFINE(Input.Environment, STRATA_ENABLED, bStrata ? 1 : 0);
 
 		if (bStrata)
 		{
 			const uint32 StrataShadingQuality = Strata::GetShadingQuality(Target.GetPlatform());
-			Input.Environment.SetDefine(TEXT("STRATA_SHADING_QUALITY"), StrataShadingQuality);
+			SET_SHADER_DEFINE(Input.Environment, STRATA_SHADING_QUALITY, StrataShadingQuality);
 
 			const bool bLowQuality = StrataShadingQuality > 1;
-			Input.Environment.SetDefine(TEXT("USE_ACHROMATIC_BXDF_ENERGY"), bLowQuality ? 1u : 0u);
+			SET_SHADER_DEFINE(Input.Environment, USE_ACHROMATIC_BXDF_ENERGY, bLowQuality ? 1u : 0u);
 
 			const uint32 StrataSheenQuality = Strata::GetSheenQuality();
 			Input.Environment.SetDefine(TEXT("STRATA_SHEEN_QUALITY"), bLowQuality ? 2 : StrataSheenQuality);
 
 			const uint32 StrataNormalQuality = Strata::GetNormalQuality();
-			Input.Environment.SetDefine(TEXT("STRATA_NORMAL_QUALITY"), StrataNormalQuality);
+			SET_SHADER_DEFINE(Input.Environment, STRATA_NORMAL_QUALITY, StrataNormalQuality);
 
 			const uint32 StrataUintPerPixel = Strata::GetBytePerPixel(Target.GetPlatform()) / 4u;
-			Input.Environment.SetDefine(TEXT("STRATA_MATERIAL_NUM_UINTS"), StrataUintPerPixel);
+			SET_SHADER_DEFINE(Input.Environment, STRATA_MATERIAL_NUM_UINTS, StrataUintPerPixel);
 
 			const bool bTileCoord8Bits = Strata::Is8bitTileCoordEnabled();
-			Input.Environment.SetDefine(TEXT("USE_8BIT_TILE_COORD"), bTileCoord8Bits ? 1 : 0);
+			SET_SHADER_DEFINE(Input.Environment, USE_8BIT_TILE_COORD, bTileCoord8Bits ? 1 : 0);
 
 			const bool bStrataDBufferPass = Strata::IsDBufferPassEnabled(Target.GetPlatform());
-			Input.Environment.SetDefine(TEXT("STRATA_USE_DBUFFER_PASS"), bStrataDBufferPass ? 1 : 0);
+			SET_SHADER_DEFINE(Input.Environment, STRATA_USE_DBUFFER_PASS, bStrataDBufferPass ? 1 : 0);
 
 			const bool bStrataGlints = Strata::IsGlintEnabled();
-			Input.Environment.SetDefine(TEXT("PLATFORM_ENABLES_STRATA_GLINTS"), bStrataGlints ? 1 : 0);
+			SET_SHADER_DEFINE(Input.Environment, PLATFORM_ENABLES_STRATA_GLINTS, bStrataGlints ? 1 : 0);
 
 			const bool bSpecularProfileEnabled = Strata::IsSpecularProfileEnabled();
-			Input.Environment.SetDefine(TEXT("PLATFORM_ENABLES_STRATA_SPECULAR_PROFILE"), bSpecularProfileEnabled ? 1 : 0);
+			SET_SHADER_DEFINE(Input.Environment, PLATFORM_ENABLES_STRATA_SPECULAR_PROFILE, bSpecularProfileEnabled ? 1 : 0);
 		}
 
 		const bool bStrataBackCompatibility = bStrata && Strata::IsBackCompatibilityEnabled();
-		Input.Environment.SetDefine(TEXT("PROJECT_STRATA_BACKCOMPATIBILITY"), bStrataBackCompatibility ? 1 : 0);
+		SET_SHADER_DEFINE(Input.Environment, PROJECT_STRATA_BACKCOMPATIBILITY, bStrataBackCompatibility ? 1 : 0);
 
 		const bool bStrataOpaqueRoughRefrac = bStrata && Strata::IsOpaqueRoughRefractionEnabled();
-		Input.Environment.SetDefine(TEXT("STRATA_OPAQUE_ROUGH_REFRACTION_ENABLED"), bStrataOpaqueRoughRefrac ? 1 : 0);
+		SET_SHADER_DEFINE(Input.Environment, STRATA_OPAQUE_ROUGH_REFRACTION_ENABLED, bStrataOpaqueRoughRefrac ? 1 : 0);
 
 		const bool bStrataAdvDebug = bStrata && Strata::IsAdvancedVisualizationEnabled();
-		Input.Environment.SetDefine(TEXT("STRATA_ADVANCED_DEBUG_ENABLED"), bStrataAdvDebug ? 1 : 0);
+		SET_SHADER_DEFINE(Input.Environment, STRATA_ADVANCED_DEBUG_ENABLED, bStrataAdvDebug ? 1 : 0);
 	}
 
 	{
 		static IConsoleVariable* CVar = IConsoleManager::Get().FindConsoleVariable(TEXT("r.Material.RoughDiffuse"));
 		const bool bMaterialRoughDiffuse = CVar && CVar->GetInt() != 0;
 		const bool bStrataRoughDiffuse = Strata::IsRoughDiffuseEnabled() && !Strata::IsBackCompatibilityEnabled();
-		Input.Environment.SetDefine(TEXT("MATERIAL_ROUGHDIFFUSE"), (bStrata ? bStrataRoughDiffuse : bMaterialRoughDiffuse) ? 1 : 0);
+		SET_SHADER_DEFINE(Input.Environment, MATERIAL_ROUGHDIFFUSE, (bStrata ? bStrataRoughDiffuse : bMaterialRoughDiffuse) ? 1 : 0);
 	}
 
 	{
 		static IConsoleVariable* CVar = IConsoleManager::Get().FindConsoleVariable(TEXT("r.Lumen.Supported"));
 		const bool bLumenSupported = CVar->GetInt() != 0;
-		Input.Environment.SetDefine(TEXT("PROJECT_SUPPORTS_LUMEN"), bLumenSupported ? 1 : 0);
+		SET_SHADER_DEFINE(Input.Environment, PROJECT_SUPPORTS_LUMEN, bLumenSupported ? 1 : 0);
 	}
 
 	{
 		const bool bSupportOIT = FDataDrivenShaderPlatformInfo::GetSupportsOIT(EShaderPlatform(Target.Platform));
 		static IConsoleVariable* CVar = IConsoleManager::Get().FindConsoleVariable(TEXT("r.OIT.SortedPixels"));
 		const bool bOIT = CVar && CVar->GetInt() != 0;
-		Input.Environment.SetDefine(TEXT("PROJECT_OIT"), (bSupportOIT && bOIT) ? 1 : 0);
+		SET_SHADER_DEFINE(Input.Environment, PROJECT_OIT, (bSupportOIT && bOIT) ? 1 : 0);
 	}
 
 	{
 		static IConsoleVariable* CVar = IConsoleManager::Get().FindConsoleVariable(TEXT("r.Material.EnergyConservation"));
 		const bool bMaterialEnergyConservation = CVar && CVar->GetInt() != 0;
-		Input.Environment.SetDefine(TEXT("LEGACY_MATERIAL_ENERGYCONSERVATION"), bMaterialEnergyConservation ? 1 : 0);
+		SET_SHADER_DEFINE(Input.Environment, LEGACY_MATERIAL_ENERGYCONSERVATION, bMaterialEnergyConservation ? 1 : 0);
 	}
 
 	{
 		static IConsoleVariable* CVar = IConsoleManager::Get().FindConsoleVariable(TEXT("r.SupportSkyAtmosphereAffectsHeightFog"));
-		Input.Environment.SetDefine(TEXT("PROJECT_SUPPORT_SKY_ATMOSPHERE_AFFECTS_HEIGHFOG"), (CVar && bSupportSkyAtmosphere) ? (CVar->GetInt() != 0) : 0);
+		SET_SHADER_DEFINE(Input.Environment, PROJECT_SUPPORT_SKY_ATMOSPHERE_AFFECTS_HEIGHFOG, (CVar && bSupportSkyAtmosphere) ? (CVar->GetInt() != 0) : 0);
 	}
 
 	{
@@ -6760,22 +6765,22 @@ void GlobalBeginCompileShader(
 		{
 			PropagateAlpha = 0;
 		}
-		Input.Environment.SetDefine(TEXT("POST_PROCESS_ALPHA"), PropagateAlpha);
+		SET_SHADER_DEFINE(Input.Environment, POST_PROCESS_ALPHA, PropagateAlpha);
 	}
 
 	if (TargetPlatform && 
 		TargetPlatform->SupportsFeature(ETargetPlatformFeatures::NormalmapLAEncodingMode))
 	{
-		Input.Environment.SetDefine(TEXT("LA_NORMALMAPS"), 1);
+		SET_SHADER_DEFINE(Input.Environment, LA_NORMALMAPS, 1);
 	}
 
-	Input.Environment.SetDefine(TEXT("PLATFORM_SUPPORTS_RENDERTARGET_WRITE_MASK"), RHISupportsRenderTargetWriteMask(EShaderPlatform(Target.Platform)) ? 1 : 0);
-	Input.Environment.SetDefine(TEXT("PLATFORM_SUPPORTS_PER_PIXEL_DBUFFER_MASK"), FDataDrivenShaderPlatformInfo::GetSupportsPerPixelDBufferMask(EShaderPlatform(Target.Platform)) ? 1 : 0);
-	Input.Environment.SetDefine(TEXT("PLATFORM_SUPPORTS_DISTANCE_FIELDS"), DoesPlatformSupportDistanceFields(EShaderPlatform(Target.Platform)) ? 1 : 0);
-	Input.Environment.SetDefine(TEXT("PLATFORM_SUPPORTS_MESH_SHADERS_TIER0"), RHISupportsMeshShadersTier0(EShaderPlatform(Target.Platform)) ? 1 : 0);
-	Input.Environment.SetDefine(TEXT("PLATFORM_SUPPORTS_MESH_SHADERS_TIER1"), RHISupportsMeshShadersTier1(EShaderPlatform(Target.Platform)) ? 1 : 0);
-	Input.Environment.SetDefine(TEXT("PLATFORM_ALLOW_SCENE_DATA_COMPRESSED_TRANSFORMS"), FDataDrivenShaderPlatformInfo::GetSupportSceneDataCompressedTransforms(EShaderPlatform(Target.Platform)) ? 1 : 0);
-	Input.Environment.SetDefine(TEXT("PLATFORM_SUPPORTS_BUFFER_LOAD_TYPE_CONVERSION"), RHISupportsBufferLoadTypeConversion(ShaderPlatform) ? 1 : 0);
+	SET_SHADER_DEFINE(Input.Environment, PLATFORM_SUPPORTS_RENDERTARGET_WRITE_MASK, RHISupportsRenderTargetWriteMask(EShaderPlatform(Target.Platform)) ? 1 : 0);
+	SET_SHADER_DEFINE(Input.Environment, PLATFORM_SUPPORTS_PER_PIXEL_DBUFFER_MASK, FDataDrivenShaderPlatformInfo::GetSupportsPerPixelDBufferMask(EShaderPlatform(Target.Platform)) ? 1 : 0);
+	SET_SHADER_DEFINE(Input.Environment, PLATFORM_SUPPORTS_DISTANCE_FIELDS, DoesPlatformSupportDistanceFields(EShaderPlatform(Target.Platform)) ? 1 : 0);
+	SET_SHADER_DEFINE(Input.Environment, PLATFORM_SUPPORTS_MESH_SHADERS_TIER0, RHISupportsMeshShadersTier0(EShaderPlatform(Target.Platform)) ? 1 : 0);
+	SET_SHADER_DEFINE(Input.Environment, PLATFORM_SUPPORTS_MESH_SHADERS_TIER1, RHISupportsMeshShadersTier1(EShaderPlatform(Target.Platform)) ? 1 : 0);
+	SET_SHADER_DEFINE(Input.Environment, PLATFORM_ALLOW_SCENE_DATA_COMPRESSED_TRANSFORMS, FDataDrivenShaderPlatformInfo::GetSupportSceneDataCompressedTransforms(EShaderPlatform(Target.Platform)) ? 1 : 0);
+	SET_SHADER_DEFINE(Input.Environment, PLATFORM_SUPPORTS_BUFFER_LOAD_TYPE_CONVERSION, RHISupportsBufferLoadTypeConversion(ShaderPlatform) ? 1 : 0);
 
 	bool bEnableBindlessMacro = false;
 	if (RHIGetBindlessSupport(ShaderPlatform) != ERHIBindlessSupport::Unsupported)
@@ -6789,18 +6794,18 @@ void GlobalBeginCompileShader(
 		{
 			bEnableBindlessMacro = true;
 			Input.Environment.CompilerFlags.Add(CFLAG_BindlessResources);
-			Input.Environment.SetDefine(TEXT("ENABLE_BINDLESS_RESOURCES"), true);
+			SET_SHADER_DEFINE(Input.Environment, ENABLE_BINDLESS_RESOURCES, true);
 		}
 
 		if (SamplersConfig == ERHIBindlessConfiguration::AllShaders || (SamplersConfig == ERHIBindlessConfiguration::RayTracingShaders && bIsRaytracingShader))
 		{
 			bEnableBindlessMacro = true;
 			Input.Environment.CompilerFlags.Add(CFLAG_BindlessSamplers);
-			Input.Environment.SetDefine(TEXT("ENABLE_BINDLESS_SAMPLERS"), true);
+			SET_SHADER_DEFINE(Input.Environment, ENABLE_BINDLESS_SAMPLERS, true);
 		}
 	}
 
-	Input.Environment.SetDefine(TEXT("PLATFORM_SUPPORTS_BINDLESS"), bEnableBindlessMacro);
+	SET_SHADER_DEFINE(Input.Environment, PLATFORM_SUPPORTS_BINDLESS, bEnableBindlessMacro);
 
 	if (CVarShadersRemoveDeadCode.GetValueOnAnyThread())
 	{
@@ -6809,25 +6814,25 @@ void GlobalBeginCompileShader(
 
 	{
 		static const auto CVar = IConsoleManager::Get().FindConsoleVariable(TEXT("r.VT.AnisotropicFiltering"));
-		Input.Environment.SetDefine(TEXT("VIRTUAL_TEXTURE_ANISOTROPIC_FILTERING"), CVar ? (CVar->GetInt() != 0) : 0);
+		SET_SHADER_DEFINE(Input.Environment, VIRTUAL_TEXTURE_ANISOTROPIC_FILTERING, CVar ? (CVar->GetInt() != 0) : 0);
 		
 		if (bIsMobilePlatform)
 		{
 			static FShaderPlatformCachedIniValue<bool> CVarVTMobileManualTrilinearFiltering(TEXT("r.VT.Mobile.ManualTrilinearFiltering"));
-			Input.Environment.SetDefine(TEXT("VIRTUAL_TEXTURE_MANUAL_TRILINEAR_FILTERING"), (CVarVTMobileManualTrilinearFiltering.Get(Target.GetPlatform()) ? 1 : 0));
+			SET_SHADER_DEFINE(Input.Environment, VIRTUAL_TEXTURE_MANUAL_TRILINEAR_FILTERING, (CVarVTMobileManualTrilinearFiltering.Get(Target.GetPlatform()) ? 1 : 0));
 		}
 	}
 
 	if (bIsMobilePlatform)
 	{
 		const bool bMobileMovableSpotlightShadowsEnabled = IsMobileMovableSpotlightShadowsEnabled(Target.GetPlatform());
-		Input.Environment.SetDefine(TEXT("PROJECT_MOBILE_ENABLE_MOVABLE_SPOTLIGHT_SHADOWS"), bMobileMovableSpotlightShadowsEnabled ? 1 : 0);
+		SET_SHADER_DEFINE(Input.Environment, PROJECT_MOBILE_ENABLE_MOVABLE_SPOTLIGHT_SHADOWS, bMobileMovableSpotlightShadowsEnabled ? 1 : 0);
 	}
 
 	{
 		using namespace UE::Color;
 		const bool bWorkingColorSpaceIsSRGB = FColorSpace::GetWorking().IsSRGB();
-		Input.Environment.SetDefine(TEXT("WORKING_COLOR_SPACE_IS_SRGB"), bWorkingColorSpaceIsSRGB ? 1 : 0);
+		SET_SHADER_DEFINE(Input.Environment, WORKING_COLOR_SPACE_IS_SRGB, bWorkingColorSpaceIsSRGB ? 1 : 0);
 		
 		// We limit matrix definitions below to WORKING_COLOR_SPACE_IS_SRGB == 0.
 		if (!bWorkingColorSpaceIsSRGB)
@@ -6853,8 +6858,8 @@ void GlobalBeginCompileShader(
 					FromXYZ.M[0][2], FromXYZ.M[1][2], FromXYZ.M[2][2]));
 
 			const FColorSpaceTransform FromSRGB(FColorSpace(EColorSpace::sRGB), WorkingColorSpace);
-			Input.Environment.SetDefine(
-				TEXT("SRGB_TO_WORKING_COLOR_SPACE_MAT"),
+			SET_SHADER_DEFINE(Input.Environment,
+				SRGB_TO_WORKING_COLOR_SPACE_MAT,
 				FString::Printf(MatrixFormat,
 					FromSRGB.M[0][0], FromSRGB.M[1][0], FromSRGB.M[2][0],
 					FromSRGB.M[0][1], FromSRGB.M[1][1], FromSRGB.M[2][1],
@@ -6863,12 +6868,12 @@ void GlobalBeginCompileShader(
 	}
 
 	const double TileSize = FLargeWorldRenderScalar::GetTileSize();
-	Input.Environment.SetDefine(TEXT("UE_LWC_RENDER_TILE_SIZE"), (float)TileSize);
-	Input.Environment.SetDefine(TEXT("UE_LWC_RENDER_TILE_SIZE_SQRT"), (float)FMath::Sqrt(TileSize));
-	Input.Environment.SetDefine(TEXT("UE_LWC_RENDER_TILE_SIZE_RSQRT"), (float)FMath::InvSqrt(TileSize));
-	Input.Environment.SetDefine(TEXT("UE_LWC_RENDER_TILE_SIZE_RCP"), (float)(1.0 / TileSize));
-	Input.Environment.SetDefine(TEXT("UE_LWC_RENDER_TILE_SIZE_FMOD_PI"), (float)FMath::Fmod(TileSize, UE_DOUBLE_PI));
-	Input.Environment.SetDefine(TEXT("UE_LWC_RENDER_TILE_SIZE_FMOD_2PI"), (float)FMath::Fmod(TileSize, 2.0 * UE_DOUBLE_PI));
+	SET_SHADER_DEFINE(Input.Environment, UE_LWC_RENDER_TILE_SIZE, (float)TileSize);
+	SET_SHADER_DEFINE(Input.Environment, UE_LWC_RENDER_TILE_SIZE_SQRT, (float)FMath::Sqrt(TileSize));
+	SET_SHADER_DEFINE(Input.Environment, UE_LWC_RENDER_TILE_SIZE_RSQRT, (float)FMath::InvSqrt(TileSize));
+	SET_SHADER_DEFINE(Input.Environment, UE_LWC_RENDER_TILE_SIZE_RCP, (float)(1.0 / TileSize));
+	SET_SHADER_DEFINE(Input.Environment, UE_LWC_RENDER_TILE_SIZE_FMOD_PI, (float)FMath::Fmod(TileSize, UE_DOUBLE_PI));
+	SET_SHADER_DEFINE(Input.Environment, UE_LWC_RENDER_TILE_SIZE_FMOD_2PI, (float)FMath::Fmod(TileSize, 2.0 * UE_DOUBLE_PI));
 
 	// Allow the target shader format to modify the shader input before we add it as a job
 	const IShaderFormat* Format = GetTargetPlatformManagerRef().FindShaderFormat(ShaderFormatName);
diff --git a/Engine/Source/Runtime/Engine/Private/ShaderCompiler/ShaderGenerationUtil.cpp b/Engine/Source/Runtime/Engine/Private/ShaderCompiler/ShaderGenerationUtil.cpp
--- a/Engine/Source/Runtime/Engine/Private/ShaderCompiler/ShaderGenerationUtil.cpp
+++ b/Engine/Source/Runtime/Engine/Private/ShaderCompiler/ShaderGenerationUtil.cpp
@@ -14,12 +14,6 @@
 #include "HAL/PlatformFileManager.h"
 #include "RenderUtils.h"
 #include "SceneManagement.h"
-
-
-#if WITH_EDITOR
-#endif
-
-
 #include "ShaderMaterial.h"
 
 static TAutoConsoleVariable<int32> CVarShaderUseGBufferRefactor(
@@ -41,26 +35,19 @@ bool NeedsVelocityDepth(EShaderPlatform TargetPlatform)
 
 #if WITH_EDITOR
 
-int FShaderCompileUtilities::FetchCompileInt(const FShaderCompilerEnvironment& OutEnvironment, const char* SrcName)
-{
-	int32 Ret = 0;
-	if (OutEnvironment.Definitions.GetDefinitionMap().Contains(SrcName))
-	{
-		const FString* Str = OutEnvironment.Definitions.GetDefinitionMap().Find(SrcName);
-
-		Ret = atoi((const char*)Str->GetCharArray().GetData());
-	}
-
-	return Ret;
-}
+#define SET_COMPILE_BOOL_IF_TRUE(X) { if (DerivedDefines.X) { OutEnvironment.SetDefine(TEXT(#X), 1); } }
 
+#define FETCH_COMPILE_BOOL(X) {																\
+	static FShaderCompilerDefineNameCache Cache_##X(TEXT(#X));								\
+	SrcDefines.X = Environment.GetIntegerValue(Cache_##X, SrcDefines.X) != 0 ? 1 : 0; }
 
-#define SET_COMPILE_BOOL_IF_TRUE(X) { if (DerivedDefines.X) { OutEnvironment.SetDefine(TEXT(#X),TEXT("1")); } }
+#define FETCH_COMPILE_INT(X) {																\
+	static FShaderCompilerDefineNameCache Cache_##X(TEXT(#X));								\
+	SrcDefines.X = Environment.GetIntegerValue(Cache_##X, SrcDefines.X); }
 
-#define FETCH_COMPILE_BOOL(X) { if (Environment.Definitions.GetDefinitionMap().Contains(#X)) SrcDefines.X = FetchCompileInt(Environment,#X) != 0 ? 1 : 0; }
-#define FETCH_COMPILE_INT(X) { if (Environment.Definitions.GetDefinitionMap().Contains(#X)) SrcDefines.X = FetchCompileInt(Environment,#X); }
 
-void FShaderCompileUtilities::ApplyFetchEnvironment(FShaderGlobalDefines& SrcDefines, const FShaderCompilerEnvironment& Environment, const EShaderPlatform Platform)
+template<typename EnvironmentType>
+void ApplyFetchEnvironmentInternal(FShaderGlobalDefines& SrcDefines, const EnvironmentType& Environment)
 {
 	FETCH_COMPILE_BOOL(GBUFFER_HAS_VELOCITY);
 	FETCH_COMPILE_BOOL(GBUFFER_HAS_TANGENT);
@@ -85,15 +72,10 @@ void FShaderCompileUtilities::ApplyFetchEnvironment(FShaderGlobalDefines& SrcDef
 	FETCH_COMPILE_BOOL(PLATFORM_SUPPORTS_DISTANCE_FIELDS);
 	FETCH_COMPILE_BOOL(COMPILE_SHADERS_FOR_DEVELOPMENT_ALLOWED);
 	FETCH_COMPILE_BOOL(PLATFORM_ALLOW_SCENE_DATA_COMPRESSED_TRANSFORMS);
-
-	// note that we are doing an if so that if we call ApplyFetchEnvironment() twice, we get the logical OR of bSupportsDualBlending support
-	if (RHISupportsDualSourceBlending(Platform))
-	{
-		SrcDefines.bSupportsDualBlending = true;
-	}
 }
 
-void FShaderCompileUtilities::ApplyFetchEnvironment(FShaderLightmapPropertyDefines& SrcDefines, const FShaderCompilerEnvironment& Environment)
+template<typename EnvironmentType>
+void ApplyFetchEnvironmentInternal(FShaderLightmapPropertyDefines& SrcDefines, const EnvironmentType& Environment)
 {
 	FETCH_COMPILE_BOOL(LQ_TEXTURE_LIGHTMAP);
 	FETCH_COMPILE_BOOL(HQ_TEXTURE_LIGHTMAP);
@@ -120,7 +102,8 @@ void FShaderCompileUtilities::ApplyFetchEnvironment(FShaderLightmapPropertyDefin
 
 }
 
-void FShaderCompileUtilities::ApplyFetchEnvironment(FShaderMaterialPropertyDefines& SrcDefines, const FShaderCompilerEnvironment& Environment)
+template<typename EnvironmentType>
+void ApplyFetchEnvironmentInternal(FShaderMaterialPropertyDefines& SrcDefines, const EnvironmentType& Environment)
 {
 	FETCH_COMPILE_BOOL(MATERIAL_ENABLE_TRANSLUCENCY_FOGGING);
 	FETCH_COMPILE_BOOL(MATERIALBLENDING_ANY_TRANSLUCENT);
@@ -236,7 +219,8 @@ void FShaderCompileUtilities::ApplyFetchEnvironment(FShaderMaterialPropertyDefin
 	FETCH_COMPILE_INT(GBUFFER_LAYOUT);
 }
 
-void FShaderCompileUtilities::ApplyFetchEnvironment(FShaderCompilerDefines& SrcDefines, const FShaderCompilerEnvironment& Environment)
+template<typename EnvironmentType>
+void ApplyFetchEnvironmentInternal(FShaderCompilerDefines& SrcDefines, const EnvironmentType& Environment)
 {
 	FETCH_COMPILE_BOOL(COMPILER_GLSL_ES3_1);
 	FETCH_COMPILE_BOOL(ES3_1_PROFILE);
@@ -253,6 +237,71 @@ void FShaderCompileUtilities::ApplyFetchEnvironment(FShaderCompilerDefines& SrcD
 	FETCH_COMPILE_BOOL(PLATFORM_SUPPORTS_DEVELOPMENT_SHADERS);
 }
 
+
+void FShaderCompileUtilities::ApplyFetchEnvironment(FShaderGlobalDefines& SrcDefines, const FShaderCompilerEnvironment& Environment, const EShaderPlatform Platform)
+{
+	ApplyFetchEnvironmentInternal(SrcDefines, Environment);
+
+	// note that we are doing an if so that if we call ApplyFetchEnvironment() twice, we get the logical OR of bSupportsDualBlending support
+	if (RHISupportsDualSourceBlending(Platform))
+	{
+		SrcDefines.bSupportsDualBlending = true;
+	}
+}
+void FShaderCompileUtilities::ApplyFetchEnvironment(FShaderLightmapPropertyDefines& SrcDefines, const FShaderCompilerEnvironment& Environment)
+{
+	ApplyFetchEnvironmentInternal(SrcDefines, Environment);
+}
+void FShaderCompileUtilities::ApplyFetchEnvironment(FShaderMaterialPropertyDefines& SrcDefines, const FShaderCompilerEnvironment& Environment)
+{
+	ApplyFetchEnvironmentInternal(SrcDefines, Environment);
+}
+void FShaderCompileUtilities::ApplyFetchEnvironment(FShaderCompilerDefines& SrcDefines, const FShaderCompilerEnvironment& Environment)
+{
+	ApplyFetchEnvironmentInternal(SrcDefines, Environment);
+}
+
+/** Dummy "environment" used to gather names of defines used in the ApplyFetchEnvironmentInternal functions */
+struct FDefineNameGatherEnvironment
+{
+	mutable TArray<FName> Names;
+
+	FORCEINLINE int32 GetIntegerValue(FShaderCompilerDefineNameCache& Cache, int32 DefaultValue) const
+	{
+		Names.Add(Cache);
+		return 0;
+	}
+};
+
+/** Singleton to initialize initial defines at startup */
+struct FShaderInitialDefinesInitializer
+{
+	FShaderInitialDefinesInitializer()
+	{
+		FShaderMaterialPropertyDefines MaterialDefines = {};
+		FShaderLightmapPropertyDefines LightmapDefines = {};
+		FShaderGlobalDefines GlobalDefines = {};
+		FShaderCompilerDefines CompilerDefines = {};
+
+		FDefineNameGatherEnvironment GatherNames;
+
+		ApplyFetchEnvironmentInternal(GlobalDefines, GatherNames);
+		ApplyFetchEnvironmentInternal(MaterialDefines, GatherNames);
+		ApplyFetchEnvironmentInternal(LightmapDefines, GatherNames);
+		ApplyFetchEnvironmentInternal(CompilerDefines, GatherNames);
+
+		FShaderCompilerDefinitions InitialDefines;
+		for (FName DefineKey : GatherNames.Names)
+		{
+			// Add a key, leaving the value unset (EShaderCompilerDefineVariant::None)
+			InitialDefines.FindOrAddMapIndex(DefineKey);
+		}
+		FShaderCompilerDefinitions::InitializeInitialDefines(InitialDefines);
+	}
+};
+static FShaderInitialDefinesInitializer GInitialDefinesInitializer;
+
+
 // if we change the logic, increment this number to force a DDC key change
 static const int32 GBufferGeneratorVersion = 5;
 
@@ -2111,7 +2160,7 @@ void FShaderCompileUtilities::GenerateBrdfHeaders(const FName& ShaderFormat)
 
 EGBufferLayout FShaderCompileUtilities::FetchGBufferLayout(const FShaderCompilerEnvironment& Environment)
 {
-	const uint32 Layout = FetchCompileInt(Environment, "GBUFFER_LAYOUT");
+	const uint32 Layout = Environment.Definitions.GetIntegerValue(TEXT("GBUFFER_LAYOUT"));
 	if (Layout >= GBL_Num)
 	{
 		return GBL_Default;
@@ -2150,7 +2199,7 @@ FGBufferParams FShaderCompileUtilities::FetchGBufferParamsPipeline(EShaderPlatfo
 	return Ret;
 }
 
-#endif
+#endif  // WITH_EDITOR
 
 FGBufferParams FShaderCompileUtilities::FetchGBufferParamsRuntime(EShaderPlatform Platform, EGBufferLayout Layout)
 {
diff --git a/Engine/Source/Runtime/Engine/Public/ShaderCompiler.h b/Engine/Source/Runtime/Engine/Public/ShaderCompiler.h
--- a/Engine/Source/Runtime/Engine/Public/ShaderCompiler.h
+++ b/Engine/Source/Runtime/Engine/Public/ShaderCompiler.h
@@ -382,7 +382,6 @@ public:
 	static void AppendGBufferDDCKeyString(const EShaderPlatform Platform, FString& KeyString);
 	static ENGINE_API void WriteGBufferInfoAutogen(EShaderPlatform TargetPlatform, ERHIFeatureLevel::Type FeatureLevel);
 
-	static int FetchCompileInt(const FShaderCompilerEnvironment& Enviroment, const char* SrcName);
 	static void ApplyFetchEnvironment(FShaderMaterialPropertyDefines& DefineData, const FShaderCompilerEnvironment& Environment);
 	static void ApplyFetchEnvironment(FShaderGlobalDefines& DefineData, const FShaderCompilerEnvironment& Environment, const EShaderPlatform Platform);
 	static void ApplyFetchEnvironment(FShaderLightmapPropertyDefines& DefineData, const FShaderCompilerEnvironment& Environment);
diff --git a/Engine/Source/Runtime/RenderCore/Private/ShaderCore.cpp b/Engine/Source/Runtime/RenderCore/Private/ShaderCore.cpp
--- a/Engine/Source/Runtime/RenderCore/Private/ShaderCore.cpp
+++ b/Engine/Source/Runtime/RenderCore/Private/ShaderCore.cpp
@@ -116,6 +116,322 @@ DEFINE_STAT(STAT_Shaders_ShaderPreloadMemory);
 DEFINE_STAT(STAT_Shaders_NumShadersRegistered);
 DEFINE_STAT(STAT_Shaders_NumShadersDuplicated);
 
+/**
+ * Singleton initial set of defines added when constructing a defines structure with bIncludeInitialDefines==true.  The
+ * advantage of using preset defines is that the index of the initial define can be cached in the FShaderCompilerDefineNameCache
+ * class, allowing direct lookup by index, bypassing the hash table.  This optimization is applied to system level defines used
+ * by every shader (ones referenced by FShaderCompileUtilities::ApplyDerivedDefines).
+ */
+FShaderCompilerDefinitions* FShaderCompilerDefinitions::GInitialDefines = nullptr;
+
+FShaderCompilerDefinitions::FShaderCompilerDefinitions(bool bIncludeInitialDefines)
+	: InitialDefineCount(0), ValueCount(0)
+{
+	if (bIncludeInitialDefines && GInitialDefines)
+	{
+		*this = *GInitialDefines;
+	}
+	else
+	{
+		Pairs.Reserve(16);
+		ValueTypes.Reserve(16);
+	}
+}
+
+void FShaderCompilerDefinitions::InitializeInitialDefines(const FShaderCompilerDefinitions& InDefines)
+{
+	check(GInitialDefines == nullptr);
+	GInitialDefines = new FShaderCompilerDefinitions;
+	*GInitialDefines = InDefines;
+	GInitialDefines->InitialDefineCount = InDefines.Pairs.Num();
+}
+
+int32 FShaderCompilerDefinitions::FindMapIndex(FName Key, uint32 KeyHash) const
+{
+	for (uint32 KeyIndex = KeyHashTable.First(KeyHash); KeyHashTable.IsValid(KeyIndex); KeyIndex = KeyHashTable.Next(KeyIndex))
+	{
+		if (Pairs[KeyIndex].Key == Key)
+		{
+			return KeyIndex;
+		}
+	}
+	return INDEX_NONE;
+}
+
+int32 FShaderCompilerDefinitions::FindOrAddMapIndex(FName Name)
+{
+	uint32 KeyHash = GetTypeHash(Name);
+	int32 KeyIndex = FindMapIndex(Name, KeyHash);
+	if (KeyIndex == INDEX_NONE)
+	{
+		KeyIndex = Pairs.Add(FPairType({ Name, {0} }));
+		ValueTypes.Add(EShaderCompilerDefineVariant::None);
+		KeyHashTable.Add(KeyHash, KeyIndex);
+	}
+	return KeyIndex;
+}
+
+int32 FShaderCompilerDefinitions::FindOrAddMapIndex(FShaderCompilerDefineNameCache& NameCache)
+{
+	// Check if this is an initial define, meaning it has a fixed map index.  If Name.MapIndex is INDEX_NONE
+	// (MapIndex not initialized yet) or InitialDefineCount is zero (FShaderCompilerDefinitions constructed without
+	// initial state), this condition will be false, and continue to the rest of the function.
+	if ((uint32)NameCache.MapIndex < InitialDefineCount)
+	{
+		return NameCache.MapIndex;
+	}
+
+	uint32 KeyHash = GetTypeHash(NameCache.Name);
+	int32 KeyIndex = FindMapIndex(NameCache.Name, KeyHash);
+	if (KeyIndex == INDEX_NONE)
+	{
+		KeyIndex = Pairs.Add(FPairType({ NameCache.Name, {0} }));
+		ValueTypes.Add(EShaderCompilerDefineVariant::None);
+		KeyHashTable.Add(KeyHash, KeyIndex);
+	}
+
+	// Initialize MapIndex if necessary.
+	if (InitialDefineCount && NameCache.MapIndex == INDEX_NONE)
+	{
+		NameCache.MapIndex = KeyIndex;
+	}
+	return KeyIndex;
+}
+
+void FShaderCompilerDefinitions::InternalSetValue(int32 Index, const TCHAR* Value)
+{
+	if (Value[0] >= '0' && Value[0] <= '9' && Value[1] == 0)
+	{
+		// If the string is a single digit integer, treat it as an integer.  Most usages of string define values are cases where
+		// clients inadvertently pass bools as text, like TEXT("0") or TEXT("1"), and we can trivially optimize that.
+		SetValueType(Index, EShaderCompilerDefineVariant::Integer);
+		Pairs[Index].ValueInteger = Value[0] - '0';
+	}
+	else if (ValueTypes[Index] == EShaderCompilerDefineVariant::String)
+	{
+		// Value was already a string variant, overwrite it with new value, rather than allocating new string
+		StringValues[Pairs[Index].ValueInteger] = Value;
+	}
+	else
+	{
+		// Set to string variant type, allocate and fill in the string
+		SetValueType(Index, EShaderCompilerDefineVariant::String);
+		Pairs[Index].ValueInteger = StringValues.Add(Value);
+	}
+}
+
+int32 FShaderCompilerDefinitions::GetIntegerValue(FName Name) const
+{
+	int32 Result = 0;
+	int32 KeyIndex = FindMapIndex(Name, GetTypeHash(Name));
+	if (KeyIndex != INDEX_NONE)
+	{
+		// For None, Integer, or Unsigned, return ValueInteger (None will have a default of zero)
+		if (ValueTypes[KeyIndex] <= EShaderCompilerDefineVariant::Unsigned)
+		{
+			Result = Pairs[KeyIndex].ValueInteger;
+		}
+		else if (ValueTypes[KeyIndex] == EShaderCompilerDefineVariant::String)
+		{
+			Result = FCString::Atoi(*StringValues[Pairs[KeyIndex].ValueInteger]);
+		}
+		else
+		{
+			// EShaderCompilerDefineVariant::Float
+			Result = (int32)Pairs[KeyIndex].ValueFloat;
+		}
+	}
+	return Result;
+}
+
+int32 FShaderCompilerDefinitions::GetIntegerValue(FShaderCompilerDefineNameCache& NameCache, int32 ResultIfNotFound) const
+{
+	int32 Result = ResultIfNotFound;
+	int32 KeyIndex;
+
+	// Check if this is an initial define, meaning it has a fixed map index.  If Name.MapIndex is INDEX_NONE
+	// (MapIndex not initialized yet) or InitialDefineCount is zero (FShaderCompilerDefinitions constructed without
+	// initial state), this condition will be false, and continue to the rest of the function.
+	if ((uint32)NameCache.MapIndex < InitialDefineCount)
+	{
+		KeyIndex = NameCache.MapIndex;
+	}
+	else
+	{
+		KeyIndex = FindMapIndex(NameCache.Name, GetTypeHash(NameCache.Name));
+	}
+
+	if (KeyIndex != INDEX_NONE)
+	{
+		// Initialize MapIndex if necessary.  If the define is not an initial define, its index will be greater than
+		// or equal to InitialDefineCount, indicating that it doesn't have a fixed map index.
+		if (InitialDefineCount && NameCache.MapIndex == INDEX_NONE)
+		{
+			NameCache.MapIndex = KeyIndex;
+		}
+
+		EShaderCompilerDefineVariant ValueType = ValueTypes[KeyIndex];
+		if (ValueType == EShaderCompilerDefineVariant::Integer ||
+			ValueType == EShaderCompilerDefineVariant::Unsigned)
+		{
+			Result = Pairs[KeyIndex].ValueInteger;
+		}
+		else if (ValueType == EShaderCompilerDefineVariant::String)
+		{
+			Result = FCString::Atoi(*StringValues[Pairs[KeyIndex].ValueInteger]);
+		}
+		else if (ValueType == EShaderCompilerDefineVariant::Float)
+		{
+			// EShaderCompilerDefineVariant::Float
+			Result = (int32)Pairs[KeyIndex].ValueFloat;
+		}
+	}
+	return Result;
+}
+
+RENDERCORE_API void FShaderCompilerDefinitions::Empty()
+{
+	if (ValueCount)
+	{
+		KeyHashTable.Free();
+		Pairs.Empty();
+		ValueTypes.Empty();
+		StringValues.Empty();
+		ValueCount = 0;
+
+		// Re-copy initial defines if originally constructed with initial defines
+		if (InitialDefineCount)
+		{
+			*this = *GInitialDefines;
+		}
+	}
+}
+
+void FShaderCompilerDefinitions::Merge(const FShaderCompilerDefinitions& Other)
+{
+	for (FConstIterator OtherIt(Other); OtherIt; ++OtherIt)
+	{
+		int32 OtherIndex = OtherIt.GetIndex();
+		switch (Other.ValueTypes[OtherIndex])
+		{
+		case EShaderCompilerDefineVariant::Integer:
+			SetDefine(OtherIt.Key(), Other.Pairs[OtherIndex].ValueInteger);
+			break;
+		case EShaderCompilerDefineVariant::Unsigned:
+			SetDefine(OtherIt.Key(), Other.Pairs[OtherIndex].ValueUnsigned);
+			break;
+		case EShaderCompilerDefineVariant::Float:
+			SetDefine(OtherIt.Key(), Other.Pairs[OtherIndex].ValueFloat);
+			break;
+		case EShaderCompilerDefineVariant::String:
+			SetDefine(OtherIt.Key(), OtherIt.Value());
+			break;
+		}
+	}
+}
+
+RENDERCORE_API FShaderCompilerDefinitions& FShaderCompilerDefinitions::operator=(const FShaderCompilerDefinitions& Other)
+{
+	KeyHashTable = Other.KeyHashTable;
+	Pairs = Other.Pairs;
+	ValueTypes = Other.ValueTypes;
+	StringValues = Other.StringValues;
+	InitialDefineCount = Other.InitialDefineCount;
+	ValueCount = Other.ValueCount;
+	return *this;
+}
+
+FArchive& operator<<(FArchive& Ar, FShaderCompilerDefinitions& Defs)
+{
+	if (Ar.IsSaving())
+	{
+		// Only write set values in the map
+		Ar << Defs.ValueCount;
+		for (FShaderCompilerDefinitions::FConstIterator DefineIt(Defs); DefineIt; ++DefineIt)
+		{
+			int32 Index = DefineIt.GetIndex();
+			Ar << Defs.Pairs[Index];
+			Ar << Defs.ValueTypes[Index];
+		}
+		Ar << Defs.StringValues;
+	}
+	else if (Ar.IsLoading())
+	{
+		Ar << Defs.ValueCount;
+		for (uint32 ValueIndex = 0; ValueIndex < Defs.ValueCount; ValueIndex++)
+		{
+			FShaderCompilerDefinitions::FPairType Pair;
+			Ar << Pair;
+			int32 Index = Defs.FindOrAddMapIndex(Pair.Key);
+			Defs.Pairs[Index] = Pair;
+			Ar << Defs.ValueTypes[Index];
+		}
+		Ar << Defs.StringValues;
+	}
+	return Ar;
+}
+
+RENDERCORE_API FShaderCompilerDefinitions::FConstIterator::FConstIterator(const FShaderCompilerDefinitions& InDefines)
+	: Defines(InDefines), Index(-1)
+{
+	// NULL terminate these strings to start
+	KeyStringBuffer[0] = 0;
+	ValueStringBuffer[0] = 0;
+
+	// Index set to -1 above, advance to first valid element
+	++(*this);
+}
+
+RENDERCORE_API const TCHAR* FShaderCompilerDefinitions::FConstIterator::Value()
+{
+	const TCHAR* Result;
+
+	if (Defines.ValueTypes[Index] == EShaderCompilerDefineVariant::Integer)
+	{
+		int32 ValueInteger = Defines.Pairs[Index].ValueInteger;
+		if (ValueInteger >= 0 && ValueInteger <= 9)
+		{
+			ValueStringBuffer[0] = (TCHAR)ValueInteger + '0';
+			ValueStringBuffer[1] = 0;
+		}
+		else
+		{
+			FCString::Sprintf(ValueStringBuffer, TEXT("%d"), ValueInteger);
+		}
+		Result = ValueStringBuffer;
+	}
+	else if (Defines.ValueTypes[Index] == EShaderCompilerDefineVariant::Unsigned)
+	{
+		uint32 ValueUnsigned = Defines.Pairs[Index].ValueUnsigned;
+		if (ValueUnsigned >= 0 && ValueUnsigned <= 9)
+		{
+			ValueStringBuffer[0] = (TCHAR)ValueUnsigned + '0';
+			ValueStringBuffer[1] = 0;
+		}
+		else
+		{
+			FCString::Sprintf(ValueStringBuffer, TEXT("%u"), ValueUnsigned);
+		}
+		Result = ValueStringBuffer;
+	}
+	else if (Defines.ValueTypes[Index] == EShaderCompilerDefineVariant::Float)
+	{
+		// Make sure the printed value perfectly matches the given number
+		FCString::Sprintf(ValueStringBuffer, TEXT("%#.9gf"), Defines.Pairs[Index].ValueFloat);
+		Result = ValueStringBuffer;
+	}
+	else
+	{
+		check(Defines.ValueTypes[Index] == EShaderCompilerDefineVariant::String);
+		Result = *Defines.StringValues[Defines.Pairs[Index].ValueInteger];
+	}
+
+	return Result;
+}
+
+
+
+
 // Apply lock striping as we're mostly reader lock bound.
 constexpr int32 GSHADERFILECACHE_BUCKETS = 31; /* prime number for best distribution using modulo */
 
@@ -756,13 +1072,6 @@ void FShaderResourceTableMap::FixupOnLoad(const TMap<FString, FUniformBufferEntr
 	}
 }
 
-void FShaderCompilerDefinitions::SetFloatDefine(const TCHAR* Name, float Value)
-{
-	// Make sure the printed value perfectly matches the given number
-	FString Define = FString::Printf(TEXT("%#.9gf"), Value);
-	Definitions.Add(Name, MoveTemp(Define));
-}
-
 void FShaderCompilerOutput::GenerateOutputHash()
 {
 	FSHA1 HashState;
@@ -2109,6 +2418,11 @@ void BuildShaderFileToUniformBufferMap(TMap<FString, TArray<const TCHAR*> >& Sha
 				SearchKeyWithSpace(FString(ShaderVariable).ToUpper() + TEXT(" ."))
 			{}
 
+			bool operator<(const FShaderVariable& Other) const
+			{
+				return FCString::Strcmp(OriginalShaderVariable, Other.OriginalShaderVariable) < 0;
+			}
+
 			const TCHAR* OriginalShaderVariable;
 			FString SearchKey;
 			FString SearchKeyWithSpace;
@@ -2120,6 +2434,10 @@ void BuildShaderFileToUniformBufferMap(TMap<FString, TArray<const TCHAR*> >& Sha
 			SearchKeys.Add(FShaderVariable(StructIt->GetShaderVariableName()));
 		}
 
+		// Sort SearchKeys for determinism in the generated ShaderFileToUniformBufferVariables maps, to improve consistency for A/B testing.
+		// Order of items in FShaderParametersMetadata::GetStructList() is otherwise dependent on arbitrary startup constructor order.
+		SearchKeys.Sort();
+
 		TArray<UE::Tasks::TTask<void>> Tasks;
 		Tasks.Reserve(ShaderSourceFiles.Num());
 
diff --git a/Engine/Source/Runtime/RenderCore/Private/ShaderParameterMetadata.cpp b/Engine/Source/Runtime/RenderCore/Private/ShaderParameterMetadata.cpp
--- a/Engine/Source/Runtime/RenderCore/Private/ShaderParameterMetadata.cpp
+++ b/Engine/Source/Runtime/RenderCore/Private/ShaderParameterMetadata.cpp
@@ -69,6 +69,14 @@ TMap<FHashedName, FShaderParametersMetadata*>& FShaderParametersMetadata::GetNam
 	return NameStructMap;
 }
 
+#if WITH_EDITOR
+TMap<FString, FShaderParametersMetadata*>& FShaderParametersMetadata::GetStringStructMap()
+{
+	static TMap<FString, FShaderParametersMetadata*> StringStructMap;
+	return StringStructMap;
+}
+#endif  // WITH_EDITOR
+
 void FShaderParametersMetadata::FMember::GenerateShaderParameterType(
 	FString& Result,
 	bool bSupportsPrecisionModifier,
@@ -459,6 +467,10 @@ FShaderParametersMetadata::~FShaderParametersMetadata()
 
 void FShaderParametersMetadata::InitializeAllUniformBufferStructs()
 {
+#if WITH_EDITOR
+	TMap<FString, FShaderParametersMetadata*>& StringStructMap = FShaderParametersMetadata::GetStringStructMap();
+#endif // WITH_EDITOR
+
 	for (TLinkedList<FShaderParametersMetadata*>::TIterator StructIt(FShaderParametersMetadata::GetStructList()); StructIt; StructIt.Next())
 	{
 #if WITH_EDITOR
@@ -466,6 +478,8 @@ void FShaderParametersMetadata::InitializeAllUniformBufferStructs()
 		{
 			StructIt->InitializeUniformBufferDeclaration();
 		}
+
+		StringStructMap.Add(StructIt->GetShaderVariableName(), *StructIt);
 #endif // WITH_EDITOR
 
 		if (!StructIt->IsLayoutInitialized())
diff --git a/Engine/Source/Runtime/RenderCore/Private/ShaderParameters.cpp b/Engine/Source/Runtime/RenderCore/Private/ShaderParameters.cpp
--- a/Engine/Source/Runtime/RenderCore/Private/ShaderParameters.cpp
+++ b/Engine/Source/Runtime/RenderCore/Private/ShaderParameters.cpp
@@ -360,8 +360,14 @@ FString UE::ShaderParameters::CreateUniformBufferShaderDeclaration(const TCHAR*
 	return CreateHLSLUniformBufferDeclaration(Name, UniformBufferStruct);
 }
 
-static FShaderParametersMetadata* FindShaderParametersMetadataWithVariableName(const TCHAR* InVariableName)
+static FShaderParametersMetadata* FindShaderParametersMetadataWithVariableName(uint32 InVariableNameHash, FStringView InVariableNameView)
 {
+#if WITH_EDITOR
+	TMap<FString, FShaderParametersMetadata*>& StringStructMap = FShaderParametersMetadata::GetStringStructMap();
+
+	FShaderParametersMetadata** FoundMetadata = StringStructMap.FindByHash(InVariableNameHash, InVariableNameView);
+	return FoundMetadata ? *FoundMetadata : nullptr;
+#else // WITH_EDITOR
 	for (FShaderParametersMetadata* Metadata : *FShaderParametersMetadata::GetStructList())
 	{
 		if (InVariableName == Metadata->GetShaderVariableName())
@@ -369,6 +375,7 @@ static FShaderParametersMetadata* FindShaderParametersMetadataWithVariableName(c
 			return Metadata;
 		}
 	}
+#endif
 	return nullptr;
 }
 
@@ -381,7 +388,10 @@ void CacheUniformBufferIncludes(TMap<const TCHAR*, FCachedUniformBufferDeclarati
 		FCachedUniformBufferDeclaration& BufferDeclaration = It.Value();
 		check(BufferDeclaration.Declaration.Get() == NULL);
 
-		if (const FShaderParametersMetadata* Metadata = FindShaderParametersMetadataWithVariableName(UniformBufferName))
+		FStringView UniformBufferNameView(UniformBufferName);
+		uint32 UniformBufferNameHash = GetTypeHash(UniformBufferNameView);
+
+		if (const FShaderParametersMetadata* Metadata = FindShaderParametersMetadataWithVariableName(UniformBufferNameHash, UniformBufferNameView))
 		{
 			FString* NewDeclaration = new FString(UE::ShaderParameters::CreateUniformBufferShaderDeclaration(UniformBufferName, *Metadata));
 			check(!NewDeclaration->IsEmpty());
@@ -399,9 +409,10 @@ void UE::ShaderParameters::AddUniformBufferIncludesToEnvironment(FShaderCompiler
 	for (const TCHAR* UniformBufferName : InUniformBufferNames)
 	{
 		FStringView UniformBufferNameView(UniformBufferName);
-		if (!OutEnvironment.UniformBufferMap.FindByHash(FCrc::Strihash_DEPRECATED(UniformBufferName), UniformBufferName))
+		uint32 UniformBufferNameHash = GetTypeHash(UniformBufferNameView);
+		if (!OutEnvironment.UniformBufferMap.FindByHash(UniformBufferNameHash, UniformBufferNameView))
 		{
-			if (const FShaderParametersMetadata* Metadata = FindShaderParametersMetadataWithVariableName(UniformBufferName))
+			if (const FShaderParametersMetadata* Metadata = FindShaderParametersMetadataWithVariableName(UniformBufferNameHash, UniformBufferNameView))
 			{
 				const FThreadSafeSharedStringPtr UniformBufferDeclaration = Metadata->GetUniformBufferDeclarationPtr();
 
diff --git a/Engine/Source/Runtime/RenderCore/Public/ShaderCore.h b/Engine/Source/Runtime/RenderCore/Public/ShaderCore.h
--- a/Engine/Source/Runtime/RenderCore/Public/ShaderCore.h
+++ b/Engine/Source/Runtime/RenderCore/Public/ShaderCore.h
@@ -297,83 +297,229 @@ public:
 	TMap<FString,FParameterAllocation> ParameterMap;
 };
 
+/** Utility class for caching FName and other info for a define.  Don't use directly, use SET_SHADER_DEFINE or SET_SHADER_DEFINE_AND_COMPILE_ARGUMENT macros. */
+class FShaderCompilerDefineNameCache
+{
+public:
+	FShaderCompilerDefineNameCache(const TCHAR* InName)
+		: Name(InName), MapIndex(INDEX_NONE)
+	{}
+
+	operator FName() const
+	{
+		return Name;
+	}
+
+private:
+	FName Name;
+	int32 MapIndex;
+
+	friend class FShaderCompilerDefinitions;
+};
+
+enum class EShaderCompilerDefineVariant : uint8
+{
+	None = 0,
+	Integer,
+	Unsigned,
+	Float,
+	String
+};
+
 /** Container for shader compiler definitions. */
 class FShaderCompilerDefinitions
 {
 public:
+	RENDERCORE_API FShaderCompilerDefinitions(bool bIncludeInitialDefines = false);
+
+	/** Value types supported:   bool, int32, uint32, float, const TCHAR*, FString& */
+	template <typename ValueType>
+	void SetDefine(FName Name, ValueType Value)									{ InternalSetValue(FindOrAddMapIndex(Name), Value); }
+	void SetDefine(FName Name, const TCHAR* Value)								{ InternalSetValue(FindOrAddMapIndex(Name), Value); }
+	void SetDefine(FName Name, const FString& Value)							{ InternalSetValue(FindOrAddMapIndex(Name), *Value); }
+
+	template <typename ValueType>
+	void SetDefine(FShaderCompilerDefineNameCache& Name, ValueType Value)		{ InternalSetValue(FindOrAddMapIndex(Name), Value); }
+	void SetDefine(FShaderCompilerDefineNameCache& Name, const TCHAR* Value)	{ InternalSetValue(FindOrAddMapIndex(Name), Value); }
+	void SetDefine(FShaderCompilerDefineNameCache& Name, const FString& Value)	{ InternalSetValue(FindOrAddMapIndex(Name), *Value); }
 
-	FShaderCompilerDefinitions()
+	RENDERCORE_API int32 GetIntegerValue(FName Name) const;
+	RENDERCORE_API int32 GetIntegerValue(FShaderCompilerDefineNameCache& NameCache, int32 ResultIfNotFound) const;
+
+	bool Contains(FName Name) const
 	{
-		// Presize to reduce re-hashing while building shader jobs
-		Definitions.Empty(50);
+		int32 KeyIndex = FindMapIndex(Name, GetTypeHash(Name));
+		return KeyIndex != INDEX_NONE && ValueTypes[KeyIndex] != EShaderCompilerDefineVariant::None;
 	}
 
-	/**
-	 * Works for TCHAR
-	 * e.g. SetDefine(TEXT("NUM_SAMPLES"), TEXT("1"));
-	 */
-	void SetDefine(const TCHAR* Name, const TCHAR* Value)
+	FORCEINLINE int32 Num() const
 	{
-		Definitions.Add(Name, Value);
+		return ValueCount;
 	}
 
-	void SetDefine(const TCHAR* Name, const FString& Value)
+	RENDERCORE_API void Empty();
+	RENDERCORE_API void Merge(const FShaderCompilerDefinitions& Other);
+
+	RENDERCORE_API FShaderCompilerDefinitions& operator=(const FShaderCompilerDefinitions& Other);
+
+	RENDERCORE_API friend FArchive& operator<<(FArchive& Ar,FShaderCompilerDefinitions& Defs);
+
+	class FConstIterator
 	{
-		Definitions.Add(Name, Value);
+	public:
+		RENDERCORE_API FConstIterator(const FShaderCompilerDefinitions& InDefines);
+
+		/** conversion to "bool" returning true if the iterator is valid. */
+		FORCEINLINE explicit operator bool() const
+		{
+			return Index < Defines.Pairs.Num();
+		}
+		/** inverse of the "bool" operator */
+		FORCEINLINE bool operator !() const
+		{
+			return !(bool)*this;
+		}
+
+		FORCEINLINE FConstIterator& operator++()
+		{
+			Index++;
+
+			// Skip over None values in the map
+			int32 PairNum = Defines.Pairs.Num();
+			while (Index < PairNum && Defines.ValueTypes[Index] == EShaderCompilerDefineVariant::None)
+			{
+				Index++;
+			}
+			return *this;
+		}
+
+		// Note that the output of Key() is transient, only valid until the iterator is incremented!
+		FORCEINLINE const TCHAR* Key()
+		{
+			Defines.Pairs[Index].Key.ToString(KeyStringBuffer);
+			return KeyStringBuffer;
+		}
+
+		// Note that the output of Value() is transient, only valid until the iterator is incremented!
+		RENDERCORE_API const TCHAR* Value();
+
+		FORCEINLINE int32 GetIndex() const
+		{
+			return Index;
+		}
+
+	private:
+		const FShaderCompilerDefinitions& Defines;
+		int32 Index;
+		TCHAR KeyStringBuffer[FName::StringBufferSize];
+		TCHAR ValueStringBuffer[32];
+	};
+
+
+private:
+	RENDERCORE_API int32 FindMapIndex(FName Key, uint32 KeyHash) const;
+	RENDERCORE_API int32 FindOrAddMapIndex(FName Name);
+	RENDERCORE_API int32 FindOrAddMapIndex(FShaderCompilerDefineNameCache& NameCache);
+
+	RENDERCORE_API void InternalSetValue(int32 Index, const TCHAR* Value);
+
+	void InternalSetValue(int32 Index, bool Value)
+	{
+		SetValueType(Index, EShaderCompilerDefineVariant::Integer);
+		Pairs[Index].ValueInteger = Value ? 1 : 0;
 	}
 
-	void SetDefine(const TCHAR* Name, bool Value)
+	void InternalSetValue(int32 Index, int32 Value)
 	{
-		Definitions.Add(Name, Value ? TEXT("1") : TEXT("0"));
+		SetValueType(Index, EShaderCompilerDefineVariant::Integer);
+		Pairs[Index].ValueInteger = Value;
 	}
 
-	void SetDefine(const TCHAR* Name, uint32 Value)
+	void InternalSetValue(int32 Index, uint32 Value)
 	{
-		// can be optimized
-		switch (Value)
-		{
-		// Avoid Printf for common cases
-		case 0u: Definitions.Add(Name, TEXT("0")); break;
-		case 1u: Definitions.Add(Name, TEXT("1")); break;
-		default: Definitions.Add(Name, FString::Printf(TEXT("%u"), Value)); break;
-		}
+		SetValueType(Index, EShaderCompilerDefineVariant::Unsigned);
+		Pairs[Index].ValueUnsigned = Value;
 	}
 
-	void SetDefine(const TCHAR* Name, int32 Value)
+	void InternalSetValue(int32 Index, float Value)
 	{
-		// can be optimized
-		switch (Value)
+		SetValueType(Index, EShaderCompilerDefineVariant::Float);
+		Pairs[Index].ValueFloat = Value;
+	}
+
+	FORCEINLINE void SetValueType(int32 Index, EShaderCompilerDefineVariant InValueType)
+	{
+		EShaderCompilerDefineVariant& ValueType = ValueTypes[Index];
+		if (ValueType == EShaderCompilerDefineVariant::None)
 		{
-		case 0: Definitions.Add(Name, TEXT("0")); break;
-		case 1: Definitions.Add(Name, TEXT("1")); break;
-		default: Definitions.Add(Name, FString::Printf(TEXT("%d"), Value));
+			ValueCount++;
 		}
+		ValueType = InValueType;
 	}
 
-	/**
-	 * Works for float
-	 */
-	RENDERCORE_API void SetFloatDefine(const TCHAR* Name, float Value);
+	/** Called from FShaderInitialDefinesInitializer */
+	RENDERCORE_API static void InitializeInitialDefines(const FShaderCompilerDefinitions& InDefines);
 
-	const TMap<FString,FString>& GetDefinitionMap() const
+	/** Needed for the deprecated FShaderCompilerEnvironment::GetDefinitions, do not use! */
+	const TMap<FString, FString>& GetDefinitionMap() const
 	{
-		return Definitions;
+		return UnusedStringDefinitions;
 	}
+	friend struct FShaderCompilerEnvironment;
 
-	friend FArchive& operator<<(FArchive& Ar,FShaderCompilerDefinitions& Defs)
+	struct FPairType
 	{
-		return Ar << Defs.Definitions;
-	}
+		FName Key;
+		union
+		{
+			int32 ValueInteger;
+			uint32 ValueUnsigned;
+			float ValueFloat;
+		};
 
-	void Merge(const FShaderCompilerDefinitions& Other)
-	{
-		Definitions.Append(Other.Definitions);
-	}
+		FORCEINLINE friend FArchive& operator<<(FArchive& Ar, FPairType& Pair)
+		{
+			// Avoid dynamic allocation when serializing keys.  We need to serialize the FName as a string, because the index in the FName is
+			// non-deterministic (serialization generates hash values that need to be deterministic) and process specific (data needs to be
+			// serialized to worker processes).  Compiler defines must be ANSI text, so serialize as such to save memory.
+			TCHAR KeyBuffer[FName::StringBufferSize];
+			ANSICHAR KeyBufferAnsi[FName::StringBufferSize];
+			uint32 StringLength;
+
+			if (Ar.IsLoading())
+			{
+				Ar << StringLength;
+				check(StringLength < FName::StringBufferSize);
+				Ar.Serialize(KeyBufferAnsi, StringLength + 1);
+				Pair.Key = FName(KeyBufferAnsi);
+			}
+			else
+			{
+				StringLength = Pair.Key.ToString(KeyBuffer, FName::StringBufferSize);
+				for (uint32 CharIndex = 0; CharIndex < StringLength + 1; CharIndex++)
+				{
+					KeyBufferAnsi[CharIndex] = (ANSICHAR)KeyBuffer[CharIndex];
+				}
+				Ar << StringLength;
+				Ar.Serialize(KeyBufferAnsi, StringLength + 1);
+			}
+			Ar << Pair.ValueInteger;
+			return Ar;
+		}
+	};
+	FHashTable KeyHashTable;
+	TArray<FPairType> Pairs;
+	TArray<EShaderCompilerDefineVariant> ValueTypes;
+	TArray<FString> StringValues;
+	uint32 InitialDefineCount;			// Number of items that came from GInitialDefines
+	uint32 ValueCount;					// Number of valid values (ValueType != EShaderCompilerDefineVariant::None)
 
-private:
+	/** Unused data kept around as something to return from GetDefinitionMap() call, which is used only in deprecated FShaderCompilerEnvironment::GetDefinitions call */
+	TMap<FString,FString> UnusedStringDefinitions;
 
-	/** Map: definition -> value. */
-	TMap<FString,FString> Definitions;
+	RENDERCORE_API static FShaderCompilerDefinitions* GInitialDefines;
+
+	friend struct FShaderInitialDefinesInitializer;
 };
 
 struct FShaderResourceTable
@@ -550,6 +696,7 @@ struct FShaderCompilerEnvironment
 
 	/** Default constructor. */
 	FShaderCompilerEnvironment()
+		: Definitions(true)			/* Enable initial defines -- helps performance here, but not necessary for standalone "added" defines declared elsewhere */
 	{
 		// Presize to reduce re-hashing while building shader jobs
 		IncludeVirtualPathToContentsMap.Empty(15);
@@ -581,19 +728,48 @@ struct FShaderCompilerEnvironment
 	 * e.g. SetDefine(TEXT("NAME"), TEXT("Test"));
 	 * e.g. SetDefine(TEXT("NUM_SAMPLES"), 1);
 	 * e.g. SetDefine(TEXT("DOIT"), true);
+	 *
+	 * Or use optimized macros, which can cache FName and map lookups to improve performance:
+	 * e.g. SET_SHADER_DEFINE(NAME, TEXT("Test"));
+	 * e.g. SET_SHADER_DEFINE(NUM_SAMPLES, 1);
+	 * e.g. SET_SHADER_DEFINE(DOIT, true);
 	 */
 	void SetDefine(const TCHAR* Name, const TCHAR* Value)	{ Definitions.SetDefine(Name, Value); }
 	void SetDefine(const TCHAR* Name, const FString& Value) { Definitions.SetDefine(Name, Value); }
 	void SetDefine(const TCHAR* Name, uint32 Value)			{ Definitions.SetDefine(Name, Value); }
 	void SetDefine(const TCHAR* Name, int32 Value)			{ Definitions.SetDefine(Name, Value); }
 	void SetDefine(const TCHAR* Name, bool Value)			{ Definitions.SetDefine(Name, Value); }
-	void SetDefine(const TCHAR* Name, float Value)			{ Definitions.SetFloatDefine(Name, Value); }
+	void SetDefine(const TCHAR* Name, float Value)			{ Definitions.SetDefine(Name, Value); }
+
+	void SetDefine(FName Name, const TCHAR* Value)			{ Definitions.SetDefine(Name, Value); }
+	void SetDefine(FName Name, const FString& Value)		{ Definitions.SetDefine(Name, Value); }
+	void SetDefine(FName Name, uint32 Value)				{ Definitions.SetDefine(Name, Value); }
+	void SetDefine(FName Name, int32 Value)					{ Definitions.SetDefine(Name, Value); }
+	void SetDefine(FName Name, bool Value)					{ Definitions.SetDefine(Name, Value); }
+	void SetDefine(FName Name, float Value)					{ Definitions.SetDefine(Name, Value); }
+
+	void SetDefine(FShaderCompilerDefineNameCache& Name, const TCHAR* Value)	{ Definitions.SetDefine(Name, Value); }
+	void SetDefine(FShaderCompilerDefineNameCache& Name, const FString& Value)	{ Definitions.SetDefine(Name, Value); }
+	void SetDefine(FShaderCompilerDefineNameCache& Name, uint32 Value)			{ Definitions.SetDefine(Name, Value); }
+	void SetDefine(FShaderCompilerDefineNameCache& Name, int32 Value)			{ Definitions.SetDefine(Name, Value); }
+	void SetDefine(FShaderCompilerDefineNameCache& Name, bool Value)			{ Definitions.SetDefine(Name, Value); }
+	void SetDefine(FShaderCompilerDefineNameCache& Name, float Value)			{ Definitions.SetDefine(Name, Value); }
+
+	FORCEINLINE int32 GetIntegerValue(FName Name) const
+	{
+		return Definitions.GetIntegerValue(Name);
+	}
+	FORCEINLINE int32 GetIntegerValue(FShaderCompilerDefineNameCache& NameCache, int32 ResultIfNotFound = 0) const
+	{
+		return Definitions.GetIntegerValue(NameCache, ResultIfNotFound);
+	}
 
 	template <typename ValueType> void SetDefineIfUnset(const TCHAR* Name, ValueType Value)
 	{
-		if (!Definitions.GetDefinitionMap().Contains(Name))
+		FName NameKey(Name);
+		if (!Definitions.Contains(NameKey))
 		{
-			SetDefine(Name, Value);
+			Definitions.SetDefine(NameKey, Value);
 		}
 	}
 
@@ -739,6 +915,22 @@ private:
 	TMap<FString, TVariant<bool, float, int32, uint32, FString>> CompileArgs;
 };
 
+
+/** Optimized define setting macros that cache the FName lookup, and potentially the map index. */
+#define SET_SHADER_DEFINE(ENVIRONMENT, NAME, VALUE) \
+	do {																	\
+		static FShaderCompilerDefineNameCache Cache_##NAME(TEXT(#NAME));	\
+		(ENVIRONMENT).SetDefine(Cache_##NAME, VALUE);						\
+	} while(0)
+
+#define SET_SHADER_DEFINE_AND_COMPILE_ARGUMENT(ENVIRONMENT, NAME, VALUE) \
+	do {																	\
+		static FShaderCompilerDefineNameCache Cache_##NAME(TEXT(#NAME));	\
+		(ENVIRONMENT).SetDefine(Cache_##NAME, VALUE);						\
+		(ENVIRONMENT).SetCompileArgument(TEXT(#NAME), VALUE);				\
+	} while(0)
+
+
 struct FSharedShaderCompilerEnvironment final : public FShaderCompilerEnvironment, public FRefCountBase
 {
 	virtual ~FSharedShaderCompilerEnvironment() = default;
diff --git a/Engine/Source/Runtime/RenderCore/Public/ShaderParameterMetadata.h b/Engine/Source/Runtime/RenderCore/Public/ShaderParameterMetadata.h
--- a/Engine/Source/Runtime/RenderCore/Public/ShaderParameterMetadata.h
+++ b/Engine/Source/Runtime/RenderCore/Public/ShaderParameterMetadata.h
@@ -372,6 +372,9 @@ public:
 	static RENDERCORE_API TLinkedList<FShaderParametersMetadata*>*& GetStructList();
 	/** Speed up finding the uniform buffer by its name */
 	static RENDERCORE_API TMap<FHashedName, FShaderParametersMetadata*>& GetNameStructMap();
+#if WITH_EDITOR
+	static RENDERCORE_API TMap<FString, FShaderParametersMetadata*>& GetStringStructMap();
+#endif // WITH_EDITOR
 
 	/** Initialize all the global shader parameter structs. */
 	static RENDERCORE_API void InitializeAllUniformBufferStructs();
