Skeletal Mesh: Just ignore the parts of morph targets that exceed the LODs or sections of the mesh they sit on.

#rnx

diff --git a/Engine/Source/Runtime/Engine/Private/Components/SkinnedMeshComponent.cpp b/Engine/Source/Runtime/Engine/Private/Components/SkinnedMeshComponent.cpp
--- a/Engine/Source/Runtime/Engine/Private/Components/SkinnedMeshComponent.cpp
+++ b/Engine/Source/Runtime/Engine/Private/Components/SkinnedMeshComponent.cpp
@@ -484,18 +484,15 @@ void USkinnedMeshComponent::UpdateMorphMaterialUsageOnProxy()
 			}
 
 			const TArray<FMorphTargetLODModel>& MorphLODModels = MorphTarget->GetMorphLODModels();
-			if (ensure(MorphLODModels.Num() == LODRenderData.Num()))
+			for (int32 LODIndex = 0; LODIndex < FMath::Min(LODRenderData.Num(), MorphLODModels.Num()); LODIndex++)
 			{
-				for (int32 LODIndex = 0; LODIndex < MorphLODModels.Num(); LODIndex++)
+				const FSkeletalMeshLODRenderData& RenderData = LODRenderData[LODIndex];
+				
+				for(const int32 MorphSectionIndex: MorphLODModels[LODIndex].SectionIndices)
 				{
-					const FSkeletalMeshLODRenderData& RenderData = LODRenderData[LODIndex];
-					
-					for(const int32 MorphSectionIndex: MorphLODModels[LODIndex].SectionIndices)
+					if (RenderData.RenderSections.IsValidIndex(MorphSectionIndex))
 					{
-						if (ensure(RenderData.RenderSections.IsValidIndex(MorphSectionIndex)))
-						{
-							UsedMaterialIndices.Add(RenderData.RenderSections[MorphSectionIndex].MaterialIndex);
-						}
+						UsedMaterialIndices.Add(RenderData.RenderSections[MorphSectionIndex].MaterialIndex);
 					}
 				}
 			}
