Fix crash in BuildSignedDistanceFieldBuildSectionData(...)

diff --git a/Engine/Source/Runtime/Engine/Private/DistanceFieldAtlas.cpp b/Engine/Source/Runtime/Engine/Private/DistanceFieldAtlas.cpp
--- a/Engine/Source/Runtime/Engine/Private/DistanceFieldAtlas.cpp
+++ b/Engine/Source/Runtime/Engine/Private/DistanceFieldAtlas.cpp
@@ -158,14 +158,15 @@ void BuildSignedDistanceFieldBuildSectionData(UStaticMesh* Mesh, uint32 LODIndex
 	{
 		const FMeshSectionInfo& Section = SectionInfoMap.Get(LODIndex, SectionIndex);
 
-		if (!OutData.IsValidIndex(Section.MaterialIndex))
+		FSignedDistanceFieldBuildSectionData& SectionData = OutData[SectionIndex];
+		SectionData.bAffectDistanceFieldLighting = Section.bAffectDistanceFieldLighting;
+
+		if (!StaticMaterials.IsValidIndex(Section.MaterialIndex))
 		{
+			// TODO: Should maybe log warning here
 			continue;
 		}
 
-		FSignedDistanceFieldBuildSectionData& SectionData = OutData[SectionIndex];
-		SectionData.bAffectDistanceFieldLighting = Section.bAffectDistanceFieldLighting;
-
 		UMaterialInterface* MaterialInterface = StaticMaterials[Section.MaterialIndex].MaterialInterface;
 		if (MaterialInterface)
 		{
@@ -186,12 +187,12 @@ void FDistanceFieldVolumeData::CacheDerivedData(const FString& InStaticMeshDeriv
 
 	FString DistanceFieldKey = BuildDistanceFieldDerivedDataKey(InStaticMeshDerivedDataKey);
 
-	for (int32 MaterialIndex = 0; MaterialIndex < BuildSectionData.Num(); MaterialIndex++)
+	for (int32 SectionIndex = 0; SectionIndex < BuildSectionData.Num(); SectionIndex++)
 	{
 		DistanceFieldKey += FString::Printf(TEXT("_M%u_%u_%u"), 
-			(uint32)BuildSectionData[MaterialIndex].BlendMode,
-			BuildSectionData[MaterialIndex].bTwoSided ? 1 : 0,
-			BuildSectionData[MaterialIndex].bAffectDistanceFieldLighting ? 1 : 0);
+			(uint32)BuildSectionData[SectionIndex].BlendMode,
+			BuildSectionData[SectionIndex].bTwoSided ? 1 : 0,
+			BuildSectionData[SectionIndex].bAffectDistanceFieldLighting ? 1 : 0);
 	}
 
 	TArray<uint8> DerivedData;
