World Partition Streaming Generation:
- Do actor filter validation in streaming generation phase, not during context creation.
- Don't create the streaming generation context when validating, since no error detection should be done in it.

diff --git a/Engine/Plugins/Editor/DataValidation/Source/DataValidation/Private/WorldPartitionChangelistValidator.cpp b/Engine/Plugins/Editor/DataValidation/Source/DataValidation/Private/WorldPartitionChangelistValidator.cpp
--- a/Engine/Plugins/Editor/DataValidation/Source/DataValidation/Private/WorldPartitionChangelistValidator.cpp
+++ b/Engine/Plugins/Editor/DataValidation/Source/DataValidation/Private/WorldPartitionChangelistValidator.cpp
@@ -441,7 +441,7 @@ void UWorldPartitionChangelistValidator::OnLevelInstanceInvalidWorldAsset(const
 	}
 }
 
-void UWorldPartitionChangelistValidator::OnActorFilterFailed(const FWorldPartitionActorDescView& ActorDescView)
+void UWorldPartitionChangelistValidator::OnInvalidActorFilterReference(const FWorldPartitionActorDescView& ActorDescView, const FWorldPartitionActorDescView& ReferenceActorDescView)
 {
 	// Not a validation error
 }
diff --git a/Engine/Plugins/Editor/DataValidation/Source/DataValidation/Private/WorldPartitionChangelistValidator.h b/Engine/Plugins/Editor/DataValidation/Source/DataValidation/Private/WorldPartitionChangelistValidator.h
--- a/Engine/Plugins/Editor/DataValidation/Source/DataValidation/Private/WorldPartitionChangelistValidator.h
+++ b/Engine/Plugins/Editor/DataValidation/Source/DataValidation/Private/WorldPartitionChangelistValidator.h
@@ -46,5 +46,5 @@ protected:
 	virtual void OnDataLayerAssetConflict(const UDataLayerInstanceWithAsset* DataLayerInstance, const UDataLayerInstanceWithAsset* ConflictingDataLayerInstance) override;
 	virtual void OnActorNeedsResave(const FWorldPartitionActorDescView& ActorDescView) override;
 	virtual void OnLevelInstanceInvalidWorldAsset(const FWorldPartitionActorDescView& ActorDescView, FName WorldAsset, ELevelInstanceInvalidReason Reason) override;
-	virtual void OnActorFilterFailed(const FWorldPartitionActorDescView& ActorDescView) override;
+	virtual void OnInvalidActorFilterReference(const FWorldPartitionActorDescView& ActorDescView, const FWorldPartitionActorDescView& ReferenceActorDescView) override;
 };
diff --git a/Engine/Source/Runtime/Engine/Private/WorldPartition/ErrorHandling/WorldPartitionStreamingGenerationLogErrorHandler.cpp b/Engine/Source/Runtime/Engine/Private/WorldPartition/ErrorHandling/WorldPartitionStreamingGenerationLogErrorHandler.cpp
--- a/Engine/Source/Runtime/Engine/Private/WorldPartition/ErrorHandling/WorldPartitionStreamingGenerationLogErrorHandler.cpp
+++ b/Engine/Source/Runtime/Engine/Private/WorldPartition/ErrorHandling/WorldPartitionStreamingGenerationLogErrorHandler.cpp
@@ -91,9 +91,9 @@ void FStreamingGenerationLogErrorHandler::OnLevelInstanceInvalidWorldAsset(const
 	};
 }
 
-void FStreamingGenerationLogErrorHandler::OnActorFilterFailed(const FWorldPartitionActorDescView& ActorDescView)
+void FStreamingGenerationLogErrorHandler::OnInvalidActorFilterReference(const FWorldPartitionActorDescView& ActorDescView, const FWorldPartitionActorDescView& ReferenceActorDescView)
 {
-	UE_ASSET_LOG_ACTORDESCVIEW(LogWorldPartition, Log, ActorDescView, TEXT("Actor %s will not be filtered out because it is referenced/referencing other Actor(s) not part of the filter"), *GetActorName(ActorDescView));
+	UE_ASSET_LOG_ACTORDESCVIEW(LogWorldPartition, Log, ReferenceActorDescView, TEXT("Actor %s will not be filtered out because it is referenced by Actor %s not part of the filter"), *GetActorName(ReferenceActorDescView), *GetActorName(ActorDescView));
 }
 
 #endif
diff --git a/Engine/Source/Runtime/Engine/Private/WorldPartition/ErrorHandling/WorldPartitionStreamingGenerationProxyErrorHandler.cpp b/Engine/Source/Runtime/Engine/Private/WorldPartition/ErrorHandling/WorldPartitionStreamingGenerationProxyErrorHandler.cpp
--- a/Engine/Source/Runtime/Engine/Private/WorldPartition/ErrorHandling/WorldPartitionStreamingGenerationProxyErrorHandler.cpp
+++ b/Engine/Source/Runtime/Engine/Private/WorldPartition/ErrorHandling/WorldPartitionStreamingGenerationProxyErrorHandler.cpp
@@ -63,8 +63,8 @@ void FStreamingGenerationProxyErrorHandler::OnLevelInstanceInvalidWorldAsset(con
 	InnerErrorHandler->OnLevelInstanceInvalidWorldAsset(ActorDescView, WorldAsset, Reason);
 }
 
-void FStreamingGenerationProxyErrorHandler::OnActorFilterFailed(const FWorldPartitionActorDescView& ActorDescView)
+void FStreamingGenerationProxyErrorHandler::OnInvalidActorFilterReference(const FWorldPartitionActorDescView& ActorDescView, const FWorldPartitionActorDescView& ReferenceActorDescView)
 {
-	InnerErrorHandler->OnActorFilterFailed(ActorDescView);
+	InnerErrorHandler->OnInvalidActorFilterReference(ActorDescView, ReferenceActorDescView);
 }
 #endif
diff --git a/Engine/Source/Runtime/Engine/Private/WorldPartition/ErrorHandling/WorldPartitionStreamingGenerationTokenizedMessageErrorHandler.cpp b/Engine/Source/Runtime/Engine/Private/WorldPartition/ErrorHandling/WorldPartitionStreamingGenerationTokenizedMessageErrorHandler.cpp
--- a/Engine/Source/Runtime/Engine/Private/WorldPartition/ErrorHandling/WorldPartitionStreamingGenerationTokenizedMessageErrorHandler.cpp
+++ b/Engine/Source/Runtime/Engine/Private/WorldPartition/ErrorHandling/WorldPartitionStreamingGenerationTokenizedMessageErrorHandler.cpp
@@ -198,12 +198,14 @@ void ITokenizedMessageErrorHandler::OnLevelInstanceInvalidWorldAsset(const FWorl
 	HandleTokenizedMessage(MoveTemp(Message));
 }
 
-void ITokenizedMessageErrorHandler::OnActorFilterFailed(const FWorldPartitionActorDescView& ActorDescView)
+void ITokenizedMessageErrorHandler::OnInvalidActorFilterReference(const FWorldPartitionActorDescView& ActorDescView, const FWorldPartitionActorDescView& ReferenceActorDescView)
 {
-	TSharedRef<FTokenizedMessage> Message = FTokenizedMessage::Create(EMessageSeverity::Warning);
-	Message->AddToken(FTextToken::Create(LOCTEXT("TokenMessage_WorldPartition_ActorFilterFailed", "Actor will not be filtered out because it is referenced/referencing other Actor(s) not part of the filter")))
+	TSharedRef<FTokenizedMessage> Message = FTokenizedMessage::Create(EMessageSeverity::Error);
+	Message->AddToken(FTextToken::Create(LOCTEXT("TokenMessage_WorldPartition_Actor", "Actor")))
+		->AddToken(FActorToken::Create(ReferenceActorDescView.GetActorSoftPath().ToString(), ReferenceActorDescView.GetGuid(), FText::FromString(GetActorName(ReferenceActorDescView))))
+		->AddToken(FTextToken::Create(LOCTEXT("TokenMessage_WorldPartition_ActorFilterFailed", "will not be filtered out because it is referenced by actor")))
 		->AddToken(FActorToken::Create(ActorDescView.GetActorSoftPath().ToString(), ActorDescView.GetGuid(), FText::FromString(GetActorName(ActorDescView))))
-		->AddToken(FMapErrorToken::Create(TEXT("WorldPartition_ActorFilterFailed_CheckForErrors")));
+		->AddToken(FMapErrorToken::Create(TEXT("WorldPartition_MissingActorReference_CheckForErrors")));
 
 	HandleTokenizedMessage(MoveTemp(Message));
 }
diff --git a/Engine/Source/Runtime/Engine/Private/WorldPartition/RuntimeSpatialHash/RuntimeSpatialHashGridHelper.cpp b/Engine/Source/Runtime/Engine/Private/WorldPartition/RuntimeSpatialHash/RuntimeSpatialHashGridHelper.cpp
--- a/Engine/Source/Runtime/Engine/Private/WorldPartition/RuntimeSpatialHash/RuntimeSpatialHashGridHelper.cpp
+++ b/Engine/Source/Runtime/Engine/Private/WorldPartition/RuntimeSpatialHash/RuntimeSpatialHashGridHelper.cpp
@@ -178,16 +178,18 @@ FSquare2DGridHelper GetPartitionedActors(const FBox& WorldBounds, const FSpatial
 		if (GRuntimeSpatialHashPlacePartitionActorsUsingLocationEffective)
 		{
 			bool bUseLocation = true;
-			for (const FGuid& ActorGuid : InActorSetInstance->ActorSet->Actors)
+			InActorSetInstance->ForEachActor([InActorSetInstance, &bUseLocation](const FGuid& ActorGuid)
 			{
 				const FWorldPartitionActorDescView& ActorDescView = InActorSetInstance->ContainerInstance->ActorDescViewMap->FindByGuidChecked(ActorGuid);
 
 				if (!ActorDescView.GetActorNativeClass()->IsChildOf<APartitionActor>())
 				{
 					bUseLocation = false;
-					break;
+					return false;
 				}
-			}
+
+				return true;
+			});
 			if (bUseLocation)
 			{
 				// Find grid level that best matches actor set bounds
@@ -203,6 +205,8 @@ FSquare2DGridHelper GetPartitionedActors(const FBox& WorldBounds, const FSpatial
 	for (const IStreamingGenerationContext::FActorSetInstance* ActorSetInstance : ActorSetInstances)
 	{
 		check(ActorSetInstance->ActorSet->Actors.Num() > 0);
+		check(!ActorSetInstance->FilteredActors || (ActorSetInstance->ActorSet->Actors.Num() != ActorSetInstance->FilteredActors->Num()));
+
 		FSquare2DGridHelper::FGridLevel::FGridCell* GridCell = nullptr;
 
 		if (ActorSetInstance->bIsSpatiallyLoaded)
diff --git a/Engine/Source/Runtime/Engine/Private/WorldPartition/RuntimeSpatialHash/RuntimeSpatialHashHLOD.cpp b/Engine/Source/Runtime/Engine/Private/WorldPartition/RuntimeSpatialHash/RuntimeSpatialHashHLOD.cpp
--- a/Engine/Source/Runtime/Engine/Private/WorldPartition/RuntimeSpatialHash/RuntimeSpatialHashHLOD.cpp
+++ b/Engine/Source/Runtime/Engine/Private/WorldPartition/RuntimeSpatialHash/RuntimeSpatialHashHLOD.cpp
@@ -193,14 +193,14 @@ static TArray<FGuid> GenerateHLODsForGrid(UWorldPartition* WorldPartition, const
 				TArray<IStreamingGenerationContext::FActorInstance> ActorInstances;
 				for (const IStreamingGenerationContext::FActorSetInstance* ActorSetInstance : GridCellDataChunk.GetActorSetInstances())
 				{
-					for (const FGuid& ActorGuid : ActorSetInstance->ActorSet->Actors)
+					ActorSetInstance->ForEachActor([&ActorInstances, ActorSetInstance, &CellBounds](const FGuid& ActorGuid)
 					{
 						const IStreamingGenerationContext::FActorInstance& ActorInstance = ActorInstances.Emplace_GetRef(ActorGuid, ActorSetInstance);
 						const FBox RuntimeBounds = ActorInstance.GetActorDescView().GetRuntimeBounds();
 						check(RuntimeBounds.IsValid);
 						CellBounds.Min.Z = FMath::Min(CellBounds.Min.Z, RuntimeBounds.Min.Z);
 						CellBounds.Max.Z = FMath::Max(CellBounds.Max.Z, RuntimeBounds.Max.Z);
-					}
+					});
 				}
 
 				// Ensure the Z bounds are valid
diff --git a/Engine/Source/Runtime/Engine/Private/WorldPartition/WorldPartitionRuntimeHash.cpp b/Engine/Source/Runtime/Engine/Private/WorldPartition/WorldPartitionRuntimeHash.cpp
--- a/Engine/Source/Runtime/Engine/Private/WorldPartition/WorldPartitionRuntimeHash.cpp
+++ b/Engine/Source/Runtime/Engine/Private/WorldPartition/WorldPartitionRuntimeHash.cpp
@@ -139,7 +139,7 @@ bool UWorldPartitionRuntimeHash::ConditionalRegisterAlwaysLoadedActorsForPIE(con
 {
 	if (bIsMainWorldPartition && bIsMainContainer && bIsCellAlwaysLoaded && !IsRunningCookCommandlet())
 	{
-		for (const FGuid& ActorGuid : ActorSetInstance->ActorSet->Actors)
+		ActorSetInstance->ForEachActor([this, ActorSetInstance](const FGuid& ActorGuid)
 		{
 			IStreamingGenerationContext::FActorInstance ActorInstance(ActorGuid, ActorSetInstance);
 			const FWorldPartitionActorDescView& ActorDescView = ActorInstance.GetActorDescView();
@@ -151,7 +151,7 @@ bool UWorldPartitionRuntimeHash::ConditionalRegisterAlwaysLoadedActorsForPIE(con
 			{
 				AlwaysLoadedActorsForPIE.Emplace(Reference, AlwaysLoadedActor);
 			}
-		}
+		});
 
 		return true;
 	}
@@ -168,7 +168,7 @@ bool UWorldPartitionRuntimeHash::PopulateCellActorInstances(const TArray<const I
 		// In PIE, always loaded actors of an instanced world partition will go in the always loaded cell.
 		if (!ConditionalRegisterAlwaysLoadedActorsForPIE(ActorSetInstance, bIsMainWorldPartition, ActorSetInstance->ContainerID.IsMainContainer(), bIsCellAlwaysLoaded))
 		{
-			for (const FGuid& ActorGuid : ActorSetInstance->ActorSet->Actors)
+			ActorSetInstance->ForEachActor([this, ActorSetInstance, &OutCellActorInstances](const FGuid& ActorGuid)
 			{
 				// Actors that return true to ShouldLevelKeepRefIfExternal will always be part of the partitioned persistent level of a
 				// world partition. In PIE, for an instanced world partition, we don't want this actor to be both in the persistent level
@@ -182,7 +182,7 @@ bool UWorldPartitionRuntimeHash::PopulateCellActorInstances(const TArray<const I
 				{
 					OutCellActorInstances.Emplace(ActorInstance);
 				}
-			}
+			});
 		}
 	}
 
diff --git a/Engine/Source/Runtime/Engine/Private/WorldPartition/WorldPartitionStreamingGeneration.cpp b/Engine/Source/Runtime/Engine/Private/WorldPartition/WorldPartitionStreamingGeneration.cpp
--- a/Engine/Source/Runtime/Engine/Private/WorldPartition/WorldPartitionStreamingGeneration.cpp
+++ b/Engine/Source/Runtime/Engine/Private/WorldPartition/WorldPartitionStreamingGeneration.cpp
@@ -116,12 +116,6 @@ class FWorldPartitionStreamingGenerator
 				}
 			}
 
-			TMap<FActorContainerID, TSet<FGuid>> FilteredOutActorsPerContainer;
-			for (const auto& [ContainerID, ContainerCollectionInstanceDescriptor] : StreamingGenerator->ContainerCollectionInstanceDescriptorsMap)
-			{
-				FilteredOutActorsPerContainer.Append(ContainerCollectionInstanceDescriptor.FilteredActors);
-			}
-
 			ActorSetInstances.Empty();
 			for (const auto& [ContainerID, ContainerCollectionInstanceDescriptor] : StreamingGenerator->ContainerCollectionInstanceDescriptorsMap)
 			{
@@ -133,9 +127,6 @@ class FWorldPartitionStreamingGenerator
 					const FActorSet& ActorSet = *ActorSetPtr;
 					const FWorldPartitionActorDescView& ReferenceActorDescView = ActorSetContainer.ActorDescViewMap->FindByGuidChecked(ActorSet.Actors[0]);
 
-					TSet<FGuid> FilteredActors;
-					TSet<FGuid>* FilteredOutActors = FilteredOutActorsPerContainer.Find(ContainerID);
-
 					// Validate assumptions
 					for (const FGuid& ActorGuid : ActorSet.Actors)
 					{
@@ -143,48 +134,35 @@ class FWorldPartitionStreamingGenerator
 						check(ActorDescView.GetRuntimeGrid() == ReferenceActorDescView.GetRuntimeGrid());
 						check(ActorDescView.GetIsSpatiallyLoaded() == ReferenceActorDescView.GetIsSpatiallyLoaded());
 						check(ActorDescView.GetContentBundleGuid() == ReferenceActorDescView.GetContentBundleGuid());
-
-						if (FilteredOutActors && FilteredOutActors->Contains(ActorGuid))
-						{
-							FilteredActors.Add(ActorGuid);
-						}
 					}
 
-					// We can only filter the complete set because of dependencies between actors of the same set
-					if (FilteredActors.Num() == ActorSet.Actors.Num())
-					{
-						continue;
-					}
-					else if (FilteredActors.Num() > 0)
+					// Skip if all actors are filtered out for this container
+					if (!ContainerCollectionInstanceDescriptor.FilteredActors || (ContainerCollectionInstanceDescriptor.FilteredActors->Num() != ActorSet.Actors.Num()))
 					{
-						// Report failed filtering
-						for (const FGuid& FilteredActorGuid : FilteredActors)
-						{
-							StreamingGenerator->ErrorHandler->OnActorFilterFailed(ActorSetContainer.ActorDescViewMap->FindByGuidChecked(FilteredActorGuid));
-						}
-					}
-
-					FActorSetInstance& ActorSetInstance = ActorSetInstances.Emplace_GetRef();
-					const FContainerCollectionInstanceDescriptor::FPerInstanceData& PerInstanceData = ContainerCollectionInstanceDescriptor.GetInstanceData(ReferenceActorDescView.GetGuid());
+						FActorSetInstance& ActorSetInstance = ActorSetInstances.Emplace_GetRef();
+						const FContainerCollectionInstanceDescriptor::FPerInstanceData& PerInstanceData = ContainerCollectionInstanceDescriptor.GetInstanceData(ReferenceActorDescView.GetGuid());
 				
-					ActorSetInstance.ContainerInstance = &ActorSetContainer;
-					ActorSetInstance.ActorSet = &ActorSet;
-					ActorSetInstance.ContainerID = ContainerCollectionInstanceDescriptor.ID;
-					ActorSetInstance.Transform = ContainerCollectionInstanceDescriptor.Transform;
-					ActorSetInstance.bIsSpatiallyLoaded = PerInstanceData.bIsSpatiallyLoaded;
-					ActorSetInstance.ContentBundleID = ContainerCollectionInstanceDescriptor.ContentBundleID;
-					ActorSetInstance.RuntimeGrid = PerInstanceData.RuntimeGrid;
-					ActorSetInstance.DataLayers = StreamingGenerator->GetRuntimeDataLayerInstances(PerInstanceData.DataLayers);
-
-					ActorSetInstance.Bounds.Init();
-					for (const FGuid& ActorGuid : ActorSetInstance.ActorSet->Actors)
-					{
-						const FWorldPartitionActorDescView& ActorDescView = ActorSetContainer.ActorDescViewMap->FindByGuidChecked(ActorGuid);
-						const FBox RuntimeBounds = ActorDescView.GetRuntimeBounds();
-						if (RuntimeBounds.IsValid)
+						ActorSetInstance.ContainerInstance = &ActorSetContainer;
+						ActorSetInstance.ActorSet = &ActorSet;
+						ActorSetInstance.FilteredActors = ContainerCollectionInstanceDescriptor.FilteredActors;
+						ActorSetInstance.ContainerID = ContainerCollectionInstanceDescriptor.ID;
+						ActorSetInstance.Transform = ContainerCollectionInstanceDescriptor.Transform;
+						ActorSetInstance.bIsSpatiallyLoaded = PerInstanceData.bIsSpatiallyLoaded;
+						ActorSetInstance.ContentBundleID = ContainerCollectionInstanceDescriptor.ContentBundleID;
+						ActorSetInstance.RuntimeGrid = PerInstanceData.RuntimeGrid;
+						ActorSetInstance.DataLayers = StreamingGenerator->GetRuntimeDataLayerInstances(PerInstanceData.DataLayers);
+
+						ActorSetInstance.Bounds.Init();
+						const FTransform& ContainerCollectionInstanceDescriptorTransform = ContainerCollectionInstanceDescriptor.Transform;
+						ActorSetInstance.ForEachActor([this, &ActorSetContainer, &ActorSetInstance, &ContainerCollectionInstanceDescriptorTransform](const FGuid& ActorGuid)
 						{
-							ActorSetInstance.Bounds += RuntimeBounds.TransformBy(ContainerCollectionInstanceDescriptor.Transform);
-						}
+							const FWorldPartitionActorDescView& ActorDescView = ActorSetContainer.ActorDescViewMap->FindByGuidChecked(ActorGuid);
+							const FBox RuntimeBounds = ActorDescView.GetRuntimeBounds();
+							if (RuntimeBounds.IsValid)
+							{
+								ActorSetInstance.Bounds += RuntimeBounds.TransformBy(ContainerCollectionInstanceDescriptorTransform);
+							}
+						});
 					}
 				}
 			}
@@ -258,6 +236,7 @@ class FWorldPartitionStreamingGenerator
 	{
 		FContainerCollectionInstanceDescriptor()
 			: Bounds(ForceInit)
+			, FilteredActors(nullptr)
 		{}
 
 		FBox Bounds;
@@ -268,7 +247,7 @@ class FWorldPartitionStreamingGenerator
 		FActorContainerID ID;
 		FActorContainerID ParentID;
 		FGuid ContentBundleID;
-		TMap<FActorContainerID, TSet<FGuid>> FilteredActors;
+		TSet<FGuid>* FilteredActors;
 
 		struct FPerInstanceData
 		{
@@ -580,8 +559,10 @@ class FWorldPartitionStreamingGenerator
 				{
 					if (const FWorldPartitionActorFilter* ContainerFilter = ContainerCollectionInstanceView.GetContainerFilter())
 					{						
-						SubContainerInstanceDescriptor.FilteredActors = WorldPartitionSubsystem->GetFilteredActorsPerContainer(SubContainerInstanceDescriptor.ID, ContainerCollectionInstanceView.GetContainerPackage().ToString(), *ContainerFilter);
+						ContainerFilteredActors.Append(WorldPartitionSubsystem->GetFilteredActorsPerContainer(SubContainerInstanceDescriptor.ID, ContainerCollectionInstanceView.GetContainerPackage().ToString(), *ContainerFilter));
 					}
+
+					SubContainerInstanceDescriptor.FilteredActors = ContainerFilteredActors.Find(SubContainerInstanceDescriptor.ID);
 				}
 
 				CreateActorDescriptorViewsRecursive(SubContainerInstanceDescriptor);
@@ -979,7 +960,7 @@ class FWorldPartitionStreamingGenerator
 
 			NbErrorsDetected = 0;
 
-			ContainerCollectionDescriptor.ActorDescViewMap.ForEachActorDescView([this, &ContainerCollectionInstanceDescriptor, &NbErrorsDetected, PassType](FWorldPartitionActorDescView& ActorDescView)
+			ContainerCollectionDescriptor.ActorDescViewMap.ForEachActorDescView([this, &ContainerCollectionDescriptor, &ContainerCollectionInstanceDescriptor, &NbErrorsDetected, PassType](FWorldPartitionActorDescView& ActorDescView)
 			{
 				FContainerCollectionInstanceDescriptor::FPerInstanceData& PerInstanceData = ContainerCollectionInstanceDescriptor.GetInstanceData(ActorDescView.GetGuid());
 
@@ -1001,6 +982,29 @@ PRAGMA_ENABLE_DEPRECATION_WARNINGS
 						NbErrorsDetected++;
 					}
 				}
+
+				if (TSet<FGuid>* FilteredActors = ContainerCollectionInstanceDescriptor.FilteredActors)
+				{
+					const bool IsReferencerFiltered = FilteredActors->Contains(ActorDescView.GetGuid());
+					for (const FGuid& ReferenceGuid : ActorDescView.GetReferences())
+					{
+						const bool IsReferenceeFiltered = FilteredActors->Contains(ReferenceGuid);
+						if (IsReferenceeFiltered && !IsReferencerFiltered)
+						{
+							if (PassType == EPassType::ErrorReporting)
+							{
+								const FWorldPartitionActorDescView& ReferenceActorDesc = ContainerCollectionDescriptor.ActorDescViewMap.FindByGuidChecked(ReferenceGuid);
+								ErrorHandler->OnInvalidActorFilterReference(ActorDescView, ReferenceActorDesc);
+							}
+							else
+							{
+								FilteredActors->Remove(ReferenceGuid);
+							}
+
+							NbErrorsDetected++;
+						}
+					}
+				}
 			});
 		}
 	}
@@ -1062,9 +1066,6 @@ public:
 	void PreparationPhase(const FStreamingGenerationActorDescCollection& ActorDescCollection)
 	{
 		CreateActorContainers(ActorDescCollection);
-
-		// Construct the streaming generation context
-		StreamingGenerationContext = MakeUnique<FStreamingGenerationContext>(this, ActorDescCollection);
 	}
 
 	static TUniquePtr<FArchive> CreateDumpStateLogArchive(const TCHAR* Suffix, bool bTimeStamped = true)
@@ -1174,8 +1175,14 @@ public:
 		DumpContainerInstances(FActorContainerID());
 	}
 
-	const FStreamingGenerationContext* GetStreamingGenerationContext()
+	const FStreamingGenerationContext* GetStreamingGenerationContext(const FStreamingGenerationActorDescCollection& ActorDescCollection)
 	{
+		if (!StreamingGenerationContext)
+		{
+			// Construct the streaming generation context
+			StreamingGenerationContext = MakeUnique<FStreamingGenerationContext>(this, ActorDescCollection);
+		}
+
 		return StreamingGenerationContext.Get();
 	}
 
@@ -1208,6 +1215,9 @@ private:
 
 	/** List of current container instances on the stack to detect circular references */
 	TSet<FName> ContainerInstancesStack;
+
+	/** Maps containers IDs to their filtered actors use while creating FContainerCollectionInstanceDescriptor */
+	TMap<FActorContainerID, TSet<FGuid>> ContainerFilteredActors;
 };
 
 bool UWorldPartition::GenerateStreaming(TArray<FString>* OutPackagesToGenerate)
@@ -1304,7 +1314,7 @@ bool UWorldPartition::GenerateContainerStreaming(const FGenerateStreamingParams&
 	StreamingPolicy = NewObject<UWorldPartitionStreamingPolicy>(const_cast<UWorldPartition*>(this), WorldPartitionStreamingPolicyClass.Get(), NAME_None, bIsPIE ? RF_Transient : RF_NoFlags);
 
 	check(RuntimeHash);
-	if (RuntimeHash->GenerateStreaming(StreamingPolicy, StreamingGenerator.GetStreamingGenerationContext(), InContext.PackagesToGenerate))
+	if (RuntimeHash->GenerateStreaming(StreamingPolicy, StreamingGenerator.GetStreamingGenerationContext(InParams.ActorDescCollection), InContext.PackagesToGenerate))
 	{
 		if (HierarchicalLogAr.IsValid())
 		{
@@ -1356,7 +1366,7 @@ void UWorldPartition::GenerateHLOD(ISourceControlHelper* SourceControlHelper, bo
 		FHierarchicalLogArchive HierarchicalLogAr(*LogFileAr);
 		StreamingGenerator.DumpStateLog(HierarchicalLogAr);
 
-		RuntimeHash->GenerateHLOD(SourceControlHelper, StreamingGenerator.GetStreamingGenerationContext(), bCreateActorsOnly);
+		RuntimeHash->GenerateHLOD(SourceControlHelper, StreamingGenerator.GetStreamingGenerationContext(ActorDescCollection), bCreateActorsOnly);
 	});	
 }
 
diff --git a/Engine/Source/Runtime/Engine/Public/WorldPartition/ErrorHandling/WorldPartitionStreamingGenerationErrorHandler.h b/Engine/Source/Runtime/Engine/Public/WorldPartition/ErrorHandling/WorldPartitionStreamingGenerationErrorHandler.h
--- a/Engine/Source/Runtime/Engine/Public/WorldPartition/ErrorHandling/WorldPartitionStreamingGenerationErrorHandler.h
+++ b/Engine/Source/Runtime/Engine/Public/WorldPartition/ErrorHandling/WorldPartitionStreamingGenerationErrorHandler.h
@@ -85,7 +85,7 @@ public:
 
 	virtual void OnLevelInstanceInvalidWorldAsset(const FWorldPartitionActorDescView& ActorDescView, FName WorldAsset, ELevelInstanceInvalidReason Reason) = 0;
 
-	virtual void OnActorFilterFailed(const FWorldPartitionActorDescView& ActorDescView) = 0;
+	virtual void OnInvalidActorFilterReference(const FWorldPartitionActorDescView& ActorDescView, const FWorldPartitionActorDescView& ReferenceActorDescView) = 0;
 
 	// Helpers
 	static ENGINE_API FString GetActorName(const FWorldPartitionActorDescView& ActorDescView);
diff --git a/Engine/Source/Runtime/Engine/Public/WorldPartition/ErrorHandling/WorldPartitionStreamingGenerationLogErrorHandler.h b/Engine/Source/Runtime/Engine/Public/WorldPartition/ErrorHandling/WorldPartitionStreamingGenerationLogErrorHandler.h
--- a/Engine/Source/Runtime/Engine/Public/WorldPartition/ErrorHandling/WorldPartitionStreamingGenerationLogErrorHandler.h
+++ b/Engine/Source/Runtime/Engine/Public/WorldPartition/ErrorHandling/WorldPartitionStreamingGenerationLogErrorHandler.h
@@ -19,6 +19,6 @@ public:
 	ENGINE_API virtual void OnDataLayerAssetConflict(const UDataLayerInstanceWithAsset* DataLayerInstance, const UDataLayerInstanceWithAsset* ConflictingDataLayerInstance) override;
 	ENGINE_API virtual void OnActorNeedsResave(const FWorldPartitionActorDescView& ActorDescView) override;
 	ENGINE_API virtual void OnLevelInstanceInvalidWorldAsset(const FWorldPartitionActorDescView& ActorDescView, FName WorldAsset, ELevelInstanceInvalidReason Reason) override;
-	ENGINE_API virtual void OnActorFilterFailed(const FWorldPartitionActorDescView& ActorDescView) override;
+	ENGINE_API virtual void OnInvalidActorFilterReference(const FWorldPartitionActorDescView& ActorDescView, const FWorldPartitionActorDescView& ReferenceActorDescView) override;
 };
 #endif
diff --git a/Engine/Source/Runtime/Engine/Public/WorldPartition/ErrorHandling/WorldPartitionStreamingGenerationNullErrorHandler.h b/Engine/Source/Runtime/Engine/Public/WorldPartition/ErrorHandling/WorldPartitionStreamingGenerationNullErrorHandler.h
--- a/Engine/Source/Runtime/Engine/Public/WorldPartition/ErrorHandling/WorldPartitionStreamingGenerationNullErrorHandler.h
+++ b/Engine/Source/Runtime/Engine/Public/WorldPartition/ErrorHandling/WorldPartitionStreamingGenerationNullErrorHandler.h
@@ -19,6 +19,6 @@ public:
 	virtual void OnDataLayerAssetConflict(const UDataLayerInstanceWithAsset* DataLayerInstance, const UDataLayerInstanceWithAsset* ConflictingDataLayerInstance) override {};
 	virtual void OnActorNeedsResave(const FWorldPartitionActorDescView& ActorDescView) override {};
 	virtual void OnLevelInstanceInvalidWorldAsset(const FWorldPartitionActorDescView& ActorDescView, FName WorldAsset, ELevelInstanceInvalidReason Reason) override {};
-	virtual void OnActorFilterFailed(const FWorldPartitionActorDescView& ActorDescView) override {};
+	virtual void OnInvalidActorFilterReference(const FWorldPartitionActorDescView& ActorDescView, const FWorldPartitionActorDescView& ReferenceActorDescView) override {};
 };
 #endif
diff --git a/Engine/Source/Runtime/Engine/Public/WorldPartition/ErrorHandling/WorldPartitionStreamingGenerationProxyErrorHandler.h b/Engine/Source/Runtime/Engine/Public/WorldPartition/ErrorHandling/WorldPartitionStreamingGenerationProxyErrorHandler.h
--- a/Engine/Source/Runtime/Engine/Public/WorldPartition/ErrorHandling/WorldPartitionStreamingGenerationProxyErrorHandler.h
+++ b/Engine/Source/Runtime/Engine/Public/WorldPartition/ErrorHandling/WorldPartitionStreamingGenerationProxyErrorHandler.h
@@ -23,7 +23,7 @@ public:
 	ENGINE_API virtual void OnDataLayerAssetConflict(const UDataLayerInstanceWithAsset* DataLayerInstance, const UDataLayerInstanceWithAsset* ConflictingDataLayerInstance) override;
 	ENGINE_API virtual void OnActorNeedsResave(const FWorldPartitionActorDescView& ActorDescView) override;
 	ENGINE_API virtual void OnLevelInstanceInvalidWorldAsset(const FWorldPartitionActorDescView& ActorDescView, FName WorldAsset, ELevelInstanceInvalidReason Reason) override;
-	ENGINE_API virtual void OnActorFilterFailed(const FWorldPartitionActorDescView& ActorDescView) override;
+	ENGINE_API virtual void OnInvalidActorFilterReference(const FWorldPartitionActorDescView& ActorDescView, const FWorldPartitionActorDescView& ReferenceActorDescView) override;
 
 protected:
 	IStreamingGenerationErrorHandler* InnerErrorHandler;
diff --git a/Engine/Source/Runtime/Engine/Public/WorldPartition/ErrorHandling/WorldPartitionStreamingGenerationTokenizedMessageErrorHandler.h b/Engine/Source/Runtime/Engine/Public/WorldPartition/ErrorHandling/WorldPartitionStreamingGenerationTokenizedMessageErrorHandler.h
--- a/Engine/Source/Runtime/Engine/Public/WorldPartition/ErrorHandling/WorldPartitionStreamingGenerationTokenizedMessageErrorHandler.h
+++ b/Engine/Source/Runtime/Engine/Public/WorldPartition/ErrorHandling/WorldPartitionStreamingGenerationTokenizedMessageErrorHandler.h
@@ -20,7 +20,7 @@ class ITokenizedMessageErrorHandler : public IStreamingGenerationErrorHandler
 	ENGINE_API virtual void OnDataLayerAssetConflict(const UDataLayerInstanceWithAsset* DataLayerInstance, const UDataLayerInstanceWithAsset* ConflictingDataLayerInstance) override;
 	ENGINE_API virtual void OnActorNeedsResave(const FWorldPartitionActorDescView& ActorDescView) override;
 	ENGINE_API virtual void OnLevelInstanceInvalidWorldAsset(const FWorldPartitionActorDescView& ActorDescView, FName WorldAsset, ELevelInstanceInvalidReason Reason) override;
-	ENGINE_API virtual void OnActorFilterFailed(const FWorldPartitionActorDescView& ActorDescView) override;
+	ENGINE_API virtual void OnInvalidActorFilterReference(const FWorldPartitionActorDescView& ActorDescView, const FWorldPartitionActorDescView& ReferenceActorDescView) override;
 
 protected:
 	virtual void HandleTokenizedMessage(TSharedRef<FTokenizedMessage>&& ErrorMessage) = 0;
diff --git a/Engine/Source/Runtime/Engine/Public/WorldPartition/WorldPartitionStreamingGenerationContext.h b/Engine/Source/Runtime/Engine/Public/WorldPartition/WorldPartitionStreamingGenerationContext.h
--- a/Engine/Source/Runtime/Engine/Public/WorldPartition/WorldPartitionStreamingGenerationContext.h
+++ b/Engine/Source/Runtime/Engine/Public/WorldPartition/WorldPartitionStreamingGenerationContext.h
@@ -2,6 +2,7 @@
 #pragma once
 
 #include "CoreMinimal.h"
+#include "OverrideVoidReturnInvoker.h"
 #include "WorldPartition/WorldPartitionActorContainerID.h"
 
 class FActorDescViewMap;
@@ -46,6 +47,24 @@ public:
 		FActorContainerID ContainerID;
 		FTransform Transform;
 		const FActorSet* ActorSet;
+		const TSet<FGuid>* FilteredActors;
+
+		template <typename Func>
+		void ForEachActor(Func InFunc) const
+		{
+			TOverrideVoidReturnInvoker Invoker(true, InFunc);
+
+			for (const FGuid& ActorGuid : ActorSet->Actors)
+			{
+				if (!FilteredActors || !FilteredActors->Contains(ActorGuid))
+				{
+					if (!Invoker(ActorGuid))
+					{
+						break;
+					}
+				}
+			}
+		}
 	};
 
 	struct FActorInstance
