Fix for blocking SparseVolumeTexture requests sometimes dropping mips. In certain cases, non-blocking prefetches issued for blocking requests were still not done in the next frame and even though the actual requested high resolution mip finished streaming (due to blocking), it couldn't be used because the lower resolution mips were still missing. Fixed by disabling non-blocking prefetch requests for blocking requests.

#rnx

diff --git a/Engine/Source/Runtime/Engine/Private/SparseVolumeTexture/SparseVolumeTextureStreamingManager.cpp b/Engine/Source/Runtime/Engine/Private/SparseVolumeTexture/SparseVolumeTextureStreamingManager.cpp
--- a/Engine/Source/Runtime/Engine/Private/SparseVolumeTexture/SparseVolumeTextureStreamingManager.cpp
+++ b/Engine/Source/Runtime/Engine/Private/SparseVolumeTexture/SparseVolumeTextureStreamingManager.cpp
@@ -943,7 +943,12 @@ void FStreamingManager::Request(UStreamableSparseVolumeTexture* SparseVolumeText
 		}
 
 		check(StreamingWindow);
-		const int32 OffsetMagnitude = GSVTStreamingNumPrefetchFrames;
+		// No prefetching for blocking requests. Making the prefetches blocking would increase latency even more and making them non-blocking could lead
+		// to situations where lower mips are already streamed in in subsequent frames but can't be used because dependent higher mips haven't finished streaming
+		// due to non-blocking requests.
+		// SVT_TODO: This can still break if blocking and non-blocking requests of the same frames/mips are made to the same SVT. We would need to cancel already scheduled non-blocking requests and reissue them as blocking.
+		// Or alternatively we could just block on all requests if we detect this case. If we had a single DDC request owner per request, we could just selectively wait on already scheduled non-blocking requests.
+		const int32 OffsetMagnitude = !bBlocking ? GSVTStreamingNumPrefetchFrames : 0;
 		const int32 LowerFrameOffset = StreamingWindow->bPlayBackward ? -OffsetMagnitude : 0;
 		const int32 UpperFrameOffset = StreamingWindow->bPlayForward ? OffsetMagnitude : 0;
 
@@ -956,7 +961,7 @@ void FStreamingManager::Request(UStreamableSparseVolumeTexture* SparseVolumeText
 			Request.Key.FrameIndex = RequestFrameIndex;
 			Request.Key.MipLevelIndex = FMath::Clamp(MipLevel + RequestMipLevelOffset, 0, SVTInfo->PerFrameInfo[RequestFrameIndex].NumMipLevels);
 			Request.Priority = FMath::Max(0, OffsetMagnitude - FMath::Abs(i));
-			if (bBlocking && (i == 0)) // Only block on the actually requested frame
+			if (bBlocking)
 			{
 				Request.Priority = FStreamingRequest::BlockingPriority;
 			}
