Add animstats.h headers to node that use verbose profiling

diff --git a/Engine/Plugins/Animation/ControlRig/Source/ControlRig/Private/AnimNode_ControlRig.cpp b/Engine/Plugins/Animation/ControlRig/Source/ControlRig/Private/AnimNode_ControlRig.cpp
--- a/Engine/Plugins/Animation/ControlRig/Source/ControlRig/Private/AnimNode_ControlRig.cpp
+++ b/Engine/Plugins/Animation/ControlRig/Source/ControlRig/Private/AnimNode_ControlRig.cpp
@@ -9,6 +9,7 @@
 #include "Animation/NodeMappingContainer.h"
 #include "AnimationRuntime.h"
 #include "Animation/AnimCurveUtils.h"
+#include "Animation/AnimStats.h"
 
 #include UE_INLINE_GENERATED_CPP_BY_NAME(AnimNode_ControlRig)
 
diff --git a/Engine/Plugins/Animation/IKRig/Source/IKRig/Private/AnimNodes/AnimNode_IKRig.cpp b/Engine/Plugins/Animation/IKRig/Source/IKRig/Private/AnimNodes/AnimNode_IKRig.cpp
--- a/Engine/Plugins/Animation/IKRig/Source/IKRig/Private/AnimNodes/AnimNode_IKRig.cpp
+++ b/Engine/Plugins/Animation/IKRig/Source/IKRig/Private/AnimNodes/AnimNode_IKRig.cpp
@@ -9,6 +9,7 @@
 #include "Components/SkeletalMeshComponent.h"
 #include "ActorComponents/IKRigInterface.h"
 #include "Animation/AnimInstanceProxy.h"
+#include "Animation/AnimStats.h"
 #include "Algo/ForEach.h"
 #include "SceneManagement.h"
 
diff --git a/Engine/Plugins/Animation/IKRig/Source/IKRigEditor/Private/RetargetEditor/IKRetargetAnimInstance.cpp b/Engine/Plugins/Animation/IKRig/Source/IKRigEditor/Private/RetargetEditor/IKRetargetAnimInstance.cpp
--- a/Engine/Plugins/Animation/IKRig/Source/IKRigEditor/Private/RetargetEditor/IKRetargetAnimInstance.cpp
+++ b/Engine/Plugins/Animation/IKRig/Source/IKRigEditor/Private/RetargetEditor/IKRetargetAnimInstance.cpp
@@ -1,6 +1,7 @@
 // Copyright Epic Games, Inc. All Rights Reserved.
 
 #include "RetargetEditor/IKRetargetAnimInstance.h"
+#include "Animation/AnimStats.h"
 #include "Engine/SkeletalMesh.h"
 #include "RetargetEditor/IKRetargetAnimInstanceProxy.h"
 #include "RetargetEditor/IKRetargeterController.h"
diff --git a/Engine/Plugins/Experimental/Animation/CurveExpression/Source/Runtime/Private/AnimNode_RemapCurves.cpp b/Engine/Plugins/Experimental/Animation/CurveExpression/Source/Runtime/Private/AnimNode_RemapCurves.cpp
--- a/Engine/Plugins/Experimental/Animation/CurveExpression/Source/Runtime/Private/AnimNode_RemapCurves.cpp
+++ b/Engine/Plugins/Experimental/Animation/CurveExpression/Source/Runtime/Private/AnimNode_RemapCurves.cpp
@@ -3,6 +3,7 @@
 #include "AnimNode_RemapCurves.h"
 
 #include "Animation/AnimInstanceProxy.h"
+#include "Animation/AnimStats.h"
 
 void FAnimNode_RemapCurves::Evaluate_AnyThread(FPoseContext& Output)
 {
diff --git a/Engine/Plugins/Experimental/Animation/CurveExpression/Source/Runtime/Private/AnimNode_RemapCurvesFromMesh.cpp b/Engine/Plugins/Experimental/Animation/CurveExpression/Source/Runtime/Private/AnimNode_RemapCurvesFromMesh.cpp
--- a/Engine/Plugins/Experimental/Animation/CurveExpression/Source/Runtime/Private/AnimNode_RemapCurvesFromMesh.cpp
+++ b/Engine/Plugins/Experimental/Animation/CurveExpression/Source/Runtime/Private/AnimNode_RemapCurvesFromMesh.cpp
@@ -5,6 +5,7 @@
 #include "CurveExpressionCustomVersion.h"
 #include "Animation/AnimInstance.h"
 #include "Animation/AnimNodeFunctionRef.h"
+#include "Animation/AnimStats.h"
 #include "Components/SkeletalMeshComponent.h"
 #include "Animation/ExposedValueHandler.h"
 #include "Engine/SkeletalMesh.h"
diff --git a/Engine/Plugins/Experimental/Animation/PoseCorrectives/Source/PoseCorrectives/Private/AnimNode_CorrectPose.cpp b/Engine/Plugins/Experimental/Animation/PoseCorrectives/Source/PoseCorrectives/Private/AnimNode_CorrectPose.cpp
--- a/Engine/Plugins/Experimental/Animation/PoseCorrectives/Source/PoseCorrectives/Private/AnimNode_CorrectPose.cpp
+++ b/Engine/Plugins/Experimental/Animation/PoseCorrectives/Source/PoseCorrectives/Private/AnimNode_CorrectPose.cpp
@@ -3,7 +3,7 @@
 #include "AnimNode_CorrectPose.h"
 
 #include "Animation/AnimInstanceProxy.h"
-
+#include "Animation/AnimStats.h"
 
 FAnimNode_CorrectPose::FAnimNode_CorrectPose()
 {
diff --git a/Engine/Plugins/Experimental/Animation/PoseSearch/Source/Runtime/Private/AnimNode_MotionMatching.cpp b/Engine/Plugins/Experimental/Animation/PoseSearch/Source/Runtime/Private/AnimNode_MotionMatching.cpp
--- a/Engine/Plugins/Experimental/Animation/PoseSearch/Source/Runtime/Private/AnimNode_MotionMatching.cpp
+++ b/Engine/Plugins/Experimental/Animation/PoseSearch/Source/Runtime/Private/AnimNode_MotionMatching.cpp
@@ -6,6 +6,7 @@
 #include "Animation/AnimNode_Inertialization.h"
 #include "Animation/AnimRootMotionProvider.h"
 #include "Animation/AnimSequence.h"
+#include "Animation/AnimStats.h"
 #include "Animation/BlendSpace.h"
 #include "Components/SkeletalMeshComponent.h"
 #include "PoseSearch/AnimNode_PoseSearchHistoryCollector.h"
diff --git a/Engine/Plugins/Experimental/Animation/PoseSearch/Source/Runtime/Private/AnimNode_PoseSearchHistoryCollector.cpp b/Engine/Plugins/Experimental/Animation/PoseSearch/Source/Runtime/Private/AnimNode_PoseSearchHistoryCollector.cpp
--- a/Engine/Plugins/Experimental/Animation/PoseSearch/Source/Runtime/Private/AnimNode_PoseSearchHistoryCollector.cpp
+++ b/Engine/Plugins/Experimental/Animation/PoseSearch/Source/Runtime/Private/AnimNode_PoseSearchHistoryCollector.cpp
@@ -3,6 +3,7 @@
 #include "PoseSearch/AnimNode_PoseSearchHistoryCollector.h"
 #include "Animation/AnimInstanceProxy.h"
 #include "Animation/AnimNodeMessages.h"
+#include "Animation/AnimStats.h"
 #include "Engine/SkinnedAsset.h"
 
 #include UE_INLINE_GENERATED_CPP_BY_NAME(AnimNode_PoseSearchHistoryCollector)
diff --git a/Engine/Plugins/Experimental/ChaosVehiclesPlugin/Source/ChaosVehicles/Private/AnimNode_StageCoachWheelController.cpp b/Engine/Plugins/Experimental/ChaosVehiclesPlugin/Source/ChaosVehicles/Private/AnimNode_StageCoachWheelController.cpp
--- a/Engine/Plugins/Experimental/ChaosVehiclesPlugin/Source/ChaosVehicles/Private/AnimNode_StageCoachWheelController.cpp
+++ b/Engine/Plugins/Experimental/ChaosVehiclesPlugin/Source/ChaosVehicles/Private/AnimNode_StageCoachWheelController.cpp
@@ -2,6 +2,7 @@
 
 #include "AnimNode_StageCoachWheelController.h"
 #include "AnimationRuntime.h"
+#include "Animation/AnimStats.h"
 #include "WheeledVehiclePawn.h"
 #include "ChaosVehicleMovementComponent.h"
 
diff --git a/Engine/Plugins/Experimental/ChaosVehiclesPlugin/Source/ChaosVehicles/Private/AnimNode_WheelController.cpp b/Engine/Plugins/Experimental/ChaosVehiclesPlugin/Source/ChaosVehicles/Private/AnimNode_WheelController.cpp
--- a/Engine/Plugins/Experimental/ChaosVehiclesPlugin/Source/ChaosVehicles/Private/AnimNode_WheelController.cpp
+++ b/Engine/Plugins/Experimental/ChaosVehiclesPlugin/Source/ChaosVehicles/Private/AnimNode_WheelController.cpp
@@ -2,6 +2,7 @@
 
 #include "AnimNode_WheelController.h"
 #include "AnimationRuntime.h"
+#include "Animation/AnimStats.h"
 #include "WheeledVehiclePawn.h"
 #include "ChaosVehicleMovementComponent.h"
 
diff --git a/Engine/Source/Runtime/AnimGraphRuntime/Private/AnimNodes/AnimNode_AimOffsetLookAt.cpp b/Engine/Source/Runtime/AnimGraphRuntime/Private/AnimNodes/AnimNode_AimOffsetLookAt.cpp
--- a/Engine/Source/Runtime/AnimGraphRuntime/Private/AnimNodes/AnimNode_AimOffsetLookAt.cpp
+++ b/Engine/Source/Runtime/AnimGraphRuntime/Private/AnimNodes/AnimNode_AimOffsetLookAt.cpp
@@ -2,6 +2,7 @@
 
 #include "AnimNodes/AnimNode_AimOffsetLookAt.h"
 #include "Animation/AnimInstanceProxy.h"
+#include "Animation/AnimStats.h"
 #include "AnimationRuntime.h"
 #include "Animation/BlendSpace.h"
 #include "Animation/BlendSpace1D.h"
diff --git a/Engine/Source/Runtime/AnimGraphRuntime/Private/AnimNodes/AnimNode_ApplyAdditive.cpp b/Engine/Source/Runtime/AnimGraphRuntime/Private/AnimNodes/AnimNode_ApplyAdditive.cpp
--- a/Engine/Source/Runtime/AnimGraphRuntime/Private/AnimNodes/AnimNode_ApplyAdditive.cpp
+++ b/Engine/Source/Runtime/AnimGraphRuntime/Private/AnimNodes/AnimNode_ApplyAdditive.cpp
@@ -2,6 +2,7 @@
 
 #include "AnimNodes/AnimNode_ApplyAdditive.h"
 #include "AnimationRuntime.h"
+#include "Animation/AnimStats.h"
 #include "Animation/AnimInstanceProxy.h"
 #include "Animation/AnimTrace.h"
 
diff --git a/Engine/Source/Runtime/AnimGraphRuntime/Private/AnimNodes/AnimNode_BlendBoneByChannel.cpp b/Engine/Source/Runtime/AnimGraphRuntime/Private/AnimNodes/AnimNode_BlendBoneByChannel.cpp
--- a/Engine/Source/Runtime/AnimGraphRuntime/Private/AnimNodes/AnimNode_BlendBoneByChannel.cpp
+++ b/Engine/Source/Runtime/AnimGraphRuntime/Private/AnimNodes/AnimNode_BlendBoneByChannel.cpp
@@ -4,6 +4,7 @@
 #include "AnimationRuntime.h"
 #include "Animation/AnimInstanceProxy.h"
 #include "Animation/AnimTrace.h"
+#include "Animation/AnimStats.h"
 
 #include UE_INLINE_GENERATED_CPP_BY_NAME(AnimNode_BlendBoneByChannel)
 
diff --git a/Engine/Source/Runtime/AnimGraphRuntime/Private/AnimNodes/AnimNode_BlendSpacePlayer.cpp b/Engine/Source/Runtime/AnimGraphRuntime/Private/AnimNodes/AnimNode_BlendSpacePlayer.cpp
--- a/Engine/Source/Runtime/AnimGraphRuntime/Private/AnimNodes/AnimNode_BlendSpacePlayer.cpp
+++ b/Engine/Source/Runtime/AnimGraphRuntime/Private/AnimNodes/AnimNode_BlendSpacePlayer.cpp
@@ -7,6 +7,7 @@
 #include "AnimGraphRuntimeTrace.h"
 #include "Animation/AnimSync.h"
 #include "Animation/AnimSyncScope.h"
+#include "Animation/AnimStats.h"
 
 #include UE_INLINE_GENERATED_CPP_BY_NAME(AnimNode_BlendSpacePlayer)
 
diff --git a/Engine/Source/Runtime/AnimGraphRuntime/Private/AnimNodes/AnimNode_CopyPoseFromMesh.cpp b/Engine/Source/Runtime/AnimGraphRuntime/Private/AnimNodes/AnimNode_CopyPoseFromMesh.cpp
--- a/Engine/Source/Runtime/AnimGraphRuntime/Private/AnimNodes/AnimNode_CopyPoseFromMesh.cpp
+++ b/Engine/Source/Runtime/AnimGraphRuntime/Private/AnimNodes/AnimNode_CopyPoseFromMesh.cpp
@@ -4,6 +4,7 @@
 
 #include "Animation/AnimCurveUtils.h"
 #include "Animation/AnimInstanceProxy.h"
+#include "Animation/AnimStats.h"
 #include "Animation/AnimTrace.h"
 #include "Components/SkeletalMeshComponent.h"
 #include "Engine/SkeletalMesh.h"
diff --git a/Engine/Source/Runtime/AnimGraphRuntime/Private/AnimNodes/AnimNode_CurveSource.cpp b/Engine/Source/Runtime/AnimGraphRuntime/Private/AnimNodes/AnimNode_CurveSource.cpp
--- a/Engine/Source/Runtime/AnimGraphRuntime/Private/AnimNodes/AnimNode_CurveSource.cpp
+++ b/Engine/Source/Runtime/AnimGraphRuntime/Private/AnimNodes/AnimNode_CurveSource.cpp
@@ -4,6 +4,7 @@
 #include "AnimationRuntime.h"
 #include "Animation/AnimCurveUtils.h"
 #include "Animation/AnimInstanceProxy.h"
+#include "Animation/AnimStats.h"
 #include "Animation/AnimTrace.h"
 
 #include UE_INLINE_GENERATED_CPP_BY_NAME(AnimNode_CurveSource)
diff --git a/Engine/Source/Runtime/AnimGraphRuntime/Private/AnimNodes/AnimNode_MakeDynamicAdditive.cpp b/Engine/Source/Runtime/AnimGraphRuntime/Private/AnimNodes/AnimNode_MakeDynamicAdditive.cpp
--- a/Engine/Source/Runtime/AnimGraphRuntime/Private/AnimNodes/AnimNode_MakeDynamicAdditive.cpp
+++ b/Engine/Source/Runtime/AnimGraphRuntime/Private/AnimNodes/AnimNode_MakeDynamicAdditive.cpp
@@ -2,6 +2,7 @@
 
 #include "AnimNodes/AnimNode_MakeDynamicAdditive.h"
 #include "AnimationRuntime.h"
+#include "Animation/AnimStats.h"
 #include "Animation/AnimTrace.h"
 
 #include UE_INLINE_GENERATED_CPP_BY_NAME(AnimNode_MakeDynamicAdditive)
diff --git a/Engine/Source/Runtime/AnimGraphRuntime/Private/AnimNodes/AnimNode_ModifyCurve.cpp b/Engine/Source/Runtime/AnimGraphRuntime/Private/AnimNodes/AnimNode_ModifyCurve.cpp
--- a/Engine/Source/Runtime/AnimGraphRuntime/Private/AnimNodes/AnimNode_ModifyCurve.cpp
+++ b/Engine/Source/Runtime/AnimGraphRuntime/Private/AnimNodes/AnimNode_ModifyCurve.cpp
@@ -4,6 +4,7 @@
 #include "AnimationRuntime.h"
 #include "Animation/AnimInstanceProxy.h"
 #include "Animation/AnimCurveTypes.h"
+#include "Animation/AnimStats.h"
 
 #include UE_INLINE_GENERATED_CPP_BY_NAME(AnimNode_ModifyCurve)
 
diff --git a/Engine/Source/Runtime/AnimGraphRuntime/Private/AnimNodes/AnimNode_MultiWayBlend.cpp b/Engine/Source/Runtime/AnimGraphRuntime/Private/AnimNodes/AnimNode_MultiWayBlend.cpp
--- a/Engine/Source/Runtime/AnimGraphRuntime/Private/AnimNodes/AnimNode_MultiWayBlend.cpp
+++ b/Engine/Source/Runtime/AnimGraphRuntime/Private/AnimNodes/AnimNode_MultiWayBlend.cpp
@@ -2,6 +2,7 @@
 
 #include "AnimNodes/AnimNode_MultiWayBlend.h"
 #include "AnimationRuntime.h"
+#include "Animation/AnimStats.h"
 
 #include UE_INLINE_GENERATED_CPP_BY_NAME(AnimNode_MultiWayBlend)
 
diff --git a/Engine/Source/Runtime/AnimGraphRuntime/Private/AnimNodes/AnimNode_RandomPlayer.cpp b/Engine/Source/Runtime/AnimGraphRuntime/Private/AnimNodes/AnimNode_RandomPlayer.cpp
--- a/Engine/Source/Runtime/AnimGraphRuntime/Private/AnimNodes/AnimNode_RandomPlayer.cpp
+++ b/Engine/Source/Runtime/AnimGraphRuntime/Private/AnimNodes/AnimNode_RandomPlayer.cpp
@@ -6,6 +6,7 @@
 #include "Animation/AnimInstanceProxy.h"
 #include "Animation/AnimSequence.h"
 #include "Animation/AnimTrace.h"
+#include "Animation/AnimStats.h"
 #include "Animation/AnimSyncScope.h"
 
 #include UE_INLINE_GENERATED_CPP_BY_NAME(AnimNode_RandomPlayer)
diff --git a/Engine/Source/Runtime/AnimGraphRuntime/Private/AnimNodes/AnimNode_RefPose.cpp b/Engine/Source/Runtime/AnimGraphRuntime/Private/AnimNodes/AnimNode_RefPose.cpp
--- a/Engine/Source/Runtime/AnimGraphRuntime/Private/AnimNodes/AnimNode_RefPose.cpp
+++ b/Engine/Source/Runtime/AnimGraphRuntime/Private/AnimNodes/AnimNode_RefPose.cpp
@@ -1,6 +1,7 @@
 // Copyright Epic Games, Inc. All Rights Reserved.
 
 #include "AnimNodes/AnimNode_RefPose.h"
+#include "Animation/AnimStats.h"
 #include "Animation/AnimTrace.h"
 
 /////////////////////////////////////////////////////
diff --git a/Engine/Source/Runtime/AnimGraphRuntime/Private/AnimNodes/AnimNode_RotateRootBone.cpp b/Engine/Source/Runtime/AnimGraphRuntime/Private/AnimNodes/AnimNode_RotateRootBone.cpp
--- a/Engine/Source/Runtime/AnimGraphRuntime/Private/AnimNodes/AnimNode_RotateRootBone.cpp
+++ b/Engine/Source/Runtime/AnimGraphRuntime/Private/AnimNodes/AnimNode_RotateRootBone.cpp
@@ -1,6 +1,7 @@
 // Copyright Epic Games, Inc. All Rights Reserved.
 
 #include "AnimNodes/AnimNode_RotateRootBone.h"
+#include "Animation/AnimStats.h"
 #include "Animation/AnimTrace.h"
 #include "Animation/AnimRootMotionProvider.h"
 
diff --git a/Engine/Source/Runtime/AnimGraphRuntime/Private/AnimNodes/AnimNode_Slot.cpp b/Engine/Source/Runtime/AnimGraphRuntime/Private/AnimNodes/AnimNode_Slot.cpp
--- a/Engine/Source/Runtime/AnimGraphRuntime/Private/AnimNodes/AnimNode_Slot.cpp
+++ b/Engine/Source/Runtime/AnimGraphRuntime/Private/AnimNodes/AnimNode_Slot.cpp
@@ -3,6 +3,7 @@
 #include "AnimNodes/AnimNode_Slot.h"
 #include "Animation/AnimInstanceProxy.h"
 #include "Animation/AnimMontage.h"
+#include "Animation/AnimStats.h"
 #include "Animation/AnimTrace.h"
 #include "Animation/AnimNode_Inertialization.h"
 
diff --git a/Engine/Source/Runtime/AnimGraphRuntime/Private/AnimNodes/AnimNode_TwoWayBlend.cpp b/Engine/Source/Runtime/AnimGraphRuntime/Private/AnimNodes/AnimNode_TwoWayBlend.cpp
--- a/Engine/Source/Runtime/AnimGraphRuntime/Private/AnimNodes/AnimNode_TwoWayBlend.cpp
+++ b/Engine/Source/Runtime/AnimGraphRuntime/Private/AnimNodes/AnimNode_TwoWayBlend.cpp
@@ -2,6 +2,7 @@
 
 #include "AnimNodes/AnimNode_TwoWayBlend.h"
 #include "Animation/AnimInstanceProxy.h"
+#include "Animation/AnimStats.h"
 #include "AnimationRuntime.h"
 
 #include UE_INLINE_GENERATED_CPP_BY_NAME(AnimNode_TwoWayBlend)
diff --git a/Engine/Source/Runtime/AnimGraphRuntime/Private/BoneControllers/AnimNodeRotationMultiplier.cpp b/Engine/Source/Runtime/AnimGraphRuntime/Private/BoneControllers/AnimNodeRotationMultiplier.cpp
--- a/Engine/Source/Runtime/AnimGraphRuntime/Private/BoneControllers/AnimNodeRotationMultiplier.cpp
+++ b/Engine/Source/Runtime/AnimGraphRuntime/Private/BoneControllers/AnimNodeRotationMultiplier.cpp
@@ -4,6 +4,7 @@
 #include "BoneIndices.h"
 #include "BonePose.h"
 #include "BoneControllers/AnimNode_RotationMultiplier.h"
+#include "Animation/AnimStats.h"
 #include "Animation/AnimTrace.h"
 
 /////////////////////////////////////////////////////
diff --git a/Engine/Source/Runtime/AnimGraphRuntime/Private/BoneControllers/AnimNode_ApplyLimits.cpp b/Engine/Source/Runtime/AnimGraphRuntime/Private/BoneControllers/AnimNode_ApplyLimits.cpp
--- a/Engine/Source/Runtime/AnimGraphRuntime/Private/BoneControllers/AnimNode_ApplyLimits.cpp
+++ b/Engine/Source/Runtime/AnimGraphRuntime/Private/BoneControllers/AnimNode_ApplyLimits.cpp
@@ -3,6 +3,7 @@
 #include "BoneControllers/AnimNode_ApplyLimits.h"
 #include "AnimationCoreLibrary.h"
 #include "Animation/AnimInstanceProxy.h"
+#include "Animation/AnimStats.h"
 #include "AnimationRuntime.h"
 #include "AngularLimit.h"
 
diff --git a/Engine/Source/Runtime/AnimGraphRuntime/Private/BoneControllers/AnimNode_BoneDrivenController.cpp b/Engine/Source/Runtime/AnimGraphRuntime/Private/BoneControllers/AnimNode_BoneDrivenController.cpp
--- a/Engine/Source/Runtime/AnimGraphRuntime/Private/BoneControllers/AnimNode_BoneDrivenController.cpp
+++ b/Engine/Source/Runtime/AnimGraphRuntime/Private/BoneControllers/AnimNode_BoneDrivenController.cpp
@@ -5,6 +5,7 @@
 #include "Components/SkeletalMeshComponent.h"
 #include "Curves/CurveFloat.h"
 #include "Animation/AnimInstanceProxy.h"
+#include "Animation/AnimStats.h"
 #include "Animation/AnimTrace.h"
 
 #include UE_INLINE_GENERATED_CPP_BY_NAME(AnimNode_BoneDrivenController)
diff --git a/Engine/Source/Runtime/AnimGraphRuntime/Private/BoneControllers/AnimNode_CCDIK.cpp b/Engine/Source/Runtime/AnimGraphRuntime/Private/BoneControllers/AnimNode_CCDIK.cpp
--- a/Engine/Source/Runtime/AnimGraphRuntime/Private/BoneControllers/AnimNode_CCDIK.cpp
+++ b/Engine/Source/Runtime/AnimGraphRuntime/Private/BoneControllers/AnimNode_CCDIK.cpp
@@ -2,6 +2,7 @@
 
 #include "BoneControllers/AnimNode_CCDIK.h"
 #include "Animation/AnimTypes.h"
+#include "Animation/AnimStats.h"
 #include "AnimationRuntime.h"
 #include "DrawDebugHelpers.h"
 #include "EngineDefines.h"
diff --git a/Engine/Source/Runtime/AnimGraphRuntime/Private/BoneControllers/AnimNode_Constraint.cpp b/Engine/Source/Runtime/AnimGraphRuntime/Private/BoneControllers/AnimNode_Constraint.cpp
--- a/Engine/Source/Runtime/AnimGraphRuntime/Private/BoneControllers/AnimNode_Constraint.cpp
+++ b/Engine/Source/Runtime/AnimGraphRuntime/Private/BoneControllers/AnimNode_Constraint.cpp
@@ -3,6 +3,7 @@
 #include "BoneControllers/AnimNode_Constraint.h"
 #include "AnimationCoreLibrary.h"
 #include "AnimationRuntime.h"
+#include "Animation/AnimStats.h"
 #include "SceneManagement.h"
 #include "Components/SkeletalMeshComponent.h"
 #include "Animation/AnimTrace.h"
diff --git a/Engine/Source/Runtime/AnimGraphRuntime/Private/BoneControllers/AnimNode_CopyBone.cpp b/Engine/Source/Runtime/AnimGraphRuntime/Private/BoneControllers/AnimNode_CopyBone.cpp
--- a/Engine/Source/Runtime/AnimGraphRuntime/Private/BoneControllers/AnimNode_CopyBone.cpp
+++ b/Engine/Source/Runtime/AnimGraphRuntime/Private/BoneControllers/AnimNode_CopyBone.cpp
@@ -3,6 +3,7 @@
 #include "BoneControllers/AnimNode_CopyBone.h"
 #include "AnimationRuntime.h"
 #include "Animation/AnimInstanceProxy.h"
+#include "Animation/AnimStats.h"
 #include "Animation/AnimTrace.h"
 
 #include UE_INLINE_GENERATED_CPP_BY_NAME(AnimNode_CopyBone)
diff --git a/Engine/Source/Runtime/AnimGraphRuntime/Private/BoneControllers/AnimNode_CopyBoneDelta.cpp b/Engine/Source/Runtime/AnimGraphRuntime/Private/BoneControllers/AnimNode_CopyBoneDelta.cpp
--- a/Engine/Source/Runtime/AnimGraphRuntime/Private/BoneControllers/AnimNode_CopyBoneDelta.cpp
+++ b/Engine/Source/Runtime/AnimGraphRuntime/Private/BoneControllers/AnimNode_CopyBoneDelta.cpp
@@ -4,6 +4,7 @@
 #include "Components/SkeletalMeshComponent.h"
 #include "AnimationRuntime.h"
 #include "Animation/AnimInstanceProxy.h"
+#include "Animation/AnimStats.h"
 
 #include UE_INLINE_GENERATED_CPP_BY_NAME(AnimNode_CopyBoneDelta)
 
diff --git a/Engine/Source/Runtime/AnimGraphRuntime/Private/BoneControllers/AnimNode_Fabrik.cpp b/Engine/Source/Runtime/AnimGraphRuntime/Private/BoneControllers/AnimNode_Fabrik.cpp
--- a/Engine/Source/Runtime/AnimGraphRuntime/Private/BoneControllers/AnimNode_Fabrik.cpp
+++ b/Engine/Source/Runtime/AnimGraphRuntime/Private/BoneControllers/AnimNode_Fabrik.cpp
@@ -5,6 +5,7 @@
 #include "Components/SkeletalMeshComponent.h"
 #include "DrawDebugHelpers.h"
 #include "Animation/AnimInstanceProxy.h"
+#include "Animation/AnimStats.h"
 #include "FABRIK.h"
 
 #include UE_INLINE_GENERATED_CPP_BY_NAME(AnimNode_Fabrik)
diff --git a/Engine/Source/Runtime/AnimGraphRuntime/Private/BoneControllers/AnimNode_HandIKRetargeting.cpp b/Engine/Source/Runtime/AnimGraphRuntime/Private/BoneControllers/AnimNode_HandIKRetargeting.cpp
--- a/Engine/Source/Runtime/AnimGraphRuntime/Private/BoneControllers/AnimNode_HandIKRetargeting.cpp
+++ b/Engine/Source/Runtime/AnimGraphRuntime/Private/BoneControllers/AnimNode_HandIKRetargeting.cpp
@@ -1,6 +1,7 @@
 // Copyright Epic Games, Inc. All Rights Reserved.
 
 #include "BoneControllers/AnimNode_HandIKRetargeting.h"
+#include "Animation/AnimStats.h"
 #include "Animation/AnimTrace.h"
 
 #include UE_INLINE_GENERATED_CPP_BY_NAME(AnimNode_HandIKRetargeting)
diff --git a/Engine/Source/Runtime/AnimGraphRuntime/Private/BoneControllers/AnimNode_LookAt.cpp b/Engine/Source/Runtime/AnimGraphRuntime/Private/BoneControllers/AnimNode_LookAt.cpp
--- a/Engine/Source/Runtime/AnimGraphRuntime/Private/BoneControllers/AnimNode_LookAt.cpp
+++ b/Engine/Source/Runtime/AnimGraphRuntime/Private/BoneControllers/AnimNode_LookAt.cpp
@@ -6,6 +6,7 @@
 #include "SceneManagement.h"
 #include "Engine/SkeletalMeshSocket.h"
 #include "Animation/AnimInstanceProxy.h"
+#include "Animation/AnimStats.h"
 #include "AnimationCoreLibrary.h"
 #include "Engine/Engine.h"
 #include "Materials/MaterialInstanceDynamic.h"
diff --git a/Engine/Source/Runtime/AnimGraphRuntime/Private/BoneControllers/AnimNode_ModifyBone.cpp b/Engine/Source/Runtime/AnimGraphRuntime/Private/BoneControllers/AnimNode_ModifyBone.cpp
--- a/Engine/Source/Runtime/AnimGraphRuntime/Private/BoneControllers/AnimNode_ModifyBone.cpp
+++ b/Engine/Source/Runtime/AnimGraphRuntime/Private/BoneControllers/AnimNode_ModifyBone.cpp
@@ -3,6 +3,7 @@
 #include "BoneControllers/AnimNode_ModifyBone.h"
 #include "AnimationRuntime.h"
 #include "Animation/AnimInstanceProxy.h"
+#include "Animation/AnimStats.h"
 #include "Animation/AnimTrace.h"
 
 /////////////////////////////////////////////////////
diff --git a/Engine/Source/Runtime/AnimGraphRuntime/Private/BoneControllers/AnimNode_ObserveBone.cpp b/Engine/Source/Runtime/AnimGraphRuntime/Private/BoneControllers/AnimNode_ObserveBone.cpp
--- a/Engine/Source/Runtime/AnimGraphRuntime/Private/BoneControllers/AnimNode_ObserveBone.cpp
+++ b/Engine/Source/Runtime/AnimGraphRuntime/Private/BoneControllers/AnimNode_ObserveBone.cpp
@@ -3,6 +3,7 @@
 #include "BoneControllers/AnimNode_ObserveBone.h"
 #include "AnimationRuntime.h"
 #include "Animation/AnimInstanceProxy.h"
+#include "Animation/AnimStats.h"
 #include "Animation/AnimTrace.h"
 
 #include UE_INLINE_GENERATED_CPP_BY_NAME(AnimNode_ObserveBone)
diff --git a/Engine/Source/Runtime/AnimGraphRuntime/Private/BoneControllers/AnimNode_ResetRoot.cpp b/Engine/Source/Runtime/AnimGraphRuntime/Private/BoneControllers/AnimNode_ResetRoot.cpp
--- a/Engine/Source/Runtime/AnimGraphRuntime/Private/BoneControllers/AnimNode_ResetRoot.cpp
+++ b/Engine/Source/Runtime/AnimGraphRuntime/Private/BoneControllers/AnimNode_ResetRoot.cpp
@@ -2,6 +2,7 @@
 
 #include "BoneControllers/AnimNode_ResetRoot.h"
 #include "Animation/AnimInstanceProxy.h"
+#include "Animation/AnimStats.h"
 #include "AnimationRuntime.h"
 
 #include UE_INLINE_GENERATED_CPP_BY_NAME(AnimNode_ResetRoot)
diff --git a/Engine/Source/Runtime/AnimGraphRuntime/Private/BoneControllers/AnimNode_ScaleChainLength.cpp b/Engine/Source/Runtime/AnimGraphRuntime/Private/BoneControllers/AnimNode_ScaleChainLength.cpp
--- a/Engine/Source/Runtime/AnimGraphRuntime/Private/BoneControllers/AnimNode_ScaleChainLength.cpp
+++ b/Engine/Source/Runtime/AnimGraphRuntime/Private/BoneControllers/AnimNode_ScaleChainLength.cpp
@@ -4,6 +4,7 @@
 #include "SceneManagement.h"
 #include "Engine/SkeletalMeshSocket.h"
 #include "Animation/AnimInstanceProxy.h"
+#include "Animation/AnimStats.h"
 
 #include UE_INLINE_GENERATED_CPP_BY_NAME(AnimNode_ScaleChainLength)
 
diff --git a/Engine/Source/Runtime/AnimGraphRuntime/Private/BoneControllers/AnimNode_SkeletalControlBase.cpp b/Engine/Source/Runtime/AnimGraphRuntime/Private/BoneControllers/AnimNode_SkeletalControlBase.cpp
--- a/Engine/Source/Runtime/AnimGraphRuntime/Private/BoneControllers/AnimNode_SkeletalControlBase.cpp
+++ b/Engine/Source/Runtime/AnimGraphRuntime/Private/BoneControllers/AnimNode_SkeletalControlBase.cpp
@@ -3,6 +3,7 @@
 #include "BoneControllers/AnimNode_SkeletalControlBase.h"
 #include "Engine/SkeletalMesh.h"
 #include "Animation/AnimInstanceProxy.h"
+#include "Animation/AnimStats.h"
 #include "Engine/SkeletalMeshSocket.h"
 
 #include UE_INLINE_GENERATED_CPP_BY_NAME(AnimNode_SkeletalControlBase)
diff --git a/Engine/Source/Runtime/AnimGraphRuntime/Private/BoneControllers/AnimNode_SplineIK.cpp b/Engine/Source/Runtime/AnimGraphRuntime/Private/BoneControllers/AnimNode_SplineIK.cpp
--- a/Engine/Source/Runtime/AnimGraphRuntime/Private/BoneControllers/AnimNode_SplineIK.cpp
+++ b/Engine/Source/Runtime/AnimGraphRuntime/Private/BoneControllers/AnimNode_SplineIK.cpp
@@ -4,6 +4,7 @@
 #include "Animation/AnimTypes.h"
 #include "AnimationRuntime.h"
 #include "Animation/AnimInstanceProxy.h"
+#include "Animation/AnimStats.h"
 #include "Components/SkeletalMeshComponent.h"
 #include "Engine/SkeletalMesh.h"
 #include "SplineIK.h"
diff --git a/Engine/Source/Runtime/AnimGraphRuntime/Private/BoneControllers/AnimNode_SpringBone.cpp b/Engine/Source/Runtime/AnimGraphRuntime/Private/BoneControllers/AnimNode_SpringBone.cpp
--- a/Engine/Source/Runtime/AnimGraphRuntime/Private/BoneControllers/AnimNode_SpringBone.cpp
+++ b/Engine/Source/Runtime/AnimGraphRuntime/Private/BoneControllers/AnimNode_SpringBone.cpp
@@ -4,6 +4,7 @@
 #include "Components/SkeletalMeshComponent.h"
 #include "GameFramework/WorldSettings.h"
 #include "Animation/AnimInstanceProxy.h"
+#include "Animation/AnimStats.h"
 #include "Animation/AnimTrace.h"
 
 #include UE_INLINE_GENERATED_CPP_BY_NAME(AnimNode_SpringBone)
diff --git a/Engine/Source/Runtime/Engine/Classes/Animation/AnimNode_CustomProperty.h b/Engine/Source/Runtime/Engine/Classes/Animation/AnimNode_CustomProperty.h
--- a/Engine/Source/Runtime/Engine/Classes/Animation/AnimNode_CustomProperty.h
+++ b/Engine/Source/Runtime/Engine/Classes/Animation/AnimNode_CustomProperty.h
@@ -6,6 +6,7 @@
 #include "UObject/ObjectMacros.h"
 #include "Templates/SubclassOf.h"
 #include "Animation/AnimNodeBase.h"
+#include "Animation/AnimStats.h"
 #include "AnimNode_CustomProperty.generated.h"
 
 /** 
@@ -108,7 +109,7 @@ protected:
 #if ANIMNODE_STATS_VERBOSE
 	// Cached StatID for this node
 	TStatId StatID;
-	virtual void InitializeStatID() { StatID = FDynamicStats::CreateStatId<FStatGroup_STATGROUP_Anim>(FString(TEXT("Unknown"))); }
+	virtual void InitializeStatID(); { StatID = FDynamicStats::CreateStatId<FStatGroup_STATGROUP_Anim>(FString(TEXT("Unknown"))); }
 #endif // ANIMNODE_STATS_VERBOSE
 
 	friend class UAnimGraphNode_CustomProperty;
diff --git a/Engine/Source/Runtime/Engine/Private/Animation/AnimNodeSpaceConversions.cpp b/Engine/Source/Runtime/Engine/Private/Animation/AnimNodeSpaceConversions.cpp
--- a/Engine/Source/Runtime/Engine/Private/Animation/AnimNodeSpaceConversions.cpp
+++ b/Engine/Source/Runtime/Engine/Private/Animation/AnimNodeSpaceConversions.cpp
@@ -1,6 +1,7 @@
 // Copyright Epic Games, Inc. All Rights Reserved.
 
 #include "Animation/AnimNodeSpaceConversions.h"
+#include "Animation/AnimStats.h"
 
 #include UE_INLINE_GENERATED_CPP_BY_NAME(AnimNodeSpaceConversions)
 
diff --git a/Engine/Source/Runtime/Engine/Private/Animation/AnimNode_ApplyMeshSpaceAdditive.cpp b/Engine/Source/Runtime/Engine/Private/Animation/AnimNode_ApplyMeshSpaceAdditive.cpp
--- a/Engine/Source/Runtime/Engine/Private/Animation/AnimNode_ApplyMeshSpaceAdditive.cpp
+++ b/Engine/Source/Runtime/Engine/Private/Animation/AnimNode_ApplyMeshSpaceAdditive.cpp
@@ -3,6 +3,7 @@
 #include "Animation/AnimNode_ApplyMeshSpaceAdditive.h"
 #include "Animation/AnimInstanceProxy.h"
 #include "Animation/AnimationPoseData.h"
+#include "Animation/AnimStats.h"
 #include "Animation/ExposedValueHandler.h"
 #include "AnimationRuntime.h"
 
diff --git a/Engine/Source/Runtime/Engine/Private/Animation/AnimNode_Inertialization.cpp b/Engine/Source/Runtime/Engine/Private/Animation/AnimNode_Inertialization.cpp
--- a/Engine/Source/Runtime/Engine/Private/Animation/AnimNode_Inertialization.cpp
+++ b/Engine/Source/Runtime/Engine/Private/Animation/AnimNode_Inertialization.cpp
@@ -3,6 +3,7 @@
 #include "Animation/AnimNode_Inertialization.h"
 #include "Animation/AnimInstanceProxy.h"
 #include "Animation/AnimNode_SaveCachedPose.h"
+#include "Animation/AnimStats.h"
 #include "Animation/BlendProfile.h"
 #include "Algo/MaxElement.h"
 #include "Components/SkeletalMeshComponent.h"
diff --git a/Engine/Source/Runtime/Engine/Private/Animation/AnimNode_SaveCachedPose.cpp b/Engine/Source/Runtime/Engine/Private/Animation/AnimNode_SaveCachedPose.cpp
--- a/Engine/Source/Runtime/Engine/Private/Animation/AnimNode_SaveCachedPose.cpp
+++ b/Engine/Source/Runtime/Engine/Private/Animation/AnimNode_SaveCachedPose.cpp
@@ -1,6 +1,7 @@
 // Copyright Epic Games, Inc. All Rights Reserved.
 
 #include "Animation/AnimNode_SaveCachedPose.h"
+#include "Animation/AnimStats.h"
 #include "Animation/AnimInstanceProxy.h"
 
 #include UE_INLINE_GENERATED_CPP_BY_NAME(AnimNode_SaveCachedPose)
