Substrate - make sure we disable all feature when compiling but it does not fit in memory.

diff --git a/Engine/Source/Runtime/Engine/Private/Materials/HLSLMaterialTranslator.cpp b/Engine/Source/Runtime/Engine/Private/Materials/HLSLMaterialTranslator.cpp
--- a/Engine/Source/Runtime/Engine/Private/Materials/HLSLMaterialTranslator.cpp
+++ b/Engine/Source/Runtime/Engine/Private/Materials/HLSLMaterialTranslator.cpp
@@ -11558,7 +11558,9 @@ bool FHLSLMaterialTranslator::FStrataCompilationContext::StrataGenerateDerivedMa
 
 	if (!StrataSimplificationStatus.bRunFullSimplification)
 	{
-		// Generate LayerDepth value for all operators/bsdfs
+		//
+		// Generate LayerDepth value for all operators/bsdfs for progressive material simplification
+		//
 		int VOpTopBranchCountTaken = 0;
 		int VOpBottomBranchCountTaken = 0;
 		std::function<void(FStrataOperator&)> WalkOperatorsForDepth = [&](FStrataOperator& CurrentOperator) -> void
@@ -11637,16 +11639,32 @@ bool FHLSLMaterialTranslator::FStrataCompilationContext::StrataGenerateDerivedMa
 	bool bFirstLoop = true;
 	do 
 	{
-		if (!bFirstLoop && !StrataSimplificationStatus.bRunFullSimplification && !StrataSimplificationStatus.bMaterialFitsInMemoryBudget && StrataSimplificationStatus.OperatorSimplificationOrder.Num() > 0)
+		if (!bFirstLoop && !StrataSimplificationStatus.bFullSimplificationStepHasBeenRun && !StrataSimplificationStatus.bMaterialFitsInMemoryBudget && StrataSimplificationStatus.OperatorSimplificationOrder.Num() > 0)
 		{
 			// Mark the deepest operator for parameter blending
 			FStrataSimplificationStatus::FOperatorToSimplify& OperatorToSimplify = StrataSimplificationStatus.OperatorSimplificationOrder.Top();
 			StrataMaterialExpressionRegisteredOperators[OperatorToSimplify.Data.Index].bNodeRequestParameterBlending = true;
 			StrataSimplificationStatus.OperatorSimplificationOrder.Pop(/*bAllowShrinking*/false);
+
+			// Mark that this is similar to have run full simplification
+			StrataSimplificationStatus.bFullSimplificationStepHasBeenRun |= StrataSimplificationStatus.OperatorSimplificationOrder.Num() == 0;
 		}
-		else if (!bFirstLoop && !StrataSimplificationStatus.bRunFullSimplification && !StrataSimplificationStatus.bMaterialFitsInMemoryBudget && StrataSimplificationStatus.OperatorSimplificationOrder.Num() == 0)
+		else if (!bFirstLoop && StrataSimplificationStatus.bFullSimplificationStepHasBeenRun && !StrataSimplificationStatus.bMaterialFitsInMemoryBudget && StrataSimplificationStatus.OperatorSimplificationOrder.Num() == 0 && !StrataSimplificationStatus.bSlabSimplificationStepHasBeenRun)
 		{
-			StrataSimplificationStatus.bRunFullSimplification = !StrataSimplificationStatus.bMaterialFitsInMemoryBudget;
+			for (auto& It : StrataMaterialExpressionRegisteredOperators)
+			{
+				// Disable all optional features for now to fit.
+				// STRATA_TODO we will need to refine that to account for platforms supporting SSS for instance.
+				It.bBSDFHasSSS = false;
+				It.bBSDFHasMFPPluggedIn = false;
+				It.bBSDFHasEdgeColor = false;
+				It.bBSDFHasFuzz = false;
+				It.bBSDFHasSecondRoughnessOrSimpleClearCoat = false;
+				It.bBSDFHasAnisotropy = false;
+				It.bBSDFHasGlint = false;
+				It.bBSDFHasSpecularProfile = false;
+			}
+			StrataSimplificationStatus.bSlabSimplificationStepHasBeenRun = true;
 		}
 		else if (bFirstLoop)
 		{
@@ -12153,7 +12171,7 @@ bool FHLSLMaterialTranslator::FStrataCompilationContext::StrataGenerateDerivedMa
 			}
 
 			StrataSimplificationStatus.bMaterialFitsInMemoryBudget = StrataMaterialRequestedSizeByte <= StrataBytePerPixel;
-			if (!StrataSimplificationStatus.bMaterialFitsInMemoryBudget && StrataSimplificationStatus.bRunFullSimplification)
+			if (!StrataSimplificationStatus.bMaterialFitsInMemoryBudget && StrataSimplificationStatus.bFullSimplificationStepHasBeenRun && StrataSimplificationStatus.bSlabSimplificationStepHasBeenRun)
 			{
 				// If we have already run the full simplification but the material still does not fit in memory, we must fail the material compilation.
 				Compiler->Errorf(TEXT("Material %s could not be simplified to fit in strata per pixel (asset: %s).\r\n"), *CompilerMaterial->GetDebugName(), *CompilerMaterial->GetAssetPath().ToString());
@@ -12164,6 +12182,7 @@ bool FHLSLMaterialTranslator::FStrataCompilationContext::StrataGenerateDerivedMa
 				// Record the original requested byte size before simplification, only for the first pass.
 				StrataSimplificationStatus.OriginalRequestedByteSize = StrataMaterialRequestedSizeByte;
 			}
+			StrataSimplificationStatus.bFullSimplificationStepHasBeenRun |= StrataSimplificationStatus.bRunFullSimplification;
 
 			const uint32 RequestedSizeInUint = FMath::DivideAndRoundUp(StrataMaterialRequestedSizeByte, 4u);
 			check(RequestedSizeInUint < 256u);
diff --git a/Engine/Source/Runtime/Engine/Private/Materials/HLSLMaterialTranslator.h b/Engine/Source/Runtime/Engine/Private/Materials/HLSLMaterialTranslator.h
--- a/Engine/Source/Runtime/Engine/Private/Materials/HLSLMaterialTranslator.h
+++ b/Engine/Source/Runtime/Engine/Private/Materials/HLSLMaterialTranslator.h
@@ -436,6 +436,9 @@ protected:
 
 		uint32 OriginalRequestedByteSize = 0;
 		bool bRunFullSimplification = false;	// Simple implementation for now: if the material does not fit, we simply everything.
+		bool bFullSimplificationStepHasBeenRun = false;
+
+		bool bSlabSimplificationStepHasBeenRun = false;
 
 		struct FOperatorToSimplify
 		{
diff --git a/Engine/Source/Runtime/Engine/Private/Materials/MaterialExpressions.cpp b/Engine/Source/Runtime/Engine/Private/Materials/MaterialExpressions.cpp
--- a/Engine/Source/Runtime/Engine/Private/Materials/MaterialExpressions.cpp
+++ b/Engine/Source/Runtime/Engine/Private/Materials/MaterialExpressions.cpp
@@ -23539,30 +23539,22 @@ UMaterialExpressionStrataSlabBSDF::UMaterialExpressionStrataSlabBSDF(const FObje
 FName CreateSpecularProfileParameterName(USpecularProfile* InProfile);
 int32 UMaterialExpressionStrataSlabBSDF::Compile(class FMaterialCompiler* Compiler, int32 OutputIndex)
 {
-	int32 RoughnessCodeChunk = CompileWithDefaultFloat1(Compiler, Roughness, 0.5f);
-	int32 AnisotropyCodeChunk = CompileWithDefaultFloat1(Compiler, Anisotropy, 0.0f);
+	FStrataOperator& StrataOperator = Compiler->StrataCompilationGetOperator(Compiler->StrataTreeStackGetPathUniqueId());
+
 
 	// As long as both roughness are potentially different, we must take it into account in our encoding.
 	// We also cannot ignore the tangent when using the default Tangent because GetTangentBasis
 	// used in StrataGetBSDFSharedBasis cannot be relied on for smooth tangent used for lighting on any mesh.
-	const bool bHasAnisotropy = HasAnisotropy();
-
-	const bool bHasEdgeColor = HasEdgeColor();
-	const bool bHasFuzz = HasFuzz();
-	const bool bHasFuzzRoughness = HasFuzzRoughness();
-	const bool bHasSecondRoughness = HasSecondRoughness();
-	const bool bHasMFPPluggedIn = HasMFPPluggedIn();
-	const bool bHasSSS = HasSSS();
-	const bool bHasSpecularProfile = HasSpecularProfile();
+	const bool bHasAnisotropy		= StrataOperator.bBSDFHasAnisotropy > 0;
 
 	int32 SSSProfileCodeChunk = INDEX_NONE;
-	if (bHasSSS)
+	if (StrataOperator.bBSDFHasSSS > 0)
 	{
 		SSSProfileCodeChunk = Compiler->ForceCast(Compiler->ScalarParameter(GetSubsurfaceProfileParameterName(), 1.0f), MCT_Float1);
 	}
 
 	int32 SpecularProfileCodeChunk = INDEX_NONE;
-	if (bHasSpecularProfile)
+	if (StrataOperator.bBSDFHasSpecularProfile > 0)
 	{
 		const FName SpecularProfileParameterName = CreateSpecularProfileParameterName(SpecularProfile);
 		SpecularProfileCodeChunk = Compiler->ForceCast(Compiler->ScalarParameter(SpecularProfileParameterName, 1.0f), MCT_Float1);
@@ -23575,7 +23567,6 @@ int32 UMaterialExpressionStrataSlabBSDF::Compile(class FMaterialCompiler* Compil
 	int32 TangentCodeChunk = bHasAnisotropy ? CompileWithDefaultTangentWS(Compiler, Tangent) : INDEX_NONE;
 	const FStrataRegisteredSharedLocalBasis NewRegisteredSharedLocalBasis = StrataCompilationInfoCreateSharedLocalBasis(Compiler, NormalCodeChunk, TangentCodeChunk);
 
-	FStrataOperator& StrataOperator = Compiler->StrataCompilationGetOperator(Compiler->StrataTreeStackGetPathUniqueId());
 	StrataOperator.BSDFRegisteredSharedLocalBasis = NewRegisteredSharedLocalBasis;
 	
 	int32 ThicknesCodeChunk = INDEX_NONE;
@@ -23590,27 +23581,73 @@ int32 UMaterialExpressionStrataSlabBSDF::Compile(class FMaterialCompiler* Compil
 	}
 	check(ThicknesCodeChunk != INDEX_NONE);
 
+	int32 DiffuseAlbedoCodeChunk		= CompileWithDefaultFloat3(Compiler, DiffuseAlbedo, 0.18f, 0.18f, 0.18f);
+	int32 F0CodeChunk					= CompileWithDefaultFloat3(Compiler, F0, DefaultF0, DefaultF0, DefaultF0);
+	int32 RoughnessCodeChunk			= CompileWithDefaultFloat1(Compiler, Roughness, 0.5f);
+	int32 AnisotropyCodeChunk			= CompileWithDefaultFloat1(Compiler, Anisotropy, 0.0f);
+	int32 F90CodeChunk					= CompileWithDefaultFloat3(Compiler, F90, 1.0f, 1.0f, 1.0f);
+	int32 SSSMFPCodeChunk				= CompileWithDefaultFloat3(Compiler, SSSMFP, 0.0f, 0.0f, 0.0f);
+	int32 SSSMFPScaleCodeChunk			= CompileWithDefaultFloat1(Compiler, SSSMFPScale, 1.0f);
+	int32 SSSPhaseAnisotropyCodeChunk	= CompileWithDefaultFloat1(Compiler, SSSPhaseAnisotropy, 0.0f);
+	int32 SecondRoughnessCodeChunk		= CompileWithDefaultFloat1(Compiler, SecondRoughness, 0.0f);
+	int32 SecondRoughnessWeightCodeChunk= CompileWithDefaultFloat1(Compiler, SecondRoughnessWeight, 0.0f);
+	int32 FuzzAmountCodeChunk			= CompileWithDefaultFloat1(Compiler, FuzzAmount, 0.0f);
+	int32 FuzzColorCodeChunk			= CompileWithDefaultFloat3(Compiler, FuzzColor, 0.0f, 0.0f, 0.0f);
+	int32 FuzzRoughnessCodeChunk		= HasFuzzRoughness() ? CompileWithDefaultFloat1(Compiler, FuzzRoughness, 0.5f) : RoughnessCodeChunk;
+	int32 GlintValueCodeChunk			= CompileWithDefaultFloat1(Compiler, GlintValue, 0.0f);
+	int32 GlintUVCodeChunk				= CompileWithDefaultFloat2(Compiler, GlintUV, 0.0f, 0.0f);
+
+	// Disable some features if requested by the simplification process
+	if (StrataOperator.bBSDFHasMFPPluggedIn == 0)
+	{
+		SSSMFPCodeChunk = Compiler->Constant3(0.0f, 0.0f, 0.0f);
+	}
+	if (StrataOperator.bBSDFHasEdgeColor == 0)
+	{
+		F90CodeChunk = Compiler->Constant3(1.0f, 1.0f, 1.0f);
+	}
+	if (StrataOperator.bBSDFHasFuzz == 0)
+	{
+		FuzzAmountCodeChunk = Compiler->Constant(0.0f);
+	}
+	if (StrataOperator.bBSDFHasSecondRoughnessOrSimpleClearCoat == 0)
+	{
+		SecondRoughnessWeightCodeChunk = Compiler->Constant(0.0f);
+	}
+	if (StrataOperator.bBSDFHasAnisotropy == 0)
+	{
+		AnisotropyCodeChunk = Compiler->Constant(0.0f);
+	}
+	if (StrataOperator.bBSDFHasGlint == 0)
+	{
+		GlintValueCodeChunk = Compiler->Constant(0.0f);
+	}
+	if (StrataOperator.bBSDFHasSpecularProfile == 0)
+	{
+		SpecularProfileCodeChunk = INDEX_NONE;
+	}
+
 	int32 OutputCodeChunk = Compiler->StrataSlabBSDF(
-		CompileWithDefaultFloat3(Compiler, DiffuseAlbedo, 0.18f, 0.18f, 0.18f),
-		CompileWithDefaultFloat3(Compiler, F0, DefaultF0, DefaultF0, DefaultF0),
-		CompileWithDefaultFloat3(Compiler, F90, 1.0f, 1.0f, 1.0f),
+		DiffuseAlbedoCodeChunk,
+		F0CodeChunk,
+		F90CodeChunk,
 		RoughnessCodeChunk,
 		AnisotropyCodeChunk,
-		SSSProfileCodeChunk != INDEX_NONE ? SSSProfileCodeChunk : Compiler->Constant(0.0f),	
-		CompileWithDefaultFloat3(Compiler, SSSMFP, 0.0f, 0.0f, 0.0f),
-		CompileWithDefaultFloat1(Compiler, SSSMFPScale, 1.0f),
-		CompileWithDefaultFloat1(Compiler, SSSPhaseAnisotropy, 0.0f),
+		SSSProfileCodeChunk != INDEX_NONE ? SSSProfileCodeChunk : Compiler->Constant(0.0f),
+		SSSMFPCodeChunk,
+		SSSMFPScaleCodeChunk,
+		SSSPhaseAnisotropyCodeChunk,
 		bUseSSSDiffusion ? Compiler->Constant(1.0f) : Compiler->Constant(0.0f),
 		CompileWithDefaultFloat3(Compiler, EmissiveColor, 0.0f, 0.0f, 0.0f),
-		CompileWithDefaultFloat1(Compiler, SecondRoughness, 0.0f),
-		CompileWithDefaultFloat1(Compiler, SecondRoughnessWeight, 0.0f),
+		SecondRoughnessCodeChunk,
+		SecondRoughnessWeightCodeChunk,
 		Compiler->Constant(0.0f),										// SecondRoughnessAsSimpleClearCoat
-		CompileWithDefaultFloat1(Compiler, FuzzAmount, 0.0f),
-		CompileWithDefaultFloat3(Compiler, FuzzColor, 0.0f, 0.0f, 0.0f),
-		bHasFuzzRoughness ? CompileWithDefaultFloat1(Compiler, FuzzRoughness, 0.5f) : RoughnessCodeChunk,
+		FuzzAmountCodeChunk,
+		FuzzColorCodeChunk,
+		FuzzRoughnessCodeChunk,
 		ThicknesCodeChunk,
-		CompileWithDefaultFloat1(Compiler, GlintValue, 0.0f),
-		CompileWithDefaultFloat2(Compiler, GlintUV, 0.0f, 0.0f),
+		GlintValueCodeChunk,
+		GlintUVCodeChunk,
 		SpecularProfileCodeChunk != INDEX_NONE ? SpecularProfileCodeChunk : Compiler->Constant(0.0f),
 		StrataOperator.bIsBottom > 0 ? true : false,
 		NormalCodeChunk,
