Fix a compiler warning in MaterialRenderProxy about bitfield not preserving sign.

- clang 16.0.6 produces this (-Wsingle-bit-bitfield-constant-conversion)

diff --git a/Engine/Source/Runtime/Engine/Public/Materials/MaterialRenderProxy.h b/Engine/Source/Runtime/Engine/Public/Materials/MaterialRenderProxy.h
--- a/Engine/Source/Runtime/Engine/Public/Materials/MaterialRenderProxy.h
+++ b/Engine/Source/Runtime/Engine/Public/Materials/MaterialRenderProxy.h
@@ -237,11 +237,11 @@ private:
 	mutable int32 UniformExpressionCacheSerialNumber = 0;
 
 	/** For tracking down a bug accessing a deleted proxy. */
-	mutable int8 MarkedForGarbageCollection : 1;
-	mutable int8 DeletedFlag : 1;
-	mutable int8 ReleaseResourceFlag : 1;
+	mutable uint8 MarkedForGarbageCollection : 1;
+	mutable uint8 DeletedFlag : 1;
+	mutable uint8 ReleaseResourceFlag : 1;
 	/** If any VT producer destroyed callbacks have been registered */
-	mutable int8 HasVirtualTextureCallbacks : 1;
+	mutable uint8 HasVirtualTextureCallbacks : 1;
 
 #if WITH_EDITOR
 	/**
