Refactor the UPlayerInput::ProcessInputStack function to separate out the evaluation of the key map fromt he broadcasting of delegates.

This will allow us to change the consumption state of Keys or Delegate bindings from subclasses such as Enhanced Input, which means that we can mark Input Actions to explictly block any legacy key mappings, gameplay debugger keybindings, etc.

diff --git a/Engine/Plugins/EnhancedInput/Source/EnhancedInput/Private/EnhancedInputSubsystemInterface.cpp b/Engine/Plugins/EnhancedInput/Source/EnhancedInput/Private/EnhancedInputSubsystemInterface.cpp
--- a/Engine/Plugins/EnhancedInput/Source/EnhancedInput/Private/EnhancedInputSubsystemInterface.cpp
+++ b/Engine/Plugins/EnhancedInput/Source/EnhancedInput/Private/EnhancedInputSubsystemInterface.cpp
@@ -827,6 +827,7 @@ void IEnhancedInputSubsystemInterface::RebuildControlMappings()
 	// Clear existing mappings, but retain the mapping array for later processing
 	TArray<FEnhancedActionKeyMapping> OldMappings(MoveTemp(PlayerInput->EnhancedActionMappings));
 	PlayerInput->ClearAllMappings();
+	PlayerInput->KeyConsumptionData.Reset();
 	AppliedContextRedirects.Reset();
 
 	// Order contexts by priority
@@ -923,6 +924,13 @@ void IEnhancedInputSubsystemInterface::RebuildControlMappings()
 			{
 				Mapping.Key = PlayerMappedKey;
 
+				// If this Input Action is flagged to consume input, then mark it's key state as being consumed every tick.
+				// This has the affect where the base UPlayerInput class will not fire any legacy bindings
+				if (Mapping.Action->bConsumesActionAndAxisMappings)
+				{
+					PlayerInput->KeyConsumptionData.Add(Mapping.Action, { Mapping.Key, static_cast<ETriggerEvent>(Mapping.Action->TriggerEventsThatConsumeLegacyKeys) });	
+				}
+				
 				if (Mapping.Action && !AppliedKeys.Contains(Mapping.Key))
 				{
 					// TODO: Wasteful query as we've already established chord state within ReorderMappings. Store TOptional bConsumeInput per mapping, allowing override? Query override via delegate?
@@ -931,7 +939,7 @@ void IEnhancedInputSubsystemInterface::RebuildControlMappings()
 						return Cast<const UInputTriggerChordAction>(Trigger) != nullptr;
 					};
 					bool bHasActionChords = HasTriggerWith(IsChord, Mapping.Action->Triggers);
-					bool bHasChords = HasTriggerWith(IsChord, Mapping.Triggers) || bHasActionChords;
+					bool bHasChords = bHasActionChords || HasTriggerWith(IsChord, Mapping.Triggers);
 
 					// Chorded actions can't consume input or they would hide the action they are chording.
 					if (!bHasChords && Mapping.Action->bConsumeInput)
diff --git a/Engine/Plugins/EnhancedInput/Source/EnhancedInput/Private/EnhancedPlayerInput.cpp b/Engine/Plugins/EnhancedInput/Source/EnhancedInput/Private/EnhancedPlayerInput.cpp
--- a/Engine/Plugins/EnhancedInput/Source/EnhancedInput/Private/EnhancedPlayerInput.cpp
+++ b/Engine/Plugins/EnhancedInput/Source/EnhancedInput/Private/EnhancedPlayerInput.cpp
@@ -277,48 +277,48 @@ float UEnhancedPlayerInput::GetEffectiveTimeDilation() const
 	return 1.0f;
 }
 
-void UEnhancedPlayerInput::ProcessInputStack(const TArray<UInputComponent*>& InputComponentStack, const float DeltaTime, const bool bGamePaused)
+void UEnhancedPlayerInput::EvaluateKeyMapState(const float DeltaTime, const bool bGamePaused, OUT TArray<TPair<FKey, FKeyState*>>& KeysWithEvents)
 {
-	// We need to grab the down states of all keys before calling Super::ProcessInputStack as it will leave bDownPrevious in the same state as bDown (i.e. this frame, not last).
-	static TMap<FKey, bool> KeyDownPrevious;
+	TRACE_CPUPROFILER_EVENT_SCOPE(EnhPIS_KeyDownPrev);
+
+	const UEnhancedInputDeveloperSettings* Settings = GetDefault<UEnhancedInputDeveloperSettings>();
+	bool bWasAnyKeyDownLastFrame = false;
+
+	KeyDownPrevious.Reset();
+	KeyDownPrevious.Reserve(GetKeyStateMap().Num());
+	for (TPair<FKey, FKeyState>& KeyPair : GetKeyStateMap())
 	{
-		TRACE_CPUPROFILER_EVENT_SCOPE(EnhPIS_KeyDownPrev);
-
-		const UEnhancedInputDeveloperSettings* Settings = GetDefault<UEnhancedInputDeveloperSettings>();
-		bool bWasAnyKeyDownLastFrame = false;
-
-		KeyDownPrevious.Reset();
-		KeyDownPrevious.Reserve(GetKeyStateMap().Num());
-		for (TPair<FKey, FKeyState>& KeyPair : GetKeyStateMap())
-		{
-			const FKeyState& KeyState = KeyPair.Value;
-			// TODO: Can't just use bDownPrevious as paired axis event edges may not fire due to axial deadzoning/missing axis properties. Need to change how this is detected in PlayerInput.cpp.
-			bool bWasDown = KeyState.bDownPrevious || KeyState.EventCounts[IE_Pressed].Num() || KeyState.EventCounts[IE_Repeat].Num();
-			bWasDown |= KeyPair.Key.IsAnalog() && KeyState.RawValue.SizeSquared() != 0;	// Analog inputs should pulse every (non-zero) tick to retain compatibility with UE4.
-
-			// When UPlayerInput::FlushPressedKeys is called any keys that are down will have their RawValue set to 0, and their bDown/bDownPrevious state will be reset to false.
-			// However, their "Value" will not be reset until UPlayerInput::ProcessInputStack. We need to detect if this key was down previously after a flush
-			// so that the Enhanced Input action will correctly fire the triggered values.
-			const bool bKeyWasJustFlushed = 
-				bIsFlushingInputThisFrame && 
-				Settings->bSendTriggeredEventsWhenInputIsFlushed && 
-				!KeyState.Value.IsZero() && 
-				KeyState.RawValue.IsZero() && 
-				!KeyState.bDown;
-
-			bWasDown |= bKeyWasJustFlushed;
+		const FKeyState& KeyState = KeyPair.Value;
+		// TODO: Can't just use bDownPrevious as paired axis event edges may not fire due to axial deadzoning/missing axis properties. Need to change how this is detected in PlayerInput.cpp.
+		bool bWasDown = KeyState.bDownPrevious || KeyState.EventCounts[IE_Pressed].Num() || KeyState.EventCounts[IE_Repeat].Num();
+		bWasDown |= KeyPair.Key.IsAnalog() && KeyState.RawValue.SizeSquared() != 0;	// Analog inputs should pulse every (non-zero) tick to retain compatibility with UE4.
+
+		// When UPlayerInput::FlushPressedKeys is called any keys that are down will have their RawValue set to 0, and their bDown/bDownPrevious state will be reset to false.
+		// However, their "Value" will not be reset until UPlayerInput::ProcessInputStack. We need to detect if this key was down previously after a flush
+		// so that the Enhanced Input action will correctly fire the triggered values.
+		const bool bKeyWasJustFlushed = 
+			bIsFlushingInputThisFrame && 
+			Settings->bSendTriggeredEventsWhenInputIsFlushed && 
+			!KeyState.Value.IsZero() && 
+			KeyState.RawValue.IsZero() && 
+			!KeyState.bDown;
+
+		bWasDown |= bKeyWasJustFlushed;
 			
-			// Keep track of the state of every key so that when we are done iterating we can have a meaningful value for EKeys::AnyKey
-			bWasAnyKeyDownLastFrame |= bWasDown;
-
-			KeyDownPrevious.Emplace(KeyPair.Key, bWasDown);
-		}
+		// Keep track of the state of every key so that when we are done iterating we can have a meaningful value for EKeys::AnyKey
+		bWasAnyKeyDownLastFrame |= bWasDown;
 
-		KeyDownPrevious.Emplace(EKeys::AnyKey, bWasAnyKeyDownLastFrame);
+		KeyDownPrevious.Emplace(KeyPair.Key, bWasDown);
 	}
 
-	Super::ProcessInputStack(InputComponentStack, DeltaTime, bGamePaused);
+	KeyDownPrevious.Emplace(EKeys::AnyKey, bWasAnyKeyDownLastFrame);
+	
+	
+	Super::EvaluateKeyMapState(DeltaTime, bGamePaused, KeysWithEvents);
+}
 
+void UEnhancedPlayerInput::EvaluateInputDelegates(const TArray<UInputComponent*>& InputComponentStack, const float DeltaTime, const bool bGamePaused, const TArray<TPair<FKey, FKeyState*>>& KeysWithEvents)
+{
 	TRACE_CPUPROFILER_EVENT_SCOPE(EnhPIS_Main);
 
 	// Process Action bindings
@@ -601,6 +601,14 @@ void UEnhancedPlayerInput::ProcessInputStack(const TArray<UInputComponent*>& Inp
 				// Search for the action instance data a second time as a previous delegate call may have deleted it.
 				if (const FInputActionInstance* ActionData = FindActionInstanceData(DelegateAction))
 				{
+					// If this enhanced input delegate has triggered and is flagged to consume legacy keys, then mark it as such.
+					if (const FKeyConsumptionOptions* ConsumptionData = KeyConsumptionData.Find(ActionData->GetSourceAction()))
+					{
+						if (static_cast<uint8>(ConsumptionData->EventsToCauseConsumption & Delegate->GetTriggerEvent()) != 0)
+						{
+							ConsumeKey(ConsumptionData->KeyToConsume);
+						}
+					}
 					Delegate->Execute(*ActionData);
 				}
 			}
@@ -719,6 +727,13 @@ void UEnhancedPlayerInput::ProcessInputStack(const TArray<UInputComponent*>& Inp
 	LastFrameTime = CurrentTime;
 	KeysPressedThisTick.Reset();
 	bIsFlushingInputThisFrame = false;
+
+	Super::EvaluateInputDelegates(InputComponentStack, DeltaTime, bGamePaused, KeysWithEvents);
+}
+
+void UEnhancedPlayerInput::ProcessInputStack(const TArray<UInputComponent*>& InputComponentStack, const float DeltaTime, const bool bGamePaused)
+{
+	Super::ProcessInputStack(InputComponentStack, DeltaTime, bGamePaused);
 }
 
 void UEnhancedPlayerInput::FlushPressedKeys()
diff --git a/Engine/Plugins/EnhancedInput/Source/EnhancedInput/Public/EnhancedPlayerInput.h b/Engine/Plugins/EnhancedInput/Source/EnhancedInput/Public/EnhancedPlayerInput.h
--- a/Engine/Plugins/EnhancedInput/Source/EnhancedInput/Public/EnhancedPlayerInput.h
+++ b/Engine/Plugins/EnhancedInput/Source/EnhancedInput/Public/EnhancedPlayerInput.h
@@ -26,6 +26,18 @@ struct FInjectedInput
 	TArray<UInputModifier*> Modifiers;
 };
 
+USTRUCT()
+struct FKeyConsumptionOptions
+{
+	GENERATED_BODY()
+	
+	/** The key that should be consumed if the trigger state is reached */
+	FKey KeyToConsume = EKeys::Invalid;
+		
+	/** A bitmask of trigger events that when reached, should cause the key to be marked as consumed. */
+	ETriggerEvent EventsToCauseConsumption = ETriggerEvent::None;
+};
+
 USTRUCT()
 struct FInjectedInputArray
 {
@@ -75,6 +87,9 @@ public:
 	float GetEffectiveTimeDilation() const;
 	
 protected:
+
+	virtual void EvaluateKeyMapState(const float DeltaTime, const bool bGamePaused, OUT TArray<TPair<FKey, FKeyState*>>& KeysWithEvents) override;
+	virtual void EvaluateInputDelegates(const TArray<UInputComponent*>& InputComponentStack, const float DeltaTime, const bool bGamePaused, const TArray<TPair<FKey, FKeyState*>>& KeysWithEvents) override;
 	
 	// Causes key to be consumed if it is affecting an action.
 	virtual bool IsKeyHandledByAction(FKey Key) const override;
@@ -84,7 +99,10 @@ protected:
 
 	/** This player's version of the Action Mappings */
 	const TArray<FEnhancedActionKeyMapping>& GetEnhancedActionMappings() const { return EnhancedActionMappings; }
-	
+
+	/** Array of data that represents what keys should be consumed if an enhanced input action is in a specific triggered state */
+	UPROPERTY()
+	TMap<TObjectPtr<const UInputAction>, FKeyConsumptionOptions> KeyConsumptionData;
 
 private:
 
@@ -165,6 +183,9 @@ private:
 
 protected:
 
+	// We need to grab the down states of all keys before calling Super::ProcessInputStack as it will leave bDownPrevious in the same state as bDown (i.e. this frame, not last).
+	TMap<FKey, bool> KeyDownPrevious;
+	
 	/** 
 	* If true, then FlushPressedKeys has been called and the input key state map has been flushed.
 	* 
diff --git a/Engine/Plugins/EnhancedInput/Source/EnhancedInput/Public/InputAction.h b/Engine/Plugins/EnhancedInput/Source/EnhancedInput/Public/InputAction.h
--- a/Engine/Plugins/EnhancedInput/Source/EnhancedInput/Public/InputAction.h
+++ b/Engine/Plugins/EnhancedInput/Source/EnhancedInput/Public/InputAction.h
@@ -85,18 +85,30 @@ public:
 	UPROPERTY(EditAnywhere, BlueprintReadOnly, Category = "Description")
 	FText ActionDescription = FText::GetEmpty();
 
+	// Should this action be able to trigger whilst the game is paused - Replaces bExecuteWhenPaused
+    UPROPERTY(EditAnywhere, BlueprintReadOnly, Category = Action)
+    bool bTriggerWhenPaused = false;
+	
 	// Should this action swallow any inputs bound to it or allow them to pass through to affect lower priority bound actions?
-	UPROPERTY(EditAnywhere, BlueprintReadWrite, Category = Action)
+	UPROPERTY(EditAnywhere, BlueprintReadWrite, Category = "Input Consumption", meta=(DisplayName="Consume Lower Priority Enhanced Input Mappings"))
 	bool bConsumeInput = true;
 
-	// Should this action be able to trigger whilst the game is paused - Replaces bExecuteWhenPaused
-	UPROPERTY(EditAnywhere, BlueprintReadOnly, Category = Action)
-	bool bTriggerWhenPaused = false;
+	/**
+	 * Should this Input Action consume any legacy Action and Axis key mappings?
+	 * If true, then any key mapping to this input action will consume(aka block) the legacy key
+	 * mapping from firing delegates.
+	 */
+	UPROPERTY(EditAnywhere, BlueprintReadWrite, Category = "Input Consumption")
+	bool bConsumesActionAndAxisMappings = false;
 
 	// This action's mappings are not intended to be automatically overridden by higher priority context mappings. Users must explicitly remove the mapping first. NOTE: It is the responsibility of the author of the mapping code to enforce this!
 	UPROPERTY(EditAnywhere, BlueprintReadOnly, Category = Action)
 	bool bReserveAllMappings = false;	// TODO: Need something more complex than this?
 
+	/** A bitmask of trigger events that, when reached, will consume any FKeys mapped to this input action. */
+	UPROPERTY(EditAnywhere, BlueprintReadWrite, Category = "Input Consumption", meta=(EditCondition = "bConsumesActionAndAxisMappings", Bitmask, BitmaskEnum="/Script/EnhancedInput.ETriggerEvent"))
+	int32 TriggerEventsThatConsumeLegacyKeys;
+	
 	// The type that this action returns from a GetActionValue query or action event
 	UPROPERTY(EditAnywhere, BlueprintReadOnly, Category = Action, AssetRegistrySearchable)
 	EInputActionValueType ValueType = EInputActionValueType::Boolean;
diff --git a/Engine/Source/Runtime/Engine/Classes/GameFramework/PlayerInput.h b/Engine/Source/Runtime/Engine/Classes/GameFramework/PlayerInput.h
--- a/Engine/Source/Runtime/Engine/Classes/GameFramework/PlayerInput.h
+++ b/Engine/Source/Runtime/Engine/Classes/GameFramework/PlayerInput.h
@@ -606,6 +606,22 @@ public:
 	/** Process the frame's input events given the current input component stack. */
 	ENGINE_API virtual void ProcessInputStack(const TArray<UInputComponent*>& InputComponentStack, const float DeltaTime, const bool bGamePaused);
 
+protected:
+	
+	/**
+	 * Evaluates the current state of the KeyStateMap this tick.
+	 * Should only be called from UPlayerInput::ProcessInputStack
+	 */
+	ENGINE_API virtual void EvaluateKeyMapState(const float DeltaTime, const bool bGamePaused, OUT TArray<TPair<FKey, FKeyState*>>& KeysWithEvents);
+
+	/**
+	 * Evaluates the state of any bound input delegates on the given Input Component Stack.
+	 * This should only be called from UPlayerInput::ProcessInputStack
+	 */
+	ENGINE_API virtual void EvaluateInputDelegates(const TArray<UInputComponent*>& InputComponentStack, const float DeltaTime, const bool bGamePaused, const TArray<TPair<FKey, FKeyState*>>& KeysWithEvents);
+	
+public:
+
 	/** Rather than processing input, consume it and discard without doing anything useful with it.  Like calling ProcessInputStack() and ignoring all results. */
 	ENGINE_API void DiscardPlayerInput();
 
@@ -756,8 +772,11 @@ private:
 	/* Returns the summed values of all the components of this axis this frame
 	 * @param AxisBinding - the action to determine if it ocurred
 	 * @param KeysToConsume - array to collect the keys associated with this binding that should be consumed
+	 * @param bHadAnyNonConsumedKeys - If false, then all keys related to the axis binding have been consumed and there were none that could be used to determine the value
+	 *
+	 * @return The value of the axis based on the all the key mappings to it added together
 	 */
-	ENGINE_API float DetermineAxisValue(const FInputAxisBinding& AxisBinding, const bool bGamePaused, TArray<FKey>& KeysToConsume);
+	ENGINE_API float DetermineAxisValue(const FInputAxisBinding& AxisBinding, const bool bGamePaused, TArray<FKey>& KeysToConsume, OUT bool& bHadAnyNonConsumedKeys) const;
 
 	/** Utility function to ensure the key mapping cache maps are built */
 	FORCEINLINE void ConditionalBuildKeyMappings() const
diff --git a/Engine/Source/Runtime/Engine/Private/UserInterface/PlayerInput.cpp b/Engine/Source/Runtime/Engine/Private/UserInterface/PlayerInput.cpp
--- a/Engine/Source/Runtime/Engine/Private/UserInterface/PlayerInput.cpp
+++ b/Engine/Source/Runtime/Engine/Private/UserInterface/PlayerInput.cpp
@@ -37,6 +37,17 @@
 TArray<FInputActionKeyMapping> UPlayerInput::EngineDefinedActionMappings;
 TArray<FInputAxisKeyMapping> UPlayerInput::EngineDefinedAxisMappings;
 
+namespace UE
+{
+	namespace Input
+	{
+		static bool AxisEventsCanBeConsumed = true;
+		static FAutoConsoleVariableRef CVarShouldOnlyTriggerLastActionInChord(TEXT("Input.AxisEventsCanBeConsumed"),
+			AxisEventsCanBeConsumed,
+			TEXT("If true and all FKey's for a given Axis Event are consumed, then the axis delegate will not fire."));
+	}
+}
+
 /** Runtime struct that gathers up the different kinds of delegates that might be issued */
 struct FDelegateDispatchDetails
 {
@@ -1024,11 +1035,12 @@
 	}
 }
 
-float UPlayerInput::DetermineAxisValue(const FInputAxisBinding& AxisBinding, const bool bGamePaused, TArray<FKey>& KeysToConsume)
+float UPlayerInput::DetermineAxisValue(const FInputAxisBinding& AxisBinding, const bool bGamePaused, TArray<FKey>& KeysToConsume, OUT bool& bHadAnyNonConsumedKeys) const
 {
 	ConditionalBuildKeyMappings();
 
-	float AxisValue = 0.f;
+	float AxisValue = 0.0f;
+	bHadAnyNonConsumedKeys = false;
 
 	FAxisKeyDetails* KeyDetails = AxisKeyMap.Find(AxisBinding.AxisName);
 	if (KeyDetails)
@@ -1047,6 +1059,7 @@
 				{
 					KeysToConsume.AddUnique(KeyMapping.Key);
 				}
+				bHadAnyNonConsumedKeys = true;
 			}
 		}
 
@@ -1092,41 +1105,7 @@
 {
 	ConditionalBuildKeyMappings();
 
-	// We collect axis contributions by delegate, so we can sum up
-	// contributions from multiple bindings.
-	struct FAxisDelegateDetails
-	{
-		FInputAxisUnifiedDelegate Delegate;
-		float Value;
-
-		FAxisDelegateDetails(FInputAxisUnifiedDelegate InDelegate, const float InValue)
-			: Delegate(MoveTemp(InDelegate))
-			, Value(InValue)
-		{
-		}
-	};
-	struct FVectorAxisDelegateDetails
-	{
-		FInputVectorAxisUnifiedDelegate Delegate;
-		FVector Value;
-
-		FVectorAxisDelegateDetails(FInputVectorAxisUnifiedDelegate InDelegate, const FVector InValue)
-			: Delegate(MoveTemp(InDelegate))
-			, Value(InValue)
-		{
-		}
-	};
-
-	static TArray<FAxisDelegateDetails> AxisDelegates;
-	static TArray<FVectorAxisDelegateDetails> VectorAxisDelegates;
-	static TArray<FDelegateDispatchDetails> NonAxisDelegates;
-	static TArray<FKey> KeysToConsume;
-	static TArray<FDelegateDispatchDetails> FoundChords;
 	static TArray<TPair<FKey, FKeyState*>> KeysWithEvents;
-	static TArray<TSharedPtr<FInputActionBinding>> PotentialActions;
 
-	// must be called non-recursively and on the game thread
-	check(IsInGameThread() && !AxisDelegates.Num() && !VectorAxisDelegates.Num() && !NonAxisDelegates.Num() && !KeysToConsume.Num() && !FoundChords.Num() && !EventIndices.Num() && !KeysWithEvents.Num() && !PotentialActions.Num());
-
 	// @todo, if input is coming in asynchronously, we should buffer anything that comes in during playerinput() and not include it
 	// in this processing batch
@@ -1136,7 +1115,29 @@
 	{
 		PlayerController->PreProcessInput(DeltaTime, bGamePaused);
 	}
+	
+	// Evaluate the current state of the key map. This will take the event accumulators and put them into the actual values of the key state map.
+	EvaluateKeyMapState(DeltaTime, bGamePaused, OUT KeysWithEvents);
+	
+	// Determine potential delegates
+	EvaluateInputDelegates(InputComponentStack, DeltaTime, bGamePaused, KeysWithEvents);
 
+	if (PlayerController)
+	{
+		PlayerController->PostProcessInput(DeltaTime, bGamePaused);
+	}
+	
+	FinishProcessingPlayerInput();
+
+	TouchEventLocations.Reset();
+	KeysWithEvents.Reset();
+}
+
+void UPlayerInput::EvaluateKeyMapState(const float DeltaTime, const bool bGamePaused, OUT TArray<TPair<FKey, FKeyState*>>& KeysWithEvents)
+{
+	// Must be called non-recursively on the game thread
+	check(IsInGameThread() && !KeysWithEvents.Num());
+	
 	// copy data from accumulators to the real values
 	for (TMap<FKey,FKeyState>::TIterator It(KeyStateMap); It; ++It)
 	{
@@ -1210,7 +1211,35 @@
 		KeyState->PairSampledAxes = 0;
 	}
 	EventCount = 0;
+}
 
+void UPlayerInput::EvaluateInputDelegates(const TArray<UInputComponent*>& InputComponentStack, const float DeltaTime, const bool bGamePaused, const TArray<TPair<FKey, FKeyState*>>& KeysWithEvents)
+{
+	// We collect axis contributions by delegate, so we can sum up
+	// contributions from multiple bindings.
+	struct FAxisDelegateDetails
+	{
+		FInputAxisUnifiedDelegate Delegate;
+		float Value;
+
+		FAxisDelegateDetails(FInputAxisUnifiedDelegate InDelegate, const float InValue)
+			: Delegate(MoveTemp(InDelegate))
+			, Value(InValue)
+		{
+		}
+	};
+	struct FVectorAxisDelegateDetails
+	{
+		FInputVectorAxisUnifiedDelegate Delegate;
+		FVector Value;
+
+		FVectorAxisDelegateDetails(FInputVectorAxisUnifiedDelegate InDelegate, const FVector InValue)
+			: Delegate(MoveTemp(InDelegate))
+			, Value(InValue)
+		{
+		}
+	};
+	
 	struct FDelegateDispatchDetailsSorter
 	{
 		bool operator()( const FDelegateDispatchDetails& A, const FDelegateDispatchDetails& B ) const
@@ -1219,6 +1248,17 @@
 		}
 	};
 
+	
+	static TArray<FAxisDelegateDetails> AxisDelegates;
+	static TArray<FVectorAxisDelegateDetails> VectorAxisDelegates;
+	static TArray<FDelegateDispatchDetails> NonAxisDelegates;
+	static TArray<FKey> KeysToConsume;
+	static TArray<FDelegateDispatchDetails> FoundChords;	
+	static TArray<TSharedPtr<FInputActionBinding>> PotentialActions;
+
+	// must be called non-recursively and on the game thread
+	check(IsInGameThread() && !AxisDelegates.Num() && !VectorAxisDelegates.Num() && !NonAxisDelegates.Num() && !KeysToConsume.Num() && !FoundChords.Num() && !EventIndices.Num() && !PotentialActions.Num());
+
 	int32 StackIndex = InputComponentStack.Num()-1;
 
 	// Walk the stack, top to bottom
@@ -1343,11 +1383,15 @@
 				EventIndices.Reset();
 			}
 
+			// A flag that is set in DetermineAxisValue. If false, then all keys related to the axis binding have been consumed.
+			bool bHadAnyKeys = false;
+			
 			// Run though game axis bindings and accumulate axis values
 			for (FInputAxisBinding& AB : IC->AxisBindings)
 			{
-				AB.AxisValue = DetermineAxisValue(AB, bGamePaused, KeysToConsume);
-				if (AB.AxisDelegate.IsBound())
+				AB.AxisValue = DetermineAxisValue(AB, bGamePaused, KeysToConsume, bHadAnyKeys); 
+				
+				if ((bHadAnyKeys || !UE::Input::AxisEventsCanBeConsumed) && AB.AxisDelegate.IsBound())
 				{
 					AxisDelegates.Emplace(FAxisDelegateDetails(AB.AxisDelegate, AB.AxisValue));
 				}
@@ -1473,17 +1517,9 @@
 		}
 	}
 
-	if (PlayerController)
-	{
-		PlayerController->PostProcessInput(DeltaTime, bGamePaused);
-	}
-	
-	FinishProcessingPlayerInput();
 	AxisDelegates.Reset();
 	VectorAxisDelegates.Reset();
 	NonAxisDelegates.Reset();
-	TouchEventLocations.Reset();
-	KeysWithEvents.Reset();
 }
 
 void UPlayerInput::DiscardPlayerInput()
