Fix uninitialized read causing selection outline color flickering

diff --git a/Engine/Source/Runtime/Engine/Private/PrimitiveSceneProxy.cpp b/Engine/Source/Runtime/Engine/Private/PrimitiveSceneProxy.cpp
--- a/Engine/Source/Runtime/Engine/Private/PrimitiveSceneProxy.cpp
+++ b/Engine/Source/Runtime/Engine/Private/PrimitiveSceneProxy.cpp
@@ -379,6 +379,7 @@ FPrimitiveSceneProxy::FPrimitiveSceneProxy(const UPrimitiveComponent* InComponen
 #if WITH_EDITOR
 // by default we are always drawn
 ,	HiddenEditorViews(0)
+,   SelectionOutlineColorIndex(0)
 ,	DrawInAnyEditMode(0)
 ,   bIsFoliage(false)
 #endif
