Fixed garbage object reference in ChildActorComponent::ChildActor on clients when the child actor is destroyed. Made clients register for EndPlay (and changing the destroy callback to receive EndPlay instead as a catch-all to other destruction paths).

diff --git a/Engine/Source/Runtime/Engine/Classes/Components/ChildActorComponent.h b/Engine/Source/Runtime/Engine/Classes/Components/ChildActorComponent.h
--- a/Engine/Source/Runtime/Engine/Classes/Components/ChildActorComponent.h
+++ b/Engine/Source/Runtime/Engine/Classes/Components/ChildActorComponent.h
@@ -117,7 +117,7 @@ private:
 	TSubclassOf<AActor>	ChildActorClass;
 
 	/** The actor that we spawned and own */
-	UPROPERTY(Replicated, BlueprintReadOnly, Category=ChildActorComponent, TextExportTransient, NonPIEDuplicateTransient, meta=(AllowPrivateAccess="true"))
+	UPROPERTY(Replicated, BlueprintReadOnly, ReplicatedUsing=OnRep_ChildActor, Category=ChildActorComponent, TextExportTransient, NonPIEDuplicateTransient, meta=(AllowPrivateAccess="true"))
 	TObjectPtr<AActor>	ChildActor;
 
 	/** Property to point to the template child actor for details panel purposes */
@@ -218,7 +218,10 @@ private:
 	ENGINE_API bool IsBeingRemovedFromLevel() const;
 
 	UFUNCTION()
-	ENGINE_API void OnChildActorDestroyed(AActor* DestroyedActor);
+	void OnRep_ChildActor();
+
+	UFUNCTION()
+	ENGINE_API void OnChildActorEndPlay(AActor* Actor, EEndPlayReason::Type EndPlayReason);
 };
 
 struct FActorParentComponentSetter
diff --git a/Engine/Source/Runtime/Engine/Private/Components/ChildActorComponent.cpp b/Engine/Source/Runtime/Engine/Private/Components/ChildActorComponent.cpp
--- a/Engine/Source/Runtime/Engine/Private/Components/ChildActorComponent.cpp
+++ b/Engine/Source/Runtime/Engine/Private/Components/ChildActorComponent.cpp
@@ -655,44 +655,56 @@ bool UChildActorComponent::IsBeingRemovedFromLevel() const
 	return false;
 }
 
-void UChildActorComponent::OnChildActorDestroyed(AActor* DestroyedActor)
+void UChildActorComponent::OnRep_ChildActor()
+{
+	if (ChildActor)
+	{
+		ChildActor->OnEndPlay.AddUniqueDynamic(this, &ThisClass::OnChildActorEndPlay);
+	}
+}
+
+void UChildActorComponent::OnChildActorEndPlay(AActor* Actor, EEndPlayReason::Type EndPlayReason)
 {
 	if (GExitPurge)
 	{
 		return;
 	}
 
-	if (DestroyedActor && (DestroyedActor->HasAuthority() || !IsChildActorReplicated()) && !IsBeingRemovedFromLevel())
+	if (EndPlayReason == EEndPlayReason::Destroyed)
 	{
-		UWorld* World = DestroyedActor->GetWorld();
-		// World may be nullptr during shutdown
-		if (World != nullptr)
+		if (Actor && (Actor->HasAuthority() || !IsChildActorReplicated()) && !IsBeingRemovedFromLevel())
 		{
-			UClass* ChildClass = DestroyedActor->GetClass();
-
-			// We would like to make certain that our name is not going to accidentally get taken from us while we're destroyed
-			// so we increment ClassUnique beyond our index to be certain of it.  This is ... a bit hacky.
-			if (!GFastPathUniqueNameGeneration)
+			UWorld* World = Actor->GetWorld();
+			// World may be nullptr during shutdown
+			if (World != nullptr)
 			{
-				UpdateSuffixForNextNewObject(DestroyedActor->GetOuter(), ChildClass, [DestroyedActor](int32& Index) { Index = FMath::Max(Index, DestroyedActor->GetFName().GetNumber()); });
-			}
+				UClass* ChildClass = Actor->GetClass();
 
-			// If we are getting here due to garbage collection we can't rename, so we'll have to abandon this child actor name and pick up a new one
-			if (!IsGarbageCollecting())
-			{
-				const FString ObjectBaseName = FString::Printf(TEXT("DESTROYED_%s_CHILDACTOR"), *ChildClass->GetName());
-				DestroyedActor->Rename(*MakeUniqueObjectName(DestroyedActor->GetOuter(), ChildClass, *ObjectBaseName).ToString(), nullptr, REN_DoNotDirty | REN_ForceNoResetLoaders);
-			}
-			else
-			{
-				ChildActorName = NAME_None;
-				if (CachedInstanceData)
+				// We would like to make certain that our name is not going to accidentally get taken from us while we're destroyed
+				// so we increment ClassUnique beyond our index to be certain of it.  This is ... a bit hacky.
+				if (!GFastPathUniqueNameGeneration)
+				{
+					UpdateSuffixForNextNewObject(Actor->GetOuter(), ChildClass, [Actor](int32& Index) { Index = FMath::Max(Index, Actor->GetFName().GetNumber()); });
+				}
+
+				// If we are getting here due to garbage collection we can't rename, so we'll have to abandon this child actor name and pick up a new one
+				if (!IsGarbageCollecting())
+				{
+					const FString ObjectBaseName = FString::Printf(TEXT("DESTROYED_%s_CHILDACTOR"), *ChildClass->GetName());
+					Actor->Rename(*MakeUniqueObjectName(Actor->GetOuter(), ChildClass, *ObjectBaseName).ToString(), nullptr, REN_DoNotDirty | REN_ForceNoResetLoaders);
+				}
+				else
 				{
-					CachedInstanceData->ChildActorName = NAME_None;
+					ChildActorName = NAME_None;
+					if (CachedInstanceData)
+					{
+						CachedInstanceData->ChildActorName = NAME_None;
+					}
 				}
 			}
 		}
 	}
+
 	// While reinstancing we need the reference to remain valid so that we can
 	// overwrite it when references to old instances are updated to references
 	// to new instances:
@@ -799,7 +811,7 @@ void UChildActorComponent::CreateChildActor(TFunction<void(AActor*)> CustomizerF
 					}
 					else
 					{
-						ChildActor->OnDestroyed.AddDynamic(this, &ThisClass::OnChildActorDestroyed);
+						ChildActor->OnEndPlay.AddDynamic(this, &ThisClass::OnChildActorEndPlay);
 					}
 
 					ChildActorName = ChildActor->GetFName();
