Fix two issues with DistanceFieldReplacementMesh
- if base mesh was Nanite, DistanceFieldReplacementMesh wouldn't be used at all
- when DistanceFieldReplacementMesh was used, BuildSectionData was incorrect since it was setup using the base mesh which might use sections with different materials/sections.

#prelight 649339214d740a8c19d3583c

diff --git a/Engine/Source/Runtime/Engine/Private/DistanceFieldAtlas.cpp b/Engine/Source/Runtime/Engine/Private/DistanceFieldAtlas.cpp
--- a/Engine/Source/Runtime/Engine/Private/DistanceFieldAtlas.cpp
+++ b/Engine/Source/Runtime/Engine/Private/DistanceFieldAtlas.cpp
@@ -183,7 +183,7 @@ void FDistanceFieldVolumeData::CacheDerivedData(const FString& InStaticMeshDeriv
 	TArray<FSignedDistanceFieldBuildSectionData> BuildSectionData;
 
 	const uint32 LODIndex = 0;
-	BuildSignedDistanceFieldBuildSectionData(Mesh, LODIndex, BuildSectionData);
+	BuildSignedDistanceFieldBuildSectionData(GenerateSource, LODIndex, BuildSectionData);
 
 	FString DistanceFieldKey = BuildDistanceFieldDerivedDataKey(InStaticMeshDerivedDataKey);
 
@@ -225,10 +225,10 @@ void FDistanceFieldVolumeData::CacheDerivedData(const FString& InStaticMeshDeriv
 		NewTask->SectionData = MoveTemp(BuildSectionData);
 
 		// Nanite overrides source static mesh with a coarse representation. Need to load original data before we build the mesh SDF.
-		if (Mesh->IsNaniteEnabled())
+		if (GenerateSource->IsNaniteEnabled())
 		{
 			IMeshBuilderModule& MeshBuilderModule = IMeshBuilderModule::GetForPlatform(TargetPlatform);
-			if (!MeshBuilderModule.BuildMeshVertexPositions(Mesh, NewTask->SourceMeshData.TriangleIndices, NewTask->SourceMeshData.VertexPositions))
+			if (!MeshBuilderModule.BuildMeshVertexPositions(GenerateSource, NewTask->SourceMeshData.TriangleIndices, NewTask->SourceMeshData.VertexPositions))
 			{
 				UE_LOG(LogStaticMesh, Error, TEXT("Failed to build static mesh. See previous line(s) for details."));
 			}
