Fix FRuntimePartitionStreamingData::LoadingRange not initialized properly automated tests.

diff --git a/Engine/Source/Runtime/Engine/Public/WorldPartition/RuntimeHashSet/WorldPartitionRuntimeHashSet.h b/Engine/Source/Runtime/Engine/Public/WorldPartition/RuntimeHashSet/WorldPartitionRuntimeHashSet.h
--- a/Engine/Source/Runtime/Engine/Public/WorldPartition/RuntimeHashSet/WorldPartitionRuntimeHashSet.h
+++ b/Engine/Source/Runtime/Engine/Public/WorldPartition/RuntimeHashSet/WorldPartitionRuntimeHashSet.h
@@ -71,7 +71,7 @@ struct FRuntimePartitionStreamingData
 	FName Name;
 
 	UPROPERTY()
-	float LoadingRange;
+	float LoadingRange = 0.0f;
 
 	UPROPERTY()
 	TArray<TObjectPtr<UWorldPartitionRuntimeCell>> RuntimeCells;
