cluster unions will now export their navmesh geometry from their child components

diff --git a/Engine/Source/Runtime/Engine/Classes/PhysicsEngine/ClusterUnionComponent.h b/Engine/Source/Runtime/Engine/Classes/PhysicsEngine/ClusterUnionComponent.h
--- a/Engine/Source/Runtime/Engine/Classes/PhysicsEngine/ClusterUnionComponent.h
+++ b/Engine/Source/Runtime/Engine/Classes/PhysicsEngine/ClusterUnionComponent.h
@@ -304,6 +304,7 @@ public:
 	ENGINE_API virtual bool OverlapComponentWithResult(const FVector& Pos, const FQuat& Rot, const FPhysicsGeometry& Geometry, ECollisionChannel TraceChannel, const struct FCollisionQueryParams& Params, const struct FCollisionResponseParams& ResponseParams, const struct FCollisionObjectQueryParams& ObjectParams, TArray<FOverlapResult>& OutOverlap) const override;
 	ENGINE_API virtual bool ComponentOverlapComponentWithResultImpl(const class UPrimitiveComponent* const PrimComp, const FVector& Pos, const FQuat& Rot, const FCollisionQueryParams& Params, TArray<FOverlapResult>& OutOverlap) const override;
 	virtual bool ShouldDispatchWakeEvents(FName BoneName) const override { return true; }
+	virtual bool DoCustomNavigableGeometryExport(FNavigableGeometryExport& GeomExport) const override;
 	//~ End UPrimitiveComponent Interface
 
 	//~ Begin USceneComponent Interface
diff --git a/Engine/Source/Runtime/Engine/Private/PhysicsEngine/ClusterUnionComponent.cpp b/Engine/Source/Runtime/Engine/Private/PhysicsEngine/ClusterUnionComponent.cpp
--- a/Engine/Source/Runtime/Engine/Private/PhysicsEngine/ClusterUnionComponent.cpp
+++ b/Engine/Source/Runtime/Engine/Private/PhysicsEngine/ClusterUnionComponent.cpp
@@ -1257,3 +1257,19 @@ int32 UClusterUnionComponent::NumChildClusterComponents() const
 {
 	return PendingComponentSync.Num() + ComponentToPhysicsObjects.Num();
 }
+
+bool UClusterUnionComponent::DoCustomNavigableGeometryExport(FNavigableGeometryExport& GeomExport) const
+{
+	bool bHasData = false;
+	VisitAllCurrentChildComponents(
+		[&bHasData, &GeomExport](UPrimitiveComponent* Component)
+		{
+			if (Component)
+			{
+				bHasData |= Component->DoCustomNavigableGeometryExport(GeomExport);
+			}
+			return true;
+		}
+	);
+	return bHasData;
+}
