Fix compilation issue due to typo when ANIMNODE_STATS_VERBOSE is enabled

#rnx

diff --git a/Engine/Source/Runtime/Engine/Classes/Animation/AnimNode_CustomProperty.h b/Engine/Source/Runtime/Engine/Classes/Animation/AnimNode_CustomProperty.h
--- a/Engine/Source/Runtime/Engine/Classes/Animation/AnimNode_CustomProperty.h
+++ b/Engine/Source/Runtime/Engine/Classes/Animation/AnimNode_CustomProperty.h
@@ -109,7 +109,7 @@ protected:
 #if ANIMNODE_STATS_VERBOSE
 	// Cached StatID for this node
 	TStatId StatID;
-	virtual void InitializeStatID(); { StatID = FDynamicStats::CreateStatId<FStatGroup_STATGROUP_Anim>(FString(TEXT("Unknown"))); }
+	virtual void InitializeStatID() { StatID = FDynamicStats::CreateStatId<FStatGroup_STATGROUP_Anim>(FString(TEXT("Unknown"))); }
 #endif // ANIMNODE_STATS_VERBOSE
 
 	friend class UAnimGraphNode_CustomProperty;
