Fix incorrect BuildUnsorted overload being picked when compiling EngineTests

CIS fix

diff --git a/Engine/Plugins/Animation/ControlRig/Source/ControlRig/Private/AnimNode_ControlRigBase.cpp b/Engine/Plugins/Animation/ControlRig/Source/ControlRig/Private/AnimNode_ControlRigBase.cpp
--- a/Engine/Plugins/Animation/ControlRig/Source/ControlRig/Private/AnimNode_ControlRigBase.cpp
+++ b/Engine/Plugins/Animation/ControlRig/Source/ControlRig/Private/AnimNode_ControlRigBase.cpp
@@ -385,7 +385,7 @@ void FAnimNode_ControlRigBase::UpdateOutput(UControlRig* ControlRig, FPoseContex
         };
 
 		FBlendedCurve ControlRigCurves;
-		UE::Anim::FCurveUtils::BuildUnsorted(ControlRigCurves, Curves.Num(), GetNameFromIndex, GetValueFromIndex, GetIsValidCurveFromIndex);
+		UE::Anim::FCurveUtils::BuildUnsortedValidated(ControlRigCurves, Curves.Num(), GetNameFromIndex, GetValueFromIndex, GetIsValidCurveFromIndex);
 		InOutput.Curve.Combine(ControlRigCurves);
 	}
 
diff --git a/Engine/Source/Runtime/Engine/Public/Animation/AnimCurveUtils.h b/Engine/Source/Runtime/Engine/Public/Animation/AnimCurveUtils.h
--- a/Engine/Source/Runtime/Engine/Public/Animation/AnimCurveUtils.h
+++ b/Engine/Source/Runtime/Engine/Public/Animation/AnimCurveUtils.h
@@ -321,7 +321,7 @@ public:
 
 	// Helper function for building curves, applying filtering through InValidityPredicate
 	template<typename NamePredicateType, typename ValuePredicateType, typename ValidityPredicateType, typename CurveAllocatorType, typename CurveElementType>
-	static void BuildUnsorted(TBaseBlendedCurve<CurveAllocatorType, CurveElementType>& OutCurve, int32 InNumElements, const NamePredicateType& InNamePredicate, const ValuePredicateType& InValuePredicate, const ValidityPredicateType& InValidityPredicate)
+	static void BuildUnsortedValidated(TBaseBlendedCurve<CurveAllocatorType, CurveElementType>& OutCurve, int32 InNumElements, const NamePredicateType& InNamePredicate, const ValuePredicateType& InValuePredicate, const ValidityPredicateType& InValidityPredicate)
 	{
 		CURVE_PROFILE_CYCLE_COUNTER(FCurveUtils_BuildUnsorted);
 		BuildLinearUnfiltered(OutCurve, InNumElements, InNamePredicate, InValuePredicate, InValidityPredicate);
