turn on acceleration structure for cluster unions

diff --git a/Engine/Source/Runtime/Engine/Private/PhysicsEngine/ClusterUnionComponent.cpp b/Engine/Source/Runtime/Engine/Private/PhysicsEngine/ClusterUnionComponent.cpp
--- a/Engine/Source/Runtime/Engine/Private/PhysicsEngine/ClusterUnionComponent.cpp
+++ b/Engine/Source/Runtime/Engine/Private/PhysicsEngine/ClusterUnionComponent.cpp
@@ -22,7 +22,7 @@ DEFINE_LOG_CATEGORY(LogClusterUnion);
 
 namespace
 {
-	bool bUseClusterUnionAccelerationStructure = false;
+	bool bUseClusterUnionAccelerationStructure = true;
 	FAutoConsoleVariableRef CVarUseClusterUnionAccelerationStructure(TEXT("ClusterUnion.UseAccelerationStructure"), bUseClusterUnionAccelerationStructure, TEXT("Whether component level sweeps and overlaps against cluster unions should use an acceleration structure instead."));
 
 	bool bClusterUnionCallFlushNetDormancy = true;
