Instanced streaming levels don't build an instancing context anymore, but rely on a remapping functor instead; this speed up the streaming level setup process.
Most of the time, the linker's automatic instancing context is enough to remap all packages, but the functor is still needed for non-OFPA actors referencing OFPA actors.

diff --git a/Engine/Source/Runtime/Engine/Private/LevelStreaming.cpp b/Engine/Source/Runtime/Engine/Private/LevelStreaming.cpp
--- a/Engine/Source/Runtime/Engine/Private/LevelStreaming.cpp
+++ b/Engine/Source/Runtime/Engine/Private/LevelStreaming.cpp
@@ -1637,20 +1637,24 @@ bool ULevelStreaming::RequestLevel(UWorld* PersistentWorld, bool bAllowLevelLoad
 					NewPackage->SetFlags(RF_Transient);
 				}
 
-				// When loading an instanced package we need to build an instancing context in case non external actors part of the level are 
+				// When loading an instanced package we need to invoke an instancing context function in case non external actors part of the level are 
 				// pulling on external actors.
-				FString ExternalActorsPath = ULevel::GetExternalActorsPath(PackagePath.GetPackageName());
-				TArray<FString> ActorPackageNames = ULevel::GetOnDiskExternalActorPackages(ExternalActorsPath);
-								
-				InstancingContextPtr = &InstancingContext;
-				for (const FString& ActorPackageName : ActorPackageNames)
+				const FString ExternalActorsPathStr = ULevel::GetExternalActorsPath(PackagePath.GetPackageName());
+				const FString DesiredPackageNameStr = DesiredPackageName.ToString();
+
+				InstancingContext.AddPackageMappingFunc([ExternalActorsPathStr, DesiredPackageNameStr](FName Original)
 				{
-					const FString InstancedName = ULevel::GetExternalActorPackageInstanceName(DesiredPackageName.ToString(), ActorPackageName);
-					InstancingContext.AddPackageMapping(FName(*ActorPackageName), FName(*InstancedName));
-				}
+					const FString OriginalStr = Original.ToString();
+					if (OriginalStr.StartsWith(ExternalActorsPathStr))
+					{
+						return FName(*ULevel::GetExternalActorPackageInstanceName(DesiredPackageNameStr, OriginalStr));
+					}
+					return Original;
+				});
+
+				InstancingContextPtr = &InstancingContext;
 			}
 #endif
-
 			LoadPackageAsync(PackagePath, DesiredPackageName, FLoadPackageAsyncDelegate::CreateUObject(this, &ULevelStreaming::AsyncLevelLoadComplete), PackageFlags, PIEInstanceID, GetPriority(), InstancingContextPtr);
 
 			// streamingServer: server loads everything?
