Fixed Reset() racing updates happening in a task.

#rnx

diff --git a/Engine/Source/Runtime/Engine/Private/Components/SkeletalMeshComponent.cpp b/Engine/Source/Runtime/Engine/Private/Components/SkeletalMeshComponent.cpp
--- a/Engine/Source/Runtime/Engine/Private/Components/SkeletalMeshComponent.cpp
+++ b/Engine/Source/Runtime/Engine/Private/Components/SkeletalMeshComponent.cpp
@@ -1974,6 +1974,8 @@ void USkeletalMeshComponent::RecalcRequiredBones(int32 LODIndex)
 	BoneSpaceTransforms = GetSkeletalMeshAsset()->GetRefSkeleton().GetRefBonePose();
 	PRAGMA_ENABLE_DEPRECATION_WARNINGS
 
+	HandleExistingParallelEvaluationTask(true, false);
+
 	// If we had cached our shared bone container, reset it
 	if (SharedRequiredBones)
 	{
