TSAN - bUseRawDataOnly is written while being read from other thread
#fix swap to an atomic bool

diff --git a/Engine/Source/Runtime/Engine/Classes/Animation/AnimSequence.h b/Engine/Source/Runtime/Engine/Classes/Animation/AnimSequence.h
--- a/Engine/Source/Runtime/Engine/Classes/Animation/AnimSequence.h
+++ b/Engine/Source/Runtime/Engine/Classes/Animation/AnimSequence.h
@@ -793,8 +793,7 @@ public:
 
 	// Should we be always using our raw data (i.e is our compressed data stale)
 	UE_DEPRECATED(5.2, "bUseRawDataOnly public access will be deprecated")
-	bool bUseRawDataOnly;
-
+	TAtomic<bool> bUseRawDataOnly;
 public:
 	/** Authored Sync markers */
 	UPROPERTY()
diff --git a/Engine/Source/Runtime/Engine/Private/Animation/AnimSequence.cpp b/Engine/Source/Runtime/Engine/Private/Animation/AnimSequence.cpp
--- a/Engine/Source/Runtime/Engine/Private/Animation/AnimSequence.cpp
+++ b/Engine/Source/Runtime/Engine/Private/Animation/AnimSequence.cpp
@@ -702,7 +702,9 @@ void UAnimSequence::Serialize(FArchive& Ar)
 			SerializeCompressedData(Ar,false);
 			if (!bIsTransacting)
 			{
-				Ar << bUseRawDataOnly;
+				bool bTemp = bUseRawDataOnly;
+				Ar << bTemp;
+				bUseRawDataOnly = bTemp;
 			}
 			PRAGMA_ENABLE_DEPRECATION_WARNINGS
 		}
