Fixed a comment type: "Create a rotation from an this axis" -> "an axis"

diff --git a/Engine/Source/Runtime/Engine/Classes/Kismet/KismetMathLibrary.h b/Engine/Source/Runtime/Engine/Classes/Kismet/KismetMathLibrary.h
--- a/Engine/Source/Runtime/Engine/Classes/Kismet/KismetMathLibrary.h
+++ b/Engine/Source/Runtime/Engine/Classes/Kismet/KismetMathLibrary.h
@@ -1355,7 +1355,7 @@ class UKismetMathLibrary : public UBlueprintFunctionLibrary
 	UFUNCTION(BlueprintPure, meta=(DisplayName = "Rotation From X Vector", ScriptMethod = "Rotator", Keywords="rotation rotate cast convert", BlueprintAutocast), Category="Math|Conversions")
 	static ENGINE_API FRotator Conv_VectorToRotator(FVector InVec);
 
-	/** Create a rotation from an this axis and supplied angle (in degrees) */
+	/** Create a rotation from an axis and supplied angle (in degrees) */
 	UFUNCTION(BlueprintPure, meta=(ScriptMethod = "RotatorFromAxisAndAngle", Keywords="make construct build rotate rotation"), Category="Math|Vector")
 	static ENGINE_API FRotator RotatorFromAxisAndAngle(FVector Axis, float Angle);
 
