World Partition Class Descirptor Registry: properly handle the case where a redirector points to a deleted class.

diff --git a/Engine/Source/Runtime/Engine/Private/WorldPartition/WorldPartitionClassDescRegistry.cpp b/Engine/Source/Runtime/Engine/Private/WorldPartition/WorldPartitionClassDescRegistry.cpp
--- a/Engine/Source/Runtime/Engine/Private/WorldPartition/WorldPartitionClassDescRegistry.cpp
+++ b/Engine/Source/Runtime/Engine/Private/WorldPartition/WorldPartitionClassDescRegistry.cpp
@@ -306,42 +306,46 @@ void FWorldPartitionClassDescRegistry::PrefetchClassDescs(const TArray<FTopLevel
 	// Register current class descriptors
 	Algo::ForEach(Assets, [this, &AssetRegistry](FAssetData& AssetData)
 	{
-		const FTopLevelAssetPath AssetDataClassName(GetAssetDataClassName(AssetData));
-		AssetData.AssetName = *AssetDataClassName.GetAssetName().ToString();
-
-		// Lookup for an already registered class
-		if (!ClassByPath.Contains(FTopLevelAssetPath(AssetData.ToSoftObjectPath().ToString())))
+		// Asset could have been invalidated if a redirector chain wasn't properly resolved, etc.
+		if (AssetData.IsValid())
 		{
-			bool bOldAsset = false;
-			FAssetPackageData PackageData;
-			if (AssetRegistry.TryGetAssetPackageData(AssetData.PackageName, PackageData) == UE::AssetRegistry::EExists::Exists)
+			const FTopLevelAssetPath AssetDataClassName(GetAssetDataClassName(AssetData));
+			AssetData.AssetName = *AssetDataClassName.GetAssetName().ToString();
+
+			// Lookup for an already registered class
+			if (!ClassByPath.Contains(FTopLevelAssetPath(AssetData.ToSoftObjectPath().ToString())))
 			{
-				for (const UE::AssetRegistry::FPackageCustomVersion& CustomVersion : PackageData.GetCustomVersions())
+				bool bOldAsset = false;
+				FAssetPackageData PackageData;
+				if (AssetRegistry.TryGetAssetPackageData(AssetData.PackageName, PackageData) == UE::AssetRegistry::EExists::Exists)
 				{
-					if(CustomVersion.Key == FFortniteMainBranchObjectVersion::GUID)
+					for (const UE::AssetRegistry::FPackageCustomVersion& CustomVersion : PackageData.GetCustomVersions())
 					{
-						if (CustomVersion.Version < FFortniteMainBranchObjectVersion::WorldPartitionActorClassDescSerialize)
+						if(CustomVersion.Key == FFortniteMainBranchObjectVersion::GUID)
 						{
-							bOldAsset = true;
+							if (CustomVersion.Version < FFortniteMainBranchObjectVersion::WorldPartitionActorClassDescSerialize)
+							{
+								bOldAsset = true;
+							}
+							break;
 						}
-						break;
 					}
 				}
-			}
-			else
-			{
-				bOldAsset = true;
-			}
+				else
+				{
+					bOldAsset = true;
+				}
 
-			if (TUniquePtr<FWorldPartitionActorDesc> ClassDesc = !bOldAsset ? FWorldPartitionActorDescUtils::GetActorDescriptorFromAssetData(AssetData) : nullptr; ClassDesc.IsValid())
-			{
-				FWorldPartitionActorDesc* Result = ClassDesc.Release();
-				RegisterClassDescriptor(Result);
-			}
-			else
-			{
-				// For missing class descriptors, register a null entry so we won't be trying to register them again for future objects
-				ClassByPath.Add(FTopLevelAssetPath(FTopLevelAssetPath(AssetData.ToSoftObjectPath().ToString())), nullptr);
+				if (TUniquePtr<FWorldPartitionActorDesc> ClassDesc = !bOldAsset ? FWorldPartitionActorDescUtils::GetActorDescriptorFromAssetData(AssetData) : nullptr; ClassDesc.IsValid())
+				{
+					FWorldPartitionActorDesc* Result = ClassDesc.Release();
+					RegisterClassDescriptor(Result);
+				}
+				else
+				{
+					// For missing class descriptors, register a null entry so we won't be trying to register them again for future objects
+					ClassByPath.Add(FTopLevelAssetPath(FTopLevelAssetPath(AssetData.ToSoftObjectPath().ToString())), nullptr);
+				}
 			}
 		}
 	});
