Packed Level Actor: Avoid updating blueprints if packing result is unchanged
- Add TransformUtilities::GetRoundedTransformCRC32 and use it in ISMComponentBatcher

diff --git a/Engine/Source/Runtime/Core/Private/Misc/TransformUtilities.cpp b/Engine/Source/Runtime/Core/Private/Misc/TransformUtilities.cpp
new file mode 100644
--- /dev/null
+++ b/Engine/Source/Runtime/Core/Private/Misc/TransformUtilities.cpp
@@ -0,0 +1,27 @@
+// Copyright Epic Games, Inc. All Rights Reserved.
+
+#include "Misc/TransformUtilities.h"
+#include "Misc/Crc.h"
+#include "Math/UnrealMathUtility.h"
+#include "Math/Vector.h"
+#include "Math/Transform.h"
+
+uint32 TransformUtilities::GetRoundedTransformCRC32(const FTransform& InTransform)
+{
+	FVector Location(InTransform.GetLocation());
+	// Include transform - round sufficiently to ensure stability
+	FIntVector IntLocation(FMath::RoundToInt32(Location.X), FMath::RoundToInt32(Location.Y), FMath::RoundToInt32(Location.Z));
+	uint32 CRC = FCrc::TypeCrc32(IntLocation);
+
+	FRotator Rotator(InTransform.Rotator().GetDenormalized());
+	const int32 MAX_DEGREES = 360;
+	FIntVector IntRotation(FMath::RoundToInt32(Rotator.Pitch) % MAX_DEGREES, FMath::RoundToInt32(Rotator.Yaw) % MAX_DEGREES, FMath::RoundToInt32(Rotator.Roll) % MAX_DEGREES);
+	CRC = FCrc::TypeCrc32(IntRotation, CRC);
+
+	const float SCALE_FACTOR = 100;
+	FVector Scale(InTransform.GetScale3D());
+	FIntVector IntScale(FMath::RoundToInt32(Scale.X * SCALE_FACTOR), FMath::RoundToInt32(Scale.Y * SCALE_FACTOR), FMath::RoundToInt32(Scale.Z * SCALE_FACTOR));
+	return FCrc::TypeCrc32(IntScale, CRC);
+}
+
+
diff --git a/Engine/Source/Runtime/Core/Public/Misc/TransformUtilities.h b/Engine/Source/Runtime/Core/Public/Misc/TransformUtilities.h
new file mode 100644
--- /dev/null
+++ b/Engine/Source/Runtime/Core/Public/Misc/TransformUtilities.h
@@ -0,0 +1,10 @@
+// Copyright Epic Games, Inc. All Rights Reserved.
+
+#pragma once
+
+#include "Math/MathFwd.h"
+
+namespace TransformUtilities
+{
+	CORE_API uint32 GetRoundedTransformCRC32(const FTransform& InTransform);
+}
\ No newline at end of file
diff --git a/Engine/Source/Runtime/Engine/Private/ISMPartition/ISMComponentBatcher.cpp b/Engine/Source/Runtime/Engine/Private/ISMPartition/ISMComponentBatcher.cpp
--- a/Engine/Source/Runtime/Engine/Private/ISMPartition/ISMComponentBatcher.cpp
+++ b/Engine/Source/Runtime/Engine/Private/ISMPartition/ISMComponentBatcher.cpp
@@ -2,6 +2,9 @@
 
 #include "ISMPartition/ISMComponentBatcher.h"
 #include "Components/StaticMeshComponent.h"
+#include "Serialization/ArchiveCrc32.h"
+#include "Misc/TransformUtilities.h"
+#include "Templates/TypeHash.h"
 
 void FISMComponentBatcher::Add(const UActorComponent* InComponent)
 {
@@ -15,6 +18,7 @@ void FISMComponentBatcher::Add(const UActorComponent* InComponent, TFunctionRef<
 
 void FISMComponentBatcher::AddInternal(const UActorComponent* InComponent, TOptional<TFunctionRef<FTransform(const FTransform&)>> InTransformFunc)
 {
+	Hash = 0; // Invalidate
 	int32 NewNumCustomDataFloats = 0;
 	int32 NewNumInstances = 0;
 
@@ -121,3 +125,20 @@ void FISMComponentBatcher::InitComponent(UInstancedStaticMeshComponent* ISMCompo
 		}
 	}		
 }
+
+void FISMComponentBatcher::ComputeHash() const
+{
+	uint32 CRC = 0;
+	for (const FTransform& InstanceTransform : InstancesTransformsWS)
+	{
+		CRC = HashCombine(TransformUtilities::GetRoundedTransformCRC32(InstanceTransform), CRC);
+	}
+	
+	FArchiveCrc32 Ar(CRC);
+	FISMComponentBatcher& This = *const_cast<FISMComponentBatcher*>(this);
+
+	Ar << This.InstancesCustomData;
+	Ar << This.RandomSeeds;
+
+	Hash = Ar.GetCrc();
+}
diff --git a/Engine/Source/Runtime/Engine/Private/PackedLevelActor/PackedLevelActorBuilder.cpp b/Engine/Source/Runtime/Engine/Private/PackedLevelActor/PackedLevelActorBuilder.cpp
--- a/Engine/Source/Runtime/Engine/Private/PackedLevelActor/PackedLevelActorBuilder.cpp
+++ b/Engine/Source/Runtime/Engine/Private/PackedLevelActor/PackedLevelActorBuilder.cpp
@@ -32,6 +32,7 @@
 #include "ContentBrowserModule.h"
 #include "EditorDirectories.h"
 #include "WorldPartition/Filter/WorldPartitionActorFilter.h"
+#include "Serialization/ArchiveCrc32.h"
 
 
 #define LOCTEXT_NAMESPACE "FPackedLevelActorBuilder"
@@ -194,12 +195,16 @@ bool FPackedLevelActorBuilder::PackActor(FPackedLevelActorBuilderContext& InCont
 	InContext.SetRelativePivotTransform(RelativePivotTransform);	
 	ClusterActor(InContext, CastChecked<AActor>(InContext.GetLevelInstanceToPack()));
 
+	TArray<uint32> HashArray;
 	for (const auto& Pair : InContext.GetClusters())
 	{
 		TUniquePtr<IPackedLevelActorBuilder>& Builder = Builders.FindChecked(Pair.Key.GetBuilderID());
-		Builder->PackActors(InContext, Pair.Key, Pair.Value);
+		HashArray.Add(Builder->PackActors(InContext, Pair.Key, Pair.Value));
 	}
-			
+	HashArray.Sort();
+	FArchiveCrc32 Ar;
+	Ar << HashArray;
+	InContext.GetPackedLevelActor()->SetPackedHash(Ar.GetCrc());
 	InContext.Report(Log);
 	return true;
 }
@@ -295,6 +300,11 @@ ALevelInstance* FPackedLevelActorBuilder::CreateTransientLevelInstanceForPacking
 	UWorld* World = GEditor->GetEditorWorldContext().World();
 	check(World);
 	SpawnParams.OverrideLevel = World->PersistentLevel;
+
+	// Provide the same ActorGuid all the time so that sorting of Component by ObjectPath are deterministic (used in FPackedLevelActorISMBuilder::PackActors)
+	// It is safe because CreateTransientLevelInstanceForPacking is always called to Create a Level Instance on the stack, pack it and destroy it so no Guid conflicts can happen.
+	const FGuid PackingActorGuid(0x626EAF37, 0x1F7E47B9, 0x8C4788DC, 0x85F52822);
+	SpawnParams.OverrideActorGuid = PackingActorGuid;
 	APackedLevelActor* LevelInstance = World->SpawnActor<APackedLevelActor>(InLocation, InRotator, SpawnParams);
 	LevelInstance->SetFilter(InFilter, false);
 	LevelInstance->SetShouldLoadForPacking(true);
@@ -424,6 +434,15 @@ bool FPackedLevelActorBuilder::CreateOrUpdateBlueprintFromPacked(APackedLevelAct
 		return false;
 	}
 	
+	// No Need to update the Blueprint because packing result is the same
+	FMessageLog Log("PackedLevelActor");
+	uint32 PackedHash = CastChecked<APackedLevelActor>(BP->GeneratedClass->GetDefaultObject())->GetPackedHash();
+	if (PackedHash != 0 && PackedHash == InActor->GetPackedHash())
+	{
+		Log.Info(FText::Format(LOCTEXT("PackedBlueprintUpdated", "Packed Blueprint '{0}' already up to date"), FText::FromString(InBlueprintAsset.ToString())));
+		return true;
+	}
+
 	BP->Modify();
 
 	TArray<USCS_Node*> AllNodes = BP->SimpleConstructionScript->GetAllNodes();
@@ -473,6 +492,7 @@ bool FPackedLevelActorBuilder::CreateOrUpdateBlueprintFromPacked(APackedLevelAct
 		PackedLevelActor->SetWorldAsset(InActor->GetWorldAsset());
 		PackedLevelActor->SetFilter(InActor->GetFilter(), false);
 		PackedLevelActor->SetPackedVersion(NewVersion);
+		PackedLevelActor->SetPackedHash(InActor->GetPackedHash());
 		// match root component mobility to source actor
 		USceneComponent* Root = PackedLevelActor->GetRootComponent();
 		Root->SetMobility(InActor->GetRootComponent()->Mobility);
@@ -486,7 +506,6 @@ bool FPackedLevelActorBuilder::CreateOrUpdateBlueprintFromPacked(APackedLevelAct
 			
 	// Synchronous compile
 	FKismetEditorUtilities::CompileBlueprint(BP, EBlueprintCompileOptions::SkipGarbageCollection);
-	FMessageLog Log("PackedLevelActor");
 	Log.Info(FText::Format(LOCTEXT("PackedBlueprintUpdated", "Updated Packed Blueprint '{0}'"), FText::FromString(InBlueprintAsset.ToString())));
 
 	if (bCheckoutAndSave)
diff --git a/Engine/Source/Runtime/Engine/Private/PackedLevelActor/PackedLevelActorISMBuilder.cpp b/Engine/Source/Runtime/Engine/Private/PackedLevelActor/PackedLevelActorISMBuilder.cpp
--- a/Engine/Source/Runtime/Engine/Private/PackedLevelActor/PackedLevelActorISMBuilder.cpp
+++ b/Engine/Source/Runtime/Engine/Private/PackedLevelActor/PackedLevelActorISMBuilder.cpp
@@ -13,6 +13,8 @@
 #include "Engine/StaticMesh.h"
 
 #include "ISMPartition/ISMComponentBatcher.h"
+#include "Templates/TypeHash.h"
+#include "Misc/Crc.h"
 
 FPackedLevelActorBuilderID FPackedLevelActorISMBuilder::BuilderID = 'ISMP';
 
@@ -46,7 +48,7 @@ void FPackedLevelActorISMBuilder::GetPackClusters(FPackedLevelActorBuilderContex
 	}
 }
 
-void FPackedLevelActorISMBuilder::PackActors(FPackedLevelActorBuilderContext& InContext, const FPackedLevelActorBuilderClusterID& InClusterID, const TArray<UActorComponent*>& InComponents) const
+uint32 FPackedLevelActorISMBuilder::PackActors(FPackedLevelActorBuilderContext& InContext, const FPackedLevelActorBuilderClusterID& InClusterID, const TArray<UActorComponent*>& InComponents) const
 {
 	check(InClusterID.GetBuilderID() == GetID());
 	APackedLevelActor* PackingActor = InContext.GetPackedLevelActor();
@@ -56,7 +58,15 @@ void FPackedLevelActorISMBuilder::PackActors(FPackedLevelActorBuilderContext& In
 	check(ISMCluster);
 
 	FISMComponentBatcher ISMComponentBatcher;
-	ISMComponentBatcher.Append(InComponents, [&ActorTransform, &InContext](const FTransform& InTransform) { return InTransform.GetRelativeTransform(ActorTransform) * InContext.GetRelativePivotTransform() * ActorTransform; });
+	TArray<UActorComponent*> Components(InComponents);
+	
+	// Sort by path name to generate stable hash
+	Components.Sort([](const UActorComponent& LHS, const UActorComponent& RHS)
+	{
+		return LHS.GetPathName() < RHS.GetPathName();
+	});
+	
+	ISMComponentBatcher.Append(Components, [&ActorTransform, &InContext](const FTransform& InTransform) { return InTransform.GetRelativeTransform(ActorTransform) * InContext.GetRelativePivotTransform() * ActorTransform; });
 
 	TSubclassOf<UInstancedStaticMeshComponent> ComponentClass = UInstancedStaticMeshComponent::StaticClass();
 	if (ISMComponentBatcher.GetNumInstances() > 1 && ISMCluster->ISMDescriptor.StaticMesh && !ISMCluster->ISMDescriptor.StaticMesh->IsNaniteEnabled())
@@ -75,6 +85,10 @@ void FPackedLevelActorISMBuilder::PackActors(FPackedLevelActorBuilderContext& In
 	ISMComponentBatcher.InitComponent(PackComponent);
 
 	PackComponent->RegisterComponent();
+
+	uint32 Hash = FCrc::StrCrc32(*ComponentClass->GetPathName(), 0);
+	Hash = HashCombine(Hash, HashCombine(InClusterID.GetHash(), ISMComponentBatcher.GetHash()));
+	return Hash;
 }
 
 FPackedLevelActorISMBuilderCluster::FPackedLevelActorISMBuilderCluster(FPackedLevelActorBuilderID InBuilderID, UStaticMeshComponent* InComponent)
diff --git a/Engine/Source/Runtime/Engine/Private/PackedLevelActor/PackedLevelActorISMBuilder.h b/Engine/Source/Runtime/Engine/Private/PackedLevelActor/PackedLevelActorISMBuilder.h
--- a/Engine/Source/Runtime/Engine/Private/PackedLevelActor/PackedLevelActorISMBuilder.h
+++ b/Engine/Source/Runtime/Engine/Private/PackedLevelActor/PackedLevelActorISMBuilder.h
@@ -20,7 +20,7 @@ public:
 
 	virtual FPackedLevelActorBuilderID GetID() const override;
 	virtual void GetPackClusters(FPackedLevelActorBuilderContext& InContext, AActor* InActor) const override;
-	virtual void PackActors(FPackedLevelActorBuilderContext& InContext, const FPackedLevelActorBuilderClusterID& InClusterID, const TArray<UActorComponent*>& InComponents) const override;
+	virtual uint32 PackActors(FPackedLevelActorBuilderContext& InContext, const FPackedLevelActorBuilderClusterID& InClusterID, const TArray<UActorComponent*>& InComponents) const override;
 };
 
 class FPackedLevelActorISMBuilderCluster : public FPackedLevelActorBuilderCluster
diff --git a/Engine/Source/Runtime/Engine/Private/PackedLevelActor/PackedLevelActorRecursiveBuilder.h b/Engine/Source/Runtime/Engine/Private/PackedLevelActor/PackedLevelActorRecursiveBuilder.h
--- a/Engine/Source/Runtime/Engine/Private/PackedLevelActor/PackedLevelActorRecursiveBuilder.h
+++ b/Engine/Source/Runtime/Engine/Private/PackedLevelActor/PackedLevelActorRecursiveBuilder.h
@@ -19,7 +19,7 @@ public:
 
 	virtual FPackedLevelActorBuilderID GetID() const override;
 	virtual void GetPackClusters(FPackedLevelActorBuilderContext& InContext, AActor* InActor) const override;
-	virtual void PackActors(FPackedLevelActorBuilderContext& InContext, const FPackedLevelActorBuilderClusterID& InClusterID, const TArray<UActorComponent*>& InComponents) const override { }
+	virtual uint32 PackActors(FPackedLevelActorBuilderContext& InContext, const FPackedLevelActorBuilderClusterID& InClusterID, const TArray<UActorComponent*>& InComponents) const override { return 0; }
 };
 
 #endif
\ No newline at end of file
diff --git a/Engine/Source/Runtime/Engine/Public/ISMPartition/ISMComponentBatcher.h b/Engine/Source/Runtime/Engine/Public/ISMPartition/ISMComponentBatcher.h
--- a/Engine/Source/Runtime/Engine/Public/ISMPartition/ISMComponentBatcher.h
+++ b/Engine/Source/Runtime/Engine/Public/ISMPartition/ISMComponentBatcher.h
@@ -63,9 +63,21 @@ public:
 	 */
 	ENGINE_API void InitComponent(UInstancedStaticMeshComponent* ISMComponent) const;
 
+	
+	inline uint32 GetHash() const
+	{
+		if (Hash == 0)
+		{
+			ComputeHash();
+		}
+		return Hash;
+	}
+
 private:
-	ENGINE_API void AddInternal(const UActorComponent* InComponent, TOptional<TFunctionRef<FTransform(const FTransform&)>> InTransformFunc);
+	void ENGINE_API AddInternal(const UActorComponent* InComponent, TOptional<TFunctionRef<FTransform(const FTransform&)>> InTransformFunc);
+	void ENGINE_API ComputeHash() const;
 
+	mutable uint32 Hash = 0;
 	int32 NumInstances = 0;
 	int32 NumCustomDataFloats = 0;
 
diff --git a/Engine/Source/Runtime/Engine/Public/PackedLevelActor/IPackedLevelActorBuilder.h b/Engine/Source/Runtime/Engine/Public/PackedLevelActor/IPackedLevelActorBuilder.h
--- a/Engine/Source/Runtime/Engine/Public/PackedLevelActor/IPackedLevelActorBuilder.h
+++ b/Engine/Source/Runtime/Engine/Public/PackedLevelActor/IPackedLevelActorBuilder.h
@@ -21,7 +21,7 @@ public:
 	virtual ~IPackedLevelActorBuilder() {}
 	virtual FPackedLevelActorBuilderID GetID() const = 0;
 	virtual void GetPackClusters(FPackedLevelActorBuilderContext& InContext, AActor* InActor) const = 0;
-	virtual void PackActors(FPackedLevelActorBuilderContext& InBuilder, const FPackedLevelActorBuilderClusterID& InClusterID, const TArray<UActorComponent*>& InComponents) const = 0;
+	virtual uint32 PackActors(FPackedLevelActorBuilderContext& InBuilder, const FPackedLevelActorBuilderClusterID& InClusterID, const TArray<UActorComponent*>& InComponents) const = 0;
 protected:
 	FPackedLevelActorBuilder& Owner;
 };
diff --git a/Engine/Source/Runtime/Engine/Public/PackedLevelActor/PackedLevelActor.h b/Engine/Source/Runtime/Engine/Public/PackedLevelActor/PackedLevelActor.h
--- a/Engine/Source/Runtime/Engine/Public/PackedLevelActor/PackedLevelActor.h
+++ b/Engine/Source/Runtime/Engine/Public/PackedLevelActor/PackedLevelActor.h
@@ -61,7 +61,10 @@ public:
 
 	ENGINE_API virtual bool CanEditChange(const FProperty* InProperty) const override;
 		
-	void SetPackedVersion(const FGuid& Version) { PackedVersion = Version; }
+	void SetPackedVersion(const FGuid& InVersion) { PackedVersion = InVersion; }
+
+	uint32 GetPackedHash() const { return PackedHash; }
+	void SetPackedHash(uint32 InHash) { PackedHash = InHash; }
 
 	ENGINE_API virtual bool IsHiddenEd() const override;
 	ENGINE_API virtual bool IsHLODRelevant() const override;
@@ -93,6 +96,9 @@ public:
 private:
 	UPROPERTY()
 	FGuid PackedVersion;
+
+	UPROPERTY()
+	uint32 PackedHash;
 #endif
 
 #if WITH_EDITOR
