World Partition Class Descriptors: resolve class asset redirectors before querying class descriptor registry.

diff --git a/Engine/Source/Runtime/Engine/Private/WorldPartition/WorldPartitionActorDescArchive.cpp b/Engine/Source/Runtime/Engine/Private/WorldPartition/WorldPartitionActorDescArchive.cpp
--- a/Engine/Source/Runtime/Engine/Private/WorldPartition/WorldPartitionActorDescArchive.cpp
+++ b/Engine/Source/Runtime/Engine/Private/WorldPartition/WorldPartitionActorDescArchive.cpp
@@ -6,6 +6,7 @@
 #include "WorldPartition/WorldPartitionActorDesc.h"
 #include "WorldPartition/WorldPartitionHelpers.h"
 #include "WorldPartition/WorldPartitionLog.h"
+#include "Misc/RedirectCollector.h"
 #include "UObject/CoreRedirects.h"
 #include "UObject/UE5MainStreamObjectVersion.h"
 #include "UObject/FortniteSeasonBranchObjectVersion.h"
@@ -73,6 +74,12 @@ PRAGMA_ENABLE_DEPRECATION_WARNINGS
 				{
 					InOutClassPath = FTopLevelAssetPath(RedirectedClassRedirect.ToString());
 				}
+
+				FSoftObjectPath FoundRedirection = GRedirectCollector.GetAssetPathRedirection(FSoftObjectPath(InOutClassPath.ToString()));
+				if (FoundRedirection.IsValid())
+				{
+					InOutClassPath = FoundRedirection.GetAssetPath();
+				}
 			}
 		};
 
@@ -91,7 +98,7 @@ PRAGMA_ENABLE_DEPRECATION_WARNINGS
 			bIsMissingClassDesc = true;
 
 			ClassDesc = ClassDescRegistry.GetClassDescDefault(FTopLevelAssetPath(TEXT("/Script/Engine.Actor")));
-			check(ClassDesc);			
+			check(ClassDesc);
 
 			UE_LOG(LogWorldPartition, Log, TEXT("Can't find class descriptor '%s' for loading '%s', using '%s'"), *ClassPath.ToString(), *InActorDesc->GetActorSoftPath().ToString(), *ClassDesc->GetActorSoftPath().ToString());
 		}
@@ -115,4 +122,4 @@ FArchive& FActorDescArchive::operator<<(FSoftObjectPath& Value)
 
 	return *this;
 }
-#endif
+#endif
\ No newline at end of file
