Packed Level Actors:
- Fix PackedLevelActors not being properly unloaded when deleted
- Fix Packing failing when packing while packing world already loaded (use FPreviewScene to isolate packing into its own world)

#rnx

diff --git a/Engine/Source/Editor/LevelInstanceEditor/Private/LevelInstanceEditorModule.cpp b/Engine/Source/Editor/LevelInstanceEditor/Private/LevelInstanceEditorModule.cpp
--- a/Engine/Source/Editor/LevelInstanceEditor/Private/LevelInstanceEditorModule.cpp
+++ b/Engine/Source/Editor/LevelInstanceEditor/Private/LevelInstanceEditorModule.cpp
@@ -8,6 +8,7 @@
 #include "LevelInstance/LevelInstanceInterface.h"
 #include "LevelInstance/LevelInstanceActor.h"
 #include "PackedLevelActor/PackedLevelActor.h"
+#include "PackedLevelActor/PackedLevelActorBuilder.h"
 #include "LevelInstanceEditorSettings.h"
 #include "ToolMenus.h"
 #include "Editor.h"
@@ -669,11 +670,12 @@ namespace LevelInstanceMenuUtils
 			FScopedSlowTask SlowTask(0.0f, LOCTEXT("UpdatePackedBlueprintsProgress", "Updating Packed Blueprints..."));
 			TSet<TSoftObjectPtr<UBlueprint>> BlueprintAssets;
 			FPackedLevelActorUtils::GetPackedBlueprintsForWorldAsset(TSoftObjectPtr<UWorld>(WorldAsset.GetSoftObjectPath()), BlueprintAssets, false);
+			TSharedPtr<FPackedLevelActorBuilder> Builder = FPackedLevelActorBuilder::CreateDefaultBuilder();
 			for (TSoftObjectPtr<UBlueprint> BlueprintAsset : BlueprintAssets)
 			{
 				if (UBlueprint* Blueprint = BlueprintAsset.Get())
 				{
-					FPackedLevelActorUtils::UpdateBlueprint(Blueprint, false);
+					Builder->UpdateBlueprint(Blueprint, false);
 				}
 			}
 		});
diff --git a/Engine/Source/Editor/LevelInstanceEditor/Private/PackedLevelActorUtils.cpp b/Engine/Source/Editor/LevelInstanceEditor/Private/PackedLevelActorUtils.cpp
--- a/Engine/Source/Editor/LevelInstanceEditor/Private/PackedLevelActorUtils.cpp
+++ b/Engine/Source/Editor/LevelInstanceEditor/Private/PackedLevelActorUtils.cpp
@@ -178,7 +178,7 @@ bool FPackedLevelActorUtils::CreateOrUpdateBlueprint(TSoftObjectPtr<UWorld> InWo
 
 void FPackedLevelActorUtils::UpdateBlueprint(UBlueprint* InBlueprint, bool bCheckoutAndSave)
 {
-	return FPackedLevelActorBuilder::CreateDefaultBuilder()->UpdateBlueprint(InBlueprint, bCheckoutAndSave);
+	FPackedLevelActorBuilder::CreateDefaultBuilder()->UpdateBlueprint(InBlueprint, bCheckoutAndSave);
 }
 
 #undef LOCTEXT_NAMESPACE
diff --git a/Engine/Source/Runtime/Engine/Private/LevelInstance/LevelInstanceSubsystem.cpp b/Engine/Source/Runtime/Engine/Private/LevelInstance/LevelInstanceSubsystem.cpp
--- a/Engine/Source/Runtime/Engine/Private/LevelInstance/LevelInstanceSubsystem.cpp
+++ b/Engine/Source/Runtime/Engine/Private/LevelInstance/LevelInstanceSubsystem.cpp
@@ -1528,12 +1528,6 @@ void ULevelInstanceSubsystem::OnActorDeleted(AActor* Actor)
 			return;
 		}
 
-		const bool bIsEditingLevelInstance = IsEditingLevelInstance(LevelInstance);
-		if (!bIsEditingLevelInstance && Actor->IsA<APackedLevelActor>())
-		{
-			return;
-		}
-
 		const bool bAlreadyRooted = Actor->IsRooted();
 		// Unloading LevelInstances leads to GC and Actor can be collected. Add to root temp. It will get collected after the OnActorDeleted callbacks
 		if (!bAlreadyRooted)
@@ -1541,8 +1535,10 @@ void ULevelInstanceSubsystem::OnActorDeleted(AActor* Actor)
 			Actor->AddToRoot();
 		}
 
+		const bool bIsEditingLevelInstance = IsEditingLevelInstance(LevelInstance);
+
 		FScopedSlowTask SlowTask(0, LOCTEXT("UnloadingLevelInstances", "Unloading Level Instances..."), !GetWorld()->IsGameWorld());
-		SlowTask.MakeDialog();
+		SlowTask.MakeDialogDelayed(1.0f);
 		check(!IsEditingLevelInstanceDirty(LevelInstance) && !HasDirtyChildrenLevelInstances(LevelInstance));
 		if (bIsEditingLevelInstance)
 		{
diff --git a/Engine/Source/Runtime/Engine/Private/PackedLevelActor/PackedLevelActorBuilder.cpp b/Engine/Source/Runtime/Engine/Private/PackedLevelActor/PackedLevelActorBuilder.cpp
--- a/Engine/Source/Runtime/Engine/Private/PackedLevelActor/PackedLevelActorBuilder.cpp
+++ b/Engine/Source/Runtime/Engine/Private/PackedLevelActor/PackedLevelActorBuilder.cpp
@@ -53,6 +53,7 @@ void FPackedLevelActorBuilderContext::DiscardActor(AActor* InActor)
 }
 
 FPackedLevelActorBuilder::FPackedLevelActorBuilder()
+	: PreviewScene(FPreviewScene::ConstructionValues().SetTransactional(false))
 {
 }
 
@@ -297,8 +298,7 @@ ALevelInstance* FPackedLevelActorBuilder::CreateTransientLevelInstanceForPacking
 	SpawnParams.bNoFail = true;
 	SpawnParams.ObjectFlags |= RF_Transient;
 
-	UWorld* World = GEditor->GetEditorWorldContext().World();
-	check(World);
+	UWorld* World = PreviewScene.GetWorld();
 	SpawnParams.OverrideLevel = World->PersistentLevel;
 
 	// Provide the same ActorGuid all the time so that sorting of Component by ObjectPath are deterministic (used in FPackedLevelActorISMBuilder::PackActors)
@@ -381,7 +381,7 @@ bool FPackedLevelActorBuilder::CreateOrUpdateBlueprintFromUnpacked(ILevelInstanc
 	SpawnParams.ObjectFlags |= RF_Transient;
 	SpawnParams.ObjectFlags &= ~RF_Transactional;
 	
-	UWorld* World = LevelInstanceActor->GetWorld();
+	UWorld* World = PreviewScene.GetWorld();
 	SpawnParams.OverrideLevel = World->PersistentLevel;
 
 	APackedLevelActor* PackedLevelActor = World->SpawnActor<APackedLevelActor>(LevelInstanceActor->GetActorLocation(), LevelInstanceActor->GetActorRotation(), SpawnParams);
diff --git a/Engine/Source/Runtime/Engine/Public/PackedLevelActor/PackedLevelActorBuilder.h b/Engine/Source/Runtime/Engine/Public/PackedLevelActor/PackedLevelActorBuilder.h
--- a/Engine/Source/Runtime/Engine/Public/PackedLevelActor/PackedLevelActorBuilder.h
+++ b/Engine/Source/Runtime/Engine/Public/PackedLevelActor/PackedLevelActorBuilder.h
@@ -10,6 +10,7 @@
 #include "PackedLevelActor/IPackedLevelActorBuilder.h"
 #include "UObject/SoftObjectPtr.h"
 #include "Containers/Set.h"
+#include "PreviewScene.h"
 
 class ILevelInstanceInterface;
 class APackedLevelActor;
@@ -78,6 +79,7 @@ private:
 	
 	TSet<UClass*> ClassDiscards;
 	TMap<FPackedLevelActorBuilderID, TUniquePtr<IPackedLevelActorBuilder>> Builders;
+	FPreviewScene PreviewScene;
 };
 
 
