Add an option to fallback to the default material when creating the render proxy if precached PSOs are not ready.

This augments the current option to skip render proxy creation altogether if PSOs are not ready. There are now three options under what I named the "render proxy creation delay strategy":
* Always create render proxy, incurring a blocking wait at draw time if the PSO is not ready by then (r.PSOPrecache.ProxyCreationWhenPSOReady=0).
* Always skip render proxy creation until PSOs have been successfully precached (r.PSOPrecache.ProxyCreationWhenPSOReady=1, r.PSOPrecache.ProxyCreationDelayStrategy=0).
* Create a render proxy that uses the default material if the PSO has not finished precaching by creation time. Currently implemented for static and skinned mesh components (incl. spline and instanced) but not for Niagara. (r.PSOPrecache.ProxyCreationWhenPSOReady=1, r.PSOPrecache.ProxyCreationDelayStrategy=1)

diff --git a/Engine/Plugins/FX/Niagara/Source/Niagara/Private/NiagaraComponent.cpp b/Engine/Plugins/FX/Niagara/Source/Niagara/Private/NiagaraComponent.cpp
--- a/Engine/Plugins/FX/Niagara/Source/Niagara/Private/NiagaraComponent.cpp
+++ b/Engine/Plugins/FX/Niagara/Source/Niagara/Private/NiagaraComponent.cpp
@@ -2172,7 +2172,7 @@ FPrimitiveSceneProxy* UNiagaraComponent::CreateSceneProxy()
 			PrecacheAssetPSOs(Asset);
 		}
 
-		if (IsPSOPrecaching())
+		if (CheckPSOPrecachingAndBoostPriority() && GetPSOPrecacheProxyCreationStrategy() != EPSOPrecacheProxyCreationStrategy::AlwaysCreate)
 		{
 			UE_LOG(LogNiagara, Verbose, TEXT("Skipping CreateSceneProxy for UNiagaraComponent %s (UNiagaraSystem PSOs are still compiling)"), *GetFullName());
 			return nullptr;
diff --git a/Engine/Source/Runtime/Engine/Classes/Components/PrimitiveComponent.h b/Engine/Source/Runtime/Engine/Classes/Components/PrimitiveComponent.h
--- a/Engine/Source/Runtime/Engine/Classes/Components/PrimitiveComponent.h
+++ b/Engine/Source/Runtime/Engine/Classes/Components/PrimitiveComponent.h
@@ -960,8 +960,17 @@ public:
 	/** Schedule task to mark render state dirty when the PSO precaching tasks are done */
 	ENGINE_API void RequestRecreateRenderStateWhenPSOPrecacheFinished(const FGraphEventArray& PSOPrecacheCompileEvents);
 
-	/** Check if PSOs are still precaching and boost priority if not done yet */
-	ENGINE_API bool IsPSOPrecaching();
+	/** Check if PSOs are still precaching */
+	ENGINE_API bool IsPSOPrecaching() const;
+
+	/** Whether the render proxy should fallback to the default material because the PSOs are still precaching */
+	ENGINE_API bool ShouldRenderProxyFallbackToDefaultMaterial() const;
+
+	/**
+	 * Check if PSOs are still precaching and boost priority if not done yet.
+	 * Returns true if the PSOs are still precaching.
+	 */
+	ENGINE_API bool CheckPSOPrecachingAndBoostPriority();
 
 	/**
 	 * Set of actors to ignore during component sweeps in MoveComponent().
diff --git a/Engine/Source/Runtime/Engine/Private/Components/PrimitiveComponent.cpp b/Engine/Source/Runtime/Engine/Private/Components/PrimitiveComponent.cpp
--- a/Engine/Source/Runtime/Engine/Private/Components/PrimitiveComponent.cpp
+++ b/Engine/Source/Runtime/Engine/Private/Components/PrimitiveComponent.cpp
@@ -4532,8 +4532,9 @@ void UPrimitiveComponent::PrecachePSOs()
 
 void UPrimitiveComponent::RequestRecreateRenderStateWhenPSOPrecacheFinished(const FGraphEventArray& PSOPrecacheCompileEvents)
 {
-	// Mark the render state dirty when all PSOs are compiled so the proxy gets recreated
-	if (ProxyCreationWhenPSOReady() && !PSOPrecacheCompileEvents.IsEmpty())
+	// If the proxy creation strategy relies on knowing when the precached PSO has been compiled,
+	// schedule a task to mark the render state dirty when all PSOs are compiled so the proxy gets recreated.
+	if (GetPSOPrecacheProxyCreationStrategy() != EPSOPrecacheProxyCreationStrategy::AlwaysCreate && !PSOPrecacheCompileEvents.IsEmpty())
 	{
 		PSOPrecacheCompileEvent = TGraphTask<FMarkRenderStateDirtyTask>::CreateTask(&PSOPrecacheCompileEvents).ConstructAndDispatchWhenReady(this);
 	}
@@ -4541,7 +4542,17 @@ void UPrimitiveComponent::RequestRecreateRenderStateWhenPSOPrecacheFinished(cons
 	bPSOPrecacheCalled = true;
 }
 
-bool UPrimitiveComponent::IsPSOPrecaching()
+bool UPrimitiveComponent::IsPSOPrecaching() const
+{
+	return PSOPrecacheCompileEvent && !PSOPrecacheCompileEvent->IsComplete();
+}
+
+bool UPrimitiveComponent::ShouldRenderProxyFallbackToDefaultMaterial() const
+{
+	return IsPSOPrecaching() && GetPSOPrecacheProxyCreationStrategy() == EPSOPrecacheProxyCreationStrategy::UseDefaultMaterialUntilPSOPrecached;
+}
+
+bool UPrimitiveComponent::CheckPSOPrecachingAndBoostPriority()
 {
 	ensure(!IsComponentPSOPrecachingEnabled() || bPSOPrecacheCalled);
 
@@ -4558,7 +4569,7 @@ bool UPrimitiveComponent::IsPSOPrecaching()
 		PSOPrecacheCompileEvent = nullptr;
 	}
 
-	return PSOPrecacheCompileEvent != nullptr;
+	return IsPSOPrecaching();
 }
 
 #if WITH_EDITOR
diff --git a/Engine/Source/Runtime/Engine/Private/Components/SkinnedMeshComponent.cpp b/Engine/Source/Runtime/Engine/Private/Components/SkinnedMeshComponent.cpp
--- a/Engine/Source/Runtime/Engine/Private/Components/SkinnedMeshComponent.cpp
+++ b/Engine/Source/Runtime/Engine/Private/Components/SkinnedMeshComponent.cpp
@@ -526,12 +526,17 @@ FPrimitiveSceneProxy* USkinnedMeshComponent::CreateSceneProxy()
 	FSkeletalMeshSceneProxy* Result = nullptr;
 	FSkeletalMeshRenderData* SkelMeshRenderData = GetSkeletalMeshRenderData();
 
+	if (CheckPSOPrecachingAndBoostPriority() && GetPSOPrecacheProxyCreationStrategy() == EPSOPrecacheProxyCreationStrategy::DelayUntilPSOPrecached)
+	{
+		UE_LOG(LogSkinnedMeshComp, Verbose, TEXT("Skipping CreateSceneProxy for USkinnedMeshComponent %s (USkinnedMeshComponent PSOs are still compiling)"), *GetFullName());
+		return nullptr;
+	}
+
 	// Only create a scene proxy for rendering if properly initialized
 	if (SkelMeshRenderData &&
 		SkelMeshRenderData->LODRenderData.IsValidIndex(GetPredictedLODLevel()) &&
 		!bHideSkin &&
-		MeshObject &&
-		!IsPSOPrecaching())
+		MeshObject)
 	{
 		// Only create a scene proxy if the bone count being used is supported, or if we don't have a skeleton (this is the case with destructibles)
 		int32 MinLODIndex = ComputeMinLOD();
diff --git a/Engine/Source/Runtime/Engine/Private/Components/SplineMeshComponent.cpp b/Engine/Source/Runtime/Engine/Private/Components/SplineMeshComponent.cpp
--- a/Engine/Source/Runtime/Engine/Private/Components/SplineMeshComponent.cpp
+++ b/Engine/Source/Runtime/Engine/Private/Components/SplineMeshComponent.cpp
@@ -813,6 +813,11 @@ void USplineMeshComponent::CollectPSOPrecacheData(const FPSOPrecacheParams& Base
 FPrimitiveSceneProxy* USplineMeshComponent::CreateStaticMeshSceneProxy(Nanite::FMaterialAudit& NaniteMaterials, bool bCreateNanite)
 {
 	LLM_SCOPE(ELLMTag::StaticMesh);
+
+     if (CheckPSOPrecachingAndBoostPriority() && GetPSOPrecacheProxyCreationStrategy() == EPSOPrecacheProxyCreationStrategy::DelayUntilPSOPrecached)
+	{
+		return nullptr;
+	}
 	
 	if (bCreateNanite && ShouldRenderNaniteSplineMeshes())
 	{
diff --git a/Engine/Source/Runtime/Engine/Private/HierarchicalInstancedStaticMesh.cpp b/Engine/Source/Runtime/Engine/Private/HierarchicalInstancedStaticMesh.cpp
--- a/Engine/Source/Runtime/Engine/Private/HierarchicalInstancedStaticMesh.cpp
+++ b/Engine/Source/Runtime/Engine/Private/HierarchicalInstancedStaticMesh.cpp
@@ -3301,6 +3301,12 @@ FPrimitiveSceneProxy* UHierarchicalInstancedStaticMeshComponent::CreateSceneProx
 	}
 	ProxySize = 0;
 
+	if (CheckPSOPrecachingAndBoostPriority() && GetPSOPrecacheProxyCreationStrategy() == EPSOPrecacheProxyCreationStrategy::DelayUntilPSOPrecached)
+	{
+		UE_LOG(LogStaticMesh, Verbose, TEXT("Skipping CreateSceneProxy for UHierarchicalInstancedStaticMeshComponent %s (UHierarchicalInstancedStaticMeshComponent PSOs are still compiling)"), *GetFullName());
+		return nullptr;
+	}
+
 	// Verify that the mesh is valid before using it.
 	const bool bMeshIsValid =
 		// Make sure we have instances.
@@ -3308,8 +3314,7 @@ FPrimitiveSceneProxy* UHierarchicalInstancedStaticMeshComponent::CreateSceneProx
 		// Make sure we have an actual static mesh.
 		GetStaticMesh() &&
 		!GetStaticMesh()->IsCompiling() &&
-		GetStaticMesh()->HasValidRenderData(false) &&
-		!IsPSOPrecaching();
+		GetStaticMesh()->HasValidRenderData(false);
 
 	if (bMeshIsValid)
 	{
diff --git a/Engine/Source/Runtime/Engine/Private/InstancedStaticMesh.cpp b/Engine/Source/Runtime/Engine/Private/InstancedStaticMesh.cpp
--- a/Engine/Source/Runtime/Engine/Private/InstancedStaticMesh.cpp
+++ b/Engine/Source/Runtime/Engine/Private/InstancedStaticMesh.cpp
@@ -2736,6 +2736,12 @@ FPrimitiveSceneProxy* UInstancedStaticMeshComponent::CreateStaticMeshSceneProxy(
 
 	LLM_SCOPE(ELLMTag::InstancedMesh);
 
+	if (CheckPSOPrecachingAndBoostPriority() && GetPSOPrecacheProxyCreationStrategy() == EPSOPrecacheProxyCreationStrategy::DelayUntilPSOPrecached)
+	{
+		UE_LOG(LogStaticMesh, Verbose, TEXT("Skipping CreateSceneProxy for UInstancedStaticMeshComponent %s (UInstancedStaticMeshComponent PSOs are still compiling)"), *GetFullName());
+		return nullptr;
+	}
+
 	if (bCreateNanite)
 	{
 		return ::new Nanite::FSceneProxy(NaniteMaterials, this);
@@ -2765,8 +2771,7 @@ FPrimitiveSceneProxy* UInstancedStaticMeshComponent::CreateSceneProxy()
 		// make sure we have an actual static mesh
 		GetStaticMesh() &&
 		GetStaticMesh()->IsCompiling() == false &&
-		GetStaticMesh()->HasValidRenderData() &&
-		!IsPSOPrecaching();
+		GetStaticMesh()->HasValidRenderData();
 
 	if (bMeshIsValid)
 	{
diff --git a/Engine/Source/Runtime/Engine/Private/PSOPrecache.cpp b/Engine/Source/Runtime/Engine/Private/PSOPrecache.cpp
--- a/Engine/Source/Runtime/Engine/Private/PSOPrecache.cpp
+++ b/Engine/Source/Runtime/Engine/Private/PSOPrecache.cpp
@@ -68,7 +68,19 @@ int32 GPSOProxyCreationWhenPSOReady = 0;
 static FAutoConsoleVariableRef CVarPSOProxyCreationWhenPSOReady(
 	TEXT("r.PSOPrecache.ProxyCreationWhenPSOReady"),
 	GPSOProxyCreationWhenPSOReady,
-	TEXT("Delay the component proxy creation when the requested PSOs for precaching are still compiling."),
+	TEXT("Delay the component proxy creation when the requested PSOs for precaching are still compiling.\n")
+	TEXT(" 0: always create regardless of PSOs status (default)\n")
+	TEXT(" 1: delay the creation of the render proxy depending on the specific strategy controlled by r.PSOPrecache.ProxyCreationDelayStrategy\n"),
+	ECVF_ReadOnly
+);
+
+int32 GPSOProxyCreationDelayStrategy = 0;
+static FAutoConsoleVariableRef CVarPSOProxyCreationDelayStrategy(
+	TEXT("r.PSOPrecache.ProxyCreationDelayStrategy"),
+	GPSOProxyCreationDelayStrategy,
+	TEXT("Control the component proxy creation strategy when the requested PSOs for precaching are still compiling. Ignored if r.PSOPrecache.ProxyCreationWhenPSOReady = 0.\n")
+	TEXT(" 0: delay creation until PSOs are ready (default)\n")
+	TEXT(" 1: create a proxy using the default material until PSOs are ready. Currently implemented for static and skinned meshes - Niagara components will delay creation instead"),
 	ECVF_ReadOnly
 );
 
@@ -84,6 +96,24 @@ bool IsResourcePSOPrecachingEnabled()
 	return FApp::CanEverRender() && PipelineStateCache::IsPSOPrecachingEnabled() && GPSOPrecacheResources && !GIsEditor;
 }
 
+EPSOPrecacheProxyCreationStrategy GetPSOPrecacheProxyCreationStrategy()
+{
+	if (GPSOProxyCreationWhenPSOReady != 1)
+	{
+		return EPSOPrecacheProxyCreationStrategy::AlwaysCreate;
+	}
+
+	switch (GPSOProxyCreationDelayStrategy)
+	{
+	case 1:
+		return EPSOPrecacheProxyCreationStrategy::UseDefaultMaterialUntilPSOPrecached;
+	case 0:
+		[[fallthrough]];
+	default:
+		return EPSOPrecacheProxyCreationStrategy::DelayUntilPSOPrecached;
+	}
+}
+
 bool ProxyCreationWhenPSOReady()
 {
 	return FApp::CanEverRender() && PipelineStateCache::IsPSOPrecachingEnabled() && GPSOProxyCreationWhenPSOReady && !GIsEditor;
diff --git a/Engine/Source/Runtime/Engine/Private/Particles/ParticleSystemRender.cpp b/Engine/Source/Runtime/Engine/Private/Particles/ParticleSystemRender.cpp
--- a/Engine/Source/Runtime/Engine/Private/Particles/ParticleSystemRender.cpp
+++ b/Engine/Source/Runtime/Engine/Private/Particles/ParticleSystemRender.cpp
@@ -6912,7 +6912,7 @@ FPrimitiveSceneProxy* UParticleSystemComponent::CreateSceneProxy()
 			PrecacheAssetPSOs(Template);
 		}
 
-		if (IsPSOPrecaching())
+		if (CheckPSOPrecachingAndBoostPriority() && GetPSOPrecacheProxyCreationStrategy() != EPSOPrecacheProxyCreationStrategy::AlwaysCreate)
 		{
 			UE_LOG(LogParticles, Verbose, TEXT("Skipping CreateSceneProxy for UParticleSystemComponent %s (UParticleSystem PSOs are still compiling)"), *GetFullName());
 			return nullptr;
diff --git a/Engine/Source/Runtime/Engine/Private/SkeletalMesh.cpp b/Engine/Source/Runtime/Engine/Private/SkeletalMesh.cpp
--- a/Engine/Source/Runtime/Engine/Private/SkeletalMesh.cpp
+++ b/Engine/Source/Runtime/Engine/Private/SkeletalMesh.cpp
@@ -5640,6 +5640,8 @@ FSkeletalMeshSceneProxy::FSkeletalMeshSceneProxy(const USkinnedMeshComponent* Co
 		bAlwaysHasVelocity = true;
 	}
 
+	bool bForceDefaultMaterial = Component->ShouldRenderProxyFallbackToDefaultMaterial();
+
 	// setup materials and performance classification for each LOD.
 	extern bool GForceDefaultMaterial;
 	bool bCastShadow = Component->CastShadow;
@@ -5674,7 +5676,7 @@ FSkeletalMeshSceneProxy::FSkeletalMeshSceneProxy(const USkinnedMeshComponent* Co
 
 			// If the material is NULL, or isn't flagged for use with skeletal meshes, it will be replaced by the default material.
 			UMaterialInterface* Material = Component->GetMaterial(UseMaterialIndex);
-			if (GForceDefaultMaterial && Material && !IsTranslucentBlendMode(*Material))
+			if (bForceDefaultMaterial || (GForceDefaultMaterial && Material && !IsTranslucentBlendMode(*Material)))
 			{
 				Material = UMaterial::GetDefaultMaterial(MD_Surface);
 				MaterialRelevance |= Material->GetRelevance(FeatureLevel);
diff --git a/Engine/Source/Runtime/Engine/Private/StaticMeshRender.cpp b/Engine/Source/Runtime/Engine/Private/StaticMeshRender.cpp
--- a/Engine/Source/Runtime/Engine/Private/StaticMeshRender.cpp
+++ b/Engine/Source/Runtime/Engine/Private/StaticMeshRender.cpp
@@ -265,6 +265,8 @@ FStaticMeshSceneProxy::FStaticMeshSceneProxy(UStaticMeshComponent* InComponent,
 	}
 #endif
 
+	bool bForceDefaultMaterial = InComponent->ShouldRenderProxyFallbackToDefaultMaterial();
+
 	// Find the first LOD with any vertices (ie that haven't been stripped)
 	int FirstAvailableLOD = 0;
 	for (; FirstAvailableLOD < RenderData->LODResources.Num(); FirstAvailableLOD++)
@@ -275,7 +277,7 @@ FStaticMeshSceneProxy::FStaticMeshSceneProxy(UStaticMeshComponent* InComponent,
 		}
 	}
 
-	if (GForceDefaultMaterial)
+	if (bForceDefaultMaterial || GForceDefaultMaterial)
 	{
 		MaterialRelevance |= UMaterial::GetDefaultMaterial(MD_Surface)->GetRelevance(FeatureLevel);
 	}
@@ -2206,6 +2208,8 @@ FStaticMeshSceneProxy::FLODInfo::FLODInfo(const UStaticMeshComponent* InComponen
 		SetGlobalVolumeLightmap(true);
 	}
 
+	bool bForceDefaultMaterial = InComponent->ShouldRenderProxyFallbackToDefaultMaterial();
+
 	bool bMeshMapBuildDataOverriddenByLightmapPreview = false;
 
 #if WITH_EDITOR
@@ -2326,7 +2330,7 @@ FStaticMeshSceneProxy::FLODInfo::FLODInfo(const UStaticMeshComponent* InComponen
 		SectionInfo.MaterialIndex = Section.MaterialIndex;
 #endif
 
-		if (GForceDefaultMaterial && SectionInfo.Material && !IsTranslucentBlendMode(*SectionInfo.Material))
+		if (bForceDefaultMaterial || (GForceDefaultMaterial && SectionInfo.Material && !IsTranslucentBlendMode(*SectionInfo.Material)))
 		{
 			SectionInfo.Material = UMaterial::GetDefaultMaterial(MD_Surface);
 		}
@@ -2623,7 +2627,7 @@ FPrimitiveSceneProxy* UStaticMeshComponent::CreateSceneProxy()
 		return nullptr;
 	}
 
-	if (IsPSOPrecaching())
+	if (CheckPSOPrecachingAndBoostPriority() && GetPSOPrecacheProxyCreationStrategy() == EPSOPrecacheProxyCreationStrategy::DelayUntilPSOPrecached)
 	{
 		UE_LOG(LogStaticMesh, Verbose, TEXT("Skipping CreateSceneProxy for StaticMeshComponent %s (Static mesh component PSOs are still compiling)"), *GetFullName());
 		return nullptr;
diff --git a/Engine/Source/Runtime/Engine/Public/PSOPrecache.h b/Engine/Source/Runtime/Engine/Public/PSOPrecache.h
--- a/Engine/Source/Runtime/Engine/Public/PSOPrecache.h
+++ b/Engine/Source/Runtime/Engine/Public/PSOPrecache.h
@@ -236,6 +236,24 @@ extern ENGINE_API bool IsComponentPSOPrecachingEnabled();
  */
 extern ENGINE_API bool IsResourcePSOPrecachingEnabled();
 
+enum class EPSOPrecacheProxyCreationStrategy : uint8
+{
+	// Always create the render proxy regardless of whether the PSO has finished precaching or not. 
+	// This will introduce a blocking wait when the proxy is rendered if the PSO is not ready.
+	AlwaysCreate = 0, 
+
+	// Delay the creation of the render proxy until the PSO has finished precaching. 
+	// This effectively skips drawing components until the PSO is ready, when the proxy will be created.
+	DelayUntilPSOPrecached = 1, 
+	
+	// Create a render proxy that uses the default material if the PSO has not finished precaching by creation time.
+	// The proxy will be re-created with the actual materials once the PSO is ready.
+	// Currently implemented only for static and skinned mesh components, while Niagara components will skip render proxy creation altogether.
+	UseDefaultMaterialUntilPSOPrecached = 2
+};
+
+extern ENGINE_API EPSOPrecacheProxyCreationStrategy GetPSOPrecacheProxyCreationStrategy();
+
 /**
  * Delay component proxy creation when it's requested PSOs are still precaching
  */
