Bump DDC keys because of possible corruption.

diff --git a/Engine/Source/Runtime/Engine/Private/AudioDerivedData.cpp b/Engine/Source/Runtime/Engine/Private/AudioDerivedData.cpp
--- a/Engine/Source/Runtime/Engine/Private/AudioDerivedData.cpp
+++ b/Engine/Source/Runtime/Engine/Private/AudioDerivedData.cpp
@@ -165,8 +165,8 @@ Derived data key generation.
 
 // If you want to bump this version, generate a new guid using
 // VS->Tools->Create GUID and paste it here. https://www.guidgen.com works too.
-#define AUDIO_DERIVEDDATA_VER				TEXT("4CC71E548D764B6989612C03CBF8E376")				// Change this if you want to bump all audio formats at once
-#define STREAMEDAUDIO_DERIVEDDATA_VER		TEXT("F6F8FA984D96440B8A28C73F2ABAE692")				// This depends on the above key, but will regenerate all streaming chunk data derived from the compressed audio
+#define AUDIO_DERIVEDDATA_VER				TEXT("78D1D5180FCF4D3FB7E8B2A4799AD116")				// Change this if you want to bump all audio formats at once
+#define STREAMEDAUDIO_DERIVEDDATA_VER		TEXT("1A69BFC09E8146E08A1B7FD9BD5B9BA4")				// This depends on the above key, but will regenerate all streaming chunk data derived from the compressed audio
 
 #if WITH_EDITORONLY_DATA
 
