Fix for crash when linking/unlinking anim layers from UGameplayAbility::EndAbility

diff --git a/Engine/Source/Runtime/Engine/Private/Animation/AnimInstance.cpp b/Engine/Source/Runtime/Engine/Private/Animation/AnimInstance.cpp
--- a/Engine/Source/Runtime/Engine/Private/Animation/AnimInstance.cpp
+++ b/Engine/Source/Runtime/Engine/Private/Animation/AnimInstance.cpp
@@ -3281,7 +3281,10 @@ void UAnimInstance::LinkAnimClassLayers(TSubclassOf<UAnimInstance> InClass)
 		}
 	};
 
-	PerformLinkedLayerOverlayOperation(InClass, SelectResolvedClassIfValid);
+	if (GetSkelMeshComponent()->IsRegistered())
+	{
+		PerformLinkedLayerOverlayOperation(InClass, SelectResolvedClassIfValid);
+	}
 }
 
 void UAnimInstance::UnlinkAnimClassLayers(TSubclassOf<UAnimInstance> InClass)
@@ -3305,7 +3308,10 @@ void UAnimInstance::UnlinkAnimClassLayers(TSubclassOf<UAnimInstance> InClass)
 		}
 	};
 
-	PerformLinkedLayerOverlayOperation(InClass, ConditionallySelectDefaultClass);
+	if (GetSkelMeshComponent()->IsRegistered())
+	{
+		PerformLinkedLayerOverlayOperation(InClass, ConditionallySelectDefaultClass);
+	}
 }
 
 
@@ -3325,7 +3331,10 @@ void UAnimInstance::InitializeGroupedLayers(bool bInDeferSubGraphInitialization)
 		}
 	};
 
-	PerformLinkedLayerOverlayOperation(nullptr, SelectResolvedClassIfValid, bInDeferSubGraphInitialization);
+	if (GetSkelMeshComponent()->IsRegistered())
+	{
+		PerformLinkedLayerOverlayOperation(nullptr, SelectResolvedClassIfValid, bInDeferSubGraphInitialization);
+	}
 }
 
 void UAnimInstance::AddExternalNotifyHandler(UObject* ExternalHandlerObject, FName NotifyEventName)
