Fixing issue where a character could intend to stay vertical (aligned with gravity) but unable to change pitch or roll to get there.
Includes fixes for custom movement modes that were not considered as "ShouldRemainVertical"

diff --git a/Engine/Source/Runtime/Engine/Private/Components/CharacterMovementComponent.cpp b/Engine/Source/Runtime/Engine/Private/Components/CharacterMovementComponent.cpp
--- a/Engine/Source/Runtime/Engine/Private/Components/CharacterMovementComponent.cpp
+++ b/Engine/Source/Runtime/Engine/Private/Components/CharacterMovementComponent.cpp
@@ -310,6 +310,13 @@ namespace CharacterMovementCVars
 		TEXT("When enabled, during a correction, restore the last good rotation before re-simulating saved moves if the server didn't specify one. This improves visual quality with options like bOrientToMovement or bUseControllerDesiredRotation that rotate over time."),
 		ECVF_Default);
 
+	static int32 bPreventNonVerticalOrientationBlock = 1;
+	FAutoConsoleVariableRef CVarPreventNonVerticalOrientationBlock(
+		TEXT("p.PreventNonVerticalOrientationBlock"),
+		bPreventNonVerticalOrientationBlock,
+		TEXT("When enabled, this allows a character that's supposed to remain vertical to snap to a vertical orientation even if RotationRate settings would block it. See @ShouldRemainVertical and @RotationRate."),
+		ECVF_Default);
+
 #if !UE_BUILD_SHIPPING
 
 	int32 NetShowCorrections = 0;
@@ -6198,7 +6205,9 @@ void UCharacterMovementComponent::PhysicsRotation(float DeltaTime)
 		return;
 	}
 
-	if (ShouldRemainVertical())
+	const bool bWantsToBeVertical = ShouldRemainVertical();
+
+	if (bWantsToBeVertical)
 	{
 		if (HasCustomGravity())
 		{
@@ -6225,6 +6234,19 @@ void UCharacterMovementComponent::PhysicsRotation(float DeltaTime)
 
 	if (!CurrentRotation.Equals(DesiredRotation, AngleTolerance))
 	{
+		// If we'd be prevented from becoming vertical, override the non-yaw rotation rates to allow the character to snap upright
+		if (CharacterMovementCVars::bPreventNonVerticalOrientationBlock && bWantsToBeVertical)
+		{
+			if (FMath::IsNearlyZero(DeltaRot.Pitch))
+			{
+				DeltaRot.Pitch = 360.0;
+			}
+			if (FMath::IsNearlyZero(DeltaRot.Roll))
+			{
+				DeltaRot.Roll = 360.0;
+			}
+		}
+
 		// PITCH
 		if (!FMath::IsNearlyEqual(CurrentRotation.Pitch, DesiredRotation.Pitch, AngleTolerance))
 		{
