Get Socket is confusing
#fix add GetBoneTransform UFunction on SkinnedMeshComponent
#misc corrected typo in comment

diff --git a/Engine/Source/Runtime/Engine/Classes/Components/SceneComponent.h b/Engine/Source/Runtime/Engine/Classes/Components/SceneComponent.h
--- a/Engine/Source/Runtime/Engine/Classes/Components/SceneComponent.h
+++ b/Engine/Source/Runtime/Engine/Classes/Components/SceneComponent.h
@@ -778,7 +778,7 @@ public:
 	/** 
 	 * Get world-space socket or bone location.
 	 * @param InSocketName Name of the socket or the bone to get the transform 
-	 * @return Socket transform in world space if socket if found. Otherwise it will return component's transform in world space.
+	 * @return Socket transform in world space if socket is found. Otherwise it will return component's transform in world space.
 	 */
 	UFUNCTION(BlueprintCallable, Category="Transformation", meta=(Keywords="Bone"))
 	ENGINE_API virtual FVector GetSocketLocation(FName InSocketName) const;
diff --git a/Engine/Source/Runtime/Engine/Classes/Components/SkinnedMeshComponent.h b/Engine/Source/Runtime/Engine/Classes/Components/SkinnedMeshComponent.h
--- a/Engine/Source/Runtime/Engine/Classes/Components/SkinnedMeshComponent.h
+++ b/Engine/Source/Runtime/Engine/Classes/Components/SkinnedMeshComponent.h
@@ -1174,6 +1174,14 @@ public:
 		return LeaderPoseComponent.IsValid() ? LeaderPoseComponent.Get() : this;
 	}
 
+	/** 
+	 * Get world-space bone transform.
+	 * @param InBoneName Name of the the bone to get the transform 
+	 * @return Bone transform in world space if bone is found. Otherwise it will return component's transform in world space.
+	 */
+	UFUNCTION(BlueprintCallable, Category="Transformation", meta=(Keywords="Bone"))
+	FTransform GetBoneTransform(FName InBoneName, ERelativeTransformSpace TransformSpace = RTS_World) const { return GetSocketTransform(InBoneName, TransformSpace); }
+
 	/**
 	* Returns color of the vertex.
 	*
