MorphUP Doesn't Play Edited Curve Animation
#fix only accept linear or constant interpolation curve keys when testing whether or not a curve is constant

diff --git a/Engine/Source/Runtime/Engine/Private/Curves/RichCurve.cpp b/Engine/Source/Runtime/Engine/Private/Curves/RichCurve.cpp
--- a/Engine/Source/Runtime/Engine/Private/Curves/RichCurve.cpp
+++ b/Engine/Source/Runtime/Engine/Private/Curves/RichCurve.cpp
@@ -490,7 +490,7 @@ bool FRichCurve::IsConstant(float Tolerance) const
 	const FRichCurveKey& RefKey = Keys[0];
 	for (const FRichCurveKey& Key : Keys)
 	{
-		if (!FMath::IsNearlyEqual(Key.Value, RefKey.Value, Tolerance))
+		if (!FMath::IsNearlyEqual(Key.Value, RefKey.Value, Tolerance) || Key.InterpMode == ERichCurveInterpMode::RCIM_Cubic)
 		{
 			return false;
 		}
