Fix Substrate legacy material conversion with VT input.

diff --git a/Engine/Source/Runtime/Engine/Private/Materials/HLSLMaterialTranslator.cpp b/Engine/Source/Runtime/Engine/Private/Materials/HLSLMaterialTranslator.cpp
--- a/Engine/Source/Runtime/Engine/Private/Materials/HLSLMaterialTranslator.cpp
+++ b/Engine/Source/Runtime/Engine/Private/Materials/HLSLMaterialTranslator.cpp
@@ -791,6 +791,13 @@ void FHLSLMaterialTranslator::ValidateVtPropertyLimits()
 
 		const EShaderFrequency ShaderFrequencyToValidate = FMaterialAttributeDefinitionMap::GetShaderFrequency(PropertyToValidate);
 
+		// When converting legacy material the opacty mask is automatically connected when the material uses MaterialAttributes. 
+		// To avoid false positive, we bypass VT validation for OpactyMask input
+		if (Strata::IsStrataEnabled() && PropertyToValidate == MP_OpacityMask && bStrataUsesConversionFromLegacy)
+		{
+			continue;
+		}
+
 		// check to see if this is a property that doesn't support virtual texture connections
 		if (PropertyToValidate == MP_OpacityMask || ShaderFrequencyToValidate != SF_Pixel)
 		{
