Fix assertion failures caused by [REDACTED_PERSON] components creating scene proxies while their LODs are streamed out.
Also fixes similar potential issues for HISM, and moves the check for LOD resources to a common location.

#rnx

diff --git a/Engine/Source/Runtime/Engine/Private/InstancedStaticMesh.cpp b/Engine/Source/Runtime/Engine/Private/InstancedStaticMesh.cpp
--- a/Engine/Source/Runtime/Engine/Private/InstancedStaticMesh.cpp
+++ b/Engine/Source/Runtime/Engine/Private/InstancedStaticMesh.cpp
@@ -2747,16 +2747,6 @@ FPrimitiveSceneProxy* UInstancedStaticMeshComponent::CreateStaticMeshSceneProxy(
 		return ::new Nanite::FSceneProxy(NaniteMaterials, this);
 	}
 
-	// Validate the LOD resources here
-	const FStaticMeshLODResourcesArray& LODResources = GetStaticMesh()->GetRenderData()->LODResources;
-	const int32 SMCurrentMinLOD = GetStaticMesh()->GetMinLODIdx();
-	const int32 EffectiveMinLOD = bOverrideMinLOD ? MinLOD : SMCurrentMinLOD;
-	if (LODResources.Num() == 0 || LODResources[FMath::Clamp<int32>(EffectiveMinLOD, 0, LODResources.Num() - 1)].VertexBuffers.StaticMeshVertexBuffer.GetNumVertices() == 0)
-	{
-		UE_LOG(LogStaticMesh, Verbose, TEXT("Skipping CreateSceneProxy for InstancedStaticMeshComponent %s (LOD problems)"), *GetFullName());
-		return nullptr;
-	}
-
 	return ::new FInstancedStaticMeshSceneProxy(this, GetWorld()->GetFeatureLevel());
 }
 
