Fixed SparseVolumeTextureSamplePageTable() and SparseVolumeTextureLoadPageTable() sometimes loading from the page table at incorrect coordinates when the page table is using a page table offset. Fixed incorrect page table offset in SVT sequences causing parts of the volume to appear cut off at higher mip levels. Fixed empty vdb frames during import causing massive VolumeResolution values in the resulting SVT. Fixed StreamingManager triggering an assert when all voxels of a higher mip level have been stripped away due to occupancy bitmask based compression. It seems that higher mip levels can sometimes be empty despite lower mip levels having valid data. I suspect that this is due to precision issues when averaging 8 small float values and then quantizing them during the mip mapping.
Removed unused lambda in SVT factory.

#rnx

diff --git a/Engine/Shaders/Private/SparseVolumeTexture/SparseVolumeTextureCommon.ush b/Engine/Shaders/Private/SparseVolumeTexture/SparseVolumeTextureCommon.ush
--- a/Engine/Shaders/Private/SparseVolumeTexture/SparseVolumeTextureCommon.ush
+++ b/Engine/Shaders/Private/SparseVolumeTexture/SparseVolumeTextureCommon.ush
@@ -88,7 +88,7 @@ float3 SparseVolumeTextureSamplePageTable(Texture3D<uint> PageTable, FSparseVolu
 	const float3 MipPageTableOffset = floor(Uniforms.PageTableOffset * RcpMipLevelFactor);
 	const float3 PageTableCoord = VolumePageCoord * RcpMipLevelFactor - MipPageTableOffset;
 
-	const uint PackedPhysicalTileCoord = PageTable.Load(int4(int3(PageTableCoord), MipLevel)).x;
+	const uint PackedPhysicalTileCoord = PageTable.Load(int4(floor(PageTableCoord), MipLevel)).x;
 
 	const float3 VoxelCoord = SparseVolumeTextureGetVoxelCoord(PackedPhysicalTileCoord, PageTableCoord);
 	const float3 VoxelUVW = VoxelCoord * Uniforms.TileDataTexelSize;
@@ -106,7 +106,7 @@ int3 SparseVolumeTextureLoadPageTable(Texture3D<uint> PageTable, FSparseVolumeTe
 	const float3 MipPageTableOffset = floor(Uniforms.PageTableOffset * RcpMipLevelFactor);
 	const float3 PageTableCoord = VolumeMipPageCoord - MipPageTableOffset;
 
-	const uint PackedPhysicalTileCoord = PageTable.Load(int4(int3(PageTableCoord), MipLevel)).x;
+	const uint PackedPhysicalTileCoord = PageTable.Load(int4(floor(PageTableCoord), MipLevel)).x;
 
 	const float3 VoxelCoord = SparseVolumeTextureGetVoxelCoord(PackedPhysicalTileCoord, PageTableCoord);
 	return int3(VoxelCoord);
diff --git a/Engine/Source/Editor/SparseVolumeTexture/Private/SparseVolumeTextureFactory.cpp b/Engine/Source/Editor/SparseVolumeTexture/Private/SparseVolumeTextureFactory.cpp
--- a/Engine/Source/Editor/SparseVolumeTexture/Private/SparseVolumeTextureFactory.cpp
+++ b/Engine/Source/Editor/SparseVolumeTexture/Private/SparseVolumeTextureFactory.cpp
@@ -548,18 +548,6 @@ UObject* USparseVolumeTextureFactory::ImportInternal(UClass* InClass, UObject* I
 		}
 	};
 
-	auto ComputeNumMipLevels = [](const FIntVector3& VolumeBoundsMin, const FIntVector3& VolumeBoundsMax)
-	{
-		int32 Levels = 1;
-		FIntVector3 Resolution = VolumeBoundsMax - VolumeBoundsMin;
-		while (Resolution.X > SPARSE_VOLUME_TILE_RES || Resolution.Y > SPARSE_VOLUME_TILE_RES || Resolution.Z > SPARSE_VOLUME_TILE_RES)
-		{
-			Resolution /= 2;
-			++Levels;
-		}
-		return Levels;
-	};
-
 	FIntVector3 VolumeBoundsMin = FIntVector3(INT32_MAX, INT32_MAX, INT32_MAX);
 	FIntVector3 VolumeBoundsMax = FIntVector3(INT32_MIN, INT32_MIN, INT32_MIN);
 
diff --git a/Engine/Source/Editor/SparseVolumeTexture/Private/SparseVolumeTextureOpenVDBUtility.cpp b/Engine/Source/Editor/SparseVolumeTexture/Private/SparseVolumeTextureOpenVDBUtility.cpp
--- a/Engine/Source/Editor/SparseVolumeTexture/Private/SparseVolumeTextureOpenVDBUtility.cpp
+++ b/Engine/Source/Editor/SparseVolumeTexture/Private/SparseVolumeTextureOpenVDBUtility.cpp
@@ -416,9 +416,9 @@ public:
 				GridToComponentMappings[SourceGridIndex].Add(Mapping);
 			}
 		}
-
-		SVTCreateInfo.VirtualVolumeAABBMin= SmallestAABBMin - VolumeBoundsMin;
-		SVTCreateInfo.VirtualVolumeAABBMax = LargestAABBMax - VolumeBoundsMin;
+		const bool bEmptyBounds = SmallestAABBMin.X >= LargestAABBMax.X || SmallestAABBMin.Y >= LargestAABBMax.Y || SmallestAABBMin.Z >= LargestAABBMax.Z;
+		SVTCreateInfo.VirtualVolumeAABBMin = bEmptyBounds ? FIntVector3(INT32_MAX) : (SmallestAABBMin - VolumeBoundsMin);
+		SVTCreateInfo.VirtualVolumeAABBMax = bEmptyBounds ? FIntVector3(INT32_MIN) : (LargestAABBMax - VolumeBoundsMin);
 
 		return true;
 	}
diff --git a/Engine/Source/Runtime/Engine/Classes/SparseVolumeTexture/SparseVolumeTexture.h b/Engine/Source/Runtime/Engine/Classes/SparseVolumeTexture/SparseVolumeTexture.h
--- a/Engine/Source/Runtime/Engine/Classes/SparseVolumeTexture/SparseVolumeTexture.h
+++ b/Engine/Source/Runtime/Engine/Classes/SparseVolumeTexture/SparseVolumeTexture.h
@@ -44,6 +44,13 @@ struct FHeader
 
 	FHeader() = default;
 	ENGINE_API FHeader(const FIntVector3& AABBMin, const FIntVector3& AABBMax, EPixelFormat FormatA, EPixelFormat FormatB, const FVector4f& FallbackValueA, const FVector4f& FallbackValueB);
+
+	// PageTableVolumeAABBMin needs to be aligned to a power of two such that (PageTableVolumeAABBMin / pow(2, MipLevel)) results in an integer value for all mip levels. Otherwise we could end up
+	// with higher mip levels that are shifted in world space (due to PageTableVolumeAABBMin getting rounded down for every mip level) and would therefore not have corresponding voxels for voxels of lower mip levels.
+	// Such a shifted page table mip level causes clearly visible artifacts because the volume looks cut off where the page table ends.
+	// This problem can be avoided by aligning PageTableVolumeAABBMin to pow(2, (NumMipLevelsGlobal - 1)), where NumMipLevelsGlobal is the maximum number of mip levels in the entire animated SVT sequence.
+	ENGINE_API void UpdatePageTableFromGlobalNumMipLevels(int32 NumMipLevelsGlobal);
+	ENGINE_API bool Validate(bool bPrintToLog);
 };
 
 // Describes a mip level of a SVT frame in terms of the sizes and offsets of the data in the built bulk data.
@@ -373,10 +380,10 @@ public:
 	// and doesn't need to match the exact number if it is not known at the time.
 	ENGINE_API virtual bool BeginInitialize(int32 NumExpectedFrames);
 	ENGINE_API virtual bool AppendFrame(UE::SVT::FTextureData& UncookedFrame);
-	ENGINE_API virtual bool EndInitialize(int32 NumMipLevels = INDEX_NONE /*Create entire mip chain by default*/);
+	ENGINE_API virtual bool EndInitialize();
 
 	// Convenience function wrapping the multi-phase initialization functions above
-	virtual bool Initialize(const TArrayView<UE::SVT::FTextureData>& UncookedData, int32 NumMipLevels = INDEX_NONE /*Create entire mip chain by default*/);
+	virtual bool Initialize(const TArrayView<UE::SVT::FTextureData>& UncookedData);
 	// Consider using USparseVolumeTextureFrame::GetFrameAndIssueStreamingRequest() if the frame should have streaming requests issued.
 	USparseVolumeTextureFrame* GetFrame(int32 FrameIndex) const { return Frames.IsValidIndex(FrameIndex) ? Frames[FrameIndex] : nullptr; }
 
diff --git a/Engine/Source/Runtime/Engine/Private/SparseVolumeTexture/SparseVolumeTexture.cpp b/Engine/Source/Runtime/Engine/Private/SparseVolumeTexture/SparseVolumeTexture.cpp
--- a/Engine/Source/Runtime/Engine/Private/SparseVolumeTexture/SparseVolumeTexture.cpp
+++ b/Engine/Source/Runtime/Engine/Private/SparseVolumeTexture/SparseVolumeTexture.cpp
@@ -89,13 +89,14 @@ namespace UE
 namespace SVT
 {
 
-static int32 ComputeNumMipLevels(const FIntVector3& InResolution)
+static int32 ComputeNumMipLevels(const FIntVector3& InVirtualVolumeMin, const FIntVector3& InVirtualVolumeMax)
 {
+	FHeader DummyHeader(InVirtualVolumeMin, InVirtualVolumeMax, PF_Unknown, PF_Unknown, FVector4f(), FVector4f());
 	int32 Levels = 1;
-	FIntVector3 Resolution = InResolution;
-	while (Resolution.X > SPARSE_VOLUME_TILE_RES || Resolution.Y > SPARSE_VOLUME_TILE_RES || Resolution.Z > SPARSE_VOLUME_TILE_RES)
+	FIntVector3 PageTableResolution = DummyHeader.PageTableVolumeResolution;
+	while (PageTableResolution.X > 1 || PageTableResolution.Y > 1 || PageTableResolution.Z > 1)
 	{
-		Resolution /= 2;
+		PageTableResolution /= 2;
 		++Levels;
 	}
 	return Levels;
@@ -130,6 +131,60 @@ FHeader::FHeader(const FIntVector3& AABBMin, const FIntVector3& AABBMax, EPixelF
 	FallbackValues[1] = FallbackValueB;
 }
 
+void FHeader::UpdatePageTableFromGlobalNumMipLevels(int32 NumMipLevelsGlobal)
+{
+	check(NumMipLevelsGlobal < 32);
+	const int32 Alignment = 1 << (NumMipLevelsGlobal - 1);
+
+	PageTableVolumeAABBMin = ((VirtualVolumeAABBMin / SPARSE_VOLUME_TILE_RES) / Alignment) * Alignment;
+	PageTableVolumeAABBMax = (VirtualVolumeAABBMax + FIntVector3(SPARSE_VOLUME_TILE_RES - 1)) / SPARSE_VOLUME_TILE_RES;
+	PageTableVolumeResolution = PageTableVolumeAABBMax - PageTableVolumeAABBMin;
+
+	// We need to ensure a power of two resolution for the page table in order to fit all mips of the page table into the physical mips of the texture resource.
+	PageTableVolumeResolution.X = FMath::RoundUpToPowerOfTwo(PageTableVolumeResolution.X);
+	PageTableVolumeResolution.Y = FMath::RoundUpToPowerOfTwo(PageTableVolumeResolution.Y);
+	PageTableVolumeResolution.Z = FMath::RoundUpToPowerOfTwo(PageTableVolumeResolution.Z);
+	PageTableVolumeAABBMax = PageTableVolumeAABBMin + PageTableVolumeResolution;
+}
+
+bool FHeader::Validate(bool bPrintToLog)
+{
+	for (int32 i = 0; i < 2; ++i)
+	{
+		if (!SVT::IsSupportedFormat(AttributesFormats[i]))
+		{
+			if (bPrintToLog)
+			{
+				UE_LOG(LogSparseVolumeTexture, Warning, TEXT("'%s' is not a supported SparseVolumeTexture format!"), GPixelFormats[AttributesFormats[i]].Name);
+			}
+			return false;
+		}
+	}
+
+	if (PageTableVolumeResolution.X > SVT::MaxVolumeTextureDim || PageTableVolumeResolution.Y > SVT::MaxVolumeTextureDim || PageTableVolumeResolution.Z > SVT::MaxVolumeTextureDim)
+	{
+		if (bPrintToLog)
+		{
+			UE_LOG(LogSparseVolumeTexture, Warning, TEXT("SparseVolumeTexture page table texture dimensions exceed limit (%ix%ix%i): %ix%ix%i"),
+				SVT::MaxVolumeTextureDim, SVT::MaxVolumeTextureDim, SVT::MaxVolumeTextureDim,
+				PageTableVolumeResolution.X, PageTableVolumeResolution.Y, PageTableVolumeResolution.Z);
+		}
+		return false;
+	}
+
+	const int64 PageTableSizeBytes = (int64)PageTableVolumeResolution.X * (int64)PageTableVolumeResolution.Y * (int64)PageTableVolumeResolution.Z * sizeof(uint32);
+	if (PageTableSizeBytes > SVT::MaxResourceSize)
+	{
+		if (bPrintToLog)
+		{
+			UE_LOG(LogSparseVolumeTexture, Warning, TEXT("SparseVolumeTexture page table texture memory size (%ll) exceeds the 2048MB GPU resource limit!"), (long long)PageTableSizeBytes);
+		}
+		return false;
+	}
+
+	return true;
+}
+
 ////////////////////////////////////////////////////////////////////////////////////////////////
 
 void FResources::Serialize(FArchive& Ar, UObject* Owner, bool bCooked)
@@ -256,10 +311,10 @@ bool FResources::Build(USparseVolumeTextureFrame* Owner, UE::Serialization::FEdi
 		AddressingInfo.AddressY = Owner->GetTextureAddressY();
 		AddressingInfo.AddressZ = Owner->GetTextureAddressZ();
 
-		const int32 NumMipLevelsRequested = -1; // generate entire mip chain
+		const int32 NumMipLevelsGlobal = Owner->GetNumMipLevels();
 		const bool bMoveMip0FromSource = true; // we have no need to keep SourceTextureData around
 		FTextureData DerivedTextureData;
-		if (!SourceTextureData.BuildDerivedData(AddressingInfo, NumMipLevelsRequested, bMoveMip0FromSource, DerivedTextureData))
+		if (!SourceTextureData.BuildDerivedData(AddressingInfo, NumMipLevelsGlobal, bMoveMip0FromSource, DerivedTextureData))
 		{
 			return false;
 		}
@@ -1062,7 +1117,16 @@ bool UStreamableSparseVolumeTexture::AppendFrame(UE::SVT::FTextureData& Uncooked
 	VolumeBoundsMax.Y = FMath::Max(VolumeBoundsMax.Y, UncookedFrame.Header.VirtualVolumeAABBMax.Y);
 	VolumeBoundsMax.Z = FMath::Max(VolumeBoundsMax.Z, UncookedFrame.Header.VirtualVolumeAABBMax.Z);
 
-	VolumeResolution = VolumeBoundsMax;
+	if (VolumeBoundsMin.X >= VolumeBoundsMax.X || VolumeBoundsMin.Y >= VolumeBoundsMax.Y || VolumeBoundsMin.Z >= VolumeBoundsMax.Z)
+	{
+		// Force a minimum resolution of (1, 1, 1) if the frames were empty so far
+		VolumeResolution = FIntVector3(1, 1, 1);
+	}
+	else
+	{
+		VolumeResolution = VolumeBoundsMax;
+	}
+	
 
 	USparseVolumeTextureFrame* Frame = NewObject<USparseVolumeTextureFrame>(this);
 	if (Frame->Initialize(this, Frames.Num(), UncookedFrame))
@@ -1077,7 +1141,7 @@ bool UStreamableSparseVolumeTexture::AppendFrame(UE::SVT::FTextureData& Uncooked
 #endif
 }
 
-bool UStreamableSparseVolumeTexture::EndInitialize(int32 InNumMipLevels)
+bool UStreamableSparseVolumeTexture::EndInitialize()
 {
 #if WITH_EDITORONLY_DATA
 	if (InitState != EInitState_Pending)
@@ -1105,10 +1169,10 @@ bool UStreamableSparseVolumeTexture::EndInitialize(int32 InNumMipLevels)
 			VolumeBoundsMin.X, VolumeBoundsMin.Y, VolumeBoundsMin.Z);
 	}
 
-	const int32 NumMipLevelsFullMipChain = UE::SVT::ComputeNumMipLevels(VolumeResolution);
+	const int32 NumMipLevelsFullMipChain = UE::SVT::ComputeNumMipLevels(VolumeBoundsMin, VolumeBoundsMax);
 	check(NumMipLevelsFullMipChain > 0);
 
-	NumMipLevels = (InNumMipLevels <= INDEX_NONE) ? NumMipLevelsFullMipChain : FMath::Clamp(InNumMipLevels, 1, NumMipLevelsFullMipChain);
+	NumMipLevels = NumMipLevelsFullMipChain;
 	NumFrames = Frames.Num();
 
 	for (USparseVolumeTextureFrame* Frame : Frames)
@@ -1124,7 +1188,7 @@ bool UStreamableSparseVolumeTexture::EndInitialize(int32 InNumMipLevels)
 #endif
 }
 
-bool UStreamableSparseVolumeTexture::Initialize(const TArrayView<UE::SVT::FTextureData>& InUncookedData, int32 InNumMipLevels)
+bool UStreamableSparseVolumeTexture::Initialize(const TArrayView<UE::SVT::FTextureData>& InUncookedData)
 {
 	if (InUncookedData.IsEmpty())
 	{
@@ -1143,7 +1207,7 @@ bool UStreamableSparseVolumeTexture::Initialize(const TArrayView<UE::SVT::FTextu
 			return false;
 		}
 	}
-	if (!EndInitialize(InNumMipLevels))
+	if (!EndInitialize())
 	{
 		return false;
 	}
diff --git a/Engine/Source/Runtime/Engine/Private/SparseVolumeTexture/SparseVolumeTextureData.cpp b/Engine/Source/Runtime/Engine/Private/SparseVolumeTexture/SparseVolumeTextureData.cpp
--- a/Engine/Source/Runtime/Engine/Private/SparseVolumeTexture/SparseVolumeTextureData.cpp
+++ b/Engine/Source/Runtime/Engine/Private/SparseVolumeTexture/SparseVolumeTextureData.cpp
@@ -60,30 +60,8 @@ bool FTextureData::Create(const ITextureDataProvider& DataProvider)
 	const FTextureDataCreateInfo CreateInfo = DataProvider.GetCreateInfo();
 
 	Header = FHeader(CreateInfo.VirtualVolumeAABBMin, CreateInfo.VirtualVolumeAABBMax, CreateInfo.AttributesFormats[0], CreateInfo.AttributesFormats[1], CreateInfo.FallbackValues[0], CreateInfo.FallbackValues[1]);
-
-	for (int32 i = 0; i < 2; ++i)
+	if (!Header.Validate(true /*bPrintToLog*/))
 	{
-		if (!SVT::IsSupportedFormat(Header.AttributesFormats[i]))
-		{
-			UE_LOG(LogSparseVolumeTextureData, Warning, TEXT("'%s' is not a supported SparseVolumeTexture format!"), GPixelFormats[Header.AttributesFormats[i]].Name);
-			return false;
-		}
-	}
-
-	if (Header.PageTableVolumeResolution.X > SVT::MaxVolumeTextureDim
-		|| Header.PageTableVolumeResolution.Y > SVT::MaxVolumeTextureDim
-		|| Header.PageTableVolumeResolution.Z > SVT::MaxVolumeTextureDim)
-	{
-		UE_LOG(LogSparseVolumeTextureData, Warning, TEXT("SparseVolumeTexture page table texture dimensions exceed limit (%ix%ix%i): %ix%ix%i"),
-			SVT::MaxVolumeTextureDim, SVT::MaxVolumeTextureDim, SVT::MaxVolumeTextureDim,
-			Header.PageTableVolumeResolution.X, Header.PageTableVolumeResolution.Y, Header.PageTableVolumeResolution.Z);
-		return false;
-	}
-
-	const int64 PageTableSizeBytes = (int64)Header.PageTableVolumeResolution.X * (int64)Header.PageTableVolumeResolution.Y * (int64)Header.PageTableVolumeResolution.Z * sizeof(uint32);
-	if (PageTableSizeBytes > SVT::MaxResourceSize)
-	{
-		UE_LOG(LogSparseVolumeTextureData, Warning, TEXT("SparseVolumeTexture page table texture memory size (%ll) exceeds the 2048MB GPU resource limit!"), (long long)PageTableSizeBytes);
 		return false;
 	}
 
@@ -174,7 +152,7 @@ bool FTextureData::Create(const ITextureDataProvider& DataProvider)
 	DataProvider.IteratePhysicalSource([&](const FIntVector3& Coord, int32 AttributesIdx, int32 ComponentIdx, float VoxelValue)
 		{
 			const FIntVector3 GridCoord = Coord;
-			check(IsInBounds(GridCoord, FIntVector3::ZeroValue, Header.VirtualVolumeAABBMax));
+			check(IsInBounds(GridCoord, Header.VirtualVolumeAABBMin, Header.VirtualVolumeAABBMax));
 			const FIntVector3 PageCoord = (GridCoord / SPARSE_VOLUME_TILE_RES) - Header.PageTableVolumeAABBMin;
 			check(IsInBounds(PageCoord, FIntVector3::ZeroValue, Header.PageTableVolumeResolution));
 			const int32 PageIndex = PageCoord.Z * (Header.PageTableVolumeResolution.Y * Header.PageTableVolumeResolution.X) + PageCoord.Y * Header.PageTableVolumeResolution.X + PageCoord.X;
@@ -194,33 +172,11 @@ bool FTextureData::Create(const ITextureDataProvider& DataProvider)
 bool FTextureData::CreateFromDense(const FTextureDataCreateInfo& CreateInfo, const TArrayView64<uint8>& VoxelDataA, const TArrayView64<uint8>& VoxelDataB)
 {	
 	Header = FHeader(CreateInfo.VirtualVolumeAABBMin, CreateInfo.VirtualVolumeAABBMax, CreateInfo.AttributesFormats[0], CreateInfo.AttributesFormats[1], CreateInfo.FallbackValues[0], CreateInfo.FallbackValues[1]);
-
-	for (int32 i = 0; i < 2; ++i)
-	{
-		if (!SVT::IsSupportedFormat(Header.AttributesFormats[i]))
-		{
-			UE_LOG(LogSparseVolumeTextureData, Warning, TEXT("'%s' is not a supported SparseVolumeTexture format!"), GPixelFormats[Header.AttributesFormats[i]].Name);
-			return false;
-		}
-	}
-
-	if (Header.PageTableVolumeResolution.X > SVT::MaxVolumeTextureDim
-		|| Header.PageTableVolumeResolution.Y > SVT::MaxVolumeTextureDim
-		|| Header.PageTableVolumeResolution.Z > SVT::MaxVolumeTextureDim)
-	{
-		UE_LOG(LogSparseVolumeTextureData, Warning, TEXT("SparseVolumeTexture page table texture dimensions exceed limit (%ix%ix%i): %ix%ix%i"),
-			SVT::MaxVolumeTextureDim, SVT::MaxVolumeTextureDim, SVT::MaxVolumeTextureDim,
-			Header.PageTableVolumeResolution.X, Header.PageTableVolumeResolution.Y, Header.PageTableVolumeResolution.Z);
-		return false;
-	}
-
-	const int64 PageTableSizeBytes = (int64)Header.PageTableVolumeResolution.X * (int64)Header.PageTableVolumeResolution.Y * (int64)Header.PageTableVolumeResolution.Z * sizeof(uint32);
-	if (PageTableSizeBytes > SVT::MaxResourceSize)
+	if (!Header.Validate(true /*bPrintToLog*/))
 	{
-		UE_LOG(LogSparseVolumeTextureData, Warning, TEXT("SparseVolumeTexture page table texture memory size (%ll) exceeds the 2048MB GPU resource limit!"), (long long)PageTableSizeBytes);
 		return false;
 	}
-
+	
 	const int64 FormatSize[] = { GPixelFormats[Header.AttributesFormats[0]].BlockBytes, GPixelFormats[Header.AttributesFormats[1]].BlockBytes };
 	const FIntVector3 FractionalPageOffset = Header.VirtualVolumeAABBMin % SPARSE_VOLUME_TILE_RES; // Offset from page origin to source data volume origin
 	uint8 NullTileValuesU8[2][sizeof(float) * 4] = {};
@@ -510,10 +466,10 @@ bool FTextureData::GenerateMipMaps(const FTextureDataAddressingInfo& AddressingI
 	if (NumMipLevels == -1)
 	{
 		NumMipLevels = 1;
-		FIntVector3 Resolution = Header.VirtualVolumeResolution;
-		while (Resolution.X > SPARSE_VOLUME_TILE_RES || Resolution.Y > SPARSE_VOLUME_TILE_RES || Resolution.Z > SPARSE_VOLUME_TILE_RES)
+		FIntVector3 PageTableResolution = Header.PageTableVolumeResolution;
+		while (PageTableResolution.X > 1 || PageTableResolution.Y > 1 || PageTableResolution.Z > 1)
 		{
-			Resolution /= 2;
+			PageTableResolution /= 2;
 			++NumMipLevels;
 		}
 	}
@@ -526,7 +482,7 @@ bool FTextureData::GenerateMipMaps(const FTextureDataAddressingInfo& AddressingI
 	{
 		MipMaps.AddDefaulted();
 
-		const FIntVector3 PageTableVolumeAABBMin = (Header.VirtualVolumeAABBMin / SPARSE_VOLUME_TILE_RES) >> MipLevel;
+		const FIntVector3 PageTableVolumeAABBMin = Header.PageTableVolumeAABBMin >> MipLevel;
 		const FIntVector3 PageTableVolumeResolution = ShiftRightAndMax(Header.PageTableVolumeResolution, MipLevel, 1);
 		check(FMath::IsPowerOfTwo(PageTableVolumeResolution.X)
 			&& FMath::IsPowerOfTwo(PageTableVolumeResolution.Y)
@@ -767,10 +723,15 @@ bool FTextureData::DeduplicateTiles()
 	return true;
 }
 
-bool FTextureData::BuildDerivedData(const FTextureDataAddressingInfo& AddressingInfo, int32 NumMipLevels, bool bMoveMip0FromThis, FTextureData& OutDerivedData)
+bool FTextureData::BuildDerivedData(const FTextureDataAddressingInfo& AddressingInfo, int32 NumMipLevelsGlobal, bool bMoveMip0FromThis, FTextureData& OutDerivedData)
 {
 	OutDerivedData = FTextureData{};
 	OutDerivedData.Header = Header;
+	OutDerivedData.Header.UpdatePageTableFromGlobalNumMipLevels(NumMipLevelsGlobal);
+	if (!OutDerivedData.Header.Validate(true /*bPrintToLog*/))
+	{
+		return false;
+	}
 	OutDerivedData.MipMaps.SetNum(1);
 	if (bMoveMip0FromThis)
 	{
@@ -782,6 +743,45 @@ bool FTextureData::BuildDerivedData(const FTextureDataAddressingInfo& Addressing
 		OutDerivedData.MipMaps[0] = MipMaps[0];
 	}
 
+	// Page table changed due to UpdatePageTableFromGlobalNumMipLevels(), so we need to recreate it based on the old one
+	if (OutDerivedData.Header.PageTableVolumeAABBMin != Header.PageTableVolumeAABBMin
+		|| OutDerivedData.Header.PageTableVolumeAABBMax != Header.PageTableVolumeAABBMax
+		|| OutDerivedData.Header.PageTableVolumeResolution != Header.PageTableVolumeResolution)
+	{
+		TArray<uint32> PageTableNew;
+		PageTableNew.SetNumZeroed(OutDerivedData.Header.PageTableVolumeResolution.X * OutDerivedData.Header.PageTableVolumeResolution.Y * OutDerivedData.Header.PageTableVolumeResolution.Z);
+
+		// Iterate over new page table and copy entries from old one
+		for (int32 Z = OutDerivedData.Header.PageTableVolumeAABBMin.Z; Z < OutDerivedData.Header.PageTableVolumeAABBMax.Z; ++Z)
+		{
+			for (int32 Y = OutDerivedData.Header.PageTableVolumeAABBMin.Y; Y < OutDerivedData.Header.PageTableVolumeAABBMax.Y; ++Y)
+			{
+				for (int32 X = OutDerivedData.Header.PageTableVolumeAABBMin.X; X < OutDerivedData.Header.PageTableVolumeAABBMax.X; ++X)
+				{
+					FIntVector3 GlobalPageTableCoord(X, Y, Z);
+					if (IsInBounds(GlobalPageTableCoord, Header.PageTableVolumeAABBMin, Header.PageTableVolumeAABBMax))
+					{
+						const FIntVector3 LocalPageTableCoordOld = GlobalPageTableCoord - Header.PageTableVolumeAABBMin;
+						const int32 PageTableIndexOld = (LocalPageTableCoordOld.Z * Header.PageTableVolumeResolution.Y * Header.PageTableVolumeResolution.X) 
+							+ (LocalPageTableCoordOld.Y * Header.PageTableVolumeResolution.X) 
+							+ LocalPageTableCoordOld.X;
+						check(OutDerivedData.MipMaps[0].PageTable.IsValidIndex(PageTableIndexOld));
+
+						const FIntVector3 LocalPageTableCoordNew = GlobalPageTableCoord - OutDerivedData.Header.PageTableVolumeAABBMin;
+						const int32 PageTableIndexNew = (LocalPageTableCoordNew.Z * OutDerivedData.Header.PageTableVolumeResolution.Y * OutDerivedData.Header.PageTableVolumeResolution.X) 
+							+ (LocalPageTableCoordNew.Y * OutDerivedData.Header.PageTableVolumeResolution.X) 
+							+ LocalPageTableCoordNew.X;
+						check(PageTableNew.IsValidIndex(PageTableIndexNew));
+
+						PageTableNew[PageTableIndexNew] = OutDerivedData.MipMaps[0].PageTable[PageTableIndexOld];
+					}
+				}
+			}
+		}
+
+		OutDerivedData.MipMaps[0].PageTable = MoveTemp(PageTableNew);
+	}
+
 	// Generate border voxels of mip0
 	{
 		// Collect the page table coordinates of all non-zero pages
@@ -809,7 +809,7 @@ bool FTextureData::BuildDerivedData(const FTextureDataAddressingInfo& Addressing
 	}
 
 	// Generate all remaining mips. Also generates border voxels.
-	if (!OutDerivedData.GenerateMipMaps(AddressingInfo, NumMipLevels))
+	if (!OutDerivedData.GenerateMipMaps(AddressingInfo))
 	{
 		return false;
 	}
diff --git a/Engine/Source/Runtime/Engine/Private/SparseVolumeTexture/SparseVolumeTextureStreamingManager.cpp b/Engine/Source/Runtime/Engine/Private/SparseVolumeTexture/SparseVolumeTextureStreamingManager.cpp
--- a/Engine/Source/Runtime/Engine/Private/SparseVolumeTexture/SparseVolumeTextureStreamingManager.cpp
+++ b/Engine/Source/Runtime/Engine/Private/SparseVolumeTexture/SparseVolumeTextureStreamingManager.cpp
@@ -1364,8 +1364,6 @@ void FStreamingManager::AddInternal(FRDGBuilder& GraphBuilder, FNewSparseVolumeT
 			const FMipLevelStreamingInfo* RootStreamingInfo = !Resources->MipLevelStreamingInfo.IsEmpty() ? &Resources->MipLevelStreamingInfo.Last() : nullptr;
 			if (!Resources->RootData.IsEmpty() && RootStreamingInfo)
 			{
-				check((RootStreamingInfo->TileDataSize[0] > 0) || (RootStreamingInfo->TileDataSize[1] > 0));
-
 				const uint32 TileCoord = SVTInfo.TileDataTexture->Allocate();
 				check(TileCoord != INDEX_NONE);
 				FrameInfo.TileAllocations.Last()[0] = TileCoord;
diff --git a/Engine/Source/Runtime/Engine/Public/SparseVolumeTexture/SparseVolumeTextureData.h b/Engine/Source/Runtime/Engine/Public/SparseVolumeTexture/SparseVolumeTextureData.h
--- a/Engine/Source/Runtime/Engine/Public/SparseVolumeTexture/SparseVolumeTextureData.h
+++ b/Engine/Source/Runtime/Engine/Public/SparseVolumeTexture/SparseVolumeTextureData.h
@@ -55,7 +55,7 @@ struct FTextureData
 	ENGINE_API bool CreateFromDense(const FTextureDataCreateInfo& CreateInfo, const TArrayView64<uint8>& VoxelDataA, const TArrayView64<uint8>& VoxelDataB);
 	ENGINE_API void CreateDefault();
 	ENGINE_API FVector4f Load(const FIntVector3& VolumeCoord, int32 MipLevel, int32 AttributesIdx, const FTextureDataAddressingInfo& AddressingInfo) const;
-	ENGINE_API bool BuildDerivedData(const FTextureDataAddressingInfo& AddressingInfo, int32 NumMipLevels, bool bMoveMip0FromThis, FTextureData& OutDerivedData);
+	ENGINE_API bool BuildDerivedData(const FTextureDataAddressingInfo& AddressingInfo, int32 NumMipLevelsGlobal, bool bMoveMip0FromThis, FTextureData& OutDerivedData);
 
 private:
 
