World Partition 2D Debug Display:
- Cut by two the number of wortld to screen transforms sent to FWorldPartitionDraw2DContext.
- Fix FWorldPartitionDraw2DContext::PushDrawTile not properly clipping input tile box.
- Implement initial 2D debug for runtime spatial hash set.

diff --git a/Engine/Source/Runtime/Engine/Private/WorldPartition/RuntimeHashSet/RuntimePartition.cpp b/Engine/Source/Runtime/Engine/Private/WorldPartition/RuntimeHashSet/RuntimePartition.cpp
--- a/Engine/Source/Runtime/Engine/Private/WorldPartition/RuntimeHashSet/RuntimePartition.cpp
+++ b/Engine/Source/Runtime/Engine/Private/WorldPartition/RuntimeHashSet/RuntimePartition.cpp
@@ -3,6 +3,18 @@
 #include "WorldPartition/RuntimeHashSet/RuntimePartition.h"
 
 #if WITH_EDITOR
+void URuntimePartition::PostEditChangeProperty(FPropertyChangedEvent& PropertyChangedEvent)
+{
+	const FName PropertyName = PropertyChangedEvent.GetPropertyName();
+
+	if (PropertyName == GET_MEMBER_NAME_CHECKED(URuntimePartition, LoadingRange))
+	{
+		LoadingRange = FMath::Max(LoadingRange, 0);
+	}
+
+	Super::PostEditChangeProperty(PropertyChangedEvent);
+}
+
 bool URuntimePartition::PopulateCellActorInstances(const TArray<const IStreamingGenerationContext::FActorSetInstance*>& ActorSetInstances, bool bIsMainWorldPartition, bool bIsCellAlwaysLoaded, TArray<IStreamingGenerationContext::FActorInstance>& OutCellActorInstances)
 {
 	UWorldPartitionRuntimeHash* RuntimeHash = GetTypedOuter<UWorldPartitionRuntimeHash>();
diff --git a/Engine/Source/Runtime/Engine/Private/WorldPartition/RuntimeHashSet/RuntimePartitionLevelStreaming.cpp b/Engine/Source/Runtime/Engine/Private/WorldPartition/RuntimeHashSet/RuntimePartitionLevelStreaming.cpp
--- a/Engine/Source/Runtime/Engine/Private/WorldPartition/RuntimeHashSet/RuntimePartitionLevelStreaming.cpp
+++ b/Engine/Source/Runtime/Engine/Private/WorldPartition/RuntimeHashSet/RuntimePartitionLevelStreaming.cpp
@@ -67,6 +67,7 @@ bool URuntimePartitionLevelStreaming::GenerateStreaming(const TArray<const IStre
 			CellDesc.bBlockOnSlowStreaming = bBlockOnSlowStreaming;
 			CellDesc.bClientOnlyVisible = bClientOnlyVisible;
 			CellDesc.Priority = Priority;
+			CellDesc.Level = 0;
 			CellDesc.ActorInstances = SubLevelActorSetInstances;
 
 			for (const IStreamingGenerationContext::FActorInstance& ActorInstance : CellDesc.ActorInstances)
diff --git a/Engine/Source/Runtime/Engine/Private/WorldPartition/RuntimeHashSet/WorldPartitionRuntimeHashSet.cpp b/Engine/Source/Runtime/Engine/Private/WorldPartition/RuntimeHashSet/WorldPartitionRuntimeHashSet.cpp
--- a/Engine/Source/Runtime/Engine/Private/WorldPartition/RuntimeHashSet/WorldPartitionRuntimeHashSet.cpp
+++ b/Engine/Source/Runtime/Engine/Private/WorldPartition/RuntimeHashSet/WorldPartitionRuntimeHashSet.cpp
@@ -3,11 +3,16 @@
 #include "WorldPartition/RuntimeHashSet/WorldPartitionRuntimeHashSet.h"
 #include "WorldPartition/RuntimeHashSet/RuntimePartition.h"
 #include "WorldPartition/RuntimeHashSet/RuntimePartitionPersistent.h"
+#include "WorldPartition/ContentBundle/ContentBundleWorldSubsystem.h"
 #include "WorldPartition/ContentBundle/ContentBundleDescriptor.h"
+#include "WorldPartition/ContentBundle/ContentBundleBase.h"
+#include "WorldPartition/DataLayer/DataLayerManager.h"
 #include "WorldPartition/WorldPartitionStreamingGenerationContext.h"
 #include "WorldPartition/WorldPartitionRuntimeSpatialHash.h"
 #include "WorldPartition/WorldPartitionStreamingPolicy.h"
 #include "WorldPartition/WorldPartitionStreamingSource.h"
+#include "WorldPartition/WorldPartitionDraw2DContext.h"
+#include "WorldPartition/WorldPartitionDebugHelper.h"
 #include "WorldPartition/DataLayer/DataLayersID.h"
 #include "WorldPartition/HLOD/HLODLayer.h"
 #include "Algo/ForEach.h"
@@ -23,6 +28,18 @@ FAutoConsoleCommand WorldPartitionRuntimeHashSetEnable(
 	})
 );
 
+static int32 GShowRuntimeHashSetDebugDisplayLevel = 0;
+static FAutoConsoleVariableRef CVarShowRuntimeHashSetDebugDisplayLevel(
+	TEXT("wp.Runtime.HashSet.ShowtDebugDisplayLevel"),
+	GShowRuntimeHashSetDebugDisplayLevel,
+	TEXT("Used to choose which level to display when showing world partition partitions."));
+
+static int32 GShowRuntimeHashSetDebugDisplayLevelCount = 1;
+static FAutoConsoleVariableRef CVarShowRuntimeHashSetDebugDisplayLevelCount(
+	TEXT("wp.Runtime.HashSet.ShowDebugDisplayLevelCount"),
+	GShowRuntimeHashSetDebugDisplayLevelCount,
+	TEXT("Used to choose how many levels to display when showing world partition runtime partitions."));
+
 #if WITH_EDITOR
 void FRuntimePartitionDesc::UpdateHLODPartitionLayers()
 {
@@ -139,6 +156,256 @@ void UWorldPartitionRuntimeHashSet::PostLoad()
 }
 
 #if WITH_EDITOR
+static EWorldPartitionRuntimeCellVisualizeMode GetStreamingCellVisualizeMode()
+{
+	const EWorldPartitionRuntimeCellVisualizeMode VisualizeMode = FWorldPartitionDebugHelper::IsRuntimeSpatialHashCellStreamingPriorityShown() ? EWorldPartitionRuntimeCellVisualizeMode::StreamingPriority : EWorldPartitionRuntimeCellVisualizeMode::StreamingStatus;
+	return VisualizeMode;
+}
+
+static TMap<FName, FColor> GetDataLayerDebugColors(const UWorldPartition* InWorldPartition)
+{
+	TMap<FName, FColor> DebugColors;
+	if (const UDataLayerManager* DataLayerManager = InWorldPartition->GetDataLayerManager())
+	{
+		DataLayerManager->ForEachDataLayerInstance([&DebugColors](UDataLayerInstance* DataLayerInstance)
+		{
+			DebugColors.Add(DataLayerInstance->GetDataLayerFName(), DataLayerInstance->GetDebugColor());
+			return true;
+		});
+	}
+	return DebugColors;
+}
+
+bool UWorldPartitionRuntimeHashSet::Draw2D(FWorldPartitionDraw2DContext& DrawContext) const
+{
+	TRACE_CPUPROFILER_EVENT_SCOPE(UWorldPartitionRuntimeHashSet::Draw2D);
+
+	const UWorldPartition* WorldPartition = GetOuterUWorldPartition();
+	const TArray<FWorldPartitionStreamingSource>& Sources = WorldPartition->GetStreamingSources();
+	
+	if (!Sources.Num())
+	{
+		return false;
+	}
+
+	TMap<FName, TArray<const FRuntimePartitionStreamingData*>> FilteredStreamingObjects;
+	ForEachStreamingObject([&FilteredStreamingObjects](const URuntimeHashSetExternalStreamingObject* CurStreamingObject)
+	{
+		for (const FRuntimePartitionStreamingData& RuntimeStreamingData : CurStreamingObject->RuntimeStreamingData)
+		{
+			if (FWorldPartitionDebugHelper::IsDebugRuntimeHashGridShown(RuntimeStreamingData.Name))
+			{
+				FilteredStreamingObjects.FindOrAdd(RuntimeStreamingData.Name).Add(&RuntimeStreamingData);
+			}
+		}
+	});
+
+	if (!FilteredStreamingObjects.Num())
+	{
+		return false;
+	}
+
+	const FTransform WorldPartitionTransform = WorldPartition->GetInstanceTransform();
+	const FTransform2D LocalWPToGlobal(FQuat2D(FMath::DegreesToRadians(WorldPartitionTransform.Rotator().Yaw)), FVector2D(WorldPartitionTransform.GetLocation()));
+	const FTransform2D GlobalToLocalWP = LocalWPToGlobal.Inverse();
+	const FBox2D& WorldRegion = DrawContext.GetWorldRegion();
+	const FVector2D WorldRegionSize = WorldRegion.GetSize();
+	TArray<FVector2D> Points;
+	Points.Add(GlobalToLocalWP.TransformPoint(WorldRegion.Min));
+	Points.Add(GlobalToLocalWP.TransformPoint(WorldRegion.Min + FVector2D(WorldRegionSize.X, 0)));
+	Points.Add(GlobalToLocalWP.TransformPoint(WorldRegion.Min + WorldRegionSize));
+	Points.Add(GlobalToLocalWP.TransformPoint(WorldRegion.Min + FVector2D(0, WorldRegionSize.Y)));
+	const FBox2D LocalRegion(Points);
+
+	const FBox2D& CanvasRegion = DrawContext.GetCanvasRegion();
+	const int32 GridScreenWidthDivider = DrawContext.IsDetailedMode() ? FilteredStreamingObjects.Num() : 1;
+	const float GridScreenWidthShrinkSize = GridScreenWidthDivider > 1 ? 20.f : 0.f;
+	const float CanvasMaxScreenWidth = CanvasRegion.GetSize().X;
+	const float GridMaxScreenWidth = CanvasMaxScreenWidth / GridScreenWidthDivider;
+	const float GridEffectiveScreenWidth = FMath::Min(GridMaxScreenWidth, CanvasRegion.GetSize().Y) - GridScreenWidthShrinkSize;
+	const FVector2D PartitionCanvasSize = FVector2D(CanvasRegion.GetSize().GetMin());
+	const FVector2D GridScreenExtent = FVector2D(GridEffectiveScreenWidth, GridEffectiveScreenWidth);
+	const FVector2D GridScreenHalfExtent = 0.5f * GridScreenExtent;
+	const FVector2D GridScreenInitialOffset = CanvasRegion.Min;
+
+	auto DrawStreamingData = [this](const FRuntimePartitionStreamingData* StreamingData, const FBox2D& Region2D, const FBox2D& GridScreenBounds, TFunctionRef<FVector2D(const FVector2D&, bool)> InWorldToScreen, FWorldPartitionDraw2DContext& DrawContext)
+	{
+		const UWorldPartition* WorldPartition = GetOuterUWorldPartition();
+		const TArray<FWorldPartitionStreamingSource>& Sources = WorldPartition->GetStreamingSources();
+
+		const FBox Region(FVector(Region2D.Min.X, Region2D.Min.Y, 0), FVector(Region2D.Max.X, Region2D.Max.Y, 0));
+		const UWorld* OwningWorld = WorldPartition->GetWorld();
+		const EWorldPartitionRuntimeCellVisualizeMode VisualizeMode = GetStreamingCellVisualizeMode();
+		const UContentBundleManager* ContentBundleManager = OwningWorld->ContentBundleManager;
+		TMap<FName, FColor> DataLayerDebugColors = GetDataLayerDebugColors(WorldPartition);
+
+		auto WorldToScreen = [&](const FVector2D& Pos, bool bIsLocal = true)
+		{
+			return InWorldToScreen(Pos, bIsLocal);
+		};
+
+		auto LocalDrawTile = [&](const FVector2D& Min, const FVector2D& Size, const FLinearColor& Color)
+		{
+			const FVector2D ScreenMin = WorldToScreen(Min);
+			const FVector2D ScreenMax = WorldToScreen(Min + Size);
+			DrawContext.PushDrawTile(GridScreenBounds, ScreenMin, ScreenMax, Color);
+		};
+
+		auto LocalDrawBox = [&](const FVector2D& Min, const FVector2D& Size, const FLinearColor& Color, float LineThickness, const FBox2D* CustomGridScreenBounds = nullptr)
+		{
+			const FVector2D ScreenMin = WorldToScreen(Min);
+			const FVector2D ScreenMax = WorldToScreen(Min + Size);
+			DrawContext.PushDrawBox(GridScreenBounds, ScreenMin, ScreenMax, Color, LineThickness);
+		};
+
+		auto LocalDrawSegment = [&](const FVector2D& Start, const FVector2D& End, const FLinearColor& Color, float LineThickness)
+		{
+			const FVector2D ScreenStart = WorldToScreen(Start);
+			const FVector2D ScreenEnd = WorldToScreen(End);
+			DrawContext.PushDrawSegment(GridScreenBounds, ScreenStart, ScreenEnd, Color, LineThickness);
+		};
+
+		TArray<const UWorldPartitionRuntimeCell*> FilteredCells;
+		const FBox Region3D(FVector(Region.Min.X, Region.Min.Y, -HALF_WORLD_MAX), FVector(Region.Max.X, Region.Max.Y, HALF_WORLD_MAX));		
+		StreamingData->SpatialIndex->ForEachIntersectingElement(Region3D, [&FilteredCells](UWorldPartitionRuntimeCell* Cell)
+		{
+			UWorldPartitionRuntimeCellDataSpatialHashSet* RuntimeCellDataHashSet = CastChecked<UWorldPartitionRuntimeCellDataSpatialHashSet>(Cell->RuntimeCellData);
+
+			if ((RuntimeCellDataHashSet->Level >= GShowRuntimeHashSetDebugDisplayLevel) && (RuntimeCellDataHashSet->Level < (GShowRuntimeHashSetDebugDisplayLevel + GShowRuntimeHashSetDebugDisplayLevelCount)))
+			{
+				FilteredCells.Add(Cell);
+			}
+		});
+
+		if (FilteredCells.Num())
+		{
+			for (const UWorldPartitionRuntimeCell* Cell : FilteredCells)
+			{
+				const FVector2D CellBoundsSize = FVector2D(Cell->GetCellBounds().GetSize());
+				const FVector2D CellBoundsMin = FVector2D(Cell->GetCellBounds().Min);
+				LocalDrawTile(CellBoundsMin, CellBoundsSize, Cell->GetDebugColor(VisualizeMode).CopyWithNewOpacity(0.25f));
+				LocalDrawBox(CellBoundsMin, CellBoundsSize, FLinearColor::Black, 1);
+			}
+		}
+
+		// Draw X/Y Axis
+		if (DrawContext.GetDrawGridAxis())
+		{
+			DrawContext.PushDrawSegment(GridScreenBounds, WorldToScreen(FVector2D(-1638400.f, 0.f), false), WorldToScreen(FVector2D(1638400.f, 0.f), false), FLinearColor::Red, 3);
+			DrawContext.PushDrawSegment(GridScreenBounds, WorldToScreen(FVector2D(0.f, -1638400.f), false), WorldToScreen(FVector2D(0.f, 1638400.f), false), FLinearColor::Green, 3);
+		}
+
+		// Draw Grid Bounds
+		if (DrawContext.GetDrawGridBounds())
+		{
+			const FBox2D Bounds = GridScreenBounds.ExpandBy(FVector2D(10));
+			const FVector2D Size = GridScreenBounds.GetSize();
+			DrawContext.PushDrawBox(Bounds, GridScreenBounds.Min, GridScreenBounds.Max, FLinearColor::White, 1);
+		}
+
+		// Draw Streaming Sources
+		for (const FWorldPartitionStreamingSource& Source : Sources)
+		{
+			const FColor Color = Source.GetDebugColor();
+			// @todo_jfd
+			const FSoftObjectPath HLODLayer;
+			Source.ForEachShape(StreamingData->LoadingRange, StreamingData->Name, HLODLayer, true, [&Color, &WorldToScreen, &GridScreenBounds, &DrawContext, this](const FSphericalSector& Shape)
+			{
+				check(!Shape.IsNearlyZero())
+
+				// Spherical Sector
+				const FVector2D Center2D(FVector2D(Shape.GetCenter()));
+				const FSphericalSector::FReal Angle = Shape.GetAngle();
+				const int32 MaxSegments = FMath::Max(4, FMath::CeilToInt(64 * Angle / 360.f));
+				const float AngleIncrement = Angle / MaxSegments;
+				const FVector2D Axis = FVector2D(Shape.GetAxis());
+				const FVector Startup = FRotator(0, -0.5f * Angle, 0).RotateVector(Shape.GetScaledAxis());
+
+				FVector2D LineStart = FVector2D(Startup);
+				if (!Shape.IsSphere())
+				{
+					// Draw sector start axis
+					DrawContext.PushDrawSegment(GridScreenBounds, WorldToScreen(Center2D), WorldToScreen(Center2D + LineStart), Color, 2);
+				}
+				// Draw sector Arc
+				for (int32 i = 1; i <= MaxSegments; i++)
+				{
+					FVector2D LineEnd = FVector2D(FRotator(0, AngleIncrement * i, 0).RotateVector(Startup));
+					DrawContext.PushDrawSegment(GridScreenBounds, WorldToScreen(Center2D + LineStart), WorldToScreen(Center2D + LineEnd), Color, 2);
+					LineStart = LineEnd;
+				}
+				// If sphere, close circle, else draw sector end axis
+				DrawContext.PushDrawSegment(GridScreenBounds, WorldToScreen(Center2D + LineStart), WorldToScreen(Center2D + (Shape.IsSphere() ? FVector2D(Startup) : FVector2D::ZeroVector)), Color, 2);
+
+				// Draw direction vector
+				DrawContext.PushDrawSegment(GridScreenBounds, WorldToScreen(Center2D), WorldToScreen(Center2D + Axis * Shape.GetRadius()), Color, 2);
+			});
+		}
+	};
+
+	FBox GridsShapeBounds(ForceInit);
+	FBox2D GridsBounds(ForceInit);
+	int32 GridIndex = 0;
+	for (auto& [Name, StreamingDataList] : FilteredStreamingObjects)
+	{
+		for (const FWorldPartitionStreamingSource& Source : Sources)
+		{
+			// @todo_jfd
+			const FSoftObjectPath HLODLayer;
+			Source.ForEachShape(StreamingDataList[0]->LoadingRange, Name, HLODLayer, true, [&GridsShapeBounds](const FSphericalSector& Shape) { GridsShapeBounds += Shape.CalcBounds(); });
+		}
+
+		const FVector2D GridReferenceWorldPos = FVector2D(WorldRegion.GetCenter());
+		const FVector2D WorldRegionExtent = FVector2D(WorldRegion.GetExtent().GetMax());
+		const FVector2D GridScreenOffset = GridScreenInitialOffset + ((float)GridIndex * FVector2D(GridMaxScreenWidth, 0.f)) + GridScreenHalfExtent + FVector2D(GridScreenWidthShrinkSize * 0.5f);
+		const FVector2D WorldToScreenScale = GridScreenHalfExtent / WorldRegionExtent;
+		const FBox2D GridScreenBounds(GridScreenOffset - GridScreenHalfExtent, GridScreenOffset + GridScreenHalfExtent);
+
+		auto WorldToScreen = [&](const FVector2D& LocalWorldPos, bool bIsLocalWorldPos = true)
+		{
+			FVector2D GlocalPos = bIsLocalWorldPos ? LocalWPToGlobal.TransformPoint(LocalWorldPos) : LocalWorldPos;
+			return (WorldToScreenScale * (GlocalPos - GridReferenceWorldPos)) + GridScreenOffset;
+		};
+
+		for (const FRuntimePartitionStreamingData* StreamingData : StreamingDataList)
+		{
+			DrawStreamingData(StreamingData, LocalRegion, GridScreenBounds, WorldToScreen, DrawContext);
+		}
+
+		GridsBounds += GridScreenBounds;
+
+		if (DrawContext.IsDetailedMode())
+		{
+			FVector2D GridInfoPos = GridScreenOffset - GridScreenHalfExtent;
+			FWorldPartitionCanvasMultiLineText MultiLineText;
+			MultiLineText.Emplace(UWorld::RemovePIEPrefix(FPaths::GetBaseFilename(WorldPartition->GetPackage()->GetName())), FLinearColor::White);
+			FString GridInfoText = FString::Printf(TEXT("%s | %d m"), *Name.ToString(), int32(StreamingDataList[0]->LoadingRange * 0.01f));
+			MultiLineText.Emplace(GridInfoText, FLinearColor::Yellow);
+			FWorldPartitionCanvasMultiLineTextItem Item(GridInfoPos, MultiLineText);
+			DrawContext.PushDrawText(Item);
+			++GridIndex;
+		}
+	}
+
+	FBox2D DesiredWorldBounds(ForceInit);
+	if (GridsShapeBounds.IsValid)
+	{
+		// Convert to 2D
+		FBox2D GridsShapeBounds2D = FBox2D(FVector2D(GridsShapeBounds.Min.X, GridsShapeBounds.Min.Y), FVector2D(GridsShapeBounds.Max.X, GridsShapeBounds.Max.Y));
+		FVector2D CenterGlobalPos = LocalWPToGlobal.TransformPoint(GridsShapeBounds2D.GetCenter());
+		// Use max extent of X/Y
+		FVector2D Extent = FVector2D(GridsShapeBounds2D.GetExtent().GetMax());
+		// Transform to global space
+		GridsShapeBounds2D = FBox2D(CenterGlobalPos - Extent, CenterGlobalPos + Extent);
+		// Expand by 10% computed bounds
+		DesiredWorldBounds = GridsShapeBounds2D.ExpandBy(GridsShapeBounds2D.GetExtent() * 0.1f);
+	}
+	DrawContext.SetDesiredWorldBounds(DesiredWorldBounds);
+	DrawContext.SetUsedCanvasBounds(GridsBounds);
+
+	return true;
+}
+
 void UWorldPartitionRuntimeHashSet::SetDefaultValues()
 {
 }
@@ -316,7 +583,10 @@ bool UWorldPartitionRuntimeHashSet::GenerateStreaming(UWorldPartitionStreamingPo
 		RuntimeCell->SetBlockOnSlowLoading(CellDesc.bBlockOnSlowStreaming);
 		RuntimeCell->SetIsHLOD(false);
 		RuntimeCell->SetGuid(CellGuid);
-		RuntimeCell->RuntimeCellData->DebugName = CellObjectName;
+
+		UWorldPartitionRuntimeCellDataSpatialHashSet* RuntimeCellDataHashSet = CastChecked<UWorldPartitionRuntimeCellDataSpatialHashSet>(RuntimeCell->RuntimeCellData);
+		RuntimeCellDataHashSet->DebugName = CellObjectName;
+		RuntimeCellDataHashSet->Level = CellDesc.Level;
 
 		return RuntimeCell;
 	};
@@ -325,7 +595,7 @@ bool UWorldPartitionRuntimeHashSet::GenerateStreaming(UWorldPartitionStreamingPo
 	TMap<URuntimePartition*, FRuntimePartitionStreamingData> RuntimePartitionsStreamingData;
 	for (const FCellDescInstance& CellDescInstance : RuntimeCellDescsInstances)
 	{
-		UWorldPartitionRuntimeCell* RuntimeCell = RuntimeCells.Emplace_GetRef(CreateRuntimeCellFromCellDesc(CellDescInstance, CellDescInstance.DataLayers, StreamingPolicy->GetRuntimeCellClass(), UWorldPartitionRuntimeCellData::StaticClass()));
+		UWorldPartitionRuntimeCell* RuntimeCell = RuntimeCells.Emplace_GetRef(CreateRuntimeCellFromCellDesc(CellDescInstance, CellDescInstance.DataLayers, StreamingPolicy->GetRuntimeCellClass(), UWorldPartitionRuntimeCellDataSpatialHashSet::StaticClass()));
 		PopulateRuntimeCell(RuntimeCell, CellDescInstance.ActorInstances, nullptr);
 
 		// Override the cell bounds if the runtime partition provided one
@@ -597,7 +867,7 @@ void UWorldPartitionRuntimeHashSet::ForEachStreamingCellsSources(const TArray<FW
 	UWorldPartitionRuntimeHash::FStreamingSourceCells ActivateStreamingSourceCells;
 	UWorldPartitionRuntimeHash::FStreamingSourceCells LoadStreamingSourceCells;
 
-	auto ForEachStreamingCells = [this, &Sources, &Func, &ActivateStreamingSourceCells, &LoadStreamingSourceCells](FStaticSpatialIndexType* InSpatialIndex, float InLoadingRange, FName InGridName)
+	auto ForEachStreamingCells = [this, &Sources, &Func, &ActivateStreamingSourceCells, &LoadStreamingSourceCells](FStaticSpatialIndexType* InSpatialIndex, int32 InLoadingRange, FName InGridName)
 	{
 		if (InSpatialIndex)
 		{
diff --git a/Engine/Source/Runtime/Engine/Private/WorldPartition/WorldPartitionDraw2DContext.h b/Engine/Source/Runtime/Engine/Private/WorldPartition/WorldPartitionDraw2DContext.h
--- a/Engine/Source/Runtime/Engine/Private/WorldPartition/WorldPartitionDraw2DContext.h
+++ b/Engine/Source/Runtime/Engine/Private/WorldPartition/WorldPartitionDraw2DContext.h
@@ -165,28 +165,25 @@ public:
 		return false;
 	};
 
-	bool PushDrawBox(const FBox2D& InBounds, const FVector2D& A, const FVector2D& B, const FVector2D& C, const FVector2D& D, const FLinearColor& InColor, float InLineThickness)
+	bool PushDrawBox(const FBox2D& InBounds, const FVector2D& Min, const FVector2D& Max, const FLinearColor& InColor, float InLineThickness)
 	{
-		PushDrawSegment(InBounds, A, B, InColor, InLineThickness);
-		PushDrawSegment(InBounds, B, C, InColor, InLineThickness);
-		PushDrawSegment(InBounds, C, D, InColor, InLineThickness);
-		PushDrawSegment(InBounds, A, D, InColor, InLineThickness);
+		PushDrawSegment(InBounds, Min, FVector2D(Max.X, Min.Y), InColor, InLineThickness);
+		PushDrawSegment(InBounds, FVector2D(Max.X, Min.Y), Max, InColor, InLineThickness);
+		PushDrawSegment(InBounds, Max, FVector2D(Min.X, Max.Y), InColor, InLineThickness);
+		PushDrawSegment(InBounds, FVector2D(Min.X, Max.Y), Min, InColor, InLineThickness);
 		return true;
 	};
 
-	bool PushDrawTile(const FBox2D& InBounds, const FVector2D& A, const FVector2D& B, const FVector2D& C, const FVector2D& D, const FLinearColor& InColor)
+	bool PushDrawTile(const FBox2D& InBounds, const FVector2D& Min, const FVector2D& Max, const FLinearColor& InColor)
 	{
-		if (!InBounds.IsInside(A) ||
-			!InBounds.IsInside(B) ||
-			!InBounds.IsInside(C) ||
-			!InBounds.IsInside(D))
+		const FBox2D ClipBox = InBounds.Overlap(FBox2D(Min, Max));
+		if (ClipBox.bIsValid)
 		{
-			return false;
+			FWorldPartitionCanvasBoxItem Box(ClipBox.Min, FVector2D(ClipBox.Max.X, ClipBox.Min.Y), ClipBox.Max, FVector2D(ClipBox.Min.X, ClipBox.Max.Y), InColor);
+			CanvasItems.AddBox(Box);
+			return true;
 		}
-
-		FWorldPartitionCanvasBoxItem Box(A, B, C, D, InColor);
-		CanvasItems.AddBox(Box);
-		return true;
+		return false;
 	};
 
 	bool PushDrawText(FWorldPartitionCanvasMultiLineTextItem& InMultiLineText)
diff --git a/Engine/Source/Runtime/Engine/Private/WorldPartition/WorldPartitionRuntimeLevelStreamingCell.cpp b/Engine/Source/Runtime/Engine/Private/WorldPartition/WorldPartitionRuntimeLevelStreamingCell.cpp
--- a/Engine/Source/Runtime/Engine/Private/WorldPartition/WorldPartitionRuntimeLevelStreamingCell.cpp
+++ b/Engine/Source/Runtime/Engine/Private/WorldPartition/WorldPartitionRuntimeLevelStreamingCell.cpp
@@ -201,9 +201,7 @@ FLinearColor UWorldPartitionRuntimeLevelStreamingCell::GetDebugColor(EWorldParti
 		case EWorldPartitionRuntimeCellVisualizeMode::StreamingStatus:
 		{
 			// Return streaming status color
-			FLinearColor Color = LevelStreaming ? ULevelStreaming::GetLevelStreamingStatusColor(GetStreamingStatus()) : FLinearColor::Black;
-			Color.A = 0.25f;
-			return Color;
+			return LevelStreaming ? ULevelStreaming::GetLevelStreamingStatusColor(LevelStreaming->GetLevelStreamingStatus()) : FLinearColor::Black;
 		}
 		default:
 		{
diff --git a/Engine/Source/Runtime/Engine/Private/WorldPartition/WorldPartitionRuntimeSpatialHash.cpp b/Engine/Source/Runtime/Engine/Private/WorldPartition/WorldPartitionRuntimeSpatialHash.cpp
--- a/Engine/Source/Runtime/Engine/Private/WorldPartition/WorldPartitionRuntimeSpatialHash.cpp
+++ b/Engine/Source/Runtime/Engine/Private/WorldPartition/WorldPartitionRuntimeSpatialHash.cpp
@@ -803,27 +803,23 @@ void FSpatialHashStreamingGrid::Draw2D(const UWorldPartitionRuntimeSpatialHash*
 
 	auto LocalDrawTile = [&](const FVector2D& Min, const FVector2D& Size, const FLinearColor& Color)
 	{
-		FVector2D A = WorldToScreen(Min);
-		FVector2D B = WorldToScreen(Min + FVector2D(Size.X, 0));
-		FVector2D C = WorldToScreen(Min + Size);
-		FVector2D D = WorldToScreen(Min + FVector2D(0, Size.Y));
-		DrawContext.PushDrawTile(GridScreenBounds, A, B, C, D, Color);
+		const FVector2D ScreenMin = WorldToScreen(Min);
+		const FVector2D ScreenMax = WorldToScreen(Min + Size);
+		DrawContext.PushDrawTile(GridScreenBounds, ScreenMin, ScreenMax, Color);
 	};
 
 	auto LocalDrawBox = [&](const FVector2D& Min, const FVector2D& Size, const FLinearColor& Color, float LineThickness, const FBox2D* CustomGridScreenBounds = nullptr)
 	{
-		FVector2D A = WorldToScreen(Min);
-		FVector2D B = WorldToScreen(Min + FVector2D(Size.X, 0));
-		FVector2D C = WorldToScreen(Min + Size);
-		FVector2D D = WorldToScreen(Min + FVector2D(0, Size.Y));
-		DrawContext.PushDrawBox(GridScreenBounds, A, B, C, D, Color, LineThickness);
+		const FVector2D ScreenMin = WorldToScreen(Min);
+		const FVector2D ScreenMax = WorldToScreen(Min + Size);
+		DrawContext.PushDrawBox(GridScreenBounds, ScreenMin, ScreenMax, Color, LineThickness);
 	};
 
 	auto LocalDrawSegment = [&](const FVector2D& Start, const FVector2D& End, const FLinearColor& Color, float LineThickness)
 	{
-		FVector2D A = WorldToScreen(Start);
-		FVector2D B = WorldToScreen(End);
-		DrawContext.PushDrawSegment(GridScreenBounds, A, B, Color, LineThickness);
+		const FVector2D ScreenStart = WorldToScreen(Start);
+		const FVector2D ScreenEnd = WorldToScreen(End);
+		DrawContext.PushDrawSegment(GridScreenBounds, ScreenStart, ScreenEnd, Color, LineThickness);
 	};
 
 	TArray<const UWorldPartitionRuntimeCell*> FilteredCells;
@@ -853,7 +849,7 @@ void FSpatialHashStreamingGrid::Draw2D(const UWorldPartitionRuntimeSpatialHash*
 					// Draw Cell using its debug color
 					FVector2D StartPos = CellWorldBounds.Min + CellOffset;
 					{
-						LocalDrawTile(StartPos, CellBoundsSize, Cell->GetDebugColor(VisualizeMode));
+						LocalDrawTile(StartPos, CellBoundsSize, Cell->GetDebugColor(VisualizeMode).CopyWithNewOpacity(0.25f));
 					}
 
 					CellOffset.Y += CellBoundsSize.Y;
@@ -907,7 +903,7 @@ void FSpatialHashStreamingGrid::Draw2D(const UWorldPartitionRuntimeSpatialHash*
 	{
 		FBox2D Bounds = GridScreenBounds.ExpandBy(FVector2D(10));
 		FVector2D Size = GridScreenBounds.GetSize();
-		DrawContext.PushDrawBox(Bounds, GridScreenBounds.Min, GridScreenBounds.Min + FVector2D(Size.X, 0), GridScreenBounds.Max, GridScreenBounds.Min + FVector2D(0, Size.Y), DebugColor, 1);
+		DrawContext.PushDrawBox(Bounds, GridScreenBounds.Min, GridScreenBounds.Max, DebugColor, 1);
 	}
 
 	// Draw WorldBounds
diff --git a/Engine/Source/Runtime/Engine/Public/WorldPartition/RuntimeHashSet/RuntimePartition.h b/Engine/Source/Runtime/Engine/Public/WorldPartition/RuntimeHashSet/RuntimePartition.h
--- a/Engine/Source/Runtime/Engine/Public/WorldPartition/RuntimeHashSet/RuntimePartition.h
+++ b/Engine/Source/Runtime/Engine/Public/WorldPartition/RuntimeHashSet/RuntimePartition.h
@@ -13,6 +13,10 @@ class URuntimePartition : public UObject
 
 public:
 #if WITH_EDITOR
+	//~ Begin UObject Interface.
+	virtual void PostEditChangeProperty(FPropertyChangedEvent& PropertyChangedEvent) override;
+	//~ End UObject Interface.
+
 	struct FCellDesc
 	{
 		FName Name;
@@ -23,6 +27,9 @@ public:
 		FGuid ContentBundleID;
 		int32 Priority;
 
+		/** Optional level value that can be used to filter debug display */
+		int32 Level;
+
 		TArray<IStreamingGenerationContext::FActorInstance> ActorInstances;
 	};
 
@@ -45,7 +52,7 @@ public:
 	int32 Priority;
 
 	UPROPERTY(EditAnywhere, Category = RuntimeSettings)
-	float LoadingRange;
+	int32 LoadingRange;
 
 	UPROPERTY()
 	bool bIsHLODSetup;
diff --git a/Engine/Source/Runtime/Engine/Public/WorldPartition/RuntimeHashSet/WorldPartitionRuntimeHashSet.h b/Engine/Source/Runtime/Engine/Public/WorldPartition/RuntimeHashSet/WorldPartitionRuntimeHashSet.h
--- a/Engine/Source/Runtime/Engine/Public/WorldPartition/RuntimeHashSet/WorldPartitionRuntimeHashSet.h
+++ b/Engine/Source/Runtime/Engine/Public/WorldPartition/RuntimeHashSet/WorldPartitionRuntimeHashSet.h
@@ -71,7 +71,7 @@ struct FRuntimePartitionStreamingData
 	FName Name;
 
 	UPROPERTY()
-	float LoadingRange = 0.0f;
+	int32 LoadingRange = 0;
 
 	UPROPERTY()
 	TArray<TObjectPtr<UWorldPartitionRuntimeCell>> RuntimeCells;
@@ -109,6 +109,16 @@ public:
 	TArray<FRuntimePartitionStreamingData> RuntimeStreamingData;
 };
 
+UCLASS(Within = WorldPartitionRuntimeCell, MinimalAPI)
+class UWorldPartitionRuntimeCellDataSpatialHashSet : public UWorldPartitionRuntimeCellData
+{
+	GENERATED_BODY()
+
+public:
+	UPROPERTY()
+	int32 Level;
+};
+
 UCLASS(HideDropdown, MinimalAPI)
 class UWorldPartitionRuntimeHashSet : public UWorldPartitionRuntimeHash
 {
@@ -123,6 +133,8 @@ class UWorldPartitionRuntimeHashSet : public UWorldPartitionRuntimeHash
 
 public:
 #if WITH_EDITOR
+	ENGINE_API virtual bool Draw2D(FWorldPartitionDraw2DContext& DrawContext) const override;
+
 	// Streaming generation interface
 	ENGINE_API virtual void SetDefaultValues() override;
 	ENGINE_API virtual bool SupportsHLODs() const override;
