Fix non-unity build.

diff --git a/Engine/Source/Runtime/Engine/Private/WorldPartition/RuntimeHashSet/WorldPartitionRuntimeHashSet.cpp b/Engine/Source/Runtime/Engine/Private/WorldPartition/RuntimeHashSet/WorldPartitionRuntimeHashSet.cpp
--- a/Engine/Source/Runtime/Engine/Private/WorldPartition/RuntimeHashSet/WorldPartitionRuntimeHashSet.cpp
+++ b/Engine/Source/Runtime/Engine/Private/WorldPartition/RuntimeHashSet/WorldPartitionRuntimeHashSet.cpp
@@ -17,6 +17,7 @@
 #include "WorldPartition/HLOD/HLODLayer.h"
 #include "Algo/ForEach.h"
 #include "Algo/Transform.h"
+#include "Math/TransformCalculus2D.h"
 #include "Misc/ArchiveMD5.h"
 
 FAutoConsoleCommand WorldPartitionRuntimeHashSetEnable(
