Warning in Cook Content Worker: Failed to finish async Animation Compression task for F_MED_Blue_Jay_Head_Facial_Poses, as the generated data is not valid.
#fix downgrade warning to display UE_LOG when trying to finish async task with invalid data, until root-cause of GC causing this is found. Behaviour self-corrects and problematic assets are compressed correctly later in the cook.

diff --git a/Engine/Source/Runtime/Engine/Private/Animation/AnimSequence.cpp b/Engine/Source/Runtime/Engine/Private/Animation/AnimSequence.cpp
--- a/Engine/Source/Runtime/Engine/Private/Animation/AnimSequence.cpp
+++ b/Engine/Source/Runtime/Engine/Private/Animation/AnimSequence.cpp
@@ -5505,7 +5505,7 @@ void UAnimSequence::FinishAsyncTasks()
 				else
 				{
 					// Failed to compress
-					UE_LOG(LogAnimation, Warning, TEXT("Failed to finish async Animation Compression task for %s, as the generated data is not valid."), *GetName());					
+					UE_LOG(LogAnimation, Display, TEXT("Failed to finish async Animation Compression task for %s, as the generated data is not valid."), *GetName());					
 					ResetData();
 				}
 			}
