Add PreLoginAsync to GameMode to allow for online requests during login.

diff --git a/Engine/Source/Runtime/Engine/Classes/Engine/World.h b/Engine/Source/Runtime/Engine/Classes/Engine/World.h
--- a/Engine/Source/Runtime/Engine/Classes/Engine/World.h
+++ b/Engine/Source/Runtime/Engine/Classes/Engine/World.h
@@ -3635,6 +3635,21 @@ public:
 	virtual void NotifyControlMessage(UNetConnection* Connection, uint8 MessageType, class FInBunch& Bunch) override;
 	//~ End FNetworkNotify Interface
 
+	/**
+	 * Log error and close connection for prelogin failures.
+	 */
+	bool PreLoginCheckError(UNetConnection* Connection, const FString& ErrorMsg);
+
+	/**
+	 * Check GameMode PreLogin results and welcome player if needed.
+	 */
+	void PreLoginComplete(const FString& ErrorMsg, TWeakObjectPtr<UNetConnection> WeakConnection);
+
+	/**
+	 * Check GameMode PreLogin results for split screen player joins.
+	 */
+	void PreLoginCompleteSplit(const FString& ErrorMsg, TWeakObjectPtr<UNetConnection> WeakConnection, FUniqueNetIdRepl SplitRequestUniqueIdRepl, FString SplitRequestURL);
+
 	/** Welcome a new player joining this server. */
 	void WelcomePlayer(UNetConnection* Connection);
 
diff --git a/Engine/Source/Runtime/Engine/Classes/GameFramework/GameModeBase.h b/Engine/Source/Runtime/Engine/Classes/GameFramework/GameModeBase.h
--- a/Engine/Source/Runtime/Engine/Classes/GameFramework/GameModeBase.h
+++ b/Engine/Source/Runtime/Engine/Classes/GameFramework/GameModeBase.h
@@ -288,6 +288,16 @@ public:
 	 */
 	ENGINE_API virtual void PreLogin(const FString& Options, const FString& Address, const FUniqueNetIdRepl& UniqueId, FString& ErrorMessage);
 
+	DECLARE_DELEGATE_OneParam(FOnPreLoginCompleteDelegate, const FString& /*ErrorMsg*/);
+
+	/**
+	 * Async version of PreLogin that MUST call OnComplete when done, otherwise incoming client connections will hang or timeout.
+	 * This allows servers to make requests to other services for verification of login credentials without blocking the game thread.
+	 * By default, this calls PreLogin() and OnComplete immediately to not break backwards compatibility. When overriding this, be sure
+	 * checks in PreLogin are also run, and that the GameModePreLoginEvent is broadcast.
+	 */
+	virtual void PreLoginAsync(const FString& Options, const FString& Address, const FUniqueNetIdRepl& UniqueId, const FOnPreLoginCompleteDelegate& OnComplete);
+
 	/**
 	 * Called to login new players by creating a player controller, overridable by the game
 	 *
diff --git a/Engine/Source/Runtime/Engine/Private/GameModeBase.cpp b/Engine/Source/Runtime/Engine/Private/GameModeBase.cpp
--- a/Engine/Source/Runtime/Engine/Private/GameModeBase.cpp
+++ b/Engine/Source/Runtime/Engine/Private/GameModeBase.cpp
@@ -669,6 +669,13 @@ void AGameModeBase::PreLogin(const FString& Options, const FString& Address, con
 	FGameModeEvents::GameModePreLoginEvent.Broadcast(this, UniqueId, ErrorMessage);
 }
 
+void AGameModeBase::PreLoginAsync(const FString& Options, const FString& Address, const FUniqueNetIdRepl& UniqueId, const FOnPreLoginCompleteDelegate& OnComplete)
+{
+	FString ErrorMessage;
+	PreLogin(Options, Address, UniqueId, ErrorMessage);
+	OnComplete.ExecuteIfBound(ErrorMessage);
+}
+
 APlayerController* AGameModeBase::Login(UPlayer* NewPlayer, ENetRole InRemoteRole, const FString& Portal, const FString& Options, const FUniqueNetIdRepl& UniqueId, FString& ErrorMessage)
 {
 	if (GameSession == nullptr)
diff --git a/Engine/Source/Runtime/Engine/Private/World.cpp b/Engine/Source/Runtime/Engine/Private/World.cpp
--- a/Engine/Source/Runtime/Engine/Private/World.cpp
+++ b/Engine/Source/Runtime/Engine/Private/World.cpp
@@ -6426,26 +6426,16 @@ void UWorld::NotifyControlMessage(UNetConnection* Connection, uint8 MessageType,
 					Connection->SetPlayerOnlinePlatformName(FName(*OnlinePlatformName));
 
 					// ask the game code if this player can join
-					FString ErrorMsg;
 					AGameModeBase* GameMode = GetAuthGameMode();
-
+					AGameModeBase::FOnPreLoginCompleteDelegate OnComplete = AGameModeBase::FOnPreLoginCompleteDelegate::CreateUObject(
+						this, &UWorld::PreLoginComplete, TWeakObjectPtr<UNetConnection>(Connection));
 					if (GameMode)
 					{
-						GameMode->PreLogin(Tmp, Connection->LowLevelGetRemoteAddress(), Connection->PlayerId, ErrorMsg);
-					}
-					if (!ErrorMsg.IsEmpty())
-					{
-						UE_LOG(LogNet, Log, TEXT("PreLogin failure: %s"), *ErrorMsg);
-						NETWORK_PROFILER(GNetworkProfiler.TrackEvent(TEXT("PRELOGIN FAILURE"), *ErrorMsg, Connection));
-
-						Connection->SendCloseReason(ENetCloseResult::PreLoginFailure);
-						FNetControlMessage<NMT_Failure>::Send(Connection, ErrorMsg);
-						Connection->FlushNet(true);
-						Connection->Close(ENetCloseResult::PreLoginFailure);
+						GameMode->PreLoginAsync(Tmp, Connection->LowLevelGetRemoteAddress(), Connection->PlayerId, OnComplete);
 					}
 					else
 					{
-						WelcomePlayer(Connection);
+						OnComplete.ExecuteIfBound(FString());
 					}
 				}
 				else
@@ -6565,66 +6555,17 @@ void UWorld::NotifyControlMessage(UNetConnection* Connection, uint8 MessageType,
 					const TCHAR* Tmp = *SplitRequestURL;
 					for (; *Tmp && *Tmp != '?'; Tmp++);
 
-					// go through the same full login process for the split player even though it's all in the same frame
-					FString ErrorMsg;
+					// go through the same full login process for the split player
 					AGameModeBase* GameMode = GetAuthGameMode();
+					AGameModeBase::FOnPreLoginCompleteDelegate OnComplete = AGameModeBase::FOnPreLoginCompleteDelegate::CreateUObject(
+						this, &UWorld::PreLoginCompleteSplit, TWeakObjectPtr<UNetConnection>(Connection), SplitRequestUniqueIdRepl, SplitRequestURL);
 					if (GameMode)
 					{
-						GameMode->PreLogin(Tmp, Connection->LowLevelGetRemoteAddress(), SplitRequestUniqueIdRepl, ErrorMsg);
-					}
-					if (!ErrorMsg.IsEmpty())
-					{
-						// if any splitscreen viewport fails to join, all viewports on that client also fail
-						UE_LOG(LogNet, Log, TEXT("PreLogin failure: %s"), *ErrorMsg);
-						NETWORK_PROFILER(GNetworkProfiler.TrackEvent(TEXT("PRELOGIN FAILURE"), *ErrorMsg, Connection));
-
-						Connection->SendCloseReason(ENetCloseResult::PreLoginFailure);
-						FNetControlMessage<NMT_Failure>::Send(Connection, ErrorMsg);
-						Connection->FlushNet(true);
-						Connection->Close(ENetCloseResult::PreLoginFailure);
+						GameMode->PreLoginAsync(Tmp, Connection->LowLevelGetRemoteAddress(), SplitRequestUniqueIdRepl, OnComplete);
 					}
 					else
 					{
-						// create a child network connection using the existing connection for its parent
-						check(Connection->GetUChildConnection() == NULL);
-#if WITH_EDITORONLY_DATA
-						check(CurrentLevel);
-#else
-						ULevel* CurrentLevel = PersistentLevel;
-#endif
-
-						UChildConnection* ChildConn = NetDriver->CreateChild(Connection);
-						ChildConn->PlayerId = SplitRequestUniqueIdRepl;
-						ChildConn->SetPlayerOnlinePlatformName(Connection->GetPlayerOnlinePlatformName());
-						ChildConn->RequestURL = SplitRequestURL;
-						ChildConn->SetClientWorldPackageName(CurrentLevel->GetOutermost()->GetFName());
-
-						// create URL from string
-						FURL JoinSplitURL(NULL, *SplitRequestURL, TRAVEL_Absolute);
-
-						UE_LOG(LogNet, Log, TEXT("JOINSPLIT: Join request: URL=%s"), *JoinSplitURL.ToString());
-						APlayerController* PC = SpawnPlayActor(ChildConn, ROLE_AutonomousProxy, JoinSplitURL, ChildConn->PlayerId, ErrorMsg, uint8(Connection->Children.Num()));
-						if (PC == NULL)
-						{
-							// Failed to connect.
-							UE_LOG(LogNet, Log, TEXT("JOINSPLIT: Join failure: %s"), *ErrorMsg);
-							NETWORK_PROFILER(GNetworkProfiler.TrackEvent(TEXT("JOINSPLIT FAILURE"), *ErrorMsg, Connection));
-							// remove the child connection
-							Connection->Children.Remove(ChildConn);
-
-							// if any splitscreen viewport fails to join, all viewports on that client also fail
-							Connection->SendCloseReason(ENetCloseResult::JoinSplitFailure);
-							FNetControlMessage<NMT_Failure>::Send(Connection, ErrorMsg);
-							Connection->FlushNet(true);
-							Connection->Close(ENetCloseResult::JoinSplitFailure);
-						}
-						else
-						{
-							// Successfully spawned in game.
-							UE_LOG(LogNet, Log, TEXT("JOINSPLIT: Succeeded: %s PlayerId: %s"),
-								*ChildConn->PlayerController->PlayerState->GetPlayerName(),
-								*ChildConn->PlayerController->PlayerState->GetUniqueId().ToDebugString());
-						}
+						OnComplete.ExecuteIfBound(FString());
 					}
 				}
 
@@ -6675,6 +6616,93 @@ void UWorld::NotifyControlMessage(UNetConnection* Connection, uint8 MessageType,
 	}
 }
 
+bool UWorld::PreLoginCheckError(UNetConnection* Connection, const FString& ErrorMsg)
+{
+	if (Connection)
+	{
+		if (ErrorMsg.IsEmpty())
+		{
+			return true;
+		}
+
+		UE_LOG(LogNet, Log, TEXT("PreLogin failure: %s"), *ErrorMsg);
+		NETWORK_PROFILER(GNetworkProfiler.TrackEvent(TEXT("PRELOGIN FAILURE"), *ErrorMsg, Connection));
+		Connection->SendCloseReason(ENetCloseResult::PreLoginFailure);
+		FString ErrorMsgCopy(ErrorMsg); // Needed because Send() can't handle const FString.
+		FNetControlMessage<NMT_Failure>::Send(Connection, ErrorMsgCopy);
+		Connection->FlushNet(true);
+		Connection->Close(ENetCloseResult::PreLoginFailure);
+	}
+	else
+	{
+		UE_LOG(LogNet, Log, TEXT("PreLogin failure: connection was null"));
+	}
+
+	return false;
+}
+
+void UWorld::PreLoginComplete(const FString& ErrorMsg, TWeakObjectPtr<UNetConnection> WeakConnection)
+{
+	UNetConnection* Connection = WeakConnection.Get();
+	if (!PreLoginCheckError(Connection, ErrorMsg))
+	{
+		return;
+	}
+
+	WelcomePlayer(Connection);
+}
+
+void UWorld::PreLoginCompleteSplit(const FString& ErrorMsg, TWeakObjectPtr<UNetConnection> WeakConnection, FUniqueNetIdRepl SplitRequestUniqueIdRepl, FString SplitRequestURL)
+{
+	UNetConnection* Connection = WeakConnection.Get();
+	if (!PreLoginCheckError(Connection, ErrorMsg))
+	{
+		return;
+	}
+
+	// create a child network connection using the existing connection for its parent
+	check(Connection->GetUChildConnection() == NULL);
+#if WITH_EDITORONLY_DATA
+	check(CurrentLevel);
+#else
+	ULevel* CurrentLevel = PersistentLevel;
+#endif
+
+	UChildConnection* ChildConn = NetDriver->CreateChild(Connection);
+	ChildConn->PlayerId = SplitRequestUniqueIdRepl;
+	ChildConn->SetPlayerOnlinePlatformName(Connection->GetPlayerOnlinePlatformName());
+	ChildConn->RequestURL = SplitRequestURL;
+	ChildConn->SetClientWorldPackageName(CurrentLevel->GetOutermost()->GetFName());
+
+	// create URL from string
+	FURL JoinSplitURL(NULL, *SplitRequestURL, TRAVEL_Absolute);
+
+	UE_LOG(LogNet, Log, TEXT("JOINSPLIT: Join request: URL=%s"), *JoinSplitURL.ToString());
+	FString SpawnErrorMsg;
+	APlayerController* PC = SpawnPlayActor(ChildConn, ROLE_AutonomousProxy, JoinSplitURL, ChildConn->PlayerId, SpawnErrorMsg, uint8(Connection->Children.Num()));
+	if (PC == NULL)
+	{
+		// Failed to connect.
+		UE_LOG(LogNet, Log, TEXT("JOINSPLIT: Join failure: %s"), *SpawnErrorMsg);
+		NETWORK_PROFILER(GNetworkProfiler.TrackEvent(TEXT("JOINSPLIT FAILURE"), *SpawnErrorMsg, Connection));
+		// remove the child connection
+		Connection->Children.Remove(ChildConn);
+
+		// if any splitscreen viewport fails to join, all viewports on that client also fail
+		Connection->SendCloseReason(ENetCloseResult::JoinSplitFailure);
+		FNetControlMessage<NMT_Failure>::Send(Connection, SpawnErrorMsg);
+		Connection->FlushNet(true);
+		Connection->Close(ENetCloseResult::JoinSplitFailure);
+	}
+	else
+	{
+		// Successfully spawned in game.
+		UE_LOG(LogNet, Log, TEXT("JOINSPLIT: Succeeded: %s PlayerId: %s"),
+			*ChildConn->PlayerController->PlayerState->GetPlayerName(),
+			*ChildConn->PlayerController->PlayerState->GetUniqueId().ToDebugString());
+	}
+}
+
 bool UWorld::Listen( FURL& InURL )
 {
 #if WITH_SERVER_CODE
