Add a cook stat to track all the materials we cook out to help catch regressions in the number of materials we cook for a project.

- Note this will only be accurate since we do full cooks and will always load and cook each material.
- If you are doing incremental cooks this stat will fluctuate and we will have to revist this.

diff --git a/Engine/Source/Runtime/Engine/Private/Materials/MaterialInterface.cpp b/Engine/Source/Runtime/Engine/Private/Materials/MaterialInterface.cpp
--- a/Engine/Source/Runtime/Engine/Private/Materials/MaterialInterface.cpp
+++ b/Engine/Source/Runtime/Engine/Private/Materials/MaterialInterface.cpp
@@ -33,6 +33,7 @@
 #include "MaterialDomain.h"
 #include "MaterialShaderQualitySettings.h"
 #include "Materials/MaterialRenderProxy.h"
+#include "ProfilingDebugging/CookStats.h"
 #include "ShaderPlatformQualitySettings.h"
 #include "ObjectCacheContext.h"
 #include "MaterialCachedData.h"
@@ -47,6 +48,20 @@
 
 #define LOCTEXT_NAMESPACE "MaterialInterface"
 
+#if ENABLE_COOK_STATS
+namespace MaterialCookStats
+{
+	static int32 NumMaterialsCooked = 0;
+
+	static FCookStatsManager::FAutoRegisterCallback RegisterCookStats([](FCookStatsManager::AddStatFuncRef AddStat)
+	{
+		AddStat(TEXT("Material"), FCookStatsManager::CreateKeyValueArray(
+			TEXT("NumMaterialsCooked"), NumMaterialsCooked
+		));
+	});
+}
+#endif
+
 /**
  * This is used to deprecate data that has been built with older versions.
  * To regenerate the data, commands like "BUILDMATERIALTEXTURESTREAMINGDATA" can be used in the editor.
@@ -1628,6 +1643,12 @@ void UMaterialInterface::PreSave(FObjectPreSaveContext ObjectSaveContext)
 	}
 #endif // WITH_EDITORONLY_DATA
 
+#if WITH_EDITOR
+	if (ObjectSaveContext.IsCooking())
+	{
+		COOK_STAT(MaterialCookStats::NumMaterialsCooked++);
+	}
+#endif // WITH_EDITOR
 }
 
 void UMaterialInterface::AddAssetUserData(UAssetUserData* InUserData)
