Fix address sanitaizer issue with player input. The "ConsumeKey" function can re-allocate the KeyStateMap because of the call to FindOrAdd, which would in turn invalidate the pointers in the "KeysWithEvents" map. Instead of calling consume key directly, we can just do a non-mutating call to "Find" instead.

#rnx

diff --git a/Engine/Source/Runtime/Engine/Private/UserInterface/PlayerInput.cpp b/Engine/Source/Runtime/Engine/Private/UserInterface/PlayerInput.cpp
--- a/Engine/Source/Runtime/Engine/Private/UserInterface/PlayerInput.cpp
+++ b/Engine/Source/Runtime/Engine/Private/UserInterface/PlayerInput.cpp
@@ -1457,7 +1457,10 @@ void UPlayerInput::EvaluateInputDelegates(const TArray<UInputComponent*>& InputC
 			// we do this after finishing the whole component, so we don't consume a key while there might be more bindings to it
 			for (int32 KeyIndex=0; KeyIndex<KeysToConsume.Num(); ++KeyIndex)
 			{
-				ConsumeKey(KeysToConsume[KeyIndex]);
+				if (FKeyState* KeyState = KeyStateMap.Find(KeysToConsume[KeyIndex]))
+				{
+					KeyState->bConsumed = true;	
+				}
 			}
 			KeysToConsume.Reset();
 			FoundChords.Reset();
