Fix actor reconstruction to correctly set transform when a native root component changes class during reparenting

diff --git a/Engine/Source/Runtime/Engine/Private/ActorConstruction.cpp b/Engine/Source/Runtime/Engine/Private/ActorConstruction.cpp
--- a/Engine/Source/Runtime/Engine/Private/ActorConstruction.cpp
+++ b/Engine/Source/Runtime/Engine/Private/ActorConstruction.cpp
@@ -846,11 +846,21 @@ bool AActor::ExecuteConstruction(const FTransform& Transform, const FRotationCon
 					if (SceneComponent->CreationMethod == EComponentCreationMethod::Native && SceneComponent->GetOuter()->IsA<AActor>())
 					{
 						// If RootComponent is not set, the first unattached native scene component will be used as root. This matches what's done in FixupNativeActorComponents().
-						// @TODO - consider removing this; keeping here as a fallback just in case it wasn't set prior to SCS execution, but in most cases now this should be valid. 
+												// In cases like BP reparenting between native classes, this is needed to fix up changes in root component type
 						if (RootComponent == nullptr && SceneComponent->GetAttachParent() == nullptr)
 						{
 							// Note: All native scene components should already have been registered at this point, so we don't need to register the component here.
 							SetRootComponent(SceneComponent);
+
+							// Update the transform on the newly set root component
+							if (ensure(RootComponent) && !bIsDefaultTransform)
+							{
+								if (TransformRotationCache)
+								{
+									RootComponent->SetRelativeRotationCache(*TransformRotationCache);
+								}
+								RootComponent->SetWorldTransform(Transform, /*bSweep=*/false, /*OutSweepHitResult=*/nullptr, ETeleportType::TeleportPhysics);
+							}
 						}
 
 						NativeSceneComponents.Add(SceneComponent);
