Fix a crash in GetDynamicRayTracingInstances when spline mesh was using a static mesh that does not support ray tracing.

diff --git a/Engine/Source/Runtime/Engine/Private/SplineMeshSceneProxy.cpp b/Engine/Source/Runtime/Engine/Private/SplineMeshSceneProxy.cpp
--- a/Engine/Source/Runtime/Engine/Private/SplineMeshSceneProxy.cpp
+++ b/Engine/Source/Runtime/Engine/Private/SplineMeshSceneProxy.cpp
@@ -194,7 +194,7 @@ bool FSplineMeshSceneProxy::GetCollisionMeshElement(int32 LODIndex, int32 BatchI
 #if RHI_RAYTRACING
 void FSplineMeshSceneProxy::GetDynamicRayTracingInstances(struct FRayTracingMaterialGatheringContext& Context, TArray<FRayTracingInstance>& OutRayTracingInstances)
 {
-	if (CVarRayTracingSplineMeshes.GetValueOnRenderThread() == 0)
+	if (CVarRayTracingSplineMeshes.GetValueOnRenderThread() == 0  || !bSupportRayTracing)
 	{
 		return;
 	}
