More command list conversions. Streaming buffers, vertex factories, and TUniformBufferRef.

diff --git a/Engine/Plugins/2D/Paper2D/Source/Paper2D/Private/PaperRenderSceneProxy.cpp b/Engine/Plugins/2D/Paper2D/Source/Paper2D/Private/PaperRenderSceneProxy.cpp
--- a/Engine/Plugins/2D/Paper2D/Source/Paper2D/Private/PaperRenderSceneProxy.cpp
+++ b/Engine/Plugins/2D/Paper2D/Source/Paper2D/Private/PaperRenderSceneProxy.cpp
@@ -255,13 +255,13 @@ void FPaperRenderSceneProxy::RecreateCachedRenderData(FRHICommandListBase& RHICm
 			//When the buffer reallocates, the factory needs to bind the buffers and SRV again, we just init again.
 			if (bFactoryRequiresReInitialization)
 			{
-				VertexFactory.Init(&VertexBuffer);
+				VertexFactory.Init(RHICmdList, &VertexBuffer);
 			}
 		}
 		else
 		{
 			VertexBuffer.InitResource(RHICmdList);
-			VertexFactory.Init(&VertexBuffer);
+			VertexFactory.Init(RHICmdList, &VertexBuffer);
 		}
 	}
 }
@@ -270,11 +270,12 @@ void FPaperRenderSceneProxy::CreateRenderThreadResources()
 {
 	if (bSpritesUseVertexBufferPath && (Vertices.Num() > 0))
 	{
+		FRHICommandList& RHICmdList = FRHICommandListImmediate::Get();
 		VertexBuffer.Vertices = Vertices;
 
 		// Init the resources
-		VertexBuffer.InitResource(FRHICommandListImmediate::Get());
-		VertexFactory.Init(&VertexBuffer);
+		VertexBuffer.InitResource(RHICmdList);
+		VertexFactory.Init(RHICmdList, &VertexBuffer);
 	}
 }
 
diff --git a/Engine/Plugins/2D/Paper2D/Source/Paper2D/Private/PaperSpriteVertexBuffer.cpp b/Engine/Plugins/2D/Paper2D/Source/Paper2D/Private/PaperSpriteVertexBuffer.cpp
--- a/Engine/Plugins/2D/Paper2D/Source/Paper2D/Private/PaperSpriteVertexBuffer.cpp
+++ b/Engine/Plugins/2D/Paper2D/Source/Paper2D/Private/PaperSpriteVertexBuffer.cpp
@@ -209,39 +209,27 @@ FPaperSpriteVertexFactory::FPaperSpriteVertexFactory(ERHIFeatureLevel::Type Feat
 {
 }
 
-void FPaperSpriteVertexFactory::Init(const FPaperSpriteVertexBuffer* InVertexBuffer)
+void FPaperSpriteVertexFactory::Init(FRHICommandListBase& RHICmdList, const FPaperSpriteVertexBuffer* InVertexBuffer)
 {
-	if (IsInRenderingThread())
-	{
-		FLocalVertexFactory::FDataType VertexData;
-		VertexData.NumTexCoords = 1;
-
-		//SRV setup
-		VertexData.LightMapCoordinateIndex = 0;
-		VertexData.TangentsSRV = InVertexBuffer->TangentBufferSRV;
-		VertexData.TextureCoordinatesSRV = InVertexBuffer->TexCoordBufferSRV;
-		VertexData.ColorComponentsSRV = InVertexBuffer->ColorBufferSRV;
-		VertexData.PositionComponentSRV = InVertexBuffer->PositionBufferSRV;
-
-		// Vertex Streams
-		VertexData.PositionComponent = FVertexStreamComponent(&InVertexBuffer->PositionBuffer, 0, sizeof(FVector3f), VET_Float3, EVertexStreamUsage::Default);
-		VertexData.TangentBasisComponents[0] = FVertexStreamComponent(&InVertexBuffer->TangentBuffer, 0, 2 * sizeof(FPackedNormal), VET_PackedNormal, EVertexStreamUsage::ManualFetch);
-		VertexData.TangentBasisComponents[1] = FVertexStreamComponent(&InVertexBuffer->TangentBuffer, sizeof(FPackedNormal), 2 * sizeof(FPackedNormal), VET_PackedNormal, EVertexStreamUsage::ManualFetch);
-		VertexData.ColorComponent = FVertexStreamComponent(&InVertexBuffer->ColorBuffer, 0, sizeof(FColor), VET_Color, EVertexStreamUsage::ManualFetch);
-		VertexData.TextureCoordinates.Add(FVertexStreamComponent(&InVertexBuffer->TexCoordBuffer, 0, sizeof(FVector2f), VET_Float2, EVertexStreamUsage::ManualFetch));
-
-		SetData(VertexData);
-		VertexBuffer = InVertexBuffer;
-
-		InitResource(FRHICommandListImmediate::Get());
-	}
-	else
-	{
-		FPaperSpriteVertexFactory* ThisFactory = this;
-		ENQUEUE_RENDER_COMMAND(SpriteVertexFactoryInit)(
-			[ThisFactory, InVertexBuffer](FRHICommandListImmediate& RHICmdList)
-		{
-			ThisFactory->Init(InVertexBuffer);
-		});
-	}
+	FLocalVertexFactory::FDataType VertexData;
+	VertexData.NumTexCoords = 1;
+
+	//SRV setup
+	VertexData.LightMapCoordinateIndex = 0;
+	VertexData.TangentsSRV = InVertexBuffer->TangentBufferSRV;
+	VertexData.TextureCoordinatesSRV = InVertexBuffer->TexCoordBufferSRV;
+	VertexData.ColorComponentsSRV = InVertexBuffer->ColorBufferSRV;
+	VertexData.PositionComponentSRV = InVertexBuffer->PositionBufferSRV;
+
+	// Vertex Streams
+	VertexData.PositionComponent = FVertexStreamComponent(&InVertexBuffer->PositionBuffer, 0, sizeof(FVector3f), VET_Float3, EVertexStreamUsage::Default);
+	VertexData.TangentBasisComponents[0] = FVertexStreamComponent(&InVertexBuffer->TangentBuffer, 0, 2 * sizeof(FPackedNormal), VET_PackedNormal, EVertexStreamUsage::ManualFetch);
+	VertexData.TangentBasisComponents[1] = FVertexStreamComponent(&InVertexBuffer->TangentBuffer, sizeof(FPackedNormal), 2 * sizeof(FPackedNormal), VET_PackedNormal, EVertexStreamUsage::ManualFetch);
+	VertexData.ColorComponent = FVertexStreamComponent(&InVertexBuffer->ColorBuffer, 0, sizeof(FColor), VET_Color, EVertexStreamUsage::ManualFetch);
+	VertexData.TextureCoordinates.Add(FVertexStreamComponent(&InVertexBuffer->TexCoordBuffer, 0, sizeof(FVector2f), VET_Float2, EVertexStreamUsage::ManualFetch));
+
+	SetData(RHICmdList, VertexData);
+	VertexBuffer = InVertexBuffer;
+
+	InitResource(RHICmdList);
 }
diff --git a/Engine/Plugins/2D/Paper2D/Source/Paper2D/Private/PaperSpriteVertexBuffer.h b/Engine/Plugins/2D/Paper2D/Source/Paper2D/Private/PaperSpriteVertexBuffer.h
--- a/Engine/Plugins/2D/Paper2D/Source/Paper2D/Private/PaperSpriteVertexBuffer.h
+++ b/Engine/Plugins/2D/Paper2D/Source/Paper2D/Private/PaperSpriteVertexBuffer.h
@@ -87,7 +87,7 @@ public:
 	FPaperSpriteVertexFactory(ERHIFeatureLevel::Type FeatureLevel);
 
 	/* Initializes this factory with a given vertex buffer. */
-	void Init(const FPaperSpriteVertexBuffer* VertexBuffer);
+	void Init(FRHICommandListBase& RHICmdList, const FPaperSpriteVertexBuffer* VertexBuffer);
 
 private:
 	/* Vertex buffer used to initialize this factory. */
diff --git a/Engine/Plugins/Editor/UVEditor/Source/UVEditorTools/Private/Drawing/BasicLineSetComponent.cpp b/Engine/Plugins/Editor/UVEditor/Source/UVEditorTools/Private/Drawing/BasicLineSetComponent.cpp
--- a/Engine/Plugins/Editor/UVEditor/Source/UVEditorTools/Private/Drawing/BasicLineSetComponent.cpp
+++ b/Engine/Plugins/Editor/UVEditor/Source/UVEditorTools/Private/Drawing/BasicLineSetComponent.cpp
@@ -119,7 +119,7 @@ public:
 					VertexBuffers.StaticMeshVertexBuffer.BindTangentVertexBuffer(&VertexFactory, Data);
 					VertexBuffers.StaticMeshVertexBuffer.BindTexCoordVertexBuffer(&VertexFactory, Data);
 					VertexBuffers.ColorVertexBuffer.BindColorVertexBuffer(&VertexFactory, Data);
-					VertexFactory.SetData(Data);
+					VertexFactory.SetData(RHICmdList, Data);
 
 					VertexFactory.InitResource(RHICmdList);
 					IndexBuffer.InitResource(RHICmdList);
diff --git a/Engine/Plugins/Editor/UVEditor/Source/UVEditorTools/Private/Drawing/BasicPointSetComponent.cpp b/Engine/Plugins/Editor/UVEditor/Source/UVEditorTools/Private/Drawing/BasicPointSetComponent.cpp
--- a/Engine/Plugins/Editor/UVEditor/Source/UVEditorTools/Private/Drawing/BasicPointSetComponent.cpp
+++ b/Engine/Plugins/Editor/UVEditor/Source/UVEditorTools/Private/Drawing/BasicPointSetComponent.cpp
@@ -112,7 +112,7 @@ public:
 			VertexBuffers.StaticMeshVertexBuffer.BindTangentVertexBuffer(&VertexFactory, Data);
 			VertexBuffers.StaticMeshVertexBuffer.BindTexCoordVertexBuffer(&VertexFactory, Data);
 			VertexBuffers.ColorVertexBuffer.BindColorVertexBuffer(&VertexFactory, Data);
-			VertexFactory.SetData(Data);
+			VertexFactory.SetData(RHICmdList, Data);
 
 			VertexFactory.InitResource(RHICmdList);
 			IndexBuffer.InitResource(RHICmdList);
diff --git a/Engine/Plugins/Editor/UVEditor/Source/UVEditorTools/Private/Drawing/BasicTriangleSetComponent.cpp b/Engine/Plugins/Editor/UVEditor/Source/UVEditorTools/Private/Drawing/BasicTriangleSetComponent.cpp
--- a/Engine/Plugins/Editor/UVEditor/Source/UVEditorTools/Private/Drawing/BasicTriangleSetComponent.cpp
+++ b/Engine/Plugins/Editor/UVEditor/Source/UVEditorTools/Private/Drawing/BasicTriangleSetComponent.cpp
@@ -108,7 +108,7 @@ public:
 					VertexBuffers.StaticMeshVertexBuffer.BindTangentVertexBuffer(&VertexFactory, Data);
 					VertexBuffers.StaticMeshVertexBuffer.BindTexCoordVertexBuffer(&VertexFactory, Data);
 					VertexBuffers.ColorVertexBuffer.BindColorVertexBuffer(&VertexFactory, Data);
-					VertexFactory.SetData(Data);
+					VertexFactory.SetData(RHICmdList, Data);
 
 					VertexFactory.InitResource(RHICmdList);
 					IndexBuffer.InitResource(RHICmdList);
diff --git a/Engine/Plugins/Experimental/Dataflow/Source/DataflowEnginePlugin/Private/Dataflow/DataflowEngineSceneProxy.cpp b/Engine/Plugins/Experimental/Dataflow/Source/DataflowEnginePlugin/Private/Dataflow/DataflowEngineSceneProxy.cpp
--- a/Engine/Plugins/Experimental/Dataflow/Source/DataflowEnginePlugin/Private/Dataflow/DataflowEngineSceneProxy.cpp
+++ b/Engine/Plugins/Experimental/Dataflow/Source/DataflowEnginePlugin/Private/Dataflow/DataflowEngineSceneProxy.cpp
@@ -243,7 +243,7 @@ void FDataflowEngineSceneProxy::CreateMeshRenderThreadResources()
 	VertexBuffers.StaticMeshVertexBuffer.BindTangentVertexBuffer(&VertexFactory, Data);
 	VertexBuffers.StaticMeshVertexBuffer.BindTexCoordVertexBuffer(&VertexFactory, Data);
 	VertexBuffers.ColorVertexBuffer.BindColorVertexBuffer(&VertexFactory, Data);
-	VertexFactory.SetData(Data);
+	VertexFactory.SetData(RHICmdList, Data);
 
 	VertexFactory.InitResource(RHICmdList);
 	IndexBuffer.InitResource(RHICmdList);
@@ -396,7 +396,7 @@ void FDataflowEngineSceneProxy::CreateInstancedVertexRenderThreadResources()
 				BoxVertexBuffers.StaticMeshVertexBuffer.BindTangentVertexBuffer(&VertexFactory, Data);
 				BoxVertexBuffers.StaticMeshVertexBuffer.BindTexCoordVertexBuffer(&VertexFactory, Data);
 				BoxVertexBuffers.ColorVertexBuffer.BindColorVertexBuffer(&VertexFactory, Data);
-				BoxVertexFactory.SetData(Data);
+				BoxVertexFactory.SetData(RHICmdList, Data);
 
 				BoxVertexFactory.InitResource(RHICmdList);
 				BoxIndexBuffer.InitResource(RHICmdList);
diff --git a/Engine/Plugins/Experimental/GPULightmass/Source/GPULightmass/Private/Scene/Scene.cpp b/Engine/Plugins/Experimental/GPULightmass/Source/GPULightmass/Private/Scene/Scene.cpp
--- a/Engine/Plugins/Experimental/GPULightmass/Source/GPULightmass/Private/Scene/Scene.cpp
+++ b/Engine/Plugins/Experimental/GPULightmass/Source/GPULightmass/Private/Scene/Scene.cpp
@@ -1291,7 +1291,7 @@ void FScene::AddGeometryInstanceFromComponent(ULandscapeComponent* InComponent)
 				0.f,
 				(float)((InstanceRenderStateRef->SubsectionSizeVerts >> LodIndex) - 1),
 				1.f / (float)((InstanceRenderStateRef->SubsectionSizeVerts >> LodIndex) - 1));
-			InstanceRenderStateRef->LandscapeFixedGridUniformShaderParameters[LodIndex].SetContents(Parameters);
+			InstanceRenderStateRef->LandscapeFixedGridUniformShaderParameters[LodIndex].SetContents(RHICmdList, Parameters);
 		}
 
 		{
@@ -1355,8 +1355,8 @@ void FScene::AddGeometryInstanceFromComponent(ULandscapeComponent* InComponent)
 			LandscapeParams.XYOffsetmapTextureSampler = GBlackTexture->SamplerStateRHI;
 
 			InstanceRenderStateRef->LandscapeUniformShaderParameters = MakeUnique<TUniformBuffer<FLandscapeUniformShaderParameters>>();
+			InstanceRenderStateRef->LandscapeUniformShaderParameters->SetContents(RHICmdList, LandscapeParams);
 			InstanceRenderStateRef->LandscapeUniformShaderParameters->InitResource(RHICmdList);
-			InstanceRenderStateRef->LandscapeUniformShaderParameters->SetContents(LandscapeParams);
 		}
 
 		for (int32 LODIndex = 0; LODIndex < InstanceLightmapRenderStateInitializers.Num(); LODIndex++)
diff --git a/Engine/Plugins/FX/Niagara/Source/Niagara/Private/NiagaraComponent.cpp b/Engine/Plugins/FX/Niagara/Source/Niagara/Private/NiagaraComponent.cpp
--- a/Engine/Plugins/FX/Niagara/Source/Niagara/Private/NiagaraComponent.cpp
+++ b/Engine/Plugins/FX/Niagara/Source/Niagara/Private/NiagaraComponent.cpp
@@ -305,6 +305,8 @@ const FVector3f& FNiagaraSceneProxy::GetLWCRenderTile() const
 
 TUniformBuffer<FPrimitiveUniformShaderParameters>* FNiagaraSceneProxy::GetCustomUniformBufferResource(bool bHasVelocity, const FBox& InstanceBounds) const
 {
+	FRHICommandListBase& RHICmdList = FRHICommandListImmediate::Get();
+
 	// Use a hash to determine if we can re-use any uniform buffer
 	uint32 KeyHash = HashCombine(bHasVelocity, InstanceBounds.IsValid);
 
@@ -373,8 +375,8 @@ TUniformBuffer<FPrimitiveUniformShaderParameters>* FNiagaraSceneProxy::GetCustom
 		}
 
 		CustomUBRef = new TUniformBuffer<FPrimitiveUniformShaderParameters>();
-		CustomUBRef->SetContents(UBBuilder.Build());
-		CustomUBRef->InitResource(FRHICommandListImmediate::Get());
+		CustomUBRef->SetContents(RHICmdList, UBBuilder.Build());
+		CustomUBRef->InitResource(RHICmdList);
 	}
 
 	return CustomUBRef;
diff --git a/Engine/Plugins/FX/Niagara/Source/Niagara/Private/NiagaraRendererVolumes.cpp b/Engine/Plugins/FX/Niagara/Source/Niagara/Private/NiagaraRendererVolumes.cpp
--- a/Engine/Plugins/FX/Niagara/Source/Niagara/Private/NiagaraRendererVolumes.cpp
+++ b/Engine/Plugins/FX/Niagara/Source/Niagara/Private/NiagaraRendererVolumes.cpp
@@ -103,6 +103,8 @@ void FNiagaraRendererVolumes::CreateRenderThreadResources()
 {
 	FNiagaraRenderer::CreateRenderThreadResources();
 
+	FRHICommandListBase& RHICmdList = FRHICommandListImmediate::Get();
+
 	FLocalVertexFactory::FDataType VFData;
 	VFData.PositionComponent = FVertexStreamComponent(&GNullVertexBuffer, 0, 0, VET_Float3);
 	VFData.PositionComponentSRV = GNullVertexBuffer.VertexBufferSRV;
@@ -118,8 +120,8 @@ void FNiagaraRendererVolumes::CreateRenderThreadResources()
 	VFData.ColorComponent = FVertexStreamComponent(&GNullColorVertexBuffer, 0, 0, VET_Color, EVertexStreamUsage::ManualFetch);
 	VFData.ColorComponentsSRV = GNullColorVertexBuffer.VertexBufferSRV;
 
-	VertexFactory.SetData(VFData);
-	VertexFactory.InitResource(FRHICommandListImmediate::Get());
+	VertexFactory.SetData(RHICmdList, VFData);
+	VertexFactory.InitResource(RHICmdList);
 }
 
 void FNiagaraRendererVolumes::ReleaseRenderThreadResources()
diff --git a/Engine/Plugins/Runtime/MeshModelingToolset/Source/ModelingComponents/Private/Drawing/LineSetComponent.cpp b/Engine/Plugins/Runtime/MeshModelingToolset/Source/ModelingComponents/Private/Drawing/LineSetComponent.cpp
--- a/Engine/Plugins/Runtime/MeshModelingToolset/Source/ModelingComponents/Private/Drawing/LineSetComponent.cpp
+++ b/Engine/Plugins/Runtime/MeshModelingToolset/Source/ModelingComponents/Private/Drawing/LineSetComponent.cpp
@@ -124,7 +124,7 @@ public:
 			VertexBuffers.StaticMeshVertexBuffer.BindTangentVertexBuffer(&VertexFactory, Data);
 			VertexBuffers.StaticMeshVertexBuffer.BindTexCoordVertexBuffer(&VertexFactory, Data);
 			VertexBuffers.ColorVertexBuffer.BindColorVertexBuffer(&VertexFactory, Data);
-			VertexFactory.SetData(Data);
+			VertexFactory.SetData(RHICmdList, Data);
 
 			VertexFactory.InitResource(RHICmdList);
 			IndexBuffer.InitResource(RHICmdList);
diff --git a/Engine/Plugins/Runtime/MeshModelingToolset/Source/ModelingComponents/Private/Drawing/MeshWireframeComponent.cpp b/Engine/Plugins/Runtime/MeshModelingToolset/Source/ModelingComponents/Private/Drawing/MeshWireframeComponent.cpp
--- a/Engine/Plugins/Runtime/MeshModelingToolset/Source/ModelingComponents/Private/Drawing/MeshWireframeComponent.cpp
+++ b/Engine/Plugins/Runtime/MeshModelingToolset/Source/ModelingComponents/Private/Drawing/MeshWireframeComponent.cpp
@@ -252,7 +252,7 @@ public:
 			VertexBuffers.StaticMeshVertexBuffer.BindTangentVertexBuffer(&VertexFactory, Data);
 			VertexBuffers.StaticMeshVertexBuffer.BindTexCoordVertexBuffer(&VertexFactory, Data);
 			VertexBuffers.ColorVertexBuffer.BindColorVertexBuffer(&VertexFactory, Data);
-			VertexFactory.SetData(Data);
+			VertexFactory.SetData(RHICmdList, Data);
 
 			VertexFactory.InitResource(RHICmdList);
 			IndexBuffer.InitResource(RHICmdList);
diff --git a/Engine/Plugins/Runtime/MeshModelingToolset/Source/ModelingComponents/Private/Drawing/PointSetComponent.cpp b/Engine/Plugins/Runtime/MeshModelingToolset/Source/ModelingComponents/Private/Drawing/PointSetComponent.cpp
--- a/Engine/Plugins/Runtime/MeshModelingToolset/Source/ModelingComponents/Private/Drawing/PointSetComponent.cpp
+++ b/Engine/Plugins/Runtime/MeshModelingToolset/Source/ModelingComponents/Private/Drawing/PointSetComponent.cpp
@@ -141,7 +141,7 @@ public:
 			VertexBuffers.StaticMeshVertexBuffer.BindTangentVertexBuffer(&VertexFactory, Data);
 			VertexBuffers.StaticMeshVertexBuffer.BindTexCoordVertexBuffer(&VertexFactory, Data);
 			VertexBuffers.ColorVertexBuffer.BindColorVertexBuffer(&VertexFactory, Data);
-			VertexFactory.SetData(Data);
+			VertexFactory.SetData(RHICmdList, Data);
 
 			VertexFactory.InitResource(RHICmdList);
 			IndexBuffer.InitResource(RHICmdList);
diff --git a/Engine/Plugins/Runtime/MeshModelingToolset/Source/ModelingComponents/Private/Drawing/TriangleSetComponent.cpp b/Engine/Plugins/Runtime/MeshModelingToolset/Source/ModelingComponents/Private/Drawing/TriangleSetComponent.cpp
--- a/Engine/Plugins/Runtime/MeshModelingToolset/Source/ModelingComponents/Private/Drawing/TriangleSetComponent.cpp
+++ b/Engine/Plugins/Runtime/MeshModelingToolset/Source/ModelingComponents/Private/Drawing/TriangleSetComponent.cpp
@@ -110,7 +110,7 @@ public:
 				VertexBuffers.StaticMeshVertexBuffer.BindTangentVertexBuffer(&VertexFactory, Data);
 				VertexBuffers.StaticMeshVertexBuffer.BindTexCoordVertexBuffer(&VertexFactory, Data);
 				VertexBuffers.ColorVertexBuffer.BindColorVertexBuffer(&VertexFactory, Data);
-				VertexFactory.SetData(Data);
+				VertexFactory.SetData(RHICmdList, Data);
 
 				VertexFactory.InitResource(RHICmdList);
 				IndexBuffer.InitResource(RHICmdList);
diff --git a/Engine/Plugins/Runtime/nDisplay/Source/DisplayCluster/Private/Game/EngineClasses/Scene/DisplayClusterWidgetComponent.cpp b/Engine/Plugins/Runtime/nDisplay/Source/DisplayCluster/Private/Game/EngineClasses/Scene/DisplayClusterWidgetComponent.cpp
--- a/Engine/Plugins/Runtime/nDisplay/Source/DisplayCluster/Private/Game/EngineClasses/Scene/DisplayClusterWidgetComponent.cpp
+++ b/Engine/Plugins/Runtime/nDisplay/Source/DisplayCluster/Private/Game/EngineClasses/Scene/DisplayClusterWidgetComponent.cpp
@@ -57,7 +57,7 @@ public:
 			this->StaticMeshVertexBuffers.StaticMeshVertexBuffer.BindPackedTexCoordVertexBuffer(&this->VertexFactory, Data);
 			this->StaticMeshVertexBuffers.StaticMeshVertexBuffer.BindLightMapVertexBuffer(&this->VertexFactory, Data, 0);
 			this->StaticMeshVertexBuffers.ColorVertexBuffer.BindColorVertexBuffer(&this->VertexFactory, Data);
-			this->VertexFactory.SetData(Data);
+			this->VertexFactory.SetData(RHICmdList, Data);
 
 			this->VertexFactory.InitResource(RHICmdList);
 		});
diff --git a/Engine/Source/Runtime/Engine/Classes/Engine/InstancedStaticMesh.h b/Engine/Source/Runtime/Engine/Classes/Engine/InstancedStaticMesh.h
--- a/Engine/Source/Runtime/Engine/Classes/Engine/InstancedStaticMesh.h
+++ b/Engine/Source/Runtime/Engine/Classes/Engine/InstancedStaticMesh.h
@@ -299,14 +299,14 @@ public:
 	/**
 	 * An implementation of the interface used by TSynchronizedResource to update the resource with new data from the game thread.
 	 */
-	void SetData(const FDataType& InData, const FInstancedStaticMeshDataType* InInstanceData)
+	void SetData(FRHICommandListBase& RHICmdList, const FDataType& InData, const FInstancedStaticMeshDataType* InInstanceData)
 	{
 		Data = InData;
 		if (InInstanceData)
 		{
 			InstanceData = *InInstanceData;
 		}
-		UpdateRHI(FRenderResource::GetCommandList());
+		UpdateRHI(RHICmdList);
 	}
 
 	/**
@@ -477,7 +477,7 @@ public:
 	/** Feature level used when creating instance data */
 	ERHIFeatureLevel::Type FeatureLevel;
 
-	ENGINE_API void BindBuffersToVertexFactories();
+	ENGINE_API void BindBuffersToVertexFactories(FRHICommandListBase& RHICmdList);
 
 private:
 	void InitVertexFactories();
@@ -527,7 +527,7 @@ public:
 
 	ENGINE_API virtual void OnTransformChanged() override;
 
-	ENGINE_API virtual void UpdateInstances_RenderThread(const FInstanceUpdateCmdBuffer& CmdBuffer, const FBoxSphereBounds& InBounds, const FBoxSphereBounds& InLocalBounds, const FBoxSphereBounds& InStaticMeshBounds) override;
+	ENGINE_API virtual void UpdateInstances_RenderThread(FRHICommandListBase& RHICmdList, const FInstanceUpdateCmdBuffer& CmdBuffer, const FBoxSphereBounds& InBounds, const FBoxSphereBounds& InLocalBounds, const FBoxSphereBounds& InStaticMeshBounds) override;
 
 	virtual FPrimitiveViewRelevance GetViewRelevance(const FSceneView* View) const override
 	{
diff --git a/Engine/Source/Runtime/Engine/Private/Animation/SkinWeightProfile.cpp b/Engine/Source/Runtime/Engine/Private/Animation/SkinWeightProfile.cpp
--- a/Engine/Source/Runtime/Engine/Private/Animation/SkinWeightProfile.cpp
+++ b/Engine/Source/Runtime/Engine/Private/Animation/SkinWeightProfile.cpp
@@ -603,8 +603,7 @@ void FSkinWeightProfilesData::ReleaseCPUResources()
 	ResetGPUReadback();
 }
 
-template <bool bRenderThread>
-void FSkinWeightProfilesData::CreateRHIBuffers_Internal(TArray<TPair<FName, FSkinWeightRHIInfo>>& OutBuffers)
+void FSkinWeightProfilesData::CreateRHIBuffers(FRHICommandListBase& RHICmdList, TArray<TPair<FName, FSkinWeightRHIInfo>>& OutBuffers)
 {
 	const int32 NumActiveProfiles = ProfileNameToBuffer.Num();
 	check(BaseBuffer || !NumActiveProfiles);
@@ -614,14 +613,7 @@ void FSkinWeightProfilesData::CreateRHIBuffers_Internal(TArray<TPair<FName, FSki
 		const FName& ProfileName = It->Key;
 		FSkinWeightVertexBuffer* OverrideBuffer = It->Value;
 		ApplyOverrideProfile(OverrideBuffer, ProfileName);
-		if (bRenderThread)
-		{
-			OutBuffers.Emplace(ProfileName, OverrideBuffer->CreateRHIBuffer_RenderThread());
-		}
-		else
-		{
-			OutBuffers.Emplace(ProfileName, OverrideBuffer->CreateRHIBuffer_Async());
-		}
+		OutBuffers.Emplace(ProfileName, OverrideBuffer->CreateRHIBuffer(RHICmdList));
 	}
 }
 
@@ -757,12 +749,13 @@ void FSkinWeightProfilesData::ApplyOverrideProfile(FSkinWeightVertexBuffer* Over
 
 void FSkinWeightProfilesData::CreateRHIBuffers_RenderThread(TArray<TPair<FName, FSkinWeightRHIInfo>>& OutBuffers)
 {
-	CreateRHIBuffers_Internal<true>(OutBuffers);
+	CreateRHIBuffers(FRHICommandListImmediate::Get(), OutBuffers);
 }
 
 void FSkinWeightProfilesData::CreateRHIBuffers_Async(TArray<TPair<FName, FSkinWeightRHIInfo>>& OutBuffers)
 {
-	CreateRHIBuffers_Internal<false>(OutBuffers);
+	FRHIAsyncCommandList CommandList;
+	CreateRHIBuffers(*CommandList, OutBuffers);
 }
 
 void FSkinWeightProfilesData::InitRHIForStreaming(const TArray<TPair<FName, FSkinWeightRHIInfo>>& IntermediateBuffers, FRHIResourceUpdateBatcher& Batcher)
diff --git a/Engine/Source/Runtime/Engine/Private/Components/HeterogeneousVolumeComponent.cpp b/Engine/Source/Runtime/Engine/Private/Components/HeterogeneousVolumeComponent.cpp
--- a/Engine/Source/Runtime/Engine/Private/Components/HeterogeneousVolumeComponent.cpp
+++ b/Engine/Source/Runtime/Engine/Private/Components/HeterogeneousVolumeComponent.cpp
@@ -117,7 +117,7 @@ FHeterogeneousVolumeSceneProxy::FHeterogeneousVolumeSceneProxy(UHeterogeneousVol
 			Self->StaticMeshVertexBuffers.StaticMeshVertexBuffer.BindPackedTexCoordVertexBuffer(&Self->VertexFactory, Data);
 			Self->StaticMeshVertexBuffers.StaticMeshVertexBuffer.BindLightMapVertexBuffer(&Self->VertexFactory, Data, 0);
 			Self->StaticMeshVertexBuffers.ColorVertexBuffer.BindColorVertexBuffer(&Self->VertexFactory, Data);
-			Self->VertexFactory.SetData(Data);
+			Self->VertexFactory.SetData(RHICmdList, Data);
 
 			Self->VertexFactory.InitResource(RHICmdList);
 		}
diff --git a/Engine/Source/Runtime/Engine/Private/Components/MaterialBillboardComponent.cpp b/Engine/Source/Runtime/Engine/Private/Components/MaterialBillboardComponent.cpp
--- a/Engine/Source/Runtime/Engine/Private/Components/MaterialBillboardComponent.cpp
+++ b/Engine/Source/Runtime/Engine/Private/Components/MaterialBillboardComponent.cpp
@@ -107,7 +107,7 @@ public:
 			Self->StaticMeshVertexBuffers.StaticMeshVertexBuffer.BindPackedTexCoordVertexBuffer(&Self->VertexFactory, Data);
 			Self->StaticMeshVertexBuffers.StaticMeshVertexBuffer.BindLightMapVertexBuffer(&Self->VertexFactory, Data, 0);
 			Self->StaticMeshVertexBuffers.ColorVertexBuffer.BindColorVertexBuffer(&Self->VertexFactory, Data);
-			Self->VertexFactory.SetData(Data);
+			Self->VertexFactory.SetData(RHICmdList, Data);
 
 			Self->VertexFactory.InitResource(RHICmdList);
 		});
@@ -252,9 +252,7 @@ public:
 				Self->StaticMeshVertexBuffers.StaticMeshVertexBuffer.BindPackedTexCoordVertexBuffer(VertexFactoryPtr, Data);
 				Self->StaticMeshVertexBuffers.StaticMeshVertexBuffer.BindLightMapVertexBuffer(VertexFactoryPtr, Data, 0);
 				Self->StaticMeshVertexBuffers.ColorVertexBuffer.BindColorVertexBuffer(VertexFactoryPtr, Data);
-				VertexFactoryPtr->SetData(Data);
-
-				VertexFactoryPtr->UpdateRHI(RHICmdList);
+				VertexFactoryPtr->SetData(RHICmdList, Data);
 			});
 		}
 	}
diff --git a/Engine/Source/Runtime/Engine/Private/Components/SplineMeshComponent.cpp b/Engine/Source/Runtime/Engine/Private/Components/SplineMeshComponent.cpp
--- a/Engine/Source/Runtime/Engine/Private/Components/SplineMeshComponent.cpp
+++ b/Engine/Source/Runtime/Engine/Private/Components/SplineMeshComponent.cpp
@@ -292,11 +292,11 @@ void FSplineMeshSceneProxy::InitVertexFactory(USplineMeshComponent* InComponent,
 	int32 LightMapCoordinateIndex = Parent->GetLightMapCoordinateIndex();
 	// Initialize the static mesh's vertex factory.
 	ENQUEUE_RENDER_COMMAND(InitSplineMeshVertexFactory)(
-		[VertexFactory, RenderData2, bOverrideColorVertexBuffer, LightMapCoordinateIndex](FRHICommandListImmediate& RHICmdList)
+		[VertexFactory, RenderData2, bOverrideColorVertexBuffer, LightMapCoordinateIndex](FRHICommandListBase& RHICmdList)
 	{
 		FLocalVertexFactory::FDataType Data;
 		InitSplineMeshVertexFactoryComponents(RenderData2->VertexBuffers, VertexFactory, LightMapCoordinateIndex, bOverrideColorVertexBuffer, Data);
-		VertexFactory->SetData(Data);
+		VertexFactory->SetData(RHICmdList, Data);
 		VertexFactory->InitResource(RHICmdList);
 	});
 }
diff --git a/Engine/Source/Runtime/Engine/Private/DynamicMeshBuilder.cpp b/Engine/Source/Runtime/Engine/Private/DynamicMeshBuilder.cpp
--- a/Engine/Source/Runtime/Engine/Private/DynamicMeshBuilder.cpp
+++ b/Engine/Source/Runtime/Engine/Private/DynamicMeshBuilder.cpp
@@ -97,7 +97,7 @@ public:
 	 * @param Args The buffer size in bytes.
 	 * @returns A suitably sized buffer or NULL on failure.
 	 */
-	FBufferRHIRef CreateResource(FGlobalDynamicMeshPoolPolicy::CreationArguments Args)
+	FBufferRHIRef CreateResource(FRHICommandListBase& RHICmdList, FGlobalDynamicMeshPoolPolicy::CreationArguments Args)
 	{
 		FGlobalDynamicMeshPoolPolicy::CreationArguments BufferSize = GetPoolBucketSize(GetPoolBucketIndex(Args));
 		// The use of BUF_Static is deliberate - on OS X the buffer backing-store orphaning & reallocation will dominate execution time
@@ -155,7 +155,7 @@ public:
 	 * @param Args The buffer size in bytes.
 	 * @returns A suitably sized buffer or NULL on failure.
 	 */
-	FBufferRHIRef CreateResource(FGlobalDynamicMeshPoolPolicy::CreationArguments Args)
+	FBufferRHIRef CreateResource(FRHICommandListBase& RHICmdList, FGlobalDynamicMeshPoolPolicy::CreationArguments Args)
 	{
 		FGlobalDynamicMeshPoolPolicy::CreationArguments BufferSize = GetPoolBucketSize(GetPoolBucketIndex(Args));
 		FRHIResourceCreateInfo CreateInfo(TEXT("FGlobalDynamicMeshVertexPolicy"));
@@ -607,7 +607,7 @@ public:
 				Data.TangentBasisComponents[1] = FVertexStreamComponent(&PooledVertexBuffer->TangentBuffer, sizeof(FPackedNormal), 2 * sizeof(FPackedNormal), VET_PackedNormal, EVertexStreamUsage::ManualFetch);
 				Data.ColorComponent = FVertexStreamComponent(&PooledVertexBuffer->ColorBuffer, 0, sizeof(FColor), VET_Color, EVertexStreamUsage::ManualFetch);
 			}
-			VertexFactory->SetData(Data);
+			VertexFactory->SetData(RHICmdList, Data);
 		});
 
 		FLocalVertexFactory::InitResource(RHICmdList);
@@ -864,7 +864,7 @@ void FDynamicMeshBuilder::GetMesh(
 			}
 			else
 			{
-				OneFrameResources->PrimitiveUniformBuffer->SetContents(PrimitiveParams);
+				OneFrameResources->PrimitiveUniformBuffer->SetContents(RHICmdList, PrimitiveParams);
 			}
 
 			OneFrameResources->PrimitiveUniformBuffer->InitResource(RHICmdList);
@@ -955,7 +955,7 @@ void FDynamicMeshBuilder::GetMeshElement(const FPrimitiveUniformShaderParameters
 			}
 			else
 			{
-				OneFrameResource.PrimitiveUniformBuffer->SetContents(PrimitiveParams);
+				OneFrameResource.PrimitiveUniformBuffer->SetContents(RHICmdList, PrimitiveParams);
 			}
 
 			OneFrameResource.PrimitiveUniformBuffer->InitResource(RHICmdList);
@@ -1025,7 +1025,7 @@ void FDynamicMeshBuilder::Draw(FPrimitiveDrawInterface* PDI,const FMatrix& Local
 		}
 		else
 		{
-			PrimitiveUniformBuffer->SetContents(PrimitiveParams);
+			PrimitiveUniformBuffer->SetContents(FRHICommandListImmediate::Get(), PrimitiveParams);
 		}
 		PDI->RegisterDynamicResource(PrimitiveUniformBuffer);
 
diff --git a/Engine/Source/Runtime/Engine/Private/GPUSkinCache.cpp b/Engine/Source/Runtime/Engine/Private/GPUSkinCache.cpp
--- a/Engine/Source/Runtime/Engine/Private/GPUSkinCache.cpp
+++ b/Engine/Source/Runtime/Engine/Private/GPUSkinCache.cpp
@@ -1872,7 +1872,7 @@ void FGPUSkinCache::ProcessRayTracingGeometryToUpdate(FRHICommandListImmediate&
 		check(LODIndex < SkelMeshRenderData.LODRenderData.Num());
 		FSkeletalMeshLODRenderData& LODModel = SkelMeshRenderData.LODRenderData[LODIndex];
 
- 		SkinCacheEntry->GPUSkin->UpdateRayTracingGeometry(LODModel, LODIndex, VertexBufffers);
+ 		SkinCacheEntry->GPUSkin->UpdateRayTracingGeometry(RHICmdList, LODModel, LODIndex, VertexBufffers);
 	}
 }
 
diff --git a/Engine/Source/Runtime/Engine/Private/GPUSkinVertexFactory.cpp b/Engine/Source/Runtime/Engine/Private/GPUSkinVertexFactory.cpp
--- a/Engine/Source/Runtime/Engine/Private/GPUSkinVertexFactory.cpp
+++ b/Engine/Source/Runtime/Engine/Private/GPUSkinVertexFactory.cpp
@@ -178,9 +178,8 @@ uint32 FSharedPoolPolicyData::BucketSizes[NumPoolBucketSizes] = {
 /*-----------------------------------------------------------------------------
  FBoneBufferPoolPolicy
  -----------------------------------------------------------------------------*/
-FVertexBufferAndSRV FBoneBufferPoolPolicy::CreateResource(CreationArguments Args)
+FVertexBufferAndSRV FBoneBufferPoolPolicy::CreateResource(FRHICommandListBase& RHICmdList, CreationArguments Args)
 {
-	FRHICommandListBase& RHICmdList = FRHICommandListImmediate::Get();
 	uint32 BufferSize = GetPoolBucketSize(GetPoolBucketIndex(Args));
 	// in VisualStudio the copy constructor call on the return argument can be optimized out
 	// see https://msdn.microsoft.com/en-us/library/ms364057.aspx#nrvo_cpp05_topic3
@@ -200,9 +199,8 @@ void FBoneBufferPoolPolicy::FreeResource(FVertexBufferAndSRV Resource)
 {
 }
 
-FVertexBufferAndSRV FClothBufferPoolPolicy::CreateResource(CreationArguments Args)
+FVertexBufferAndSRV FClothBufferPoolPolicy::CreateResource(FRHICommandListBase& RHICmdList, CreationArguments Args)
 {
-	FRHICommandListBase& RHICmdList = FRHICommandListImmediate::Get();
 	uint32 BufferSize = GetPoolBucketSize(GetPoolBucketIndex(Args));
 	// in VisualStudio the copy constructor call on the return argument can be optimized out
 	// see https://msdn.microsoft.com/en-us/library/ms364057.aspx#nrvo_cpp05_topic3
diff --git a/Engine/Source/Runtime/Engine/Private/InstancedStaticMesh.cpp b/Engine/Source/Runtime/Engine/Private/InstancedStaticMesh.cpp
--- a/Engine/Source/Runtime/Engine/Private/InstancedStaticMesh.cpp
+++ b/Engine/Source/Runtime/Engine/Private/InstancedStaticMesh.cpp
@@ -1222,13 +1222,10 @@ void InitInstancedStaticMeshVertexFactoryComponents(
 	}
 }
 
-void FInstancedStaticMeshRenderData::BindBuffersToVertexFactories()
+void FInstancedStaticMeshRenderData::BindBuffersToVertexFactories(FRHICommandListBase& RHICmdList)
 {
 	TRACE_CPUPROFILER_EVENT_SCOPE_STR("FInstancedStaticMeshRenderData::BindBuffersToVertexFactories");
 
-	check(IsInRenderingThread());
-	FRHICommandListBase& RHICmdList = FRHICommandListImmediate::Get();
-
 	FStaticMeshInstanceBuffer* InstanceBuffer = nullptr;
 	const bool bRHISupportsManualVertexFetch = RHISupportsManualVertexFetch(GShaderPlatformForFeatureLevel[FeatureLevel]);
 	const bool bCanUseGPUScene = UseGPUScene(GMaxRHIShaderPlatform, FeatureLevel);
@@ -1252,7 +1249,7 @@ void FInstancedStaticMeshRenderData::BindBuffersToVertexFactories()
 			ColorVertexBuffer = Component->LODData[LODIndex].OverrideVertexColors;
 		}
 		InitInstancedStaticMeshVertexFactoryComponents(RenderData->VertexBuffers, ColorVertexBuffer, InstanceBuffer, &VertexFactory, LightMapCoordinateIndex, bRHISupportsManualVertexFetch, Data, InstanceData);
-		VertexFactory.SetData(Data, InstanceBuffer ? &InstanceData : nullptr);
+		VertexFactory.SetData(RHICmdList, Data, InstanceBuffer ? &InstanceData : nullptr);
 		VertexFactory.InitResource(RHICmdList);
 	}
 }
@@ -1266,9 +1263,9 @@ void FInstancedStaticMeshRenderData::InitVertexFactories()
 	}
 
 	ENQUEUE_RENDER_COMMAND(InstancedStaticMeshRenderData_InitVertexFactories)(
-		[this](FRHICommandListImmediate& RHICmdList)
+		[this](FRHICommandListBase& RHICmdList)
 		{
-			BindBuffersToVertexFactories();
+			BindBuffersToVertexFactories(RHICmdList);
 		});
 }
 
@@ -2033,12 +2030,12 @@ void FInstancedStaticMeshSceneProxy::OnTransformChanged()
 	}
 }
 
-void FInstancedStaticMeshSceneProxy::UpdateInstances_RenderThread(const FInstanceUpdateCmdBuffer& CmdBuffer, const FBoxSphereBounds& InBounds, const FBoxSphereBounds& InLocalBounds, const FBoxSphereBounds& InStaticMeshBounds)
+void FInstancedStaticMeshSceneProxy::UpdateInstances_RenderThread(FRHICommandListBase& RHICmdList, const FInstanceUpdateCmdBuffer& CmdBuffer, const FBoxSphereBounds& InBounds, const FBoxSphereBounds& InLocalBounds, const FBoxSphereBounds& InStaticMeshBounds)
 {
 	// This will flush GPU instance data, create buffers, srvs and vertex factories.
-	InstancedRenderData.BindBuffersToVertexFactories();
+	InstancedRenderData.BindBuffersToVertexFactories(RHICmdList);
 
-	return FPrimitiveSceneProxy::UpdateInstances_RenderThread(CmdBuffer, InBounds, InLocalBounds, InStaticMeshBounds);
+	return FPrimitiveSceneProxy::UpdateInstances_RenderThread(RHICmdList, CmdBuffer, InBounds, InLocalBounds, InStaticMeshBounds);
 }
 
 void FInstancedStaticMeshSceneProxy::SetupInstancedMeshBatch(int32 LODIndex, int32 BatchIndex, FMeshBatch& OutMeshBatch) const
diff --git a/Engine/Source/Runtime/Engine/Private/LocalVertexFactory.cpp b/Engine/Source/Runtime/Engine/Private/LocalVertexFactory.cpp
--- a/Engine/Source/Runtime/Engine/Private/LocalVertexFactory.cpp
+++ b/Engine/Source/Runtime/Engine/Private/LocalVertexFactory.cpp
@@ -378,23 +378,17 @@ void FLocalVertexFactory::GetVertexElements(ERHIFeatureLevel::Type FeatureLevel,
 
 void FLocalVertexFactory::SetData(const FDataType& InData)
 {
-	{
-		//const int NumTexCoords = InData.NumTexCoords;
-		//const int LightMapCoordinateIndex = InData.LightMapCoordinateIndex;
-		//check(NumTexCoords > 0);
-		//check(LightMapCoordinateIndex < NumTexCoords && LightMapCoordinateIndex >= 0);
-		//check(InData.PositionComponentSRV);
-		//check(InData.TangentsSRV);
-		//check(InData.TextureCoordinatesSRV);
-		//check(InData.ColorComponentsSRV);
-	}
+	SetData(FRHICommandListImmediate::Get(), InData);
+}
 
+void FLocalVertexFactory::SetData(FRHICommandListBase& RHICmdList, const FDataType& InData)
+{
 	// The shader code makes assumptions that the color component is a FColor, performing swizzles on ES3 and Metal platforms as necessary
 	// If the color is sent down as anything other than VET_Color then you'll get an undesired swizzle on those platforms
 	check((InData.ColorComponent.Type == VET_None) || (InData.ColorComponent.Type == VET_Color));
 
 	Data = InData;
-	UpdateRHI(FRHICommandListImmediate::Get());
+	UpdateRHI(RHICmdList);
 }
 
 /**
diff --git a/Engine/Source/Runtime/Engine/Private/Particles/ParticleGpuSimulation.cpp b/Engine/Source/Runtime/Engine/Private/Particles/ParticleGpuSimulation.cpp
--- a/Engine/Source/Runtime/Engine/Private/Particles/ParticleGpuSimulation.cpp
+++ b/Engine/Source/Runtime/Engine/Private/Particles/ParticleGpuSimulation.cpp
@@ -2654,6 +2654,8 @@ public:
 	{
 		auto FeatureLevel = ViewFamily.GetFeatureLevel();
 
+		FRHICommandListBase& RHICmdList = FRHICommandListImmediate::Get();
+
 		if (RHISupportsGPUParticles())
 		{
 			SCOPE_CYCLE_COUNTER(STAT_GPUSpritePreRenderTime);
@@ -2674,15 +2676,13 @@ public:
 
 				if (bUseLocalSpace == false)
 				{
-					Proxy->UpdateWorldSpacePrimitiveUniformBuffer();
+					Proxy->UpdateWorldSpacePrimitiveUniformBuffer(RHICmdList);
 				}
 
 				const bool bTranslucent = RendersWithTranslucentMaterial();
 				const bool bAllowSorting = FXConsoleVariables::bAllowGPUSorting
 					&& bTranslucent;
 
-				FRHICommandListBase& RHICmdList = FRHICommandListImmediate::Get();
-
 				// Iterate over views and assign parameters for each.
 				FParticleSimulationResources* SimulationResources = FXSystem->GetParticleSimulationResources();
 				FGPUSpriteCollectorResources& CollectorResources = Collector.AllocateOneFrameResource<FGPUSpriteCollectorResources>(FeatureLevel);
diff --git a/Engine/Source/Runtime/Engine/Private/Particles/ParticleSystemRender.cpp b/Engine/Source/Runtime/Engine/Private/Particles/ParticleSystemRender.cpp
--- a/Engine/Source/Runtime/Engine/Private/Particles/ParticleSystemRender.cpp
+++ b/Engine/Source/Runtime/Engine/Private/Particles/ParticleSystemRender.cpp
@@ -813,6 +813,8 @@ void FDynamicSpriteEmitterData::GetDynamicMeshElementsEmitter(const FParticleSys
 
 	const auto FeatureLevel = View->GetFeatureLevel();
 
+	FRHICommandListBase& RHICmdList = FRHICommandListImmediate::Get();
+
 	// Sort and generate particles for this view.
 	const FDynamicSpriteEmitterReplayDataBase* SourceData = GetSourceData();
 
@@ -844,11 +846,11 @@ void FDynamicSpriteEmitterData::GetDynamicMeshElementsEmitter(const FParticleSys
 			SpriteVertexFactory->SetParticleFactoryType(PVFT_Sprite);
 			SpriteVertexFactory->SetNumVertsInInstanceBuffer(SourceData->RequiredModule->bCutoutTexureIsValid ? SourceData->RequiredModule->NumBoundingVertices : 4);
 			SpriteVertexFactory->SetUsesDynamicParameter(bUsesDynamicParameter, bUsesDynamicParameter ? GetDynamicParameterVertexStride() : 0);
-			SpriteVertexFactory->InitResource(FRHICommandListImmediate::Get());
+			SpriteVertexFactory->InitResource(RHICmdList);
 
 			if (SourceData->bUseLocalSpace == false)
 			{
-				Proxy->UpdateWorldSpacePrimitiveUniformBuffer();
+				Proxy->UpdateWorldSpacePrimitiveUniformBuffer(RHICmdList);
 			}
 
 			FGlobalDynamicVertexBuffer& DynamicVertexBuffer = Collector.GetDynamicVertexBuffer();
@@ -1439,7 +1441,7 @@ void FDynamicMeshEmitterData::GetDynamicMeshElementsEmitter(const FParticleSyste
 			MeshVertexFactory->SetInstanceBuffer(Allocation.VertexBuffer, Allocation.VertexOffset, InstanceVertexStride);
 			MeshVertexFactory->SetDynamicParameterBuffer(DynamicParameterAllocation.VertexBuffer, DynamicParameterAllocation.VertexOffset, GetDynamicParameterVertexStride());
 
-			Proxy->UpdateWorldSpacePrimitiveUniformBuffer();
+			Proxy->UpdateWorldSpacePrimitiveUniformBuffer(RHICmdList);
 			MeshVertexFactory->GetInstanceVerticesCPU() = InstanceVerticesCPU;
 
 			const bool bIsWireframe = AllowDebugViewmodes() && View->Family->EngineShowFlags.Wireframe;
@@ -2369,6 +2371,9 @@ void FDynamicBeam2EmitterData::GetDynamicMeshElementsEmitter(const FParticleSyst
 	{
 		return;
 	}
+
+	FRHICommandListBase& RHICmdList = FRHICommandListImmediate::Get();
+
 	FIndexBuffer* IndexBuffer = nullptr;
 	uint32 FirstIndex = 0;
 	int32 OutTriangleCount = 0;
@@ -2393,7 +2398,7 @@ void FDynamicBeam2EmitterData::GetDynamicMeshElementsEmitter(const FParticleSyst
 
 	if (Source.bUseLocalSpace == false)
 	{
-		Proxy->UpdateWorldSpacePrimitiveUniformBuffer();
+		Proxy->UpdateWorldSpacePrimitiveUniformBuffer(RHICmdList);
 	}
 
 	auto FeatureLevel = View->GetFeatureLevel();
@@ -2401,7 +2406,7 @@ void FDynamicBeam2EmitterData::GetDynamicMeshElementsEmitter(const FParticleSyst
 	FParticleBeamTrailVertexFactory* BeamTrailVertexFactory = &CollectorResources.VertexFactory;
 	BeamTrailVertexFactory->SetParticleFactoryType(PVFT_BeamTrail);
 	BeamTrailVertexFactory->SetUsesDynamicParameter(bUsesDynamicParameter);
-	BeamTrailVertexFactory->InitResource(FRHICommandListImmediate::Get());
+	BeamTrailVertexFactory->InitResource(RHICmdList);
 
 	// Create and set the uniform buffer for this emitter.
 	BeamTrailVertexFactory->SetBeamTrailUniformBuffer(CreateBeamTrailUniformBuffer(Proxy, &Source, View));
@@ -5123,6 +5128,9 @@ void FDynamicTrailsEmitterData::GetDynamicMeshElementsEmitter(const FParticleSys
 	{
 		return;
 	}
+
+	FRHICommandListBase& RHICmdList = FRHICommandListImmediate::Get();
+
 	const bool bIsWireframe = ViewFamily.EngineShowFlags.Wireframe;
 	FIndexBuffer* IndexBuffer = nullptr;
 	uint32 FirstIndex = 0;
@@ -5153,14 +5161,14 @@ void FDynamicTrailsEmitterData::GetDynamicMeshElementsEmitter(const FParticleSys
 
 	if (SourcePointer->bUseLocalSpace == false)
 	{
-		Proxy->UpdateWorldSpacePrimitiveUniformBuffer();
+		Proxy->UpdateWorldSpacePrimitiveUniformBuffer(RHICmdList);
 	}
 
 	FDynamicBeamTrailCollectorResources& CollectorResources = Collector.AllocateOneFrameResource<FDynamicBeamTrailCollectorResources>(FeatureLevel);
 	FParticleBeamTrailVertexFactory* BeamTrailVertexFactory = &CollectorResources.VertexFactory;
 	BeamTrailVertexFactory->SetParticleFactoryType(PVFT_BeamTrail);
 	BeamTrailVertexFactory->SetUsesDynamicParameter(bUsesDynamicParameter);
-	BeamTrailVertexFactory->InitResource(FRHICommandListImmediate::Get());
+	BeamTrailVertexFactory->InitResource(RHICmdList);
 
 	// Create and set the uniform buffer for this emitter.
 	BeamTrailVertexFactory->SetBeamTrailUniformBuffer(CreateBeamTrailUniformBuffer(Proxy, SourcePointer, View));
@@ -6839,7 +6847,7 @@ void FParticleSystemSceneProxy::OnTransformChanged()
 	WorldSpaceUBHash = 0;
 }
 
-void FParticleSystemSceneProxy::UpdateWorldSpacePrimitiveUniformBuffer() const
+void FParticleSystemSceneProxy::UpdateWorldSpacePrimitiveUniformBuffer(FRHICommandListBase& RHICmdList) const
 {
 	// Hash custom floats because we need to invalidate this UB if they don't match otherwise updates to the buffer won't work
 	uint32 NewWorldSpaceUBHash = 0;
@@ -6855,6 +6863,7 @@ void FParticleSystemSceneProxy::UpdateWorldSpacePrimitiveUniformBuffer() const
 	{
 		WorldSpaceUBHash = NewWorldSpaceUBHash;
 		WorldSpacePrimitiveUniformBuffer.SetContents(
+			RHICmdList,
 			FPrimitiveUniformShaderParametersBuilder{}
 			.Defaults()
 				.LocalToWorld(FMatrix::Identity)
@@ -6873,7 +6882,7 @@ void FParticleSystemSceneProxy::UpdateWorldSpacePrimitiveUniformBuffer() const
 
 	if ( bNeedsInit)
 	{
-		WorldSpacePrimitiveUniformBuffer.InitResource(FRHICommandListImmediate::Get());
+		WorldSpacePrimitiveUniformBuffer.InitResource(RHICmdList);
 	}
 }
 
diff --git a/Engine/Source/Runtime/Engine/Private/PositionVertexBuffer.cpp b/Engine/Source/Runtime/Engine/Private/PositionVertexBuffer.cpp
--- a/Engine/Source/Runtime/Engine/Private/PositionVertexBuffer.cpp
+++ b/Engine/Source/Runtime/Engine/Private/PositionVertexBuffer.cpp
@@ -204,20 +204,20 @@ void FPositionVertexBuffer::operator=(const FPositionVertexBuffer &Other)
 	VertexData = NULL;
 }
 
-template <bool bRenderThread>
-FBufferRHIRef FPositionVertexBuffer::CreateRHIBuffer_Internal()
+FBufferRHIRef FPositionVertexBuffer::CreateRHIBuffer(FRHICommandListBase& RHICmdList)
 {
-	return CreateRHIBuffer<bRenderThread>(VertexData, NumVertices, BUF_Static | BUF_ShaderResource, TEXT("FPositionVertexBuffer"));
+	return FRenderResource::CreateRHIBuffer(RHICmdList, VertexData, NumVertices, BUF_Static | BUF_ShaderResource, TEXT("FPositionVertexBuffer"));
 }
 
 FBufferRHIRef FPositionVertexBuffer::CreateRHIBuffer_RenderThread()
 {
-	return CreateRHIBuffer_Internal<true>();
+	return CreateRHIBuffer(FRHICommandListExecutor::GetImmediateCommandList());
 }
 
 FBufferRHIRef FPositionVertexBuffer::CreateRHIBuffer_Async()
 {
-	return CreateRHIBuffer_Internal<false>();
+	FRHIAsyncCommandList CommandList;
+	return CreateRHIBuffer(*CommandList);
 }
 
 void FPositionVertexBuffer::InitRHIForStreaming(FRHIBuffer* IntermediateBuffer, FRHIResourceUpdateBatcher& Batcher)
@@ -239,7 +239,7 @@ void FPositionVertexBuffer::InitRHI(FRHICommandListBase& RHICmdList)
 {
 	TRACE_CPUPROFILER_EVENT_SCOPE(FPositionVertexBuffer::InitRHI);
 
-	VertexBufferRHI = CreateRHIBuffer_RenderThread();
+	VertexBufferRHI = CreateRHIBuffer(RHICmdList);
 	// we have decide to create the SRV based on GMaxRHIShaderPlatform because this is created once and shared between feature levels for editor preview.
 	// Also check to see whether cpu access has been activated on the vertex data
 	if (VertexBufferRHI)
diff --git a/Engine/Source/Runtime/Engine/Private/PrimitiveSceneProxy.cpp b/Engine/Source/Runtime/Engine/Private/PrimitiveSceneProxy.cpp
--- a/Engine/Source/Runtime/Engine/Private/PrimitiveSceneProxy.cpp
+++ b/Engine/Source/Runtime/Engine/Private/PrimitiveSceneProxy.cpp
@@ -789,7 +789,7 @@ void FPrimitiveSceneProxy::SetTransform(const FMatrix& InLocalToWorld, const FBo
 	OnTransformChanged();
 }
 
-void FPrimitiveSceneProxy::UpdateInstances_RenderThread(const FInstanceUpdateCmdBuffer& CmdBuffer, const FBoxSphereBounds& InBounds, const FBoxSphereBounds& InLocalBounds, const FBoxSphereBounds& InStaticMeshBounds)
+void FPrimitiveSceneProxy::UpdateInstances_RenderThread(FRHICommandListBase& RHICmdList, const FInstanceUpdateCmdBuffer& CmdBuffer, const FBoxSphereBounds& InBounds, const FBoxSphereBounds& InLocalBounds, const FBoxSphereBounds& InStaticMeshBounds)
 {
 	TRACE_CPUPROFILER_EVENT_SCOPE_STR("FPrimitiveSceneProxy::UpdateInstances_RenderThread");
 
diff --git a/Engine/Source/Runtime/Engine/Private/RawIndexBuffer.cpp b/Engine/Source/Runtime/Engine/Private/RawIndexBuffer.cpp
--- a/Engine/Source/Runtime/Engine/Private/RawIndexBuffer.cpp
+++ b/Engine/Source/Runtime/Engine/Private/RawIndexBuffer.cpp
@@ -320,8 +320,7 @@ FIndexArrayView FRawStaticIndexBuffer::GetArrayView() const
 	return FIndexArrayView(IndexStorage.GetData(),NumIndices,b32Bit);
 }
 
-template <bool bRenderThread>
-FBufferRHIRef FRawStaticIndexBuffer::CreateRHIBuffer_Internal()
+FBufferRHIRef FRawStaticIndexBuffer::CreateRHIBuffer(FRHICommandListBase& RHICmdList)
 {
 	const uint32 IndexStride = b32Bit ? sizeof(uint32) : sizeof(uint16);
 	const uint32 SizeInBytes = IndexStorage.Num();
@@ -342,35 +341,27 @@ FBufferRHIRef FRawStaticIndexBuffer::CreateRHIBuffer_Internal()
 		const static FLazyName ClassName16(TEXT("FRawStaticIndexBuffer16"));
 
 		// Create the index buffer.
-		FBufferRHIRef Ret;
 		FRHIResourceCreateInfo CreateInfo(Is32Bit() ? TEXT("FRawStaticIndexBuffer32") : TEXT("FRawStaticIndexBuffer16"), &IndexStorage);
 		CreateInfo.ClassName = Is32Bit() ? ClassName32 : ClassName16;
 		CreateInfo.OwnerName = GetOwnerName();
 		CreateInfo.bWithoutNativeResource = !SizeInBytes;
-		if (bRenderThread)
-		{
-			Ret = FRHICommandListImmediate::Get().CreateIndexBuffer(IndexStride, SizeInBytes, BufferFlags, CreateInfo);
-		}
-		else
-		{
-			FRHIAsyncCommandList CommandList;
-			Ret = CommandList->CreateBuffer(SizeInBytes, BufferFlags | EBufferUsageFlags::IndexBuffer, IndexStride, ERHIAccess::SRVMask, CreateInfo);
-		}
 
-		Ret->SetOwnerName(GetOwnerName());
-		return Ret;
+		FBufferRHIRef Buffer = RHICmdList.CreateBuffer(SizeInBytes, BufferFlags | EBufferUsageFlags::IndexBuffer, IndexStride, ERHIAccess::VertexOrIndexBuffer, CreateInfo);
+		Buffer->SetOwnerName(GetOwnerName());
+		return Buffer;
 	}
 	return nullptr;
 }
 
 FBufferRHIRef FRawStaticIndexBuffer::CreateRHIBuffer_RenderThread()
 {
-	return CreateRHIBuffer_Internal<true>();
+	return CreateRHIBuffer(FRHICommandListExecutor::GetImmediateCommandList());
 }
 
 FBufferRHIRef FRawStaticIndexBuffer::CreateRHIBuffer_Async()
 {
-	return CreateRHIBuffer_Internal<false>();
+	FRHIAsyncCommandList CommandList;
+	return CreateRHIBuffer(*CommandList);
 }
 
 void FRawStaticIndexBuffer::InitRHIForStreaming(FRHIBuffer* IntermediateBuffer, FRHIResourceUpdateBatcher& Batcher)
@@ -392,7 +383,7 @@ void FRawStaticIndexBuffer::ReleaseRHIForStreaming(FRHIResourceUpdateBatcher& Ba
 void FRawStaticIndexBuffer::InitRHI(FRHICommandListBase& RHICmdList)
 {
 	TRACE_CPUPROFILER_EVENT_SCOPE(FRawStaticIndexBuffer::InitRHI);
-	IndexBufferRHI = CreateRHIBuffer_RenderThread();
+	IndexBufferRHI = CreateRHIBuffer(RHICmdList);
 }
 
 void FRawStaticIndexBuffer::Serialize(FArchive& Ar, bool bNeedsCPUAccess)
@@ -480,13 +471,13 @@ void FRawStaticIndexBuffer16or32Interface::ReleaseRHIForStreaming(FRHIResourceUp
 }
 
 FBufferRHIRef FRawStaticIndexBuffer16or32Interface::CreateRHIIndexBufferInternal(
+	FRHICommandListBase& RHICmdList,
 	const TCHAR* InDebugName,
 	const FName& InOwnerName,
 	int32 IndexCount,
 	size_t IndexSize,
 	FResourceArrayInterface* ResourceArray,
-	bool bNeedSRV,
-	bool bRenderThread
+	bool bNeedSRV
 )
 {
 	// Create the index buffer.
@@ -501,21 +492,12 @@ FBufferRHIRef FRawStaticIndexBuffer16or32Interface::CreateRHIIndexBufferInternal
 		Flags |= EBufferUsageFlags::ShaderResource;
 	}
 
-	FBufferRHIRef Ret;
 	const uint32 Size = IndexCount * IndexSize;
 	CreateInfo.bWithoutNativeResource = !Size;
-	if (bRenderThread)
-	{
-		Ret = FRHICommandListImmediate::Get().CreateIndexBuffer(IndexSize, Size, Flags, CreateInfo);
-	}
-	else
-	{
-		FRHIAsyncCommandList CommandList;
-		Ret = CommandList->CreateBuffer(Size, Flags | EBufferUsageFlags::IndexBuffer, IndexSize, ERHIAccess::SRVMask, CreateInfo);
-	}
 
-	Ret->SetOwnerName(InOwnerName);
-	return Ret;
+	FBufferRHIRef Buffer = RHICmdList.CreateBuffer(Size, Flags | EBufferUsageFlags::IndexBuffer, IndexSize, ERHIAccess::VertexOrIndexBuffer, CreateInfo);
+	Buffer->SetOwnerName(InOwnerName);
+	return Buffer;
 }
 
 /*-----------------------------------------------------------------------------
diff --git a/Engine/Source/Runtime/Engine/Private/RayTracingSkinnedGeometry.cpp b/Engine/Source/Runtime/Engine/Private/RayTracingSkinnedGeometry.cpp
--- a/Engine/Source/Runtime/Engine/Private/RayTracingSkinnedGeometry.cpp
+++ b/Engine/Source/Runtime/Engine/Private/RayTracingSkinnedGeometry.cpp
@@ -36,6 +36,7 @@ FAutoConsoleVariableRef CVarSkinnedGeometryMaxRayTracingPrimitivesPerCmdList(
 
 void FRayTracingSkinnedGeometryUpdateQueue::Add(FRayTracingGeometry* InRayTracingGeometry, const FRayTracingAccelerationStructureSize& StructureSize, EAccelerationStructureBuildMode InBuildMode)
 {
+	FScopeLock Lock(&CS);
 	FRayTracingUpdateInfo* CurrentUpdateInfo = ToUpdate.Find(InRayTracingGeometry);
 	if (CurrentUpdateInfo == nullptr)
 	{
@@ -54,6 +55,7 @@ void FRayTracingSkinnedGeometryUpdateQueue::Add(FRayTracingGeometry* InRayTracin
 
 void FRayTracingSkinnedGeometryUpdateQueue::Remove(FRayTracingGeometry* RayTracingGeometry, uint32 EstimatedMemory)
 {
+	FScopeLock Lock(&CS);
 	if (ToUpdate.Find(RayTracingGeometry) != nullptr)
 	{
 		ToUpdate.Remove(RayTracingGeometry);
@@ -86,6 +88,8 @@ void FRayTracingSkinnedGeometryUpdateQueue::Commit(FRHICommandListImmediate & RH
 
 	if (ToUpdate.Num())
 	{
+		FScopeLock Lock(&CS);
+
 		// If we have more deferred deleted data than set limit then force flush to make sure all pending releases have actually been freed
 		// before reallocating a lot of new BLAS data
 		if (EstimatedMemoryPendingRelease >= GMemoryLimitForBatchedRayTracingGeometryUpdates * 1024ull * 1024ull)
diff --git a/Engine/Source/Runtime/Engine/Private/RendererSettings.cpp b/Engine/Source/Runtime/Engine/Private/RendererSettings.cpp
--- a/Engine/Source/Runtime/Engine/Private/RendererSettings.cpp
+++ b/Engine/Source/Runtime/Engine/Private/RendererSettings.cpp
@@ -529,10 +529,10 @@ void URendererSettings::UpdateWorkingColorSpaceAndChromaticities()
 	FColorSpace UpdatedColorSpace = FColorSpace::GetWorking();
 
 	ENQUEUE_RENDER_COMMAND(WorkingColorSpaceCommand)(
-		[UpdatedColorSpace](FRHICommandList&)
+		[UpdatedColorSpace](FRHICommandList& RHICmdList)
 		{
 			// Set or update the global uniform buffer for Working Color Space conversions.
-			GDefaultWorkingColorSpaceUniformBuffer.Update(UpdatedColorSpace);
+			GDefaultWorkingColorSpaceUniformBuffer.Update(RHICmdList, UpdatedColorSpace);
 		});
 }
 
diff --git a/Engine/Source/Runtime/Engine/Private/Rendering/ColorVertexBuffer.cpp b/Engine/Source/Runtime/Engine/Private/Rendering/ColorVertexBuffer.cpp
--- a/Engine/Source/Runtime/Engine/Private/Rendering/ColorVertexBuffer.cpp
+++ b/Engine/Source/Runtime/Engine/Private/Rendering/ColorVertexBuffer.cpp
@@ -385,20 +385,20 @@ uint32 FColorVertexBuffer::GetAllocatedSize() const
 	}
 }
 
-template <bool bRenderThread>
-FBufferRHIRef FColorVertexBuffer::CreateRHIBuffer_Internal()
+FBufferRHIRef FColorVertexBuffer::CreateRHIBuffer(FRHICommandListBase& RHICmdList)
 {
-	return CreateRHIBuffer<bRenderThread>(VertexData, NumVertices, BUF_Static | BUF_ShaderResource, TEXT("FColorVertexBuffer"));
+	return FRenderResource::CreateRHIBuffer(RHICmdList, VertexData, NumVertices, BUF_Static | BUF_ShaderResource, TEXT("FColorVertexBuffer"));
 }
 
 FBufferRHIRef FColorVertexBuffer::CreateRHIBuffer_RenderThread()
 {
-	return CreateRHIBuffer_Internal<true>();
+	return CreateRHIBuffer(FRHICommandListImmediate::Get());
 }
 
 FBufferRHIRef FColorVertexBuffer::CreateRHIBuffer_Async()
 {
-	return CreateRHIBuffer_Internal<false>();
+	FRHIAsyncCommandList CommandList;
+	return CreateRHIBuffer(*CommandList);
 }
 
 void FColorVertexBuffer::InitRHIForStreaming(FRHIBuffer* IntermediateBuffer, FRHIResourceUpdateBatcher& Batcher)
@@ -422,7 +422,7 @@ void FColorVertexBuffer::InitRHI(FRHICommandListBase& RHICmdList)
 	TRACE_CPUPROFILER_EVENT_SCOPE(FColorVertexBuffer::InitRHI);
 	SCOPED_LOADTIMER(FColorVertexBuffer_InitRHI);
 
-	VertexBufferRHI = CreateRHIBuffer_RenderThread();
+	VertexBufferRHI = CreateRHIBuffer(RHICmdList);
 	if (VertexBufferRHI && RHISupportsManualVertexFetch(GMaxRHIShaderPlatform))
 	{
 		ColorComponentsSRV = RHICmdList.CreateShaderResourceView(VertexBufferRHI, 4, PF_R8G8B8A8);
diff --git a/Engine/Source/Runtime/Engine/Private/Rendering/MultiSizeIndexContainer.cpp b/Engine/Source/Runtime/Engine/Private/Rendering/MultiSizeIndexContainer.cpp
--- a/Engine/Source/Runtime/Engine/Private/Rendering/MultiSizeIndexContainer.cpp
+++ b/Engine/Source/Runtime/Engine/Private/Rendering/MultiSizeIndexContainer.cpp
@@ -4,6 +4,7 @@
 #include "EngineLogs.h"
 #include "RawIndexBuffer.h"
 #include "Stats/Stats.h"
+#include "RenderingThread.h"
 
 FMultiSizeIndexContainer::~FMultiSizeIndexContainer()
 {
@@ -193,36 +194,31 @@ void FMultiSizeIndexContainer::SerializeMetaData(FArchive& Ar, bool bNeedsCPUAcc
 	IndexBuffer->SerializeMetaData(Ar);
 }
 
-FBufferRHIRef FMultiSizeIndexContainer::CreateRHIBuffer_RenderThread()
+FBufferRHIRef FMultiSizeIndexContainer::CreateRHIBuffer(FRHICommandListBase& RHICmdList)
 {
 	if (IndexBuffer)
 	{
 		if (DataTypeSize == sizeof(uint16))
 		{
-			return static_cast<FRawStaticIndexBuffer16or32<uint16>*>(IndexBuffer)->CreateRHIBuffer_RenderThread();
+			return static_cast<FRawStaticIndexBuffer16or32<uint16>*>(IndexBuffer)->CreateRHIBuffer(RHICmdList);
 		}
 		else
 		{
-			return static_cast<FRawStaticIndexBuffer16or32<uint32>*>(IndexBuffer)->CreateRHIBuffer_RenderThread();
+			return static_cast<FRawStaticIndexBuffer16or32<uint32>*>(IndexBuffer)->CreateRHIBuffer(RHICmdList);
 		}
 	}
 	return nullptr;
 }
 
+FBufferRHIRef FMultiSizeIndexContainer::CreateRHIBuffer_RenderThread()
+{
+	return CreateRHIBuffer(FRHICommandListImmediate::Get());
+}
+
 FBufferRHIRef FMultiSizeIndexContainer::CreateRHIBuffer_Async()
 {
-	if (IndexBuffer)
-	{
-		if (DataTypeSize == sizeof(uint16))
-		{
-			return static_cast<FRawStaticIndexBuffer16or32<uint16>*>(IndexBuffer)->CreateRHIBuffer_Async();
-		}
-		else
-		{
-			return static_cast<FRawStaticIndexBuffer16or32<uint32>*>(IndexBuffer)->CreateRHIBuffer_Async();
-		}
-	}
-	return nullptr;
+	FRHIAsyncCommandList RHICmdList;
+	return CreateRHIBuffer(*RHICmdList);
 }
 
 void FMultiSizeIndexContainer::InitRHIForStreaming(FRHIBuffer* IntermediateBuffer, FRHIResourceUpdateBatcher& Batcher)
diff --git a/Engine/Source/Runtime/Engine/Private/Rendering/SkeletalMeshVertexClothBuffer.cpp b/Engine/Source/Runtime/Engine/Private/Rendering/SkeletalMeshVertexClothBuffer.cpp
--- a/Engine/Source/Runtime/Engine/Private/Rendering/SkeletalMeshVertexClothBuffer.cpp
+++ b/Engine/Source/Runtime/Engine/Private/Rendering/SkeletalMeshVertexClothBuffer.cpp
@@ -75,20 +75,20 @@ void FSkeletalMeshVertexClothBuffer::ClearMetaData()
 	NumVertices = 0;
 }
 
-template <bool bRenderThread>
-FBufferRHIRef FSkeletalMeshVertexClothBuffer::CreateRHIBuffer_Internal()
+FBufferRHIRef FSkeletalMeshVertexClothBuffer::CreateRHIBuffer(FRHICommandListBase& RHICmdList)
 {
-	return CreateRHIBuffer<bRenderThread>(VertexData, NumVertices, BUF_Static | BUF_ShaderResource, TEXT("FSkeletalMeshVertexClothBuffer"));
+	return FRenderResource::CreateRHIBuffer(RHICmdList, VertexData, NumVertices, BUF_Static | BUF_ShaderResource, TEXT("FSkeletalMeshVertexClothBuffer"));
 }
 
 FBufferRHIRef FSkeletalMeshVertexClothBuffer::CreateRHIBuffer_RenderThread()
 {
-	return CreateRHIBuffer_Internal<true>();
+	return CreateRHIBuffer(FRHICommandListImmediate::Get());
 }
 
 FBufferRHIRef FSkeletalMeshVertexClothBuffer::CreateRHIBuffer_Async()
 {
-	return CreateRHIBuffer_Internal<false>();
+	FRHIAsyncCommandList CommandList;
+	return CreateRHIBuffer(*CommandList);
 }
 
 void FSkeletalMeshVertexClothBuffer::InitRHIForStreaming(FRHIBuffer* IntermediateBuffer, FRHIResourceUpdateBatcher& Batcher)
@@ -114,7 +114,7 @@ void FSkeletalMeshVertexClothBuffer::InitRHI(FRHICommandListBase& RHICmdList)
 {
 	SCOPED_LOADTIMER(FSkeletalMeshVertexClothBuffer_InitRHI);
 
-	VertexBufferRHI = CreateRHIBuffer_RenderThread();
+	VertexBufferRHI = CreateRHIBuffer(RHICmdList);
 	if (VertexBufferRHI)
 	{
 		VertexBufferSRV = RHICmdList.CreateShaderResourceView(VertexBufferRHI, 16, PF_A32B32G32R32F);
diff --git a/Engine/Source/Runtime/Engine/Private/Rendering/SkinWeightVertexBuffer.cpp b/Engine/Source/Runtime/Engine/Private/Rendering/SkinWeightVertexBuffer.cpp
--- a/Engine/Source/Runtime/Engine/Private/Rendering/SkinWeightVertexBuffer.cpp
+++ b/Engine/Source/Runtime/Engine/Private/Rendering/SkinWeightVertexBuffer.cpp
@@ -128,7 +128,7 @@ void FSkinWeightLookupVertexBuffer::AllocateData()
 void FSkinWeightLookupVertexBuffer::InitRHI(FRHICommandListBase& RHICmdList)
 {
 	// BUF_ShaderResource is needed for support of the SkinCache (we could make is dependent on GEnableGPUSkinCacheShaders or are there other users?)
-	VertexBufferRHI = CreateRHIBuffer_RenderThread();
+	VertexBufferRHI = CreateRHIBuffer(RHICmdList);
 	if (VertexBufferRHI)
 	{
 		bool bSRV = GPixelFormats[PixelFormat].Supported;
@@ -153,12 +153,13 @@ void FSkinWeightLookupVertexBuffer::ReleaseRHI()
 
 FBufferRHIRef FSkinWeightLookupVertexBuffer::CreateRHIBuffer_RenderThread()
 {
-	return CreateRHIBuffer_Internal<true>();
+	return CreateRHIBuffer(FRHICommandListImmediate::Get());
 }
 
 FBufferRHIRef FSkinWeightLookupVertexBuffer::CreateRHIBuffer_Async()
 {
-	return CreateRHIBuffer_Internal<false>();
+	FRHIAsyncCommandList CommandList;
+	return CreateRHIBuffer(*CommandList);
 }
 
 void FSkinWeightLookupVertexBuffer::InitRHIForStreaming(FRHIBuffer* IntermediateBuffer, FRHIResourceUpdateBatcher& Batcher)
@@ -192,10 +193,9 @@ void FSkinWeightLookupVertexBuffer::SetWeightOffsetAndInfluenceCount(uint32 Vert
 	*DataUInt32 = (WeightOffset << 8) | InfluenceCount;
 }
 
-template <bool bRenderThread>
-FBufferRHIRef FSkinWeightLookupVertexBuffer::CreateRHIBuffer_Internal()
+FBufferRHIRef FSkinWeightLookupVertexBuffer::CreateRHIBuffer(FRHICommandListBase& RHICmdList)
 {
-	return CreateRHIBuffer<bRenderThread>(LookupData, NumVertices, BUF_Static | BUF_ShaderResource | BUF_SourceCopy, TEXT("FSkinWeightLookupVertexBuffer"));
+	return FRenderResource::CreateRHIBuffer(RHICmdList, LookupData, NumVertices, BUF_Static | BUF_ShaderResource | BUF_SourceCopy, TEXT("FSkinWeightLookupVertexBuffer"));
 }
 
 /*-----------------------------------------------------------------------------
@@ -393,21 +393,21 @@ void FSkinWeightDataVertexBuffer::CopyMetaData(const FSkinWeightDataVertexBuffer
 	NumBoneWeights = Other.NumBoneWeights;
 }
 
-template <bool bRenderThread>
-FBufferRHIRef FSkinWeightDataVertexBuffer::CreateRHIBuffer_Internal()
+FBufferRHIRef FSkinWeightDataVertexBuffer::CreateRHIBuffer(FRHICommandListBase& RHICmdList)
 {
 	// BUF_ShaderResource is needed for support of the SkinCache (we could make is dependent on GEnableGPUSkinCacheShaders or are there other users?)
-	return CreateRHIBuffer<bRenderThread>(WeightData, NumBoneWeights, BUF_Static | BUF_ShaderResource | BUF_SourceCopy, TEXT("FSkinWeightDataVertexBuffer"));
+	return FRenderResource::CreateRHIBuffer(RHICmdList, WeightData, NumBoneWeights, BUF_Static | BUF_ShaderResource | BUF_SourceCopy, TEXT("FSkinWeightDataVertexBuffer"));
 }
 
 FBufferRHIRef FSkinWeightDataVertexBuffer::CreateRHIBuffer_RenderThread()
 {
-	return CreateRHIBuffer_Internal<true>();
+	return CreateRHIBuffer(FRHICommandListImmediate::Get());
 }
 
 FBufferRHIRef FSkinWeightDataVertexBuffer::CreateRHIBuffer_Async()
 {
-	return CreateRHIBuffer_Internal<false>();
+	FRHIAsyncCommandList CommandList;
+	return CreateRHIBuffer(*CommandList);
 }
 
 void FSkinWeightDataVertexBuffer::InitRHIForStreaming(FRHIBuffer* IntermediateBuffer, FRHIResourceUpdateBatcher& Batcher)
@@ -436,8 +436,8 @@ void FSkinWeightDataVertexBuffer::InitRHI(FRHICommandListBase& RHICmdList)
 	SCOPED_LOADTIMER(FSkinWeightVertexBuffer_InitRHI);
 
 	// BUF_ShaderResource is needed for support of the SkinCache (we could make is dependent on GEnableGPUSkinCacheShaders or are there other users?)
-	VertexBufferRHI = CreateRHIBuffer_RenderThread();
-	
+	VertexBufferRHI = CreateRHIBuffer(RHICmdList);
+
 	bool bSRV = VertexBufferRHI && GPixelFormats[GetPixelFormat()].Supported;
 	// When bAllowCPUAccess is true, the meshes is likely going to be used for Niagara to spawn particles on mesh surface.
 	// And it can be the case for CPU *and* GPU access: no differenciation today. That is why we create a SRV in this case.
@@ -583,7 +583,7 @@ void FSkinWeightDataVertexBuffer::SetBoneWeight(uint32 VertexWeightOffset, uint3
 		const uint32 BoneWeightOffset = GetBoneIndexByteSize() * VertexInfluenceCount;
 		uint8* BoneData = Data + VertexWeightOffset + BoneWeightOffset;
 		uint8* BoneVertexData = &BoneData[InfluenceIndex * GetBoneWeightByteSize()];  
-		
+
 		if (Use16BitBoneWeight())
 		{
 			*reinterpret_cast<uint16*>(BoneVertexData) = BoneWeight;
@@ -592,7 +592,7 @@ void FSkinWeightDataVertexBuffer::SetBoneWeight(uint32 VertexWeightOffset, uint3
 		{
 			*BoneVertexData = static_cast<uint8>(BoneWeight >> 8);
 		}
-		
+
 	}
 	else
 	{
@@ -703,19 +703,25 @@ void FSkinWeightVertexBuffer::CopyMetaData(const FSkinWeightVertexBuffer& Other)
 	LookupVertexBuffer.CopyMetaData(Other.LookupVertexBuffer);
 }
 
-FSkinWeightRHIInfo FSkinWeightVertexBuffer::CreateRHIBuffer_RenderThread()
+FSkinWeightRHIInfo FSkinWeightVertexBuffer::CreateRHIBuffer(FRHICommandListBase& RHICmdList)
 {
 	FSkinWeightRHIInfo RHIInfo;
-	RHIInfo.DataVertexBufferRHI = DataVertexBuffer.CreateRHIBuffer_RenderThread();
-	RHIInfo.LookupVertexBufferRHI = LookupVertexBuffer.CreateRHIBuffer_RenderThread();
+	RHIInfo.DataVertexBufferRHI = DataVertexBuffer.CreateRHIBuffer(RHICmdList);
+	RHIInfo.LookupVertexBufferRHI = LookupVertexBuffer.CreateRHIBuffer(RHICmdList);
 	return RHIInfo;
 }
 
+FSkinWeightRHIInfo FSkinWeightVertexBuffer::CreateRHIBuffer_RenderThread()
+{
+	return CreateRHIBuffer(FRHICommandListImmediate::Get());
+}
+
 FSkinWeightRHIInfo FSkinWeightVertexBuffer::CreateRHIBuffer_Async()
 {
 	FSkinWeightRHIInfo RHIInfo;
-	RHIInfo.DataVertexBufferRHI = DataVertexBuffer.CreateRHIBuffer_Async();
-	RHIInfo.LookupVertexBufferRHI = LookupVertexBuffer.CreateRHIBuffer_Async();
+	FRHIAsyncCommandList CommandList;
+	RHIInfo.DataVertexBufferRHI = DataVertexBuffer.CreateRHIBuffer(*CommandList);
+	RHIInfo.LookupVertexBufferRHI = LookupVertexBuffer.CreateRHIBuffer(*CommandList);
 	return RHIInfo;
 }
 
diff --git a/Engine/Source/Runtime/Engine/Private/Rendering/StaticMeshVertexBuffer.cpp b/Engine/Source/Runtime/Engine/Private/Rendering/StaticMeshVertexBuffer.cpp
--- a/Engine/Source/Runtime/Engine/Private/Rendering/StaticMeshVertexBuffer.cpp
+++ b/Engine/Source/Runtime/Engine/Private/Rendering/StaticMeshVertexBuffer.cpp
@@ -267,36 +267,36 @@ void FStaticMeshVertexBuffer::operator=(const FStaticMeshVertexBuffer &Other)
 	bUseHighPrecisionTangentBasis = Other.bUseHighPrecisionTangentBasis;
 }
 
-template <bool bRenderThread>
-FBufferRHIRef FStaticMeshVertexBuffer::CreateTangentsRHIBuffer_Internal()
+FBufferRHIRef FStaticMeshVertexBuffer::CreateTangentsRHIBuffer(FRHICommandListBase& RHICmdList)
 {
-	return CreateRHIBuffer<bRenderThread>(TangentsData, GetNumVertices(), BUF_Static | BUF_ShaderResource, TEXT("TangentsRHIBuffer"));
+	return CreateRHIBuffer(RHICmdList, TangentsData, GetNumVertices(), BUF_Static | BUF_ShaderResource, TEXT("TangentsRHIBuffer"));
 }
 
 FBufferRHIRef FStaticMeshVertexBuffer::CreateTangentsRHIBuffer_RenderThread()
 {
-	return CreateTangentsRHIBuffer_Internal<true>();
+	return CreateTangentsRHIBuffer(FRHICommandListImmediate::Get());
 }
 
 FBufferRHIRef FStaticMeshVertexBuffer::CreateTangentsRHIBuffer_Async()
 {
-	return CreateTangentsRHIBuffer_Internal<false>();
+	FRHIAsyncCommandList CommandList;
+	return CreateTangentsRHIBuffer(*CommandList);
 }
 
-template <bool bRenderThread>
-FBufferRHIRef FStaticMeshVertexBuffer::CreateTexCoordRHIBuffer_Internal()
+FBufferRHIRef FStaticMeshVertexBuffer::CreateTexCoordRHIBuffer(FRHICommandListBase& RHICmdList)
 {
-	return CreateRHIBuffer<bRenderThread>(TexcoordData, GetNumTexCoords(), BUF_Static | BUF_ShaderResource, TEXT("TexCoordRHIBuffer"));
+	return CreateRHIBuffer(RHICmdList, TexcoordData, GetNumTexCoords(), BUF_Static | BUF_ShaderResource, TEXT("TexCoordRHIBuffer"));
 }
 
 FBufferRHIRef FStaticMeshVertexBuffer::CreateTexCoordRHIBuffer_RenderThread()
 {
-	return CreateTexCoordRHIBuffer_Internal<true>();
+	return CreateTexCoordRHIBuffer(FRHICommandListImmediate::Get());
 }
 
 FBufferRHIRef FStaticMeshVertexBuffer::CreateTexCoordRHIBuffer_Async()
 {
-	return CreateTexCoordRHIBuffer_Internal<false>();
+	FRHIAsyncCommandList CommandList;
+	return CreateTexCoordRHIBuffer(*CommandList);
 }
 
 void FStaticMeshVertexBuffer::InitRHIForStreaming(
@@ -333,7 +333,7 @@ void FStaticMeshVertexBuffer::InitRHI(FRHICommandListBase& RHICmdList)
 
 	const bool bHadTangentsData = TangentsData != nullptr;
 	const bool bCreateTangentsSRV = bHadTangentsData && TangentsData->GetAllowCPUAccess();
-	TangentsVertexBuffer.VertexBufferRHI = CreateTangentsRHIBuffer_RenderThread();
+	TangentsVertexBuffer.VertexBufferRHI = CreateTangentsRHIBuffer(RHICmdList);
 	if (TangentsVertexBuffer.VertexBufferRHI && (bCreateTangentsSRV || RHISupportsManualVertexFetch(GMaxRHIShaderPlatform) || FLocalVertexFactory::IsGPUSkinPassThroughSupported(GMaxRHIShaderPlatform)))
 	{
 		uint32       Stride = GetUseHighPrecisionTangentBasis() ? 8 : 4;
@@ -343,7 +343,7 @@ void FStaticMeshVertexBuffer::InitRHI(FRHICommandListBase& RHICmdList)
 
 	const bool bHadTexCoordData = TexcoordData != nullptr;
 	const bool bCreateTexCoordSRV = bHadTexCoordData && TexcoordData->GetAllowCPUAccess();
-	TexCoordVertexBuffer.VertexBufferRHI = CreateTexCoordRHIBuffer_RenderThread();
+	TexCoordVertexBuffer.VertexBufferRHI = CreateTexCoordRHIBuffer(RHICmdList);
 	if (TexCoordVertexBuffer.VertexBufferRHI && (bCreateTexCoordSRV || RHISupportsManualVertexFetch(GMaxRHIShaderPlatform)))
 	{
 		uint32       Stride = GetUseFullPrecisionUVs() ? 8 : 4;
diff --git a/Engine/Source/Runtime/Engine/Private/SceneManagement.cpp b/Engine/Source/Runtime/Engine/Private/SceneManagement.cpp
--- a/Engine/Source/Runtime/Engine/Private/SceneManagement.cpp
+++ b/Engine/Source/Runtime/Engine/Private/SceneManagement.cpp
@@ -129,7 +129,7 @@ FFrozenSceneViewMatricesGuard::~FFrozenSceneViewMatricesGuard()
 IMPLEMENT_STATIC_UNIFORM_BUFFER_SLOT(WorkingColorSpace);
 IMPLEMENT_STATIC_UNIFORM_BUFFER_STRUCT(FWorkingColorSpaceShaderParameters, "WorkingColorSpace", WorkingColorSpace);
 
-void FDefaultWorkingColorSpaceUniformBuffer::Update(const UE::Color::FColorSpace& InColorSpace)
+void FDefaultWorkingColorSpaceUniformBuffer::Update(FRHICommandListBase& RHICmdList, const UE::Color::FColorSpace& InColorSpace)
 {
 	using namespace UE::Color;
 
@@ -147,7 +147,7 @@ void FDefaultWorkingColorSpaceUniformBuffer::Update(const UE::Color::FColorSpace
 
 	Parameters.bIsSRGB = InColorSpace.IsSRGB();
 
-	SetContents(Parameters);
+	SetContents(RHICmdList, Parameters);
 }
 
 TGlobalResource<FDefaultWorkingColorSpaceUniformBuffer> GDefaultWorkingColorSpaceUniformBuffer;
@@ -480,8 +480,9 @@ void FDynamicPrimitiveUniformBuffer::Set(
 	bool bOutputVelocity,
 	const FCustomPrimitiveData* CustomPrimitiveData)
 {
-	check(IsInRenderingThread());
+	FRHICommandListBase& RHICmdList = FRHICommandListImmediate::Get();
 	UniformBuffer.SetContents(
+		RHICmdList,
 		FPrimitiveUniformShaderParametersBuilder{}
 		.Defaults()
 			.LocalToWorld(LocalToWorld)
@@ -496,7 +497,7 @@ void FDynamicPrimitiveUniformBuffer::Set(
 			.CustomPrimitiveData(CustomPrimitiveData)
 		.Build()
 	);
-	UniformBuffer.InitResource(FRHICommandListImmediate::Get());
+	UniformBuffer.InitResource(RHICmdList);
 }
 
 void FDynamicPrimitiveUniformBuffer::Set(
@@ -1265,7 +1266,7 @@ void FDefaultLightmapResourceClusterUniformBuffer::InitRHI(FRHICommandListBase&
 {
 	FLightmapResourceClusterShaderParameters Parameters;
 	GetLightmapClusterResourceParameters(GMaxRHIFeatureLevel, FLightmapClusterResourceInput(), nullptr, Parameters);
-	SetContents(Parameters);
+	SetContentsNoUpdate(Parameters);
 	Super::InitRHI(RHICmdList);
 }
 
@@ -1547,7 +1548,7 @@ void FDefaultMobileReflectionCaptureUniformBuffer::InitRHI(FRHICommandListBase&
 	Parameters.Params = FVector4f(1.f, 0.f, 0.f, 0.f);
 	Parameters.Texture = GBlackTextureCube->TextureRHI;
 	Parameters.TextureSampler = GBlackTextureCube->SamplerStateRHI;
-	SetContents(Parameters);
+	SetContentsNoUpdate(Parameters);
 	Super::InitRHI(RHICmdList);
 }
 
diff --git a/Engine/Source/Runtime/Engine/Private/SkeletalMesh.cpp b/Engine/Source/Runtime/Engine/Private/SkeletalMesh.cpp
--- a/Engine/Source/Runtime/Engine/Private/SkeletalMesh.cpp
+++ b/Engine/Source/Runtime/Engine/Private/SkeletalMesh.cpp
@@ -198,31 +198,6 @@ FArchive& operator<<(FArchive& Ar, FClothBufferIndexMapping& ClothBufferIndexMap
 		<< ClothBufferIndexMapping.LODBiasStride;
 }
 
-void FreeSkeletalMeshBuffersSinkCallback()
-{
-	static const auto CVar = IConsoleManager::Get().FindTConsoleVariableDataInt(TEXT("r.FreeSkeletalMeshBuffers"));
-	static bool bPreviousFreeSkeletalMeshBuffers = false;
-	const bool bFreeSkeletalMeshBuffers = CVar->GetValueOnGameThread() == 1;
-
-	// We previously relied on this sink occurring periodically to free helper structures, so run it if we're not using the new logic that frees on submit.
-	if (!GFreeStructuresOnRHIBufferCreation || (bPreviousFreeSkeletalMeshBuffers != bFreeSkeletalMeshBuffers))
-	{
-		bPreviousFreeSkeletalMeshBuffers = bFreeSkeletalMeshBuffers;
-		if (bFreeSkeletalMeshBuffers)
-		{
-			FlushRenderingCommands();
-			for (TObjectIterator<USkeletalMesh> It; It; ++It)
-			{
-				USkeletalMesh* SkeletalMesh = *It;
-				if (!SkeletalMesh->HasPendingInitOrStreaming() && !SkeletalMesh->GetResourceForRendering()->RequiresCPUSkinning(GMaxRHIFeatureLevel))
-				{
-					SkeletalMesh->ReleaseCPUResources();
-				}
-			}
-		}
-	}
-}
-
 /*-----------------------------------------------------------------------------
 	FClothingAssetData
 -----------------------------------------------------------------------------*/
diff --git a/Engine/Source/Runtime/Engine/Private/SkeletalMeshDeformerHelpers.cpp b/Engine/Source/Runtime/Engine/Private/SkeletalMeshDeformerHelpers.cpp
--- a/Engine/Source/Runtime/Engine/Private/SkeletalMeshDeformerHelpers.cpp
+++ b/Engine/Source/Runtime/Engine/Private/SkeletalMeshDeformerHelpers.cpp
@@ -129,7 +129,7 @@ FRDGBuffer* FSkeletalMeshDeformerHelpers::AllocateVertexFactoryPositionBuffer(FR
 			VertexBuffers[SectionIndex] = DeformerGeometry.Position->GetRHI();
 		}
 
-		MeshObjectGPU->UpdateRayTracingGeometry(LODModel, InLodIndex, VertexBuffers);
+		MeshObjectGPU->UpdateRayTracingGeometry(GraphBuilder.RHICmdList, LODModel, InLodIndex, VertexBuffers);
 #endif // RHI_RAYTRACING
 	}
 
@@ -254,13 +254,13 @@ void FSkeletalMeshDeformerHelpers::UpdateVertexFactoryBufferOverrides(FSkeletalM
 
 void FSkeletalMeshDeformerHelpers::ResetVertexFactoryBufferOverrides(FSkeletalMeshObject* InMeshObject, int32 LODIndex)
 {
-	check(IsInRenderingThread());
-
 	if (InMeshObject->IsCPUSkinned())
 	{
 		return;
 	}
 
+	FRHICommandListBase& RHICmdList = FRHICommandListImmediate::Get();
+
 	FSkeletalMeshObjectGPUSkin* MeshObjectGPU = static_cast<FSkeletalMeshObjectGPUSkin*>(InMeshObject);
 	FMeshDeformerGeometry& DeformerGeometry = MeshObjectGPU->GetDeformerGeometry(LODIndex);
 
@@ -284,7 +284,7 @@ void FSkeletalMeshDeformerHelpers::ResetVertexFactoryBufferOverrides(FSkeletalMe
 		TargetVertexFactory->ResetVertexAttributes();
 		FGPUSkinPassthroughVertexFactory::FDataType Data;
 		BaseVertexFactory->CopyDataTypeForLocalVertexFactory(Data);
-		TargetVertexFactory->SetData(Data);
+		TargetVertexFactory->SetData(RHICmdList, Data);
 	}
 
 #if RHI_RAYTRACING
@@ -295,6 +295,6 @@ void FSkeletalMeshDeformerHelpers::ResetVertexFactoryBufferOverrides(FSkeletalMe
 
 	TArray<FBufferRHIRef> VertexBuffers;
 	VertexBuffers.Init(VertexBuffer, NumSections);
-	MeshObjectGPU->UpdateRayTracingGeometry(LODModel, LODIndex, VertexBuffers);
+	MeshObjectGPU->UpdateRayTracingGeometry(FRHICommandList::Get(RHICmdList), LODModel, LODIndex, VertexBuffers);
 #endif // RHI_RAYTRACING
 }
diff --git a/Engine/Source/Runtime/Engine/Private/SkeletalRenderCPUSkin.cpp b/Engine/Source/Runtime/Engine/Private/SkeletalRenderCPUSkin.cpp
--- a/Engine/Source/Runtime/Engine/Private/SkeletalRenderCPUSkin.cpp
+++ b/Engine/Source/Runtime/Engine/Private/SkeletalRenderCPUSkin.cpp
@@ -312,7 +312,7 @@ void FSkeletalMeshObjectCPUSkin::CacheVertices(int32 LODIndex, bool bForce) cons
 			MeshLODptr->StaticMeshVertexBuffer.BindLightMapVertexBuffer(VertexFactoryPtr, Data, 0);
 			MeshLODptr->MeshObjectColorBuffer->BindColorVertexBuffer(VertexFactoryPtr, Data);
 
-			VertexFactoryPtr->SetData(Data);
+			VertexFactoryPtr->SetData(RHICmdList, Data);
 			VertexFactoryPtr->InitResource(RHICmdList);
 		});
 	}
@@ -364,7 +364,7 @@ void FSkeletalMeshObjectCPUSkin::FSkeletalMeshObjectLOD::InitResources(FSkelMesh
 			Self->StaticMeshVertexBuffer.BindLightMapVertexBuffer(VertexFactoryPtr, Data, 0);
 			Self->MeshObjectColorBuffer->BindColorVertexBuffer(VertexFactoryPtr, Data);
 
-			VertexFactoryPtr->SetData(Data);
+			VertexFactoryPtr->SetData(RHICmdList, Data);
 			VertexFactoryPtr->InitResource(RHICmdList);
 		});
 
diff --git a/Engine/Source/Runtime/Engine/Private/SkeletalRenderGPUSkin.cpp b/Engine/Source/Runtime/Engine/Private/SkeletalRenderGPUSkin.cpp
--- a/Engine/Source/Runtime/Engine/Private/SkeletalRenderGPUSkin.cpp
+++ b/Engine/Source/Runtime/Engine/Private/SkeletalRenderGPUSkin.cpp
@@ -773,10 +773,8 @@ void FSkeletalMeshObjectGPUSkin::ProcessUpdatedDynamicData(EGPUSkinCacheEntryMod
 
 #if RHI_RAYTRACING
 
-void FSkeletalMeshObjectGPUSkin::UpdateRayTracingGeometry(FSkeletalMeshLODRenderData& LODModel, uint32 LODIndex, TArray<FBufferRHIRef>& VertexBufffers)
+void FSkeletalMeshObjectGPUSkin::UpdateRayTracingGeometry(FRHICommandList& RHICmdList, FSkeletalMeshLODRenderData& LODModel, uint32 LODIndex, TArray<FBufferRHIRef>& VertexBufffers)
 {
-	FRHICommandListBase& RHICmdList = FRHICommandListImmediate::Get();
-
 	if (IsRayTracingEnabled() && bSupportRayTracing)
 	{
 		const bool bAnySegmentUsesWorldPositionOffset = DynamicData != nullptr ? DynamicData->bAnySegmentUsesWorldPositionOffset : false;
@@ -1863,7 +1861,7 @@ static void CreatePassthroughVertexFactory(ERHIFeatureLevel::Type InFeatureLevel
 		{
 			FLocalVertexFactory::FDataType Data;
 			SourceVertexFactory->CopyDataTypeForLocalVertexFactory(Data);
-			NewPassthroughVertexFactory->SetData(Data);
+			NewPassthroughVertexFactory->SetData(RHICmdList, Data);
 			NewPassthroughVertexFactory->InitResource(RHICmdList);
 		}
 	);
diff --git a/Engine/Source/Runtime/Engine/Private/SkeletalRenderGPUSkin.h b/Engine/Source/Runtime/Engine/Private/SkeletalRenderGPUSkin.h
--- a/Engine/Source/Runtime/Engine/Private/SkeletalRenderGPUSkin.h
+++ b/Engine/Source/Runtime/Engine/Private/SkeletalRenderGPUSkin.h
@@ -395,7 +395,7 @@ public:
 	 * VSinCS path is still required for world position offset materials but this can still use 
 	 * the updated vertex buffers from here with a passthrough vertex factory.
 	 */
-	ENGINE_API void UpdateRayTracingGeometry(FSkeletalMeshLODRenderData& LODModel, uint32 LODIndex, TArray<FBufferRHIRef>& VertexBufffers);
+	ENGINE_API void UpdateRayTracingGeometry(FRHICommandList& RHICmdList, FSkeletalMeshLODRenderData& LODModel, uint32 LODIndex, TArray<FBufferRHIRef>& VertexBufffers);
 
 #endif // RHI_RAYTRACING
 
diff --git a/Engine/Source/Runtime/Engine/Private/SkeletalRenderStatic.cpp b/Engine/Source/Runtime/Engine/Private/SkeletalRenderStatic.cpp
--- a/Engine/Source/Runtime/Engine/Private/SkeletalRenderStatic.cpp
+++ b/Engine/Source/Runtime/Engine/Private/SkeletalRenderStatic.cpp
@@ -72,7 +72,7 @@ void FSkeletalMeshObjectStatic::InitResources(USkinnedMeshComponent* InMeshCompo
 						LODIndex = LODIndex, 
 						SkelMeshRenderData = SkelLOD.SkelMeshRenderData, 
 						&RayTracingGeometry = SkelLOD.SkelMeshRenderData->LODRenderData[LODIndex].StaticRayTracingGeometry,
-						&bReferencedByStaticSkeletalMeshObjects_RenderThread = SkelLOD.SkelMeshRenderData->LODRenderData[LODIndex].bReferencedByStaticSkeletalMeshObjects_RenderThread](FRHICommandListImmediate& RHICmdList)
+						&bReferencedByStaticSkeletalMeshObjects_RenderThread = SkelLOD.SkelMeshRenderData->LODRenderData[LODIndex].bReferencedByStaticSkeletalMeshObjects_RenderThread](FRHICommandList& RHICmdList)
 						{
 							FRayTracingGeometryInitializer Initializer;
 							static const FName DebugName("FSkeletalMeshObjectLOD");
@@ -205,7 +205,7 @@ void FSkeletalMeshObjectStatic::FSkeletalMeshObjectLOD::InitResources(FSkelMeshC
 	FColorVertexBuffer* ColorVertexBufferPtr = ColorVertexBuffer;
 
 	ENQUEUE_RENDER_COMMAND(InitSkeletalMeshStaticSkinVertexFactory)(
-		[VertexFactoryPtr, PositionVertexBufferPtr, StaticMeshVertexBufferPtr, ColorVertexBufferPtr](FRHICommandListImmediate& RHICmdList)
+		[VertexFactoryPtr, PositionVertexBufferPtr, StaticMeshVertexBufferPtr, ColorVertexBufferPtr](FRHICommandList& RHICmdList)
 		{
 			FLocalVertexFactory::FDataType Data;
 			PositionVertexBufferPtr->InitResource(RHICmdList);
@@ -218,7 +218,7 @@ void FSkeletalMeshObjectStatic::FSkeletalMeshObjectLOD::InitResources(FSkelMeshC
 			StaticMeshVertexBufferPtr->BindLightMapVertexBuffer(VertexFactoryPtr, Data, 0);
 			ColorVertexBufferPtr->BindColorVertexBuffer(VertexFactoryPtr, Data);
 
-			VertexFactoryPtr->SetData(Data);
+			VertexFactoryPtr->SetData(RHICmdList, Data);
 			VertexFactoryPtr->InitResource(RHICmdList);
 		});
 
diff --git a/Engine/Source/Runtime/Engine/Private/StaticMesh.cpp b/Engine/Source/Runtime/Engine/Private/StaticMesh.cpp
--- a/Engine/Source/Runtime/Engine/Private/StaticMesh.cpp
+++ b/Engine/Source/Runtime/Engine/Private/StaticMesh.cpp
@@ -979,7 +979,7 @@ void FStaticMeshVertexFactories::InitVertexFactory(
 			Data.bIsCoarseProxy			= Params.bIsCoarseProxy;
 			Data.StaticMesh				= Params.StaticMesh;
 		#endif
-			Params.VertexFactory->SetData(Data);
+			Params.VertexFactory->SetData(RHICmdList, Data);
 			Params.VertexFactory->InitResource(RHICmdList);
 		});
 }
@@ -1087,10 +1087,8 @@ float FStaticMeshAreaWeightedSectionSampler::GetWeights(TArray<float>& OutWeight
 	return Total;
 }
 
-static inline void InitOrUpdateResource(FRenderResource* Resource)
+static inline void InitOrUpdateResource(FRHICommandListBase& RHICmdList, FRenderResource* Resource)
 {
-	FRHICommandListImmediate& RHICmdList = FRHICommandListImmediate::Get();
-
 	if (!Resource->IsInitialized())
 	{
 		Resource->InitResource(RHICmdList);
@@ -1146,9 +1144,9 @@ void FStaticMeshVertexBuffers::InitModelVF(FLocalVertexFactory* VertexFactory)
 		Self->StaticMeshVertexBuffer.BindPackedTexCoordVertexBuffer(VertexFactory, Data);
 		Self->StaticMeshVertexBuffer.BindLightMapVertexBuffer(VertexFactory, Data, 1);
 		FColorVertexBuffer::BindDefaultColorVertexBuffer(VertexFactory, Data, FColorVertexBuffer::NullBindStride::ZeroForDefaultBufferBind);
-		VertexFactory->SetData(Data);
+		VertexFactory->SetData(RHICmdList, Data);
 
-		InitOrUpdateResource(VertexFactory);
+		InitOrUpdateResource(RHICmdList, VertexFactory);
 	});
 }
 
@@ -1166,9 +1164,9 @@ void FStaticMeshVertexBuffers::InitWithDummyData(FLocalVertexFactory* VertexFact
 	ENQUEUE_RENDER_COMMAND(StaticMeshVertexBuffersLegacyInit)(
 		[VertexFactory, Self, LightMapIndex](FRHICommandListImmediate& RHICmdList)
 	{
-		InitOrUpdateResource(&Self->PositionVertexBuffer);
-		InitOrUpdateResource(&Self->StaticMeshVertexBuffer);
-		InitOrUpdateResource(&Self->ColorVertexBuffer);
+		InitOrUpdateResource(RHICmdList, &Self->PositionVertexBuffer);
+		InitOrUpdateResource(RHICmdList, &Self->StaticMeshVertexBuffer);
+		InitOrUpdateResource(RHICmdList, &Self->ColorVertexBuffer);
 
 		FLocalVertexFactory::FDataType Data;
 		Self->PositionVertexBuffer.BindPositionVertexBuffer(VertexFactory, Data);
@@ -1176,9 +1174,9 @@ void FStaticMeshVertexBuffers::InitWithDummyData(FLocalVertexFactory* VertexFact
 		Self->StaticMeshVertexBuffer.BindPackedTexCoordVertexBuffer(VertexFactory, Data);
 		Self->StaticMeshVertexBuffer.BindLightMapVertexBuffer(VertexFactory, Data, LightMapIndex);
 		Self->ColorVertexBuffer.BindColorVertexBuffer(VertexFactory, Data);
-		VertexFactory->SetData(Data);
+		VertexFactory->SetData(RHICmdList, Data);
 
-		InitOrUpdateResource(VertexFactory);
+		InitOrUpdateResource(RHICmdList, VertexFactory);
 	});
 }
 
@@ -1224,9 +1222,9 @@ void FStaticMeshVertexBuffers::InitFromDynamicVertex(FLocalVertexFactory* Vertex
 	ENQUEUE_RENDER_COMMAND(StaticMeshVertexBuffersLegacyInit)(
 		[VertexFactory, Self, LightMapIndex](FRHICommandListImmediate& RHICmdList)
 		{
-			InitOrUpdateResource(&Self->PositionVertexBuffer);
-			InitOrUpdateResource(&Self->StaticMeshVertexBuffer);
-			InitOrUpdateResource(&Self->ColorVertexBuffer);
+			InitOrUpdateResource(RHICmdList, &Self->PositionVertexBuffer);
+			InitOrUpdateResource(RHICmdList, &Self->StaticMeshVertexBuffer);
+			InitOrUpdateResource(RHICmdList, &Self->ColorVertexBuffer);
 
 			FLocalVertexFactory::FDataType Data;
 			Self->PositionVertexBuffer.BindPositionVertexBuffer(VertexFactory, Data);
@@ -1234,9 +1232,9 @@ void FStaticMeshVertexBuffers::InitFromDynamicVertex(FLocalVertexFactory* Vertex
 			Self->StaticMeshVertexBuffer.BindPackedTexCoordVertexBuffer(VertexFactory, Data);
 			Self->StaticMeshVertexBuffer.BindLightMapVertexBuffer(VertexFactory, Data, LightMapIndex);
 			Self->ColorVertexBuffer.BindColorVertexBuffer(VertexFactory, Data);
-			VertexFactory->SetData(Data);
+			VertexFactory->SetData(RHICmdList, Data);
 
-			InitOrUpdateResource(VertexFactory);
+			InitOrUpdateResource(RHICmdList, VertexFactory);
 		});
 };
 
diff --git a/Engine/Source/Runtime/Engine/Private/Streaming/SkeletalMeshUpdate.cpp b/Engine/Source/Runtime/Engine/Private/Streaming/SkeletalMeshUpdate.cpp
--- a/Engine/Source/Runtime/Engine/Private/Streaming/SkeletalMeshUpdate.cpp
+++ b/Engine/Source/Runtime/Engine/Private/Streaming/SkeletalMeshUpdate.cpp
@@ -58,28 +58,32 @@ FSkeletalMeshUpdate::FSkeletalMeshUpdate(const USkeletalMesh* InMesh)
 
 void FSkeletalMeshStreamIn::FIntermediateBuffers::CreateFromCPUData_RenderThread(FSkeletalMeshLODRenderData& LODResource)
 {
+	FRHICommandListBase& RHICmdList = FRHICommandListImmediate::Get();
+
 	FStaticMeshVertexBuffers& VBs = LODResource.StaticVertexBuffers;
-	TangentsVertexBuffer = VBs.StaticMeshVertexBuffer.CreateTangentsRHIBuffer_RenderThread();
-	TexCoordVertexBuffer = VBs.StaticMeshVertexBuffer.CreateTexCoordRHIBuffer_RenderThread();
-	PositionVertexBuffer = VBs.PositionVertexBuffer.CreateRHIBuffer_RenderThread();
-	ColorVertexBuffer = VBs.ColorVertexBuffer.CreateRHIBuffer_RenderThread();
-	LODResource.SkinWeightProfilesData.CreateRHIBuffers_RenderThread(AltSkinWeightVertexBuffers);
-	SkinWeightVertexBuffer = LODResource.SkinWeightVertexBuffer.CreateRHIBuffer_RenderThread();
-	ClothVertexBuffer = LODResource.ClothVertexBuffer.CreateRHIBuffer_RenderThread();
-	IndexBuffer = LODResource.MultiSizeIndexContainer.CreateRHIBuffer_RenderThread();
+	TangentsVertexBuffer = VBs.StaticMeshVertexBuffer.CreateTangentsRHIBuffer(RHICmdList);
+	TexCoordVertexBuffer = VBs.StaticMeshVertexBuffer.CreateTexCoordRHIBuffer(RHICmdList);
+	PositionVertexBuffer = VBs.PositionVertexBuffer.CreateRHIBuffer(RHICmdList);
+	ColorVertexBuffer = VBs.ColorVertexBuffer.CreateRHIBuffer(RHICmdList);
+	LODResource.SkinWeightProfilesData.CreateRHIBuffers(RHICmdList, AltSkinWeightVertexBuffers);
+	SkinWeightVertexBuffer = LODResource.SkinWeightVertexBuffer.CreateRHIBuffer(RHICmdList);
+	ClothVertexBuffer = LODResource.ClothVertexBuffer.CreateRHIBuffer(RHICmdList);
+	IndexBuffer = LODResource.MultiSizeIndexContainer.CreateRHIBuffer(RHICmdList);
 }
 
 void FSkeletalMeshStreamIn::FIntermediateBuffers::CreateFromCPUData_Async(FSkeletalMeshLODRenderData& LODResource)
 {
+	FRHIAsyncCommandList RHICmdList;
+
 	FStaticMeshVertexBuffers& VBs = LODResource.StaticVertexBuffers;
-	TangentsVertexBuffer = VBs.StaticMeshVertexBuffer.CreateTangentsRHIBuffer_Async();
-	TexCoordVertexBuffer = VBs.StaticMeshVertexBuffer.CreateTexCoordRHIBuffer_Async();
-	PositionVertexBuffer = VBs.PositionVertexBuffer.CreateRHIBuffer_Async();
-	ColorVertexBuffer = VBs.ColorVertexBuffer.CreateRHIBuffer_Async();
-	LODResource.SkinWeightProfilesData.CreateRHIBuffers_Async(AltSkinWeightVertexBuffers);
-	SkinWeightVertexBuffer = LODResource.SkinWeightVertexBuffer.CreateRHIBuffer_Async();
-	ClothVertexBuffer = LODResource.ClothVertexBuffer.CreateRHIBuffer_Async();
-	IndexBuffer = LODResource.MultiSizeIndexContainer.CreateRHIBuffer_Async();
+	TangentsVertexBuffer = VBs.StaticMeshVertexBuffer.CreateTangentsRHIBuffer(*RHICmdList);
+	TexCoordVertexBuffer = VBs.StaticMeshVertexBuffer.CreateTexCoordRHIBuffer(*RHICmdList);
+	PositionVertexBuffer = VBs.PositionVertexBuffer.CreateRHIBuffer(*RHICmdList);
+	ColorVertexBuffer = VBs.ColorVertexBuffer.CreateRHIBuffer(*RHICmdList);
+	LODResource.SkinWeightProfilesData.CreateRHIBuffers(*RHICmdList, AltSkinWeightVertexBuffers);
+	SkinWeightVertexBuffer = LODResource.SkinWeightVertexBuffer.CreateRHIBuffer(*RHICmdList);
+	ClothVertexBuffer = LODResource.ClothVertexBuffer.CreateRHIBuffer(*RHICmdList);
+	IndexBuffer = LODResource.MultiSizeIndexContainer.CreateRHIBuffer(*RHICmdList);
 }
 
 void FSkeletalMeshStreamIn::FIntermediateBuffers::SafeRelease()
diff --git a/Engine/Source/Runtime/Engine/Private/Streaming/StaticMeshUpdate.cpp b/Engine/Source/Runtime/Engine/Private/Streaming/StaticMeshUpdate.cpp
--- a/Engine/Source/Runtime/Engine/Private/Streaming/StaticMeshUpdate.cpp
+++ b/Engine/Source/Runtime/Engine/Private/Streaming/StaticMeshUpdate.cpp
@@ -64,38 +64,42 @@ FStaticMeshUpdate::FStaticMeshUpdate(const UStaticMesh* InMesh)
 
 void FStaticMeshStreamIn::FIntermediateBuffers::CreateFromCPUData_RenderThread(FStaticMeshLODResources& LODResource)
 {
+	FRHICommandListBase& RHICmdList = FRHICommandListImmediate::Get();
+
 	FStaticMeshVertexBuffers& VBs = LODResource.VertexBuffers;
-	TangentsVertexBuffer = VBs.StaticMeshVertexBuffer.CreateTangentsRHIBuffer_RenderThread();
-	TexCoordVertexBuffer = VBs.StaticMeshVertexBuffer.CreateTexCoordRHIBuffer_RenderThread();
-	PositionVertexBuffer = VBs.PositionVertexBuffer.CreateRHIBuffer_RenderThread();
-	ColorVertexBuffer = VBs.ColorVertexBuffer.CreateRHIBuffer_RenderThread();
-	IndexBuffer = LODResource.IndexBuffer.CreateRHIBuffer_RenderThread();
-	DepthOnlyIndexBuffer = LODResource.DepthOnlyIndexBuffer.CreateRHIBuffer_RenderThread();
+	TangentsVertexBuffer = VBs.StaticMeshVertexBuffer.CreateTangentsRHIBuffer(RHICmdList);
+	TexCoordVertexBuffer = VBs.StaticMeshVertexBuffer.CreateTexCoordRHIBuffer(RHICmdList);
+	PositionVertexBuffer = VBs.PositionVertexBuffer.CreateRHIBuffer(RHICmdList);
+	ColorVertexBuffer = VBs.ColorVertexBuffer.CreateRHIBuffer(RHICmdList);
+	IndexBuffer = LODResource.IndexBuffer.CreateRHIBuffer(RHICmdList);
+	DepthOnlyIndexBuffer = LODResource.DepthOnlyIndexBuffer.CreateRHIBuffer(RHICmdList);
 
 
 	if (LODResource.AdditionalIndexBuffers)
 	{
-		ReversedIndexBuffer = LODResource.AdditionalIndexBuffers->ReversedIndexBuffer.CreateRHIBuffer_RenderThread();
-		ReversedDepthOnlyIndexBuffer = LODResource.AdditionalIndexBuffers->ReversedDepthOnlyIndexBuffer.CreateRHIBuffer_RenderThread();
-		WireframeIndexBuffer = LODResource.AdditionalIndexBuffers->WireframeIndexBuffer.CreateRHIBuffer_RenderThread();
+		ReversedIndexBuffer = LODResource.AdditionalIndexBuffers->ReversedIndexBuffer.CreateRHIBuffer(RHICmdList);
+		ReversedDepthOnlyIndexBuffer = LODResource.AdditionalIndexBuffers->ReversedDepthOnlyIndexBuffer.CreateRHIBuffer(RHICmdList);
+		WireframeIndexBuffer = LODResource.AdditionalIndexBuffers->WireframeIndexBuffer.CreateRHIBuffer(RHICmdList);
 	}
 }
 
 void FStaticMeshStreamIn::FIntermediateBuffers::CreateFromCPUData_Async(FStaticMeshLODResources& LODResource)
 {
+	FRHIAsyncCommandList RHICmdList;
+
 	FStaticMeshVertexBuffers& VBs = LODResource.VertexBuffers;
-	TangentsVertexBuffer = VBs.StaticMeshVertexBuffer.CreateTangentsRHIBuffer_Async();
-	TexCoordVertexBuffer = VBs.StaticMeshVertexBuffer.CreateTexCoordRHIBuffer_Async();
-	PositionVertexBuffer = VBs.PositionVertexBuffer.CreateRHIBuffer_Async();
-	ColorVertexBuffer = VBs.ColorVertexBuffer.CreateRHIBuffer_Async();
-	IndexBuffer = LODResource.IndexBuffer.CreateRHIBuffer_Async();
-	DepthOnlyIndexBuffer = LODResource.DepthOnlyIndexBuffer.CreateRHIBuffer_Async();
+	TangentsVertexBuffer = VBs.StaticMeshVertexBuffer.CreateTangentsRHIBuffer(*RHICmdList);
+	TexCoordVertexBuffer = VBs.StaticMeshVertexBuffer.CreateTexCoordRHIBuffer(*RHICmdList);
+	PositionVertexBuffer = VBs.PositionVertexBuffer.CreateRHIBuffer(*RHICmdList);
+	ColorVertexBuffer = VBs.ColorVertexBuffer.CreateRHIBuffer(*RHICmdList);
+	IndexBuffer = LODResource.IndexBuffer.CreateRHIBuffer(*RHICmdList);
+	DepthOnlyIndexBuffer = LODResource.DepthOnlyIndexBuffer.CreateRHIBuffer(*RHICmdList);
 
 	if (LODResource.AdditionalIndexBuffers)
 	{
-		ReversedIndexBuffer = LODResource.AdditionalIndexBuffers->ReversedIndexBuffer.CreateRHIBuffer_Async();
-		ReversedDepthOnlyIndexBuffer = LODResource.AdditionalIndexBuffers->ReversedDepthOnlyIndexBuffer.CreateRHIBuffer_Async();
-		WireframeIndexBuffer = LODResource.AdditionalIndexBuffers->WireframeIndexBuffer.CreateRHIBuffer_Async();
+		ReversedIndexBuffer = LODResource.AdditionalIndexBuffers->ReversedIndexBuffer.CreateRHIBuffer(*RHICmdList);
+		ReversedDepthOnlyIndexBuffer = LODResource.AdditionalIndexBuffers->ReversedDepthOnlyIndexBuffer.CreateRHIBuffer(*RHICmdList);
+		WireframeIndexBuffer = LODResource.AdditionalIndexBuffers->WireframeIndexBuffer.CreateRHIBuffer(*RHICmdList);
 	}
 }
 
@@ -204,7 +208,7 @@ void FStaticMeshStreamIn::CreateBuffers_Internal(const FContext& Context)
 				}
 				Initializer.Type = ERayTracingGeometryInitializerType::StreamingSource;
 				IntermediateRayTracingGeometry[LODIdx].SetInitializer(Initializer);
-				IntermediateRayTracingGeometry[LODIdx].CreateRayTracingGeometryFromCPUData(LODResource.RayTracingGeometry.RawData);
+				IntermediateRayTracingGeometry[LODIdx].CreateRayTracingGeometryFromCPUData(FRHICommandListImmediate::Get(), LODResource.RayTracingGeometry.RawData);
 			}
 #endif
 		}
@@ -297,7 +301,7 @@ void FStaticMeshStreamIn::DoFinishUpdate(const FContext& Context)
 					// Under very rare circumstances that we switch ray tracing on/off right in the middle of streaming RayTracingGeometryRHI might not be valid.
 					if (IsRayTracingEnabled() && LODResource.RayTracingGeometry.RayTracingGeometryRHI.IsValid())
 					{
-						LODResource.RayTracingGeometry.RequestBuildIfNeeded(ERTAccelerationStructureBuildPriority::Normal);
+						LODResource.RayTracingGeometry.RequestBuildIfNeeded(FRHICommandListImmediate::Get(), ERTAccelerationStructureBuildPriority::Normal);
 					}
 				}
 			}
diff --git a/Engine/Source/Runtime/Engine/Private/TileRendering.cpp b/Engine/Source/Runtime/Engine/Private/TileRendering.cpp
--- a/Engine/Source/Runtime/Engine/Private/TileRendering.cpp
+++ b/Engine/Source/Runtime/Engine/Private/TileRendering.cpp
@@ -35,7 +35,7 @@ void FCanvasTileRendererItem::FTileVertexFactory::InitResource(FRHICommandListBa
 	VertexBuffers->StaticMeshVertexBuffer.BindPackedTexCoordVertexBuffer(this, VertexData);
 	VertexBuffers->StaticMeshVertexBuffer.BindLightMapVertexBuffer(this, VertexData, 0);
 	VertexBuffers->ColorVertexBuffer.BindColorVertexBuffer(this, VertexData);
-	SetData(VertexData);
+	SetData(RHICmdList, VertexData);
 
 	FLocalVertexFactory::InitResource(RHICmdList);
 }
diff --git a/Engine/Source/Runtime/Engine/Private/TriangleRendering.cpp b/Engine/Source/Runtime/Engine/Private/TriangleRendering.cpp
--- a/Engine/Source/Runtime/Engine/Private/TriangleRendering.cpp
+++ b/Engine/Source/Runtime/Engine/Private/TriangleRendering.cpp
@@ -31,7 +31,7 @@ void FCanvasTriangleRendererItem::FTriangleVertexFactory::InitResource(FRHIComma
 	VertexBuffers->StaticMeshVertexBuffer.BindPackedTexCoordVertexBuffer(this, VertexData);
 	VertexBuffers->StaticMeshVertexBuffer.BindLightMapVertexBuffer(this, VertexData, 0);
 	VertexBuffers->ColorVertexBuffer.BindColorVertexBuffer(this, VertexData);
-	SetData(VertexData);
+	SetData(RHICmdList, VertexData);
 
 	FLocalVertexFactory::InitResource(RHICmdList);
 }
diff --git a/Engine/Source/Runtime/Engine/Private/UnrealEngine.cpp b/Engine/Source/Runtime/Engine/Private/UnrealEngine.cpp
--- a/Engine/Source/Runtime/Engine/Private/UnrealEngine.cpp
+++ b/Engine/Source/Runtime/Engine/Private/UnrealEngine.cpp
@@ -1072,12 +1072,9 @@ void RefreshSamplerStatesCallback()
 
 void RefreshEngineSettings()
 {
-	extern void FreeSkeletalMeshBuffersSinkCallback();
-
 	HDRSettingChangedSinkCallback();
 	RefreshSamplerStatesCallback();
 	ScalabilityCVarsSinkCallback();
-	FreeSkeletalMeshBuffersSinkCallback();
 	SystemResolutionSinkCallback();
 }
 
diff --git a/Engine/Source/Runtime/Engine/Public/Animation/SkinWeightProfile.h b/Engine/Source/Runtime/Engine/Public/Animation/SkinWeightProfile.h
--- a/Engine/Source/Runtime/Engine/Public/Animation/SkinWeightProfile.h
+++ b/Engine/Source/Runtime/Engine/Public/Animation/SkinWeightProfile.h
@@ -163,7 +163,11 @@ struct FSkinWeightProfilesData
 
 	ENGINE_API void ReleaseCPUResources();
 
+	ENGINE_API void CreateRHIBuffers(FRHICommandListBase& RHICmdList, TArray<TPair<FName, FSkinWeightRHIInfo>>& OutBuffers);
+
+	UE_DEPRECATED(5.4, "Use CreateRHIBuffers instead.")
 	ENGINE_API void CreateRHIBuffers_RenderThread(TArray<TPair<FName, FSkinWeightRHIInfo>>& OutBuffers);
+	UE_DEPRECATED(5.4, "Use CreateRHIBuffers instead.")
 	ENGINE_API void CreateRHIBuffers_Async(TArray<TPair<FName, FSkinWeightRHIInfo>>& OutBuffers);
 
 	ENGINE_API void InitRHIForStreaming(const TArray<TPair<FName, FSkinWeightRHIInfo>>& IntermediateBuffers, FRHIResourceUpdateBatcher& Batcher);
@@ -181,9 +185,6 @@ struct FSkinWeightProfilesData
 protected:
 	ENGINE_API void ApplyOverrideProfile(FSkinWeightVertexBuffer* OverrideBuffer, const FName& ProfileName);
 
-	template <bool bRenderThread>
-	void CreateRHIBuffers_Internal(TArray<TPair<FName, FSkinWeightRHIInfo>>& OutBuffers);
-
 	FSkinWeightVertexBuffer* BaseBuffer;
 	FSkinWeightVertexBuffer* DefaultOverrideSkinWeightBuffer;
 
diff --git a/Engine/Source/Runtime/Engine/Public/GPUSkinVertexFactory.h b/Engine/Source/Runtime/Engine/Public/GPUSkinVertexFactory.h
--- a/Engine/Source/Runtime/Engine/Public/GPUSkinVertexFactory.h
+++ b/Engine/Source/Runtime/Engine/Public/GPUSkinVertexFactory.h
@@ -105,7 +105,7 @@ public:
 	/** Creates the resource 
 	 * @param Args The buffer size in bytes.
 	 */
-	FVertexBufferAndSRV CreateResource(FSharedPoolPolicyData::CreationArguments Args);
+	FVertexBufferAndSRV CreateResource(FRHICommandListBase& RHICmdList, FSharedPoolPolicyData::CreationArguments Args);
 	
 	/** Gets the arguments used to create resource
 	 * @param Resource The buffer to get data for.
@@ -137,7 +137,7 @@ public:
 	/** Creates the resource 
 	 * @param Args The buffer size in bytes.
 	 */
-	FVertexBufferAndSRV CreateResource(FSharedPoolPolicyData::CreationArguments Args);
+	FVertexBufferAndSRV CreateResource(FRHICommandListBase& RHICmdList, FSharedPoolPolicyData::CreationArguments Args);
 };
 
 /** A pool for vertex buffers with consistent usage, bucketed for efficiency. */
diff --git a/Engine/Source/Runtime/Engine/Public/LocalVertexFactory.h b/Engine/Source/Runtime/Engine/Public/LocalVertexFactory.h
--- a/Engine/Source/Runtime/Engine/Public/LocalVertexFactory.h
+++ b/Engine/Source/Runtime/Engine/Public/LocalVertexFactory.h
@@ -85,6 +85,9 @@ public:
 	/**
 	 * An implementation of the interface used by TSynchronizedResource to update the resource with new data from the game thread.
 	 */
+	ENGINE_API void SetData(FRHICommandListBase& RHICmdList, const FDataType& InData);
+
+	UE_DEPRECATED(5.4, "SetData requires a command list.")
 	ENGINE_API void SetData(const FDataType& InData);
 
 	/**
diff --git a/Engine/Source/Runtime/Engine/Public/ParticleHelper.h b/Engine/Source/Runtime/Engine/Public/ParticleHelper.h
--- a/Engine/Source/Runtime/Engine/Public/ParticleHelper.h
+++ b/Engine/Source/Runtime/Engine/Public/ParticleHelper.h
@@ -2459,7 +2459,7 @@ public:
 	 * world space primitive uniform buffer is up-to-date.
 	 * Only called in the rendering thread.
 	 */
-	void UpdateWorldSpacePrimitiveUniformBuffer() const;
+	void UpdateWorldSpacePrimitiveUniformBuffer(FRHICommandListBase& RHICmdList) const;
 
 	/** Object position in post projection space. */
 	void GetObjectPositionAndScale(const FSceneView& View, FVector2D& ObjectNDCPosition, FVector2D& ObjectMacroUVScales) const;
diff --git a/Engine/Source/Runtime/Engine/Public/PrimitiveSceneProxy.h b/Engine/Source/Runtime/Engine/Public/PrimitiveSceneProxy.h
--- a/Engine/Source/Runtime/Engine/Public/PrimitiveSceneProxy.h
+++ b/Engine/Source/Runtime/Engine/Public/PrimitiveSceneProxy.h
@@ -1598,7 +1598,7 @@ protected:
 	 * @param InLocalBounds - Primitive local space bounds.
 	 * @param InStaticMeshBounds - Bounds of the primitive mesh instance.
 	 */
-	ENGINE_API virtual void UpdateInstances_RenderThread(const FInstanceUpdateCmdBuffer& CmdBuffer, const FBoxSphereBounds& InBounds, const FBoxSphereBounds& InLocalBounds, const FBoxSphereBounds& InStaticMeshBounds);
+	ENGINE_API virtual void UpdateInstances_RenderThread(FRHICommandListBase& RHICmdList, const FInstanceUpdateCmdBuffer& CmdBuffer, const FBoxSphereBounds& InBounds, const FBoxSphereBounds& InLocalBounds, const FBoxSphereBounds& InStaticMeshBounds);
 
 	/** Updates selection for the primitive proxy. This is called in the rendering thread by SetSelection_GameThread. */
 	void SetSelection_RenderThread(const bool bInParentSelected, const bool bInIndividuallySelected);
diff --git a/Engine/Source/Runtime/Engine/Public/PrimitiveUniformShaderParameters.h b/Engine/Source/Runtime/Engine/Public/PrimitiveUniformShaderParameters.h
--- a/Engine/Source/Runtime/Engine/Public/PrimitiveUniformShaderParameters.h
+++ b/Engine/Source/Runtime/Engine/Public/PrimitiveUniformShaderParameters.h
@@ -90,7 +90,7 @@ class FIdentityPrimitiveUniformBuffer : public TUniformBuffer<FPrimitiveUniformS
 public:
 	void InitContents()
 	{
-		SetContents(GetIdentityPrimitiveParameters());
+		SetContents(FRenderResource::GetImmediateCommandList(), GetIdentityPrimitiveParameters());
 	}
 };
 
diff --git a/Engine/Source/Runtime/Engine/Public/RawIndexBuffer.h b/Engine/Source/Runtime/Engine/Public/RawIndexBuffer.h
--- a/Engine/Source/Runtime/Engine/Public/RawIndexBuffer.h
+++ b/Engine/Source/Runtime/Engine/Public/RawIndexBuffer.h
@@ -276,7 +276,11 @@ public:
 	int32 GetIndexDataSize() const { return IndexStorage.Num(); }
 
 	/** Create an RHI index buffer with CPU data. CPU data may be discarded after creation (see TResourceArray::Discard) */
+	FBufferRHIRef CreateRHIBuffer(FRHICommandListBase& RHICmdList);
+
+	UE_DEPRECATED(5.4, "Use CreateRHIBuffer instead.")
 	FBufferRHIRef CreateRHIBuffer_RenderThread();
+	UE_DEPRECATED(5.4, "Use CreateRHIBuffer instead.")
 	FBufferRHIRef CreateRHIBuffer_Async();
 
 	/** Take over ownership of IntermediateBuffer */
@@ -309,9 +313,6 @@ public:
 	inline bool Is32Bit() const { return b32Bit; }
 
 private:
-	template <bool bRenderThread>
-	FBufferRHIRef CreateRHIBuffer_Internal();
-
 	/** Storage for indices. */
 	TResourceArray<uint8, INDEXBUFFER_ALIGNMENT> IndexStorage;
 
@@ -365,13 +366,13 @@ protected:
 	void ReleaseRHIForStreaming(FRHIResourceUpdateBatcher& Batcher);
 
 	static ENGINE_API FBufferRHIRef CreateRHIIndexBufferInternal(
+		FRHICommandListBase& RHICmdList,
 		const TCHAR* InDebugName,
 		const FName& InOwnerName,
 		int32 IndexCount,
 		size_t IndexSize,
 		FResourceArrayInterface* ResourceArray,
-		bool bNeedSRV,
-		bool bRenderThread
+		bool bNeedSRV
 	);
 
 	// guaranteed only to be valid if the vertex buffer is valid and the buffer was created with the SRV flags
@@ -399,7 +400,7 @@ public:
 	virtual void InitRHI(FRHICommandListBase& RHICmdList) override
 	{
 		const bool bHadIndexData = Num() > 0;
-		IndexBufferRHI = CreateRHIBuffer_RenderThread();
+		IndexBufferRHI = CreateRHIBuffer(RHICmdList);
 
 		if (IndexBufferRHI && IsSRVNeeded(Indices.GetAllowCPUAccess()) && bHadIndexData)
 		{
@@ -497,27 +498,7 @@ public:
 	}
 
 	/** Create an RHI index buffer with CPU data. CPU data may be discarded after creation (see TResourceArray::Discard) */
-	FBufferRHIRef CreateRHIBuffer_RenderThread() { return CreateRHIBuffer_Internal<true>(); }
-	FBufferRHIRef CreateRHIBuffer_Async() { return CreateRHIBuffer_Internal<false>(); }
-
-	/** Similar to Init/ReleaseRHI but only update existing SRV so references to the SRV stays valid */
-	void InitRHIForStreaming(FRHIBuffer* IntermediateBuffer, FRHIResourceUpdateBatcher& Batcher)
-	{
-		FRawStaticIndexBuffer16or32Interface::InitRHIForStreaming(IntermediateBuffer, sizeof(INDEX_TYPE), Batcher);
-	}
-
-	void ReleaseRHIForStreaming(FRHIResourceUpdateBatcher& Batcher)
-	{
-		FRawStaticIndexBuffer16or32Interface::ReleaseRHIForStreaming(Batcher);
-	}
-
-private:
-	TResourceArray<INDEX_TYPE,INDEXBUFFER_ALIGNMENT> Indices;
-
-	int32 CachedNumIndices;
-
-	template <bool bRenderThread>
-	FBufferRHIRef CreateRHIBuffer_Internal()
+	FBufferRHIRef CreateRHIBuffer(FRHICommandListBase& RHICmdList)
 	{
 		if (CachedNumIndices)
 		{
@@ -526,15 +507,37 @@ private:
 			CachedNumIndices = Indices.Num();
 
 			return CreateRHIIndexBufferInternal(
+				RHICmdList,
 				sizeof(INDEX_TYPE) == 4 ? TEXT("FRawStaticIndexBuffer32") : TEXT("FRawStaticIndexBuffer16"),
 				GetOwnerName(),
 				Indices.Num(),
 				sizeof(INDEX_TYPE),
 				&Indices,
-				IsSRVNeeded(Indices.GetAllowCPUAccess()),
-				bRenderThread
+				IsSRVNeeded(Indices.GetAllowCPUAccess())
 			);
 		}
 		return nullptr;
 	}
+
+	UE_DEPRECATED(5.4, "Use CreateRHIBuffer instead.")
+	FBufferRHIRef CreateRHIBuffer_RenderThread();
+	UE_DEPRECATED(5.4, "Use CreateRHIBuffer instead.")
+	FBufferRHIRef CreateRHIBuffer_Async();
+
+	/** Similar to Init/ReleaseRHI but only update existing SRV so references to the SRV stays valid */
+	void InitRHIForStreaming(FRHIBuffer* IntermediateBuffer, FRHIResourceUpdateBatcher& Batcher)
+	{
+		FRawStaticIndexBuffer16or32Interface::InitRHIForStreaming(IntermediateBuffer, sizeof(INDEX_TYPE), Batcher);
+	}
+
+	void ReleaseRHIForStreaming(FRHIResourceUpdateBatcher& Batcher)
+	{
+		FRawStaticIndexBuffer16or32Interface::ReleaseRHIForStreaming(Batcher);
+	}
+
+private:
+	TResourceArray<INDEX_TYPE,INDEXBUFFER_ALIGNMENT> Indices;
+
+	int32 CachedNumIndices;
+
 };
diff --git a/Engine/Source/Runtime/Engine/Public/RayTracingSkinnedGeometry.h b/Engine/Source/Runtime/Engine/Public/RayTracingSkinnedGeometry.h
--- a/Engine/Source/Runtime/Engine/Public/RayTracingSkinnedGeometry.h
+++ b/Engine/Source/Runtime/Engine/Public/RayTracingSkinnedGeometry.h
@@ -44,6 +44,7 @@ private:
 		uint32 ScratchSize;
 	};
 
+	FCriticalSection CS;
 	TMap<FRayTracingGeometry*, FRayTracingUpdateInfo> ToUpdate;
 	
 	/** Estimate of current memory overhead from objects awaiting RHI release. */
diff --git a/Engine/Source/Runtime/Engine/Public/Rendering/ColorVertexBuffer.h b/Engine/Source/Runtime/Engine/Public/Rendering/ColorVertexBuffer.h
--- a/Engine/Source/Runtime/Engine/Public/Rendering/ColorVertexBuffer.h
+++ b/Engine/Source/Runtime/Engine/Public/Rendering/ColorVertexBuffer.h
@@ -143,7 +143,11 @@ public:
 	}
 
 	/** Create an RHI vertex buffer with CPU data. CPU data may be discarded after creation (see TResourceArray::Discard) */
+	FBufferRHIRef CreateRHIBuffer(FRHICommandListBase& RHICmdList);
+
+	UE_DEPRECATED(5.4, "Use CreateRHIBuffer instead.")
 	FBufferRHIRef CreateRHIBuffer_RenderThread();
+	UE_DEPRECATED(5.4, "Use CreateRHIBuffer instead.")
 	FBufferRHIRef CreateRHIBuffer_Async();
 
 	/** Similar to Init/ReleaseRHI but only update existing SRV so references to the SRV stays valid */
@@ -184,9 +188,6 @@ private:
 	/** Allocates the vertex data storage type. */
 	void AllocateData(bool bNeedsCPUAccess = true);
 
-	template <bool bRenderThread>
-	FBufferRHIRef CreateRHIBuffer_Internal();
-
 	/** Purposely hidden */
 	ENGINE_API FColorVertexBuffer(const FColorVertexBuffer &rhs);
 };
diff --git a/Engine/Source/Runtime/Engine/Public/Rendering/MorphTargetVertexInfoBuffers.h b/Engine/Source/Runtime/Engine/Public/Rendering/MorphTargetVertexInfoBuffers.h
--- a/Engine/Source/Runtime/Engine/Public/Rendering/MorphTargetVertexInfoBuffers.h
+++ b/Engine/Source/Runtime/Engine/Public/Rendering/MorphTargetVertexInfoBuffers.h
@@ -107,10 +107,6 @@ protected:
 
 	friend class FSkeletalMeshLODRenderData;
 	ENGINE_API friend FArchive& operator<<(FArchive& Ar, FMorphTargetVertexInfoBuffers& MorphTargetVertexInfoBuffers);
-
-private:
-	template <bool bRenderThread>
-	FBufferRHIRef CreateMorphRHIBuffer_Internal();
 };
 
 ENGINE_API FArchive& operator<<(FArchive& Ar, FMorphTargetVertexInfoBuffers& MorphTargetVertexInfoBuffers);
diff --git a/Engine/Source/Runtime/Engine/Public/Rendering/MultiSizeIndexContainer.h b/Engine/Source/Runtime/Engine/Public/Rendering/MultiSizeIndexContainer.h
--- a/Engine/Source/Runtime/Engine/Public/Rendering/MultiSizeIndexContainer.h
+++ b/Engine/Source/Runtime/Engine/Public/Rendering/MultiSizeIndexContainer.h
@@ -99,8 +99,12 @@ public:
 
 	friend FArchive& operator<<(FArchive& Ar, FMultiSizeIndexContainer& Buffer);
 
-	FBufferRHIRef CreateRHIBuffer_RenderThread();
-	FBufferRHIRef CreateRHIBuffer_Async();
+	ENGINE_API FBufferRHIRef CreateRHIBuffer(FRHICommandListBase& RHICmdList);
+
+	UE_DEPRECATED(5.4, "Use CreateRHIBuffer instead.")
+	ENGINE_API FBufferRHIRef CreateRHIBuffer_RenderThread();
+	UE_DEPRECATED(5.4, "Use CreateRHIBuffer instead.")
+	ENGINE_API FBufferRHIRef CreateRHIBuffer_Async();
 
 	void InitRHIForStreaming(FRHIBuffer* IntermediateBuffer, FRHIResourceUpdateBatcher& Batcher);
 	void ReleaseRHIForStreaming(FRHIResourceUpdateBatcher& Batcher);
diff --git a/Engine/Source/Runtime/Engine/Public/Rendering/PositionVertexBuffer.h b/Engine/Source/Runtime/Engine/Public/Rendering/PositionVertexBuffer.h
--- a/Engine/Source/Runtime/Engine/Public/Rendering/PositionVertexBuffer.h
+++ b/Engine/Source/Runtime/Engine/Public/Rendering/PositionVertexBuffer.h
@@ -104,7 +104,12 @@ public:
 	FORCEINLINE SIZE_T GetAllocatedSize() const { return (Data != nullptr) ? Stride * NumVertices : 0; }
 
 	/** Create an RHI vertex buffer with CPU data. CPU data may be discarded after creation (see TResourceArray::Discard) */
+	FBufferRHIRef CreateRHIBuffer(FRHICommandListBase& RHICmdList);
+
+	UE_DEPRECATED(5.4, "Use CreateRHIBuffer instead.")
 	FBufferRHIRef CreateRHIBuffer_RenderThread();
+
+	UE_DEPRECATED(5.4, "Use CreateRHIBuffer instead.")
 	FBufferRHIRef CreateRHIBuffer_Async();
 
 	/** Similar to Init/ReleaseRHI but only update existing SRV so references to the SRV stays valid */
@@ -143,9 +148,6 @@ private:
 
 	/** Allocates the vertex data storage type. */
 	void AllocateData(bool bInNeedsCPUAccess = true);
-
-	template <bool bRenderThread>
-	FBufferRHIRef CreateRHIBuffer_Internal();
 };
 
 #if UE_ENABLE_INCLUDE_ORDER_DEPRECATED_IN_5_3
diff --git a/Engine/Source/Runtime/Engine/Public/Rendering/SkeletalMeshVertexClothBuffer.h b/Engine/Source/Runtime/Engine/Public/Rendering/SkeletalMeshVertexClothBuffer.h
--- a/Engine/Source/Runtime/Engine/Public/Rendering/SkeletalMeshVertexClothBuffer.h
+++ b/Engine/Source/Runtime/Engine/Public/Rendering/SkeletalMeshVertexClothBuffer.h
@@ -124,7 +124,11 @@ public:
 	}
 
 	/** Create an RHI vertex buffer with CPU data. CPU data may be discarded after creation (see TResourceArray::Discard) */
+	FBufferRHIRef CreateRHIBuffer(FRHICommandListBase& RHICmdList);
+
+	UE_DEPRECATED(5.4, "Use CreateRHIBuffer instead.")
 	FBufferRHIRef CreateRHIBuffer_RenderThread();
+	UE_DEPRECATED(5.4, "Use CreateRHIBuffer instead.")
 	FBufferRHIRef CreateRHIBuffer_Async();
 
 	/** Similar to Init/ReleaseRHI but only update existing SRV so references to the SRV stays valid */
@@ -150,7 +154,4 @@ private:
 	* Allocates the vertex data storage type
 	*/
 	void AllocateData();
-
-	template <bool bRenderThread>
-	FBufferRHIRef CreateRHIBuffer_Internal();
 };
\ No newline at end of file
diff --git a/Engine/Source/Runtime/Engine/Public/Rendering/SkinWeightVertexBuffer.h b/Engine/Source/Runtime/Engine/Public/Rendering/SkinWeightVertexBuffer.h
--- a/Engine/Source/Runtime/Engine/Public/Rendering/SkinWeightVertexBuffer.h
+++ b/Engine/Source/Runtime/Engine/Public/Rendering/SkinWeightVertexBuffer.h
@@ -138,7 +138,11 @@ public:
 	{ return bNeedsCPUAccess; }
 
 	/** Create an RHI vertex buffer with CPU data. CPU data may be discarded after creation (see TResourceArray::Discard) */
+	FBufferRHIRef CreateRHIBuffer(FRHICommandListBase& RHICmdList);
+
+	UE_DEPRECATED(5.4, "Use CreateRHIBuffer instead.")
 	FBufferRHIRef CreateRHIBuffer_RenderThread();
+	UE_DEPRECATED(5.4, "Use CreateRHIBuffer instead.")
 	FBufferRHIRef CreateRHIBuffer_Async();
 
 	/** Similar to Init/ReleaseRHI but only update existing SRV so references to the SRV stays valid */
@@ -158,9 +162,6 @@ private:
 	/** Allocates the vertex data storage type. */
 	ENGINE_API void AllocateData();
 
-	template <bool bRenderThread>
-	FBufferRHIRef CreateRHIBuffer_Internal();
-
 	/** true if this vertex buffer will be used with CPU skinning. Resource arrays are set to cpu accessible if this is true */
 	bool bNeedsCPUAccess;
 
@@ -300,7 +301,11 @@ public:
 	ENGINE_API void CopyDataFromBuffer(const uint8* InSkinWeightData, uint32 InNumVertices);
 
 	/** Create an RHI vertex buffer with CPU data. CPU data may be discarded after creation (see TResourceArray::Discard) */
+	FBufferRHIRef CreateRHIBuffer(FRHICommandListBase& RHICmdList);
+
+	UE_DEPRECATED(5.4, "Use CreateRHIBuffer instead.")
 	FBufferRHIRef CreateRHIBuffer_RenderThread();
+	UE_DEPRECATED(5.4, "Use CreateRHIBuffer instead.")
 	FBufferRHIRef CreateRHIBuffer_Async();
 
 	/** Similar to Init/ReleaseRHI but only update existing SRV so references to the SRV stays valid */
@@ -360,9 +365,6 @@ private:
 
 	/** Allocates the vertex data storage type. */
 	ENGINE_API void AllocateData();
-
-	template <bool bRenderThread>
-	FBufferRHIRef CreateRHIBuffer_Internal();
 };
 
 /** A container for skin weights data vertex buffer and lookup vertex buffer. */
@@ -463,7 +465,11 @@ public:
 	const FSkinWeightLookupVertexBuffer* GetLookupVertexBuffer() const
 	{ return &LookupVertexBuffer; }
 
+	FSkinWeightRHIInfo CreateRHIBuffer(FRHICommandListBase& RHICmdList);
+
+	UE_DEPRECATED(5.4, "Use CreateRHIBuffer instead.")
 	FSkinWeightRHIInfo CreateRHIBuffer_RenderThread();
+	UE_DEPRECATED(5.4, "Use CreateRHIBuffer instead.")
 	FSkinWeightRHIInfo CreateRHIBuffer_Async();
 
 	ENGINE_API GPUSkinBoneInfluenceType GetBoneInfluenceType() const;
diff --git a/Engine/Source/Runtime/Engine/Public/Rendering/StaticMeshVertexBuffer.h b/Engine/Source/Runtime/Engine/Public/Rendering/StaticMeshVertexBuffer.h
--- a/Engine/Source/Runtime/Engine/Public/Rendering/StaticMeshVertexBuffer.h
+++ b/Engine/Source/Runtime/Engine/Public/Rendering/StaticMeshVertexBuffer.h
@@ -437,9 +437,16 @@ public:
 	}
 
 	/** Create an RHI vertex buffer with CPU data. CPU data may be discarded after creation (see TResourceArray::Discard) */
+	FBufferRHIRef CreateTangentsRHIBuffer(FRHICommandListBase& RHICmdList);
+	FBufferRHIRef CreateTexCoordRHIBuffer(FRHICommandListBase& RHICmdList);
+
+	UE_DEPRECATED(5.4, "Use CreateTangentsRHIBuffer instead.")
 	FBufferRHIRef CreateTangentsRHIBuffer_RenderThread();
+	UE_DEPRECATED(5.4, "Use CreateTangentsRHIBuffer instead.")
 	FBufferRHIRef CreateTangentsRHIBuffer_Async();
+	UE_DEPRECATED(5.4, "Use CreateTexCoordRHIBuffer instead.")
 	FBufferRHIRef CreateTexCoordRHIBuffer_RenderThread();
+	UE_DEPRECATED(5.4, "Use CreateTexCoordRHIBuffer instead.")
 	FBufferRHIRef CreateTexCoordRHIBuffer_Async();
 
 	/** Similar to Init/ReleaseRHI but only update existing SRV so references to the SRV stays valid */
@@ -534,11 +541,5 @@ private:
 	*/
 	void ConvertHalfTexcoordsToFloat(const uint8* InData);
 
-	template <bool bRenderThread>
-	FBufferRHIRef CreateTangentsRHIBuffer_Internal();
-
-	template<bool bRenderThread>
-	FBufferRHIRef CreateTexCoordRHIBuffer_Internal();
-
 	void InitTangentAndTexCoordStrides();
 };
diff --git a/Engine/Source/Runtime/Engine/Public/ResourcePool.h b/Engine/Source/Runtime/Engine/Public/ResourcePool.h
--- a/Engine/Source/Runtime/Engine/Public/ResourcePool.h
+++ b/Engine/Source/Runtime/Engine/Public/ResourcePool.h
@@ -9,6 +9,7 @@
 #include "CoreMinimal.h"
 #include "RenderResource.h"
 #include "TickableObjectRenderThread.h"
+#include "RHICommandList.h"
 
 /** A templated pool for resources that can only be freed at a 'safe' point in the frame. */
 template<typename ResourceType, class ResourcePoolPolicy, class ResourceCreationArguments>
@@ -50,7 +51,7 @@ public:
 	 * @param Args the argument object for construction.
 	 * @returns An initialised resource.
 	 */
-	ResourceType CreatePooledResource(ResourceCreationArguments Args)
+	ResourceType CreatePooledResource(FRHICommandListBase& RHICmdList, ResourceCreationArguments Args)
 	{
 		// Find the appropriate bucket based on size
 		const uint32 BucketIndex = Policy.GetPoolBucketIndex(Args);
@@ -67,7 +68,7 @@ public:
 		}
 
 		// Nothing usable was found in the free pool, create a new resource
-		return Policy.CreateResource(Args);
+		return Policy.CreateResource(RHICmdList, Args);
 	}
 	
 	/** Release a resource back into the pool.
@@ -206,11 +207,23 @@ public:
 	 * @param Args the argument object for construction.
 	 * @returns An initialised resource or the policy's NullResource if not initialised.
 	 */
+	ResourceType CreatePooledResource(FRHICommandListBase& RHICmdList, ResourceCreationArguments Args)
+	{
+		if (IsInitialized())
+		{
+			return TResourcePool<ResourceType, ResourcePoolPolicy, ResourceCreationArguments>::CreatePooledResource(RHICmdList, Args);
+		}
+		else
+		{
+			return ResourceType();
+		}
+	}
+
 	ResourceType CreatePooledResource(ResourceCreationArguments Args)
 	{
 		if (IsInitialized())
 		{
-			return TResourcePool<ResourceType, ResourcePoolPolicy, ResourceCreationArguments>::CreatePooledResource(Args);
+			return TResourcePool<ResourceType, ResourcePoolPolicy, ResourceCreationArguments>::CreatePooledResource(FRHICommandListImmediate::Get(), Args);
 		}
 		else
 		{
diff --git a/Engine/Source/Runtime/Engine/Public/SceneManagement.h b/Engine/Source/Runtime/Engine/Public/SceneManagement.h
--- a/Engine/Source/Runtime/Engine/Public/SceneManagement.h
+++ b/Engine/Source/Runtime/Engine/Public/SceneManagement.h
@@ -314,7 +314,7 @@ class FDefaultWorkingColorSpaceUniformBuffer : public TUniformBuffer<FWorkingCol
 	typedef TUniformBuffer<FWorkingColorSpaceShaderParameters> Super;
 public:
 
-	void Update(const UE::Color::FColorSpace& InColorSpace);
+	void Update(FRHICommandListBase& RHICmdList, const UE::Color::FColorSpace& InColorSpace);
 };
 
 ENGINE_API extern TGlobalResource<FDefaultWorkingColorSpaceUniformBuffer> GDefaultWorkingColorSpaceUniformBuffer;
diff --git a/Engine/Source/Runtime/Experimental/GeometryCollectionEngine/Private/GeometryCollection/GeometryCollectionRenderData.cpp b/Engine/Source/Runtime/Experimental/GeometryCollectionEngine/Private/GeometryCollection/GeometryCollectionRenderData.cpp
--- a/Engine/Source/Runtime/Experimental/GeometryCollectionEngine/Private/GeometryCollection/GeometryCollectionRenderData.cpp
+++ b/Engine/Source/Runtime/Experimental/GeometryCollectionEngine/Private/GeometryCollection/GeometryCollectionRenderData.cpp
@@ -46,7 +46,7 @@ void FBoneMapVertexBuffer::AllocateData(bool bInNeedsCPUAccess)
 void FBoneMapVertexBuffer::InitRHI(FRHICommandListBase& RHICmdList)
 {
 	const bool bHadData = BoneMapData != nullptr;
-	VertexBufferRHI = CreateRHIBuffer<true>(BoneMapData, NumVertices, BUF_Static | BUF_ShaderResource | BUF_SourceCopy, TEXT("FBoneMapVertexBuffer")); 
+	VertexBufferRHI = FRenderResource::CreateRHIBuffer(RHICmdList, BoneMapData, NumVertices, BUF_Static | BUF_ShaderResource | BUF_SourceCopy, TEXT("FBoneMapVertexBuffer")); 
 	if (VertexBufferRHI != nullptr)
 	{
 		VertexBufferSRV = RHICmdList.CreateShaderResourceView(FShaderResourceViewInitializer(bHadData ? VertexBufferRHI : nullptr, PixelFormat));
diff --git a/Engine/Source/Runtime/Experimental/GeometryCollectionEngine/Private/GeometryCollection/GeometryCollectionSceneProxy.cpp b/Engine/Source/Runtime/Experimental/GeometryCollectionEngine/Private/GeometryCollection/GeometryCollectionSceneProxy.cpp
--- a/Engine/Source/Runtime/Experimental/GeometryCollectionEngine/Private/GeometryCollection/GeometryCollectionSceneProxy.cpp
+++ b/Engine/Source/Runtime/Experimental/GeometryCollectionEngine/Private/GeometryCollection/GeometryCollectionSceneProxy.cpp
@@ -906,7 +906,7 @@ void FGeometryCollectionSceneProxy::GetDynamicRayTracingInstances(FRayTracingMat
 
 void FGeometryCollectionSceneProxy::UpdatingRayTracingGeometry_RenderingThread(TArray<FGeometryCollectionMeshElement> const& InSectionArray)
 {
-	FRHICommandListBase& RHICmdList = FRHICommandListImmediate::Get();
+	FRHICommandList& RHICmdList = FRHICommandListImmediate::Get();
 
 	if (bGeometryResourceUpdated)
 	{
diff --git a/Engine/Source/Runtime/GeometryFramework/Public/Components/BaseDynamicMeshSceneProxy.h b/Engine/Source/Runtime/GeometryFramework/Public/Components/BaseDynamicMeshSceneProxy.h
--- a/Engine/Source/Runtime/GeometryFramework/Public/Components/BaseDynamicMeshSceneProxy.h
+++ b/Engine/Source/Runtime/GeometryFramework/Public/Components/BaseDynamicMeshSceneProxy.h
@@ -139,9 +139,9 @@ public:
 			return;
 		}
 
-		InitOrUpdateResource(&this->PositionVertexBuffer);
-		InitOrUpdateResource(&this->StaticMeshVertexBuffer);
-		InitOrUpdateResource(&this->ColorVertexBuffer);
+		InitOrUpdateResource(RHICmdList, &this->PositionVertexBuffer);
+		InitOrUpdateResource(RHICmdList, &this->StaticMeshVertexBuffer);
+		InitOrUpdateResource(RHICmdList, &this->ColorVertexBuffer);
 
 		FLocalVertexFactory::FDataType Data;
 		this->PositionVertexBuffer.BindPositionVertexBuffer(&this->VertexFactory, Data);
@@ -150,9 +150,9 @@ public:
 		// currently no lightmaps support
 		//this->StaticMeshVertexBuffer.BindLightMapVertexBuffer(&this->VertexFactory, Data, LightMapIndex);
 		this->ColorVertexBuffer.BindColorVertexBuffer(&this->VertexFactory, Data);
-		this->VertexFactory.SetData(Data);
+		this->VertexFactory.SetData(RHICmdList, Data);
 
-		InitOrUpdateResource(&this->VertexFactory);
+		InitOrUpdateResource(RHICmdList, &this->VertexFactory);
 		PositionVertexBuffer.InitResource(RHICmdList);
 		StaticMeshVertexBuffer.InitResource(RHICmdList);
 		ColorVertexBuffer.InitResource(RHICmdList);
@@ -180,15 +180,15 @@ public:
 	void UploadIndexBufferUpdate()
 	{
 		// todo: can this be done with RHI locking and memcpy, like in TransferVertexUpdateToGPU?
+		FRHICommandListBase& RHICmdList = FRHICommandListImmediate::Get();
 
-		check(IsInRenderingThread());
 		if (IndexBuffer.Indices.Num() > 0)
 		{
-			InitOrUpdateResource(&IndexBuffer);
+			InitOrUpdateResource(RHICmdList, &IndexBuffer);
 		}
 		if (bEnableSecondaryIndexBuffer && SecondaryIndexBuffer.Indices.Num() > 0)
 		{
-			InitOrUpdateResource(&SecondaryIndexBuffer);
+			InitOrUpdateResource(RHICmdList, &SecondaryIndexBuffer);
 		}
 
 		InvalidateRayTracingData();
@@ -208,24 +208,24 @@ public:
 		// are any situations where we would change vertex buffer size w/o also updating the index
 		// buffers (in which case we are fully rebuilding the buffers...)
 
-		check(IsInRenderingThread());
-
 		if (TriangleCount == 0)
 		{
 			return;
 		}
 
+		FRHICommandListBase& RHICmdList = FRHICommandListImmediate::Get();
+
 		if (bPositions)
 		{
-			InitOrUpdateResource(&this->PositionVertexBuffer);
+			InitOrUpdateResource(RHICmdList, &this->PositionVertexBuffer);
 		}
 		if (bMeshAttribs)
 		{
-			InitOrUpdateResource(&this->StaticMeshVertexBuffer);
+			InitOrUpdateResource(RHICmdList, &this->StaticMeshVertexBuffer);
 		}
 		if (bColors)
 		{
-			InitOrUpdateResource(&this->ColorVertexBuffer);
+			InitOrUpdateResource(RHICmdList, &this->ColorVertexBuffer);
 		}
 
 		FLocalVertexFactory::FDataType Data;
@@ -233,9 +233,9 @@ public:
 		this->StaticMeshVertexBuffer.BindTangentVertexBuffer(&this->VertexFactory, Data);
 		this->StaticMeshVertexBuffer.BindPackedTexCoordVertexBuffer(&this->VertexFactory, Data);
 		this->ColorVertexBuffer.BindColorVertexBuffer(&this->VertexFactory, Data);
-		this->VertexFactory.SetData(Data);
+		this->VertexFactory.SetData(RHICmdList, Data);
 
-		InitOrUpdateResource(&this->VertexFactory);
+		InitOrUpdateResource(RHICmdList, &this->VertexFactory);
 
 		InvalidateRayTracingData();
 		ValidateRayTracingData();		// currently we are immediately validating. This may be revisited in future.
@@ -364,10 +364,8 @@ protected:
 	 * Initializes a render resource, or update it if already initialized.
 	 * @warning This function can only be called on the Render Thread
 	 */
-	void InitOrUpdateResource(FRenderResource* Resource)
+	void InitOrUpdateResource(FRHICommandListBase& RHICmdList, FRenderResource* Resource)
 	{
-		FRHICommandListBase& RHICmdList = FRHICommandListImmediate::Get();
-
 		if (!Resource->IsInitialized())
 		{
 			Resource->InitResource(RHICmdList);
diff --git a/Engine/Source/Runtime/Landscape/Private/LandscapeRender.cpp b/Engine/Source/Runtime/Landscape/Private/LandscapeRender.cpp
--- a/Engine/Source/Runtime/Landscape/Private/LandscapeRender.cpp
+++ b/Engine/Source/Runtime/Landscape/Private/LandscapeRender.cpp
@@ -1421,14 +1421,14 @@ void FLandscapeComponentSceneProxy::CreateRenderThreadResources()
 	LandscapeFixedGridUniformShaderParameters.AddDefaulted(NumMips);
 	for (int32 LodIndex = 0; LodIndex < NumMips; ++LodIndex)
 	{
-		LandscapeFixedGridUniformShaderParameters[LodIndex].InitResource(RHICmdList);
 		FLandscapeFixedGridUniformShaderParameters Parameters;
 		Parameters.LodValues = FVector4f(
 			static_cast<float>(LodIndex),
 			0.f,
 			(float)((SubsectionSizeVerts >> LodIndex) - 1),
 			1.f / (float)((SubsectionSizeVerts >> LodIndex) - 1));
-		LandscapeFixedGridUniformShaderParameters[LodIndex].SetContents(Parameters);
+		LandscapeFixedGridUniformShaderParameters[LodIndex].SetContents(RHICmdList, Parameters);
+		LandscapeFixedGridUniformShaderParameters[LodIndex].InitResource(RHICmdList);
 	}
 
 #if WITH_EDITOR
@@ -1998,7 +1998,7 @@ void FLandscapeComponentSceneProxy::OnTransformChanged()
 		LandscapeParams.XYOffsetmapTextureSampler = GBlackTexture->SamplerStateRHI;
 	}
 
-	LandscapeUniformShaderParameters.SetContents(LandscapeParams);
+	LandscapeUniformShaderParameters.SetContents(FRHICommandListImmediate::Get(), LandscapeParams);
 
 	// Recache mesh draw commands for changed uniform buffers
 	GetScene().UpdateCachedRenderStates(this);
diff --git a/Engine/Source/Runtime/MRMesh/Private/MRMeshComponent.cpp b/Engine/Source/Runtime/MRMesh/Private/MRMeshComponent.cpp
--- a/Engine/Source/Runtime/MRMesh/Private/MRMeshComponent.cpp
+++ b/Engine/Source/Runtime/MRMesh/Private/MRMeshComponent.cpp
@@ -207,7 +207,7 @@ static void InitVertexFactory(FLocalVertexFactory* VertexFactory, const FMRMeshP
 			NewData.ColorComponent = FVertexStreamComponent(&MRMeshSection.ColorBuffer, 0, sizeof(FColor), VET_Color, EVertexStreamUsage::ManualFetch);
 		}
 
-		VertexFactory->SetData(NewData);
+		VertexFactory->SetData(RHICmdList, NewData);
 		VertexFactory->InitResource(RHICmdList);
 	});
 }
diff --git a/Engine/Source/Runtime/RenderCore/Private/RayTracingGeometry.cpp b/Engine/Source/Runtime/RenderCore/Private/RayTracingGeometry.cpp
--- a/Engine/Source/Runtime/RenderCore/Private/RayTracingGeometry.cpp
+++ b/Engine/Source/Runtime/RenderCore/Private/RayTracingGeometry.cpp
@@ -47,7 +47,7 @@ void FRayTracingGeometry::ReleaseRHIForStreaming(FRHIResourceUpdateBatcher& Batc
 	}
 }
 
-void FRayTracingGeometry::CreateRayTracingGeometryFromCPUData(TResourceArray<uint8>& OfflineData)
+void FRayTracingGeometry::CreateRayTracingGeometryFromCPUData(FRHICommandList& RHICmdList, TResourceArray<uint8>& OfflineData)
 {
 	check(OfflineData.Num() == 0 || Initializer.OfflineData == nullptr);
 	if (OfflineData.Num())
@@ -62,24 +62,33 @@ void FRayTracingGeometry::CreateRayTracingGeometryFromCPUData(TResourceArray<uin
 	}
 
 	SetRequiresBuild(Initializer.OfflineData == nullptr);
-	RayTracingGeometryRHI = RHICreateRayTracingGeometry(Initializer);
+	RayTracingGeometryRHI = RHICmdList.CreateRayTracingGeometry(Initializer);
 }
 
-void FRayTracingGeometry::RequestBuildIfNeeded(ERTAccelerationStructureBuildPriority InBuildPriority)
+void FRayTracingGeometry::CreateRayTracingGeometryFromCPUData(TResourceArray<uint8>& OfflineData)
+{
+	CreateRayTracingGeometryFromCPUData(FRHICommandListImmediate::Get(), OfflineData);
+}
+
+void FRayTracingGeometry::RequestBuildIfNeeded(FRHICommandList& RHICmdList, ERTAccelerationStructureBuildPriority InBuildPriority)
 {
 	RayTracingGeometryRHI->SetInitializer(Initializer);
 
 	if (GetRequiresBuild())
 	{
-		RayTracingBuildRequestIndex = GRayTracingGeometryManager.RequestBuildAccelerationStructure(this, InBuildPriority);
+		RayTracingBuildRequestIndex = GRayTracingGeometryManager.RequestBuildAccelerationStructure(RHICmdList, this, InBuildPriority);
 		SetRequiresBuild(false);
 	}
 }
 
-void FRayTracingGeometry::InitRHIForDynamicRayTracing()
+void FRayTracingGeometry::RequestBuildIfNeeded(ERTAccelerationStructureBuildPriority InBuildPriority)
+{
+	RequestBuildIfNeeded(FRHICommandListImmediate::Get(), InBuildPriority);
+}
+
+void FRayTracingGeometry::InitRHIForDynamicRayTracing(FRHICommandList& RHICmdList)
 {
 	check(GetRayTracingMode() == ERayTracingMode::Dynamic);
-	FRHICommandListBase& RHICmdList = FRHICommandListImmediate::Get();
 
 	// Streaming BLAS needs special handling to not get their "streaming" type wiped out as it will cause issues down the line.	
 	// We only have to do this if the geometry was marked to be streamed in.
@@ -103,14 +112,14 @@ void FRayTracingGeometry::InitRHIForDynamicRayTracing()
 			FRayTracingGeometryInitializer IntermediateInitializer = Initializer;
 			IntermediateInitializer.Type = ERayTracingGeometryInitializerType::StreamingSource;
 
-			FRayTracingGeometryRHIRef IntermediateRayTracingGeometry = RHICreateRayTracingGeometry(IntermediateInitializer);
+			FRayTracingGeometryRHIRef IntermediateRayTracingGeometry = RHICmdList.CreateRayTracingGeometry(IntermediateInitializer);
 			InitRHIForStreaming(IntermediateRayTracingGeometry, Batcher);
 
 			// When Batcher goes out of scope it will add commands to copy the BLAS buffers on RHI thread.
 			// We need to do it before we build the current geometry (also on RHI thread).
 		}
 
-		RequestBuildIfNeeded(ERTAccelerationStructureBuildPriority::Normal);
+		RequestBuildIfNeeded(RHICmdList, ERTAccelerationStructureBuildPriority::Normal);
 	}
 	else
 	{
@@ -118,7 +127,12 @@ void FRayTracingGeometry::InitRHIForDynamicRayTracing()
 	}
 }
 
-void FRayTracingGeometry::CreateRayTracingGeometry(FRHICommandListBase& RHICmdList, ERTAccelerationStructureBuildPriority InBuildPriority)
+void FRayTracingGeometry::InitRHIForDynamicRayTracing()
+{
+	InitRHIForDynamicRayTracing(FRHICommandListImmediate::Get());
+}
+
+void FRayTracingGeometry::CreateRayTracingGeometry(FRHICommandList& RHICmdList, ERTAccelerationStructureBuildPriority InBuildPriority)
 {
 	// Release previous RHI object if any
 	ReleaseRHI();
@@ -178,7 +192,7 @@ void FRayTracingGeometry::CreateRayTracingGeometry(FRHICommandListBase& RHICmdLi
 			{
 				if (IsRayTracingEnabled())
 				{
-					RayTracingBuildRequestIndex = GRayTracingGeometryManager.RequestBuildAccelerationStructure(this, InBuildPriority);
+					RayTracingBuildRequestIndex = GRayTracingGeometryManager.RequestBuildAccelerationStructure(RHICmdList, this, InBuildPriority);
 				}
 				SetRequiresBuild(false);
 			}
@@ -198,6 +212,11 @@ void FRayTracingGeometry::CreateRayTracingGeometry(FRHICommandListBase& RHICmdLi
 	}
 }
 
+void FRayTracingGeometry::CreateRayTracingGeometry(ERTAccelerationStructureBuildPriority InBuildPriority)
+{
+	CreateRayTracingGeometry(FRHICommandListImmediate::Get(), InBuildPriority);
+}
+
 bool FRayTracingGeometry::IsValid() const
 {
 	return RayTracingGeometryRHI != nullptr && Initializer.TotalPrimitiveCount > 0 && EnumHasAnyFlags(GeometryState, EGeometryStateFlags::Valid);
@@ -211,7 +230,7 @@ void FRayTracingGeometry::InitRHI(FRHICommandListBase& RHICmdList)
 	ERTAccelerationStructureBuildPriority BuildPriority = Initializer.Type != ERayTracingGeometryInitializerType::Rendering
 		? ERTAccelerationStructureBuildPriority::Skip
 		: ERTAccelerationStructureBuildPriority::Normal;
-	CreateRayTracingGeometry(RHICmdList, BuildPriority);
+	CreateRayTracingGeometry(FRHICommandList::Get(RHICmdList), BuildPriority);
 }
 
 void FRayTracingGeometry::ReleaseRHI()
diff --git a/Engine/Source/Runtime/RenderCore/Private/RayTracingGeometryManager.cpp b/Engine/Source/Runtime/RenderCore/Private/RayTracingGeometryManager.cpp
--- a/Engine/Source/Runtime/RenderCore/Private/RayTracingGeometryManager.cpp
+++ b/Engine/Source/Runtime/RenderCore/Private/RayTracingGeometryManager.cpp
@@ -49,13 +49,13 @@ static float GetInitialBuildPriority(ERTAccelerationStructureBuildPriority InBui
 	}
 }
 
-FRayTracingGeometryManager::BuildRequestIndex FRayTracingGeometryManager::RequestBuildAccelerationStructure(FRayTracingGeometry* InGeometry, ERTAccelerationStructureBuildPriority InPriority, EAccelerationStructureBuildMode InBuildMode)
+FRayTracingGeometryManager::BuildRequestIndex FRayTracingGeometryManager::RequestBuildAccelerationStructure(FRHICommandList& RHICmdList, FRayTracingGeometry* InGeometry, ERTAccelerationStructureBuildPriority InPriority, EAccelerationStructureBuildMode InBuildMode)
 {
 	// If immediate then enqueue command directly on the immediate command list
 	if (GRayTracingMaxBuiltPrimitivesPerFrame <= 0 || InPriority == ERTAccelerationStructureBuildPriority::Immediate)
 	{
 		check(InBuildMode == EAccelerationStructureBuildMode::Build);
-		FRHICommandListExecutor::GetImmediateCommandList().BuildAccelerationStructure(InGeometry->RayTracingGeometryRHI);
+		RHICmdList.BuildAccelerationStructure(InGeometry->RayTracingGeometryRHI);
 		return INDEX_NONE;
 	}
 	else
@@ -109,7 +109,7 @@ void FRayTracingGeometryManager::ReleaseRayTracingGeometryHandle(RayTracingGeome
 	}	
 }
 
-void FRayTracingGeometryManager::Tick(bool bHasRayTracingEnableChanged)
+void FRayTracingGeometryManager::Tick(FRHICommandList& RHICmdList, bool bHasRayTracingEnableChanged)
 {
 	if (GetRayTracingMode() != ERayTracingMode::Dynamic)
 	{
@@ -150,7 +150,7 @@ void FRayTracingGeometryManager::Tick(bool bHasRayTracingEnableChanged)
 		{
 			if (Geometry->RayTracingGeometryRHI == nullptr)
 			{
-				Geometry->InitRHIForDynamicRayTracing();
+				Geometry->InitRHIForDynamicRayTracing(RHICmdList);
 			}
 		}
 	}
diff --git a/Engine/Source/Runtime/RenderCore/Private/RenderResource.cpp b/Engine/Source/Runtime/RenderCore/Private/RenderResource.cpp
--- a/Engine/Source/Runtime/RenderCore/Private/RenderResource.cpp
+++ b/Engine/Source/Runtime/RenderCore/Private/RenderResource.cpp
@@ -18,11 +18,6 @@
 /** Whether to enable mip-level fading or not: +1.0f if enabled, -1.0f if disabled. */
 float GEnableMipLevelFading = 1.0f;
 
-bool GFreeStructuresOnRHIBufferCreation = true;
-FAutoConsoleVariableRef CVarFreeStructuresOnRHIBufferCreation(
-	TEXT("r.FreeStructuresOnRHIBufferCreation"),
-	GFreeStructuresOnRHIBufferCreation,
-	TEXT("Toggles experimental method for freeing helper structures that own the resource arrays after submitting to RHI instead of in the callback sink."));
 
 class FRenderResourceList
 {
@@ -164,7 +159,7 @@ void FRenderResource::ChangeFeatureLevel(ERHIFeatureLevel::Type NewFeatureLevel)
 	});
 }
 
-FRHICommandListBase& FRenderResource::GetCommandList()
+FRHICommandListBase& FRenderResource::GetImmediateCommandList()
 {
 	check(IsInRenderingThread());
 	return FRHICommandListExecutor::GetImmediateCommandList();
@@ -256,16 +251,16 @@ FRenderResource::~FRenderResource()
 
 bool FRenderResource::ShouldFreeResourceObject(void* ResourceObject, FResourceArrayInterface* ResourceArray)
 {
-	return GFreeStructuresOnRHIBufferCreation && ResourceObject && (!ResourceArray || !ResourceArray->GetResourceDataSize());
+	return ResourceObject && (!ResourceArray || !ResourceArray->GetResourceDataSize());
 }
 
 FBufferRHIRef FRenderResource::CreateRHIBufferInternal(
+	FRHICommandListBase& RHICmdList,
 	const TCHAR* InDebugName,
 	const FName& InOwnerName,
 	uint32 ResourceCount,
 	EBufferUsageFlags InBufferUsageFlags,
 	FResourceArrayInterface* ResourceArray,
-	bool bRenderThread,
 	bool bWithoutNativeResource)
 {
 	const uint32 SizeInBytes = ResourceArray ? ResourceArray->GetResourceDataSize() : 0;
@@ -274,16 +269,7 @@ FBufferRHIRef FRenderResource::CreateRHIBufferInternal(
 	CreateInfo.OwnerName = InOwnerName;
 	CreateInfo.bWithoutNativeResource = bWithoutNativeResource;
 
-	FBufferRHIRef Buffer;
-	if (bRenderThread)
-	{
-		Buffer = FRHICommandListImmediate::Get().CreateVertexBuffer(SizeInBytes, InBufferUsageFlags, CreateInfo);
-	}
-	else
-	{
-		FRHIAsyncCommandList CommandList;
-		Buffer = CommandList->CreateBuffer(SizeInBytes, InBufferUsageFlags | EBufferUsageFlags::VertexBuffer, 0, ERHIAccess::SRVMask, CreateInfo);
-	}
+	FBufferRHIRef Buffer = RHICmdList.CreateBuffer(SizeInBytes, InBufferUsageFlags | EBufferUsageFlags::VertexBuffer, 0, ERHIAccess::VertexOrIndexBuffer, CreateInfo);
 
 	Buffer->SetOwnerName(InOwnerName);
 	return Buffer;
diff --git a/Engine/Source/Runtime/RenderCore/Public/RayTracingGeometry.h b/Engine/Source/Runtime/RenderCore/Public/RayTracingGeometry.h
--- a/Engine/Source/Runtime/RenderCore/Public/RayTracingGeometry.h
+++ b/Engine/Source/Runtime/RenderCore/Public/RayTracingGeometry.h
@@ -9,6 +9,8 @@
 #include "RHI.h"
 #endif
 
+class FRHICommandListBase;
+
 enum class ERTAccelerationStructureBuildPriority
 {
 	Immediate,
@@ -100,14 +102,27 @@ public:
 	RENDERCORE_API void InitRHIForStreaming(FRHIRayTracingGeometry* IntermediateGeometry, FRHIResourceUpdateBatcher& Batcher);
 	RENDERCORE_API void ReleaseRHIForStreaming(FRHIResourceUpdateBatcher& Batcher);
 
+	RENDERCORE_API void CreateRayTracingGeometryFromCPUData(FRHICommandList& RHICmdList, TResourceArray<uint8>& OfflineData);
+
+	UE_DEPRECATED(5.4, "CreateRayTracingGeometryFromCPUData now requires a command list.")
 	RENDERCORE_API void CreateRayTracingGeometryFromCPUData(TResourceArray<uint8>& OfflineData);
-	RENDERCORE_API void RequestBuildIfNeeded(ERTAccelerationStructureBuildPriority InBuildPriority);
 
+	RENDERCORE_API void RequestBuildIfNeeded(FRHICommandList& RHICmdList, ERTAccelerationStructureBuildPriority InBuildPriority);
+
+	UE_DEPRECATED(5.4, "RequestBuildIfNeeded now requires a command list.")
+	RENDERCORE_API void RequestBuildIfNeeded(ERTAccelerationStructureBuildPriority InBuildPriority);
+	
 	// That function is only supposed to be used when dynamic ray tracing is enabled
+	RENDERCORE_API void InitRHIForDynamicRayTracing(FRHICommandList& RHICmdList);
+
+	UE_DEPRECATED(5.4, "InitRHIForDynamicRayTracing now requires a command list.")
 	RENDERCORE_API void InitRHIForDynamicRayTracing();
 
-	RENDERCORE_API void CreateRayTracingGeometry(FRHICommandListBase& RHICmdList, ERTAccelerationStructureBuildPriority InBuildPriority);
+	RENDERCORE_API void CreateRayTracingGeometry(FRHICommandList& RHICmdList, ERTAccelerationStructureBuildPriority InBuildPriority);
 
+	UE_DEPRECATED(5.4, "CreateRayTracingGeometry now requires a command list.")
+	RENDERCORE_API void CreateRayTracingGeometry(ERTAccelerationStructureBuildPriority InBuildPriority);
+	
 	bool HasPendingBuildRequest() const
 	{
 		return RayTracingBuildRequestIndex != INDEX_NONE;
diff --git a/Engine/Source/Runtime/RenderCore/Public/RayTracingGeometryManager.h b/Engine/Source/Runtime/RenderCore/Public/RayTracingGeometryManager.h
--- a/Engine/Source/Runtime/RenderCore/Public/RayTracingGeometryManager.h
+++ b/Engine/Source/Runtime/RenderCore/Public/RayTracingGeometryManager.h
@@ -6,7 +6,7 @@
 
 #include "RHI.h"
 #include "RHIResources.h"
-#include "DynamicRHI.h"
+#include "RHICommandList.h"
 
 #if RHI_RAYTRACING
 
@@ -24,19 +24,36 @@ public:
 
 	FRayTracingGeometryManager() {}
 	~FRayTracingGeometryManager() {}
+	BuildRequestIndex RequestBuildAccelerationStructure(FRHICommandList& RHICmdList, FRayTracingGeometry* InGeometry, ERTAccelerationStructureBuildPriority InPriority)
+	{
+		return RequestBuildAccelerationStructure(RHICmdList, InGeometry, InPriority, EAccelerationStructureBuildMode::Build);
+	}	
+	RENDERCORE_API BuildRequestIndex RequestBuildAccelerationStructure(FRHICommandList& RHICmdList, FRayTracingGeometry* InGeometry, ERTAccelerationStructureBuildPriority InPriority, EAccelerationStructureBuildMode InBuildMode);
 
+	UE_DEPRECATED(5.4, "RequestBuildAccelerationStructure requires a command list.")
 	BuildRequestIndex RequestBuildAccelerationStructure(FRayTracingGeometry* InGeometry, ERTAccelerationStructureBuildPriority InPriority)
 	{
-		return RequestBuildAccelerationStructure(InGeometry, InPriority, EAccelerationStructureBuildMode::Build);
-	}	
-	RENDERCORE_API BuildRequestIndex RequestBuildAccelerationStructure(FRayTracingGeometry* InGeometry, ERTAccelerationStructureBuildPriority InPriority, EAccelerationStructureBuildMode InBuildMode);
+		return RequestBuildAccelerationStructure(FRHICommandListImmediate::Get(), InGeometry, InPriority, EAccelerationStructureBuildMode::Build);
+	}
+
+	UE_DEPRECATED(5.4, "RequestBuildAccelerationStructure requires a command list.")
+	BuildRequestIndex RequestBuildAccelerationStructure(FRayTracingGeometry* InGeometry, ERTAccelerationStructureBuildPriority InPriority, EAccelerationStructureBuildMode InBuildMode)
+	{
+		return RequestBuildAccelerationStructure(FRHICommandListImmediate::Get(), InGeometry, InPriority, InBuildMode);
+	}
 
 	RENDERCORE_API void RemoveBuildRequest(BuildRequestIndex InRequestIndex);
 	RENDERCORE_API void BoostPriority(BuildRequestIndex InRequestIndex, float InBoostValue);
 	RENDERCORE_API void ForceBuildIfPending(FRHIComputeCommandList& InCmdList, const TArrayView<const FRayTracingGeometry*> InGeometries);
 	RENDERCORE_API void ProcessBuildRequests(FRHIComputeCommandList& InCmdList, bool bInBuildAll = false);
 
-	RENDERCORE_API void Tick(bool bHasRayTracingEnableChanged);
+	RENDERCORE_API void Tick(FRHICommandList& RHICmdList, bool bHasRayTracingEnableChanged);
+
+	UE_DEPRECATED(5.4, "Tick requires a command list.")
+	void Tick(bool bHasRayTracingEnableChanged)
+	{
+		Tick(FRHICommandListImmediate::Get(), bHasRayTracingEnableChanged);
+	}
 
 	RENDERCORE_API RayTracingGeometryHandle RegisterRayTracingGeometry(FRayTracingGeometry* InGeometry);
 	RENDERCORE_API void ReleaseRayTracingGeometryHandle(RayTracingGeometryHandle Handle);
diff --git a/Engine/Source/Runtime/RenderCore/Public/RenderResource.h b/Engine/Source/Runtime/RenderCore/Public/RenderResource.h
--- a/Engine/Source/Runtime/RenderCore/Public/RenderResource.h
+++ b/Engine/Source/Runtime/RenderCore/Public/RenderResource.h
@@ -39,9 +39,6 @@
 class FRDGPooledBuffer;
 class FResourceArrayInterface;
 
-/** Experimental: whether we free helper structures after submitting to RHI. */
-extern RENDERCORE_API bool GFreeStructuresOnRHIBufferCreation;
-
 enum class ERenderResourceState : uint8
 {
 	Default,
@@ -134,7 +131,7 @@ public:
 	RENDERCORE_API virtual void InitResource(FRHICommandListBase& RHICmdList);
 
 	UE_DEPRECATED(5.3, "InitResource now requires a command list.")
-	virtual void InitResource() final { InitResource(GetCommandList()); }
+	virtual void InitResource() final { InitResource(GetImmediateCommandList()); }
 
 	/**
 	 * Prepares the resource for deletion.
@@ -170,18 +167,15 @@ protected:
 	FORCEINLINE bool HasValidFeatureLevel() const { return FeatureLevel < ERHIFeatureLevel::Num; }
 
 	// Helper for submitting a resource array to RHI and freeing eligible CPU memory
-	template<bool bRenderThread, typename T>
-	FBufferRHIRef CreateRHIBuffer(T& InOutResourceObject, const uint32 ResourceCount, EBufferUsageFlags InBufferUsageFlags, const TCHAR* InDebugName)
+	template<typename T>
+	FBufferRHIRef CreateRHIBuffer(FRHICommandListBase& RHICmdList, T& InOutResourceObject, uint32 ResourceCount, EBufferUsageFlags InBufferUsageFlags, const TCHAR* InDebugName)
 	{
 		FBufferRHIRef Buffer;
+
 		FResourceArrayInterface* RESTRICT ResourceArray = InOutResourceObject ? InOutResourceObject->GetResourceArray() : nullptr;
 		if (ResourceCount != 0)
 		{
-			Buffer = CreateRHIBufferInternal(InDebugName, GetOwnerName(), ResourceCount, InBufferUsageFlags, ResourceArray, bRenderThread, InOutResourceObject == nullptr);
-			if (!bRenderThread)
-			{
-				return Buffer;
-			}
+			Buffer = CreateRHIBufferInternal(RHICmdList, InDebugName, GetOwnerName(), ResourceCount, InBufferUsageFlags, ResourceArray, InOutResourceObject == nullptr);
 		}
 
 		// If the buffer creation emptied the resource array, delete the containing structure as well
@@ -194,7 +188,7 @@ protected:
 		return Buffer;
 	}
 
-	static RENDERCORE_API FRHICommandListBase& GetCommandList();
+	static RENDERCORE_API FRHICommandListBase& GetImmediateCommandList();
 
 	void SetInitPhase(EInitPhase InInitPhase)
 	{
@@ -206,12 +200,12 @@ protected:
 private:
 	static RENDERCORE_API bool ShouldFreeResourceObject(void* ResourceObject, FResourceArrayInterface* ResourceArray);
 	static RENDERCORE_API FBufferRHIRef CreateRHIBufferInternal(
+		FRHICommandListBase& RHICmdList,
 		const TCHAR* InDebugName,
 		const FName& InOwnerName,
 		uint32 ResourceCount,
 		EBufferUsageFlags InBufferUsageFlags,
 		FResourceArrayInterface* ResourceArray,
-		bool bRenderThread,
 		bool bWithoutNativeResource
 	);
 
@@ -561,7 +555,7 @@ private:
 		if (IsInRenderingThread())
 		{
 			// If the resource is constructed in the rendering thread, directly initialize it.
-			((ResourceType*)this)->InitResource(FRenderResource::GetCommandList());
+			((ResourceType*)this)->InitResource(FRenderResource::GetImmediateCommandList());
 		}
 		else
 		{
diff --git a/Engine/Source/Runtime/RenderCore/Public/UniformBuffer.h b/Engine/Source/Runtime/RenderCore/Public/UniformBuffer.h
--- a/Engine/Source/Runtime/RenderCore/Public/UniformBuffer.h
+++ b/Engine/Source/Runtime/RenderCore/Public/UniformBuffer.h
@@ -46,20 +46,33 @@ public:
 	}
 
 	/** Sets the contents of the uniform buffer. */
-	void SetContents(const TBufferStruct& NewContents)
+	void SetContents(FRHICommandListBase& RHICmdList, const TBufferStruct& NewContents)
 	{
 		SetContentsNoUpdate(NewContents);
-		UpdateRHI(FRenderResource::GetCommandList());
+		UpdateRHI(RHICmdList);
+	}
+
+	UE_DEPRECATED(5.4, "SetContents requires a command list.")
+	void SetContents(const TBufferStruct& NewContents)
+	{
+		SetContents(FRenderResource::GetImmediateCommandList(), NewContents);
 	}
+
 	/** Sets the contents of the uniform buffer to all zeros. */
-	void SetContentsToZero()
+	void SetContentsToZero(FRHICommandListBase& RHICmdList)
 	{
 		if (!Contents)
 		{
 			Contents = (uint8*)FMemory::Malloc(sizeof(TBufferStruct), SHADER_PARAMETER_STRUCT_ALIGNMENT);
 		}
 		FMemory::Memzero(Contents, sizeof(TBufferStruct));
-		UpdateRHI(FRenderResource::GetCommandList());
+		UpdateRHI(RHICmdList);
+	}
+
+	UE_DEPRECATED(5.4, "SetContentsToZero requires a command list.")
+	void SetContentsToZero()
+	{
+		SetContentsToZero(FRenderResource::GetImmediateCommandList());
 	}
 
 	const uint8* GetContents() const 
@@ -104,7 +117,6 @@ protected:
 	/** Sets the contents of the uniform buffer. Used within calls to InitRHI */
 	void SetContentsNoUpdate(const TBufferStruct& NewContents)
 	{
-		check(IsInRenderingThread());
 		if (!Contents)
 		{
 			Contents = (uint8*)FMemory::Malloc(sizeof(TBufferStruct), SHADER_PARAMETER_STRUCT_ALIGNMENT);
@@ -130,6 +142,6 @@ void BeginSetUniformBufferContents(
 	ENQUEUE_RENDER_COMMAND(SetUniformBufferContents)(
 		[UniformBufferPtr, Struct](FRHICommandListImmediate& RHICmdList)
 		{
-			UniformBufferPtr->SetContents(Struct);
+			UniformBufferPtr->SetContents(RHICmdList, Struct);
 		});
 }
diff --git a/Engine/Source/Runtime/Renderer/Private/DeferredShadingRenderer.cpp b/Engine/Source/Runtime/Renderer/Private/DeferredShadingRenderer.cpp
--- a/Engine/Source/Runtime/Renderer/Private/DeferredShadingRenderer.cpp
+++ b/Engine/Source/Runtime/Renderer/Private/DeferredShadingRenderer.cpp
@@ -2588,7 +2588,7 @@ void FDeferredShadingSceneRenderer::Render(FRDGBuilder& GraphBuilder)
 		// Initialize ray tracing flags, in case they weren't initialized in the CreateSceneRenderers code path
 		InitializeRayTracingFlags_RenderThread();
 
-		GRayTracingGeometryManager.Tick(bHasRayTracingEnableChanged);
+		GRayTracingGeometryManager.Tick(GraphBuilder.RHICmdList, bHasRayTracingEnableChanged);
 
 		if ((GetRayTracingMode() == ERayTracingMode::Dynamic) && bHasRayTracingEnableChanged)
 		{
diff --git a/Engine/Source/Runtime/Renderer/Private/RendererScene.cpp b/Engine/Source/Runtime/Renderer/Private/RendererScene.cpp
--- a/Engine/Source/Runtime/Renderer/Private/RendererScene.cpp
+++ b/Engine/Source/Runtime/Renderer/Private/RendererScene.cpp
@@ -6167,7 +6167,7 @@ void FScene::UpdateAllPrimitiveSceneInfos(FRDGBuilder& GraphBuilder, EUpdateAllP
 			}
 
 			// Update the Proxy's data.
-			PrimitiveSceneProxy->UpdateInstances_RenderThread(UpdateInstance.Value.CmdBuffer, UpdateInstance.Value.WorldBounds, UpdateInstance.Value.LocalBounds, UpdateInstance.Value.StaticMeshBounds);
+			PrimitiveSceneProxy->UpdateInstances_RenderThread(GraphBuilder.RHICmdList, UpdateInstance.Value.CmdBuffer, UpdateInstance.Value.WorldBounds, UpdateInstance.Value.LocalBounds, UpdateInstance.Value.StaticMeshBounds);
 
 			if (!RHISupportsVolumeTextures(GetFeatureLevel())
 				&& (PrimitiveSceneProxy->IsMovable() || PrimitiveSceneProxy->NeedsUnbuiltPreviewLighting() || PrimitiveSceneProxy->GetLightmapType() == ELightmapType::ForceVolumetric))
diff --git a/Engine/Source/Runtime/Renderer/Private/ScenePrivate.h b/Engine/Source/Runtime/Renderer/Private/ScenePrivate.h
--- a/Engine/Source/Runtime/Renderer/Private/ScenePrivate.h
+++ b/Engine/Source/Runtime/Renderer/Private/ScenePrivate.h
@@ -396,7 +396,7 @@ public:
 		FDistanceCullFadeUniformShaderParameters Parameters;
 		Parameters.FadeTimeScaleBias.X = 0.0f;
 		Parameters.FadeTimeScaleBias.Y = 1.0f;
-		SetContents(Parameters);
+		SetContents(FRenderResource::GetImmediateCommandList(), Parameters);
 	}
 };
 
@@ -413,7 +413,7 @@ public:
 	{
 		FDitherUniformShaderParameters Parameters;
 		Parameters.LODFactor = 0.0f;
-		SetContents(Parameters);
+		SetContents(FRenderResource::GetImmediateCommandList(), Parameters);
 	}
 };
 
diff --git a/Engine/Source/Runtime/Renderer/Private/VT/RuntimeVirtualTextureRender.cpp b/Engine/Source/Runtime/Renderer/Private/VT/RuntimeVirtualTextureRender.cpp
--- a/Engine/Source/Runtime/Renderer/Private/VT/RuntimeVirtualTextureRender.cpp
+++ b/Engine/Source/Runtime/Renderer/Private/VT/RuntimeVirtualTextureRender.cpp
@@ -82,7 +82,7 @@ namespace RuntimeVirtualTexture
 			Parameters.RGB_DISTANCE_TABLES[6] = FVector4f(-106, -33, 33, 106);
 			Parameters.RGB_DISTANCE_TABLES[7] = FVector4f(-183, -47, 47, 183);
 
-			SetContents(Parameters);
+			SetContentsNoUpdate(Parameters);
 		}
 	};
 
diff --git a/Engine/Source/Runtime/SlateRHIRenderer/Private/SlateElementVertexBuffer.h b/Engine/Source/Runtime/SlateRHIRenderer/Private/SlateElementVertexBuffer.h
--- a/Engine/Source/Runtime/SlateRHIRenderer/Private/SlateElementVertexBuffer.h
+++ b/Engine/Source/Runtime/SlateRHIRenderer/Private/SlateElementVertexBuffer.h
@@ -34,7 +34,7 @@ public:
 
 		if ( IsInRenderingThread() )
 		{
-			InitResource(FRenderResource::GetCommandList());
+			InitResource(FRHICommandListImmediate::Get());
 		}
 		else
 		{
