Add ENGINE_API to PreLoginAsync.

diff --git a/Engine/Source/Runtime/Engine/Classes/GameFramework/GameModeBase.h b/Engine/Source/Runtime/Engine/Classes/GameFramework/GameModeBase.h
--- a/Engine/Source/Runtime/Engine/Classes/GameFramework/GameModeBase.h
+++ b/Engine/Source/Runtime/Engine/Classes/GameFramework/GameModeBase.h
@@ -296,7 +296,7 @@ public:
 	 * By default, this calls PreLogin() and OnComplete immediately to not break backwards compatibility. When overriding this, be sure
 	 * checks in PreLogin are also run, and that the GameModePreLoginEvent is broadcast.
 	 */
-	virtual void PreLoginAsync(const FString& Options, const FString& Address, const FUniqueNetIdRepl& UniqueId, const FOnPreLoginCompleteDelegate& OnComplete);
+	ENGINE_API virtual void PreLoginAsync(const FString& Options, const FString& Address, const FUniqueNetIdRepl& UniqueId, const FOnPreLoginCompleteDelegate& OnComplete);
 
 	/**
 	 * Called to login new players by creating a player controller, overridable by the game
