[Engine]
* Fixed stress test compile (every module is only one unity file)

diff --git a/Engine/Source/Runtime/Engine/Private/Materials/MaterialInterface.cpp b/Engine/Source/Runtime/Engine/Private/Materials/MaterialInterface.cpp
--- a/Engine/Source/Runtime/Engine/Private/Materials/MaterialInterface.cpp
+++ b/Engine/Source/Runtime/Engine/Private/Materials/MaterialInterface.cpp
@@ -53,7 +53,7 @@ namespace MaterialCookStats
 {
 	static int32 NumMaterialsCooked = 0;
 
-	static FCookStatsManager::FAutoRegisterCallback RegisterCookStats([](FCookStatsManager::AddStatFuncRef AddStat)
+	static FCookStatsManager::FAutoRegisterCallback RegisterCookStats2([](FCookStatsManager::AddStatFuncRef AddStat)
 	{
 		AddStat(TEXT("Material"), FCookStatsManager::CreateKeyValueArray(
 			TEXT("NumMaterialsCooked"), NumMaterialsCooked
