[WorldPartition]
* Fixed compile error when running stresstest run

diff --git a/Engine/Source/Runtime/Engine/Private/WorldPartition/WorldPartitionRuntimeSpatialHash.cpp b/Engine/Source/Runtime/Engine/Private/WorldPartition/WorldPartitionRuntimeSpatialHash.cpp
--- a/Engine/Source/Runtime/Engine/Private/WorldPartition/WorldPartitionRuntimeSpatialHash.cpp
+++ b/Engine/Source/Runtime/Engine/Private/WorldPartition/WorldPartitionRuntimeSpatialHash.cpp
@@ -630,7 +630,7 @@ EWorldPartitionRuntimeCellVisualizeMode FSpatialHashStreamingGrid::GetStreamingC
 	return VisualizeMode;
 }
 
-static TMap<FName, FColor> GetDataLayerDebugColors(const UWorldPartition* InWorldPartition)
+static TMap<FName, FColor> GetDataLayerDebugColors2(const UWorldPartition* InWorldPartition)
 {
 	TMap<FName, FColor> DebugColors;
 	if (const UDataLayerManager* DataLayerManager = InWorldPartition->GetDataLayerManager())
@@ -650,7 +650,7 @@ void FSpatialHashStreamingGrid::Draw3D(const UWorldPartitionRuntimeSpatialHash*
 	UWorld* OwningWorld = WorldPartition->GetWorld();
 	const EWorldPartitionRuntimeCellVisualizeMode VisualizeMode = GetStreamingCellVisualizeMode();
 	const UContentBundleManager* ContentBundleManager = OwningWorld->ContentBundleManager;
-	TMap<FName, FColor> DataLayerDebugColors = GetDataLayerDebugColors(WorldPartition);
+	TMap<FName, FColor> DataLayerDebugColors = GetDataLayerDebugColors2(WorldPartition);
 
 	const FSquare2DGridHelper& Helper = GetGridHelper();
 	int32 MinGridLevel = FMath::Clamp<int32>(GShowRuntimeSpatialHashGridLevel, 0, GridLevels.Num() - 1);
@@ -791,7 +791,7 @@ void FSpatialHashStreamingGrid::Draw2D(const UWorldPartitionRuntimeSpatialHash*
 	const UWorld* OwningWorld = WorldPartition->GetWorld();
 	const EWorldPartitionRuntimeCellVisualizeMode VisualizeMode = GetStreamingCellVisualizeMode();
 	const UContentBundleManager* ContentBundleManager = OwningWorld->ContentBundleManager;
-	TMap<FName, FColor> DataLayerDebugColors = GetDataLayerDebugColors(WorldPartition);
+	TMap<FName, FColor> DataLayerDebugColors = GetDataLayerDebugColors2(WorldPartition);
 
 	int32 MinGridLevel = FMath::Clamp<int32>(GShowRuntimeSpatialHashGridLevel, 0, GridLevels.Num() - 1);
 	int32 MaxGridLevel = FMath::Clamp<int32>(MinGridLevel + GShowRuntimeSpatialHashGridLevelCount - 1, 0, GridLevels.Num() - 1);
