PLAY-7816 - GetNetMode crash in PT

Cache NetMode in a GT aligned callback to be used in PT instead of calling World->GetNetMode from PT.

diff --git a/Engine/Source/Runtime/Engine/Private/PhysicsEngine/NetworkPhysicsComponent.cpp b/Engine/Source/Runtime/Engine/Private/PhysicsEngine/NetworkPhysicsComponent.cpp
--- a/Engine/Source/Runtime/Engine/Private/PhysicsEngine/NetworkPhysicsComponent.cpp
+++ b/Engine/Source/Runtime/Engine/Private/PhysicsEngine/NetworkPhysicsComponent.cpp
@@ -83,8 +83,9 @@ float DeQuantizeTimeDilation(int8 i)
 void FNetworkPhysicsCallback::ApplyCallbacks_Internal(int32 PhysicsStep, const TArray<Chaos::ISimCallbackObject*>& SimCallbackObjects)
 {
 	QUICK_SCOPE_CYCLE_COUNTER(NetworkPhysicsComponent_ApplyCallbacks_Internal);
+	UpdateNetMode();
 
-	if(World->IsNetMode(NM_ListenServer) || World->IsNetMode(NM_DedicatedServer))
+	if ((NetMode == NM_ListenServer) || (NetMode == NM_DedicatedServer))
 	{
 		if (FPhysScene_Chaos* Scene = static_cast<FPhysScene_Chaos*>(World->GetPhysicsScene()))
 		{
@@ -119,7 +120,7 @@ int32 FNetworkPhysicsCallback::TriggerRewindIfNeeded_Internal(int32 LatestStepCo
 
 	if (RewindData)
 	{
-		if (World->GetNetMode() == NM_Client)
+		if (NetMode == NM_Client)
 		{
 			const int32 ReplicationFrame = RewindData->GetResimFrame();
 
@@ -267,7 +268,7 @@ void FNetworkPhysicsCallback::ProcessInputs_External(int32 PhysicsStep, const TA
 		}
 	}
 
-	if (World->GetNetMode() == NM_Client)
+	if (NetMode == NM_Client)
 	{
 		UpdateClientPlayer_External(PhysicsStep);
 	}
diff --git a/Engine/Source/Runtime/Engine/Public/Physics/NetworkPhysicsComponent.h b/Engine/Source/Runtime/Engine/Public/Physics/NetworkPhysicsComponent.h
--- a/Engine/Source/Runtime/Engine/Public/Physics/NetworkPhysicsComponent.h
+++ b/Engine/Source/Runtime/Engine/Public/Physics/NetworkPhysicsComponent.h
@@ -165,7 +165,10 @@ private :
  */
 struct FNetworkPhysicsCallback : public Chaos::IRewindCallback
 {
-	FNetworkPhysicsCallback(UWorld* InWorld) : World(InWorld) { }
+	FNetworkPhysicsCallback(UWorld* InWorld) : World(InWorld) 
+	{
+		UpdateNetMode();
+	}
 
 	// Delegate on the internal inputs process 
 	FOnPreProcessInputsInternal PreProcessInputsInternal;
@@ -199,9 +202,18 @@ struct FNetworkPhysicsCallback : public Chaos::IRewindCallback
 	// Update server player on GT
 	void UpdateServerPlayer_External(int32 PhysicsStep);
 
+	// Cache the current netmode for use in PT
+	void UpdateNetMode()
+	{
+		NetMode = World ? World->GetNetMode() : ENetMode::NM_Client;
+	}
+
 	// World owning that callback
 	UWorld* World = nullptr;
 
+	// Current NetMode
+	ENetMode NetMode;
+
 	// List of rewindable sim callback objects
 	TArray<Chaos::ISimCallbackObject*> RewindableCallbackObjects;
 };
