Avoid unnecessary serialization of editor-only references at cook time when analyzing Blueprint assets for external imports.

diff --git a/Engine/Source/Runtime/Engine/Private/BlueprintGeneratedClass.cpp b/Engine/Source/Runtime/Engine/Private/BlueprintGeneratedClass.cpp
--- a/Engine/Source/Runtime/Engine/Private/BlueprintGeneratedClass.cpp
+++ b/Engine/Source/Runtime/Engine/Private/BlueprintGeneratedClass.cpp
@@ -2274,6 +2274,7 @@ void UBlueprintGeneratedClass::Serialize(FArchive& Ar)
 		FImportExportCollector Collector(this->GetPackage());
 		Collector.SetCookData(Ar.GetCookData());
 		Collector.AddExportToIgnore(this);
+		Collector.SetFilterEditorOnly(Ar.IsFilterEditorOnly());
 		Collector.SerializeObjectAndReferencedExports(ClassGeneratedBy);
 		for (const TPair<FName, ESoftObjectPathCollectType>& Pair : Collector.GetImportedPackages())
 		{
