Fixed ray tracing geometry being created off the render thread with the immediate command list.

diff --git a/Engine/Source/Runtime/Engine/Private/Streaming/StaticMeshUpdate.cpp b/Engine/Source/Runtime/Engine/Private/Streaming/StaticMeshUpdate.cpp
--- a/Engine/Source/Runtime/Engine/Private/Streaming/StaticMeshUpdate.cpp
+++ b/Engine/Source/Runtime/Engine/Private/Streaming/StaticMeshUpdate.cpp
@@ -208,7 +208,11 @@ void FStaticMeshStreamIn::CreateBuffers_Internal(const FContext& Context)
 				}
 				Initializer.Type = ERayTracingGeometryInitializerType::StreamingSource;
 				IntermediateRayTracingGeometry[LODIdx].SetInitializer(Initializer);
-				IntermediateRayTracingGeometry[LODIdx].CreateRayTracingGeometryFromCPUData(FRHICommandListImmediate::Get(), LODResource.RayTracingGeometry.RawData);
+
+
+				FRHIAsyncCommandList AsyncCommandList;
+				FRHICommandList& RHICmdList = bRenderThread ? FRHICommandListImmediate::Get() : *AsyncCommandList;
+				IntermediateRayTracingGeometry[LODIdx].CreateRayTracingGeometryFromCPUData(RHICmdList, LODResource.RayTracingGeometry.RawData);
 			}
 #endif
 		}
