Fixed RHI validation errors with buffer creation.

diff --git a/Engine/Source/Runtime/Engine/Private/RawIndexBuffer.cpp b/Engine/Source/Runtime/Engine/Private/RawIndexBuffer.cpp
--- a/Engine/Source/Runtime/Engine/Private/RawIndexBuffer.cpp
+++ b/Engine/Source/Runtime/Engine/Private/RawIndexBuffer.cpp
@@ -46,7 +46,7 @@ void FRawIndexBuffer::InitRHI(FRHICommandListBase& RHICmdList)
 	{
 		// Create the index buffer.
 		FRHIResourceCreateInfo CreateInfo(TEXT("FRawIndexBuffer"));
-		IndexBufferRHI = RHICmdList.CreateBuffer(Size, BUF_Static | BUF_IndexBuffer, sizeof(uint16), ERHIAccess::VertexOrIndexBuffer, CreateInfo);
+		IndexBufferRHI = RHICmdList.CreateBuffer(Size, BUF_Static | BUF_IndexBuffer, sizeof(uint16), ERHIAccess::VertexOrIndexBuffer | ERHIAccess::SRVMask, CreateInfo);
 
 		// Initialize the buffer.
 		void* Buffer = RHICmdList.LockBuffer(IndexBufferRHI, 0, Size, RLM_WriteOnly);
@@ -106,7 +106,7 @@ void FRawIndexBuffer16or32::InitRHI(FRHICommandListBase& RHICmdList)
 	{
 		// Create the index buffer.
 		FRHIResourceCreateInfo CreateInfo(TEXT("FRawIndexBuffer"));
-		IndexBufferRHI = RHICmdList.CreateBuffer(Size, BUF_Static | BUF_IndexBuffer, IndexStride, ERHIAccess::VertexOrIndexBuffer, CreateInfo);
+		IndexBufferRHI = RHICmdList.CreateBuffer(Size, BUF_Static | BUF_IndexBuffer, IndexStride, ERHIAccess::VertexOrIndexBuffer | ERHIAccess::SRVMask, CreateInfo);
 
 		// Initialize the buffer.		
 		void* Buffer = RHICmdList.LockBuffer(IndexBufferRHI, 0, Size, RLM_WriteOnly);
@@ -346,7 +346,7 @@ FBufferRHIRef FRawStaticIndexBuffer::CreateRHIBuffer(FRHICommandListBase& RHICmd
 		CreateInfo.OwnerName = GetOwnerName();
 		CreateInfo.bWithoutNativeResource = !SizeInBytes;
 
-		FBufferRHIRef Buffer = RHICmdList.CreateBuffer(SizeInBytes, BufferFlags | EBufferUsageFlags::IndexBuffer, IndexStride, ERHIAccess::VertexOrIndexBuffer, CreateInfo);
+		FBufferRHIRef Buffer = RHICmdList.CreateBuffer(SizeInBytes, BufferFlags | EBufferUsageFlags::IndexBuffer, IndexStride, ERHIAccess::VertexOrIndexBuffer | ERHIAccess::SRVMask, CreateInfo);
 		Buffer->SetOwnerName(GetOwnerName());
 		return Buffer;
 	}
@@ -495,7 +495,7 @@ FBufferRHIRef FRawStaticIndexBuffer16or32Interface::CreateRHIIndexBufferInternal
 	const uint32 Size = IndexCount * IndexSize;
 	CreateInfo.bWithoutNativeResource = !Size;
 
-	FBufferRHIRef Buffer = RHICmdList.CreateBuffer(Size, Flags | EBufferUsageFlags::IndexBuffer, IndexSize, ERHIAccess::VertexOrIndexBuffer, CreateInfo);
+	FBufferRHIRef Buffer = RHICmdList.CreateBuffer(Size, Flags | EBufferUsageFlags::IndexBuffer, IndexSize, ERHIAccess::VertexOrIndexBuffer | ERHIAccess::SRVMask, CreateInfo);
 	Buffer->SetOwnerName(InOwnerName);
 	return Buffer;
 }
diff --git a/Engine/Source/Runtime/RenderCore/Private/RenderResource.cpp b/Engine/Source/Runtime/RenderCore/Private/RenderResource.cpp
--- a/Engine/Source/Runtime/RenderCore/Private/RenderResource.cpp
+++ b/Engine/Source/Runtime/RenderCore/Private/RenderResource.cpp
@@ -269,7 +269,7 @@ FBufferRHIRef FRenderResource::CreateRHIBufferInternal(
 	CreateInfo.OwnerName = InOwnerName;
 	CreateInfo.bWithoutNativeResource = bWithoutNativeResource;
 
-	FBufferRHIRef Buffer = RHICmdList.CreateBuffer(SizeInBytes, InBufferUsageFlags | EBufferUsageFlags::VertexBuffer, 0, ERHIAccess::VertexOrIndexBuffer, CreateInfo);
+	FBufferRHIRef Buffer = RHICmdList.CreateBuffer(SizeInBytes, InBufferUsageFlags | EBufferUsageFlags::VertexBuffer, 0, ERHIAccess::VertexOrIndexBuffer | ERHIAccess::SRVMask, CreateInfo);
 
 	Buffer->SetOwnerName(InOwnerName);
 	return Buffer;
