Fix for Material Instances resolving their parent material's name in debug data and GPU captures, instead of the instance name. Altered the default MID naming convention to be clearer and easier to follow from debug data/captures from e.g. MaterialInstanceDynamic450 to MID_ParentMaterialName_450.

diff --git a/Engine/Source/Runtime/Engine/Classes/Materials/MaterialInstanceDynamic.h b/Engine/Source/Runtime/Engine/Classes/Materials/MaterialInstanceDynamic.h
--- a/Engine/Source/Runtime/Engine/Classes/Materials/MaterialInstanceDynamic.h
+++ b/Engine/Source/Runtime/Engine/Classes/Materials/MaterialInstanceDynamic.h
@@ -153,14 +153,9 @@ class UMaterialInstanceDynamic : public UMaterialInstance
 	ENGINE_API void CopyInterpParameters(UMaterialInstance* Source);
 
 	/**
-	 * Create a material instance dynamic parented to the specified material.
-	 */
-	static ENGINE_API UMaterialInstanceDynamic* Create(class UMaterialInterface* ParentMaterial, class UObject* InOuter);
-
-	/**
-	* Create a material instance dynamic parented to the specified material with the specified name.
+	* Create a material instance dynamic parented to the specified material. [Optional] Specify name.
 	*/
-	static ENGINE_API UMaterialInstanceDynamic* Create( class UMaterialInterface* ParentMaterial, class UObject* InOuter, FName Name );
+	static ENGINE_API UMaterialInstanceDynamic* Create( class UMaterialInterface* ParentMaterial, class UObject* InOuter, FName Name = NAME_None);
 
 	/**
 	 * Set the value of the given font parameter.  
diff --git a/Engine/Source/Runtime/Engine/Private/Materials/MaterialInstanceDynamic.cpp b/Engine/Source/Runtime/Engine/Private/Materials/MaterialInstanceDynamic.cpp
--- a/Engine/Source/Runtime/Engine/Private/Materials/MaterialInstanceDynamic.cpp
+++ b/Engine/Source/Runtime/Engine/Private/Materials/MaterialInstanceDynamic.cpp
@@ -68,15 +68,6 @@ void UMaterialInstanceDynamic::InitializeMID(class UMaterialInterface* ParentMat
 	UpdateCachedDataDynamic();
 }
 
-UMaterialInstanceDynamic* UMaterialInstanceDynamic::Create(UMaterialInterface* ParentMaterial, UObject* InOuter)
-{
-	LLM_SCOPE(ELLMTag::MaterialInstance);
-	UObject* Outer = InOuter ? InOuter : GetTransientPackage();
-	UMaterialInstanceDynamic* MID = NewObject<UMaterialInstanceDynamic>(Outer);
-	MID->InitializeMID(ParentMaterial);
-	return MID;
-}
-
 UMaterialInstanceDynamic* UMaterialInstanceDynamic::Create(UMaterialInterface* ParentMaterial, UObject* InOuter, FName Name)
 {
 	LLM_SCOPE(ELLMTag::MaterialInstance);
@@ -84,39 +75,42 @@ UMaterialInstanceDynamic* UMaterialInstanceDynamic::Create(UMaterialInterface* P
 	UPackage* TransientPackage = GetTransientPackage();
 
 	UObject* Outer = InOuter ? InOuter : TransientPackage;
-	if (Name != NAME_None)
+
+	if (Name == NAME_None)
 	{
-		UMaterialInstanceDynamic* ExistingMID = FindObjectFast<UMaterialInstanceDynamic>(Outer, *Name.ToString(), true);
-		if (ExistingMID)
-		{
-			bool bRenamed = false;
+		Name = FName(FString("MID_") + ParentMaterial->GetName());
+	}
+
+	UMaterialInstanceDynamic* ExistingMID = FindObjectFast<UMaterialInstanceDynamic>(Outer, *Name.ToString(), true);
+	if (ExistingMID)
+	{
+		bool bRenamed = false;
 
 #if WITH_EDITOR
-			// try to rename to enhance construction script determinism:
-			AActor* OwningActor = Cast<AActor>(Outer);
-			if (!OwningActor)
-			{
-				OwningActor = Outer->GetTypedOuter<AActor>();
-			}
+		// try to rename to enhance construction script determinism:
+		AActor* OwningActor = Cast<AActor>(Outer);
+		if (!OwningActor)
+		{
+			OwningActor = Outer->GetTypedOuter<AActor>();
+		}
 
-			if (OwningActor && OwningActor->IsRunningUserConstructionScript())
-			{
-				bRenamed = true;
-				// a collision, we're going to move this existing mid to the transient package and claim the name 
-				// for ourself:
-				ExistingMID	->Rename(
-					nullptr,
-					TransientPackage,
-					REN_DoNotDirty | REN_DontCreateRedirectors | REN_ForceNoResetLoaders | REN_NonTransactional
-				);
-			}
+		if (OwningActor && OwningActor->IsRunningUserConstructionScript())
+		{
+			bRenamed = true;
+			// a collision, we're going to move this existing mid to the transient package and claim the name 
+			// for ourself:
+			ExistingMID	->Rename(
+				nullptr,
+				TransientPackage,
+				REN_DoNotDirty | REN_DontCreateRedirectors | REN_ForceNoResetLoaders | REN_NonTransactional
+			);
+		}
 #endif
 
-			if (!bRenamed)
-			{
-				// If a MID is made with the same name and outer as another, it will overwrite it. To avoid this we will change the name when there is a collision.
-				Name = MakeUniqueObjectName(Outer, UMaterialInstanceDynamic::StaticClass(), Name);
-			}
+		if (!bRenamed)
+		{
+			// If a MID is made with the same name and outer as another, it will overwrite it. To avoid this we will change the name when there is a collision.
+			Name = MakeUniqueObjectName(Outer, UMaterialInstanceDynamic::StaticClass(), Name);
 		}
 	}
 	UMaterialInstanceDynamic* MID = NewObject<UMaterialInstanceDynamic>(Outer, Name);
diff --git a/Engine/Source/Runtime/Renderer/Private/MeshPassProcessor.cpp b/Engine/Source/Runtime/Renderer/Private/MeshPassProcessor.cpp
--- a/Engine/Source/Runtime/Renderer/Private/MeshPassProcessor.cpp
+++ b/Engine/Source/Runtime/Renderer/Private/MeshPassProcessor.cpp
@@ -1704,7 +1704,7 @@ void FMeshDrawCommand::SetDebugData(const FPrimitiveSceneProxy* PrimitiveScenePr
 	DebugData.VertexFactoryType = VertexFactory->GetType();
 	DebugData.MeshPassType = MeshPassType;
 	DebugData.ResourceName =  PrimitiveSceneProxy ? PrimitiveSceneProxy->GetResourceName() : FName();
-	DebugData.MaterialName = Material->GetAssetName();
+	DebugData.MaterialName = MaterialRenderProxy->GetMaterialName();
 }
 #endif
 
