better particle state management for cluster unions - no longer relies on needing to pass Chaos::FPhysicsObjectHandles back and forth from the GT to the PT which would result in the possibility of dead particles still being bookkept by the cluster union

- Cluster union pull data now passes proxy and BoneId from the PT to the GT. This proxy is used to convert to a component using the scene's owning component mapping.
- Remove the need to book-keep PhysicsObjects and AllPhysicsObjects on the cluster union component. When needed, will use BoneIds - otherwise rely on the PT to cleanup itself.

diff --git a/Engine/Source/Runtime/Engine/Classes/PhysicsEngine/ClusterUnionComponent.h b/Engine/Source/Runtime/Engine/Classes/PhysicsEngine/ClusterUnionComponent.h
--- a/Engine/Source/Runtime/Engine/Classes/PhysicsEngine/ClusterUnionComponent.h
+++ b/Engine/Source/Runtime/Engine/Classes/PhysicsEngine/ClusterUnionComponent.h
@@ -36,15 +36,9 @@ struct FClusteredComponentData
 {
 	GENERATED_BODY()
 
-	// Set of physics objects that we actually added into the cluster union.
-	TSet<Chaos::FPhysicsObjectHandle> PhysicsObjects;
-
 	// Set of bone Ids that we actually added into the cluster union.
 	TSet<int32> BoneIds;
 
-	// Every physics object associated with this particular component.
-	TArray<Chaos::FPhysicsObjectHandle> AllPhysicsObjects;
-
 	// Cached acceleration structure handles - needed to properly cleanup the component from the accel structure.
 	TSet<FExternalSpatialAccelerationPayload> CachedAccelerationPayloads;
 
@@ -53,6 +47,9 @@ struct FClusteredComponentData
 	UPROPERTY()
 	TWeakObjectPtr<UClusterUnionReplicatedProxyComponent> ReplicatedProxyComponent;
 
+	UPROPERTY()
+	TWeakObjectPtr<AActor> Owner;
+
 	UPROPERTY()
 	bool bWasReplicating = true;
 
@@ -65,8 +62,7 @@ struct FClusteredActorData
 {
 	GENERATED_BODY()
 
-	UPROPERTY()
-	TSet<TWeakObjectPtr<UPrimitiveComponent>> Components;
+	TSet<TObjectKey<UPrimitiveComponent>> Components;
 
 	UPROPERTY()
 	bool bWasReplicatingMovement = true;
@@ -96,22 +92,6 @@ struct FClusterUnionPendingAddData
 	TArray<FExternalSpatialAccelerationPayload> AccelerationPayloads;
 };
 
-/**
- * For every possible particle that could ever possibly be added into the cluster union,
- * keep track of its component and its bone id.
- */
-USTRUCT()
-struct FClusterUnionParticleCandidateData
-{
-	GENERATED_BODY()
-
-	UPROPERTY()
-	TWeakObjectPtr<UPrimitiveComponent> Component;
-
-	UPROPERTY()
-	int32 BoneId = INDEX_NONE;
-};
-
 DECLARE_DYNAMIC_MULTICAST_DELEGATE_ThreeParams(FOnClusterUnionAddedComponent, UPrimitiveComponent*, Component, const TSet<int32>&, BoneIds, bool, bIsNew);
 DECLARE_DYNAMIC_MULTICAST_DELEGATE_OneParam(FOnClusterUnionRemovedComponent, UPrimitiveComponent*, Component);
 DECLARE_DYNAMIC_MULTICAST_DELEGATE_TwoParams(FOnClusterUnionBoundsChanged, UClusterUnionComponent*, Component, const FBoxSphereBounds&, Bounds);
@@ -165,7 +145,7 @@ public:
 	ENGINE_API void SyncClusterUnionFromProxy();
 
 	UFUNCTION()
-	bool IsComponentAdded(UPrimitiveComponent* Component) { return ComponentToPhysicsObjects.Contains(Component) || PendingComponentSync.Contains(Component); }
+	bool IsComponentAdded(UPrimitiveComponent* Component) { return PerComponentData.Contains(Component) || PendingComponentSync.Contains(Component); }
 
 	bool HasReceivedTransform() const { return bHasReceivedTransform; }
 
@@ -207,7 +187,7 @@ private:
 	// This way we know the right physics objects to pass when removing the component (because
 	// it's possible to get a different list of physics objects when we get to removal). A
 	// side benefit here is being able to track which components are clustered.
-	TMap<TObjectKey<UPrimitiveComponent>, FClusteredComponentData> ComponentToPhysicsObjects;
+	TMap<TObjectKey<UPrimitiveComponent>, FClusteredComponentData> PerComponentData;
 
 	// Also keep track of which actors we are clustering and their components. We make modifications on
 	// actors that get clustered so we need to make sure we undo those changes only once all its clustered
@@ -224,12 +204,6 @@ private:
 	// where the component hasn't been added to the cluster union on the GT causing a mismatch in behavior.
 	TMap<TObjectKey<UPrimitiveComponent>, FClusterUnionPendingAddData> PendingComponentSync;
 
-	// Given a unique index of a particle that we're adding to the cluster union - map it back to the component that owns it.
-	// This works decently because we assume that when we're using a cluster union component, we will only try to add to the
-	// cluster union via the GT so we can guarantee to have a decent mapping here.
-	UPROPERTY()
-	TMap<int32, FClusterUnionParticleCandidateData> UniqueIdxToComponent;
-
 	// Data that can be changed at runtime to keep state about the cluster union consistent between the server and client.
 	UPROPERTY(ReplicatedUsing=OnRep_RigidState)
 	FClusterUnionReplicatedData ReplicatedRigidState;
@@ -267,7 +241,7 @@ private:
 	// These functions only get called when the physics thread syncs to the game thread thereby enforcing a physics thread authoritative view of
 	// what particles are currently contained within the cluster union.
 	ENGINE_API void HandleAddOrModifiedClusteredComponent(UPrimitiveComponent* ChangedComponent, const TMap<int32, FTransform>& PerBoneChildToParent);
-	ENGINE_API void HandleRemovedClusteredComponent(UPrimitiveComponent* ChangedComponent, bool bDestroyReplicatedProxy);
+	ENGINE_API void HandleRemovedClusteredComponent(TObjectKey<UPrimitiveComponent> ChangedComponent);
 
 	ENGINE_API TArray<UPrimitiveComponent*> GetAllCurrentChildComponents() const;
 	ENGINE_API TArray<AActor*> GetAllCurrentActors() const;
diff --git a/Engine/Source/Runtime/Engine/Private/PhysicsEngine/ClusterUnionComponent.cpp b/Engine/Source/Runtime/Engine/Private/PhysicsEngine/ClusterUnionComponent.cpp
--- a/Engine/Source/Runtime/Engine/Private/PhysicsEngine/ClusterUnionComponent.cpp
+++ b/Engine/Source/Runtime/Engine/Private/PhysicsEngine/ClusterUnionComponent.cpp
@@ -98,20 +98,6 @@ void UClusterUnionComponent::AddComponentToCluster(UPrimitiveComponent* InCompon
 	TArray<Chaos::FPhysicsObjectHandle> AllObjects = InComponent->GetAllPhysicsObjects();
 	FLockedReadPhysicsObjectExternalInterface Interface = FPhysicsObjectExternalInterface::LockRead(AllObjects);
 
-	for (Chaos::FPhysicsObjectHandle Handle : AllObjects)
-	{
-		TArray<Chaos::FGeometryParticle*> Particles = Interface->GetAllParticles({ &Handle, 1 });
-		if (Particles.IsEmpty() || !Particles[0])
-		{
-			continue;
-		}
-
-		FClusterUnionParticleCandidateData Data;
-		Data.Component = InComponent;
-		Data.BoneId = Chaos::FPhysicsObjectInterface::GetId(Handle);
-		UniqueIdxToComponent.Add(Particles[0]->UniqueIdx().Idx, Data);
-	}
-
 	TArray<Chaos::FPhysicsObjectHandle> Objects;
 	if (BoneIds.IsEmpty())
 	{
@@ -161,6 +147,7 @@ void UClusterUnionComponent::AddComponentToCluster(UPrimitiveComponent* InCompon
 	}
 
 	// Need to listen to changes in the component's physics state. If it gets destroyed it should be removed from the cluster union as well.
+	// This technically is giving us a false sense of security because when we get this notification, the primitive component will already have no physics proxy on it so we can't actually grab physics objects/particles from it.
 	InComponent->OnComponentPhysicsStateChanged.AddUniqueDynamic(this, &UClusterUnionComponent::HandleComponentPhysicsStateChangePostAddIntoClusterUnion);
 
 	PendingComponentSync.Add(InComponent, PendingData);
@@ -185,7 +172,7 @@ void UClusterUnionComponent::RemoveComponentFromCluster(UPrimitiveComponent* InC
 
 	// If we're still waiting on the physics sync for the added component, we still need to
 	// remove its acceleration handles from the acceleration structure since the payloads have
-	// already been added and will not yet be stored in ComponentToPhysicsObjects.
+	// already been added and will not yet be stored in PerComponentData.
 	if (FClusterUnionPendingAddData* PendingData = PendingComponentSync.Find(InComponent))
 	{
 		if (AccelerationStructure)
@@ -200,14 +187,21 @@ void UClusterUnionComponent::RemoveComponentFromCluster(UPrimitiveComponent* InC
 	}
 
 	TSet<Chaos::FPhysicsObjectHandle> PhysicsObjectsToRemove;
+
 	FLockedReadPhysicsObjectExternalInterface Interface = FPhysicsObjectExternalInterface::LockRead(GetWorld()->GetPhysicsScene());
 
-	if (FClusteredComponentData* ComponentData = ComponentToPhysicsObjects.Find(InComponent))
+	if (FClusteredComponentData* ComponentData = PerComponentData.Find(InComponent))
 	{
 		// We need to mark the replicated proxy as pending deletion.
 		// This way anyone who tries to use the replicated proxy component knows that it
 		// doesn't actually denote a meaningful cluster union relationship.
 		ComponentData->bPendingDeletion = true;
+
+		if (InComponent->HasValidPhysicsState())
+		{
+			PhysicsObjectsToRemove = TSet<Chaos::FPhysicsObjectHandle>{ GetAllPhysicsObjectsById(InComponent, ComponentData->BoneIds.Array()) };
+		}
+
 		if (IsAuthority())
 		{
 			if (UClusterUnionReplicatedProxyComponent* Component = ComponentData->ReplicatedProxyComponent.Get())
@@ -216,16 +210,6 @@ void UClusterUnionComponent::RemoveComponentFromCluster(UPrimitiveComponent* InC
 			}
 		}
 
-		if (!InComponent->HasValidPhysicsState())
-		{
-			ComponentData->PhysicsObjects.Reset();
-			ComponentData->AllPhysicsObjects.Reset();
-		}
-		else
-		{
-			PhysicsObjectsToRemove = ComponentData->PhysicsObjects;
-		}
-
 		if (AccelerationStructure)
 		{
 			for (const FExternalSpatialAccelerationPayload& Payload : ComponentData->CachedAccelerationPayloads)
@@ -236,6 +220,8 @@ void UClusterUnionComponent::RemoveComponentFromCluster(UPrimitiveComponent* InC
 		}
 	}
 
+	// If PhysicsObjectsToRemove is empty, it either means that BoneIds is empty OR the component's physics state is already destroyed.
+	// In the case of the latter, we rely on the physics thread to cleanup the cluster union manager properly and to sync back.
 	if (!PhysicsObjectsToRemove.IsEmpty())
 	{
 		PhysicsProxy->RemovePhysicsObjects_External(PhysicsObjectsToRemove);
@@ -296,9 +282,9 @@ void UClusterUnionComponent::ForceRebuildGTParticleGeometry()
 TArray<UPrimitiveComponent*> UClusterUnionComponent::GetPrimitiveComponents()
 {
 	TArray<UPrimitiveComponent*> PrimitiveComponents;
-	PrimitiveComponents.Reserve(ComponentToPhysicsObjects.Num());
+	PrimitiveComponents.Reserve(PerComponentData.Num());
 
-	for (auto Iter = ComponentToPhysicsObjects.CreateIterator(); Iter; ++Iter)
+	for (auto Iter = PerComponentData.CreateIterator(); Iter; ++Iter)
 	{
 		PrimitiveComponents.Add(Iter.Key().ResolveObjectPtr());
 	}
@@ -424,18 +410,12 @@ void UClusterUnionComponent::OnDestroyPhysicsState()
 
 	// We need to make sure we *immediately* disconnect on the GT side since there's no guarantee the normal flow
 	// will happen once we've destroyed things.
-	TSet<TObjectPtr<UPrimitiveComponent>> RemainingComponents;
-	for (const TPair<TObjectKey<UPrimitiveComponent>, FClusteredComponentData>& Kvp : ComponentToPhysicsObjects)
-	{
-		RemainingComponents.Add(Kvp.Key.ResolveObjectPtr());
-	}
+	TArray<TObjectKey<UPrimitiveComponent>> ComponentsToRemove;
+	PerComponentData.GetKeys(ComponentsToRemove);
 
-	for (TObjectPtr<UPrimitiveComponent> Component : RemainingComponents)
+	for (const TObjectKey<UPrimitiveComponent>& Component : ComponentsToRemove)
 	{
-		if (Component)
-		{
-			HandleRemovedClusteredComponent(Component, false);
-		}
+		HandleRemovedClusteredComponent(Component);
 	}
 
 	if (FPhysScene_Chaos* Scene = GetChaosScene())
@@ -542,12 +522,19 @@ void UClusterUnionComponent::HandleComponentPhysicsStateChangePostAddIntoCluster
 	RemoveComponentFromCluster(ChangedComponent);
 }
 
+DECLARE_CYCLE_STAT(TEXT("UClusterUnionComponent::SyncClusterUnionFromProxy"), STAT_ClusterUnionComponent_SyncClusterUnionFromProxy, STATGROUP_Chaos);
 void UClusterUnionComponent::SyncClusterUnionFromProxy()
 {
-	//SCOPE_CYCLE_COUNTER(STAT_ClusterUnionComponent_SyncClusterUnionFromProxy);
+	SCOPE_CYCLE_COUNTER(STAT_ClusterUnionComponent_SyncClusterUnionFromProxy);
 
 	// NOTE THAT WE ARE ON THE GAME THREAD HERE.
-	if (!PhysicsProxy)
+	if (!PhysicsProxy || !GetWorld())
+	{
+		return;
+	}
+
+	FPhysScene* Scene = GetWorld()->GetPhysicsScene();
+	if (!Scene)
 	{
 		return;
 	}
@@ -560,43 +547,40 @@ void UClusterUnionComponent::SyncClusterUnionFromProxy()
 	}
 	
 	const Chaos::FClusterUnionSyncedData& FullData = PhysicsProxy->GetSyncedData_External();
-
 	// Note that at the UClusterUnionComponent level we really only want to be dealing with components.
 	// Hence why we need to modify each of the particles that we synced from the game thread into a
 	// component + bone id combination for identification. 
 	TMap<TObjectKey<UPrimitiveComponent>, TMap<int32, FTransform>> MappedData;
 	for (const Chaos::FClusterUnionChildData& ChildData : FullData.ChildParticles)
 	{
-		if (FClusterUnionParticleCandidateData* Data = UniqueIdxToComponent.Find(ChildData.ParticleIdx.Idx))
+		// Using the scene's proxy to component mapping let's us detect a component physics state was destroyed.
+		if (UPrimitiveComponent* Component = Scene->GetOwningComponent<UPrimitiveComponent>(ChildData.Proxy))
 		{
-			if (Data->Component.IsValid())
-			{
-				MappedData.FindOrAdd(Data->Component.Get()).Add(Data->BoneId, ChildData.ChildToParent);
-			}
+			MappedData.FindOrAdd(Component).Add(ChildData.BoneId, ChildData.ChildToParent);
 		}
 	}
 
 	// We need to handle any additions, deletions, and modifications to any child in the cluster union here.
-	// If a component lives in MappedData but not in ComponentToPhysicsObjects, new component!
+	// If a component lives in MappedData but not in PerComponentData, new component!
 	// If a component lives in both, then it's a modified component.
 	for (const TPair<TObjectKey<UPrimitiveComponent>, TMap<int32, FTransform>>& Kvp : MappedData)
 	{
 		HandleAddOrModifiedClusteredComponent(Kvp.Key.ResolveObjectPtr(), Kvp.Value);
 	}
 
-	// If a component lives in ComponentToPhysicsObjects but not in MappedData, deleted component!
-	TArray<TObjectPtr<UPrimitiveComponent>> ComponentsToRemove;
-	for (const TPair<TObjectKey<UPrimitiveComponent>, FClusteredComponentData>& Kvp : ComponentToPhysicsObjects)
+	// If a component lives in PerComponentData but not in MappedData, deleted component!
+	TArray<TObjectKey<UPrimitiveComponent>> ComponentsToRemove;
+	for (const TPair<TObjectKey<UPrimitiveComponent>, FClusteredComponentData>& Kvp : PerComponentData)
 	{
 		if (!MappedData.Contains(Kvp.Key))
 		{
-			ComponentsToRemove.Add(Kvp.Key.ResolveObjectPtr());
+			ComponentsToRemove.Add(Kvp.Key);
 		}
 	}
 
-	for (TObjectPtr<UPrimitiveComponent> Component : ComponentsToRemove)
+	for (TObjectKey<UPrimitiveComponent> Component : ComponentsToRemove)
 	{
-		HandleRemovedClusteredComponent(Component, true);
+		HandleRemovedClusteredComponent(Component);
 	}
 
 	const FBoxSphereBounds OldBounds = CachedLocalBounds;
@@ -619,8 +603,8 @@ void UClusterUnionComponent::HandleAddOrModifiedClusteredComponent(UPrimitiveCom
 		return;
 	}
 
-	const bool bIsNew = !ComponentToPhysicsObjects.Contains(ChangedComponent);
-	FClusteredComponentData& ComponentData = ComponentToPhysicsObjects.FindOrAdd(ChangedComponent);
+	const bool bIsNew = !PerComponentData.Contains(ChangedComponent);
+	FClusteredComponentData& ComponentData = PerComponentData.FindOrAdd(ChangedComponent);
 
 	// A component shouldn't be in the PendingComponentSync map unless it's new.
 	FClusterUnionPendingAddData PendingData;
@@ -722,8 +706,6 @@ void UClusterUnionComponent::HandleAddOrModifiedClusteredComponent(UPrimitiveCom
 				ComponentData.CachedAccelerationPayloads.Add(Handle);
 			}
 		}
-
-		ComponentData.PhysicsObjects.Add(PhysicsObject);
 	}
 
 	// In the case where we need to keep the acceleration structure up to date, we need to make sure old bone ids are properly
@@ -741,79 +723,58 @@ void UClusterUnionComponent::HandleAddOrModifiedClusteredComponent(UPrimitiveCom
 			}
 		}
 	}
-
-	ComponentData.AllPhysicsObjects = MoveTemp(AllPhysicsObjects);
 }
 
-void UClusterUnionComponent::HandleRemovedClusteredComponent(UPrimitiveComponent* ChangedComponent, bool bDestroyReplicatedProxy)
+void UClusterUnionComponent::HandleRemovedClusteredComponent(TObjectKey<UPrimitiveComponent> RemovedComponent)
 {
-	if (!ChangedComponent)
+	if (FClusteredComponentData* Data = PerComponentData.Find(RemovedComponent))
 	{
-		return;
-	}
-
-	// At this point the component's particles are no longer a part of the cluster union. So we just need
-	// to get our book-keeping and game thread state to match that.
-	AActor* Owner = ChangedComponent->GetOwner();
-	if (!ensure(Owner))
-	{
-		return;
-	}
-
-	if (FClusteredComponentData* ComponentData = ComponentToPhysicsObjects.Find(ChangedComponent))
-	{
-		if (ChangedComponent->HasValidPhysicsState())
+		if (AccelerationStructure)
 		{
-			FLockedReadPhysicsObjectExternalInterface Interface = FPhysicsObjectExternalInterface::LockRead(ComponentData->AllPhysicsObjects);
-			for (Chaos::FPhysicsObjectHandle PhysicsObject : ComponentData->AllPhysicsObjects)
+			for (const FExternalSpatialAccelerationPayload& Handle : Data->CachedAccelerationPayloads)
 			{
-				if (Chaos::FGeometryParticle* Particle = Interface->GetParticle(PhysicsObject))
-				{
-					UniqueIdxToComponent.Remove(Particle->UniqueIdx().Idx);
-				}
-
-				if (AccelerationStructure && PhysicsObject)
-				{
-					FExternalSpatialAccelerationPayload Handle;
-					Handle.Initialize(ChangedComponent, Chaos::FPhysicsObjectInterface::GetId(PhysicsObject));
-					AccelerationStructure->RemoveElement(Handle);
-					ComponentData->CachedAccelerationPayloads.Remove(Handle);
-				}
+				AccelerationStructure->RemoveElement(Handle);
 			}
 		}
 
 		if (IsAuthority())
 		{
-			if (bDestroyReplicatedProxy && ensure(ComponentData->ReplicatedProxyComponent.IsValid()))
+			if (UClusterUnionReplicatedProxyComponent* ProxyComponent = Data->ReplicatedProxyComponent.Get())
 			{
-				UClusterUnionReplicatedProxyComponent* ProxyComponent = ComponentData->ReplicatedProxyComponent.Get();
 				ProxyComponent->DestroyComponent();
 			}
 		}
 
-		ComponentToPhysicsObjects.Remove(ChangedComponent);
-		OnComponentRemovedEvent.Broadcast(ChangedComponent);
-	}
-
+		if (AActor* Owner = Data->Owner.Get())
+		{
+			if (FClusteredActorData* ActorData = ActorToComponents.Find(Owner))
+			{
+				ActorData->Components.Remove(RemovedComponent);
 
-	if (FClusteredActorData* ActorData = ActorToComponents.Find(Owner))
-	{
-		ActorData->Components.Remove(ChangedComponent);
+				if (ActorData->Components.IsEmpty())
+				{
+					if (IsAuthority())
+					{
+						Owner->SetReplicatingMovement(ActorData->bWasReplicatingMovement);
+					}
+					ActorToComponents.Remove(Owner);
+				}
+			}
 
-		if (ActorData->Components.IsEmpty())
-		{
-			if (IsAuthority())
+			if (bClusterUnionCallFlushNetDormancy)
 			{
-				Owner->SetReplicatingMovement(ActorData->bWasReplicatingMovement);
+				Owner->FlushNetDormancy();
 			}
-			ActorToComponents.Remove(Owner);
 		}
-	}
 
-	if (bClusterUnionCallFlushNetDormancy)
-	{
-		Owner->FlushNetDormancy();
+		if (UPrimitiveComponent* ChangedComponent = RemovedComponent.ResolveObjectPtr())
+		{
+			OnComponentRemovedEvent.Broadcast(ChangedComponent);
+		}
+		PerComponentData.Remove(RemovedComponent);
 	}
+	PendingComponentsToAdd.Remove(RemovedComponent);
+	PendingComponentSync.Remove(RemovedComponent);
 }
 
 void UClusterUnionComponent::OnRep_RigidState()
@@ -1081,7 +1042,7 @@ bool UClusterUnionComponent::ComponentOverlapComponentWithResultImpl(const class
 TArray<UPrimitiveComponent*> UClusterUnionComponent::GetAllCurrentChildComponents() const
 {
 	TArray<UPrimitiveComponent*> Components;
-	Components.Reserve(ComponentToPhysicsObjects.Num() + PendingComponentSync.Num());
+	Components.Reserve(PerComponentData.Num() + PendingComponentSync.Num());
 
 	VisitAllCurrentChildComponents(
 		[&Components](UPrimitiveComponent* Component)
@@ -1123,7 +1084,7 @@ void UClusterUnionComponent::VisitAllCurrentChildComponents(const TFunction<bool
 		}
 	}
 
-	for (const TPair<TObjectKey<UPrimitiveComponent>, FClusteredComponentData>& Kvp : ComponentToPhysicsObjects)
+	for (const TPair<TObjectKey<UPrimitiveComponent>, FClusteredComponentData>& Kvp : PerComponentData)
 	{
 		if (UPrimitiveComponent* Component = Kvp.Key.ResolveObjectPtr(); Component && Component->HasValidPhysicsState() && !Kvp.Value.bPendingDeletion)
 		{
@@ -1198,7 +1159,7 @@ void UClusterUnionComponent::VisitAllCurrentChildComponentsForCollision(ECollisi
 
 TArray<int32> UClusterUnionComponent::GetAddedBoneIdsForComponent(UPrimitiveComponent* Component) const
 {
-	if (const FClusteredComponentData* Data = ComponentToPhysicsObjects.Find(Component))
+	if (const FClusteredComponentData* Data = PerComponentData.Find(Component))
 	{
 		return Data->BoneIds.Array();
 	}
@@ -1219,7 +1180,7 @@ void UClusterUnionComponent::AddReferencedObjects(UObject* InThis, FReferenceCol
 
 	{
 		const UScriptStruct* ScriptStruct = FClusteredComponentData::StaticStruct();
-		for (TPair<TObjectKey<UPrimitiveComponent>, FClusteredComponentData>& Kvp : This->ComponentToPhysicsObjects)
+		for (TPair<TObjectKey<UPrimitiveComponent>, FClusteredComponentData>& Kvp : This->PerComponentData)
 		{
 			TWeakObjectPtr<const UScriptStruct> ScriptStructPtr{ScriptStruct};
 			Collector.AddReferencedObjects(ScriptStructPtr, reinterpret_cast<void*>(&Kvp.Value), This, nullptr);
@@ -1255,7 +1216,7 @@ void UClusterUnionComponent::AddReferencedObjects(UObject* InThis, FReferenceCol
 
 int32 UClusterUnionComponent::NumChildClusterComponents() const
 {
-	return PendingComponentSync.Num() + ComponentToPhysicsObjects.Num();
+	return PendingComponentSync.Num() + PerComponentData.Num();
 }
 
 bool UClusterUnionComponent::DoCustomNavigableGeometryExport(FNavigableGeometryExport& GeomExport) const
diff --git a/Engine/Source/Runtime/Experimental/Chaos/Private/PhysicsProxy/ClusterUnionPhysicsProxy.cpp b/Engine/Source/Runtime/Experimental/Chaos/Private/PhysicsProxy/ClusterUnionPhysicsProxy.cpp
--- a/Engine/Source/Runtime/Experimental/Chaos/Private/PhysicsProxy/ClusterUnionPhysicsProxy.cpp
+++ b/Engine/Source/Runtime/Experimental/Chaos/Private/PhysicsProxy/ClusterUnionPhysicsProxy.cpp
@@ -389,11 +389,15 @@ namespace Chaos
 
 		SyncedData_External.bIsAnchored = CurrentPullData.bIsAnchored;
 		SyncedData_External.ChildParticles.Empty(CurrentPullData.ChildParticles.Num());
+
+		// I question the need to do this individual copy one by one...maybe we should've used the same type for the synced data.
 		for (const FDirtyClusterUnionParticleData& InData : CurrentPullData.ChildParticles)
 		{
 			FClusterUnionChildData ConvertedData;
 			ConvertedData.ParticleIdx = InData.ParticleIdx;
 			ConvertedData.ChildToParent = InData.ChildToParent;
+			ConvertedData.Proxy = InData.Proxy;
+			ConvertedData.BoneId = InData.BoneId;
 			SyncedData_External.ChildParticles.Add(ConvertedData);
 		}
 
@@ -506,17 +510,29 @@ namespace Chaos
 					continue;
 				}
 
-				FDirtyClusterUnionParticleData Data;
-				Data.ParticleIdx = Particle->UniqueIdx();
-				if (FPBDRigidClusteredParticleHandle* ClusteredParticle = Particle->CastToClustered())
-				{
-					Data.ChildToParent = ClusteredParticle->ChildToParent();
-				}
-				else
+				if (IPhysicsProxyBase* Proxy = Particle->PhysicsProxy())
 				{
-					Data.ChildToParent = FRigidTransform3::Identity;
+					FDirtyClusterUnionParticleData Data;
+					Data.ParticleIdx = Particle->UniqueIdx();
+					if (FPBDRigidClusteredParticleHandle* ClusteredParticle = Particle->CastToClustered())
+					{
+						Data.ChildToParent = ClusteredParticle->ChildToParent();
+					}
+					else
+					{
+						Data.ChildToParent = FRigidTransform3::Identity;
+					}
+					Data.Proxy = Proxy;
+					
+					if (Proxy->GetType() == EPhysicsProxyType::GeometryCollectionType)
+					{
+						// A bit hacky, but the only way we can communicate what particle we care about back to the cluster union on the GT.
+						FGeometryCollectionPhysicsProxy* GCProxy = static_cast<FGeometryCollectionPhysicsProxy*>(Proxy);
+						Data.BoneId = GCProxy->GetTransformGroupIndexFromHandle(Particle);
+					}
+
+					BufferData.ChildParticles.Add(Data);
 				}
-				BufferData.ChildParticles.Add(Data);
 			}
 
 			BufferData.SharedGeometry = ClusterUnion->SharedGeometry;
diff --git a/Engine/Source/Runtime/Experimental/Chaos/Public/Chaos/PullPhysicsDataImp.h b/Engine/Source/Runtime/Experimental/Chaos/Public/Chaos/PullPhysicsDataImp.h
--- a/Engine/Source/Runtime/Experimental/Chaos/Public/Chaos/PullPhysicsDataImp.h
+++ b/Engine/Source/Runtime/Experimental/Chaos/Public/Chaos/PullPhysicsDataImp.h
@@ -98,6 +98,8 @@ struct FDirtyClusterUnionParticleData
 {
 	FUniqueIdx ParticleIdx;
 	FRigidTransform3 ChildToParent;
+	IPhysicsProxyBase* Proxy = nullptr;
+	int32 BoneId = INDEX_NONE;
 };
 
 struct FDirtyClusterUnionData : public TBasePullData<FClusterUnionPhysicsProxy, FClusterUnionProxyTimestamp>
diff --git a/Engine/Source/Runtime/Experimental/Chaos/Public/PhysicsProxy/ClusterUnionPhysicsProxy.h b/Engine/Source/Runtime/Experimental/Chaos/Public/PhysicsProxy/ClusterUnionPhysicsProxy.h
--- a/Engine/Source/Runtime/Experimental/Chaos/Public/PhysicsProxy/ClusterUnionPhysicsProxy.h
+++ b/Engine/Source/Runtime/Experimental/Chaos/Public/PhysicsProxy/ClusterUnionPhysicsProxy.h
@@ -21,6 +21,8 @@ namespace Chaos
 	{
 		FUniqueIdx ParticleIdx;
 		FTransform ChildToParent;
+		IPhysicsProxyBase* Proxy = nullptr;
+		int32 BoneId = INDEX_NONE;
 	};
 
 	struct FClusterUnionInitData
