need to actually set the owner property on per-component data for cluster unions

diff --git a/Engine/Source/Runtime/Engine/Private/PhysicsEngine/ClusterUnionComponent.cpp b/Engine/Source/Runtime/Engine/Private/PhysicsEngine/ClusterUnionComponent.cpp
--- a/Engine/Source/Runtime/Engine/Private/PhysicsEngine/ClusterUnionComponent.cpp
+++ b/Engine/Source/Runtime/Engine/Private/PhysicsEngine/ClusterUnionComponent.cpp
@@ -639,6 +639,8 @@ void UClusterUnionComponent::HandleAddOrModifiedClusteredComponent(UPrimitiveCom
 					Owner->SetReplicatingMovement(false);
 				}
 			}
+
+			ComponentData.Owner = Owner;
 		}
 
 		ComponentData.bWasReplicating = ChangedComponent->GetIsReplicated();
