Handle serialization of rewind data in UNetworkPhysicsComponent in USTRUCT NetSerialize functions. This ensures that object references are serialized properly, and eliminates the need to keep track of a UPackageMap in the component and TNetRewindHistory.

diff --git a/Engine/Source/Runtime/Engine/Private/PhysicsEngine/NetworkPhysicsComponent.cpp b/Engine/Source/Runtime/Engine/Private/PhysicsEngine/NetworkPhysicsComponent.cpp
--- a/Engine/Source/Runtime/Engine/Private/PhysicsEngine/NetworkPhysicsComponent.cpp
+++ b/Engine/Source/Runtime/Engine/Private/PhysicsEngine/NetworkPhysicsComponent.cpp
@@ -79,6 +79,63 @@ float DeQuantizeTimeDilation(int8 i)
 
 }
 
+
+bool FNetworkPhysicsRewindDataProxy::NetSerializeBase(FArchive& Ar, class UPackageMap* Map, bool& bOutSuccess, TUniqueFunction<TUniquePtr<Chaos::FBaseRewindHistory>()> CreateHistoryFunction)
+{
+	Ar << Owner;
+
+	bool bHasData = History.IsValid();
+	Ar.SerializeBits(&bHasData, 1);
+
+	if (bHasData)
+	{
+		if (Ar.IsLoading() && !History.IsValid())
+		{
+			if(ensureMsgf(Owner, TEXT("FNetRewindDataBase::NetSerialize: owner is null")))
+			{
+				History = CreateHistoryFunction();
+				if (!ensureMsgf(History.IsValid(), TEXT("FNetRewindDataBase::NetSerialize: failed to create history. Owner: %s"), *GetFullNameSafe(Owner)))
+				{
+					Ar.SetError();
+					bOutSuccess = false;
+					return true;
+				}
+			}
+			else
+			{
+				Ar.SetError();
+				bOutSuccess = false;
+				return true;
+			}
+		}
+
+		History->NetSerialize(Ar, Map);
+	}
+
+	return true;
+}
+
+FNetworkPhysicsRewindDataProxy& FNetworkPhysicsRewindDataProxy::operator=(const FNetworkPhysicsRewindDataProxy& Other)
+{
+	if (&Other != this)
+	{
+		Owner = Other.Owner;
+		History = Other.History ? Other.History->Clone() : nullptr;
+	}
+
+	return *this;
+}
+
+bool FNetworkPhysicsRewindDataInputProxy::NetSerialize(FArchive& Ar, class UPackageMap* Map, bool& bOutSuccess)
+{
+	return NetSerializeBase(Ar, Map, bOutSuccess, [this]() { return Owner->ReplicatedInputs.History->CreateNew(); });
+}
+
+bool FNetworkPhysicsRewindDataStateProxy::NetSerialize(FArchive& Ar, class UPackageMap* Map, bool& bOutSuccess)
+{
+	return NetSerializeBase(Ar, Map, bOutSuccess, [this]() { return Owner->ReplicatedStates.History->CreateNew(); });
+}
+
 // after presimulate internal (asyncinput internal simulation done and the output created)
 void FNetworkPhysicsCallback::ApplyCallbacks_Internal(int32 PhysicsStep, const TArray<Chaos::ISimCallbackObject*>& SimCallbackObjects)
 {
@@ -399,18 +456,6 @@ void UNetworkPhysicsComponent::GetLifetimeReplicatedProps(TArray< FLifetimePrope
 	DOREPLIFETIME_CONDITION_NOTIFY(UNetworkPhysicsComponent, ReplicatedStates, COND_None, REPNOTIFY_Always);
 }
 
-void UNetworkPhysicsComponent::UpdatePackageMap()
-{
-	APlayerController* Controller = GetPlayerController();
-	UPackageMap* PackageMap = (Controller && Controller->GetNetConnection()) ? Controller->GetNetConnection()->PackageMap : nullptr;
-
-	if(InputsHistory && StatesHistory)
-	{
-		InputsHistory->SetPackageMap(PackageMap);
-		StatesHistory->SetPackageMap(PackageMap);
-	}
-}
-
 void UNetworkPhysicsComponent::AsyncPhysicsTickComponent(float DeltaTime, float SimTime)
 {
 	QUICK_SCOPE_CYCLE_COUNTER(NetworkPhysicsComponent_AsyncPhysicsTick);
@@ -420,7 +465,7 @@ void UNetworkPhysicsComponent::AsyncPhysicsTickComponent(float DeltaTime, float
 	if(HasServerWorld() && !HasLocalController() && InputsHistory)
 	{
 		TArray<int32> LocalFrames, ServerFrames, InputFrames;
-		InputsHistory->DebugDatas(ReplicatedInputs, LocalFrames, ServerFrames, InputFrames);
+		InputsHistory->DebugDatas(*ReplicatedInputs.History, LocalFrames, ServerFrames, InputFrames);
 
 		UE_LOG(LogTemp, Log, TEXT("SERVER | PT | AsyncPhysicsTickComponent | Receiving %d inputs from CLIENT | Component = %s"), LocalFrames.Num(), *GetFullName());
 		for (int32 FrameIndex = 0; FrameIndex < LocalFrames.Num(); ++FrameIndex)
@@ -430,8 +475,6 @@ void UNetworkPhysicsComponent::AsyncPhysicsTickComponent(float DeltaTime, float
 		}
 	}
 #endif
-	// Update the package map for serialization
-	UpdatePackageMap();
 
 	// Record the received states from the server into the history for future use
 	if (UWorld* World = GetWorld())
@@ -464,14 +507,11 @@ void UNetworkPhysicsComponent::SendLocalInputsDatas()
 		if (LocalOffset >= 0)
 		{
 			const int32 NextIndex = (InputsIndex + 1) % (InputsRedundancy + 1);
-			InputsHistory->SerializeDatas(InputsOffsets[NextIndex], InputsOffsets[InputsIndex], LocalInputs, LocalOffset);
 
-			if (HasServerWorld())
-			{
-				// if on server (Listen server) we should send the inputs onto all the clients through repnotify
-				ReplicatedInputs = LocalInputs;
-			}
-			else
+			// if on server (Listen server) we should send the inputs onto all the clients through repnotify
+			ReplicatedInputs.History = InputsHistory->CopyFramesWithOffset(InputsOffsets[NextIndex], InputsOffsets[InputsIndex], LocalOffset);
+			
+			if (!HasServerWorld())
 			{
 #if DEBUG_NETWORK_PHYSICS
 				if (APlayerController* PlayerController = GetPlayerController())
@@ -479,7 +519,7 @@ void UNetworkPhysicsComponent::SendLocalInputsDatas()
 					FAsyncPhysicsTimestamp Timestamp = PlayerController->GetAsyncPhysicsTimestamp();
 
 					TArray<int32> LocalFrames, ServerFrames, InputFrames;
-					InputsHistory->DebugDatas(LocalInputs, LocalFrames, ServerFrames, InputFrames);
+					InputsHistory->DebugDatas(*ReplicatedInputs.History, LocalFrames, ServerFrames, InputFrames);
 
 					UE_LOG(LogTemp, Log, TEXT("CLIENT | GT | SendLocalInputsDatas | Sending %d inputs from CLIENT | Component = %s"), LocalFrames.Num(), *GetFullName());
 					for (int32 FrameIndex = 0; FrameIndex < LocalFrames.Num(); ++FrameIndex)
@@ -492,7 +532,7 @@ void UNetworkPhysicsComponent::SendLocalInputsDatas()
 
 				// if on the client we should first send the replicated inputs onto the server
 				// the RPC will then resend them onto all the other clients (except the local ones)
-				ServerReceiveInputsDatas(LocalInputs);
+				ServerReceiveInputsDatas(ReplicatedInputs);
 			}
 		}
 	}
@@ -503,10 +543,9 @@ void UNetworkPhysicsComponent::SendLocalStatesDatas()
 	if (HasServerWorld() && StatesHistory)
 	{
 		const int32 NextIndex = (StatesIndex + 1) % (StatesRedundancy + 1);
-		StatesHistory->SerializeDatas(StatesOffsets[NextIndex], StatesOffsets[StatesIndex], LocalStates, 0);
 
 		// if on server we should send the states onto all the clients through repnotify
-		ReplicatedStates = LocalStates;
+		ReplicatedStates.History = StatesHistory->CopyFramesWithOffset(StatesOffsets[NextIndex], StatesOffsets[StatesIndex], 0);
 	}
 }
 
@@ -514,10 +553,8 @@ void UNetworkPhysicsComponent::CorrectServerToLocalOffset(const int32 LocalToSer
 {
 	if (HasLocalController() && !HasServerWorld() && StatesHistory)
 	{
-		const TArray<uint8> ReceivedStates = ReplicatedStates;
-
 		TArray<int32> LocalFrames, ServerFrames, InputFrames;
-		StatesHistory->DebugDatas(ReceivedStates, LocalFrames, ServerFrames, InputFrames);
+		StatesHistory->DebugDatas(*ReplicatedStates.History, LocalFrames, ServerFrames, InputFrames);
 
 		int32 ServerToLocalOffset = LocalToServerOffset;
 		for (int32 FrameIndex = 0; FrameIndex < LocalFrames.Num(); ++FrameIndex)
@@ -553,14 +590,14 @@ void UNetworkPhysicsComponent::OnRep_SetReplicatedStates()
 		{
 			if (FPhysScene* PhysScene = World->GetPhysicsScene())
 			{
-				const TArray<uint8> ReceivedStates = ReplicatedStates;
+				TSharedPtr<Chaos::FBaseRewindHistory> ReceivedStates = MakeShareable(ReplicatedStates.History->Clone().Release());
 				PhysScene->EnqueueAsyncPhysicsCommand(0, this, [this, PhysScene, ReceivedStates, LocalOffset]()
 				{
-					StatesHistory->DeserializeDatas(ReceivedStates, LocalOffset);
+					StatesHistory->ReceiveNewDatas(*ReceivedStates, LocalOffset);
 #if DEBUG_NETWORK_PHYSICS
 					{
 						TArray<int32> LocalFrames, ServerFrames, InputFrames;
-						StatesHistory->DebugDatas(ReceivedStates, LocalFrames, ServerFrames, InputFrames);
+						StatesHistory->DebugDatas(*ReceivedStates, LocalFrames, ServerFrames, InputFrames);
 
 						UE_LOG(LogTemp, Log, TEXT("CLIENT | PT | OnRep_SetReplicatedStates | Receiving %d states from SERVER | Local offset = %d | Component = %s "), LocalFrames.Num(), LocalOffset, *GetFullName());
 						for (int32 FrameIndex = 0; FrameIndex < LocalFrames.Num(); ++FrameIndex)
@@ -597,15 +634,15 @@ void UNetworkPhysicsComponent::OnRep_SetReplicatedInputs()
 		{
 			if (FPhysScene* PhysScene = World->GetPhysicsScene())
 			{
-				const TArray<uint8> ReceivedInputs = ReplicatedInputs;
+				TSharedPtr<Chaos::FBaseRewindHistory> ReceivedInputs = MakeShareable(ReplicatedInputs.History->Clone().Release());
 				PhysScene->EnqueueAsyncPhysicsCommand(0, this, [this, PhysScene, ReceivedInputs, LocalOffset]()
 				{
-					InputsHistory->DeserializeDatas(ReceivedInputs, LocalOffset);
+					InputsHistory->ReceiveNewDatas(*ReceivedInputs, LocalOffset);
 
 #if DEBUG_NETWORK_PHYSICS
 					{
 						TArray<int32> LocalFrames, ServerFrames, InputFrames;
-						InputsHistory->DebugDatas(ReceivedInputs, LocalFrames, ServerFrames, InputFrames);
+						InputsHistory->DebugDatas(*ReceivedInputs, LocalFrames, ServerFrames, InputFrames);
 
 						UE_LOG(LogTemp, Log, TEXT("CLIENT | PT | OnRep_SetReplicatedInputs | Receiving %d inputs from SERVER | Local offset = %d | Component = %s"), LocalFrames.Num(), LocalOffset, *GetFullName());
 						for (int32 FrameIndex = 0; FrameIndex < LocalFrames.Num(); ++FrameIndex)
@@ -620,12 +657,7 @@ void UNetworkPhysicsComponent::OnRep_SetReplicatedInputs()
 	}
 }
 
-bool UNetworkPhysicsComponent::ServerReceiveInputsDatas_Validate(const TArray<uint8>& ClientInputs)
-{
-	return true;
-}
-
-void UNetworkPhysicsComponent::ServerReceiveInputsDatas_Implementation(const TArray<uint8>& ClientInputs)
+void UNetworkPhysicsComponent::ServerReceiveInputsDatas_Implementation(const FNetworkPhysicsRewindDataInputProxy& ClientInputs)
 {
 	if(InputsHistory)
 	{ 
@@ -635,26 +667,28 @@ void UNetworkPhysicsComponent::ServerReceiveInputsDatas_Implementation(const TAr
 		// We could probably skip that test since the server RPC is on server
 		ensure(!HasLocalController());
 
-		ReplicatedInputs = ClientInputs;
-
 		// Record the received inputs from the client into the history for future use
+		ReplicatedInputs.History = ClientInputs.History->Clone();
+
 		if (UWorld* World = GetWorld())
 		{
 			if (FPhysScene* PhysScene = World->GetPhysicsScene())
 			{
-				const TArray<uint8> ReceivedInputs = ClientInputs;
+				// Make another copy of the client inputs for the physics thread to consume
+				TSharedPtr<Chaos::FBaseRewindHistory> ReceivedInputs = MakeShareable(ClientInputs.History->Clone().Release());
 				PhysScene->EnqueueAsyncPhysicsCommand(0, this, [this, ReceivedInputs, PhysScene]()
 				{
-					InputsHistory->DeserializeDatas(ReceivedInputs, 0);
+					InputsHistory->ReceiveNewDatas(*ReceivedInputs, 0);
+
 	#if DEBUG_NETWORK_PHYSICS
 					{
 						TArray<int32> LocalFrames, ServerFrames, InputFrames;
-						InputsHistory->DebugDatas(ReceivedInputs, LocalFrames, ServerFrames, InputFrames);
+						InputsHistory->DebugDatas(*ReceivedInputs, LocalFrames, ServerFrames, InputFrames);
 
 						const int32 CurrentFrame = PhysScene->GetSolver()->GetCurrentFrame();
 
 						const int32 EvalOffset = CurrentFrame - InputFrames[InputFrames.Num()-1] + 4;
-						UE_LOG(LogTemp, Log, TEXT("SERVER | PT | ServerReceiveInputsDatas | Receiving %d inputs from CLIENT | Inputs frame = %d | Server frame = %d | Eval Offset = %d | Component = %s | Num Bits = %d"), LocalFrames.Num(), InputFrames[InputFrames.Num() - 1], CurrentFrame, EvalOffset, *GetFullName(), ReceivedInputs.Num() * 8);
+						UE_LOG(LogTemp, Log, TEXT("SERVER | PT | ServerReceiveInputsDatas | Receiving %d inputs from CLIENT | Inputs frame = %d | Server frame = %d | Eval Offset = %d | Component = %s"), LocalFrames.Num(), InputFrames[InputFrames.Num() - 1], CurrentFrame, EvalOffset, *GetFullName());
 						for (int32 FrameIndex = 0; FrameIndex < LocalFrames.Num(); ++FrameIndex)
 						{
 							UE_LOG(LogTemp, Log, TEXT("		Recording replicated inputs at local frame = %d | server frame = %d | Solver offset = %d | Component = %s"), 
diff --git a/Engine/Source/Runtime/Engine/Public/Physics/NetworkPhysicsComponent.h b/Engine/Source/Runtime/Engine/Public/Physics/NetworkPhysicsComponent.h
--- a/Engine/Source/Runtime/Engine/Public/Physics/NetworkPhysicsComponent.h
+++ b/Engine/Source/Runtime/Engine/Public/Physics/NetworkPhysicsComponent.h
@@ -23,66 +23,56 @@ struct TNetRewindHistory : public Chaos::TDatasRewindHistory<DatasType>
 {
 	using Super = Chaos::TDatasRewindHistory<DatasType>;
 
-	FORCEINLINE TNetRewindHistory(const int32 FrameCount, const bool bIsHistoryLocal, UPackageMap* InPackageMap) :
+	FORCEINLINE TNetRewindHistory(const int32 FrameCount, const bool bIsHistoryLocal) :
 		Super(FrameCount, bIsHistoryLocal)
 	{
-		DatasWriter.SetAllowResize(true);
-		SetPackageMap(InPackageMap);
 	}
 
 	FORCEINLINE virtual ~TNetRewindHistory() {}
 
-	/** Set the package map for serialization */
-	FORCEINLINE virtual void SetPackageMap(class UPackageMap* InPackageMap) override
+	virtual TUniquePtr<Chaos::FBaseRewindHistory> CreateNew() const
 	{
-		PackageMap = InPackageMap;
+		TUniquePtr<TNetRewindHistory> Copy = MakeUnique<TNetRewindHistory>(0, Super::bIsLocalHistory);
 
-		DatasWriter.PackageMap = PackageMap;
-		DatasReader.PackageMap = PackageMap;
+		return Copy;
 	}
 
-	/** Serialize the datas to the archive */
-	FORCEINLINE virtual void SerializeDatas(const uint32 StartFrame, const uint32 EndFrame, TArray<uint8>& ArchiveDatas, const int32 FrameOffset) const override
+	virtual TUniquePtr<Chaos::FBaseRewindHistory> Clone() const
 	{
-		if (FrameOffset >= 0)
-		{
-			DatasWriter.Reset();
+		return MakeUnique<TNetRewindHistory>(*this);
+	}
 
-			uint32 FramesCount = (uint32)Super::NumValidDatas(StartFrame, EndFrame);
-			DatasWriter << FramesCount;
+	virtual TUniquePtr<Chaos::FBaseRewindHistory> CopyFramesWithOffset(const uint32 StartFrame, const uint32 EndFrame, const int32 FrameOffset) override
+	{
+		uint32 FramesCount = (uint32)Super::NumValidDatas(StartFrame, EndFrame);
+			
+		TUniquePtr<TNetRewindHistory> Copy = MakeUnique<TNetRewindHistory>(FramesCount, Super::bIsLocalHistory);
 
-			DatasType FrameDatas;
-			for (uint32 FrameIndex = StartFrame; FrameIndex < EndFrame; ++FrameIndex)
+		DatasType FrameDatas;
+		for (uint32 FrameIndex = StartFrame; FrameIndex < EndFrame; ++FrameIndex)
+		{
+			const int32 LocalFrame = FrameIndex % Super::NumFrames;
+			if (FrameIndex == Super::DatasArray[LocalFrame].LocalFrame)
 			{
-				const int32 LocalFrame = FrameIndex % Super::NumFrames;
-				if (FrameIndex == Super::DatasArray[LocalFrame].LocalFrame)
-				{
-					FrameDatas = Super::DatasArray[LocalFrame];
-					FrameDatas.ServerFrame = FrameDatas.LocalFrame + FrameOffset;
-					NetSerializeDatas(FrameDatas, DatasWriter);
-				}
+				FrameDatas = Super::DatasArray[LocalFrame];
+				FrameDatas.ServerFrame = FrameDatas.LocalFrame + FrameOffset;
+				Copy->RecordDatas(LocalFrame, &FrameDatas);
 			}
-
-			ArchiveDatas = *DatasWriter.GetBuffer();
 		}
+
+		return Copy;
 	}
 
-	/** Deserialize the datas from the archive */
-	FORCEINLINE virtual void DeserializeDatas(const TArray<uint8>& ArchiveDatas, const int32 FrameOffset) override
+	virtual void ReceiveNewDatas(Chaos::FBaseRewindHistory& NewDatas, const int32 FrameOffset) override
 	{
-		if ((FrameOffset >= 0) && !ArchiveDatas.IsEmpty())
-		{
-			// TODO : temporary copy since the archive datas are const
-			TArray<uint8> LocalDatas = ArchiveDatas;
-			DatasReader.SetData(LocalDatas.GetData(), ArchiveDatas.Num() * 8);
-
-			uint32 FramesCount = 0;
-			DatasReader << FramesCount;
+		TNetRewindHistory& NetNewDatas = static_cast<TNetRewindHistory&>(NewDatas);
 
-			DatasType FrameDatas;
-			for (uint32 FrameIndex = 0; FrameIndex < FramesCount; ++FrameIndex)
+		if ((FrameOffset >= 0) && NetNewDatas.NumFrames > 0)
+		{
+			for (int32 FrameIndex = 0; FrameIndex < NetNewDatas.NumFrames; ++FrameIndex)
 			{
-				NetSerializeDatas(FrameDatas, DatasReader);
+				DatasType& FrameDatas = NetNewDatas.DatasArray[FrameIndex];
+
 				FrameDatas.LocalFrame = FrameDatas.ServerFrame - FrameOffset;
 				if (FrameDatas.LocalFrame >= 0)
 				{
@@ -92,26 +82,43 @@ struct TNetRewindHistory : public Chaos::TDatasRewindHistory<DatasType>
 		}
 	}
 
-	/** Debug the datas from the archive */
-	FORCEINLINE virtual void DebugDatas(const TArray<uint8>& ArchiveDatas, TArray<int32>& LocalFrames, TArray<int32>& ServerFrames, TArray<int32>& InputFrames) override
+	virtual void NetSerialize(FArchive& Ar, UPackageMap* InPackageMap) override
 	{
-		if((PackageMap != nullptr) && !ArchiveDatas.IsEmpty())
+		Ar << Super::NumFrames;
+		
+		if (Super::NumFrames > GetMaxArraySize())
 		{
-			// TODO : temporary copy since the archive datas are const
-			TArray<uint8> LocalDatas = ArchiveDatas;
-			DatasReader.SetData(LocalDatas.GetData(), ArchiveDatas.Num() * 8);
+			UE_LOG(LogPhysics, Warning, TEXT("TNetRewindHistory: serialized array of size %d exceeds maximum size %d."), Super::NumFrames, GetMaxArraySize());
+			Ar.SetError();
+			return;
+		}
 
-			uint32 FramesCount = 0;
-			DatasReader << FramesCount;
+		if (Ar.IsLoading())
+		{
+			Super::DatasArray.SetNum(Super::NumFrames);
+		}
+
+		for (DatasType& Data : Super::DatasArray)
+		{
+			NetSerializeDatas(Data, Ar, InPackageMap);
+		}
+	}
+
+	/** Debug the datas from the archive */
+	FORCEINLINE virtual void DebugDatas(const Chaos::FBaseRewindHistory& NewDatas, TArray<int32>& LocalFrames, TArray<int32>& ServerFrames, TArray<int32>& InputFrames) override
+	{
+		const TNetRewindHistory& NewNetDatas = static_cast<const TNetRewindHistory&>(NewDatas);
 
-			LocalFrames.SetNum(FramesCount);
-			ServerFrames.SetNum(FramesCount);
-			InputFrames.SetNum(FramesCount);
+		if(NewNetDatas.NumFrames >= 0)
+		{
+			LocalFrames.SetNum(NewNetDatas.NumFrames);
+			ServerFrames.SetNum(NewNetDatas.NumFrames);
+			InputFrames.SetNum(NewNetDatas.NumFrames);
 
 			DatasType FrameDatas;
-			for (uint32 FrameIndex = 0; FrameIndex < FramesCount; ++FrameIndex)
+			for (int32 FrameIndex = 0; FrameIndex < NewNetDatas.NumFrames; ++FrameIndex)
 			{
-				NetSerializeDatas(FrameDatas, DatasReader);
+				FrameDatas = NewNetDatas.DatasArray[FrameIndex];
 				LocalFrames[FrameIndex] = FrameDatas.LocalFrame;
 				ServerFrames[FrameIndex] = FrameDatas.ServerFrame;
 				InputFrames[FrameIndex] = FrameDatas.InputFrame;
@@ -121,8 +128,15 @@ struct TNetRewindHistory : public Chaos::TDatasRewindHistory<DatasType>
 
 private :
 
+	/** Serialized array size limit to guard against invalid network data */
+	static int32 GetMaxArraySize()
+	{
+		static int32 MaxArraySize = UPhysicsSettings::Get()->GetPhysicsHistoryCount() * 4;
+		return MaxArraySize;
+	}
+
 	/** Use net serialize path to serialize datas  */
-	FORCEINLINE bool NetSerializeDatas(DatasType& FrameDatas, FBitArchive& Ar) const 
+	FORCEINLINE bool NetSerializeDatas(DatasType& FrameDatas, FArchive& Ar, UPackageMap* PackageMap) const 
 	{
 		bool bOutSuccess = false;
 		UScriptStruct* ScriptStruct = DatasType::StaticStruct();
@@ -149,15 +163,73 @@ private :
 		}
 		return bOutSuccess;
 	}
+};
+
+/**
+ * Base struct for replicated rewind history properties
+ */
+USTRUCT()
+struct FNetworkPhysicsRewindDataProxy
+{
+	GENERATED_BODY()
+
+	FNetworkPhysicsRewindDataProxy& operator=(const FNetworkPhysicsRewindDataProxy& Other);
 
-	// Datas bits writer to be used for serialization
-	mutable FNetBitWriter DatasWriter;
+	/** Causes the history to be serialized every time. If implemented, would prevent serializing if the history hasn't changed. */
+	bool operator==(const FNetworkPhysicsRewindDataProxy& Other) const { return false; }
 
-	// Datas bits reader to be used for serialization
-	mutable FNetBitReader DatasReader;
+protected:
+	bool NetSerializeBase(FArchive& Ar, class UPackageMap* Map, bool& bOutSuccess, TUniqueFunction<TUniquePtr<Chaos::FBaseRewindHistory>()> CreateHistoryFunction);
+
+public:
+	/** The history to be serialized */
+	TUniquePtr<Chaos::FBaseRewindHistory> History;
 
-	// Package map used for the net bit writer/reader
-	UPackageMap* PackageMap = nullptr;
+	/** Component that utilizes this data */
+	UPROPERTY()
+	TObjectPtr<UNetworkPhysicsComponent> Owner = nullptr;
+};
+
+/**
+ * Struct suitable for use as a replicated property to replicate input rewind history
+ */
+USTRUCT()
+struct FNetworkPhysicsRewindDataInputProxy : public FNetworkPhysicsRewindDataProxy
+{
+	GENERATED_BODY()
+		
+	bool NetSerialize(FArchive& Ar, class UPackageMap* Map, bool& bOutSuccess);
+};
+
+template<>
+struct TStructOpsTypeTraits<FNetworkPhysicsRewindDataInputProxy> : public TStructOpsTypeTraitsBase2<FNetworkPhysicsRewindDataInputProxy>
+{
+	enum
+	{
+		WithNetSerializer = true,
+		WithIdenticalViaEquality = true
+	};
+};
+
+/**
+ * Struct suitable for use as a replicated property to replicate state rewind history
+ */
+USTRUCT()
+struct FNetworkPhysicsRewindDataStateProxy : public FNetworkPhysicsRewindDataProxy
+{
+	GENERATED_BODY()
+
+	bool NetSerialize(FArchive& Ar, class UPackageMap* Map, bool& bOutSuccess);
+};
+
+template<>
+struct TStructOpsTypeTraits<FNetworkPhysicsRewindDataStateProxy> : public TStructOpsTypeTraitsBase2<FNetworkPhysicsRewindDataStateProxy>
+{
+	enum
+	{
+		WithNetSerializer = true,
+		WithIdenticalViaEquality = true
+	};
 };
 
 /**
@@ -306,8 +378,8 @@ public:
 	ENGINE_API void InitPhysics();
 
 	// Server RPC to receive inputs from client
-	UFUNCTION(Server, unreliable, WithValidation)
-	ENGINE_API void ServerReceiveInputsDatas(const TArray<uint8>& ClientInputs);
+	UFUNCTION(Server, unreliable)
+	ENGINE_API void ServerReceiveInputsDatas(const FNetworkPhysicsRewindDataInputProxy& ClientInputs);
 
 	// Async physics tick component function per frame from the solver
 	ENGINE_API virtual void AsyncPhysicsTickComponent(float DeltaTime, float SimTime) override;
@@ -363,9 +435,6 @@ public:
 
 protected : 
 
-	// Update the histories packagemap for serialization 
-	ENGINE_API void UpdatePackageMap();
-
 	// repnotify for the inputs on the client
 	UFUNCTION()
 	ENGINE_API void OnRep_SetReplicatedInputs();
@@ -376,17 +445,11 @@ protected :
 
 	// replicated physics inputs
 	UPROPERTY(Transient, ReplicatedUsing = OnRep_SetReplicatedInputs)
-	TArray<uint8> ReplicatedInputs;
+	FNetworkPhysicsRewindDataInputProxy ReplicatedInputs;
 
 	// replicated physics states 
 	UPROPERTY(Transient, ReplicatedUsing = OnRep_SetReplicatedStates)
-	TArray<uint8> ReplicatedStates;
-
-	// List of inputs that are processed locally
-	TArray<uint8> LocalInputs;
-
-	// List of states that are processed locally
-	TArray<uint8> LocalStates;
+	FNetworkPhysicsRewindDataStateProxy ReplicatedStates;
 
 	// Frame counter to compute the local to server offset
 	int32 FrameCounter = 0;
@@ -394,6 +457,8 @@ protected :
 private:
 
 	friend FNetworkPhysicsCallback;
+	friend struct FNetworkPhysicsRewindDataInputProxy;
+	friend struct FNetworkPhysicsRewindDataStateProxy;
 
 	// States history uses to rewind simulation 
 	TSharedPtr<Chaos::FBaseRewindHistory> StatesHistory;
@@ -436,14 +501,18 @@ FORCEINLINE void UNetworkPhysicsComponent::CreateDatasHistory(UActorComponent* H
 	const bool bIsLocalHistory = (Controller && Controller->IsLocalController());
 	const int32 NumFrames = UPhysicsSettings::Get()->GetPhysicsHistoryCount();
 	
-	UPackageMap* PackageMap = (Controller  && Controller->GetNetConnection())? Controller->GetNetConnection()->PackageMap : nullptr;
-
-	InputsHistory = MakeShared<TNetRewindHistory<typename PhysicsTraits::InputsType>>(NumFrames, bIsLocalHistory, PackageMap);
-	StatesHistory = MakeShared<TNetRewindHistory<typename PhysicsTraits::StatesType>>(NumFrames, bIsLocalHistory, PackageMap);
+	InputsHistory = MakeShared<TNetRewindHistory<typename PhysicsTraits::InputsType>>(NumFrames, bIsLocalHistory);
+	StatesHistory = MakeShared<TNetRewindHistory<typename PhysicsTraits::StatesType>>(NumFrames, bIsLocalHistory);
 
 	InputsDatas = MakeUnique<typename PhysicsTraits::InputsType>();
 	StatesDatas = MakeUnique<typename PhysicsTraits::StatesType>();
 
+	ReplicatedInputs.History = MakeUnique<TNetRewindHistory<typename PhysicsTraits::InputsType>>(NumFrames, bIsLocalHistory);
+	ReplicatedInputs.Owner = this;
+
+	ReplicatedStates.History = MakeUnique<TNetRewindHistory<typename PhysicsTraits::StatesType>>(NumFrames, bIsLocalHistory);
+	ReplicatedStates.Owner = this;
+	
 	ActorComponent = HistoryComponent;
 	
 	AddDatasHistory();
diff --git a/Engine/Source/Runtime/Experimental/Chaos/Public/RewindData.h b/Engine/Source/Runtime/Experimental/Chaos/Public/RewindData.h
--- a/Engine/Source/Runtime/Experimental/Chaos/Public/RewindData.h
+++ b/Engine/Source/Runtime/Experimental/Chaos/Public/RewindData.h
@@ -27,6 +27,12 @@ struct FBaseRewindHistory
 {
 	FORCEINLINE virtual ~FBaseRewindHistory() {}
 
+	/** Create a new, empty instance with the same concrete type as this object */
+	virtual TUniquePtr<FBaseRewindHistory> CreateNew() const = 0;
+
+	/** Create a polymorphic copy of the history */
+	virtual TUniquePtr<FBaseRewindHistory> Clone() const = 0;
+
 	/** Set the package map for serialization */
 	FORCEINLINE virtual void SetPackageMap(class UPackageMap* InPackageMap) {}
 
@@ -39,14 +45,17 @@ struct FBaseRewindHistory
 	/** Record datas into the history buffer at a given time */
 	FORCEINLINE virtual bool RecordDatas(const int32 RecordFrame, const void* HistoryDatas) { return true; }
 
-	/** Serialize the datas into an array of uint8 */
-	FORCEINLINE virtual void SerializeDatas(const uint32 StartFrame, const uint32 EndFrame, TArray<uint8>& ArchiveDatas, const int32 FrameOffset) const {}
+	/** Create a polymorphic copy of only a range of frames, applying the frame offset to the copies */
+	virtual TUniquePtr<FBaseRewindHistory> CopyFramesWithOffset(const uint32 StartFrame, const uint32 EndFrame, const int32 FrameOffset) = 0;
+
+	/** Copy new datas (received from the network) into this history */
+	virtual void ReceiveNewDatas(FBaseRewindHistory& NewDatas, const int32 FrameOffset) {}
 
-	/** Deserialize the datas from the array of uint8 */
-	FORCEINLINE virtual void DeserializeDatas(const TArray<uint8>& ArchiveDatas, const int32 FrameOffset) {}
+	/** Serialize the datas to or from a network archive */
+	virtual void NetSerialize(FArchive& Ar, UPackageMap* PackageMap) {}
 
 	/** Debug the datas from the array of uint8 that will be transferred from client to server */
-	FORCEINLINE virtual void DebugDatas(const TArray<uint8>& ArchiveDatas, TArray<int32>& LocalFrames, TArray<int32>& ServerFrames, TArray<int32>& InputFrames) {}
+	FORCEINLINE virtual void DebugDatas(const Chaos::FBaseRewindHistory& NewDatas, TArray<int32>& LocalFrames, TArray<int32>& ServerFrames, TArray<int32>& InputFrames) {}
 
 	/** Legacy interface to rewind states */
 	FORCEINLINE virtual bool RewindStates(const int32 RewindFrame, const bool bResetSolver) { return false; }
