World Partition 2D Debug Display:
- Fix support to render rotated grids.
- Merge code between different hash implementations.

diff --git a/Engine/Source/Runtime/Engine/Private/WorldPartition/RuntimeHashSet/WorldPartitionRuntimeHashSet.cpp b/Engine/Source/Runtime/Engine/Private/WorldPartition/RuntimeHashSet/WorldPartitionRuntimeHashSet.cpp
--- a/Engine/Source/Runtime/Engine/Private/WorldPartition/RuntimeHashSet/WorldPartitionRuntimeHashSet.cpp
+++ b/Engine/Source/Runtime/Engine/Private/WorldPartition/RuntimeHashSet/WorldPartitionRuntimeHashSet.cpp
@@ -245,27 +245,6 @@ bool UWorldPartitionRuntimeHashSet::Draw2D(FWorldPartitionDraw2DContext& DrawCon
 			return InWorldToScreen(Pos, bIsLocal);
 		};
 
-		auto LocalDrawTile = [&](const FVector2D& Min, const FVector2D& Size, const FLinearColor& Color)
-		{
-			const FVector2D ScreenMin = WorldToScreen(Min);
-			const FVector2D ScreenMax = WorldToScreen(Min + Size);
-			DrawContext.PushDrawTile(GridScreenBounds, ScreenMin, ScreenMax, Color);
-		};
-
-		auto LocalDrawBox = [&](const FVector2D& Min, const FVector2D& Size, const FLinearColor& Color, float LineThickness, const FBox2D* CustomGridScreenBounds = nullptr)
-		{
-			const FVector2D ScreenMin = WorldToScreen(Min);
-			const FVector2D ScreenMax = WorldToScreen(Min + Size);
-			DrawContext.PushDrawBox(GridScreenBounds, ScreenMin, ScreenMax, Color, LineThickness);
-		};
-
-		auto LocalDrawSegment = [&](const FVector2D& Start, const FVector2D& End, const FLinearColor& Color, float LineThickness)
-		{
-			const FVector2D ScreenStart = WorldToScreen(Start);
-			const FVector2D ScreenEnd = WorldToScreen(End);
-			DrawContext.PushDrawSegment(GridScreenBounds, ScreenStart, ScreenEnd, Color, LineThickness);
-		};
-
 		TArray<const UWorldPartitionRuntimeCell*> FilteredCells;
 		const FBox Region3D(FVector(Region.Min.X, Region.Min.Y, -HALF_WORLD_MAX), FVector(Region.Max.X, Region.Max.Y, HALF_WORLD_MAX));		
 		StreamingData->SpatialIndex->ForEachIntersectingElement(Region3D, [&FilteredCells](UWorldPartitionRuntimeCell* Cell)
@@ -284,8 +263,8 @@ bool UWorldPartitionRuntimeHashSet::Draw2D(FWorldPartitionDraw2DContext& DrawCon
 			{
 				const FVector2D CellBoundsSize = FVector2D(Cell->GetCellBounds().GetSize());
 				const FVector2D CellBoundsMin = FVector2D(Cell->GetCellBounds().Min);
-				LocalDrawTile(CellBoundsMin, CellBoundsSize, Cell->GetDebugColor(VisualizeMode).CopyWithNewOpacity(0.25f));
-				LocalDrawBox(CellBoundsMin, CellBoundsSize, FLinearColor::Black, 1);
+				DrawContext.LocalDrawTile(GridScreenBounds, CellBoundsMin, CellBoundsSize, Cell->GetDebugColor(VisualizeMode).CopyWithNewOpacity(0.25f), WorldToScreen);
+				DrawContext.LocalDrawBox(GridScreenBounds, CellBoundsMin, CellBoundsSize, FLinearColor::Black, 1, WorldToScreen);
 			}
 		}
 
@@ -301,7 +280,7 @@ bool UWorldPartitionRuntimeHashSet::Draw2D(FWorldPartitionDraw2DContext& DrawCon
 		{
 			const FBox2D Bounds = GridScreenBounds.ExpandBy(FVector2D(10));
 			const FVector2D Size = GridScreenBounds.GetSize();
-			DrawContext.PushDrawBox(Bounds, GridScreenBounds.Min, GridScreenBounds.Max, FLinearColor::White, 1);
+			DrawContext.PushDrawBox(Bounds, GridScreenBounds.Min, GridScreenBounds.Min + FVector2D(Size.X, 0), GridScreenBounds.Max, GridScreenBounds.Min + FVector2D(0, Size.Y), FLinearColor::White, 1);
 		}
 
 		// Draw Streaming Sources
diff --git a/Engine/Source/Runtime/Engine/Private/WorldPartition/WorldPartitionDraw2DContext.h b/Engine/Source/Runtime/Engine/Private/WorldPartition/WorldPartitionDraw2DContext.h
--- a/Engine/Source/Runtime/Engine/Private/WorldPartition/WorldPartitionDraw2DContext.h
+++ b/Engine/Source/Runtime/Engine/Private/WorldPartition/WorldPartitionDraw2DContext.h
@@ -165,18 +165,18 @@ public:
 		return false;
 	};
 
-	bool PushDrawBox(const FBox2D& InBounds, const FVector2D& Min, const FVector2D& Max, const FLinearColor& InColor, float InLineThickness)
+	bool PushDrawBox(const FBox2D& InBounds, const FVector2D& A, const FVector2D& B, const FVector2D& C, const FVector2D& D, const FLinearColor& InColor, float InLineThickness)
 	{
-		PushDrawSegment(InBounds, Min, FVector2D(Max.X, Min.Y), InColor, InLineThickness);
-		PushDrawSegment(InBounds, FVector2D(Max.X, Min.Y), Max, InColor, InLineThickness);
-		PushDrawSegment(InBounds, Max, FVector2D(Min.X, Max.Y), InColor, InLineThickness);
-		PushDrawSegment(InBounds, FVector2D(Min.X, Max.Y), Min, InColor, InLineThickness);
+		PushDrawSegment(InBounds, A, B, InColor, InLineThickness);
+		PushDrawSegment(InBounds, B, C, InColor, InLineThickness);
+		PushDrawSegment(InBounds, C, D, InColor, InLineThickness);
+		PushDrawSegment(InBounds, A, D, InColor, InLineThickness);
 		return true;
 	};
 
-	bool PushDrawTile(const FBox2D& InBounds, const FVector2D& Min, const FVector2D& Max, const FLinearColor& InColor)
+	bool PushDrawTile(const FBox2D& InBounds, const FVector2D& A, const FVector2D& B, const FVector2D& C, const FVector2D& D, const FLinearColor& InColor)
 	{
-		const FBox2D ClipBox = InBounds.Overlap(FBox2D(Min, Max));
+		const FBox2D ClipBox = InBounds.Overlap(FBox2D({ A, B, C, D}));
 		if (ClipBox.bIsValid)
 		{
 			FWorldPartitionCanvasBoxItem Box(ClipBox.Min, FVector2D(ClipBox.Max.X, ClipBox.Min.Y), ClipBox.Max, FVector2D(ClipBox.Min.X, ClipBox.Max.Y), InColor);
@@ -192,6 +192,34 @@ public:
 		return true;
 	}
 
+	template <class WorldToScreenFunc>
+	void LocalDrawTile(const FBox2D& GridScreenBounds, const FVector2D& Min, const FVector2D& Size, const FLinearColor& Color, WorldToScreenFunc WorldToScreen)
+	{
+		FVector2D A = WorldToScreen(Min);
+		FVector2D B = WorldToScreen(Min + FVector2D(Size.X, 0));
+		FVector2D C = WorldToScreen(Min + Size);
+		FVector2D D = WorldToScreen(Min + FVector2D(0, Size.Y));
+		PushDrawTile(GridScreenBounds, A, B, C, D, Color);
+	};
+
+	template <class WorldToScreenFunc>
+	void LocalDrawBox(const FBox2D& GridScreenBounds, const FVector2D& Min, const FVector2D& Size, const FLinearColor& Color, float LineThickness, WorldToScreenFunc WorldToScreen)
+	{
+		FVector2D A = WorldToScreen(Min);
+		FVector2D B = WorldToScreen(Min + FVector2D(Size.X, 0));
+		FVector2D C = WorldToScreen(Min + Size);
+		FVector2D D = WorldToScreen(Min + FVector2D(0, Size.Y));
+		PushDrawBox(GridScreenBounds, A, B, C, D, Color, LineThickness);
+	};
+
+	template <class WorldToScreenFunc>
+	void LocalDrawSegment(const FBox2D& GridScreenBounds, const FVector2D& Start, const FVector2D& End, const FLinearColor& Color, float LineThickness, WorldToScreenFunc WorldToScreen)
+	{
+		FVector2D A = WorldToScreen(Start);
+		FVector2D B = WorldToScreen(End);
+		PushDrawSegment(GridScreenBounds, A, B, Color, LineThickness);
+	};
+
 	const FWorldPartitionCanvasItems& GetCanvasItems() const { return CanvasItems; }
 
 private:
diff --git a/Engine/Source/Runtime/Engine/Private/WorldPartition/WorldPartitionRuntimeSpatialHash.cpp b/Engine/Source/Runtime/Engine/Private/WorldPartition/WorldPartitionRuntimeSpatialHash.cpp
--- a/Engine/Source/Runtime/Engine/Private/WorldPartition/WorldPartitionRuntimeSpatialHash.cpp
+++ b/Engine/Source/Runtime/Engine/Private/WorldPartition/WorldPartitionRuntimeSpatialHash.cpp
@@ -801,27 +801,6 @@ void FSpatialHashStreamingGrid::Draw2D(const UWorldPartitionRuntimeSpatialHash*
 		return InWorldToScreen(Pos, bIsLocal);
 	};
 
-	auto LocalDrawTile = [&](const FVector2D& Min, const FVector2D& Size, const FLinearColor& Color)
-	{
-		const FVector2D ScreenMin = WorldToScreen(Min);
-		const FVector2D ScreenMax = WorldToScreen(Min + Size);
-		DrawContext.PushDrawTile(GridScreenBounds, ScreenMin, ScreenMax, Color);
-	};
-
-	auto LocalDrawBox = [&](const FVector2D& Min, const FVector2D& Size, const FLinearColor& Color, float LineThickness, const FBox2D* CustomGridScreenBounds = nullptr)
-	{
-		const FVector2D ScreenMin = WorldToScreen(Min);
-		const FVector2D ScreenMax = WorldToScreen(Min + Size);
-		DrawContext.PushDrawBox(GridScreenBounds, ScreenMin, ScreenMax, Color, LineThickness);
-	};
-
-	auto LocalDrawSegment = [&](const FVector2D& Start, const FVector2D& End, const FLinearColor& Color, float LineThickness)
-	{
-		const FVector2D ScreenStart = WorldToScreen(Start);
-		const FVector2D ScreenEnd = WorldToScreen(End);
-		DrawContext.PushDrawSegment(GridScreenBounds, ScreenStart, ScreenEnd, Color, LineThickness);
-	};
-
 	TArray<const UWorldPartitionRuntimeCell*> FilteredCells;
 	for (int32 GridLevel = MinGridLevel; GridLevel <= MaxGridLevel; ++GridLevel)
 	{
@@ -849,7 +828,7 @@ void FSpatialHashStreamingGrid::Draw2D(const UWorldPartitionRuntimeSpatialHash*
 					// Draw Cell using its debug color
 					FVector2D StartPos = CellWorldBounds.Min + CellOffset;
 					{
-						LocalDrawTile(StartPos, CellBoundsSize, Cell->GetDebugColor(VisualizeMode).CopyWithNewOpacity(0.25f));
+						DrawContext.LocalDrawTile(GridScreenBounds, StartPos, CellBoundsSize, Cell->GetDebugColor(VisualizeMode).CopyWithNewOpacity(0.25f), WorldToScreen);
 					}
 
 					CellOffset.Y += CellBoundsSize.Y;
@@ -865,7 +844,7 @@ void FSpatialHashStreamingGrid::Draw2D(const UWorldPartitionRuntimeSpatialHash*
 							for (const FName& DataLayer : Cell->GetDataLayers())
 							{
 								const FColor& DataLayerColor = DataLayerDebugColors[DataLayer];
-								LocalDrawTile(StartPos + DataLayerOffset, DataLayerColoredBoxSize, DataLayerColor);
+								DrawContext.LocalDrawTile(GridScreenBounds, StartPos + DataLayerOffset, DataLayerColoredBoxSize, DataLayerColor, WorldToScreen);
 								DataLayerOffset.X += DataLayerColoredBoxSize.X;
 							}
 						}
@@ -879,14 +858,14 @@ void FSpatialHashStreamingGrid::Draw2D(const UWorldPartitionRuntimeSpatialHash*
 								FVector2D BoxSize = CellBoundsSize;
 								BoxSize.X /= 5; // Use 20% of cell's width
 								FVector2D ContentBundleOffset(CellBoundsSize.X - BoxSize.X, 0);
-								LocalDrawTile(StartPos + ContentBundleOffset, BoxSize, ContentBundle->GetDescriptor()->GetDebugColor());
+								DrawContext.LocalDrawTile(GridScreenBounds, StartPos + ContentBundleOffset, BoxSize, ContentBundle->GetDescriptor()->GetDebugColor(), WorldToScreen);
 							}
 						}
 					}
 				}
 
 				// Draw cell bounds
-				LocalDrawBox(CellWorldBounds.Min, CellWorldBounds.GetSize(), FLinearColor::Black, 1);
+				DrawContext.LocalDrawBox(GridScreenBounds, CellWorldBounds.Min, CellWorldBounds.GetSize(), FLinearColor::Black, 1, WorldToScreen);
 			}
 		});
 	}
@@ -903,13 +882,13 @@ void FSpatialHashStreamingGrid::Draw2D(const UWorldPartitionRuntimeSpatialHash*
 	{
 		FBox2D Bounds = GridScreenBounds.ExpandBy(FVector2D(10));
 		FVector2D Size = GridScreenBounds.GetSize();
-		DrawContext.PushDrawBox(Bounds, GridScreenBounds.Min, GridScreenBounds.Max, DebugColor, 1);
+		DrawContext.PushDrawBox(Bounds, GridScreenBounds.Min, GridScreenBounds.Min + FVector2D(Size.X, 0), GridScreenBounds.Max, GridScreenBounds.Min + FVector2D(0, Size.Y), DebugColor, 1);
 	}
 
 	// Draw WorldBounds
 	if (DrawContext.IsDetailedMode())
 	{
-		LocalDrawBox(FVector2D(WorldBounds.Min), FVector2D(WorldBounds.GetSize()), FLinearColor::Yellow, 1);
+		DrawContext.LocalDrawBox(GridScreenBounds, FVector2D(WorldBounds.Min), FVector2D(WorldBounds.GetSize()), FLinearColor::Yellow, 1, WorldToScreen);
 	}
 
 	// Draw Streaming Sources
