Add more verbose output for assert in UActorDescContainer::UnregisterActorDescriptor.

diff --git a/Engine/Source/Runtime/Engine/Private/WorldPartition/ActorDescContainer.cpp b/Engine/Source/Runtime/Engine/Private/WorldPartition/ActorDescContainer.cpp
--- a/Engine/Source/Runtime/Engine/Private/WorldPartition/ActorDescContainer.cpp
+++ b/Engine/Source/Runtime/Engine/Private/WorldPartition/ActorDescContainer.cpp
@@ -274,7 +274,7 @@ void UActorDescContainer::UnregisterActorDescriptor(FWorldPartitionActorDesc* Ac
 {
 	FActorDescList::RemoveActorDescriptor(ActorDesc);
 	ActorDesc->SetContainer(nullptr, nullptr);
-	verify(ActorsByName.Remove(ActorDesc->GetActorName()));	
+	verifyf(ActorsByName.Remove(ActorDesc->GetActorName()), TEXT("Missing actor '%s' from container '%s'"), *ActorDesc->GetActorName().ToString(), *ContainerPackageName.ToString());
 }
 
 bool UActorDescContainer::ShouldHandleActorEvent(const AActor* Actor)
