[shaders]
- Remove the condition causing SCW to automatically exit after a single job if HLSLcc was used. This doesn't appear to be beneficial; HLSLcc uses a lot of memory but does not appear to leak it, so we're wasting time spinning up new processes unnecessarily here. This allows us to deprecated the associated field in the compiler output and function in IShaderFormat.
- Strip out compiler invocations stats; this was not accurate in some cases and it was decided after discussion that at this point the effort to fix the problems was not worth the utility we get from it.

diff --git a/Engine/Source/Developer/ShaderFormatOpenGL/Private/ShaderFormatOpenGL.cpp b/Engine/Source/Developer/ShaderFormatOpenGL/Private/ShaderFormatOpenGL.cpp
--- a/Engine/Source/Developer/ShaderFormatOpenGL/Private/ShaderFormatOpenGL.cpp
+++ b/Engine/Source/Developer/ShaderFormatOpenGL/Private/ShaderFormatOpenGL.cpp
@@ -85,11 +85,6 @@ public:
 	{
 		return TEXT("GL");
 	}
-
-	virtual bool UsesHLSLcc(const struct FShaderCompilerInput& Input) const override
-	{
-		return !Input.Environment.CompilerFlags.Contains(CFLAG_ForceDXC);
-	}
 };
 
 /**
diff --git a/Engine/Source/Developer/ShaderFormatVectorVM/Private/ShaderFormatVectorVM.cpp b/Engine/Source/Developer/ShaderFormatVectorVM/Private/ShaderFormatVectorVM.cpp
--- a/Engine/Source/Developer/ShaderFormatVectorVM/Private/ShaderFormatVectorVM.cpp
+++ b/Engine/Source/Developer/ShaderFormatVectorVM/Private/ShaderFormatVectorVM.cpp
@@ -67,11 +67,6 @@ public:
 	{
 		return TEXT("");
 	}
-
-	virtual bool UsesHLSLcc(const struct FShaderCompilerInput& Input) const override
-	{
-		return true;
-	}
 };
 
 /**
diff --git a/Engine/Source/Developer/TargetPlatform/Public/Interfaces/IShaderFormat.h b/Engine/Source/Developer/TargetPlatform/Public/Interfaces/IShaderFormat.h
--- a/Engine/Source/Developer/TargetPlatform/Public/Interfaces/IShaderFormat.h
+++ b/Engine/Source/Developer/TargetPlatform/Public/Interfaces/IShaderFormat.h
@@ -130,6 +130,7 @@ public:
 	 * Can the shader compiler use the HLSLcc library when compiling shaders
 	 * @returns True if the shader compiler can use the HLSLcc library when compiling shaders, otherwise false.
 	 */
+	UE_DEPRECATED(5.3, "UsesHLSLcc function is no longer used")
 	virtual bool UsesHLSLcc(const FShaderCompilerInput& Input) const { return false; }
 
 	/**
diff --git a/Engine/Source/Programs/ShaderCompileWorker/Private/ShaderCompileWorker.cpp b/Engine/Source/Programs/ShaderCompileWorker/Private/ShaderCompileWorker.cpp
--- a/Engine/Source/Programs/ShaderCompileWorker/Private/ShaderCompileWorker.cpp
+++ b/Engine/Source/Programs/ShaderCompileWorker/Private/ShaderCompileWorker.cpp
@@ -278,9 +278,9 @@ public:
 				break;
 			}
 
-			if (TimeToLive == 0 || AnyJobUsedHLSLccCompiler(SingleJobs, PipelineJobs))
+			if (TimeToLive == 0)
 			{
-				UE_LOG(LogShaders, Log, TEXT("TimeToLive set to 0, or used HLSLcc compiler, exiting after single job"));
+				UE_LOG(LogShaders, Log, TEXT("TimeToLive set to 0, exiting after single job"));
 				break;
 			}
 
@@ -838,32 +838,6 @@ private:
 		}
 #endif
 	}
-	
-	static bool AnyJobUsedHLSLccCompiler(TArray<FShaderCompileJob>& SingleJobs, TArray<FShaderPipelineCompileJob>& PipelineJobs)
-	{
-		TRACE_CPUPROFILER_EVENT_SCOPE(AnyJobUsedHLSLccCompiler);
-
-		for (int32 JobIndex = 0; JobIndex < SingleJobs.Num(); JobIndex++)
-		{
-			if (SingleJobs[JobIndex].Output.bUsedHLSLccCompiler)
-			{
-				return true;
-			}
-		}
-
-		for (int32 JobIndex = 0; JobIndex < PipelineJobs.Num(); JobIndex++)
-		{
-			FShaderPipelineCompileJob& PipelineJob = PipelineJobs[JobIndex];
-			for (int32 Index = 0; Index < PipelineJob.StageJobs.Num(); ++Index)
-			{
-				if (PipelineJob.StageJobs[Index]->Output.bUsedHLSLccCompiler)
-				{
-					return true;
-				}
-			}
-		}
-		return false;
-	}
 };
 
 static void DirectCompile(const TArray<const class IShaderFormat*>& ShaderFormats)
diff --git a/Engine/Source/Runtime/Engine/Private/ShaderCompiler/ShaderCompiler.cpp b/Engine/Source/Runtime/Engine/Private/ShaderCompiler/ShaderCompiler.cpp
--- a/Engine/Source/Runtime/Engine/Private/ShaderCompiler/ShaderCompiler.cpp
+++ b/Engine/Source/Runtime/Engine/Private/ShaderCompiler/ShaderCompiler.cpp
@@ -3406,39 +3406,6 @@ static FString PrintJobsCompletedPercentageToString(int64 JobsAssigned, int64 Jo
 	return FString::Printf(TEXT("%.2f%%"), FMath::Min(JobsCompletedPercentage, 99.99));
 }
 
-static FString WriteShaderCompilerInvocationStats(const FShaderCompilerStats::FCompilerInvocations& Invocations)
-{
-	FString Stats;
-
-	auto AppendStat = [&Stats](const TCHAR* Text)
-	{
-		Stats += FString::Printf(TEXT("%s%s"), (Stats.IsEmpty() ? TEXT("") : TEXT(", ")), Text);
-	};
-	auto AppendInvocationStat = [&AppendStat](const TCHAR* CompilerName, int32 Invocations)
-	{
-		if (Invocations > 0)
-		{
-			AppendStat(*FString::Printf(TEXT("%s (%d)"), CompilerName, Invocations));
-		}
-	};
-	if (Invocations.Dxc > 0)
-	{
-		if (Invocations.DxcPrecompileSteps > 0)
-		{
-			AppendStat(*FString::Printf(TEXT("DXC (%d pre-compile steps / %d)"), Invocations.DxcPrecompileSteps, Invocations.Dxc));
-		}
-		else
-		{
-			AppendInvocationStat(TEXT("DXC"), Invocations.Dxc);
-		}
-	}
-	AppendInvocationStat(TEXT("Legacy FXC"), Invocations.Fxc);
-	AppendInvocationStat(TEXT("HLSLcc"), Invocations.Hlslcc);
-	AppendStat(*FString::Printf(TEXT("Platform compiler (%s)"), Invocations.bWasPlatformCompilerUsed ? TEXT("YES") : TEXT("NO")));
-
-	return Stats;
-}
-
 void FShaderCompilerStats::WriteStatSummary()
 {
 	const uint32 TotalCompiled = GetTotalShadersCompiled();
@@ -3453,8 +3420,6 @@ void FShaderCompilerStats::WriteStatSummary()
 		*FormatNumber(JobsCompleted),
 		*PrintJobsCompletedPercentageToString(JobsAssigned, JobsCompleted));
 
-	UE_LOG(LogShaderCompilers, Display, TEXT("Compiler invocations: %s"), *WriteShaderCompilerInvocationStats(CompilerInvocations));
-
 	if (TimesLocalWorkersWereIdle > 0.0)
 	{
 		UE_LOG(LogShaderCompilers, Display, TEXT("Average time worker was idle: %.2f s"), AccumulatedLocalWorkerIdleTime / TimesLocalWorkersWereIdle);
@@ -3721,49 +3686,6 @@ void FShaderCompilerStats::RegisterFinishedJob(FShaderCommonCompileJob& Job)
 			++NumAccumulatedShaderCodes;
 		}
 
-		// Register shader compiler invocations
-		{
-			// Record FXC invocation if target is D3D11 (no other paltform uses FXC)
-			const EShaderPlatform InputShaderPlatform = SingleJob.Input.Target.GetPlatform();
-			const bool bIsD3DPlatform = FDataDrivenShaderPlatformInfo::GetIsLanguageD3D(InputShaderPlatform);
-			const bool bUsedFxc = bIsD3DPlatform && SingleJob.Input.CanCompileWithLegacyFxc();
-			if (bUsedFxc)
-			{
-				CompilerInvocations.Fxc++;
-			}
-
-			const bool bSupportsDxc = FDataDrivenShaderPlatformInfo::GetSupportsDxc(InputShaderPlatform);
-			const bool bSupportsDxcPrecompileStep = FDataDrivenShaderPlatformInfo::GetIsPC(InputShaderPlatform);
-			if (bSupportsDxc || bSupportsDxcPrecompileStep)
-			{
-				// Record DXC pre-compile step if FXC was used and HLSL 2021 or an explicit pre-compile step was requested
-				const bool bHlslVersion2021 = SingleJob.Input.Environment.CompilerFlags.Contains(CFLAG_HLSL2021);
-				const bool bPrecompileWithDXC = SingleJob.Input.Environment.CompilerFlags.Contains(CFLAG_PrecompileWithDXC);
-				const bool bWasPrecompiledWithDXC = bUsedFxc && (bHlslVersion2021 || bPrecompileWithDXC);
-				if (bWasPrecompiledWithDXC)
-				{
-					CompilerInvocations.DxcPrecompileSteps++;
-				}
-
-				// Record DXC invocation if it was explicitly requested or if there was a pre-compile step
-				const bool bForceDxc = SingleJob.Input.Environment.CompilerFlags.Contains(CFLAG_ForceDXC);
-				if (bForceDxc || bWasPrecompiledWithDXC)
-				{
-					CompilerInvocations.Dxc++;
-				}
-			}
-
-			// Record HLSLcc invocation if output is marked as such
-			if (SingleJob.Output.bUsedHLSLccCompiler)
-			{
-				CompilerInvocations.Hlslcc++;
-			}
-
-			// Assume any console target requires a platform compiler
-			const bool bIsConsole = FDataDrivenShaderPlatformInfo::GetIsConsole(InputShaderPlatform);
-			CompilerInvocations.bWasPlatformCompilerUsed = bIsConsole;
-		}
-
 		const FString ShaderName(SingleJob.Key.ShaderType->GetName());
 		if (FShaderTimings* Existing = ShaderTimings.Find(ShaderName))
 		{
diff --git a/Engine/Source/Runtime/Engine/Public/ShaderCompiler.h b/Engine/Source/Runtime/Engine/Public/ShaderCompiler.h
--- a/Engine/Source/Runtime/Engine/Public/ShaderCompiler.h
+++ b/Engine/Source/Runtime/Engine/Public/ShaderCompiler.h
@@ -470,15 +470,6 @@ public:
 		float AverageCompileTime = 0.0f;	// stored explicitly as an optimization
 	};
 
-	struct FCompilerInvocations
-	{
-		int32 Dxc = 0;
-		int32 DxcPrecompileSteps = 0;
-		int32 Fxc = 0;
-		int32 Hlslcc = 0;
-		bool bWasPlatformCompilerUsed = false;
-	};
-
 	ENGINE_API void RegisterCookedShaders(uint32 NumCooked, float CompileTime, EShaderPlatform Platform, const FString MaterialPath, FString PermutationString = FString(""));
 	ENGINE_API void RegisterCompiledShaders(uint32 NumPermutations, EShaderPlatform Platform, const FString MaterialPath, FString PermutationString = FString(""));
 	const TSparseArray<ShaderCompilerStats>& GetShaderCompilerStats() { return CompileStats; }
@@ -587,9 +578,6 @@ private:
 
 	/** Total number of DDC hits on shader maps. */
 	uint32 ShaderMapDDCHits = 0;
-
-	/** Total number of shader compiler invocations. */
-	FCompilerInvocations CompilerInvocations;
 };
 
 
diff --git a/Engine/Source/Runtime/RenderCore/Private/ShaderCore.cpp b/Engine/Source/Runtime/RenderCore/Private/ShaderCore.cpp
--- a/Engine/Source/Runtime/RenderCore/Private/ShaderCore.cpp
+++ b/Engine/Source/Runtime/RenderCore/Private/ShaderCore.cpp
@@ -1472,11 +1472,6 @@ public:
 		}
 		Job.Output.CompileTime = FPlatformTime::Seconds() - TimeStart;
 
-		if (Compiler->UsesHLSLcc(Job.Input))
-		{
-			Job.Output.bUsedHLSLccCompiler = true;
-		}
-
 		if (CompileCount)
 		{
 			++(*CompileCount);
diff --git a/Engine/Source/Runtime/RenderCore/Public/ShaderCompilerCore.h b/Engine/Source/Runtime/RenderCore/Public/ShaderCompilerCore.h
--- a/Engine/Source/Runtime/RenderCore/Public/ShaderCompilerCore.h
+++ b/Engine/Source/Runtime/RenderCore/Public/ShaderCompilerCore.h
@@ -541,7 +541,6 @@ struct FShaderCompilerOutput
 	,	PreprocessTime(0.0)
 	,	bSucceeded(false)
 	,	bSupportsQueryingUsedAttributes(false)
-	,	bUsedHLSLccCompiler(false)
 	,	bSerializeModifiedSource(false)
 	{
 	}
@@ -560,14 +559,15 @@ struct FShaderCompilerOutput
 	UE_DEPRECATED(5.3, "bFailedRemovingUnused field is no longer used")
 	bool bFailedRemovingUnused;
 	bool bSupportsQueryingUsedAttributes;
+	UE_DEPRECATED(5.3, "bUsedHLSLccCompiler field is no longer used")
 	bool bUsedHLSLccCompiler;
 	bool bSerializeModifiedSource;
 	TArray<FString> UsedAttributes;
 
 	PRAGMA_DISABLE_DEPRECATION_WARNINGS
 	// Explicitly-defaulted copy/move ctors & assignment operators are needed temporarily due to 
-	// deprecation of bFailedRemovingUnused field. These can be removed once the deprecation
-	// window for said field ends.
+	// deprecation of bFailedRemovingUnused/bUsedHLSLccCompiler fields. These can be removed once the deprecation
+	// window for said fields ends.
 	FShaderCompilerOutput(FShaderCompilerOutput&&) = default;
 	FShaderCompilerOutput(const FShaderCompilerOutput&) = default;
 	FShaderCompilerOutput& operator=(FShaderCompilerOutput&&) = default;
