Converted GPU skin to use non-immediate command list.

diff --git a/Engine/Plugins/Animation/DeformerGraph/Source/OptimusCore/Private/DataInterfaces/OptimusDataInterfaceSkinnedMeshWrite.cpp b/Engine/Plugins/Animation/DeformerGraph/Source/OptimusCore/Private/DataInterfaces/OptimusDataInterfaceSkinnedMeshWrite.cpp
--- a/Engine/Plugins/Animation/DeformerGraph/Source/OptimusCore/Private/DataInterfaces/OptimusDataInterfaceSkinnedMeshWrite.cpp
+++ b/Engine/Plugins/Animation/DeformerGraph/Source/OptimusCore/Private/DataInterfaces/OptimusDataInterfaceSkinnedMeshWrite.cpp
@@ -192,7 +192,7 @@ void FOptimusSkinnedMeshWriteDataProviderProxy::AllocateResources(FRDGBuilder& G
 		ColorBufferUAV = GraphBuilder.CreateUAV(GraphBuilder.RegisterExternalBuffer(GWhiteVertexBufferWithRDG->Buffer), PF_B8G8R8A8);
 	}
 
-	FSkeletalMeshDeformerHelpers::UpdateVertexFactoryBufferOverrides(SkeletalMeshObject, LodIndex);
+	FSkeletalMeshDeformerHelpers::UpdateVertexFactoryBufferOverrides(GraphBuilder.RHICmdList, SkeletalMeshObject, LodIndex);
 }
 
 void FOptimusSkinnedMeshWriteDataProviderProxy::GatherDispatchData(FDispatchData const& InDispatchData)
diff --git a/Engine/Source/Runtime/Engine/Private/GPUSkinCache.cpp b/Engine/Source/Runtime/Engine/Private/GPUSkinCache.cpp
--- a/Engine/Source/Runtime/Engine/Private/GPUSkinCache.cpp
+++ b/Engine/Source/Runtime/Engine/Private/GPUSkinCache.cpp
@@ -397,7 +397,7 @@ public:
 			return IntermediateAccumulatedTangentBuffer;
 		}
 
-		void UpdateVertexFactoryDeclaration()
+		void UpdateVertexFactoryDeclaration(FRHICommandListBase& RHICmdList)
 		{
 			FGPUSkinPassthroughVertexFactory::FAddVertexAttributeDesc Desc;
 			Desc.FrameNumber = SourceVertexFactory->GetShaderData().UpdatedFrameNumber;
@@ -406,13 +406,13 @@ public:
 			Desc.SRVs[FGPUSkinPassthroughVertexFactory::Position] = GetPositionRWBuffer()->Buffer.SRV;
 			Desc.SRVs[FGPUSkinPassthroughVertexFactory::PreviousPosition] = GetPreviousPositionRWBuffer()->Buffer.SRV;
 			Desc.SRVs[FGPUSkinPassthroughVertexFactory::Tangent] = GetTangentRWBuffer()->Buffer.SRV;
-			TargetVertexFactory->SetVertexAttributes(SourceVertexFactory, Desc);
+			TargetVertexFactory->SetVertexAttributes(RHICmdList, SourceVertexFactory, Desc);
 		}
 	};
 
-	void UpdateVertexFactoryDeclaration(int32 Section)
+	void UpdateVertexFactoryDeclaration(FRHICommandListBase& RHICmdList, int32 Section)
 	{
-		DispatchData[Section].UpdateVertexFactoryDeclaration();
+		DispatchData[Section].UpdateVertexFactoryDeclaration(RHICmdList);
 	}
 
 	inline FCachedGeometry::Section GetCachedGeometry(int32 SectionIndex) const
@@ -1095,7 +1095,7 @@ class FRecomputeTangentsPerVertexPassCS : public FBaseRecomputeTangentsPerVertex
 IMPLEMENT_SHADER_TYPE(template<>, FRecomputeTangentsPerVertexPassCS<0>, TEXT("/Engine/Private/RecomputeTangentsPerVertexPass.usf"), TEXT("MainCS"), SF_Compute);
 IMPLEMENT_SHADER_TYPE(template<>, FRecomputeTangentsPerVertexPassCS<1>, TEXT("/Engine/Private/RecomputeTangentsPerVertexPass.usf"), TEXT("MainCS"), SF_Compute);
 
-void FGPUSkinCache::DispatchUpdateSkinTangents(FRHICommandListImmediate& RHICmdList, FGPUSkinCacheEntry* Entry, int32 SectionIndex, FSkinCacheRWBuffer*& StagingBuffer, bool bTrianglePass)
+void FGPUSkinCache::DispatchUpdateSkinTangents(FRHICommandList& RHICmdList, FGPUSkinCacheEntry* Entry, int32 SectionIndex, FSkinCacheRWBuffer*& StagingBuffer, bool bTrianglePass)
 {
 	FGPUSkinCacheEntry::FSectionDispatchData& DispatchData = Entry->DispatchData[SectionIndex];
 
@@ -1234,7 +1234,7 @@ void FGPUSkinCache::DispatchUpdateSkinTangents(FRHICommandListImmediate& RHICmdL
 	}
 }
 
-FGPUSkinCache::FRWBuffersAllocation* FGPUSkinCache::TryAllocBuffer(uint32 NumVertices, bool WithTangnents, bool UseIntermediateTangents, uint32 NumTriangles, FRHICommandListImmediate& RHICmdList, const FName& OwnerName)
+FGPUSkinCache::FRWBuffersAllocation* FGPUSkinCache::TryAllocBuffer(uint32 NumVertices, bool WithTangnents, bool UseIntermediateTangents, uint32 NumTriangles, FRHICommandList& RHICmdList, const FName& OwnerName)
 {
 	uint64 MaxSizeInBytes = (uint64)(GSkinCacheSceneMemoryLimitInMB * 1024.0f * 1024.0f);
 	uint64 RequiredMemInBytes = FRWBuffersAllocation::CalculateRequiredMemory(NumVertices, WithTangnents, UseIntermediateTangents, NumTriangles);
@@ -1257,7 +1257,7 @@ FGPUSkinCache::FRWBuffersAllocation* FGPUSkinCache::TryAllocBuffer(uint32 NumVer
 
 DECLARE_GPU_STAT(GPUSkinCache);
 
-void FGPUSkinCache::MakeBufferTransitions(FRHICommandListImmediate& RHICmdList, TArray<FSkinCacheRWBuffer*>& Buffers, ERHIAccess ToState)
+void FGPUSkinCache::MakeBufferTransitions(FRHICommandList& RHICmdList, TArray<FSkinCacheRWBuffer*>& Buffers, ERHIAccess ToState)
 {
 	if (Buffers.Num() > 0)
 	{
@@ -1302,7 +1302,7 @@ void FGPUSkinCache::GetBufferUAVs(const TArray<FSkinCacheRWBuffer*>& InBuffers,
 	}
 }
 
-void FGPUSkinCache::DoDispatch(FRHICommandListImmediate& RHICmdList)
+void FGPUSkinCache::DoDispatch(FRHICommandList& RHICmdList)
 {
 	int32 BatchCount = BatchDispatches.Num();
 	INC_DWORD_STAT_BY(STAT_GPUSkinCache_TotalNumChunks, BatchCount);
@@ -1562,7 +1562,7 @@ void FGPUSkinCache::DoDispatch(FRHICommandListImmediate& RHICmdList)
 		for (int32 i = 0; i < BatchCount; ++i)
 		{
 			FDispatchEntry& DispatchItem = BatchDispatches[i];
-			DispatchItem.SkinCacheEntry->UpdateVertexFactoryDeclaration(DispatchItem.Section);
+			DispatchItem.SkinCacheEntry->UpdateVertexFactoryDeclaration(RHICmdList, DispatchItem.Section);
 		}
 	}
 
@@ -1572,7 +1572,7 @@ void FGPUSkinCache::DoDispatch(FRHICommandListImmediate& RHICmdList)
 	}
 }
 
-void FGPUSkinCache::DoDispatch(FRHICommandListImmediate& RHICmdList, FGPUSkinCacheEntry* SkinCacheEntry, int32 Section, int32 RevisionNumber)
+void FGPUSkinCache::DoDispatch(FRHICommandList& RHICmdList, FGPUSkinCacheEntry* SkinCacheEntry, int32 Section, int32 RevisionNumber)
 {
 	RenderCaptureInterface::FScopedCapture RenderCapture(GNumDispatchesToCapture > 0, &RHICmdList, TEXT("GPUSkinCache"));
 	GNumDispatchesToCapture = FMath::Max(GNumDispatchesToCapture - 1, 0);
@@ -1620,14 +1620,14 @@ void FGPUSkinCache::DoDispatch(FRHICommandListImmediate& RHICmdList, FGPUSkinCac
 		DispatchUpdateSkinTangents(RHICmdList, SkinCacheEntry, Section, StagingBuffer, false);
 	}
 
-	SkinCacheEntry->UpdateVertexFactoryDeclaration(Section);
+	SkinCacheEntry->UpdateVertexFactoryDeclaration(RHICmdList, Section);
 
 	TransitionAllToReadable(RHICmdList, BuffersToTransitionToRead);
 }
 
 bool FGPUSkinCache::ProcessEntry(
 	EGPUSkinCacheEntryMode Mode,
-	FRHICommandListImmediate& RHICmdList, 
+	FRHICommandList& RHICmdList, 
 	FGPUBaseSkinVertexFactory* VertexFactory,
 	FGPUSkinPassthroughVertexFactory* TargetVertexFactory, 
 	const FSkelMeshRenderSection& BatchElement, 
@@ -1860,7 +1860,7 @@ bool FGPUSkinCache::IsGPUSkinCacheRayTracingSupported()
 
 #if RHI_RAYTRACING
 
-void FGPUSkinCache::ProcessRayTracingGeometryToUpdate(FRHICommandListImmediate& RHICmdList, FGPUSkinCacheEntry* SkinCacheEntry)
+void FGPUSkinCache::ProcessRayTracingGeometryToUpdate(FRHICommandList& RHICmdList, FGPUSkinCacheEntry* SkinCacheEntry)
 {
 	if (IsGPUSkinCacheRayTracingSupported() && SkinCacheEntry && SkinCacheEntry->GPUSkin && SkinCacheEntry->GPUSkin->bSupportRayTracing)
 	{
@@ -1878,14 +1878,14 @@ void FGPUSkinCache::ProcessRayTracingGeometryToUpdate(FRHICommandListImmediate&
 
 #endif
 
-void FGPUSkinCache::BeginBatchDispatch(FRHICommandListImmediate& RHICmdList)
+void FGPUSkinCache::BeginBatchDispatch(FRHICommandList& RHICmdList)
 {
 	check(BatchDispatches.Num() == 0);
 	bShouldBatchDispatches = true;
 	DispatchCounter = 0;
 }
 
-void FGPUSkinCache::EndBatchDispatch(FRHICommandListImmediate& RHICmdList)
+void FGPUSkinCache::EndBatchDispatch(FRHICommandList& RHICmdList)
 {
 	TRACE_CPUPROFILER_EVENT_SCOPE(FGPUSkinCache::EndBatchDispatch);
 
@@ -2011,7 +2011,7 @@ void FGPUSkinCache::PrepareUpdateSkinning(FGPUSkinCacheEntry* Entry, int32 Secti
 	check(DispatchData.PreviousPositionBuffer != DispatchData.PositionBuffer);
 }
 
-void FGPUSkinCache::DispatchUpdateSkinning(FRHICommandListImmediate& RHICmdList, FGPUSkinCacheEntry* Entry, int32 Section, uint32 RevisionNumber, TArray<FSkinCacheRWBuffer*>& BuffersToTransitionToRead)
+void FGPUSkinCache::DispatchUpdateSkinning(FRHICommandList& RHICmdList, FGPUSkinCacheEntry* Entry, int32 Section, uint32 RevisionNumber, TArray<FSkinCacheRWBuffer*>& BuffersToTransitionToRead)
 {
 	FGPUSkinCacheEntry::FSectionDispatchData& DispatchData = Entry->DispatchData[Section];
 	FGPUBaseSkinVertexFactory::FShaderDataType& ShaderData = DispatchData.SourceVertexFactory->GetShaderData();
@@ -2271,7 +2271,7 @@ void FGPUSkinCache::CVarSinkFunction()
 
 FAutoConsoleVariableSink FGPUSkinCache::CVarSink(FConsoleCommandDelegate::CreateStatic(&CVarSinkFunction));
 
-void FGPUSkinCache::IncrementDispatchCounter(FRHICommandListImmediate& RHICmdList)
+void FGPUSkinCache::IncrementDispatchCounter(FRHICommandList& RHICmdList)
 {
 	if (GSkinCacheMaxDispatchesPerCmdList > 0)
 	{
diff --git a/Engine/Source/Runtime/Engine/Private/GPUSkinVertexFactory.cpp b/Engine/Source/Runtime/Engine/Private/GPUSkinVertexFactory.cpp
--- a/Engine/Source/Runtime/Engine/Private/GPUSkinVertexFactory.cpp
+++ b/Engine/Source/Runtime/Engine/Private/GPUSkinVertexFactory.cpp
@@ -245,7 +245,7 @@ static bool DeferSkeletalLockAndFillToRHIThread()
 	return IsRunningRHIInSeparateThread() && CVarRHICmdDeferSkeletalLockAndFillToRHIThread.GetValueOnRenderThread() > 0;
 }
 
-bool FGPUBaseSkinVertexFactory::FShaderDataType::UpdateBoneData(FRHICommandListImmediate& RHICmdList, const TArray<FMatrix44f>& ReferenceToLocalMatrices,
+bool FGPUBaseSkinVertexFactory::FShaderDataType::UpdateBoneData(FRHICommandList& RHICmdList, const TArray<FMatrix44f>& ReferenceToLocalMatrices,
 	const TArray<FBoneIndexType>& BoneMap, uint32 RevisionNumber, bool bPrevious, ERHIFeatureLevel::Type InFeatureLevel, bool bUseSkinCache, bool bForceUpdateImmediately, const FName& AssetPathName)
 {
 	// stat disabled by default due to low-value/high-frequency
@@ -259,8 +259,8 @@ bool FGPUBaseSkinVertexFactory::FShaderDataType::UpdateBoneData(FRHICommandListI
 
 	if (SupportsBonesBufferSRV(InFeatureLevel))
 	{
-		check(IsInRenderingThread());
-		
+		check(IsInParallelRenderingThread());
+
 		// make sure current revision is up-to-date
 		SetCurrentRevisionNumber(RevisionNumber);
 
@@ -278,7 +278,7 @@ bool FGPUBaseSkinVertexFactory::FShaderDataType::UpdateBoneData(FRHICommandListI
 			{
 				BoneBufferPool.ReleasePooledResource(*CurrentBoneBuffer);
 			}
-			*CurrentBoneBuffer = BoneBufferPool.CreatePooledResource(VectorArraySize);
+			*CurrentBoneBuffer = BoneBufferPool.CreatePooledResource(RHICmdList, VectorArraySize);
 			check(IsValidRef(*CurrentBoneBuffer));
 			CurrentBoneBuffer->VertexBufferRHI->SetOwnerName(AssetPathName);
 		}
@@ -287,7 +287,7 @@ bool FGPUBaseSkinVertexFactory::FShaderDataType::UpdateBoneData(FRHICommandListI
 			if (!bUseSkinCache && !bForceUpdateImmediately && DeferSkeletalLockAndFillToRHIThread())
 			{
 				FRHIBuffer* VertexBuffer = CurrentBoneBuffer->VertexBufferRHI;
-				RHICmdList.EnqueueLambda([VertexBuffer, VectorArraySize, &ReferenceToLocalMatrices, &BoneMap](FRHICommandListImmediate& InRHICmdList)
+				RHICmdList.EnqueueLambda([VertexBuffer, VectorArraySize, &ReferenceToLocalMatrices, &BoneMap](FRHICommandList& InRHICmdList)
 				{
 					QUICK_SCOPE_CYCLE_COUNTER(STAT_FRHICommandUpdateBoneBuffer_Execute);
 					FMatrix3x4* LambdaChunkMatrices = (FMatrix3x4*)InRHICmdList.LockBuffer(VertexBuffer, 0, VectorArraySize, RLM_WriteOnly);
@@ -522,6 +522,11 @@ int32 FGPUBaseSkinVertexFactory::GetBoneInfluenceLimitForAsset(int32 AssetProvid
 }
 
 void FGPUBaseSkinVertexFactory::SetData(const FGPUSkinDataType* InData)
+{
+	SetData(FRHICommandListExecutor::GetImmediateCommandList(), InData);
+}
+
+void FGPUBaseSkinVertexFactory::SetData(FRHICommandListBase& RHICmdList, const FGPUSkinDataType* InData)
 {
 	check(InData);
 
@@ -531,7 +536,7 @@ void FGPUBaseSkinVertexFactory::SetData(const FGPUSkinDataType* InData)
 	}
 
 	*Data = *InData;
-	UpdateRHI(FRHICommandListImmediate::Get());
+	UpdateRHI(RHICmdList);
 }
 
 void FGPUBaseSkinVertexFactory::CopyDataTypeForLocalVertexFactory(FLocalVertexFactory::FDataType& OutDestData) const
@@ -1049,14 +1054,14 @@ IMPLEMENT_TYPE_LAYOUT(TGPUSkinAPEXClothVertexFactoryShaderParameters);
 	TGPUSkinAPEXClothVertexFactory::ClothShaderType
 -----------------------------------------------------------------------------*/
 
-bool FGPUBaseSkinAPEXClothVertexFactory::ClothShaderType::UpdateClothSimulData(FRHICommandListImmediate& RHICmdList, const TArray<FVector3f>& InSimulPositions,
+bool FGPUBaseSkinAPEXClothVertexFactory::ClothShaderType::UpdateClothSimulData(FRHICommandList& RHICmdList, const TArray<FVector3f>& InSimulPositions,
 	const TArray<FVector3f>& InSimulNormals, uint32 RevisionNumber, ERHIFeatureLevel::Type FeatureLevel, bool bForceUpdateImmediately, const FName& AssetPathName)
 {
 	QUICK_SCOPE_CYCLE_COUNTER(STAT_FGPUBaseSkinAPEXClothVertexFactory_UpdateClothSimulData);
 
 	uint32 NumSimulVerts = InSimulPositions.Num();
 
-	check(IsInRenderingThread());
+	check(IsInParallelRenderingThread());
 	
 	SetCurrentRevisionNumber(RevisionNumber);
 	FVertexBufferAndSRV* CurrentClothBuffer = &GetClothBufferForWriting();
@@ -1071,7 +1076,7 @@ bool FGPUBaseSkinAPEXClothVertexFactory::ClothShaderType::UpdateClothSimulData(F
 		{
 			ClothSimulDataBufferPool.ReleasePooledResource(*CurrentClothBuffer);
 		}
-		*CurrentClothBuffer = ClothSimulDataBufferPool.CreatePooledResource(VectorArraySize);
+		*CurrentClothBuffer = ClothSimulDataBufferPool.CreatePooledResource(RHICmdList, VectorArraySize);
 		check(IsValidRef(*CurrentClothBuffer));
 		CurrentClothBuffer->VertexBufferRHI->SetOwnerName(AssetPathName);
 	}
@@ -1081,7 +1086,7 @@ bool FGPUBaseSkinAPEXClothVertexFactory::ClothShaderType::UpdateClothSimulData(F
 		if (!bForceUpdateImmediately && DeferSkeletalLockAndFillToRHIThread())
 		{
 			FRHIBuffer* VertexBuffer = CurrentClothBuffer->VertexBufferRHI;
-			RHICmdList.EnqueueLambda([VertexBuffer, VectorArraySize, &InSimulPositions, &InSimulNormals](FRHICommandListImmediate& InRHICmdList)
+			RHICmdList.EnqueueLambda([VertexBuffer, VectorArraySize, &InSimulPositions, &InSimulNormals](FRHICommandList& InRHICmdList)
 			{
 				QUICK_SCOPE_CYCLE_COUNTER(STAT_FRHICommandUpdateBoneBuffer_Execute);
 				float* RESTRICT Data = (float* RESTRICT)InRHICmdList.LockBuffer(VertexBuffer, 0, VectorArraySize, RLM_WriteOnly);
@@ -1243,7 +1248,7 @@ bool TGPUSkinAPEXClothVertexFactory<BoneInfluenceType>::ShouldCompilePermutation
 }
 
 template <GPUSkinBoneInfluenceType BoneInfluenceType>
-void TGPUSkinAPEXClothVertexFactory<BoneInfluenceType>::SetData(const FGPUSkinDataType* InData)
+void TGPUSkinAPEXClothVertexFactory<BoneInfluenceType>::SetData(FRHICommandListBase& RHICmdList, const FGPUSkinDataType* InData)
 {
 	const FGPUSkinAPEXClothDataType* InClothData = (const FGPUSkinAPEXClothDataType*)(InData);
 	check(InClothData);
@@ -1255,7 +1260,7 @@ void TGPUSkinAPEXClothVertexFactory<BoneInfluenceType>::SetData(const FGPUSkinDa
 	}
 
 	*ClothDataPtr = *InClothData;
-	FGPUBaseSkinVertexFactory::UpdateRHI(FRHICommandListImmediate::Get());
+	FGPUBaseSkinVertexFactory::UpdateRHI(RHICmdList);
 }
 
 /**
@@ -1439,8 +1444,11 @@ void FGPUSkinPassthroughVertexFactory::CreateLooseUniformBuffer(FRHICommandListB
 
 void FGPUSkinPassthroughVertexFactory::SetVertexAttributes(FGPUBaseSkinVertexFactory const* InSourceVertexFactory, FAddVertexAttributeDesc const& InDesc)
 {
-	FRHICommandListBase& RHICmdList = FRHICommandListImmediate::Get();
+	SetVertexAttributes(FRHICommandListImmediate::Get(), InSourceVertexFactory, InDesc);
+}
 
+void FGPUSkinPassthroughVertexFactory::SetVertexAttributes(FRHICommandListBase& RHICmdList, FGPUBaseSkinVertexFactory const* InSourceVertexFactory, FAddVertexAttributeDesc const& InDesc)
+{
 	// Check for new vertex attributes.
 	bool bNeedFullUpdate = false;
 	for (int32 Index = 0; Index < InDesc.VertexAttributes.Num(); ++Index)
diff --git a/Engine/Source/Runtime/Engine/Private/SkeletalMeshDeformerHelpers.cpp b/Engine/Source/Runtime/Engine/Private/SkeletalMeshDeformerHelpers.cpp
--- a/Engine/Source/Runtime/Engine/Private/SkeletalMeshDeformerHelpers.cpp
+++ b/Engine/Source/Runtime/Engine/Private/SkeletalMeshDeformerHelpers.cpp
@@ -208,6 +208,11 @@ FRDGBuffer* FSkeletalMeshDeformerHelpers::AllocateVertexFactoryColorBuffer(FRDGB
 }
 
 void FSkeletalMeshDeformerHelpers::UpdateVertexFactoryBufferOverrides(FSkeletalMeshObject* InMeshObject, int32 InLodIndex)
+{
+	UpdateVertexFactoryBufferOverrides(FRHICommandListImmediate::Get(), InMeshObject, InLodIndex);
+}
+
+void FSkeletalMeshDeformerHelpers::UpdateVertexFactoryBufferOverrides(FRHICommandListBase& RHICmdList, FSkeletalMeshObject* InMeshObject, int32 InLodIndex)
 {
 	if (InMeshObject->IsCPUSkinned())
 	{
@@ -248,7 +253,7 @@ void FSkeletalMeshDeformerHelpers::UpdateVertexFactoryBufferOverrides(FSkeletalM
 	{
 		FGPUBaseSkinVertexFactory const* BaseVertexFactory = MeshObjectGPU->GetBaseSkinVertexFactory(InLodIndex, SectionIndex);
 		FGPUSkinPassthroughVertexFactory* TargetVertexFactory = LOD.GPUSkinVertexFactories.PassthroughVertexFactories[SectionIndex].Get();
-		TargetVertexFactory->SetVertexAttributes(BaseVertexFactory, Desc);
+		TargetVertexFactory->SetVertexAttributes(RHICmdList, BaseVertexFactory, Desc);
 	}
 }
 
diff --git a/Engine/Source/Runtime/Engine/Private/SkeletalRenderCPUSkin.cpp b/Engine/Source/Runtime/Engine/Private/SkeletalRenderCPUSkin.cpp
--- a/Engine/Source/Runtime/Engine/Private/SkeletalRenderCPUSkin.cpp
+++ b/Engine/Source/Runtime/Engine/Private/SkeletalRenderCPUSkin.cpp
@@ -186,7 +186,7 @@ void FSkeletalMeshObjectCPUSkin::UpdateSkinWeightBuffer(USkinnedMeshComponent* I
 	}
 }
 
-void FSkeletalMeshObjectCPUSkin::UpdateDynamicData_RenderThread(FRHICommandListImmediate& RHICmdList, FDynamicSkelMeshObjectDataCPUSkin* InDynamicData, uint64 FrameNumberToPrepare, uint32 RevisionNumber)
+void FSkeletalMeshObjectCPUSkin::UpdateDynamicData_RenderThread(FRHICommandList& RHICmdList, FDynamicSkelMeshObjectDataCPUSkin* InDynamicData, uint64 FrameNumberToPrepare, uint32 RevisionNumber)
 {
 	// we should be done with the old data at this point
 	delete DynamicData;
diff --git a/Engine/Source/Runtime/Engine/Private/SkeletalRenderCPUSkin.h b/Engine/Source/Runtime/Engine/Private/SkeletalRenderCPUSkin.h
--- a/Engine/Source/Runtime/Engine/Private/SkeletalRenderCPUSkin.h
+++ b/Engine/Source/Runtime/Engine/Private/SkeletalRenderCPUSkin.h
@@ -99,7 +99,7 @@ public:
 	ENGINE_API virtual void InitResources(USkinnedMeshComponent* InMeshComponent) override;
 	ENGINE_API virtual void ReleaseResources() override;
 	ENGINE_API virtual void Update(int32 LODIndex,USkinnedMeshComponent* InMeshComponent,const FMorphTargetWeightMap& InActiveMorphTargets, const TArray<float>& MorphTargetsWeights, EPreviousBoneTransformUpdateMode PreviousBoneTransformUpdateMode, const FExternalMorphWeightData& InExternalMorphWeightData) override;
-	ENGINE_API void UpdateDynamicData_RenderThread(FRHICommandListImmediate& RHICmdList, FDynamicSkelMeshObjectDataCPUSkin* InDynamicData, uint64 FrameNumberToPrepare, uint32 RevisionNumber);
+	ENGINE_API void UpdateDynamicData_RenderThread(FRHICommandList& RHICmdList, FDynamicSkelMeshObjectDataCPUSkin* InDynamicData, uint64 FrameNumberToPrepare, uint32 RevisionNumber);
 	ENGINE_API virtual void EnableOverlayRendering(bool bEnabled, const TArray<int32>* InBonesOfInterest, const TArray<UMorphTarget*>* InMorphTargetOfInterest) override;
 	ENGINE_API virtual void CacheVertices(int32 LODIndex, bool bForce) const override;
 	virtual bool IsCPUSkinned() const override { return true; }
diff --git a/Engine/Source/Runtime/Engine/Private/SkeletalRenderGPUSkin.cpp b/Engine/Source/Runtime/Engine/Private/SkeletalRenderGPUSkin.cpp
--- a/Engine/Source/Runtime/Engine/Private/SkeletalRenderGPUSkin.cpp
+++ b/Engine/Source/Runtime/Engine/Private/SkeletalRenderGPUSkin.cpp
@@ -164,13 +164,12 @@ SIZE_T FMorphVertexBufferPool::GetResourceSize() const
 	return ResourceSize;
 }
 
-void FMorphVertexBufferPool::EnableDoubleBuffer()
+void FMorphVertexBufferPool::EnableDoubleBuffer(FRHICommandListBase& RHICmdList)
 {
-	check(IsInRenderingThread());
 	bDoubleBuffer = true;
 	if (!MorphVertexBuffers[1].VertexBufferRHI.IsValid())
 	{
-		MorphVertexBuffers[1].InitResource(FRHICommandListImmediate::Get());
+		MorphVertexBuffers[1].InitResource(RHICmdList);
 	}
 }
 
@@ -482,7 +481,7 @@ FORCEINLINE bool IsDeferredSkeletalDynamicDataUpdateEnabled()
 	return CVarDeferSkeletalDynamicDataUpdateUntilGDME.GetValueOnRenderThread() > 0;
 }
 
-void FSkeletalMeshObjectGPUSkin::UpdateDynamicData_RenderThread(FGPUSkinCache* GPUSkinCache, FRHICommandListImmediate& RHICmdList, FDynamicSkelMeshObjectDataGPUSkin* InDynamicData, FSceneInterface* Scene, uint64 FrameNumberToPrepare, uint32 RevisionNumber)
+void FSkeletalMeshObjectGPUSkin::UpdateDynamicData_RenderThread(FGPUSkinCache* GPUSkinCache, FRHICommandList& RHICmdList, FDynamicSkelMeshObjectDataGPUSkin* InDynamicData, FSceneInterface* Scene, uint64 FrameNumberToPrepare, uint32 RevisionNumber)
 {
 	TRACE_CPUPROFILER_EVENT_SCOPE(GPUSkin::UpdateDynamicData_RT);
 
@@ -562,12 +561,12 @@ void FSkeletalMeshObjectGPUSkin::WaitForRHIThreadFenceForDynamicData()
 	if (RHIThreadFenceForDynamicData.GetReference())
 	{
 		FScopeCycleCounter Context(GetStatId());
-		FRHICommandListExecutor::WaitOnRHIThreadFence(RHIThreadFenceForDynamicData);
+		RHIThreadFenceForDynamicData->Wait();
 		RHIThreadFenceForDynamicData = nullptr;
 	}
 }
 
-void FSkeletalMeshObjectGPUSkin::ProcessUpdatedDynamicData(EGPUSkinCacheEntryMode Mode, FGPUSkinCache* GPUSkinCache, FRHICommandListImmediate& RHICmdList, uint32 FrameNumberToPrepare, uint32 RevisionNumber, bool bMorphNeedsUpdate, int32 LODIndex)
+void FSkeletalMeshObjectGPUSkin::ProcessUpdatedDynamicData(EGPUSkinCacheEntryMode Mode, FGPUSkinCache* GPUSkinCache, FRHICommandList& RHICmdList, uint32 FrameNumberToPrepare, uint32 RevisionNumber, bool bMorphNeedsUpdate, int32 LODIndex)
 {
 	QUICK_SCOPE_CYCLE_COUNTER(STAT_FSkeletalMeshObjectGPUSkin_ProcessUpdatedDynamicData);
 	bNeedsUpdateDeferred = false;
@@ -767,7 +766,7 @@ void FSkeletalMeshObjectGPUSkin::ProcessUpdatedDynamicData(EGPUSkinCacheEntryMod
 	if (bMorph && !bIsMobile && !SkinCacheEntry && Mode == EGPUSkinCacheEntryMode::Raster && !LOD.MorphVertexBufferPool.IsDoubleBuffered())
 	{
 		// Going through GPU skinned vertex factory, turn on double buffering for motion blur
-		LOD.MorphVertexBufferPool.EnableDoubleBuffer();
+		LOD.MorphVertexBufferPool.EnableDoubleBuffer(RHICmdList);
 	}
 }
 
@@ -1017,7 +1016,7 @@ static void CalculateMorphDeltaBoundsIncludingExternalMorphs(
 	);
 }
 
-void FSkeletalMeshObjectGPUSkin::UpdateMorphVertexBuffer(FRHICommandListImmediate& RHICmdList, EGPUSkinCacheEntryMode Mode, FSkeletalMeshObjectLOD& LOD, const FSkeletalMeshLODRenderData& LODData, 
+void FSkeletalMeshObjectGPUSkin::UpdateMorphVertexBuffer(FRHICommandList& RHICmdList, EGPUSkinCacheEntryMode Mode, FSkeletalMeshObjectLOD& LOD, const FSkeletalMeshLODRenderData& LODData, 
 															bool bGPUSkinCacheEnabled, FMorphVertexBuffer& MorphVertexBuffer)
 {
 	QUICK_SCOPE_CYCLE_COUNTER(STAT_FSkeletalMeshObjectGPUSkin_ProcessUpdatedDynamicData_UpdateMorphBuffer);
@@ -1089,7 +1088,7 @@ void FSkeletalMeshObjectGPUSkin::UpdateMorphVertexBuffer(FRHICommandListImmediat
 	else
 	{
 		// update the morph data for the lod (before SkinCache)
-		LOD.UpdateMorphVertexBufferCPU(DynamicData->ActiveMorphTargets, DynamicData->MorphTargetWeights, DynamicData->SectionIdsUseByActiveMorphTargets, bGPUSkinCacheEnabled, MorphVertexBuffer);
+		LOD.UpdateMorphVertexBufferCPU(RHICmdList, DynamicData->ActiveMorphTargets, DynamicData->MorphTargetWeights, DynamicData->SectionIdsUseByActiveMorphTargets, bGPUSkinCacheEnabled, MorphVertexBuffer);
 	}
 
 	if (LOD.MorphVertexBufferPool.IsDoubleBuffered())
@@ -1202,7 +1201,7 @@ void FGPUMorphNormalizeCS::UnsetParameters(FRHIBatchedShaderUnbinds& BatchedUnbi
 IMPLEMENT_SHADER_TYPE(, FGPUMorphNormalizeCS, TEXT("/Engine/Private/MorphTargets.usf"), TEXT("GPUMorphNormalizeCS"), SF_Compute);
 
 void FSkeletalMeshObjectGPUSkin::FSkeletalMeshObjectLOD::UpdateMorphVertexBufferGPU(
-	FRHICommandListImmediate& RHICmdList, 
+	FRHICommandList& RHICmdList, 
 	const TArray<float>& MorphTargetWeights,
 	const FMorphTargetVertexInfoBuffers& MorphTargetVertexInfoBuffers,
 	const TArray<int32>& SectionIdsUseByActiveMorphTargets,
@@ -1222,7 +1221,7 @@ void FSkeletalMeshObjectGPUSkin::FSkeletalMeshObjectLOD::UpdateMorphVertexBuffer
 		FSkeletalMeshLODRenderData& LodData = SkelMeshRenderData->LODRenderData[LODIndex];
 
 		const bool bUseGPUMorphTargets = UseGPUMorphTargets(FeatureLevel);
-		MorphVertexBuffer.RecreateResourcesIfRequired(bUseGPUMorphTargets);
+		MorphVertexBuffer.RecreateResourcesIfRequired(RHICmdList, bUseGPUMorphTargets);
 
 		SCOPED_GPU_STAT(RHICmdList, MorphTargets);
 
@@ -1363,11 +1362,10 @@ void FSkeletalMeshObjectGPUSkin::FSkeletalMeshObjectLOD::UpdateSkinWeights(FSkel
 	
 }
 
-void FSkeletalMeshObjectGPUSkin::FSkeletalMeshObjectLOD::UpdateMorphVertexBufferCPU(const FMorphTargetWeightMap& InActiveMorphTargets, const TArray<float>& MorphTargetWeights, 
+void FSkeletalMeshObjectGPUSkin::FSkeletalMeshObjectLOD::UpdateMorphVertexBufferCPU(FRHICommandList& RHICmdList, const FMorphTargetWeightMap& InActiveMorphTargets, const TArray<float>& MorphTargetWeights, 
 																					const TArray<int32>& SectionIdsUseByActiveMorphTargets, bool bGPUSkinCacheEnabled, FMorphVertexBuffer& MorphVertexBuffer)
 {
 	SCOPE_CYCLE_COUNTER(STAT_MorphVertexBuffer_Update);
-	FRHICommandListBase& RHICmdList = FRHICommandListImmediate::Get();
 
 	if (IsValidRef(MorphVertexBuffer.VertexBufferRHI))
 	{
@@ -1393,7 +1391,7 @@ void FSkeletalMeshObjectGPUSkin::FSkeletalMeshObjectLOD::UpdateMorphVertexBuffer
 		bool bBlendTangentsOnCPU = !bAllSectionsDoGPURecomputeTangent;
 		
 		const bool bUseGPUMorphTargets = UseGPUMorphTargets(FeatureLevel);
-		MorphVertexBuffer.RecreateResourcesIfRequired(bUseGPUMorphTargets);
+		MorphVertexBuffer.RecreateResourcesIfRequired(RHICmdList, bUseGPUMorphTargets);
 
 		uint32 Size = LodData.GetNumVertices() * sizeof(FMorphGPUSkinVertex);
 
@@ -1817,7 +1815,7 @@ static FGPUBaseSkinVertexFactory* CreateVertexFactory(TArray<TUniquePtr<FGPUBase
 		{
 			FGPUSkinDataType Data;
 			InitGPUSkinVertexFactoryComponents(&Data, VertexUpdateData.VertexBuffers, VertexUpdateData.VertexFactory);
-			VertexUpdateData.VertexFactory->SetData(&Data);
+			VertexUpdateData.VertexFactory->SetData(RHICmdList, &Data);
 			VertexUpdateData.VertexFactory->InitResource(RHICmdList);
 		}
 	);
@@ -1839,11 +1837,11 @@ void UpdateVertexFactory(TArray<TUniquePtr<FGPUBaseSkinVertexFactory>>& VertexFa
 
 			// update vertex factory components and sync it
 			ENQUEUE_RENDER_COMMAND(UpdateGPUSkinVertexFactory)(
-				[VertexUpdateData](FRHICommandList& CmdList)
+				[VertexUpdateData](FRHICommandList& RHICmdList)
 			{
 				FGPUSkinDataType Data;
 				InitGPUSkinVertexFactoryComponents(&Data, VertexUpdateData.VertexBuffers, VertexUpdateData.VertexFactory);
-				VertexUpdateData.VertexFactory->SetData(&Data);
+				VertexUpdateData.VertexFactory->SetData(RHICmdList, &Data);
 			});
 		}
 	}
@@ -1897,7 +1895,7 @@ static void CreateVertexFactoryMorph(TArray<TUniquePtr<FGPUBaseSkinVertexFactory
 			FGPUSkinDataType Data;
 			InitGPUSkinVertexFactoryComponents(&Data, VertexUpdateData.VertexBuffers, VertexUpdateData.VertexFactory);
 			InitMorphVertexFactoryComponents(&Data, VertexUpdateData.VertexBuffers.MorphVertexBufferPool);
-			VertexUpdateData.VertexFactory->SetData(&Data);
+			VertexUpdateData.VertexFactory->SetData(RHICmdList, &Data);
 			VertexUpdateData.VertexFactory->InitResource(RHICmdList);
 		}
 	);
@@ -1922,7 +1920,7 @@ static void UpdateVertexFactoryMorph(TArray<TUniquePtr<FGPUBaseSkinVertexFactory
 				FGPUSkinDataType Data;
 				InitGPUSkinVertexFactoryComponents(&Data, VertexUpdateData.VertexBuffers, VertexUpdateData.VertexFactory);
 				InitMorphVertexFactoryComponents(&Data, VertexUpdateData.VertexBuffers.MorphVertexBufferPool);
-				VertexUpdateData.VertexFactory->SetData(&Data);
+				VertexUpdateData.VertexFactory->SetData(RHICmdList, &Data);
 			});
 		}
 	}
@@ -1976,7 +1974,7 @@ static void CreateVertexFactoryCloth(TArray<TUniquePtr<FGPUBaseSkinAPEXClothVert
 			FGPUSkinAPEXClothDataType Data;
 			InitGPUSkinVertexFactoryComponents(&Data, VertexUpdateData.VertexBuffers, VertexUpdateData.VertexFactory);
 			InitAPEXClothVertexFactoryComponents(&Data, VertexUpdateData.VertexBuffers);
-			VertexUpdateData.VertexFactory->SetData(&Data);
+			VertexUpdateData.VertexFactory->SetData(RHICmdList, &Data);
 			VertexUpdateData.VertexFactory->InitResource(RHICmdList);
 		}
 	);
@@ -2001,7 +1999,7 @@ static void UpdateVertexFactoryCloth(TArray<TUniquePtr<FGPUBaseSkinAPEXClothVert
 				FGPUSkinAPEXClothDataType Data;
 				InitGPUSkinVertexFactoryComponents(&Data, VertexUpdateData.VertexBuffers, VertexUpdateData.VertexFactory);
 				InitAPEXClothVertexFactoryComponents(&Data, VertexUpdateData.VertexBuffers);
-				VertexUpdateData.VertexFactory->SetData(&Data);
+				VertexUpdateData.VertexFactory->SetData(RHICmdList, &Data);
 			});
 		}
 	}
@@ -2698,4 +2696,4 @@ bool FDynamicSkelMeshObjectDataGPUSkin::UpdateClothSimulationData(USkinnedMeshCo
 		return true;
 	}
 	return false;
-}
+}
\ No newline at end of file
diff --git a/Engine/Source/Runtime/Engine/Private/SkeletalRenderGPUSkin.h b/Engine/Source/Runtime/Engine/Private/SkeletalRenderGPUSkin.h
--- a/Engine/Source/Runtime/Engine/Private/SkeletalRenderGPUSkin.h
+++ b/Engine/Source/Runtime/Engine/Private/SkeletalRenderGPUSkin.h
@@ -217,11 +217,11 @@ public:
 	*/
 	virtual void ReleaseRHI();
 
-	inline void RecreateResourcesIfRequired(bool bInUsesComputeShader)
+	inline void RecreateResourcesIfRequired(FRHICommandListBase& RHICmdList, bool bInUsesComputeShader)
 	{
 		if (bUsesComputeShader != bInUsesComputeShader)
 		{
-			UpdateRHI(FRHICommandListImmediate::Get());
+			UpdateRHI(RHICmdList);
 		}
 	}
 
@@ -315,7 +315,7 @@ public:
 	void InitResources(const FName& OwnerName);
 	void ReleaseResources();
 	SIZE_T GetResourceSize() const;
-	void EnableDoubleBuffer();
+	void EnableDoubleBuffer(FRHICommandListBase& RHICmdList);
 	bool IsDoubleBuffered() const					{ return bDoubleBuffer; }
 	void SetUpdatedFrameNumber(uint32 FrameNumber)	{ UpdatedFrameNumber = FrameNumber; }
 	uint32 GetUpdatedFrameNumber() const			{ return UpdatedFrameNumber; }
@@ -352,7 +352,7 @@ public:
 	ENGINE_API virtual void InitResources(USkinnedMeshComponent* InMeshComponent) override;
 	ENGINE_API virtual void ReleaseResources() override;
 	ENGINE_API virtual void Update(int32 LODIndex,USkinnedMeshComponent* InMeshComponent,const FMorphTargetWeightMap& InActiveMorphTargets, const TArray<float>& InMorphTargetWeights, EPreviousBoneTransformUpdateMode PreviousBoneTransformUpdateMode, const FExternalMorphWeightData& InExternalMorphWeightData) override;
-	ENGINE_API void UpdateDynamicData_RenderThread(FGPUSkinCache* GPUSkinCache, FRHICommandListImmediate& RHICmdList, FDynamicSkelMeshObjectDataGPUSkin* InDynamicData, FSceneInterface* Scene, uint64 FrameNumberToPrepare, uint32 RevisionNumber);
+	ENGINE_API void UpdateDynamicData_RenderThread(FGPUSkinCache* GPUSkinCache, FRHICommandList& RHICmdList, FDynamicSkelMeshObjectDataGPUSkin* InDynamicData, FSceneInterface* Scene, uint64 FrameNumberToPrepare, uint32 RevisionNumber);
 	ENGINE_API virtual void PreGDMECallback(FGPUSkinCache* GPUSkinCache, uint32 FrameNumber) override;
 	ENGINE_API virtual const FVertexFactory* GetSkinVertexFactory(const FSceneView* View, int32 LODIndex,int32 ChunkIdx, ESkinVertexFactoryMode VFMode = ESkinVertexFactoryMode::Default) const override;
 	ENGINE_API virtual const FSkinBatchVertexFactoryUserData* GetVertexFactoryUserData(const int32 LODIndex, int32 ChunkIdx, ESkinVertexFactoryMode VFMode) const override;
@@ -636,10 +636,10 @@ protected:
 		 * @param ActiveMorphTargets - Morph to accumulate. assumed to be weighted and have valid targets
 		 * @param MorphTargetWeights - All Morph weights
 		 */
-		void UpdateMorphVertexBufferCPU(const FMorphTargetWeightMap& InActiveMorphTargets, const TArray<float>& MorphTargetWeights, const TArray<int32>& SectionIdsUseByActiveMorphTargets, 
+		void UpdateMorphVertexBufferCPU(FRHICommandList& RHICmdList, const FMorphTargetWeightMap& InActiveMorphTargets, const TArray<float>& MorphTargetWeights, const TArray<int32>& SectionIdsUseByActiveMorphTargets, 
 										bool bGPUSkinCacheEnabled, FMorphVertexBuffer& MorphVertexBuffer);
 
-		void UpdateMorphVertexBufferGPU(FRHICommandListImmediate& RHICmdList, const TArray<float>& MorphTargetWeights, const FMorphTargetVertexInfoBuffers& MorphTargetVertexInfoBuffers, 
+		void UpdateMorphVertexBufferGPU(FRHICommandList& RHICmdList, const TArray<float>& MorphTargetWeights, const FMorphTargetVertexInfoBuffers& MorphTargetVertexInfoBuffers, 
 										const TArray<int32>& SectionIdsUseByActiveMorphTargets, const FName& OwnerName, EGPUSkinCacheEntryMode Mode, FMorphVertexBuffer& MorphVertexBuffer,
 										bool bClearMorphVertexBuffer, bool bNormalizePass, const FVector4& MorphScale, const FVector4& InvMorphScale);
 
@@ -666,9 +666,9 @@ protected:
 	*/
 	ENGINE_API void ReleaseMorphResources();
 
-	ENGINE_API void ProcessUpdatedDynamicData(EGPUSkinCacheEntryMode Mode, FGPUSkinCache* GPUSkinCache, FRHICommandListImmediate& RHICmdList, uint32 FrameNumberToPrepare, uint32 RevisionNumber, bool bMorphNeedsUpdate, int32 LODIndex);
+	ENGINE_API void ProcessUpdatedDynamicData(EGPUSkinCacheEntryMode Mode, FGPUSkinCache* GPUSkinCache, FRHICommandList& RHICmdList, uint32 FrameNumberToPrepare, uint32 RevisionNumber, bool bMorphNeedsUpdate, int32 LODIndex);
 
-	ENGINE_API virtual void UpdateMorphVertexBuffer(FRHICommandListImmediate& RHICmdList, EGPUSkinCacheEntryMode Mode, FSkeletalMeshObjectLOD& LOD, const FSkeletalMeshLODRenderData& LODData, bool bGPUSkinCacheEnabled, FMorphVertexBuffer& MorphVertexBuffer);
+	ENGINE_API virtual void UpdateMorphVertexBuffer(FRHICommandList& RHICmdList, EGPUSkinCacheEntryMode Mode, FSkeletalMeshObjectLOD& LOD, const FSkeletalMeshLODRenderData& LODData, bool bGPUSkinCacheEnabled, FMorphVertexBuffer& MorphVertexBuffer);
 
 	ENGINE_API void WaitForRHIThreadFenceForDynamicData();
 
diff --git a/Engine/Source/Runtime/Engine/Public/GPUSkinCache.h b/Engine/Source/Runtime/Engine/Public/GPUSkinCache.h
--- a/Engine/Source/Runtime/Engine/Public/GPUSkinCache.h
+++ b/Engine/Source/Runtime/Engine/Public/GPUSkinCache.h
@@ -136,7 +136,7 @@ public:
 
 	bool ProcessEntry(
 		EGPUSkinCacheEntryMode Mode,
-		FRHICommandListImmediate& RHICmdList, 
+		FRHICommandList& RHICmdList, 
 		FGPUBaseSkinVertexFactory* VertexFactory,
 		FGPUSkinPassthroughVertexFactory* TargetVertexFactory, 
 		const FSkelMeshRenderSection& BatchElement, 
@@ -202,7 +202,7 @@ public:
 	{
 		friend struct FRWBufferTracker;
 
-		FRWBuffersAllocation(uint32 InNumVertices, bool InWithTangents, bool InUseIntermediateTangents, uint32 InIntermediateAccumulatedTangentsSize, FRHICommandListImmediate& RHICmdList, const FName& OwnerName)
+		FRWBuffersAllocation(uint32 InNumVertices, bool InWithTangents, bool InUseIntermediateTangents, uint32 InIntermediateAccumulatedTangentsSize, FRHICommandList& RHICmdList, const FName& OwnerName)
 			: NumVertices(InNumVertices), WithTangents(InWithTangents), UseIntermediateTangents(InUseIntermediateTangents), IntermediateAccumulatedTangentsSize(InIntermediateAccumulatedTangentsSize)
 		{
 			const static FLazyName PositionsName(TEXT("SkinCachePositions"));
@@ -417,16 +417,16 @@ public:
 	FCachedGeometry::Section GetCachedGeometry(FGPUSkinCacheEntry* InOutEntry, uint32 SectionId);
 
 #if RHI_RAYTRACING
-	void ProcessRayTracingGeometryToUpdate(FRHICommandListImmediate& RHICmdList, FGPUSkinCacheEntry* SkinCacheEntry);
+	void ProcessRayTracingGeometryToUpdate(FRHICommandList& RHICmdList, FGPUSkinCacheEntry* SkinCacheEntry);
 #endif // RHI_RAYTRACING
 
-	void BeginBatchDispatch(FRHICommandListImmediate& RHICmdList);
-	void EndBatchDispatch(FRHICommandListImmediate& RHICmdList);
+	void BeginBatchDispatch(FRHICommandList& RHICmdList);
+	void EndBatchDispatch(FRHICommandList& RHICmdList);
 
 	inline ERHIFeatureLevel::Type GetFeatureLevel() const { return FeatureLevel; }
 
 protected:
-	void MakeBufferTransitions(FRHICommandListImmediate& RHICmdList, TArray<FSkinCacheRWBuffer*>& Buffers, ERHIAccess ToState);
+	void MakeBufferTransitions(FRHICommandList& RHICmdList, TArray<FSkinCacheRWBuffer*>& Buffers, ERHIAccess ToState);
 	void GetBufferUAVs(const TArray<FSkinCacheRWBuffer*>& InBuffers, TArray<FRHIUnorderedAccessView*>& OutUAVs);
 
 	TArray<FRWBuffersAllocation*> Allocations;
@@ -434,10 +434,10 @@ protected:
 	TSet<FGPUSkinCacheEntry*> PendingProcessRTGeometryEntries;
 	TArray<FDispatchEntry> BatchDispatches;
 
-	FRWBuffersAllocation* TryAllocBuffer(uint32 NumVertices, bool WithTangnents, bool UseIntermediateTangents, uint32 NumTriangles, FRHICommandListImmediate& RHICmdList, const FName& OwnerName);
-	void DoDispatch(FRHICommandListImmediate& RHICmdList);
-	void DoDispatch(FRHICommandListImmediate& RHICmdList, FGPUSkinCacheEntry* SkinCacheEntry, int32 Section, int32 RevisionNumber);
-	void DispatchUpdateSkinTangents(FRHICommandListImmediate& RHICmdList, FGPUSkinCacheEntry* Entry, int32 SectionIndex, FSkinCacheRWBuffer*& StagingBuffer, bool bTrianglePass);
+	FRWBuffersAllocation* TryAllocBuffer(uint32 NumVertices, bool WithTangnents, bool UseIntermediateTangents, uint32 NumTriangles, FRHICommandList& RHICmdList, const FName& OwnerName);
+	void DoDispatch(FRHICommandList& RHICmdList);
+	void DoDispatch(FRHICommandList& RHICmdList, FGPUSkinCacheEntry* SkinCacheEntry, int32 Section, int32 RevisionNumber);
+	void DispatchUpdateSkinTangents(FRHICommandList& RHICmdList, FGPUSkinCacheEntry* Entry, int32 SectionIndex, FSkinCacheRWBuffer*& StagingBuffer, bool bTrianglePass);
 
 	void PrepareUpdateSkinning(
 		FGPUSkinCacheEntry* Entry, 
@@ -447,7 +447,7 @@ protected:
 		);
 
 	void DispatchUpdateSkinning(
-		FRHICommandListImmediate& RHICmdList, 
+		FRHICommandList& RHICmdList, 
 		FGPUSkinCacheEntry* Entry, 
 		int32 Section, 
 		uint32 RevisionNumber,
@@ -474,7 +474,7 @@ protected:
 	static void CVarSinkFunction();
 	static FAutoConsoleVariableSink CVarSink;
 
-	void IncrementDispatchCounter(FRHICommandListImmediate& RHICmdList);
+	void IncrementDispatchCounter(FRHICommandList& RHICmdList);
 	int32 DispatchCounter = 0;
 
 	void PrintMemorySummary() const;
diff --git a/Engine/Source/Runtime/Engine/Public/GPUSkinVertexFactory.h b/Engine/Source/Runtime/Engine/Public/GPUSkinVertexFactory.h
--- a/Engine/Source/Runtime/Engine/Public/GPUSkinVertexFactory.h
+++ b/Engine/Source/Runtime/Engine/Public/GPUSkinVertexFactory.h
@@ -211,14 +211,12 @@ public:
 		}
 
 		// @param FrameTime from GFrameTime
-		bool UpdateBoneData(FRHICommandListImmediate& RHICmdList, const TArray<FMatrix44f>& ReferenceToLocalMatrices,
+		bool UpdateBoneData(FRHICommandList& RHICmdList, const TArray<FMatrix44f>& ReferenceToLocalMatrices,
 			const TArray<FBoneIndexType>& BoneMap, uint32 RevisionNumber, bool bPrevious, ERHIFeatureLevel::Type FeatureLevel, 
 			bool bUseSkinCache, bool bForceUpdateImmediately, const FName& AssetPathName);
 
 		void ReleaseBoneData()
 		{
-			ensure(IsInRenderingThread());
-
 			UniformBuffer.SafeRelease();
 
 			for(uint32 i = 0; i < 2; ++i)
@@ -359,7 +357,10 @@ public:
 	* update the resource with new data from the game thread.
 	* @param	InData - new stream component data
 	*/
-	virtual void SetData(const FGPUSkinDataType* InData);
+	UE_DEPRECATED(5.3, "Use SetData with a command list.")
+	void SetData(const FGPUSkinDataType* InData);
+
+	virtual void SetData(FRHICommandListBase& RHICmdList, const FGPUSkinDataType* InData);
 
 	uint32 GetNumVertices() const
 	{
@@ -501,7 +502,7 @@ public:
 			Reset();
 		}
 
-		bool UpdateClothSimulData(FRHICommandListImmediate& RHICmdList, const TArray<FVector3f>& InSimulPositions, const TArray<FVector3f>& InSimulNormals, uint32 RevisionNumber, 
+		bool UpdateClothSimulData(FRHICommandList& RHICmdList, const TArray<FVector3f>& InSimulPositions, const TArray<FVector3f>& InSimulNormals, uint32 RevisionNumber, 
 									ERHIFeatureLevel::Type FeatureLevel, bool bForceUpdateImmediately, const FName& AssetPathName);
 
 		void ReleaseClothSimulData()
@@ -693,7 +694,7 @@ public:
 	* update the resource with new data from the game thread.
 	* @param	InData - new stream component data
 	*/
-	virtual void SetData(const FGPUSkinDataType* InData) override;
+	virtual void SetData(FRHICommandListBase& RHICmdList, const FGPUSkinDataType* InData) override;
 
 	virtual FGPUBaseSkinVertexFactory* GetVertexFactory() override
 	{
@@ -799,6 +800,9 @@ public:
 	 * The SRVs are cached per attribute. If any passed in SRV is changed from the cached value then we recreate the vertex factory uniform buffer here.
 	 * Note that on platforms that support manual vertex fetch, only Position will be in the final vertex stream and other attributes will be read through an SRV.
 	 */
+	void SetVertexAttributes(FRHICommandListBase& RHICmdList, FGPUBaseSkinVertexFactory const* InSourceVertexFactory, FAddVertexAttributeDesc const& InDesc);
+
+	UE_DEPRECATED(5.4, "SetVertexAttributes requires a command list.")
 	void SetVertexAttributes(FGPUBaseSkinVertexFactory const* InSourceVertexFactory, FAddVertexAttributeDesc const& InDesc);
 
 	/** 
diff --git a/Engine/Source/Runtime/Engine/Public/SkeletalMeshDeformerHelpers.h b/Engine/Source/Runtime/Engine/Public/SkeletalMeshDeformerHelpers.h
--- a/Engine/Source/Runtime/Engine/Public/SkeletalMeshDeformerHelpers.h
+++ b/Engine/Source/Runtime/Engine/Public/SkeletalMeshDeformerHelpers.h
@@ -77,6 +77,9 @@ public:
 	 * Update all of the MeshObject's passthrough vertex factories with the currently allocated vertex buffers. 
 	 * Usually call this after all AllocateVertexFactory*() functions for a frame.
 	 */
+	ENGINE_API static void UpdateVertexFactoryBufferOverrides(FRHICommandListBase& RHICmdList, FSkeletalMeshObject* InMeshObject, int32 InLodIndex);
+
+	UE_DEPRECATED(5.4, "UpdateVertexFactoryBufferOverrides requires a command list.")
 	ENGINE_API static void UpdateVertexFactoryBufferOverrides(FSkeletalMeshObject* InMeshObject, int32 InLodIndex);
 
 	/** 
