Removing URig from animation system.

diff --git a/Engine/Plugins/Editor/EngineAssetDefinitions/Source/Private/AssetDefinition_Rig.h b/Engine/Plugins/Editor/EngineAssetDefinitions/Source/Private/AssetDefinition_Rig.h
deleted file mode 100644
diff --git a/Engine/Plugins/Experimental/Dataflow/Source/DataflowNodes/Private/Dataflow/DataflowSkeletalMeshNodes.cpp b/Engine/Plugins/Experimental/Dataflow/Source/DataflowNodes/Private/Dataflow/DataflowSkeletalMeshNodes.cpp
--- a/Engine/Plugins/Experimental/Dataflow/Source/DataflowNodes/Private/Dataflow/DataflowSkeletalMeshNodes.cpp
+++ b/Engine/Plugins/Experimental/Dataflow/Source/DataflowNodes/Private/Dataflow/DataflowSkeletalMeshNodes.cpp
@@ -7,7 +7,6 @@
 #include "Logging/LogMacros.h"
 #include "UObject/UnrealTypePrivate.h"
 #include "Animation/Skeleton.h"
-#include "Animation/Rig.h"
 #include "Engine/SkeletalMesh.h"
 
 #include UE_INLINE_GENERATED_CPP_BY_NAME(DataflowSkeletalMeshNodes)
diff --git a/Engine/Plugins/Runtime/RigVM/Source/RigVMDeveloper/Private/RigVMCompiler/RigVMCodeGenerator.cpp b/Engine/Plugins/Runtime/RigVM/Source/RigVMDeveloper/Private/RigVMCompiler/RigVMCodeGenerator.cpp
--- a/Engine/Plugins/Runtime/RigVM/Source/RigVMDeveloper/Private/RigVMCompiler/RigVMCodeGenerator.cpp
+++ b/Engine/Plugins/Runtime/RigVM/Source/RigVMDeveloper/Private/RigVMCompiler/RigVMCodeGenerator.cpp
@@ -5,7 +5,6 @@
 #include "RigVMCore/RigVMStruct.h"
 #include "RigVMDeveloperModule.h"
 #include "Algo/Count.h"
-#include "Animation/Rig.h"
 #include "RigVMModel/Nodes/RigVMDispatchNode.h"
 #include "RigVMStringUtils.h"
 
diff --git a/Engine/Plugins/Runtime/RigVM/Source/RigVMDeveloper/Private/RigVMModel/Nodes/RigVMUnitNode.cpp b/Engine/Plugins/Runtime/RigVM/Source/RigVMDeveloper/Private/RigVMModel/Nodes/RigVMUnitNode.cpp
--- a/Engine/Plugins/Runtime/RigVM/Source/RigVMDeveloper/Private/RigVMModel/Nodes/RigVMUnitNode.cpp
+++ b/Engine/Plugins/Runtime/RigVM/Source/RigVMDeveloper/Private/RigVMModel/Nodes/RigVMUnitNode.cpp
@@ -2,7 +2,6 @@
 
 #include "RigVMModel/Nodes/RigVMUnitNode.h"
 
-#include "Animation/Rig.h"
 #include "RigVMCore/RigVMStruct.h"
 #include "RigVMModel/RigVMControllerActions.h"
 #include "RigVMUserWorkflowRegistry.h"
diff --git a/Engine/Source/Editor/DetailCustomizations/Private/DetailCustomizations.cpp b/Engine/Source/Editor/DetailCustomizations/Private/DetailCustomizations.cpp
--- a/Engine/Source/Editor/DetailCustomizations/Private/DetailCustomizations.cpp
+++ b/Engine/Source/Editor/DetailCustomizations/Private/DetailCustomizations.cpp
@@ -126,7 +126,6 @@
 #include "RawDistributionVectorStructCustomization.h"
 #include "ReflectionCaptureDetails.h"
 #include "RenderPassesCustomization.h"
-#include "RigDetails.h"
 #include "RotatorStructCustomization.h"
 #include "SceneCaptureDetails.h"
 #include "SceneComponentDetails.h"
@@ -400,7 +399,6 @@ void FDetailCustomizationsModule::RegisterObjectCustomizations()
 	RegisterCustomClassLayout("K2Node_GetClassDefaults", FOnGetDetailCustomizationInstance::CreateStatic(&FSkeletalControlNodeDetails::MakeInstance));
 
 	RegisterCustomClassLayout("AnimSequence", FOnGetDetailCustomizationInstance::CreateStatic(&FAnimSequenceDetails::MakeInstance));
-	RegisterCustomClassLayout("Rig", FOnGetDetailCustomizationInstance::CreateStatic(&FRigDetails::MakeInstance));
 
 	RegisterCustomClassLayout("EditorAnimSegment", FOnGetDetailCustomizationInstance::CreateStatic(&FAnimMontageSegmentDetails::MakeInstance));
 	RegisterCustomClassLayout("EditorAnimCompositeSegment", FOnGetDetailCustomizationInstance::CreateStatic(&FAnimMontageSegmentDetails::MakeInstance));
diff --git a/Engine/Source/Editor/DetailCustomizations/Private/RigDetails.cpp b/Engine/Source/Editor/DetailCustomizations/Private/RigDetails.cpp
deleted file mode 100644
diff --git a/Engine/Source/Editor/DetailCustomizations/Private/RigDetails.h b/Engine/Source/Editor/DetailCustomizations/Private/RigDetails.h
deleted file mode 100644
diff --git a/Engine/Source/Editor/Persona/Private/BoneMappingHelper.cpp b/Engine/Source/Editor/Persona/Private/BoneMappingHelper.cpp
deleted file mode 100644
diff --git a/Engine/Source/Editor/Persona/Private/BoneMappingHelper.h b/Engine/Source/Editor/Persona/Private/BoneMappingHelper.h
deleted file mode 100644
diff --git a/Engine/Source/Editor/Persona/Private/SBoneMappingBase.cpp b/Engine/Source/Editor/Persona/Private/SBoneMappingBase.cpp
--- a/Engine/Source/Editor/Persona/Private/SBoneMappingBase.cpp
+++ b/Engine/Source/Editor/Persona/Private/SBoneMappingBase.cpp
@@ -10,12 +10,9 @@
 #include "Styling/AppStyle.h"
 #include "Widgets/Input/SButton.h"
 #include "AssetNotifications.h"
-#include "Animation/Rig.h"
 #include "BoneSelectionWidget.h"
 #include "Widgets/Input/SSearchBox.h"
 #include "Widgets/Text/SInlineEditableTextBlock.h"
-#include "SRigPicker.h"
-#include "BoneMappingHelper.h"
 #include "SSkeletonWidget.h"
 #include "IEditableSkeleton.h"
 
diff --git a/Engine/Source/Editor/Persona/Private/SRigPicker.cpp b/Engine/Source/Editor/Persona/Private/SRigPicker.cpp
deleted file mode 100644
diff --git a/Engine/Source/Editor/Persona/Private/SRigPicker.h b/Engine/Source/Editor/Persona/Private/SRigPicker.h
deleted file mode 100644
diff --git a/Engine/Source/Editor/SkeletalMeshEditor/Private/SkeletalMeshEditor.cpp b/Engine/Source/Editor/SkeletalMeshEditor/Private/SkeletalMeshEditor.cpp
--- a/Engine/Source/Editor/SkeletalMeshEditor/Private/SkeletalMeshEditor.cpp
+++ b/Engine/Source/Editor/SkeletalMeshEditor/Private/SkeletalMeshEditor.cpp
@@ -77,7 +77,6 @@ const FName SkeletalMeshEditorTabs::AssetDetailsTab(TEXT("AnimAssetPropertiesTab
 const FName SkeletalMeshEditorTabs::ViewportTab(TEXT("Viewport"));
 const FName SkeletalMeshEditorTabs::AdvancedPreviewTab(TEXT("AdvancedPreviewTab"));
 const FName SkeletalMeshEditorTabs::MorphTargetsTab("MorphTargetsTab");
-const FName SkeletalMeshEditorTabs::AnimationMappingTab("AnimationMappingWindow");
 const FName SkeletalMeshEditorTabs::ToolboxDetailsTab("ToolBoxDetailsTab");
 const FName SkeletalMeshEditorTabs::CurveMetadataTab(TEXT("AnimCurveMetadataEditorTab"));
 const FName SkeletalMeshEditorTabs::FindReplaceTab("FindReplaceTab");
diff --git a/Engine/Source/Editor/SkeletonEditor/Private/EditableSkeleton.cpp b/Engine/Source/Editor/SkeletonEditor/Private/EditableSkeleton.cpp
--- a/Engine/Source/Editor/SkeletonEditor/Private/EditableSkeleton.cpp
+++ b/Engine/Source/Editor/SkeletonEditor/Private/EditableSkeleton.cpp
@@ -1405,36 +1405,6 @@ void FEditableSkeleton::RemoveCompatibleSkeleton(const USkeleton* InCompatibleSk
 	PersonaModule.BroadcastAssetFamilyChange();
 }
 
-void FEditableSkeleton::RefreshRigConfig()
-{
-	Skeleton->RefreshRigConfig();
-}
-
-void FEditableSkeleton::SetRigConfig(URig* InRig)
-{
-	const FScopedTransaction Transaction(LOCTEXT("RigAssetChanged", "Select Rig"));
-	Skeleton->Modify();
-	Skeleton->SetRigConfig(InRig);
-}
-
-void FEditableSkeleton::SetRigBoneMapping(const FName& InNodeName, const FName& InBoneName)
-{
-	const FScopedTransaction Transaction(LOCTEXT("BoneMappingChanged", "Change Bone Mapping"));
-	Skeleton->Modify();
-	Skeleton->SetRigBoneMapping(InNodeName, InBoneName);
-}
-
-void FEditableSkeleton::SetRigBoneMappings(const TMap<FName, FName>& InMappings)
-{
-	const FScopedTransaction Transaction(LOCTEXT("BoneMappingsChanged", "Change Bone Mappings"));
-	Skeleton->Modify();
-
-	for (const TPair<FName, FName>& Mapping : InMappings)
-	{
-		Skeleton->SetRigBoneMapping(Mapping.Key, Mapping.Value);
-	}
-}
-
 void FEditableSkeleton::RemoveUnusedBones()
 {
 	TArray<FName> SkeletonBones;
diff --git a/Engine/Source/Editor/SkeletonEditor/Private/EditableSkeleton.h b/Engine/Source/Editor/SkeletonEditor/Private/EditableSkeleton.h
--- a/Engine/Source/Editor/SkeletonEditor/Private/EditableSkeleton.h
+++ b/Engine/Source/Editor/SkeletonEditor/Private/EditableSkeleton.h
@@ -60,10 +60,6 @@ public:
 	virtual void RefreshRetargetSources(const TArray<FName>& InRetargetSourceNames) override;
 	virtual void AddCompatibleSkeleton(const USkeleton* InCompatibleSkeleton) override;
 	virtual void RemoveCompatibleSkeleton(const USkeleton* InCompatibleSkeleton) override;
-	virtual void RefreshRigConfig() override;
-	virtual void SetRigConfig(URig* InRig) override;
-	virtual void SetRigBoneMapping(const FName& InNodeName, const FName& InBoneName) override;
-	virtual void SetRigBoneMappings(const TMap<FName, FName>& InMappings) override;
 	virtual void RemoveUnusedBones() override;
 	virtual void UpdateSkeletonReferencePose(USkeletalMesh* InSkeletalMesh) override;
 	virtual void RegisterSlotNode(const FName& InSlotName) override;
diff --git a/Engine/Source/Editor/SkeletonEditor/Public/IEditableSkeleton.h b/Engine/Source/Editor/SkeletonEditor/Public/IEditableSkeleton.h
--- a/Engine/Source/Editor/SkeletonEditor/Public/IEditableSkeleton.h
+++ b/Engine/Source/Editor/SkeletonEditor/Public/IEditableSkeleton.h
@@ -166,20 +166,6 @@ public:
 	/** Remove a compatible skeleton */
 	virtual void RemoveCompatibleSkeleton(const USkeleton* InCompatibleSkeleton) = 0;
 
-	PRAGMA_DISABLE_DEPRECATION_WARNINGS
-	/** Refreshes the rig config, validating the mappings */
-	virtual void RefreshRigConfig() = 0;
-
-	/** Set the rig config  */
-	virtual void SetRigConfig(class URig* InRig) = 0;
-
-	/** Set a rig bone mapping */
-	virtual void SetRigBoneMapping(const FName& InNodeName, const FName& InBoneName) = 0;
-
-	/** Set multiple bone mappings */
-	virtual void SetRigBoneMappings(const TMap<FName, FName>& InMappings) = 0;
-	PRAGMA_ENABLE_DEPRECATION_WARNINGS
-
 	/** Remove any bones that are not used by any skeletal meshes */
 	virtual void RemoveUnusedBones() = 0;
 
diff --git a/Engine/Source/Editor/UnrealEd/Private/SSkeletonWidget.cpp b/Engine/Source/Editor/UnrealEd/Private/SSkeletonWidget.cpp
--- a/Engine/Source/Editor/UnrealEd/Private/SSkeletonWidget.cpp
+++ b/Engine/Source/Editor/UnrealEd/Private/SSkeletonWidget.cpp
@@ -17,7 +17,6 @@
 #include "Interfaces/IMainFrameModule.h"
 #include "ContentBrowserModule.h"
 #include "IDocumentation.h"
-#include "Animation/Rig.h"
 #include "AnimPreviewInstance.h"
 #include "AssetRegistry/AssetRegistryModule.h"
 #include "AnimationRuntime.h"
diff --git a/Engine/Source/Runtime/Engine/Classes/Animation/AnimSequence.h b/Engine/Source/Runtime/Engine/Classes/Animation/AnimSequence.h
--- a/Engine/Source/Runtime/Engine/Classes/Animation/AnimSequence.h
+++ b/Engine/Source/Runtime/Engine/Classes/Animation/AnimSequence.h
@@ -778,9 +778,7 @@ private:
 	ENGINE_API virtual void RemapTracksToNewSkeleton( USkeleton* NewSkeleton, bool bConvertSpaces ) override;
 
 	/** Retargeting functions */
-	ENGINE_API bool ConvertAnimationDataToRiggingData(FAnimSequenceTrackContainer & RiggingAnimationData);
-	ENGINE_API bool ConvertRiggingDataToAnimationData(FAnimSequenceTrackContainer & RiggingAnimationData);
-	ENGINE_API int32 GetSpaceBasedAnimationData(TArray< TArray<FTransform> > & AnimationDataInComponentSpace, FAnimSequenceTrackContainer * RiggingAnimationData) const;
+	ENGINE_API int32 GetSpaceBasedAnimationData(TArray< TArray<FTransform> > & AnimationDataInComponentSpace) const;
 #endif
 
 public:
diff --git a/Engine/Source/Runtime/Engine/Classes/Animation/Rig.h b/Engine/Source/Runtime/Engine/Classes/Animation/Rig.h
deleted file mode 100644
diff --git a/Engine/Source/Runtime/Engine/Classes/Animation/Skeleton.h b/Engine/Source/Runtime/Engine/Classes/Animation/Skeleton.h
--- a/Engine/Source/Runtime/Engine/Classes/Animation/Skeleton.h
+++ b/Engine/Source/Runtime/Engine/Classes/Animation/Skeleton.h
@@ -216,21 +216,6 @@ struct FNameMapping
 	}
 };
 
-USTRUCT()
-struct FRigConfiguration
-{
-	GENERATED_USTRUCT_BODY()
-
-	PRAGMA_DISABLE_DEPRECATION_WARNINGS
-	UPROPERTY()
-	TObjectPtr<class URig>  Rig = nullptr;
-	PRAGMA_ENABLE_DEPRECATION_WARNINGS
-
-	// @todo in the future we can make this to be run-time data
-	UPROPERTY()
-	TArray<FNameMapping> BoneMappingTable;
-};
-
 USTRUCT()
 struct FAnimSlotGroup
 {
@@ -594,11 +579,6 @@ private:
 	UPROPERTY(duplicatetransient, AssetRegistrySearchable)
 	TSoftObjectPtr<class UDataAsset> AdditionalPreviewSkeletalMeshes;
 
-	PRAGMA_DISABLE_DEPRECATION_WARNINGS
-	UPROPERTY()
-	FRigConfiguration RigConfig;
-	PRAGMA_ENABLE_DEPRECATION_WARNINGS
-
 	/** rig property will be saved separately */
 	ENGINE_API virtual void GetAssetRegistryTags(TArray<FAssetRegistryTag>& OutTags) const override;
 
@@ -1022,21 +1002,6 @@ public:
 	ENGINE_API static const FName CompatibleSkeletonsNameTag;
 	ENGINE_API static const FString CompatibleSkeletonsTagDelimiter;
 
-	PRAGMA_DISABLE_DEPRECATION_WARNINGS
-	// rig Configs
-	ENGINE_API static const FName RigTag;
-	ENGINE_API void SetRigConfig(URig * Rig);
-	ENGINE_API FName GetRigBoneMapping(const FName& NodeName) const;
-	ENGINE_API bool SetRigBoneMapping(const FName& NodeName, FName BoneName);
-	ENGINE_API FName GetRigNodeNameFromBoneName(const FName& BoneName) const;
-	// this make sure it stays within the valid range
-	ENGINE_API int32 GetMappedValidNodes(TArray<FName> &OutValidNodeNames);
-	// verify if it has all latest data
-	ENGINE_API void RefreshRigConfig();
-	int32 FindRigBoneMapping(const FName& NodeName) const;
-	ENGINE_API URig * GetRig() const;
-	PRAGMA_ENABLE_DEPRECATION_WARNINGS
-
 #endif
 
 public:
diff --git a/Engine/Source/Runtime/Engine/Private/Animation/AnimSequence.cpp b/Engine/Source/Runtime/Engine/Private/Animation/AnimSequence.cpp
--- a/Engine/Source/Runtime/Engine/Private/Animation/AnimSequence.cpp
+++ b/Engine/Source/Runtime/Engine/Private/Animation/AnimSequence.cpp
@@ -25,7 +25,6 @@
 #include "Components/SkeletalMeshComponent.h"
 #include "Animation/AnimRootMotionProvider.h"
 #include "Animation/BlendSpace.h"
-#include "Animation/Rig.h"
 #include "Animation/AnimationSettings.h"
 #include "Animation/AnimBoneCompressionSettings.h"
 #include "Animation/AnimCurveCompressionCodec.h"
@@ -2192,114 +2191,6 @@ int32 FindMeshBoneIndexFromBoneName(USkeleton * Skeleton, const FName &BoneName)
 	return BoneIndex;
 }
 
-PRAGMA_DISABLE_DEPRECATION_WARNINGS
-void FillUpTransformBasedOnRig(USkeleton* Skeleton, TArray<FTransform>& NodeSpaceBases, TArray<FTransform> &Rotations, TArray<FVector>& Translations, TArray<bool>& TranslationParentFlags)
-{
-	TArray<FTransform> SpaceBases;
-	FAnimationRuntime::FillUpComponentSpaceTransformsRetargetBasePose(Skeleton, SpaceBases);
-
-	const URig* Rig = Skeleton->GetRig();
-
-	if (Rig)
-	{
-		// this one has to collect all Nodes in Rig data
-		// since we're comparing two of them together. 
-		int32 NodeNum = Rig->GetNodeNum();
-
-		if (NodeNum > 0)
-		{
-			NodeSpaceBases.Empty(NodeNum);
-			NodeSpaceBases.AddUninitialized(NodeNum);
-
-			Rotations.Empty(NodeNum);
-			Rotations.AddUninitialized(NodeNum);
-
-			Translations.Empty(NodeNum);
-			Translations.AddUninitialized(NodeNum);
-
-			TranslationParentFlags.Empty(Translations.Num());
-			TranslationParentFlags.AddZeroed(Translations.Num());
-
-			const USkeletalMesh* PreviewMesh = Skeleton->GetPreviewMesh();
-
-			for (int32 Index = 0; Index < NodeNum; ++Index)
-			{
-				const FName NodeName = Rig->GetNodeName(Index);
-				const FName& BoneName = Skeleton->GetRigBoneMapping(NodeName);
-				const int32& BoneIndex = FindMeshBoneIndexFromBoneName(Skeleton, BoneName);
-
-				if (BoneIndex == INDEX_NONE)
-				{
-					// add identity
-					NodeSpaceBases[Index].SetIdentity();
-					Rotations[Index].SetIdentity();
-					Translations[Index] = FVector::ZeroVector;
-				}
-				else
-				{
-					// initialize with SpaceBases - assuming World Based
-					NodeSpaceBases[Index] = SpaceBases[BoneIndex];
-					Rotations[Index] = SpaceBases[BoneIndex];
-					Translations[Index] = SpaceBases[BoneIndex].GetLocation();
-
-					const FTransformBase* TransformBase = Rig->GetTransformBaseByNodeName(NodeName);
-
-					if (TransformBase != nullptr)
-					{
-						// orientation constraint			
-						const auto& RotConstraint = TransformBase->Constraints[EControlConstraint::Type::Orientation];
-
-						if (RotConstraint.TransformConstraints.Num() > 0)
-						{
-							const FName& ParentBoneName = Skeleton->GetRigBoneMapping(RotConstraint.TransformConstraints[0].ParentSpace);
-							const int32& ParentBoneIndex = FindMeshBoneIndexFromBoneName(Skeleton, ParentBoneName);
-
-							if (ParentBoneIndex != INDEX_NONE)
-							{
-								Rotations[Index] = SpaceBases[BoneIndex].GetRelativeTransform(SpaceBases[ParentBoneIndex]);
-							}
-						}
-
-						// translation constraint
-						const auto& TransConstraint = TransformBase->Constraints[EControlConstraint::Type::Translation];
-
-						if (TransConstraint.TransformConstraints.Num() > 0)
-						{
-							const FName& ParentBoneName = Skeleton->GetRigBoneMapping(TransConstraint.TransformConstraints[0].ParentSpace);
-							const int32& ParentBoneIndex = FindMeshBoneIndexFromBoneName(Skeleton, ParentBoneName);
-
-							if (ParentBoneIndex != INDEX_NONE)
-							{
-								// I think translation has to include rotation, otherwise it won't work
-								Translations[Index] = SpaceBases[BoneIndex].GetLocation() - SpaceBases[ParentBoneIndex].GetLocation();
-								TranslationParentFlags[Index] = true;
-							}
-						}
-					}
-				}
-			}
-		}
-	}
-}
-
-int32 FindValidTransformParentTrack(const URig* Rig, int32 NodeIndex, bool bTranslate, const TArray<FName>& ValidNodeNames)
-{
-	int32 ParentIndex = Rig->FindTransformParentNode(NodeIndex, bTranslate);
-
-	// verify if it exists in ValidNodeNames
-	if (ParentIndex != INDEX_NONE)
-	{
-		FName NodeName = Rig->GetNodeName(ParentIndex);
-
-		return ValidNodeNames.Find(NodeName);
-	}
-
-	return INDEX_NONE;
-}
-
-PRAGMA_ENABLE_DEPRECATION_WARNINGS
-
-
 void UAnimSequence::RemapTracksToNewSkeleton( USkeleton* NewSkeleton, bool bConvertSpaces )
 {
 	// this is not cheap, so make sure it only happens in editor
@@ -2314,220 +2205,9 @@ void UAnimSequence::RemapTracksToNewSkeleton( USkeleton* NewSkeleton, bool bConv
 	if (bConvertSpaces)
 	{
 		USkeleton* OldSkeleton = GetSkeleton();
-
-		PRAGMA_DISABLE_DEPRECATION_WARNINGS
-		// first check if both has same rig, if so, we'll retarget using it
-		if (OldSkeleton && OldSkeleton->GetRig() != nullptr && NewSkeleton->GetRig() == OldSkeleton->GetRig() && OldSkeleton->GetPreviewMesh() && NewSkeleton->GetPreviewMesh())
-		{
-			const URig* Rig = OldSkeleton->GetRig();
-
-			// we'll have to save the relative space bases transform from old ref pose to new refpose
-			TArray<FTransform> RelativeToNewSpaceBases;
-			// save the ratio of translation change
-			TArray<float> OldToNewTranslationRatio;
-			// create relative transform in component space between old skeleton and new skeleton
-			{
-				// first calculate component space ref pose to get the relative transform between
-				// two ref poses. It is very important update ref pose before getting here. 
-				TArray<FTransform> NewRotations, OldRotations, NewSpaceBases, OldSpaceBases;
-				TArray<FVector> NewTranslations, OldTranslations;
-				TArray<bool> NewTranslationParentFlags, OldTranslationParentFlags;
-				// get the spacebases transform
-				FillUpTransformBasedOnRig(NewSkeleton, NewSpaceBases, NewRotations, NewTranslations, NewTranslationParentFlags);
-				FillUpTransformBasedOnRig(OldSkeleton, OldSpaceBases, OldRotations, OldTranslations, OldTranslationParentFlags);
-
-				// now we'd like to get the relative transform from old to new ref pose in component space
-				// PK2*K2 = PK1*K1*theta where theta => P1*R1*theta = P2*R2 
-				// where	P1 - parent transform in component space for original skeleton
-				//			R1 - local space of the current bone for original skeleton
-				//			P2 - parent transform in component space for new skeleton
-				//			R2 - local space of the current bone for new skeleton
-				// what we're looking for is theta, so that we can apply that to animated transform
-				// this has to have all of nodes since comparing two skeletons, that might have different configuration
-				int32 NumNodes = Rig->GetNodeNum();
-				// saves the theta data per node
-				RelativeToNewSpaceBases.AddUninitialized(NumNodes);
-				// saves the translation conversion datao
-				OldToNewTranslationRatio.AddUninitialized(NumNodes);
-
-				const TArray<FNode>& Nodes = Rig->GetNodes();
-				// calculate the relative transform to new skeleton
-				// so that we can apply the delta in component space
-				for (int32 NodeIndex = 0; NodeIndex < NumNodes; ++NodeIndex)
-				{
-					// theta (RelativeToNewTransform) = (P1*R1)^(-1) * P2*R2 where theta => P1*R1*theta = P2*R2
-					RelativeToNewSpaceBases[NodeIndex] = NewSpaceBases[NodeIndex].GetRelativeTransform(OldSpaceBases[NodeIndex]); 
-
-					// also savees the translation difference between old to new
-					FVector OldTranslation = OldTranslations[NodeIndex];
-					FVector NewTranslation = NewTranslations[NodeIndex];
-
-					// skip root because we don't really have clear relative point to test with it
-					if (NodeIndex != 0 && NewTranslationParentFlags[NodeIndex] == OldTranslationParentFlags[NodeIndex])
-					{
-						// only do this if parent status matches, otherwise, you'll have invalid state 
-						// where one is based on shoulder, where the other is missing the shoulder node
-						float OldTranslationSize = OldTranslation.Size();
-						float NewTranslationSize = NewTranslation.Size();
-
-						OldToNewTranslationRatio[NodeIndex] = (FMath::IsNearlyZero(OldTranslationSize)) ? 1.f/*do not touch new translation size*/ : NewTranslationSize / OldTranslationSize;
-					}
-					else
-					{
-						OldToNewTranslationRatio[NodeIndex] = 1.f; // set to be 1, we don't know what it is
-					}
-
-					UE_LOG(LogAnimation, Verbose, TEXT("Retargeting (%s : %d) : OldtoNewTranslationRatio (%0.2f), Relative Transform (%s)"), *Nodes[NodeIndex].Name.ToString(), NodeIndex, 
-						OldToNewTranslationRatio[NodeIndex], *RelativeToNewSpaceBases[NodeIndex].ToString());
-					UE_LOG(LogAnimation, Verbose, TEXT("\tOldSpaceBase(%s), NewSpaceBase(%s)"), *OldSpaceBases[NodeIndex].ToString(), *NewSpaceBases[NodeIndex].ToString());
-				}
-			}
-
-			FAnimSequenceTrackContainer RiggingAnimationData;
-
-			// now convert animation data to rig data
-			ConvertAnimationDataToRiggingData(RiggingAnimationData);
-
-			// here we have to watch out the index
-			// The RiggingAnimationData will contain only the nodes that are mapped to source skeleton
-			// and here we convert everything that is in RiggingAnimationData which means based on source data
-			// when mapped back to new skeleton, it will discard results that are not mapped to target skeleton
-			
-			TArray<FName> SrcValidNodeNames;
-			int32 SrcNumTracks = OldSkeleton->GetMappedValidNodes(SrcValidNodeNames);
-
-			// now convert to space bases animation 
-			TArray< TArray<FTransform> > ComponentSpaceAnimations, ConvertedLocalSpaceAnimations, ConvertedSpaceAnimations;
-			ComponentSpaceAnimations.AddZeroed(SrcNumTracks);
-			ConvertedSpaceAnimations.AddZeroed(SrcNumTracks);
-			ConvertedLocalSpaceAnimations.AddZeroed(SrcNumTracks);
-
-			const int32 NumKeys = DataModelInterface->GetNumberOfKeys();
-			float Interval = DataModelInterface->GetFrameRate().AsInterval();
-
-			// allocate arrays
-			for(int32 SrcTrackIndex=0; SrcTrackIndex<SrcNumTracks; ++SrcTrackIndex)
-			{
-				ComponentSpaceAnimations[SrcTrackIndex].AddUninitialized(NumKeys);
-				ConvertedLocalSpaceAnimations[SrcTrackIndex].AddUninitialized(NumKeys);
-				ConvertedSpaceAnimations[SrcTrackIndex].AddUninitialized(NumKeys);
-			}
-
-			for (int32 SrcTrackIndex=0; SrcTrackIndex<SrcNumTracks; ++SrcTrackIndex)		
-			{
-				int32 NodeIndex = Rig->FindNode(SrcValidNodeNames[SrcTrackIndex]);
-				check (NodeIndex != INDEX_NONE);
-				auto& RawAnimation = RiggingAnimationData.AnimationTracks[SrcTrackIndex];
-
-				// find rotation parent node
-				int32 RotParentTrackIndex = FindValidTransformParentTrack(Rig, NodeIndex, false, SrcValidNodeNames);
-				int32 TransParentTrackIndex = FindValidTransformParentTrack(Rig, NodeIndex, true, SrcValidNodeNames);
-				// fill up keys - calculate PK1 * K1
-				for(int32 Key=0; Key<NumKeys; ++Key)
-				{
-					FTransform AnimatedLocalKey;
-					//ExtractBoneTransform(RiggingAnimationData.AnimationTracks[SrcTrackIndex], AnimatedLocalKey, Interval * Key);
-					FAnimationUtils::ExtractTransformFromTrack(RiggingAnimationData.AnimationTracks[SrcTrackIndex], Interval* Key, NumberOfSampledKeys, GetPlayLength(), Interpolation, AnimatedLocalKey);
-
-					AnimatedLocalKey.ScaleTranslation(OldToNewTranslationRatio[NodeIndex]);
-
-					if(RotParentTrackIndex != INDEX_NONE)
-					{
-						FQuat ComponentSpaceRotation = ComponentSpaceAnimations[RotParentTrackIndex][Key].GetRotation() * AnimatedLocalKey.GetRotation();
-						ComponentSpaceAnimations[SrcTrackIndex][Key].SetRotation(ComponentSpaceRotation);
-					}
-					else
-					{
-						ComponentSpaceAnimations[SrcTrackIndex][Key].SetRotation(AnimatedLocalKey.GetRotation());
-					}
-
-					if (TransParentTrackIndex != INDEX_NONE)
-					{
-						FVector ComponentSpaceTranslation = ComponentSpaceAnimations[TransParentTrackIndex][Key].TransformPosition(AnimatedLocalKey.GetTranslation());
-						ComponentSpaceAnimations[SrcTrackIndex][Key].SetTranslation(ComponentSpaceTranslation);
-						FVector ParentComponentSpaceScale3D = ComponentSpaceAnimations[TransParentTrackIndex][Key].GetScale3D();
-						ComponentSpaceAnimations[SrcTrackIndex][Key].SetScale3D(ParentComponentSpaceScale3D * AnimatedLocalKey.GetScale3D());
-					}
-					else
-					{
-						ComponentSpaceAnimations[SrcTrackIndex][Key].SetTranslation(AnimatedLocalKey.GetTranslation());
-						ComponentSpaceAnimations[SrcTrackIndex][Key].SetScale3D(AnimatedLocalKey.GetScale3D());
-					}
-				}
-			}
-
-			// now animation is converted to component space
-			TArray<struct FRawAnimSequenceTrack> NewRawAnimationData = RiggingAnimationData.AnimationTracks;
-			for (int32 SrcTrackIndex=0; SrcTrackIndex<SrcNumTracks; ++SrcTrackIndex)
-			{
-				int32 NodeIndex = Rig->FindNode(SrcValidNodeNames[SrcTrackIndex]);
-				// find rotation parent node
-				int32 RotParentTrackIndex = FindValidTransformParentTrack(Rig, NodeIndex, false, SrcValidNodeNames);
-				int32 TransParentTrackIndex = FindValidTransformParentTrack(Rig, NodeIndex, true, SrcValidNodeNames);
-
-				// clear translation;
-				RelativeToNewSpaceBases[NodeIndex].SetTranslation(FVector::ZeroVector);
-
-				for(int32 Key=0; Key<NumKeys; ++Key)
-				{
-					// now convert to the new space and save to local spaces
-					ConvertedSpaceAnimations[SrcTrackIndex][Key] = RelativeToNewSpaceBases[NodeIndex] * ComponentSpaceAnimations[SrcTrackIndex][Key];
-
-					if(RotParentTrackIndex != INDEX_NONE)
-					{
-						FQuat LocalRotation = ConvertedSpaceAnimations[RotParentTrackIndex][Key].GetRotation().Inverse() * ConvertedSpaceAnimations[SrcTrackIndex][Key].GetRotation();
-						ConvertedLocalSpaceAnimations[SrcTrackIndex][Key].SetRotation(LocalRotation);
-					}
-					else
-					{
-						ConvertedLocalSpaceAnimations[SrcTrackIndex][Key].SetRotation(ConvertedSpaceAnimations[SrcTrackIndex][Key].GetRotation());
-					}
-
-					if(TransParentTrackIndex != INDEX_NONE)
-					{
-						FTransform LocalTransform = ConvertedSpaceAnimations[SrcTrackIndex][Key].GetRelativeTransform(ConvertedSpaceAnimations[TransParentTrackIndex][Key]);
-						ConvertedLocalSpaceAnimations[SrcTrackIndex][Key].SetTranslation(LocalTransform.GetLocation());
-						ConvertedLocalSpaceAnimations[SrcTrackIndex][Key].SetScale3D(LocalTransform.GetScale3D());
-					}
-					else
-					{
-						ConvertedLocalSpaceAnimations[SrcTrackIndex][Key].SetTranslation(ConvertedSpaceAnimations[SrcTrackIndex][Key].GetTranslation());
-						ConvertedLocalSpaceAnimations[SrcTrackIndex][Key].SetScale3D(ConvertedSpaceAnimations[SrcTrackIndex][Key].GetScale3D());
-					}
-				}
-
-				auto& RawAnimation = NewRawAnimationData[SrcTrackIndex];
-				RawAnimation.PosKeys.Empty(NumKeys);
-				RawAnimation.PosKeys.AddUninitialized(NumKeys);
-				RawAnimation.RotKeys.Empty(NumKeys);
-				RawAnimation.RotKeys.AddUninitialized(NumKeys);
-				RawAnimation.ScaleKeys.Empty(NumKeys);
-				RawAnimation.ScaleKeys.AddUninitialized(NumKeys);
-
-				for(int32 Key=0; Key<NumKeys; ++Key)
-				{
-					RawAnimation.PosKeys[Key] = FVector3f(ConvertedLocalSpaceAnimations[SrcTrackIndex][Key].GetLocation());
-					RawAnimation.RotKeys[Key] = FQuat4f(ConvertedLocalSpaceAnimations[SrcTrackIndex][Key].GetRotation());
-					RawAnimation.ScaleKeys[Key] = FVector3f(ConvertedLocalSpaceAnimations[SrcTrackIndex][Key].GetScale3D());
-
-					// normalize rotation
-					RawAnimation.RotKeys[Key].Normalize();
-				}
-			}
-
-			RiggingAnimationData.AnimationTracks = MoveTemp(NewRawAnimationData);
-			RiggingAnimationData.TrackNames = MoveTemp(SrcValidNodeNames);
-
-			// set new skeleton
-			SetSkeleton(NewSkeleton);
-
-			// convert back to animated data with new skeleton
-			ConvertRiggingDataToAnimationData(RiggingAnimationData);
-		}
-		PRAGMA_ENABLE_DEPRECATION_WARNINGS
-		// @todo end rig testing
-		// @IMPORTANT: now otherwise this will try to do bone to bone mapping
-		else if(OldSkeleton && OldSkeleton != NewSkeleton)
+		
+		// this will try to do bone to bone mapping
+		if(OldSkeleton && OldSkeleton != NewSkeleton)
 		{
 			// Validate animation tracks against the new skeleton, any tracks linked to bones that do not exist in the new hierarchy are removed
 			Controller->RemoveBoneTracksMissingFromSkeleton(NewSkeleton);
@@ -2813,17 +2493,7 @@ void UAnimSequence::ReplaceReferredAnimations(const TMap<UAnimationAsset*, UAnim
 	}
 }
 
-PRAGMA_DISABLE_DEPRECATION_WARNINGS
-int32 FindParentNodeIndex(URig* Rig, USkeleton* Skeleton, FName ParentNodeName)
-{
-	const int32& ParentNodeIndex = Rig->FindNode(ParentNodeName);
-	const FName& ParentBoneName = Skeleton->GetRigBoneMapping(ParentNodeName);
-	
-	return Skeleton->GetReferenceSkeleton().FindBoneIndex(ParentBoneName);
-}
-PRAGMA_ENABLE_DEPRECATION_WARNINGS
-
-int32 UAnimSequence::GetSpaceBasedAnimationData(TArray< TArray<FTransform> >& AnimationDataInComponentSpace, FAnimSequenceTrackContainer * RiggingAnimationData) const
+int32 UAnimSequence::GetSpaceBasedAnimationData(TArray< TArray<FTransform> >& AnimationDataInComponentSpace) const
 {
 	USkeleton* MySkeleton = GetSkeleton();
 
@@ -2845,433 +2515,53 @@ int32 UAnimSequence::GetSpaceBasedAnimationData(TArray< TArray<FTransform> >& An
 	{
 		AnimationDataInComponentSpace[BoneIndex].AddUninitialized(NumKeys);
 	}
-
-	PRAGMA_DISABLE_DEPRECATION_WARNINGS
-	if (RiggingAnimationData)
+	
+	// now calculating old animated space bases
+	// this one calculates animated space per bones and per key
+	TArray<FTransform> BoneTransforms;
+	for (int32 BoneIndex = 0; BoneIndex < NumBones; ++BoneIndex)
 	{
-		const URig* Rig = MySkeleton->GetRig();
-
-		check(Rig);
-
-		// to fix the issue where parent of rig doesn't correspond to parent of this skeleton
-		// we do this in multiple iteration if needed. 
-		// this flag will be used to evaluate all of them until done
-		TArray<bool> BoneEvaluated;
-		BoneEvaluated.AddZeroed(NumBones);
-
-		bool bCompleted = false;
-		do
-		{
-			for(int32 BoneIndex = 0; BoneIndex < NumBones; ++BoneIndex)
-			{
-				if ( !BoneEvaluated[BoneIndex] )
-				{
-					const FName& BoneName = RefSkeleton.GetBoneName(BoneIndex);
-					const FName& NodeName = MySkeleton->GetRigNodeNameFromBoneName(BoneName);
-					const FTransformBase* TransformBase = Rig->GetTransformBaseByNodeName(NodeName);
-					const int32 NodeIndex = RiggingAnimationData->TrackNames.Find(NodeName);
-					if(NodeIndex != INDEX_NONE)
-					{
-						check(TransformBase);
-
-						// now calculate the component space
-						const TArray<FRigTransformConstraint>	& RotTransformConstraints = TransformBase->Constraints[EControlConstraint::Type::Orientation].TransformConstraints;
+		const FName BoneName = MySkeleton->GetReferenceSkeleton().GetBoneName(BoneIndex);
+		const int32 ParentBoneIndex = MySkeleton->GetReferenceSkeleton().GetParentIndex(BoneIndex);
 
-						FQuat ComponentRotation;
-						FTransform ComponentTranslation;
-						FVector ComponentScale;
-
-						// rotation first
-						// this is easy since we just make sure it's evaluated or not
-						{
-							const FName& ParentNodeName = RotTransformConstraints[0].ParentSpace;
-							const FName& ParentBoneName = MySkeleton->GetRigBoneMapping(ParentNodeName);
-							const int32& ParentBoneIndex = RefSkeleton.FindBoneIndex(ParentBoneName);
-
-							if(ParentBoneIndex != INDEX_NONE)
-							{
-								if (BoneEvaluated[ParentBoneIndex])
-								{
-									for(int32 Key = 0; Key < NumKeys; ++Key)
-									{
-										ComponentRotation = AnimationDataInComponentSpace[ParentBoneIndex][Key].GetRotation() * FQuat(RiggingAnimationData->AnimationTracks[NodeIndex].RotKeys[Key]);
-										AnimationDataInComponentSpace[BoneIndex][Key].SetRotation(ComponentRotation);
-									}
-
-									BoneEvaluated[BoneIndex] = true;
-								}
-							}
-							else
-							{
-								for(int32 Key = 0; Key < NumKeys; ++Key)
-								{
-									ComponentRotation = FQuat(RiggingAnimationData->AnimationTracks[NodeIndex].RotKeys[Key]);
-									AnimationDataInComponentSpace[BoneIndex][Key].SetRotation(ComponentRotation);
-								}
-
-								BoneEvaluated[BoneIndex] = true;
-							}
-						}
-
-						const TArray<FRigTransformConstraint>	& PosTransformConstraints = TransformBase->Constraints[EControlConstraint::Type::Translation].TransformConstraints;
-
-						// now time to check translation
-						// this is a bit more complicated
-						// since we have to make sure if it's true to start with
-						// did we succeed on getting rotation?
-						if (BoneEvaluated[BoneIndex])
-						{
-							const FName& ParentNodeName = PosTransformConstraints[0].ParentSpace;
-							const FName& ParentBoneName = MySkeleton->GetRigBoneMapping(ParentNodeName);
-							const int32& ParentBoneIndex = RefSkeleton.FindBoneIndex(ParentBoneName);
-
-							if(ParentBoneIndex != INDEX_NONE)
-							{
-								// this has to be check
-								if (BoneEvaluated[ParentBoneIndex])
-								{
-									for(int32 Key = 0; Key < NumKeys; ++Key)
-									{
-										const FTransform& AnimCompSpace = AnimationDataInComponentSpace[ParentBoneIndex][Key];
-										ComponentTranslation = FTransform(FVector(RiggingAnimationData->AnimationTracks[NodeIndex].PosKeys[Key])) * AnimCompSpace;
-										AnimationDataInComponentSpace[BoneIndex][Key].SetTranslation(ComponentTranslation.GetTranslation());
-
-										ComponentScale = AnimCompSpace.GetScale3D() * FVector(RiggingAnimationData->AnimationTracks[NodeIndex].ScaleKeys[Key]);
-										AnimationDataInComponentSpace[BoneIndex][Key].SetScale3D(ComponentScale);
-									}
-								}
-								else
-								{
-									// if we failed to get parent clear the flag
-									// because if translation has been calculated, BoneEvaluated[BoneIndex] might be true
-									BoneEvaluated[BoneIndex] = false;
-								}
-							}
-							else
-							{
-								for(int32 Key = 0; Key < NumKeys; ++Key)
-								{
-									ComponentTranslation = FTransform(FVector(RiggingAnimationData->AnimationTracks[NodeIndex].PosKeys[Key]));
-									AnimationDataInComponentSpace[BoneIndex][Key].SetTranslation(ComponentTranslation.GetTranslation());
-									
-									ComponentScale = FVector(RiggingAnimationData->AnimationTracks[NodeIndex].ScaleKeys[Key]);
-									AnimationDataInComponentSpace[BoneIndex][Key].SetScale3D(ComponentScale);
-								}
-							}
-						}
-					}
-					else
-					{
-						int32 ParentIndex = RefSkeleton.GetParentIndex(BoneIndex);
-						const FTransform& LocalSpace = RefSkeleton.GetRefBonePose()[BoneIndex];
-						if(ParentIndex != INDEX_NONE)
-						{
-							// if parent is evaluated, do it
-							if (BoneEvaluated[ParentIndex])
-							{
-								for(int32 Key = 0; Key < NumKeys; ++Key)
-								{
-									AnimationDataInComponentSpace[BoneIndex][Key] = LocalSpace * AnimationDataInComponentSpace[ParentIndex][Key];
-								}
-
-								BoneEvaluated[BoneIndex] = true;
-							}
-						}
-						else
-						{
-							BoneEvaluated[BoneIndex] = true;
-
-							for(int32 Key = 0; Key < NumKeys; ++Key)
-							{
-								AnimationDataInComponentSpace[BoneIndex][Key] = LocalSpace;
-							}
-						}
-					}
-				}
-			}
-
-			bCompleted = true;
-			// see if we can get out, brute force for now
-			for(int32 BoneIndex = 0; BoneIndex < NumBones && bCompleted; ++BoneIndex)
-			{
-				bCompleted &= !!BoneEvaluated[BoneIndex];
-			}
-		} while (bCompleted == false);
-	}
-	PRAGMA_ENABLE_DEPRECATION_WARNINGS
-	else
-	{
-		// now calculating old animated space bases
-		// this one calculates aniamted space per bones and per key
-		TArray<FTransform> BoneTransforms;
-		for (int32 BoneIndex = 0; BoneIndex < NumBones; ++BoneIndex)
+		if (DataModelInterface->IsValidBoneTrackName(BoneName))
 		{
-			const FName BoneName = MySkeleton->GetReferenceSkeleton().GetBoneName(BoneIndex);
-			const int32 ParentBoneIndex = MySkeleton->GetReferenceSkeleton().GetParentIndex(BoneIndex);
-
-			if (DataModelInterface->IsValidBoneTrackName(BoneName))
+			BoneTransforms.Reset();
+			DataModelInterface->GetBoneTrackTransforms(BoneName, BoneTransforms);
+			// fill up keys - calculate PK1 * K1
+			for (int32 Key = 0; Key < NumKeys; ++Key)
 			{
-				BoneTransforms.Reset();
-				DataModelInterface->GetBoneTrackTransforms(BoneName, BoneTransforms);
-				// fill up keys - calculate PK1 * K1
-				for (int32 Key = 0; Key < NumKeys; ++Key)
+				FTransform AnimatedLocalKey = BoneTransforms[Key];
+				if (ParentBoneIndex != INDEX_NONE)
 				{
-					FTransform AnimatedLocalKey = BoneTransforms[Key];
-					if (ParentBoneIndex != INDEX_NONE)
-					{
-						AnimationDataInComponentSpace[BoneIndex][Key] = AnimatedLocalKey * AnimationDataInComponentSpace[ParentBoneIndex][Key];
-					}
-					else
-					{
-						AnimationDataInComponentSpace[BoneIndex][Key] = AnimatedLocalKey;
-					}
-				}
-			}
-			else
-			{
-				// get local spaces from refpose and use that to fill it up
-				FTransform LocalTransform = MySkeleton->GetReferenceSkeleton().GetRefBonePose()[BoneIndex];
-
-				for (int32 Key = 0; Key < NumKeys; ++Key)
-				{
-					if (ParentBoneIndex != INDEX_NONE)
-					{
-						AnimationDataInComponentSpace[BoneIndex][Key] = LocalTransform * AnimationDataInComponentSpace[ParentBoneIndex][Key];
-					}
-					else
-					{
-						AnimationDataInComponentSpace[BoneIndex][Key] = LocalTransform;
-					}
+					AnimationDataInComponentSpace[BoneIndex][Key] = AnimatedLocalKey * AnimationDataInComponentSpace[ParentBoneIndex][Key];
 				}
-			}	
-		}
-
-	}
-
-	return AnimationDataInComponentSpace.Num();
-}
-
-PRAGMA_DISABLE_DEPRECATION_WARNINGS
-bool UAnimSequence::ConvertAnimationDataToRiggingData(FAnimSequenceTrackContainer& RiggingAnimationData)
-{
-	USkeleton* MySkeleton = GetSkeleton();
-	if (MySkeleton && MySkeleton->GetRig())
-	{
-		const URig* Rig = MySkeleton->GetRig();
-		TArray<FName> ValidNodeNames;
-		int32 NumNodes = MySkeleton->GetMappedValidNodes(ValidNodeNames);
-		TArray< TArray<FTransform> > AnimationDataInComponentSpace;
-		int32 NumBones = GetSpaceBasedAnimationData(AnimationDataInComponentSpace, nullptr);
-
-		ValidateModel();
-		const int32 NumModelKeys = DataModelInterface->GetNumberOfKeys();
-
-		if (NumBones > 0)
-		{
-			RiggingAnimationData.Initialize(ValidNodeNames);
-
-			// first we copy all space bases back to it
-			for (int32 NodeIndex = 0; NodeIndex < NumNodes; ++NodeIndex)
-			{
-				struct FRawAnimSequenceTrack& Track = RiggingAnimationData.AnimationTracks[NodeIndex];
-				const FName& NodeName = ValidNodeNames[NodeIndex];
-				const FName& BoneName = MySkeleton->GetRigBoneMapping(NodeName);
-				const int32& BoneIndex = MySkeleton->GetReferenceSkeleton().FindBoneIndex(BoneName);
-
-				if (ensure(BoneIndex != INDEX_NONE))
+				else
 				{
-					Track.PosKeys.Empty(NumModelKeys);
-					Track.RotKeys.Empty(NumModelKeys);
-					Track.ScaleKeys.Empty(NumModelKeys);
-					Track.PosKeys.AddUninitialized(NumModelKeys);
-					Track.RotKeys.AddUninitialized(NumModelKeys);
-					Track.ScaleKeys.AddUninitialized(NumModelKeys);
-
-					int32 RigConstraintIndex = Rig->FindTransformBaseByNodeName(NodeName);
-
-					if (RigConstraintIndex != INDEX_NONE)
-					{
-						const auto* RigConstraint = Rig->GetTransformBase(RigConstraintIndex);
-
-						// apply orientation - for now only one
-						const TArray<FRigTransformConstraint>& RotationTransformConstraint = RigConstraint->Constraints[EControlConstraint::Type::Orientation].TransformConstraints;
-
-						if (RotationTransformConstraint.Num() > 0)
-						{
-							const FName& ParentSpace = RotationTransformConstraint[0].ParentSpace;
-							const FName& ParentBoneName = MySkeleton->GetRigBoneMapping(ParentSpace);
-							const int32& ParentBoneIndex = MySkeleton->GetReferenceSkeleton().FindBoneIndex(ParentBoneName);
-							if (ParentBoneIndex != INDEX_NONE)
-							{
-								// if no rig control, component space is used
-								for (int32 KeyIndex = 0; KeyIndex < NumModelKeys; ++KeyIndex)
-								{
-									FTransform ParentTransform = AnimationDataInComponentSpace[ParentBoneIndex][KeyIndex];
-									FTransform RelativeTransform = AnimationDataInComponentSpace[BoneIndex][KeyIndex].GetRelativeTransform(ParentTransform);
-									Track.RotKeys[KeyIndex] = FQuat4f(RelativeTransform.GetRotation());
-								}
-							}
-							else
-							{
-								// if no rig control, component space is used
-								for (int32 KeyIndex = 0; KeyIndex < NumModelKeys; ++KeyIndex)
-								{
-									Track.RotKeys[KeyIndex] = FQuat4f(AnimationDataInComponentSpace[BoneIndex][KeyIndex].GetRotation());
-								}
-							}
-						}
-						else
-						{
-							// if no rig control, component space is used
-							for (int32 KeyIndex = 0; KeyIndex < NumModelKeys; ++KeyIndex)
-							{
-								Track.RotKeys[KeyIndex] = FQuat4f(AnimationDataInComponentSpace[BoneIndex][KeyIndex].GetRotation());
-							}
-						}
-
-						// apply translation - for now only one
-						const TArray<FRigTransformConstraint>& TranslationTransformConstraint = RigConstraint->Constraints[EControlConstraint::Type::Translation].TransformConstraints;
-
-						if (TranslationTransformConstraint.Num() > 0)
-						{
-							const FName& ParentSpace = TranslationTransformConstraint[0].ParentSpace;
-							const FName& ParentBoneName = MySkeleton->GetRigBoneMapping(ParentSpace);
-							const int32& ParentBoneIndex = MySkeleton->GetReferenceSkeleton().FindBoneIndex(ParentBoneName);
-							if (ParentBoneIndex != INDEX_NONE)
-							{
-								// if no rig control, component space is used
-								for (int32 KeyIndex = 0; KeyIndex < NumModelKeys; ++KeyIndex)
-								{
-									FTransform ParentTransform = AnimationDataInComponentSpace[ParentBoneIndex][KeyIndex];
-									FTransform RelativeTransform = AnimationDataInComponentSpace[BoneIndex][KeyIndex].GetRelativeTransform(ParentTransform);
-									Track.PosKeys[KeyIndex] = FVector3f(RelativeTransform.GetTranslation());
-									Track.ScaleKeys[KeyIndex] = FVector3f(RelativeTransform.GetScale3D());
-								}
-							}
-							else
-							{
-								for (int32 KeyIndex = 0; KeyIndex < NumModelKeys; ++KeyIndex)
-								{
-									Track.PosKeys[KeyIndex] = FVector3f(AnimationDataInComponentSpace[BoneIndex][KeyIndex].GetTranslation());
-									Track.ScaleKeys[KeyIndex] = FVector3f(AnimationDataInComponentSpace[BoneIndex][KeyIndex].GetScale3D());
-								}
-							}
-						}
-						else
-						{
-							for (int32 KeyIndex = 0; KeyIndex < NumModelKeys; ++KeyIndex)
-							{
-								Track.PosKeys[KeyIndex] = FVector3f(AnimationDataInComponentSpace[BoneIndex][KeyIndex].GetTranslation());
-								Track.ScaleKeys[KeyIndex] = FVector3f(AnimationDataInComponentSpace[BoneIndex][KeyIndex].GetScale3D());
-							}
-						}
-					}
-					else
-					{
-						// if no rig control, component space is used
-						for (int32 KeyIndex = 0; KeyIndex < NumModelKeys; ++KeyIndex)
-						{
-							Track.PosKeys[KeyIndex] = FVector3f(AnimationDataInComponentSpace[BoneIndex][KeyIndex].GetTranslation());
-							Track.RotKeys[KeyIndex] = FQuat4f(AnimationDataInComponentSpace[BoneIndex][KeyIndex].GetRotation());
-							Track.ScaleKeys[KeyIndex] = FVector3f(AnimationDataInComponentSpace[BoneIndex][KeyIndex].GetScale3D());
-						}
-					}
+					AnimationDataInComponentSpace[BoneIndex][Key] = AnimatedLocalKey;
 				}
 			}
 		}
-
-		return true;
-	}
-
-	return false;
-}
-PRAGMA_ENABLE_DEPRECATION_WARNINGS
-
-bool UAnimSequence::ConvertRiggingDataToAnimationData(FAnimSequenceTrackContainer& RiggingAnimationData)
-{
-	if (RiggingAnimationData.GetNum() > 0)
-	{
-		ValidateModel();
-
-		IAnimationDataController::FScopedBracket ScopedBracket(Controller, LOCTEXT("ConvertRiggingDataToAnimationData_Bracket", "Converting Rigging Data to Animation Data"));
-
-		TArray< TArray<FTransform> > AnimationDataInComponentSpace;
-		int32 NumBones = GetSpaceBasedAnimationData(AnimationDataInComponentSpace, &RiggingAnimationData);
-
-		USkeleton* MySkeleton = GetSkeleton();
-		TArray<FName> ValidNodeNames;
-		MySkeleton->GetMappedValidNodes(ValidNodeNames);
-		// remove from ValidNodeNames if it doesn't belong to AnimationTrackNames
-		for (int32 NameIndex=0; NameIndex<ValidNodeNames.Num(); ++NameIndex)
-		{
-			if (RiggingAnimationData.TrackNames.Contains(ValidNodeNames[NameIndex]) == false)
-			{
-				ValidNodeNames.RemoveAt(NameIndex);
-				--NameIndex;
-			}
-		}
-
-		int32 ValidNumNodes = ValidNodeNames.Num();
-
-		// get local spaces
-		// add all tracks?
-		Controller->RemoveAllBoneTracks();
-
-		// Reused track data arrays
-		TArray<FVector3f> PosKeys;
-		TArray<FQuat4f> RotKeys;
-		TArray<FVector3f> ScaleKeys;
-
-		const int32 NumModelKeys = DataModelInterface->GetNumberOfKeys();
-		const FReferenceSkeleton& RefSkeleton = MySkeleton->GetReferenceSkeleton();
-		for (int32 NodeIndex = 0; NodeIndex < ValidNumNodes; ++NodeIndex)
+		else
 		{
-			FName BoneName = MySkeleton->GetRigBoneMapping(ValidNodeNames[NodeIndex]);
-			int32 BoneIndex = RefSkeleton.FindBoneIndex(BoneName);
+			// get local spaces from refpose and use that to fill it up
+			FTransform LocalTransform = MySkeleton->GetReferenceSkeleton().GetRefBonePose()[BoneIndex];
 
-			if (BoneIndex != INDEX_NONE)
+			for (int32 Key = 0; Key < NumKeys; ++Key)
 			{
-				// Add a new track by name
-				Controller->AddBoneCurve(BoneName);
-
-				PosKeys.SetNumUninitialized(NumModelKeys);
-				RotKeys.SetNumUninitialized(NumModelKeys);
-				ScaleKeys.SetNumUninitialized(NumModelKeys);
-
-				const int32& ParentBoneIndex = RefSkeleton.GetParentIndex(BoneIndex);
-
-				if(ParentBoneIndex != INDEX_NONE)
+				if (ParentBoneIndex != INDEX_NONE)
 				{
-					for(int32 KeyIndex = 0; KeyIndex < NumModelKeys; ++KeyIndex)
-					{
-						const FTransform LocalTransform = AnimationDataInComponentSpace[BoneIndex][KeyIndex].GetRelativeTransform(AnimationDataInComponentSpace[ParentBoneIndex][KeyIndex]);
-
-						PosKeys[KeyIndex] = FVector3f(LocalTransform.GetTranslation());
-						RotKeys[KeyIndex] = FQuat4f(LocalTransform.GetRotation());
-						ScaleKeys[KeyIndex] = FVector3f(LocalTransform.GetScale3D());
-					}
+					AnimationDataInComponentSpace[BoneIndex][Key] = LocalTransform * AnimationDataInComponentSpace[ParentBoneIndex][Key];
 				}
 				else
 				{
-					for(int32 KeyIndex = 0; KeyIndex < NumModelKeys; ++KeyIndex)
-					{
-						const FTransform LocalTransform = AnimationDataInComponentSpace[BoneIndex][KeyIndex];
-
-						PosKeys[KeyIndex] = FVector3f(LocalTransform.GetTranslation());
-						RotKeys[KeyIndex] = FQuat4f(LocalTransform.GetRotation());
-						ScaleKeys[KeyIndex] = FVector3f(LocalTransform.GetScale3D());
-					}
+					AnimationDataInComponentSpace[BoneIndex][Key] = LocalTransform;
 				}
-
-				// Set the track's key data
-				Controller->SetBoneTrackKeys(BoneName, PosKeys, RotKeys, ScaleKeys);
 			}
-		}
-
-		return true;
+		}	
 	}
 
-	return false;
+	return AnimationDataInComponentSpace.Num();
 }
 
 void UAnimSequence::AddKeyToSequence(float Time, const FName& BoneName, const FTransform& AdditiveTransform)
diff --git a/Engine/Source/Runtime/Engine/Private/Animation/Rig.cpp b/Engine/Source/Runtime/Engine/Private/Animation/Rig.cpp
deleted file mode 100644
diff --git a/Engine/Source/Runtime/Engine/Private/Animation/Skeleton.cpp b/Engine/Source/Runtime/Engine/Private/Animation/Skeleton.cpp
--- a/Engine/Source/Runtime/Engine/Private/Animation/Skeleton.cpp
+++ b/Engine/Source/Runtime/Engine/Private/Animation/Skeleton.cpp
@@ -17,7 +17,6 @@
 #include "AssetRegistry/AssetData.h"
 #include "AssetRegistry/ARFilter.h"
 #include "AssetRegistry/AssetRegistryModule.h"
-#include "Animation/Rig.h"
 #include "Animation/BlendProfile.h"
 #include "Engine/SkinnedAsset.h"
 #include "Logging/MessageLog.h"
@@ -53,8 +52,6 @@ const FString USkeleton::CurveTagDelimiter = TEXT(";");
 
 const FName USkeleton::CompatibleSkeletonsNameTag = FName(TEXT("CompatibleSkeletonList"));
 const FString USkeleton::CompatibleSkeletonsTagDelimiter = TEXT(";");
-
-const FName USkeleton::RigTag = FName(TEXT("Rig"));
 #endif 
 
 #if WITH_EDITORONLY_DATA
@@ -1954,169 +1951,9 @@ void USkeleton::HandleVirtualBoneChanges()
 
 #if WITH_EDITOR
 
-PRAGMA_DISABLE_DEPRECATION_WARNINGS
-
-void USkeleton::SetRigConfig(URig * Rig)
-{
-	if (RigConfig.Rig != Rig)
-	{
-		RigConfig.Rig = Rig;
-		RigConfig.BoneMappingTable.Empty();
-
-		if (Rig)
-		{
-			const FReferenceSkeleton& RefSkeleton = GetReferenceSkeleton();
-			const TArray<FNode> & Nodes = Rig->GetNodes();
-			// now add bone mapping table
-			for (auto Node: Nodes)
-			{
-				// if find same bone, use that bone for mapping
-				if (RefSkeleton.FindBoneIndex(Node.Name) != INDEX_NONE)
-				{
-					RigConfig.BoneMappingTable.Add(FNameMapping(Node.Name, Node.Name));
-				}
-				else
-				{
-					RigConfig.BoneMappingTable.Add(FNameMapping(Node.Name));
-				}
-			}
-		}
-	}
-}
-
-int32 USkeleton::FindRigBoneMapping(const FName& NodeName) const
-{
-	int32 Index=0;
-	for(const auto & NameMap : RigConfig.BoneMappingTable)
-	{
-		if(NameMap.NodeName == NodeName)
-		{
-			return Index;
-		}
-
-		++Index;
-	}
-
-	return INDEX_NONE;
-}
-
-FName USkeleton::GetRigBoneMapping(const FName& NodeName) const
-{
-	int32 Index = FindRigBoneMapping(NodeName);
-
-	if (Index != INDEX_NONE)
-	{
-		return RigConfig.BoneMappingTable[Index].BoneName;
-	}
-
-	return NAME_None;
-}
-
-FName USkeleton::GetRigNodeNameFromBoneName(const FName& BoneName) const
-{
-	for(const auto & NameMap : RigConfig.BoneMappingTable)
-	{
-		if(NameMap.BoneName == BoneName)
-		{
-			return NameMap.NodeName;
-		}
-	}
-
-	return NAME_None;
-}
-
-int32 USkeleton::GetMappedValidNodes(TArray<FName> &OutValidNodeNames)
-{
-	OutValidNodeNames.Empty();
-
-	for (auto Entry : RigConfig.BoneMappingTable)
-	{
-		if (Entry.BoneName != NAME_None)
-		{
-			OutValidNodeNames.Add(Entry.NodeName);
-		}
-	}
-
-	return OutValidNodeNames.Num();
-}
-
-bool USkeleton::SetRigBoneMapping(const FName& NodeName, FName BoneName)
-{
-	// make sure it's valid
-	int32 BoneIndex = ReferenceSkeleton.FindBoneIndex(BoneName);
-
-	// @todo we need to have validation phase where you can't set same bone for different nodes
-	// but it might be annoying to do that right now since the tool is ugly
-	// so for now it lets you set everything, but in the future
-	// we'll have to add verification
-	if ( BoneIndex == INDEX_NONE )
-	{
-		BoneName = NAME_None;
-	}
-
-	int32 Index = FindRigBoneMapping(NodeName);
-
-	if(Index != INDEX_NONE)
-	{
-		RigConfig.BoneMappingTable[Index].BoneName = BoneName;
-		return true;
-	}
-
-	return false;
-}
-
-void USkeleton::RefreshRigConfig()
-{
-	if (RigConfig.Rig != NULL)
-	{
-		if (RigConfig.BoneMappingTable.Num() > 0)
-		{
-			// verify if any missing bones or anything
-			// remove if removed
-			for ( int32 TableId=0; TableId<RigConfig.BoneMappingTable.Num(); ++TableId )
-			{
-				auto & BoneMapping = RigConfig.BoneMappingTable[TableId];
-
-				if ( RigConfig.Rig->FindNode(BoneMapping.NodeName) == INDEX_NONE)
-				{
-					// if not contains, remove it
-					RigConfig.BoneMappingTable.RemoveAt(TableId);
-					--TableId;
-				}
-			}
-
-			// if the count doesn't match, there is missing nodes. 
-			if (RigConfig.Rig->GetNodeNum() != RigConfig.BoneMappingTable.Num())
-			{
-				int32 NodeNum = RigConfig.Rig->GetNodeNum();
-				for(int32 NodeId=0; NodeId<NodeNum; ++NodeId)
-				{
-					const auto* Node = RigConfig.Rig->GetNode(NodeId);
-
-					if (FindRigBoneMapping(Node->Name) == INDEX_NONE)
-					{
-						RigConfig.BoneMappingTable.Add(FNameMapping(Node->Name));
-					}
-				}
-			}
-		}
-	}
-}
-
-URig * USkeleton::GetRig() const
-{
-	return RigConfig.Rig;
-}
-PRAGMA_ENABLE_DEPRECATION_WARNINGS
-
 void USkeleton::GetAssetRegistryTags(TArray<FAssetRegistryTag>& OutTags) const
 {
 	Super::GetAssetRegistryTags(OutTags);
-	PRAGMA_DISABLE_DEPRECATION_WARNINGS
-	FString RigFullName = (RigConfig.Rig)? RigConfig.Rig->GetFullName() : TEXT("");
-	PRAGMA_ENABLE_DEPRECATION_WARNINGS
-
-	OutTags.Add(FAssetRegistryTag(USkeleton::RigTag, RigFullName, FAssetRegistryTag::TT_Hidden));
 
 	TStringBuilder<256> CompatibleSkeletonsBuilder;
 	
