World Partition: attached actors will inherit some flags from their topmost parent.
Also hide these settings from the property grids sine they are now inherited.

diff --git a/Engine/Source/Runtime/Engine/Private/ActorEditor.cpp b/Engine/Source/Runtime/Engine/Private/ActorEditor.cpp
--- a/Engine/Source/Runtime/Engine/Private/ActorEditor.cpp
+++ b/Engine/Source/Runtime/Engine/Private/ActorEditor.cpp
@@ -102,7 +102,7 @@ bool AActor::CanEditChange(const FProperty* PropertyThatWillChange) const
 				return false;
 			}
 
-			if (!OwningWorldPartition->IsStreamingEnabled())
+			if (!OwningWorldPartition->IsStreamingEnabled() || GetAttachParentActor())
 			{
 				if (bIsSpatiallyLoadedProperty || bIsRuntimeGridProperty)
 				{
diff --git a/Engine/Source/Runtime/Engine/Private/WorldPartition/WorldPartitionActorDescView.cpp b/Engine/Source/Runtime/Engine/Private/WorldPartition/WorldPartitionActorDescView.cpp
--- a/Engine/Source/Runtime/Engine/Private/WorldPartition/WorldPartitionActorDescView.cpp
+++ b/Engine/Source/Runtime/Engine/Private/WorldPartition/WorldPartitionActorDescView.cpp
@@ -12,6 +12,7 @@ FWorldPartitionActorDescView::FWorldPartitionActorDescView()
 
 FWorldPartitionActorDescView::FWorldPartitionActorDescView(const FWorldPartitionActorDesc* InActorDesc)
 	: ActorDesc(InActorDesc)
+	, ParentView(nullptr)
 	, bIsForcedNonSpatiallyLoaded(false)
 	, bIsForcedNoRuntimeGrid(false)
 	, bInvalidDataLayers(false)	
@@ -46,7 +47,17 @@ PRAGMA_ENABLE_DEPRECATION_WARNINGS
 
 FName FWorldPartitionActorDescView::GetRuntimeGrid() const
 {
-	return bIsForcedNoRuntimeGrid ? NAME_None : ActorDesc->GetRuntimeGrid();
+	if (bIsForcedNoRuntimeGrid)
+	{
+		return NAME_None;
+	}
+
+	if (ParentView)
+	{
+		return ParentView->GetRuntimeGrid();
+	}
+
+	return ActorDesc->GetRuntimeGrid();
 }
 
 bool FWorldPartitionActorDescView::GetActorIsEditorOnly() const
@@ -61,7 +72,17 @@ bool FWorldPartitionActorDescView::GetActorIsRuntimeOnly() const
 
 bool FWorldPartitionActorDescView::GetIsSpatiallyLoaded() const
 {
-	return bIsForcedNonSpatiallyLoaded ? false : ActorDesc->GetIsSpatiallyLoaded();
+	if (bIsForcedNonSpatiallyLoaded)
+	{
+		return false;
+	}
+
+	if (ParentView)
+	{
+		return ParentView->GetIsSpatiallyLoaded();
+	}
+	
+	return ActorDesc->GetIsSpatiallyLoaded();
 }
 
 bool FWorldPartitionActorDescView::GetActorIsHLODRelevant() const
@@ -274,8 +295,16 @@ bool FWorldPartitionActorDescView::GetProperty(FName PropertyName, FName* Proper
 {
 	return ActorDesc->GetProperty(PropertyName, PropertyValue);
 }
+
 bool FWorldPartitionActorDescView::HasProperty(FName PropertyName) const
 {
 	return ActorDesc->HasProperty(PropertyName);
 }
+
+void FWorldPartitionActorDescView::SetParentView(const FWorldPartitionActorDescView* InParentView)
+{
+	check(!ParentView);
+	check(GetParentActor().IsValid());
+	ParentView = InParentView;
+}
 #endif
diff --git a/Engine/Source/Runtime/Engine/Private/WorldPartition/WorldPartitionStreamingGeneration.cpp b/Engine/Source/Runtime/Engine/Private/WorldPartition/WorldPartitionStreamingGeneration.cpp
--- a/Engine/Source/Runtime/Engine/Private/WorldPartition/WorldPartitionStreamingGeneration.cpp
+++ b/Engine/Source/Runtime/Engine/Private/WorldPartition/WorldPartitionStreamingGeneration.cpp
@@ -339,6 +339,14 @@ class FWorldPartitionStreamingGenerator
 		}
 	}
 
+	void ResolveParentView(FWorldPartitionActorDescView& ActorDescView, const FActorDescViewMap& ActorDescViewMap)
+	{
+		if (FGuid ParentGuid = ActorDescView.GetParentActor(); ParentGuid.IsValid())
+		{
+			ActorDescView.SetParentView(ActorDescViewMap.FindByGuid(ParentGuid));
+		}
+	}
+
 	void CreateActorDescViewMap(const FStreamingGenerationActorDescCollection& InActorDescCollection, FActorDescViewMap& OutActorDescViewMap, TSet<FGuid>& OutEditorOnlyActorDescMap, const FActorContainerID& InContainerID, TArray<FWorldPartitionActorDescView>& OutContainerInstances)
 	{
 		// Should we handle unsaved or newly created actors?
@@ -626,6 +634,7 @@ class FWorldPartitionStreamingGenerator
 			ResolveRuntimeSpatiallyLoaded(ActorDescView);
 			ResolveRuntimeGrid(ActorDescView);
 			ResolveRuntimeDataLayers(ActorDescView, ContainerCollectionDescriptor.ActorDescViewMap);
+			ResolveParentView(ActorDescView, ContainerCollectionDescriptor.ActorDescViewMap);
 		};
 
 		ContainerCollectionDescriptor.ActorDescViewMap.ForEachActorDescView([this, &ResolveActorDescView](FWorldPartitionActorDescView& ActorDescView)
diff --git a/Engine/Source/Runtime/Engine/Public/WorldPartition/WorldPartitionActorDescView.h b/Engine/Source/Runtime/Engine/Public/WorldPartition/WorldPartitionActorDescView.h
--- a/Engine/Source/Runtime/Engine/Public/WorldPartition/WorldPartitionActorDescView.h
+++ b/Engine/Source/Runtime/Engine/Public/WorldPartition/WorldPartitionActorDescView.h
@@ -104,8 +104,11 @@ public:
 	ENGINE_API bool GetProperty(FName PropertyName, FName* PropertyValue) const;
 	ENGINE_API bool HasProperty(FName PropertyName) const;
 
+	ENGINE_API void SetParentView(const FWorldPartitionActorDescView* InParentView);
+
 protected:
 	const FWorldPartitionActorDesc* ActorDesc;
+	const FWorldPartitionActorDescView* ParentView;
 	bool bIsForcedNonSpatiallyLoaded;
 	bool bIsForcedNoRuntimeGrid;
 	bool bInvalidDataLayers;	
