Fixing compile errors in LightWeightInstanceUtils.h

diff --git a/Engine/Source/Runtime/Engine/Classes/GameFramework/LightWeightInstanceUtils.h b/Engine/Source/Runtime/Engine/Classes/GameFramework/LightWeightInstanceUtils.h
--- a/Engine/Source/Runtime/Engine/Classes/GameFramework/LightWeightInstanceUtils.h
+++ b/Engine/Source/Runtime/Engine/Classes/GameFramework/LightWeightInstanceUtils.h
@@ -7,6 +7,7 @@
 #include "Kismet/BlueprintFunctionLibrary.h"
 #include "Engine/ActorInstanceHandle.h"
 #include "Engine/EngineTypes.h"
+#include "GameFramework/LightWeightInstanceSubsystem.h"
 
 #include "LightWeightInstanceBlueprintFunctionLibrary.generated.h"
 
@@ -33,9 +34,9 @@ namespace LWIUtils
 	static bool DoesHandleSupportInterface(const FActorInstanceHandle& Handle)
 	{
 		// if we have a valid actor, see if it supports the interface
-		if (const UObject* Obj = Handle.GetActorAsUObject())
+		if (const UClass* ObjClass = Handle.GetRepresentedClass())
 		{
-			return Obj->GetClass()->ImplementsInterface(U);
+			return ObjClass->ImplementsInterface(U::StaticClass());
 		}
 
 		// no valid actor, ask the instance manager instead
@@ -45,9 +46,9 @@ namespace LWIUtils
 	template<typename I>
 	static I* FetchInterfaceFromHandle(const FActorInstanceHandle& Handle)
 	{
-		if (Handle.IsActorValid())
+		if (AActor* Actor = Handle.FetchActor())
 		{
-			return Cast<I>(Handle.FetchActor().Get());
+			return Cast<I>(Actor);
 		}
 
 		return FLightWeightInstanceSubsystem::Get().FetchInterfaceObject<I>(Handle);
