Empty Worlds created by WorldPartitions don't need InGamePerformanceTrackers

diff --git a/Engine/Source/Runtime/Engine/Private/World.cpp b/Engine/Source/Runtime/Engine/Private/World.cpp
--- a/Engine/Source/Runtime/Engine/Private/World.cpp
+++ b/Engine/Source/Runtime/Engine/Private/World.cpp
@@ -144,6 +144,10 @@ CSV_DEFINE_CATEGORY(LevelStreamingPendingPurge, (!UE_BUILD_SHIPPING));
 static int32 bDisableRemapScriptActors = 0;
 FAutoConsoleVariableRef CVarDisableRemapScriptActors(TEXT("net.DisableRemapScriptActors"), bDisableRemapScriptActors, TEXT("When set, disables name remapping of compiled script actors (for networking)"));
 
+static bool bDisableInGamePerfTrackersForUninitializedWorlds = true;
+FAutoConsoleVariableRef CVarDisableInGamePerfTrackersForUninitializedWorlds(TEXT("s.World.SkipPerfTrackerForUninitializedWorlds"), bDisableInGamePerfTrackersForUninitializedWorlds, TEXT("When set, disables allocation of InGamePerformanceTrackers for Worlds that aren't initialized."));
+
+
 static TAutoConsoleVariable<int32> CVarPurgeEditorSceneDuringPIE(
 	TEXT("r.PurgeEditorSceneDuringPIE"),
 	0,
@@ -615,7 +619,14 @@ UWorld::UWorld( const FObjectInitializer& ObjectInitializer )
 
 	FWorldDelegates::OnPostWorldCreation.Broadcast(this);
 
-	PerfTrackers = new FWorldInGamePerformanceTrackers();
+	if (!bDisableInGamePerfTrackersForUninitializedWorlds)
+	{
+		PerfTrackers = new FWorldInGamePerformanceTrackers();
+	}
+	else
+	{
+		PerfTrackers = nullptr;
+	}
 
 	IsInBlockTillLevelStreamingCompleted = 0;
 	BlockTillLevelStreamingCompletedEpoch = 0;
@@ -2330,6 +2341,12 @@ void UWorld::InitializeNewWorld(const InitializationValues IVS, bool bInSkipInit
 
 	if (!bInSkipInitWorld)
 	{
+		// If this isn't set, the PerfTrackers are already allocated in the constructor
+		if (bDisableInGamePerfTrackersForUninitializedWorlds && !PerfTrackers)
+		{
+			PerfTrackers = new FWorldInGamePerformanceTrackers();
+		}
+
 		// Initialize the world
 		InitWorld(IVS);
 
