WorldPartition: Actor Filters Fix
- Fix FilteredActors pointer being invalidated (crash)
- Fix bad check causing some unwanted cluster discards

#rnx

diff --git a/Engine/Source/Runtime/Engine/Private/WorldPartition/RuntimeSpatialHash/RuntimeSpatialHashGridHelper.cpp b/Engine/Source/Runtime/Engine/Private/WorldPartition/RuntimeSpatialHash/RuntimeSpatialHashGridHelper.cpp
--- a/Engine/Source/Runtime/Engine/Private/WorldPartition/RuntimeSpatialHash/RuntimeSpatialHashGridHelper.cpp
+++ b/Engine/Source/Runtime/Engine/Private/WorldPartition/RuntimeSpatialHash/RuntimeSpatialHashGridHelper.cpp
@@ -205,7 +205,6 @@ FSquare2DGridHelper GetPartitionedActors(const FBox& WorldBounds, const FSpatial
 	for (const IStreamingGenerationContext::FActorSetInstance* ActorSetInstance : ActorSetInstances)
 	{
 		check(ActorSetInstance->ActorSet->Actors.Num() > 0);
-		check(!ActorSetInstance->FilteredActors || (ActorSetInstance->ActorSet->Actors.Num() != ActorSetInstance->FilteredActors->Num()));
 
 		FSquare2DGridHelper::FGridLevel::FGridCell* GridCell = nullptr;
 
diff --git a/Engine/Source/Runtime/Engine/Private/WorldPartition/WorldPartitionStreamingGeneration.cpp b/Engine/Source/Runtime/Engine/Private/WorldPartition/WorldPartitionStreamingGeneration.cpp
--- a/Engine/Source/Runtime/Engine/Private/WorldPartition/WorldPartitionStreamingGeneration.cpp
+++ b/Engine/Source/Runtime/Engine/Private/WorldPartition/WorldPartitionStreamingGeneration.cpp
@@ -121,12 +121,13 @@ class FWorldPartitionStreamingGenerator
 			{
 				const FContainerCollectionDescriptor& ContainerCollectionDescriptor = StreamingGenerator->ContainerCollectionDescriptorsMap.FindChecked(ContainerCollectionInstanceDescriptor.ActorDescCollection->GetMainContainerPackageName());
 				const FActorSetContainer& ActorSetContainer = ActorSetContainers[ActorSetContainerMap.FindChecked(ContainerCollectionDescriptor.ActorDescCollection)];
-
+				const TSet<FGuid>* FilteredActors = StreamingGenerator->ContainerFilteredActors.Find(ContainerID);
 				for (const TUniquePtr<FActorSet>& ActorSetPtr : ActorSetContainer.ActorSets)
 				{
 					const FActorSet& ActorSet = *ActorSetPtr;
 					const FWorldPartitionActorDescView& ReferenceActorDescView = ActorSetContainer.ActorDescViewMap->FindByGuidChecked(ActorSet.Actors[0]);
 
+					bool bContainsUnfilteredActors = !FilteredActors;
 					// Validate assumptions
 					for (const FGuid& ActorGuid : ActorSet.Actors)
 					{
@@ -134,17 +135,18 @@ class FWorldPartitionStreamingGenerator
 						check(ActorDescView.GetRuntimeGrid() == ReferenceActorDescView.GetRuntimeGrid());
 						check(ActorDescView.GetIsSpatiallyLoaded() == ReferenceActorDescView.GetIsSpatiallyLoaded());
 						check(ActorDescView.GetContentBundleGuid() == ReferenceActorDescView.GetContentBundleGuid());
+						bContainsUnfilteredActors |= (FilteredActors && !FilteredActors->Contains(ActorGuid));
 					}
 
 					// Skip if all actors are filtered out for this container
-					if (!ContainerCollectionInstanceDescriptor.FilteredActors || (ContainerCollectionInstanceDescriptor.FilteredActors->Num() != ActorSet.Actors.Num()))
+					if (bContainsUnfilteredActors)
 					{
 						FActorSetInstance& ActorSetInstance = ActorSetInstances.Emplace_GetRef();
 						const FContainerCollectionInstanceDescriptor::FPerInstanceData& PerInstanceData = ContainerCollectionInstanceDescriptor.GetInstanceData(ReferenceActorDescView.GetGuid());
 				
 						ActorSetInstance.ContainerInstance = &ActorSetContainer;
 						ActorSetInstance.ActorSet = &ActorSet;
-						ActorSetInstance.FilteredActors = ContainerCollectionInstanceDescriptor.FilteredActors;
+						ActorSetInstance.FilteredActors = FilteredActors;
 						ActorSetInstance.ContainerID = ContainerCollectionInstanceDescriptor.ID;
 						ActorSetInstance.Transform = ContainerCollectionInstanceDescriptor.Transform;
 						ActorSetInstance.bIsSpatiallyLoaded = PerInstanceData.bIsSpatiallyLoaded;
@@ -236,7 +238,6 @@ class FWorldPartitionStreamingGenerator
 	{
 		FContainerCollectionInstanceDescriptor()
 			: Bounds(ForceInit)
-			, FilteredActors(nullptr)
 		{}
 
 		FBox Bounds;
@@ -247,7 +248,6 @@ class FWorldPartitionStreamingGenerator
 		FActorContainerID ID;
 		FActorContainerID ParentID;
 		FGuid ContentBundleID;
-		TSet<FGuid>* FilteredActors;
 
 		struct FPerInstanceData
 		{
@@ -569,8 +569,6 @@ class FWorldPartitionStreamingGenerator
 					{						
 						ContainerFilteredActors.Append(WorldPartitionSubsystem->GetFilteredActorsPerContainer(SubContainerInstanceDescriptor.ID, ContainerCollectionInstanceView.GetContainerPackage().ToString(), *ContainerFilter));
 					}
-
-					SubContainerInstanceDescriptor.FilteredActors = ContainerFilteredActors.Find(SubContainerInstanceDescriptor.ID);
 				}
 
 				CreateActorDescriptorViewsRecursive(SubContainerInstanceDescriptor);
@@ -992,7 +990,7 @@ PRAGMA_ENABLE_DEPRECATION_WARNINGS
 					}
 				}
 
-				if (TSet<FGuid>* FilteredActors = ContainerCollectionInstanceDescriptor.FilteredActors)
+				if (TSet<FGuid>* FilteredActors = ContainerFilteredActors.Find(ContainerCollectionInstanceDescriptor.ID))
 				{
 					const bool IsReferencerFiltered = FilteredActors->Contains(ActorDescView.GetGuid());
 					for (const FGuid& ReferenceGuid : ActorDescView.GetReferences())
