Add ContainerID to UWorldPartitionRuntimeLevelStreamingCell::DumpStateLog

#rnx

diff --git a/Engine/Source/Runtime/Engine/Private/WorldPartition/WorldPartitionRuntimeLevelStreamingCell.cpp b/Engine/Source/Runtime/Engine/Private/WorldPartition/WorldPartitionRuntimeLevelStreamingCell.cpp
--- a/Engine/Source/Runtime/Engine/Private/WorldPartition/WorldPartitionRuntimeLevelStreamingCell.cpp
+++ b/Engine/Source/Runtime/Engine/Private/WorldPartition/WorldPartitionRuntimeLevelStreamingCell.cpp
@@ -379,6 +379,7 @@ void UWorldPartitionRuntimeLevelStreamingCell::DumpStateLog(FHierarchicalLogArch
 	{
 		Ar.Printf(TEXT("         Actor Path: %s"), *Mapping.Path.ToString());
 		Ar.Printf(TEXT("      Actor Package: %s"), *Mapping.Package.ToString());
+		Ar.Printf(TEXT(" Actor Container ID: %s"), *Mapping.ContainerID.ToString());
 		Ar.Printf(TEXT("Actor Instance Guid: %s"), *Mapping.ActorInstanceGuid.ToString());
 	}
 }
