Disable defaulting of on to s.World.SkipPerfTrackerForUninitializedWorlds as it breaks in some conditions. A follow-up review will be posted for a more complete fix.

diff --git a/Engine/Source/Runtime/Engine/Private/World.cpp b/Engine/Source/Runtime/Engine/Private/World.cpp
--- a/Engine/Source/Runtime/Engine/Private/World.cpp
+++ b/Engine/Source/Runtime/Engine/Private/World.cpp
@@ -144,7 +144,7 @@ CSV_DEFINE_CATEGORY(LevelStreamingPendingPurge, (!UE_BUILD_SHIPPING));
 static int32 bDisableRemapScriptActors = 0;
 FAutoConsoleVariableRef CVarDisableRemapScriptActors(TEXT("net.DisableRemapScriptActors"), bDisableRemapScriptActors, TEXT("When set, disables name remapping of compiled script actors (for networking)"));
 
-static bool bDisableInGamePerfTrackersForUninitializedWorlds = true;
+static bool bDisableInGamePerfTrackersForUninitializedWorlds = false;
 FAutoConsoleVariableRef CVarDisableInGamePerfTrackersForUninitializedWorlds(TEXT("s.World.SkipPerfTrackerForUninitializedWorlds"), bDisableInGamePerfTrackersForUninitializedWorlds, TEXT("When set, disables allocation of InGamePerformanceTrackers for Worlds that aren't initialized."));
 
 
