UWorldPartitionLevelStreamingDynamic::LoadInEditor - Use the world mount point as the package root, as using /Temp may prevent  non externally referenceable packages from being loaded during import.

diff --git a/Engine/Source/Runtime/Engine/Private/WorldPartition/WorldPartitionLevelStreamingDynamic.cpp b/Engine/Source/Runtime/Engine/Private/WorldPartition/WorldPartitionLevelStreamingDynamic.cpp
--- a/Engine/Source/Runtime/Engine/Private/WorldPartition/WorldPartitionLevelStreamingDynamic.cpp
+++ b/Engine/Source/Runtime/Engine/Private/WorldPartition/WorldPartitionLevelStreamingDynamic.cpp
@@ -31,6 +31,9 @@ UWorldPartitionLevelStreamingDynamic::UWorldPartitionLevelStreamingDynamic(const
 #endif
 	, bShouldBeAlwaysLoaded(false)
 	, bHasSetLevelTransform(false)
+#if WITH_EDITORONLY_DATA
+	, bShouldPerformStandardLevelLoading(false)
+#endif
 {
 #if WITH_EDITOR
 	SetShouldBeVisibleInEditor(false);
@@ -81,8 +84,11 @@ UWorldPartitionLevelStreamingDynamic* UWorldPartitionLevelStreamingDynamic::Load
 	check(World->WorldType == EWorldType::Editor);
 	UWorldPartitionLevelStreamingDynamic* LevelStreaming = NewObject<UWorldPartitionLevelStreamingDynamic>(World, LevelStreamingName, RF_Transient);
 	
-	FString PackageName = FString::Printf(TEXT("/Temp/%s"), *LevelStreamingName.ToString());
-	TSoftObjectPtr<UWorld> WorldAsset(FSoftObjectPath(FString::Printf(TEXT("%s.%s"), *PackageName, *World->GetName())));
+	const FString WorldPackageName = World->GetPackage()->GetPathName();
+	const FStringView WorldMountPointName = FPathViews::GetMountPointNameFromPath(WorldPackageName);
+
+	FString PackageName = FString::Printf(TEXT("Temp/%s"), *LevelStreamingName.ToString());
+	TSoftObjectPtr<UWorld> WorldAsset(FSoftObjectPath(FString::Printf(TEXT("/%.*s/%s.%s"), WorldMountPointName.Len(), WorldMountPointName.GetData(), *PackageName, *World->GetName())));
 	LevelStreaming->SetWorldAsset(WorldAsset);
 	
 	LevelStreaming->LevelTransform = FTransform::Identity;
@@ -187,8 +193,7 @@ bool UWorldPartitionLevelStreamingDynamic::RequestLevel(UWorld* InPersistentWorl
 	const FName WorldAssetPackageFName = GetWorldAssetPackageFName();
 	const FString WorldAssetPackageName = GetWorldAssetPackageName();
 
-	// Use parent streaming implementation if world asset package is not a memory package
-	if (!FPackageName::IsMemoryPackage(WorldAssetPackageName) && !FPackageName::IsTempPackage(WorldAssetPackageName))
+	if (bShouldPerformStandardLevelLoading)
 	{
 		return Super::RequestLevel(InPersistentWorld, bInAllowLevelLoadRequests, InBlockPolicy);
 	}
diff --git a/Engine/Source/Runtime/Engine/Private/WorldPartition/WorldPartitionRuntimeLevelStreamingCell.cpp b/Engine/Source/Runtime/Engine/Private/WorldPartition/WorldPartitionRuntimeLevelStreamingCell.cpp
--- a/Engine/Source/Runtime/Engine/Private/WorldPartition/WorldPartitionRuntimeLevelStreamingCell.cpp
+++ b/Engine/Source/Runtime/Engine/Private/WorldPartition/WorldPartitionRuntimeLevelStreamingCell.cpp
@@ -126,6 +126,10 @@ bool UWorldPartitionRuntimeLevelStreamingCell::CreateAndSetLevelStreaming(const
 	LevelStreaming->Initialize(*this);
 	LevelStreaming->PackageNameToLoad = FName(LongPackageName);
 
+#if WITH_EDITOR
+	LevelStreaming->SetShouldPerformStandardLevelLoading(true);
+#endif
+
 	if (OwningWorld->IsPlayInEditor() && OwningWorld->GetPackage()->HasAnyPackageFlags(PKG_PlayInEditor) && OwningWorld->GetPackage()->GetPIEInstanceID() != INDEX_NONE)
 	{
 		// When renaming for PIE, make sure to keep World's name so that linker can properly remap with Package's instancing context
diff --git a/Engine/Source/Runtime/Engine/Public/WorldPartition/WorldPartitionLevelStreamingDynamic.h b/Engine/Source/Runtime/Engine/Public/WorldPartition/WorldPartitionLevelStreamingDynamic.h
--- a/Engine/Source/Runtime/Engine/Public/WorldPartition/WorldPartitionLevelStreamingDynamic.h
+++ b/Engine/Source/Runtime/Engine/Public/WorldPartition/WorldPartitionLevelStreamingDynamic.h
@@ -59,6 +59,8 @@ class UWorldPartitionLevelStreamingDynamic : public ULevelStreamingDynamic
 
 	bool GetLoadSucceeded() const { return bLoadSucceeded; }
 
+	void SetShouldPerformStandardLevelLoading(bool bInShouldPerformStandardLevelLoading) { bShouldPerformStandardLevelLoading = bInShouldPerformStandardLevelLoading; }
+
 private:
 	ENGINE_API void CreateRuntimeLevel();
 	ENGINE_API bool IssueLoadRequests();
@@ -92,6 +94,9 @@ private:
 	bool bHasSetLevelTransform;
 
 #if WITH_EDITORONLY_DATA
+	UPROPERTY()
+	bool bShouldPerformStandardLevelLoading;
+
 	UPROPERTY()
 	TObjectPtr<UActorContainer> UnsavedActorsContainer;
 #endif
