Change ChildActorComponent back to listening to OnDestroy for clearing references, just with the change of having it register during OnRegister and OnRep_ChildActor now to catch both client cases.

diff --git a/Engine/Source/Runtime/Engine/Classes/Components/ChildActorComponent.h b/Engine/Source/Runtime/Engine/Classes/Components/ChildActorComponent.h
--- a/Engine/Source/Runtime/Engine/Classes/Components/ChildActorComponent.h
+++ b/Engine/Source/Runtime/Engine/Classes/Components/ChildActorComponent.h
@@ -220,8 +220,10 @@ private:
 	UFUNCTION()
 	void OnRep_ChildActor();
 
+	void RegisterChildActorDestroyedDelegate();
+
 	UFUNCTION()
-	ENGINE_API void OnChildActorEndPlay(AActor* Actor, EEndPlayReason::Type EndPlayReason);
+	ENGINE_API void OnChildActorDestroyed(AActor* Actor);
 };
 
 struct FActorParentComponentSetter
diff --git a/Engine/Source/Runtime/Engine/Private/Components/ChildActorComponent.cpp b/Engine/Source/Runtime/Engine/Private/Components/ChildActorComponent.cpp
--- a/Engine/Source/Runtime/Engine/Private/Components/ChildActorComponent.cpp
+++ b/Engine/Source/Runtime/Engine/Private/Components/ChildActorComponent.cpp
@@ -93,6 +93,9 @@ void UChildActorComponent::OnRegister()
 		{
 			// Ensure the components replication is correctly initialized
 			SetIsReplicated(ChildActor->GetIsReplicated());
+
+			// All other paths register this delegate through CreateChildActor; this is the only path that doesn't go there
+			RegisterChildActorDestroyedDelegate();
 		}
 	}
 	else if (ChildActorClass)
@@ -656,50 +659,52 @@ bool UChildActorComponent::IsBeingRemovedFromLevel() const
 }
 
 void UChildActorComponent::OnRep_ChildActor()
+{
+	RegisterChildActorDestroyedDelegate();
+}
+
+void UChildActorComponent::RegisterChildActorDestroyedDelegate()
 {
 	if (ChildActor)
 	{
-		ChildActor->OnEndPlay.AddUniqueDynamic(this, &ThisClass::OnChildActorEndPlay);
+		ChildActor->OnDestroyed.AddUniqueDynamic(this, &ThisClass::OnChildActorDestroyed);
 	}
 }
 
-void UChildActorComponent::OnChildActorEndPlay(AActor* Actor, EEndPlayReason::Type EndPlayReason)
+void UChildActorComponent::OnChildActorDestroyed(AActor* Actor)
 {
 	if (GExitPurge)
 	{
 		return;
 	}
 
-	if (EndPlayReason == EEndPlayReason::Destroyed)
+	if (Actor && (Actor->HasAuthority() || !IsChildActorReplicated()) && !IsBeingRemovedFromLevel())
 	{
-		if (Actor && (Actor->HasAuthority() || !IsChildActorReplicated()) && !IsBeingRemovedFromLevel())
+		UWorld* World = Actor->GetWorld();
+		// World may be nullptr during shutdown
+		if (World != nullptr)
 		{
-			UWorld* World = Actor->GetWorld();
-			// World may be nullptr during shutdown
-			if (World != nullptr)
-			{
-				UClass* ChildClass = Actor->GetClass();
+			UClass* ChildClass = Actor->GetClass();
 
-				// We would like to make certain that our name is not going to accidentally get taken from us while we're destroyed
-				// so we increment ClassUnique beyond our index to be certain of it.  This is ... a bit hacky.
-				if (!GFastPathUniqueNameGeneration)
-				{
-					UpdateSuffixForNextNewObject(Actor->GetOuter(), ChildClass, [Actor](int32& Index) { Index = FMath::Max(Index, Actor->GetFName().GetNumber()); });
-				}
+			// We would like to make certain that our name is not going to accidentally get taken from us while we're destroyed
+			// so we increment ClassUnique beyond our index to be certain of it.  This is ... a bit hacky.
+			if (!GFastPathUniqueNameGeneration)
+			{
+				UpdateSuffixForNextNewObject(Actor->GetOuter(), ChildClass, [Actor](int32& Index) { Index = FMath::Max(Index, Actor->GetFName().GetNumber()); });
+			}
 
-				// If we are getting here due to garbage collection we can't rename, so we'll have to abandon this child actor name and pick up a new one
-				if (!IsGarbageCollecting())
-				{
-					const FString ObjectBaseName = FString::Printf(TEXT("DESTROYED_%s_CHILDACTOR"), *ChildClass->GetName());
-					Actor->Rename(*MakeUniqueObjectName(Actor->GetOuter(), ChildClass, *ObjectBaseName).ToString(), nullptr, REN_DoNotDirty | REN_ForceNoResetLoaders);
-				}
-				else
+			// If we are getting here due to garbage collection we can't rename, so we'll have to abandon this child actor name and pick up a new one
+			if (!IsGarbageCollecting())
+			{
+				const FString ObjectBaseName = FString::Printf(TEXT("DESTROYED_%s_CHILDACTOR"), *ChildClass->GetName());
+				Actor->Rename(*MakeUniqueObjectName(Actor->GetOuter(), ChildClass, *ObjectBaseName).ToString(), nullptr, REN_DoNotDirty | REN_ForceNoResetLoaders);
+			}
+			else
+			{
+				ChildActorName = NAME_None;
+				if (CachedInstanceData)
 				{
-					ChildActorName = NAME_None;
-					if (CachedInstanceData)
-					{
-						CachedInstanceData->ChildActorName = NAME_None;
-					}
+					CachedInstanceData->ChildActorName = NAME_None;
 				}
 			}
 		}
@@ -811,7 +816,7 @@ void UChildActorComponent::CreateChildActor(TFunction<void(AActor*)> CustomizerF
 					}
 					else
 					{
-						ChildActor->OnEndPlay.AddDynamic(this, &ThisClass::OnChildActorEndPlay);
+						RegisterChildActorDestroyedDelegate();
 					}
 
 					ChildActorName = ChildActor->GetFName();
