Fix for nullptr reference causing crash in UMaterialInstanceDynamic::Create due to oversight where it previously checked further down in the function calls.

diff --git a/Engine/Source/Runtime/Engine/Private/Materials/MaterialInstanceDynamic.cpp b/Engine/Source/Runtime/Engine/Private/Materials/MaterialInstanceDynamic.cpp
--- a/Engine/Source/Runtime/Engine/Private/Materials/MaterialInstanceDynamic.cpp
+++ b/Engine/Source/Runtime/Engine/Private/Materials/MaterialInstanceDynamic.cpp
@@ -76,41 +76,44 @@ UMaterialInstanceDynamic* UMaterialInstanceDynamic::Create(UMaterialInterface* P
 
 	UObject* Outer = InOuter ? InOuter : TransientPackage;
 
-	if (Name == NAME_None)
+	if (Name == NAME_None && ParentMaterial)
 	{
 		Name = FName(FString("MID_") + ParentMaterial->GetName());
 	}
 
-	UMaterialInstanceDynamic* ExistingMID = FindObjectFast<UMaterialInstanceDynamic>(Outer, *Name.ToString(), true);
-	if (ExistingMID)
+	if (Name != NAME_None)
 	{
-		bool bRenamed = false;
+		UMaterialInstanceDynamic* ExistingMID = FindObjectFast<UMaterialInstanceDynamic>(Outer, *Name.ToString(), true);
+		if (ExistingMID)
+		{
+			bool bRenamed = false;
 
 #if WITH_EDITOR
-		// try to rename to enhance construction script determinism:
-		AActor* OwningActor = Cast<AActor>(Outer);
-		if (!OwningActor)
-		{
-			OwningActor = Outer->GetTypedOuter<AActor>();
-		}
+			// try to rename to enhance construction script determinism:
+			AActor* OwningActor = Cast<AActor>(Outer);
+			if (!OwningActor)
+			{
+				OwningActor = Outer->GetTypedOuter<AActor>();
+			}
 
-		if (OwningActor && OwningActor->IsRunningUserConstructionScript())
-		{
-			bRenamed = true;
-			// a collision, we're going to move this existing mid to the transient package and claim the name 
-			// for ourself:
-			ExistingMID	->Rename(
-				nullptr,
-				TransientPackage,
-				REN_DoNotDirty | REN_DontCreateRedirectors | REN_ForceNoResetLoaders | REN_NonTransactional
-			);
-		}
+			if (OwningActor && OwningActor->IsRunningUserConstructionScript())
+			{
+				bRenamed = true;
+				// a collision, we're going to move this existing mid to the transient package and claim the name 
+				// for ourself:
+				ExistingMID->Rename(
+					nullptr,
+					TransientPackage,
+					REN_DoNotDirty | REN_DontCreateRedirectors | REN_ForceNoResetLoaders | REN_NonTransactional
+				);
+			}
 #endif
 
-		if (!bRenamed)
-		{
-			// If a MID is made with the same name and outer as another, it will overwrite it. To avoid this we will change the name when there is a collision.
-			Name = MakeUniqueObjectName(Outer, UMaterialInstanceDynamic::StaticClass(), Name);
+			if (!bRenamed)
+			{
+				// If a MID is made with the same name and outer as another, it will overwrite it. To avoid this we will change the name when there is a collision.
+				Name = MakeUniqueObjectName(Outer, UMaterialInstanceDynamic::StaticClass(), Name);
+			}
 		}
 	}
 	UMaterialInstanceDynamic* MID = NewObject<UMaterialInstanceDynamic>(Outer, Name);
