Simplified how curve blending is handled in DeadBlending node to always blend in and out missing curves from zero and fixed bug in existing behavior that was causing curves to pop.

diff --git a/Engine/Source/Runtime/Engine/Classes/Animation/AnimNode_DeadBlending.h b/Engine/Source/Runtime/Engine/Classes/Animation/AnimNode_DeadBlending.h
--- a/Engine/Source/Runtime/Engine/Classes/Animation/AnimNode_DeadBlending.h
+++ b/Engine/Source/Runtime/Engine/Classes/Animation/AnimNode_DeadBlending.h
@@ -103,20 +103,6 @@ private:
 	UPROPERTY(EditAnywhere, Category = Extrapolation, meta = (Min = "0.0", UIMin = "0.0"), DisplayName = "Maximum Extrapolation Half Life")
 	float ExtrapolationHalfLifeMax = 1.0f;
 
-	/**
-	 * When enabled, curves will blend in from zero when they are in the animation being transitioned to but not in the 
-	 * animation being transitioned from.
-	 */
-	UPROPERTY(EditAnywhere, Category = Curves)
-	bool bBlendInMissingCurves = true;
-
-	/**
-	 * When enabled, curves will blend out to zero when they are in the animation being transitioned from but not in the 
-	 * animation being transitioned to.
-	 */
-	UPROPERTY(EditAnywhere, Category = Curves)
-	bool bBlendOutMissingCurves = true;
-
 #if WITH_EDITORONLY_DATA
 	
 	// This setting can be used to show what the extrapolation of the animation looks like.
@@ -208,7 +194,6 @@ private:
 
 	struct FDeadBlendingCurveElement : public UE::Anim::FCurveElement
 	{
-		bool Valid = false;
 		float Velocity = 0.0f;
 		float HalfLife = 0.0f;
 
diff --git a/Engine/Source/Runtime/Engine/Private/Animation/AnimNode_DeadBlending.cpp b/Engine/Source/Runtime/Engine/Private/Animation/AnimNode_DeadBlending.cpp
--- a/Engine/Source/Runtime/Engine/Private/Animation/AnimNode_DeadBlending.cpp
+++ b/Engine/Source/Runtime/Engine/Private/Animation/AnimNode_DeadBlending.cpp
@@ -350,61 +350,45 @@ void FAnimNode_DeadBlending::InitFrom(
 
 	CurveData.CopyFrom(SrcPoseCurr.Curves.BlendedCurve);
 
-	// Record curve state
+	// Record Source Animation Curve State
 
 	UE::Anim::FNamedValueArrayUtils::Union(CurveData, SrcPoseCurr.Curves.BlendedCurve,
 		[this](FDeadBlendingCurveElement& OutResultElement, const UE::Anim::FCurveElement& InElement1, UE::Anim::ENamedValueUnionFlags InFlags)
 		{
-			// Here we need to store an additional `Valid` value because the `Union` operation used later on 
-			// can add curves which are in `InCurves` but which are not in `SrcPoseCurr` or `SrcPosePrev`.
-			// 
-			// Since we only want to extrapolate curves which are in `SrcPoseCurr` we therefore we need a `Valid` 
-			// value to indicate this. Probably it would be better to not use the `Union` operation and instead something 
-			// like `Intersection`, but right now no version of this is provided which modifies the curves in place 
-			// in the same way `Union` does.
-
-			OutResultElement.Valid = true;
 			OutResultElement.Value = InElement1.Value;
 			OutResultElement.Velocity = 0.0f;
 			OutResultElement.HalfLife = ExtrapolationHalfLifeMin;
 		});
 
+	// Record Source Animation Curve Velocity
+
 	if (SrcPoseCurr.DeltaTime > UE_SMALL_NUMBER)
 	{
-		// Record Curve Velocity
-
 		UE::Anim::FNamedValueArrayUtils::Union(CurveData, SrcPosePrev.Curves.BlendedCurve,
 			[DeltaTime = SrcPoseCurr.DeltaTime](FDeadBlendingCurveElement& OutResultElement, const UE::Anim::FCurveElement& InElement1, UE::Anim::ENamedValueUnionFlags InFlags)
 			{
-				if (OutResultElement.Valid)
+				bool bSrcCurrValid = (bool)(InFlags & UE::Anim::ENamedValueUnionFlags::ValidArg0);
+				bool bSrcPrevValid = (bool)(InFlags & UE::Anim::ENamedValueUnionFlags::ValidArg1);
+
+				if (bSrcCurrValid && bSrcPrevValid)
 				{
 					OutResultElement.Velocity = (OutResultElement.Value - InElement1.Value) / DeltaTime;
 				}
 			});
+	}
 
-		// Record Half-life
-
-		UE::Anim::FNamedValueArrayUtils::Union(CurveData, InCurves,
-			[this](FDeadBlendingCurveElement& OutResultElement, const UE::Anim::FCurveElement& InElement1, UE::Anim::ENamedValueUnionFlags InFlags)
-			{
-				if (bBlendInMissingCurves)
-				{
-					OutResultElement.Valid = true;
-					OutResultElement.Value = 0.0f;
-					OutResultElement.Velocity = 0.0f;
-				}
+	// Perform Union with Curves from Destination Animation and compute Half-life
 
-				if (OutResultElement.Valid)
-				{
-					OutResultElement.HalfLife = UE::Anim::DeadBlending::Private::ComputeDecayHalfLifeFromDiffAndVelocity(
-						InElement1.Value - OutResultElement.Value,
-						OutResultElement.Velocity,
-						ExtrapolationHalfLife,
-						ExtrapolationHalfLifeMin,
-						ExtrapolationHalfLifeMax);
-				}
-			});
-	}
+	UE::Anim::FNamedValueArrayUtils::Union(CurveData, InCurves,
+		[this](FDeadBlendingCurveElement& OutResultElement, const UE::Anim::FCurveElement& InElement1, UE::Anim::ENamedValueUnionFlags InFlags)
+		{
+			OutResultElement.HalfLife = UE::Anim::DeadBlending::Private::ComputeDecayHalfLifeFromDiffAndVelocity(
+				InElement1.Value - OutResultElement.Value,
+				OutResultElement.Velocity,
+				ExtrapolationHalfLife,
+				ExtrapolationHalfLifeMin,
+				ExtrapolationHalfLifeMax);
+		});
 
 	// Apply filtering to remove anything we don't want to inertialize
 
@@ -512,26 +496,23 @@ void FAnimNode_DeadBlending::ApplyTo(FCompactPose& InOutPose, FBlendedCurve& InO
 		{
 			// Compute Extrapolated Curve Value
 
-			if (InElement1.Valid && (bBlendOutMissingCurves || (bool)(InFlags & UE::Anim::ENamedValueUnionFlags::ValidArg0)))
-			{
-				const float ExtrapolatedCurve = UE::Anim::DeadBlending::Private::ExtrapolateCurve(
-					InElement1.Value,
-					InElement1.Velocity,
-					InertializationTime,
-					InElement1.HalfLife);
+			const float ExtrapolatedCurve = UE::Anim::DeadBlending::Private::ExtrapolateCurve(
+				InElement1.Value,
+				InElement1.Velocity,
+				InertializationTime,
+				InElement1.HalfLife);
 
 #if WITH_EDITORONLY_DATA
-				if (bShowExtrapolations)
-				{
-					OutResultElement.Value = ExtrapolatedCurve;
-					OutResultElement.Flags |= InElement1.Flags;
-					return;
-				}
-#endif
-
-				OutResultElement.Value = FMath::Lerp(OutResultElement.Value, ExtrapolatedCurve, CurveAlpha);
+			if (bShowExtrapolations)
+			{
+				OutResultElement.Value = ExtrapolatedCurve;
 				OutResultElement.Flags |= InElement1.Flags;
+				return;
 			}
+#endif
+
+			OutResultElement.Value = FMath::Lerp(OutResultElement.Value, ExtrapolatedCurve, CurveAlpha);
+			OutResultElement.Flags |= InElement1.Flags;
 		});
 }
 
