World Partition: attached actors will inherit data layers from their topmost parent.

diff --git a/Engine/Source/Runtime/Engine/Private/ActorEditor.cpp b/Engine/Source/Runtime/Engine/Private/ActorEditor.cpp
--- a/Engine/Source/Runtime/Engine/Private/ActorEditor.cpp
+++ b/Engine/Source/Runtime/Engine/Private/ActorEditor.cpp
@@ -117,7 +117,7 @@ bool AActor::CanEditChange(const FProperty* PropertyThatWillChange) const
 		return false;
 	}
 
-	if (bIsDataLayersProperty && (!SupportsDataLayerType(UDataLayerInstance::StaticClass()) || !IsUserManaged()))
+	if (bIsDataLayersProperty && (!SupportsDataLayerType(UDataLayerInstance::StaticClass()) || !IsUserManaged() || GetAttachParentActor()))
 	{
 		return false;
 	}
diff --git a/Engine/Source/Runtime/Engine/Private/WorldPartition/WorldPartitionActorDescView.cpp b/Engine/Source/Runtime/Engine/Private/WorldPartition/WorldPartitionActorDescView.cpp
--- a/Engine/Source/Runtime/Engine/Private/WorldPartition/WorldPartitionActorDescView.cpp
+++ b/Engine/Source/Runtime/Engine/Private/WorldPartition/WorldPartitionActorDescView.cpp
@@ -15,7 +15,7 @@ FWorldPartitionActorDescView::FWorldPartitionActorDescView(const FWorldPartition
 	, ParentView(nullptr)
 	, bIsForcedNonSpatiallyLoaded(false)
 	, bIsForcedNoRuntimeGrid(false)
-	, bInvalidDataLayers(false)	
+	, bIsForcedNoDataLayers(false)	
 {}
 
 const FGuid& FWorldPartitionActorDescView::GetGuid() const
@@ -103,29 +103,39 @@ void FWorldPartitionActorDescView::SetDataLayerInstanceNames(const TArray<FName>
 
 const TArray<FName>& FWorldPartitionActorDescView::GetDataLayerInstanceNames() const
 {
-	static TArray<FName> EmptyDataLayers;
-	if (bInvalidDataLayers)
+	if (bIsForcedNoDataLayers)
 	{
+		static TArray<FName> EmptyDataLayers;
 		return EmptyDataLayers;
 	}
-	else if (ResolvedDataLayerInstanceNames.IsSet())
+
+	if (ParentView)
+	{
+		return ParentView->GetDataLayerInstanceNames();
+	}
+	
+	if (ResolvedDataLayerInstanceNames.IsSet())
 	{
 		return ResolvedDataLayerInstanceNames.GetValue();
 	}
+
 	return ActorDesc->GetDataLayerInstanceNames();
 }
 
 const TArray<FName>& FWorldPartitionActorDescView::GetRuntimeDataLayerInstanceNames() const
 {
-	static TArray<FName> EmptyDataLayers;
-	if (!bInvalidDataLayers)
+	if (bIsForcedNoDataLayers || !ensure(RuntimeDataLayerInstanceNames.IsSet()))
 	{
-		if (ensure(RuntimeDataLayerInstanceNames.IsSet()))
-		{
-			return RuntimeDataLayerInstanceNames.GetValue();
-		}
+		static TArray<FName> EmptyDataLayers;
+		return EmptyDataLayers;
+	}
+
+	if (ParentView)
+	{
+		return ParentView->GetRuntimeDataLayerInstanceNames();
 	}
-	return EmptyDataLayers;
+
+	return RuntimeDataLayerInstanceNames.GetValue();
 }
 
 const TArray<FName>& FWorldPartitionActorDescView::GetTags() const
@@ -257,11 +267,11 @@ void FWorldPartitionActorDescView::SetForcedNoRuntimeGrid()
 	bIsForcedNoRuntimeGrid = true;	
 }
 
-void FWorldPartitionActorDescView::SetInvalidDataLayers()
+void FWorldPartitionActorDescView::SetForcedNoDataLayers()
 {
-	if (!bInvalidDataLayers)
+	if (!bIsForcedNoDataLayers)
 	{
-		bInvalidDataLayers = true;
+		bIsForcedNoDataLayers = true;
 		UE_LOG(LogWorldPartition, Verbose, TEXT("Actor '%s' data layers invalidated"), *GetActorLabelOrName().ToString());
 	}
 }
diff --git a/Engine/Source/Runtime/Engine/Private/WorldPartition/WorldPartitionStreamingGeneration.cpp b/Engine/Source/Runtime/Engine/Private/WorldPartition/WorldPartitionStreamingGeneration.cpp
--- a/Engine/Source/Runtime/Engine/Private/WorldPartition/WorldPartitionStreamingGeneration.cpp
+++ b/Engine/Source/Runtime/Engine/Private/WorldPartition/WorldPartitionStreamingGeneration.cpp
@@ -694,7 +694,7 @@ class FWorldPartitionStreamingGenerator
 					if (ActorDescView->GetRuntimeDataLayerInstanceNames().Num())
 					{
 						ErrorHandler->OnInvalidReferenceLevelScriptDataLayers(*ActorDescView);
-						ActorDescView->SetInvalidDataLayers();
+						ActorDescView->SetForcedNoDataLayers();
 					}
 				}
 			}
@@ -870,8 +870,8 @@ class FWorldPartitionStreamingGenerator
 								}
 								else
 								{
-									RefererActorDescView->SetInvalidDataLayers();
-									ReferenceActorDescView->SetInvalidDataLayers();
+									RefererActorDescView->SetForcedNoDataLayers();
+									ReferenceActorDescView->SetForcedNoDataLayers();
 								}
 
 								NbErrorsDetected++;
diff --git a/Engine/Source/Runtime/Engine/Public/WorldPartition/WorldPartitionActorDescView.h b/Engine/Source/Runtime/Engine/Public/WorldPartition/WorldPartitionActorDescView.h
--- a/Engine/Source/Runtime/Engine/Public/WorldPartition/WorldPartitionActorDescView.h
+++ b/Engine/Source/Runtime/Engine/Public/WorldPartition/WorldPartitionActorDescView.h
@@ -79,7 +79,7 @@ public:
 
 	ENGINE_API void SetForcedNonSpatiallyLoaded();
 	ENGINE_API void SetForcedNoRuntimeGrid();
-	ENGINE_API void SetInvalidDataLayers();
+	ENGINE_API void SetForcedNoDataLayers();
 	ENGINE_API void SetRuntimeDataLayerInstanceNames(const TArray<FName>& InRuntimeDataLayerInstanceNames);
 	ENGINE_API void SetRuntimeReferences(const TArray<FGuid>& InRuntimeReferences);
 	ENGINE_API void SetEditorReferences(const TArray<FGuid>& InEditorReferences);
@@ -111,7 +111,7 @@ protected:
 	const FWorldPartitionActorDescView* ParentView;
 	bool bIsForcedNonSpatiallyLoaded;
 	bool bIsForcedNoRuntimeGrid;
-	bool bInvalidDataLayers;	
+	bool bIsForcedNoDataLayers;	
 	TOptional<TArray<FName>> ResolvedDataLayerInstanceNames;
 	TOptional<TArray<FName>> RuntimeDataLayerInstanceNames;
 	TOptional<TArray<FGuid>> RuntimeReferences;
