Fixes issue with Orient & Scale translation retargeting when using compatible skeletons.

#UDN https://epicgames.lightning.force.com/lightning/r/Case/5004z00001oY7ceAAC

diff --git a/Engine/Source/Runtime/Engine/Private/BoneContainer.cpp b/Engine/Source/Runtime/Engine/Private/BoneContainer.cpp
--- a/Engine/Source/Runtime/Engine/Private/BoneContainer.cpp
+++ b/Engine/Source/Runtime/Engine/Private/BoneContainer.cpp
@@ -407,13 +407,29 @@ const FRetargetSourceCachedData& FBoneContainer::GetRetargetSourceCachedData(con
 	return GetRetargetSourceCachedData(InRetargetSourceName, FSkeletonRemapping(), RetargetTransforms);
 }
 
-const FRetargetSourceCachedData& FBoneContainer::GetRetargetSourceCachedData(const FName& InSourceName, const FSkeletonRemapping& InRemapping, const TArray<FTransform>& InRetargetTransforms) const
+const FRetargetSourceCachedData& FBoneContainer::GetRetargetSourceCachedData(
+	const FName& InSourceName,
+	const FSkeletonRemapping& InRemapping,
+	const TArray<FTransform>& InRetargetTransforms) const
 {
 	LLM_SCOPE_BYNAME(TEXT("Animation/BoneContainer"));
-	FRetargetSourceCachedData* RetargetSourceCachedData = RetargetSourceCachedDataLUT.Find(InSourceName);
+	
+	const USkeleton* SourceSkeleton = InRemapping.IsValid() ? InRemapping.GetSourceSkeleton().Get() : AssetSkeleton.Get();
+	check(SourceSkeleton);
+	
+	// Invalid retarget source names (not found on the skeleton) are internally treated as None, so we do the same
+	FName RetargetSourceKey = InSourceName;
+	if (!SourceSkeleton->AnimRetargetSources.Contains(InSourceName))
+	{
+		RetargetSourceKey = NAME_None;
+	}
+	
+	// Construct a key from the skeleton and the retarget source since both are necessary for uniqueness
+	FRetargetSourceCachedDataKey LUTKey(Cast<UObject>(SourceSkeleton), RetargetSourceKey);
+	FRetargetSourceCachedData* RetargetSourceCachedData = RetargetSourceCachedDataLUT.Find(LUTKey);
 	if (!RetargetSourceCachedData)
 	{
-		RetargetSourceCachedData = &RetargetSourceCachedDataLUT.Add(InSourceName);
+		RetargetSourceCachedData = &RetargetSourceCachedDataLUT.Add(LUTKey);
 
 		// Build Cached Data for OrientAndScale retargeting.
 
@@ -428,7 +444,7 @@ const FRetargetSourceCachedData& FBoneContainer::GetRetargetSourceCachedData(con
 			const int32 TargetSkeletonBoneIndex = CompactPoseToSkeletonIndex[CompactBoneIndex];
 			const int32 SourceSkeletonBoneIndex = InRemapping.IsValid() ? InRemapping.GetSourceSkeletonBoneIndex(TargetSkeletonBoneIndex) : TargetSkeletonBoneIndex;
 
-			if (SourceSkeletonBoneIndex != INDEX_NONE && AssetSkeleton.GetEvenIfUnreachable()->GetBoneTranslationRetargetingMode(SourceSkeletonBoneIndex, bDisableRetargeting) == EBoneTranslationRetargetingMode::OrientAndScale)
+			if (AssetSkeleton.GetEvenIfUnreachable()->GetBoneTranslationRetargetingMode(TargetSkeletonBoneIndex, bDisableRetargeting) == EBoneTranslationRetargetingMode::OrientAndScale)
 			{
 				if(AuthoredOnRefSkeleton.IsValidIndex(SourceSkeletonBoneIndex))
 				{
diff --git a/Engine/Source/Runtime/Engine/Public/BoneContainer.h b/Engine/Source/Runtime/Engine/Public/BoneContainer.h
--- a/Engine/Source/Runtime/Engine/Public/BoneContainer.h
+++ b/Engine/Source/Runtime/Engine/Public/BoneContainer.h
@@ -13,8 +13,6 @@
 #include "Animation/AnimCurveMetadata.h"
 #include "Animation/AnimBulkCurves.h"
 
-class USkeletalMesh;
-class USkeleton;
 class USkeletalMesh;
 struct FBoneContainer;
 struct FSkeletonRemapping;
@@ -163,6 +161,27 @@ struct FCompactPoseBoneIndexReverseIterator
 	FCompactPoseBoneIndex operator*() const { return FCompactPoseBoneIndex(Index); }
 };
 
+struct FRetargetSourceCachedDataKey
+{
+	FTopLevelAssetPath SourceSkeletonPath;
+	FName SourceRetargetName;
+ 
+	FRetargetSourceCachedDataKey(const UObject* Skeleton, const FName& InSourceRetargetName)
+		: SourceSkeletonPath(Skeleton)
+		, SourceRetargetName(InSourceRetargetName)
+	{
+	}
+ 
+	bool operator==(const FRetargetSourceCachedDataKey& Other) const
+	{
+		return SourceSkeletonPath == Other.SourceSkeletonPath && SourceRetargetName == Other.SourceRetargetName;
+	}
+	friend uint32 GetTypeHash(const FRetargetSourceCachedDataKey& Obj)
+	{
+		return HashCombine(GetTypeHash(Obj.SourceSkeletonPath), GetTypeHash(Obj.SourceRetargetName));
+	}
+};
+
 /**
 * This is a native transient structure.
 * Contains:
@@ -679,7 +698,7 @@ private:
 	 * Runtime cached data for retargeting from a specific RetargetSource to this current SkelMesh LOD.
 	 * @todo: We could also cache this once per skelmesh per lod, rather than creating it at runtime for each skelmesh instance.
 	 */
-	mutable TMap<FName, FRetargetSourceCachedData> RetargetSourceCachedDataLUT;
+	mutable TMap<FRetargetSourceCachedDataKey, FRetargetSourceCachedData> RetargetSourceCachedDataLUT;
 
 	/** Initialize FBoneContainer. */
 	ENGINE_API void Initialize(const UE::Anim::FCurveFilterSettings& CurveFilterSettings);
