Optimize UUserDefinedStruct::InitializeStruct by avoiding redundant operations for STRUCT_IsPlainOldData instances

diff --git a/Engine/Source/Runtime/Engine/Private/UserDefinedStruct.cpp b/Engine/Source/Runtime/Engine/Private/UserDefinedStruct.cpp
--- a/Engine/Source/Runtime/Engine/Private/UserDefinedStruct.cpp
+++ b/Engine/Source/Runtime/Engine/Private/UserDefinedStruct.cpp
@@ -285,25 +285,27 @@ void UUserDefinedStruct::InitializeStructIgnoreDefaults(void* Dest, int32 ArrayD
 
 void UUserDefinedStruct::InitializeStruct(void* Dest, int32 ArrayDim) const
 {
-	InitializeStructIgnoreDefaults(Dest, ArrayDim);
+	check(Dest);
 
-	if (Dest)
+	const uint8* DefaultInstance = GetDefaultInstance();
+	if ((StructFlags & STRUCT_IsPlainOldData) == 0 || !DefaultInstance)
 	{
-		const uint8* DefaultInstance = GetDefaultInstance();
-		if (DefaultInstance)
-		{
-			int32 Stride = GetStructureSize();
+		InitializeStructIgnoreDefaults(Dest, ArrayDim);
+	}
 
-			for (int32 ArrayIndex = 0; ArrayIndex < ArrayDim; ArrayIndex++)
-			{
-				void* DestStruct = (uint8*)Dest + (Stride * ArrayIndex);
-				CopyScriptStruct(DestStruct, DefaultInstance);
+	if (DefaultInstance)
+	{
+		int32 Stride = GetStructureSize();
 
-				// When copying into another struct we need to register this raw struct pointer so any deferred dependencies will get fixed later
-				FScopedPlaceholderRawContainerTracker TrackStruct(DestStruct);
-				FBlueprintSupport::RegisterDeferredDependenciesInStruct(this, DestStruct);
-			}	
-		}
+		for (int32 ArrayIndex = 0; ArrayIndex < ArrayDim; ArrayIndex++)
+		{
+			void* DestStruct = (uint8*)Dest + (Stride * ArrayIndex);
+			CopyScriptStruct(DestStruct, DefaultInstance);
+
+			// When copying into another struct we need to register this raw struct pointer so any deferred dependencies will get fixed later
+			FScopedPlaceholderRawContainerTracker TrackStruct(DestStruct);
+			FBlueprintSupport::RegisterDeferredDependenciesInStruct(this, DestStruct);
+		}	
 	}
 }
 
