[Backout] - CL26275272
Original CL Desc
-----------------------------------------------------------------
Removing URig from animation system.

diff --git a/Engine/Plugins/Editor/EngineAssetDefinitions/Source/Private/AssetDefinition_Rig.h b/Engine/Plugins/Editor/EngineAssetDefinitions/Source/Private/AssetDefinition_Rig.h
new file mode 100644
--- /dev/null
+++ b/Engine/Plugins/Editor/EngineAssetDefinitions/Source/Private/AssetDefinition_Rig.h
@@ -0,0 +1,28 @@
+// Copyright Epic Games, Inc. All Rights Reserved.
+
+#pragma once
+
+#include "Animation/Rig.h"
+#include "AssetDefinitionDefault.h"
+
+#include "AssetDefinition_Rig.generated.h"
+
+UCLASS()
+class UAssetDefinition_Rig : public UAssetDefinitionDefault
+{
+	GENERATED_BODY()
+
+public:
+	// UAssetDefinition Begin
+	virtual FText GetAssetDisplayName() const override { return NSLOCTEXT("AssetTypeActions", "AssetTypeActions_Rig", "Rig"); }
+	virtual FLinearColor GetAssetColor() const override { return FLinearColor(FColor(255,201,14)); }
+	PRAGMA_DISABLE_DEPRECATION_WARNINGS
+	virtual TSoftClassPtr<UObject> GetAssetClass() const override { return URig::StaticClass(); }
+	PRAGMA_ENABLE_DEPRECATION_WARNINGS
+	virtual TConstArrayView<FAssetCategoryPath> GetAssetCategories() const override
+	{
+		static const auto Categories = { EAssetCategoryPaths::Animation };
+		return Categories;
+	}
+	// UAssetDefinition End
+};
diff --git a/Engine/Plugins/Experimental/Dataflow/Source/DataflowNodes/Private/Dataflow/DataflowSkeletalMeshNodes.cpp b/Engine/Plugins/Experimental/Dataflow/Source/DataflowNodes/Private/Dataflow/DataflowSkeletalMeshNodes.cpp
--- a/Engine/Plugins/Experimental/Dataflow/Source/DataflowNodes/Private/Dataflow/DataflowSkeletalMeshNodes.cpp
+++ b/Engine/Plugins/Experimental/Dataflow/Source/DataflowNodes/Private/Dataflow/DataflowSkeletalMeshNodes.cpp
@@ -7,6 +7,7 @@
 #include "Logging/LogMacros.h"
 #include "UObject/UnrealTypePrivate.h"
 #include "Animation/Skeleton.h"
+#include "Animation/Rig.h"
 #include "Engine/SkeletalMesh.h"
 
 #include UE_INLINE_GENERATED_CPP_BY_NAME(DataflowSkeletalMeshNodes)
diff --git a/Engine/Plugins/Runtime/RigVM/Source/RigVMDeveloper/Private/RigVMCompiler/RigVMCodeGenerator.cpp b/Engine/Plugins/Runtime/RigVM/Source/RigVMDeveloper/Private/RigVMCompiler/RigVMCodeGenerator.cpp
--- a/Engine/Plugins/Runtime/RigVM/Source/RigVMDeveloper/Private/RigVMCompiler/RigVMCodeGenerator.cpp
+++ b/Engine/Plugins/Runtime/RigVM/Source/RigVMDeveloper/Private/RigVMCompiler/RigVMCodeGenerator.cpp
@@ -5,6 +5,7 @@
 #include "RigVMCore/RigVMStruct.h"
 #include "RigVMDeveloperModule.h"
 #include "Algo/Count.h"
+#include "Animation/Rig.h"
 #include "RigVMModel/Nodes/RigVMDispatchNode.h"
 #include "RigVMStringUtils.h"
 
diff --git a/Engine/Plugins/Runtime/RigVM/Source/RigVMDeveloper/Private/RigVMModel/Nodes/RigVMUnitNode.cpp b/Engine/Plugins/Runtime/RigVM/Source/RigVMDeveloper/Private/RigVMModel/Nodes/RigVMUnitNode.cpp
--- a/Engine/Plugins/Runtime/RigVM/Source/RigVMDeveloper/Private/RigVMModel/Nodes/RigVMUnitNode.cpp
+++ b/Engine/Plugins/Runtime/RigVM/Source/RigVMDeveloper/Private/RigVMModel/Nodes/RigVMUnitNode.cpp
@@ -2,6 +2,7 @@
 
 #include "RigVMModel/Nodes/RigVMUnitNode.h"
 
+#include "Animation/Rig.h"
 #include "RigVMCore/RigVMStruct.h"
 #include "RigVMModel/RigVMControllerActions.h"
 #include "RigVMUserWorkflowRegistry.h"
diff --git a/Engine/Source/Editor/DetailCustomizations/Private/DetailCustomizations.cpp b/Engine/Source/Editor/DetailCustomizations/Private/DetailCustomizations.cpp
--- a/Engine/Source/Editor/DetailCustomizations/Private/DetailCustomizations.cpp
+++ b/Engine/Source/Editor/DetailCustomizations/Private/DetailCustomizations.cpp
@@ -126,6 +126,7 @@
 #include "RawDistributionVectorStructCustomization.h"
 #include "ReflectionCaptureDetails.h"
 #include "RenderPassesCustomization.h"
+#include "RigDetails.h"
 #include "RotatorStructCustomization.h"
 #include "SceneCaptureDetails.h"
 #include "SceneComponentDetails.h"
@@ -399,6 +400,7 @@ void FDetailCustomizationsModule::RegisterObjectCustomizations()
 	RegisterCustomClassLayout("K2Node_GetClassDefaults", FOnGetDetailCustomizationInstance::CreateStatic(&FSkeletalControlNodeDetails::MakeInstance));
 
 	RegisterCustomClassLayout("AnimSequence", FOnGetDetailCustomizationInstance::CreateStatic(&FAnimSequenceDetails::MakeInstance));
+	RegisterCustomClassLayout("Rig", FOnGetDetailCustomizationInstance::CreateStatic(&FRigDetails::MakeInstance));
 
 	RegisterCustomClassLayout("EditorAnimSegment", FOnGetDetailCustomizationInstance::CreateStatic(&FAnimMontageSegmentDetails::MakeInstance));
 	RegisterCustomClassLayout("EditorAnimCompositeSegment", FOnGetDetailCustomizationInstance::CreateStatic(&FAnimMontageSegmentDetails::MakeInstance));
diff --git a/Engine/Source/Editor/DetailCustomizations/Private/RigDetails.cpp b/Engine/Source/Editor/DetailCustomizations/Private/RigDetails.cpp
new file mode 100644
--- /dev/null
+++ b/Engine/Source/Editor/DetailCustomizations/Private/RigDetails.cpp
@@ -0,0 +1,564 @@
+// Copyright Epic Games, Inc. All Rights Reserved.
+
+#include "RigDetails.h"
+
+#include "Animation/Rig.h"
+#include "DetailCategoryBuilder.h"
+#include "DetailLayoutBuilder.h"
+#include "DetailWidgetRow.h"
+#include "Fonts/SlateFontInfo.h"
+#include "Framework/SlateDelegates.h"
+#include "HAL/PlatformCrt.h"
+#include "IDetailChildrenBuilder.h"
+#include "Internationalization/Internationalization.h"
+#include "Misc/AssertionMacros.h"
+#include "Misc/Attribute.h"
+#include "PropertyCustomizationHelpers.h"
+#include "PropertyEditorModule.h"
+#include "PropertyHandle.h"
+#include "ScopedTransaction.h"
+#include "Serialization/Archive.h"
+#include "SlotBase.h"
+#include "Templates/Casts.h"
+#include "Templates/UnrealTemplate.h"
+#include "Types/SlateStructs.h"
+#include "UObject/NameTypes.h"
+#include "UObject/Object.h"
+#include "Widgets/DeclarativeSyntaxSupport.h"
+#include "Widgets/Input/SButton.h"
+#include "Widgets/Input/SCheckBox.h"
+#include "Widgets/Input/SEditableTextBox.h"
+#include "Widgets/Layout/SBox.h"
+#include "Widgets/SBoxPanel.h"
+#include "Widgets/Text/STextBlock.h"
+
+class SWidget;
+
+#define LOCTEXT_NAMESPACE	"RigDetails"
+
+// Table Headers for node list
+#define NODE_TABLE_DISPLAYNAME	TEXT("DisplayName")
+#define NODE_TABLE_NODENAME		TEXT("NodeName")
+#define NODE_TABLE_PARENTNAME	TEXT("ParentName")
+
+// January 2023 - This whole compilation unit is deprecated to be removed in 5.4
+PRAGMA_DISABLE_DEPRECATION_WARNINGS
+TSharedRef<IDetailCustomization> FRigDetails::MakeInstance()
+{
+	return MakeShareable(new FRigDetails);
+}
+
+void FRigDetails::CustomizeDetails(IDetailLayoutBuilder& DetailBuilder)
+{
+	DetailBuilderPtr = &DetailBuilder;
+	TArray<TWeakObjectPtr<UObject>> Objects;
+
+	DetailBuilder.GetObjectsBeingCustomized(Objects);
+
+	// if more than one, do not customize?
+	if (Objects.Num() > 1)
+	{
+		return;
+	}
+
+	ItemBeingEdited = Objects[0];
+
+	IDetailCategoryBuilder& NodeCategory = DetailBuilder.EditCategory("Node");
+	IDetailCategoryBuilder& TransformBaseCategory = DetailBuilder.EditCategory("Constraint Setup");
+
+	TransformBasesPropertyHandle = DetailBuilder.GetProperty("TransformBases");
+	NodesPropertyHandle = DetailBuilder.GetProperty("Nodes");
+
+
+	TSharedRef<FDetailArrayBuilder> NodeArrayBuilder = MakeShareable(new FDetailArrayBuilder(NodesPropertyHandle.ToSharedRef()));
+	NodeArrayBuilder->OnGenerateArrayElementWidget(FOnGenerateArrayElementWidget::CreateSP(this, &FRigDetails::GenerateNodeArrayElementWidget, &DetailBuilder));
+
+	NodeCategory.AddCustomBuilder( NodeArrayBuilder, false );
+
+	TSharedRef<FDetailArrayBuilder> TransformBaseArrayBuilder = MakeShareable(new FDetailArrayBuilder(TransformBasesPropertyHandle.ToSharedRef()));
+	TransformBaseArrayBuilder->OnGenerateArrayElementWidget(FOnGenerateArrayElementWidget::CreateSP(this, &FRigDetails::GenerateTransformBaseArrayElementWidget, &DetailBuilder));
+
+	// add custom menu
+	// -> set all to world
+	// -> set all to default parent
+	TransformBaseCategory.AddCustomRow(FText::GetEmpty())
+	[
+		// two button 1. view 2. save to base pose
+		SNew(SHorizontalBox)
+
+		+SHorizontalBox::Slot()
+		.AutoWidth()
+		.Padding(5)
+		.HAlign(HAlign_Center)
+		[
+			SNew(SButton)
+			.OnClicked(FOnClicked::CreateSP(this, &FRigDetails::OnSetAllToWorld))
+			.HAlign(HAlign_Center)
+			.VAlign(VAlign_Center)
+			.Text(LOCTEXT("SetAllToWorld_ButtonLabel", "Set All Constraints to World"))
+		]
+
+		+SHorizontalBox::Slot()
+		.AutoWidth()
+		.Padding(5)
+		.HAlign(HAlign_Center)
+		[
+			SNew(SButton)
+			.OnClicked(FOnClicked::CreateSP(this, &FRigDetails::OnSetAllToParent))
+			.HAlign(HAlign_Center)
+			.VAlign(VAlign_Center)
+			.Text(LOCTEXT("SetAllToParent_ButtonLabel", "Set All Constraints to Parent"))
+		]
+	];
+
+	TransformBaseCategory.AddCustomBuilder( TransformBaseArrayBuilder, false );
+}
+
+void FRigDetails::GenerateNodeArrayElementWidget(TSharedRef<IPropertyHandle> PropertyHandle, int32 ArrayIndex, IDetailChildrenBuilder& ChildrenBuilder, IDetailLayoutBuilder* DetailLayout)
+{
+	TSharedRef<IPropertyHandle> DisplayNameProp = PropertyHandle->GetChildHandle(GET_MEMBER_NAME_CHECKED(FNode, DisplayName)).ToSharedRef();
+	TSharedRef<IPropertyHandle> NodeNameProp = PropertyHandle->GetChildHandle(GET_MEMBER_NAME_CHECKED(FNode, Name)).ToSharedRef();
+	TSharedRef<IPropertyHandle> ParentNameProp = PropertyHandle->GetChildHandle(GET_MEMBER_NAME_CHECKED(FNode, ParentName)).ToSharedRef();
+	TSharedRef<IPropertyHandle> AdvancedProp = PropertyHandle->GetChildHandle(GET_MEMBER_NAME_CHECKED(FNode, bAdvanced)).ToSharedRef();
+
+	TSharedPtr<SEditableTextBox> DisplayTextBox;
+
+	// the interface will be node [display name] [parent node]
+	// delegate for display name
+	FText NodeName, ParentNodeName, DisplayString;
+	check (NodeNameProp->GetValueAsDisplayText(NodeName) != FPropertyAccess::Fail);
+	check (ParentNameProp->GetValueAsDisplayText(ParentNodeName) != FPropertyAccess::Fail);
+	check (DisplayNameProp->GetValueAsDisplayText(DisplayString) != FPropertyAccess::Fail);
+
+	ChildrenBuilder.AddCustomRow(FText::GetEmpty())
+	[
+		SNew(SHorizontalBox)
+
+		+SHorizontalBox::Slot()
+		.Padding(5, 2)
+		.AutoWidth()
+		[
+			SNew(SBox)
+			.WidthOverride(100)
+			.Content()
+			[
+				SNew(STextBlock)
+				.Text(NodeName)
+				.Font(DetailLayout->GetDetailFontBold())
+			]
+		]
+
+		+SHorizontalBox::Slot()
+		.Padding(5, 2)
+		.AutoWidth()
+		[
+			SNew(SBox)
+			.WidthOverride(150)
+			.Content()
+			[
+				SNew(STextBlock)
+				.Text(FText::Format(LOCTEXT("ParentNameFmt", " [Parent : {0}] "), ParentNodeName))
+				.Font(DetailLayout->GetDetailFont())
+			]
+		]
+
+		+SHorizontalBox::Slot()
+		.Padding(5, 2)
+		.AutoWidth()
+		[
+			SNew(STextBlock)
+			.Text(LOCTEXT("DisplayNameLabel", "Display Name"))
+			.Font(DetailLayout->GetDetailFontBold())
+		]
+
+		+SHorizontalBox::Slot()
+		.Padding(5, 2)
+		.AutoWidth()
+		.HAlign(HAlign_Left)
+		[
+			SNew(SBox)
+			.WidthOverride(250)
+			.Content()
+			[
+				SAssignNew(DisplayTextBox, SEditableTextBox)
+				.Text( TAttribute<FText>::Create( TAttribute<FText>::FGetter::CreateSP(this, &FRigDetails::GetDisplayName, DisplayNameProp) ) )
+				.Font(DetailLayout->GetDetailFont())
+				.OnTextChanged(this, &FRigDetails::OnDisplayNameChanged, DisplayNameProp, ArrayIndex)
+				.OnTextCommitted(this, &FRigDetails::OnDisplayNameCommitted, DisplayNameProp, ArrayIndex)
+				.MinDesiredWidth(200)
+			]
+		]
+
+		+SHorizontalBox::Slot()
+		.Padding(5, 2)
+		.AutoWidth()
+		[
+			SNew(STextBlock)
+			.Text(LOCTEXT("AdvancedLabel", "Advanced"))
+			.Font(DetailLayout->GetDetailFontBold())
+		]
+
+		+SHorizontalBox::Slot()
+		.Padding(5, 2)
+		.FillWidth(1)
+		.HAlign(HAlign_Left)
+		[
+			SNew(SBox)
+			.WidthOverride(250)
+			.Content()
+			[
+				SNew(SCheckBox)
+				.IsChecked(this, &FRigDetails::AdvancedCheckBoxIsChecked, AdvancedProp)
+				.OnCheckStateChanged(this, &FRigDetails::OnAdvancedCheckBoxStateChanged, AdvancedProp)
+			]
+		]
+	];
+
+	// we're adding on demand because if array expands, we'll have to handle it dynamically
+	if (!DisplayNameTextBoxes.IsValidIndex(ArrayIndex))
+	{
+		DisplayNameTextBoxes.AddZeroed(ArrayIndex-DisplayNameTextBoxes.Num()+1);
+	}
+
+	DisplayNameTextBoxes[ArrayIndex] = DisplayTextBox;
+};
+
+FText FRigDetails::GetDisplayName(TSharedRef<IPropertyHandle> DisplayNameProp) const
+{
+	FText DisplayText;
+	
+	check (FPropertyAccess::Success == DisplayNameProp->GetValueAsDisplayText(DisplayText));
+
+	return DisplayText;
+}
+
+void FRigDetails::ValidErrorMessage(const FString& DisplayString, int32 ArrayIndex)
+{
+	if(DisplayNameTextBoxes.IsValidIndex(ArrayIndex))
+	{
+		DisplayNameTextBoxes[ArrayIndex]->SetError(TEXT(""));
+
+		if(DisplayString.Len() == 0)
+		{
+			DisplayNameTextBoxes[ArrayIndex]->SetError(TEXT("Name can't be empty"));
+		}
+		else
+		{
+			// verify if this name is unique
+			URig * Rig = Cast<URig>(ItemBeingEdited.Get());
+			if(Rig)
+			{
+				FString NewText = DisplayString;
+				// make sure that name is unique
+				const TArray<FNode> & Nodes = Rig->GetNodes();
+				int32 NodeIndex = 0;
+				for(auto Node : Nodes)
+				{
+					if(NodeIndex++ != ArrayIndex && Node.DisplayName == NewText)
+					{
+						DisplayNameTextBoxes[ArrayIndex]->SetError(TEXT("Name should be unique."));
+					}
+				}
+			}
+		}
+	}
+}
+
+void FRigDetails::OnDisplayNameChanged(const FText& Text, TSharedRef<IPropertyHandle> DisplayNameProp, int32 ArrayIndex)
+{
+	// still set it since you don't know what they come up with 
+	DisplayNameProp->SetValueFromFormattedString(Text.ToString());
+	ValidErrorMessage(Text.ToString(), ArrayIndex);
+}
+
+void FRigDetails::OnDisplayNameCommitted(const FText& Text, ETextCommit::Type CommitType, TSharedRef<IPropertyHandle> DisplayNameProp, int32 ArrayIndex)
+{
+	// @todo error check here? I basically need mirror string to avoid the issue
+	DisplayNameProp->SetValueFromFormattedString(Text.ToString());
+}
+
+void FRigDetails::GenerateTransformBaseArrayElementWidget(TSharedRef<IPropertyHandle> PropertyHandle, int32 ArrayIndex, IDetailChildrenBuilder& ChildrenBuilder, IDetailLayoutBuilder* DetailLayout)
+{
+	TSharedRef<IPropertyHandle> NodeNameProp = PropertyHandle->GetChildHandle(GET_MEMBER_NAME_CHECKED(FTransformBase, Node)).ToSharedRef();
+	TSharedPtr<IPropertyHandleArray> ConstraintsProp = PropertyHandle->GetChildHandle(GET_MEMBER_NAME_CHECKED(FTransformBase, Constraints))->AsArray();
+
+	// translation
+	TSharedPtr<IPropertyHandleArray> TransformConstraintsProp_T = ConstraintsProp->GetElement(0)->GetChildHandle(GET_MEMBER_NAME_CHECKED(FTransformBaseConstraint, TransformConstraints))->AsArray();
+	TSharedRef<IPropertyHandle> ParentNameProp_T = TransformConstraintsProp_T->GetElement(0)->GetChildHandle(GET_MEMBER_NAME_CHECKED(FRigTransformConstraint, ParentSpace)).ToSharedRef();
+
+	// orientation
+	TSharedPtr<IPropertyHandleArray> TransformConstraintsProp_R = ConstraintsProp->GetElement(1)->GetChildHandle(GET_MEMBER_NAME_CHECKED(FTransformBaseConstraint, TransformConstraints))->AsArray();
+	TSharedRef<IPropertyHandle> ParentNameProp_R = TransformConstraintsProp_R->GetElement(0)->GetChildHandle(GET_MEMBER_NAME_CHECKED(FRigTransformConstraint, ParentSpace)).ToSharedRef();
+	
+	// the interface will be node [display name] [parent node]
+	// delegate for display name
+	FText NodeName;
+	FString ParentNodeName_T, ParentNodeName_R;
+	check(NodeNameProp->GetValueAsDisplayText(NodeName) != FPropertyAccess::Fail);
+	check(ParentNameProp_T->GetValueAsDisplayString(ParentNodeName_T) != FPropertyAccess::Fail);
+	check(ParentNameProp_R->GetValueAsDisplayString(ParentNodeName_R) != FPropertyAccess::Fail);
+
+	// we add them on demand since the array can expand or reduce
+	if(!ParentSpaceOptionList.IsValidIndex(ArrayIndex))
+	{
+		int32 NumToAdd = ArrayIndex-ParentSpaceOptionList.Num()+1;
+		for(int32 Idx = 0 ; Idx < NumToAdd ; ++Idx)
+		{
+			ParentSpaceOptionList.Add(new TArray<TSharedPtr<FString>>);
+		}
+	}
+
+	// ParentSpaceComboBoxees creates 2 per element - one for translation, and one for rotation
+	if(!ParentSpaceComboBoxes.IsValidIndex(ArrayIndex*2))
+	{
+		int32 NumNewItem = ArrayIndex-(ParentSpaceComboBoxes.Num()/2)+1;
+		ParentSpaceComboBoxes.AddZeroed(NumNewItem*2);
+	}
+
+	// create string list for picking parent node
+	// make sure you don't include itself and find what is curretn selected item
+	TArray<TSharedPtr<FString>>& ParentNodeOptions = ParentSpaceOptionList[ArrayIndex];
+
+	ParentNodeOptions.Empty();
+	ParentNodeOptions.Add(MakeShareable(new FString(URig::WorldNodeName.ToString())));
+	URig * Rig = Cast<URig>(ItemBeingEdited.Get());
+	check (Rig);
+	const TArray<FNode> & Nodes = Rig->GetNodes();
+	if ( Nodes.Num() <= 0 )
+	{
+		return;
+	}
+
+	int32 NodeIndex = 0, ParentIndex_T=0, ParentIndex_R=0;
+	const FNode& CurNode = Nodes[ArrayIndex];
+	for(auto Node : Nodes)
+	{
+		if (NodeIndex != ArrayIndex)
+		{
+			ParentNodeOptions.Add(MakeShareable(new FString(Node.Name.ToString())));
+
+			if (Node.Name.ToString() == ParentNodeName_T)
+			{
+				ParentIndex_T = ParentNodeOptions.Num()-1;
+			}
+
+			if(Node.Name.ToString() == ParentNodeName_R)
+			{
+				ParentIndex_R = ParentNodeOptions.Num()-1;
+			}
+		}
+
+		NodeIndex++;
+	}
+
+	ChildrenBuilder.AddCustomRow(FText::GetEmpty())
+	[
+		SNew(SHorizontalBox)
+
+		+SHorizontalBox::Slot()
+		.Padding(5, 2)
+		.AutoWidth()
+		[
+			SNew(SBox)
+			.WidthOverride(100)
+			.Content()
+			[
+				SNew(STextBlock)
+				.Text(NodeName)
+				.Font(DetailLayout->GetDetailFontBold())
+			]
+		]
+
+		+SHorizontalBox::Slot()
+		.Padding(5, 2)
+		.AutoWidth()
+		[
+			SNew(SVerticalBox)
+
+			+SVerticalBox::Slot()
+			.Padding(2)
+			[
+				SNew(SHorizontalBox)
+
+				+SHorizontalBox::Slot()
+				.Padding(2)
+				.AutoWidth()
+				[
+					SNew(SBox)
+					.WidthOverride(100)
+					.Content()
+					[
+						SNew(STextBlock)
+						.Text(LOCTEXT("TranslationLabel", "Translation"))
+						.Font(DetailLayout->GetDetailFontBold())
+					]
+				]
+
+				+SHorizontalBox::Slot()
+				.Padding(2)
+				.FillWidth(1)
+				[
+					SNew(SBox)
+					.WidthOverride(250)
+					.Content()
+					[
+
+						SAssignNew(ParentSpaceComboBoxes[ArrayIndex*2], SComboBox< TSharedPtr<FString> >)
+						.OptionsSource(&ParentNodeOptions)
+						.InitiallySelectedItem(ParentNodeOptions[ParentIndex_T])
+						.OnSelectionChanged(this, &FRigDetails::OnParentSpaceSelectionChanged, ParentNameProp_T)
+						.OnGenerateWidget(this, &FRigDetails::MakeItemWidget)
+						.OnComboBoxOpening(this, &FRigDetails::OnComboBoxOopening, ParentNameProp_T, ArrayIndex, true)
+						.HasDownArrow(true)
+						[
+							SNew(STextBlock)
+							.Text(this, &FRigDetails::GetSelectedTextLabel, ParentNameProp_T)
+						]
+					]
+				]
+			]
+
+			+SVerticalBox::Slot()
+			.Padding(2)
+			[
+				SNew(SHorizontalBox)
+
+				+SHorizontalBox::Slot()
+				.Padding(2)
+				.AutoWidth()
+				[
+					SNew(SBox)
+					.WidthOverride(100)
+					.Content()
+					[
+						SNew(STextBlock)
+						.Text(LOCTEXT("OrientationLabel", "Orientation"))
+						.Font(DetailLayout->GetDetailFontBold())
+					]
+				]
+
+				+SHorizontalBox::Slot()
+				.Padding(2)
+				.FillWidth(1)
+				[
+					SNew(SBox)
+					.WidthOverride(250)
+					.Content()
+					[
+						SAssignNew(ParentSpaceComboBoxes[ArrayIndex*2+1], SComboBox< TSharedPtr<FString> >)
+						.OptionsSource(&ParentNodeOptions)
+						.InitiallySelectedItem(ParentNodeOptions[ParentIndex_R])
+						.OnSelectionChanged(this, &FRigDetails::OnParentSpaceSelectionChanged, ParentNameProp_R)
+						.OnGenerateWidget(this, &FRigDetails::MakeItemWidget)
+						.OnComboBoxOpening(this, &FRigDetails::OnComboBoxOopening, ParentNameProp_R, ArrayIndex, true)
+						.HasDownArrow(true)
+						[
+							SNew(STextBlock)
+							.Text(this, &FRigDetails::GetSelectedTextLabel, ParentNameProp_R)
+						]
+					]
+				]
+			]
+		]
+	];
+}
+
+void FRigDetails::OnParentSpaceSelectionChanged(TSharedPtr<FString> SelectedItem, ESelectInfo::Type SelectInfo, TSharedRef<IPropertyHandle> ParentSpacePropertyHandle)
+{
+	if (SelectInfo == ESelectInfo::OnKeyPress || SelectInfo == ESelectInfo::OnMouseClick)
+	{
+		if(SelectedItem.IsValid())
+		{
+			ParentSpacePropertyHandle->SetValueFromFormattedString(*SelectedItem.Get());
+		}
+	}
+}
+
+FReply FRigDetails::OnSetAllToWorld()
+{
+	URig * Rig = Cast<URig>(ItemBeingEdited.Get());
+	check(Rig);
+
+	const FScopedTransaction Transaction(LOCTEXT("SetAllToWorld_Action", "Set All Transform Constraints to World"));
+	Rig->Modify();
+	Rig->SetAllConstraintsToWorld();
+	DetailBuilderPtr->ForceRefreshDetails();
+
+	return FReply::Handled();
+}
+
+FReply FRigDetails::OnSetAllToParent()
+{
+	URig * Rig = Cast<URig>(ItemBeingEdited.Get());
+	check(Rig);
+
+	const FScopedTransaction Transaction(LOCTEXT("SetAllToParent_Action", "Set All Transform Constraints to Parent"));
+	Rig->Modify();
+	Rig->SetAllConstraintsToParents();
+	DetailBuilderPtr->ForceRefreshDetails();
+
+	return FReply::Handled();
+}
+
+/** Called to create a widget for each string */
+TSharedRef<SWidget> FRigDetails::MakeItemWidget(TSharedPtr<FString> StringItem)
+{
+	check(StringItem.IsValid());
+
+	return SNew(STextBlock)
+		.Text(FText::FromString(*StringItem.Get()));
+}
+/** Helper method to get the text for a given item in the combo box */
+FText FRigDetails::GetSelectedTextLabel(TSharedRef<IPropertyHandle> ParentSpacePropertyHandle) const
+{
+	FString DisplayText;
+
+	if (ParentSpacePropertyHandle->GetValueAsDisplayString(DisplayText) != FPropertyAccess::Fail)
+	{
+		return FText::FromString(DisplayText);
+	}
+
+	return LOCTEXT("Unknown", "Unknown");
+}
+
+void FRigDetails::OnComboBoxOopening(TSharedRef<IPropertyHandle> ParentSpacePropertyHandle, int32 ArrayIndex, bool bTranslation)
+{
+	FString PropertyValue = GetSelectedTextLabel(ParentSpacePropertyHandle).ToString();
+
+	// now find exact data
+	TArray<TSharedPtr<FString>> & ParentOptions = ParentSpaceOptionList[ArrayIndex];
+	TSharedPtr<FString> SelectedItem;
+	for (auto Option : ParentOptions)
+	{
+		if (*Option.Get() == PropertyValue)
+		{
+			SelectedItem = Option;
+			break;
+		}
+	}
+
+	int32 ComboBoxIndex = (bTranslation)? ArrayIndex*2 : ArrayIndex*2+1;
+	TSharedPtr< SComboBox<TSharedPtr<FString>> > ComboBox = ParentSpaceComboBoxes[ComboBoxIndex];
+
+	ComboBox->SetSelectedItem(SelectedItem);
+}
+
+void FRigDetails::OnAdvancedCheckBoxStateChanged(ECheckBoxState NewState, TSharedRef<IPropertyHandle> PropertyHandle)
+{
+	bool bValue = (NewState == ECheckBoxState::Checked)? true : false;
+	PropertyHandle->SetValue(bValue);
+}
+
+ECheckBoxState FRigDetails::AdvancedCheckBoxIsChecked(TSharedRef<IPropertyHandle> PropertyHandle) const
+{
+	bool bValue = false;
+	// multi value doesn't work in array, so i'm not handling multi value
+	if (PropertyHandle->GetValue(bValue) != FPropertyAccess::Fail)
+	{
+		return (bValue)? ECheckBoxState::Checked: ECheckBoxState::Unchecked;
+	}
+
+	return ECheckBoxState::Undetermined;
+}
+#undef LOCTEXT_NAMESPACE
+PRAGMA_ENABLE_DEPRECATION_WARNINGS
\ No newline at end of file
diff --git a/Engine/Source/Editor/DetailCustomizations/Private/RigDetails.h b/Engine/Source/Editor/DetailCustomizations/Private/RigDetails.h
new file mode 100644
--- /dev/null
+++ b/Engine/Source/Editor/DetailCustomizations/Private/RigDetails.h
@@ -0,0 +1,82 @@
+// Copyright Epic Games, Inc. All Rights Reserved.
+
+#pragma once
+
+#include "Containers/Array.h"
+#include "Containers/BitArray.h"
+#include "Containers/IndirectArray.h"
+#include "Containers/Set.h"
+#include "Containers/SparseArray.h"
+#include "Containers/UnrealString.h"
+#include "Delegates/Delegate.h"
+#include "HAL/Platform.h"
+#include "IDetailCustomization.h"
+#include "Input/Reply.h"
+#include "Internationalization/Text.h"
+#include "Misc/Optional.h"
+#include "Styling/SlateTypes.h"
+#include "Templates/SharedPointer.h"
+#include "Templates/TypeHash.h"
+#include "Types/SlateEnums.h"
+#include "UObject/WeakObjectPtr.h"
+#include "UObject/WeakObjectPtrTemplates.h"
+#include "Widgets/Input/SComboBox.h"
+
+class IDetailChildrenBuilder;
+class IDetailLayoutBuilder;
+class IPropertyHandle;
+class SEditableTextBox;
+class SWidget;
+class UObject;
+
+class FRigDetails : public IDetailCustomization
+{
+public:
+	/** Makes a new instance of this detail layout class for a specific detail view requesting it */
+	static TSharedRef<IDetailCustomization> MakeInstance();
+
+	/** IDetailCustomization interface */
+	virtual void CustomizeDetails( IDetailLayoutBuilder& DetailBuilder ) override;
+
+private:
+
+	/** Generate a widget for a movie array element */
+	void GenerateNodeArrayElementWidget(TSharedRef<IPropertyHandle> PropertyHandle, int32 ArrayIndex, IDetailChildrenBuilder& ChildrenBuilder, IDetailLayoutBuilder* DetailLayout);
+	void GenerateTransformBaseArrayElementWidget(TSharedRef<IPropertyHandle> PropertyHandle, int32 ArrayIndex, IDetailChildrenBuilder& ChildrenBuilder, IDetailLayoutBuilder* DetailLayout);
+
+	TSharedPtr<IPropertyHandle> NodesPropertyHandle;
+	TSharedPtr<IPropertyHandle> TransformBasesPropertyHandle;
+
+	// array custom boxes - these will stay around as long as this window is up
+	TArray<TSharedPtr<SEditableTextBox>>	DisplayNameTextBoxes;
+	TIndirectArray<TArray<TSharedPtr<FString>>>		ParentSpaceOptionList;
+	TArray<TSharedPtr<SComboBox< TSharedPtr<FString> >>>	ParentSpaceComboBoxes;
+
+	/** we only support one item */
+	TWeakObjectPtr<UObject> ItemBeingEdited;
+
+	// node display text handler
+	void ValidErrorMessage(const FString& DisplayString, int32 ArrayIndex);
+	FText GetDisplayName(TSharedRef<IPropertyHandle> DisplayNameProp) const;
+	void OnDisplayNameChanged(const FText& Text, TSharedRef<IPropertyHandle> DisplayNameProp, int32 ArrayIndex);
+	void OnDisplayNameCommitted(const FText& Text, ETextCommit::Type CommitType, TSharedRef<IPropertyHandle> DisplayNameProp, int32 ArrayIndex);
+
+	// combo box handler
+	void OnParentSpaceSelectionChanged(TSharedPtr<FString> SelectedItem, ESelectInfo::Type SelectInfo, TSharedRef<IPropertyHandle> ParentSpacePropertyHandle);
+	/** Called to create a widget for each string */
+	TSharedRef<SWidget> MakeItemWidget(TSharedPtr<FString> StringItem);
+	// Helper method to get the text for a given item in the combo box 
+	FText GetSelectedTextLabel(TSharedRef<IPropertyHandle> ParentSpacePropertyHandle) const;
+	void OnComboBoxOopening(TSharedRef<IPropertyHandle> ParentSpacePropertyHandle, int32 ArrayIndex, bool bTranslation);
+	// check box handler
+	// Callback for changing this row's Share check box state.
+	void OnAdvancedCheckBoxStateChanged(ECheckBoxState NewState, TSharedRef<IPropertyHandle> PropertyHandle);
+	ECheckBoxState AdvancedCheckBoxIsChecked(TSharedRef<IPropertyHandle> PropertyHandle) const;
+
+	// button handlers
+	FReply OnSetAllToWorld();
+	FReply OnSetAllToParent();
+
+	IDetailLayoutBuilder * DetailBuilderPtr;
+};
+
diff --git a/Engine/Source/Editor/Persona/Private/BoneMappingHelper.cpp b/Engine/Source/Editor/Persona/Private/BoneMappingHelper.cpp
new file mode 100644
--- /dev/null
+++ b/Engine/Source/Editor/Persona/Private/BoneMappingHelper.cpp
@@ -0,0 +1,318 @@
+// Copyright Epic Games, Inc. All Rights Reserved.
+
+
+#include "BoneMappingHelper.h"
+#include "AnimationRuntime.h"
+
+//////////////////////////////////////////////////////////////////////////
+// FBoneDescription
+//////////////////////////////////////////////////////////////////////////
+
+int32 FBoneDescription::GetBestIndex() const
+{
+	int32 BestIndex = INDEX_NONE;
+	float BestScore = 0.f;
+
+	for (int32 Index = 0; Index < Scores.Num(); ++Index)
+	{
+		if (Scores[Index] > BestScore)
+		{
+			BestScore = Scores[Index];
+			BestIndex = Index;
+		}
+	}
+
+	return BestIndex;
+}
+
+float FBoneDescription::CalculateNameScore(const FName& Name1, const FName& Name2) const
+{
+	const FString String1 = Name1.ToString();
+	const FString String2 = Name2.ToString();
+	const int32 Distance = FAnimationRuntime::GetStringDistance(String1, String2);
+	const float MaxLength = static_cast<float>(FMath::Max(String1.Len() , String2.Len()));
+	return (MaxLength) ? (MaxLength - Distance) / (MaxLength) : 1.f;
+}
+
+float FBoneDescription::CalculateScore(const FBoneDescription& Other) const
+{
+	// if they don't have parent, it's root, so just give whole score
+	if (Other.BoneInfo.ParentIndex == INDEX_NONE && BoneInfo.ParentIndex == INDEX_NONE)
+	{
+		return 1.f;
+	}
+
+	// each element will exit from [0, 1], and we'll apply weight to it
+
+	// check direction of facing [-1, 1]
+	float Cosine = static_cast<float>(FVector::DotProduct(DirFromParent, Other.DirFromParent));
+	float Score_DirFromParent = FMath::Max((Cosine-0.5f)*2.f, 0.f); // scale to only care for range of [0.5-1]
+	Score_DirFromParent = FMath::Clamp(Score_DirFromParent, 0.f, 1.f);
+
+	// check direction of facing [-1, 1]
+	Cosine = static_cast<float>(FVector::DotProduct(DirFromRoot, Other.DirFromRoot));
+	float Score_DirFromRoot = FMath::Max((Cosine - 0.5f)*2.f, 0.f); // scale to only care for range of [0.5-1]
+	Score_DirFromRoot = FMath::Clamp(Score_DirFromRoot, 0.f, 1.f);
+
+	// check number of children
+	float Score_NumChildren = 0.f;
+	const int32 MaxNumChildren = FMath::Max(Other.NumChildren, NumChildren);
+	if (MaxNumChildren > 0)
+	{
+		Score_NumChildren = 1.f - static_cast<float>(FMath::Abs(Other.NumChildren - NumChildren) / MaxNumChildren);
+	}
+	else
+	{
+		// no children - leaf node
+		Score_NumChildren = 1.f;
+	}
+
+	Score_NumChildren = FMath::Clamp(Score_NumChildren, 0.f, 1.f);
+
+	// score of ratio from parent  - if you're here, you should have parent, so it shouldn't be 0.f, most likely
+	float Score_RatioFromParent = 0.f;
+	if (Other.RatioFromParent > RatioFromParent)
+	{
+		Score_RatioFromParent = (Other.RatioFromParent > 0.f) ? RatioFromParent / Other.RatioFromParent : 0.f;
+	}
+	else
+	{
+		Score_RatioFromParent = (RatioFromParent > 0.f) ? Other.RatioFromParent / RatioFromParent : 0.f;
+	}
+
+	Score_RatioFromParent = FMath::Clamp(Score_RatioFromParent, 0.f, 1.f);
+
+	// check normalized position - since this is normalized, it should stay within 1
+	const FVector DiffNormalizedPosition = (Other.NormalizedPosition - NormalizedPosition).GetAbs();
+	constexpr float MaxNormalizedPosition = 3.f; /* since 1^2+1^2+1^2 = 3*/
+	float Score_NormalizedPosition = (MaxNormalizedPosition - static_cast<float>(DiffNormalizedPosition.SizeSquared()))/ MaxNormalizedPosition;
+	Score_NormalizedPosition = FMath::Clamp(Score_NormalizedPosition, 0.f, 1.f);
+
+	float Score_NameMatching = CalculateNameScore(BoneInfo.Name, Other.BoneInfo.Name);
+	Score_NameMatching = FMath::Clamp(Score_NameMatching, 0.f, 1.f);
+
+	// now come up with full score
+	static float Weight_DirFromParent = 2.f;
+	static float Weight_NumChildren = 0.5f;
+	static float Weight_NormalizedPosition = 1.0f; // location can be very confusing, so give less weight on this
+	static float Weight_RatioFromParent = 1.f;
+	static float Weight_NameMatching = 2.0f;
+	static float Weight_DirFromRoot = 0.f;
+
+	const float FinalScore = (Score_DirFromParent * Weight_DirFromParent + Score_NumChildren * Weight_NumChildren + Score_NormalizedPosition * Weight_NormalizedPosition + Score_RatioFromParent * Weight_RatioFromParent + Score_NameMatching * Weight_NameMatching + Score_DirFromRoot * Weight_DirFromRoot) / (Weight_DirFromParent + Weight_NumChildren + Weight_NormalizedPosition + Weight_RatioFromParent + Weight_NameMatching + Weight_DirFromRoot);
+	UE_LOG(LogAnimation, Log, TEXT("Calculate Score - [%s] - [%s] (Score_DirFromParent(%0.2f), Score_NumChildren(%0.2f), Score_NormalizedPosition(%0.2f), Score_RatioFromParent(%0.2f), Score_NameMatching(%0.2f), Score_DirFromRoot(%0.2f) )"), *BoneInfo.Name.ToString(), *Other.BoneInfo.Name.ToString(), Score_DirFromParent, Score_NumChildren, Score_NormalizedPosition, Score_RatioFromParent, Score_NameMatching, Score_DirFromRoot);
+	return FinalScore;
+}
+
+//////////////////////////////////////////////////////////////////////////
+// FBoneMappingHelper
+//////////////////////////////////////////////////////////////////////////
+
+FBoneMappingHelper::FBoneMappingHelper(const FReferenceSkeleton& InRefSkeleton1, const FReferenceSkeleton&  InRefSkeleton2)
+{
+	Initialize(0, InRefSkeleton1);
+	Initialize(1, InRefSkeleton2);
+}
+
+void FBoneMappingHelper::Initialize(int32 Index, const FReferenceSkeleton&  InRefSkeleton)
+{
+	int32 TotalNum = InRefSkeleton.GetNum();
+	TArray<FBoneDescription>& BoneDescList = BoneDescs[Index];
+
+	BoneDescList.Reset(TotalNum);
+	RefSkeleton[Index] = InRefSkeleton;
+
+	if (TotalNum > 0)
+	{
+		BoneDescList.AddZeroed(TotalNum);
+
+		// fill up component space transforms
+		TArray<FTransform> ComponentSpaceTransforms;
+		FAnimationRuntime::FillUpComponentSpaceTransforms(InRefSkeleton, InRefSkeleton.GetRefBonePose(), ComponentSpaceTransforms);
+
+		// get extent
+		FBox BoxExtent;
+		BoxExtent.Init();
+
+		for (int32 BoneIndex = 0; BoneIndex < TotalNum; ++BoneIndex)
+		{
+			BoxExtent += ComponentSpaceTransforms[BoneIndex].GetLocation();
+		}
+
+		const FVector MeshBoxSize = BoxExtent.GetSize();
+		const FVector RootPosition = ComponentSpaceTransforms[0].GetTranslation();
+
+		if (MeshBoxSize.GetMin() > SMALL_NUMBER)
+		{
+			// now fill up data
+			const TArray<FMeshBoneInfo>& MeshInfoList = InRefSkeleton.GetRefBoneInfo();
+			for (int32 BoneIndex = 0; BoneIndex < TotalNum; ++BoneIndex)
+			{
+				FBoneDescription& BoneDesc = BoneDescList[BoneIndex];
+				BoneDesc.BoneInfo = MeshInfoList[BoneIndex];
+				BoneDesc.NumChildren = 0;
+
+				// normalized position
+				const FVector& Position = ComponentSpaceTransforms[BoneIndex].GetLocation();
+				BoneDesc.NormalizedPosition = (Position - BoxExtent.Min) / MeshBoxSize;
+
+				// increase parent child
+				if (BoneDesc.BoneInfo.ParentIndex != INDEX_NONE)
+				{
+					const int32 ParentIndex = BoneDesc.BoneInfo.ParentIndex;
+					++BoneDescList[ParentIndex].NumChildren;
+
+					const FVector& ParentPosition = ComponentSpaceTransforms[ParentIndex].GetLocation();
+					FVector ToChild = (Position-ParentPosition);
+					BoneDesc.RatioFromParent = static_cast<float>(ToChild.Size() / MeshBoxSize.Size());
+					BoneDesc.DirFromParent = ToChild.GetSafeNormal();
+					BoneDesc.DirFromRoot = (Position-RootPosition).GetSafeNormal(); // based on whole mesh size
+				}
+			}
+		}
+	}
+}
+
+void FBoneMappingHelper::TryMatch(TMap<FName, FName>& OutBestMatches)
+{
+	TArray<FBoneDescription>& BoneDescArray0 = BoneDescs[0];
+	TArray<FBoneDescription>& BoneDescArray1 = BoneDescs[1];
+
+	for (int32 BoneIndex0 = 0; BoneIndex0 < BoneDescArray0.Num(); ++BoneIndex0)
+	{
+		BoneDescArray0[BoneIndex0].ResetScore(BoneDescArray1.Num());
+	}
+
+	for (int32 BoneIndex1 = 0; BoneIndex1 < BoneDescArray1.Num(); ++BoneIndex1)
+	{
+		BoneDescArray1[BoneIndex1].ResetScore(BoneDescArray0.Num());
+	}
+
+	for (int32 BoneIndex0 = 0; BoneIndex0 < BoneDescArray0.Num(); ++BoneIndex0)
+	{
+		FBoneDescription& BoneDesc0 = BoneDescArray0[BoneIndex0];
+		for (int32 BoneIndex1 = 0; BoneIndex1 < BoneDescArray1.Num(); ++BoneIndex1)
+		{
+			FBoneDescription& BoneDesc1 = BoneDescArray1[BoneIndex1];
+
+			// set score in both container
+			float Score = BoneDesc0.CalculateScore(BoneDesc1);
+			BoneDesc0.SetScore(BoneIndex1, Score);
+			BoneDesc1.SetScore(BoneIndex0, Score);
+		}
+	}
+
+#define MAX_CANDIDATE 10
+
+	// first find best matches up to MAX_CANDIDATE for each
+	struct FCandidate
+	{
+		FName BoneNames[MAX_CANDIDATE];
+		float Scores[MAX_CANDIDATE];
+		float StdDev;
+
+		FCandidate()
+		{
+			FMemory::Memzero(Scores);
+			StdDev = 0.f;
+		}
+
+		void Set(FName InBoneName, float InScore, int32 Index)
+		{
+			BoneNames[Index] = InBoneName;
+			Scores[Index] = InScore;
+		}
+
+		void CalculateStdDev()
+		{
+			float Avg = 0.f;
+
+			for (int32 Index = 0; Index < MAX_CANDIDATE; ++Index)
+			{
+				Avg += Scores[Index];
+			}
+
+			Avg /= MAX_CANDIDATE;
+			
+			float AccumulatedDev = 0.f;
+
+			for (int32 Index = 0; Index < MAX_CANDIDATE; ++Index)
+			{
+				AccumulatedDev += FMath::Square((Scores[Index]) - Avg);
+			}
+
+			StdDev = FGenericPlatformMath::Sqrt(AccumulatedDev/MAX_CANDIDATE);
+		}
+	};
+
+	TMap<FName, FCandidate> Candidates;
+
+	// find the best score
+	for (int32 BoneIndex0 = 0; BoneIndex0 < BoneDescArray0.Num(); ++BoneIndex0)
+	{
+		int32 BoneIndex1 = BoneDescArray0[BoneIndex0].GetBestIndex();
+
+		if (BoneIndex1 != INDEX_NONE)
+		{
+			FCandidate Candidate;
+			FName Bone0Name = BoneDescArray0[BoneIndex0].BoneInfo.Name;
+			FName Bone1Name = BoneDescArray1[BoneIndex1].BoneInfo.Name;
+			float Score = BoneDescArray0[BoneIndex0].GetScore(BoneIndex1);
+			Candidate.Set(Bone1Name, Score, 0);
+
+			UE_LOG(LogAnimation, Log, TEXT("Bone Match [%s] - [%s] (score %0.2f)"), *Bone0Name.ToString(), *Bone1Name.ToString(), Score);
+
+			// print next 3 best score, clear current one
+			for (int32 LoopCount = 0; LoopCount < MAX_CANDIDATE-1 && BoneIndex1 != INDEX_NONE; ++LoopCount)
+			{
+				BoneDescArray0[BoneIndex0].SetScore(BoneIndex1, 0.f);
+
+				BoneIndex1 = BoneDescArray0[BoneIndex0].GetBestIndex();
+				if (BoneIndex1 != INDEX_NONE)
+				{
+					Bone1Name = BoneDescArray1[BoneIndex1].BoneInfo.Name;
+					Score = BoneDescArray0[BoneIndex0].GetScore(BoneIndex1);
+
+					Candidate.Set(Bone1Name, Score, LoopCount + 1);
+					UE_LOG(LogAnimation, Log, TEXT(" Candidate %d. - [%s] (score %0.2f)"), LoopCount + 1, *Bone1Name.ToString(), Score);
+				}
+			}
+
+			Candidate.CalculateStdDev();
+			Candidates.Add(Bone0Name, Candidate);
+		}
+		else
+		{
+			UE_LOG(LogAnimation, Log, TEXT("Bone [%s] does not have a match"), *BoneDescArray0[BoneIndex0].BoneInfo.Name.ToString());
+		}
+	}
+
+	struct FCandidateSortCallback
+	{
+		FORCEINLINE bool operator()(const FCandidate& A, const FCandidate& B) const
+		{
+			return A.StdDev > B.StdDev;
+		}
+	};
+
+	Candidates.ValueSort(FCandidateSortCallback());
+
+	TArray<FName> UsedNames;
+	for (const TPair<FName, FCandidate>& CandidatePair : Candidates)
+	{
+		const FCandidate& Candidate = CandidatePair.Value;
+		for (int32 Index = 0; Index < MAX_CANDIDATE; ++Index)
+		{
+			FName BestMatchName = Candidate.BoneNames[Index];
+			// see if it's already used by other joint
+			if (UsedNames.Find(BestMatchName) == INDEX_NONE)
+			{
+				// if that's case, just ignore and move on
+				OutBestMatches.Add(CandidatePair.Key, BestMatchName);
+				UsedNames.Add(BestMatchName);
+				break;
+			}
+		}
+	}
+}
+
diff --git a/Engine/Source/Editor/Persona/Private/BoneMappingHelper.h b/Engine/Source/Editor/Persona/Private/BoneMappingHelper.h
new file mode 100644
--- /dev/null
+++ b/Engine/Source/Editor/Persona/Private/BoneMappingHelper.h
@@ -0,0 +1,62 @@
+// Copyright Epic Games, Inc. All Rights Reserved.
+
+#pragma once
+
+#include "CoreMinimal.h"
+#include "ReferenceSkeleton.h"
+
+//////////////////////////////////////////////////////////////////////////
+// FBoneDescription
+//////////////////////////////////////////////////////////////////////////
+struct FBoneDescription
+{
+	FMeshBoneInfo BoneInfo;
+
+	FVector NormalizedPosition; // based on whole mesh size
+	FVector DirFromParent;
+	FVector DirFromRoot;
+	float	RatioFromParent; // based on whole mesh size
+	int32	NumChildren;
+
+	TArray<float> Scores;
+
+	void ResetScore(int32 NewCount)
+	{
+		Scores.Reset(NewCount);
+		Scores.AddZeroed(NewCount);
+	}
+
+	void SetScore(int32 Index, float NewScore)
+	{
+		Scores[Index] = NewScore;
+	}
+
+	float GetScore(int32 Index) const
+	{
+		return Scores[Index];
+	}
+
+	float CalculateScore(const FBoneDescription& Other) const;
+	float CalculateNameScore(const FName& Name1, const FName& Name2) const;
+	int32 GetBestIndex() const;
+};
+
+//////////////////////////////////////////////////////////////////////////
+// BoneMappingHelper Class
+
+struct FBoneMappingHelper
+{
+	// input - 2 skeletons
+	FReferenceSkeleton RefSkeleton[2];
+
+	// initialize data
+	FBoneMappingHelper(const FReferenceSkeleton& InRefSkeleton1, const FReferenceSkeleton&  InRefSkeleton2);
+
+	void TryMatch(TMap<FName, FName>& OutBestMatches);
+
+private:
+	// BoneDescription array for each bone
+	TArray<FBoneDescription>	BoneDescs[2];
+
+	void Initialize(int32 Index, const FReferenceSkeleton&  InRefSkeleton);
+};
diff --git a/Engine/Source/Editor/Persona/Private/SBoneMappingBase.cpp b/Engine/Source/Editor/Persona/Private/SBoneMappingBase.cpp
--- a/Engine/Source/Editor/Persona/Private/SBoneMappingBase.cpp
+++ b/Engine/Source/Editor/Persona/Private/SBoneMappingBase.cpp
@@ -10,9 +10,12 @@
 #include "Styling/AppStyle.h"
 #include "Widgets/Input/SButton.h"
 #include "AssetNotifications.h"
+#include "Animation/Rig.h"
 #include "BoneSelectionWidget.h"
 #include "Widgets/Input/SSearchBox.h"
 #include "Widgets/Text/SInlineEditableTextBlock.h"
+#include "SRigPicker.h"
+#include "BoneMappingHelper.h"
 #include "SSkeletonWidget.h"
 #include "IEditableSkeleton.h"
 
diff --git a/Engine/Source/Editor/Persona/Private/SRigPicker.cpp b/Engine/Source/Editor/Persona/Private/SRigPicker.cpp
new file mode 100644
--- /dev/null
+++ b/Engine/Source/Editor/Persona/Private/SRigPicker.cpp
@@ -0,0 +1,133 @@
+// Copyright Epic Games, Inc. All Rights Reserved.
+
+#include "SRigPicker.h"
+#include "Modules/ModuleManager.h"
+#include "Framework/MultiBox/MultiBoxBuilder.h"
+#include "Widgets/Layout/SBox.h"
+#include "IContentBrowserSingleton.h"
+#include "ContentBrowserModule.h"
+#include "Animation/Rig.h"
+
+#define LOCTEXT_NAMESPACE "RigPicker"
+
+static const FVector2f ContentBrowserWindowSize(300.0f, 300.0f);
+
+PRAGMA_DISABLE_DEPRECATION_WARNINGS
+URig* SRigPicker::EngineHumanoidRig = NULL;
+
+void SRigPicker::Construct(const FArguments& InArgs)
+{
+	if (EngineHumanoidRig == NULL)
+	{
+		EngineHumanoidRig = LoadObject<URig>(nullptr, TEXT("/Engine/EngineMeshes/Humanoid.Humanoid"), nullptr, LOAD_None, nullptr);
+	}
+
+	CurrentObject = InArgs._InitialObject;
+	ShouldFilterAsset = InArgs._OnShouldFilterAsset;
+	OnSetReference = InArgs._OnSetReference;
+	OnClose = InArgs._OnClose;
+
+	FMenuBuilder MenuBuilder(true, NULL);
+
+	MenuBuilder.BeginSection(NAME_None, LOCTEXT("EngineAssets", "Default"));
+	{
+		if (EngineHumanoidRig)
+		{
+			MenuBuilder.AddMenuEntry(
+				LOCTEXT("PickDefaultRig", "Select Humanoid Rig"),
+				LOCTEXT("PickDefaultRig_Tooltip", "Select Engine Defined Humanoid Rig"),
+				FSlateIcon(),
+				FUIAction(FExecuteAction::CreateSP(this, &SRigPicker::OnSelectDefault)));
+		}
+
+		MenuBuilder.AddMenuEntry(
+			LOCTEXT("ClearAsset", "Clear"),
+			LOCTEXT("ClearAsset_ToolTip", "Clears the asset set on this field"),
+			FSlateIcon(),
+			FUIAction(FExecuteAction::CreateSP(this, &SRigPicker::OnClear))
+			);
+	}
+	MenuBuilder.EndSection();
+
+	MenuBuilder.BeginSection(NAME_None, LOCTEXT("BrowseHeader", "Browse"));
+	{
+		TSharedPtr<SWidget> MenuContent;
+
+		FContentBrowserModule& ContentBrowserModule = FModuleManager::Get().LoadModuleChecked<FContentBrowserModule>(TEXT("ContentBrowser"));
+
+		FAssetPickerConfig AssetPickerConfig;
+		AssetPickerConfig.Filter.ClassPaths.Add(URig::StaticClass()->GetClassPathName());
+		// Allow child classes
+		AssetPickerConfig.Filter.bRecursiveClasses = true;
+		// Set a delegate for setting the asset from the picker
+		AssetPickerConfig.OnAssetSelected = FOnAssetSelected::CreateSP(this, &SRigPicker::OnAssetSelected);
+		// Use the list view by default
+		AssetPickerConfig.InitialAssetViewType = EAssetViewType::List;
+		// The initial selection should be the current value
+		AssetPickerConfig.InitialAssetSelection = CurrentObject;
+		// We'll do clearing ourselves
+		AssetPickerConfig.bAllowNullSelection = false;
+		// Focus search box
+		AssetPickerConfig.bFocusSearchBoxWhenOpened = true;
+		// Apply custom filter
+		AssetPickerConfig.OnShouldFilterAsset = FOnShouldFilterAsset::CreateSP(this, &SRigPicker::OnShouldFilterAsset);
+		// Don't allow dragging
+		AssetPickerConfig.bAllowDragging = false;
+
+		MenuContent =
+			SNew(SBox)
+			.WidthOverride(ContentBrowserWindowSize.X)
+			.HeightOverride(ContentBrowserWindowSize.Y)
+			[
+				ContentBrowserModule.Get().CreateAssetPicker(AssetPickerConfig)
+			];
+
+		MenuBuilder.AddWidget(MenuContent.ToSharedRef(), FText::GetEmpty(), true);
+	}
+	MenuBuilder.EndSection();
+
+	ChildSlot
+	[
+		MenuBuilder.MakeWidget()
+	];
+}
+
+void SRigPicker::OnClear()
+{
+	SetValue(NULL);
+	OnClose.ExecuteIfBound();
+}
+
+void SRigPicker::OnAssetSelected(const FAssetData& AssetData)
+{
+	SetValue(AssetData);
+	OnClose.ExecuteIfBound();
+}
+
+bool SRigPicker::OnShouldFilterAsset(const struct FAssetData& AssetData)
+{
+	if (EngineHumanoidRig && AssetData.GetSoftObjectPath() == FSoftObjectPath(EngineHumanoidRig))
+	{
+		return true;
+	}
+
+	if (ShouldFilterAsset.IsBound())
+	{
+		return ShouldFilterAsset.Execute(AssetData);
+	}
+
+	return false;
+}
+void SRigPicker::SetValue(const FAssetData& AssetData)
+{
+	OnSetReference.ExecuteIfBound(AssetData.GetAsset());
+}
+
+void SRigPicker::OnSelectDefault()
+{
+	OnAssetSelected(EngineHumanoidRig);
+}
+
+#undef LOCTEXT_NAMESPACE
+
+PRAGMA_ENABLE_DEPRECATION_WARNINGS
\ No newline at end of file
diff --git a/Engine/Source/Editor/Persona/Private/SRigPicker.h b/Engine/Source/Editor/Persona/Private/SRigPicker.h
new file mode 100644
--- /dev/null
+++ b/Engine/Source/Editor/Persona/Private/SRigPicker.h
@@ -0,0 +1,69 @@
+// Copyright Epic Games, Inc. All Rights Reserved.
+#pragma once
+
+#include "CoreMinimal.h"
+#include "Widgets/DeclarativeSyntaxSupport.h"
+#include "Widgets/SCompoundWidget.h"
+#include "AssetRegistry/AssetData.h"
+#include "PropertyCustomizationHelpers.h"
+
+class URig;
+
+/** Called when an asset is selected in the asset view */
+DECLARE_DELEGATE_OneParam(FOnAssetSelected, const struct FAssetData& /*AssetData*/);
+
+/**
+ * A widget used to pick rig 
+ */
+class SRigPicker : public SCompoundWidget
+{
+public:
+	DECLARE_DELEGATE_OneParam(FOnSetReference, UObject*)
+
+	SLATE_BEGIN_ARGS( SRigPicker )
+		: _InitialObject(NULL)
+	{}
+		SLATE_ARGUMENT( FAssetData, InitialObject )
+		SLATE_EVENT( FOnShouldFilterAsset, OnShouldFilterAsset )
+		SLATE_EVENT( FOnSetReference, OnSetReference )
+		SLATE_EVENT( FSimpleDelegate, OnClose )
+	SLATE_END_ARGS()
+
+	/**
+	 * Construct the widget.
+	 * @param	InArgs				Arguments for widget construction
+	 * @param	InPropertyHandle	The property handle that this widget will operate on.
+	 */
+	void Construct( const FArguments& InArgs );
+
+private:
+	/** 
+	 * Delegate for handling selection in the asset browser.
+	 * @param	AssetData	The chosen asset data
+	 */
+	void OnAssetSelected( const FAssetData& AssetData );
+	void OnSelectDefault();
+	void OnClear();
+	bool OnShouldFilterAsset(const struct FAssetData& AssetData);
+
+	/** 
+	 * Set the value of the asset referenced by this property editor.
+	 * Will set the underlying property handle if there is one.
+	 * @param	AssetData	The asset data for the object to set the reference to
+	 */
+	void SetValue( const FAssetData& AssetData );
+private:
+	FAssetData CurrentObject;
+
+	/** Delegate for filtering valid assets */
+	FOnShouldFilterAsset ShouldFilterAsset;
+
+	/** Delegate to call when our object value should be set */
+	FOnSetReference OnSetReference;
+
+	/** Delegate for closing the containing menu */
+	FSimpleDelegate OnClose;
+PRAGMA_DISABLE_DEPRECATION_WARNINGS
+	static URig* EngineHumanoidRig;
+PRAGMA_ENABLE_DEPRECATION_WARNINGS
+};
diff --git a/Engine/Source/Editor/SkeletalMeshEditor/Private/SkeletalMeshEditor.cpp b/Engine/Source/Editor/SkeletalMeshEditor/Private/SkeletalMeshEditor.cpp
--- a/Engine/Source/Editor/SkeletalMeshEditor/Private/SkeletalMeshEditor.cpp
+++ b/Engine/Source/Editor/SkeletalMeshEditor/Private/SkeletalMeshEditor.cpp
@@ -77,6 +77,7 @@ const FName SkeletalMeshEditorTabs::AssetDetailsTab(TEXT("AnimAssetPropertiesTab
 const FName SkeletalMeshEditorTabs::ViewportTab(TEXT("Viewport"));
 const FName SkeletalMeshEditorTabs::AdvancedPreviewTab(TEXT("AdvancedPreviewTab"));
 const FName SkeletalMeshEditorTabs::MorphTargetsTab("MorphTargetsTab");
+const FName SkeletalMeshEditorTabs::AnimationMappingTab("AnimationMappingWindow");
 const FName SkeletalMeshEditorTabs::ToolboxDetailsTab("ToolBoxDetailsTab");
 const FName SkeletalMeshEditorTabs::CurveMetadataTab(TEXT("AnimCurveMetadataEditorTab"));
 const FName SkeletalMeshEditorTabs::FindReplaceTab("FindReplaceTab");
diff --git a/Engine/Source/Editor/SkeletonEditor/Private/EditableSkeleton.cpp b/Engine/Source/Editor/SkeletonEditor/Private/EditableSkeleton.cpp
--- a/Engine/Source/Editor/SkeletonEditor/Private/EditableSkeleton.cpp
+++ b/Engine/Source/Editor/SkeletonEditor/Private/EditableSkeleton.cpp
@@ -1405,6 +1405,36 @@ void FEditableSkeleton::RemoveCompatibleSkeleton(const USkeleton* InCompatibleSk
 	PersonaModule.BroadcastAssetFamilyChange();
 }
 
+void FEditableSkeleton::RefreshRigConfig()
+{
+	Skeleton->RefreshRigConfig();
+}
+
+void FEditableSkeleton::SetRigConfig(URig* InRig)
+{
+	const FScopedTransaction Transaction(LOCTEXT("RigAssetChanged", "Select Rig"));
+	Skeleton->Modify();
+	Skeleton->SetRigConfig(InRig);
+}
+
+void FEditableSkeleton::SetRigBoneMapping(const FName& InNodeName, const FName& InBoneName)
+{
+	const FScopedTransaction Transaction(LOCTEXT("BoneMappingChanged", "Change Bone Mapping"));
+	Skeleton->Modify();
+	Skeleton->SetRigBoneMapping(InNodeName, InBoneName);
+}
+
+void FEditableSkeleton::SetRigBoneMappings(const TMap<FName, FName>& InMappings)
+{
+	const FScopedTransaction Transaction(LOCTEXT("BoneMappingsChanged", "Change Bone Mappings"));
+	Skeleton->Modify();
+
+	for (const TPair<FName, FName>& Mapping : InMappings)
+	{
+		Skeleton->SetRigBoneMapping(Mapping.Key, Mapping.Value);
+	}
+}
+
 void FEditableSkeleton::RemoveUnusedBones()
 {
 	TArray<FName> SkeletonBones;
diff --git a/Engine/Source/Editor/SkeletonEditor/Private/EditableSkeleton.h b/Engine/Source/Editor/SkeletonEditor/Private/EditableSkeleton.h
--- a/Engine/Source/Editor/SkeletonEditor/Private/EditableSkeleton.h
+++ b/Engine/Source/Editor/SkeletonEditor/Private/EditableSkeleton.h
@@ -60,6 +60,10 @@ public:
 	virtual void RefreshRetargetSources(const TArray<FName>& InRetargetSourceNames) override;
 	virtual void AddCompatibleSkeleton(const USkeleton* InCompatibleSkeleton) override;
 	virtual void RemoveCompatibleSkeleton(const USkeleton* InCompatibleSkeleton) override;
+	virtual void RefreshRigConfig() override;
+	virtual void SetRigConfig(URig* InRig) override;
+	virtual void SetRigBoneMapping(const FName& InNodeName, const FName& InBoneName) override;
+	virtual void SetRigBoneMappings(const TMap<FName, FName>& InMappings) override;
 	virtual void RemoveUnusedBones() override;
 	virtual void UpdateSkeletonReferencePose(USkeletalMesh* InSkeletalMesh) override;
 	virtual void RegisterSlotNode(const FName& InSlotName) override;
diff --git a/Engine/Source/Editor/SkeletonEditor/Public/IEditableSkeleton.h b/Engine/Source/Editor/SkeletonEditor/Public/IEditableSkeleton.h
--- a/Engine/Source/Editor/SkeletonEditor/Public/IEditableSkeleton.h
+++ b/Engine/Source/Editor/SkeletonEditor/Public/IEditableSkeleton.h
@@ -166,6 +166,20 @@ public:
 	/** Remove a compatible skeleton */
 	virtual void RemoveCompatibleSkeleton(const USkeleton* InCompatibleSkeleton) = 0;
 
+	PRAGMA_DISABLE_DEPRECATION_WARNINGS
+	/** Refreshes the rig config, validating the mappings */
+	virtual void RefreshRigConfig() = 0;
+
+	/** Set the rig config  */
+	virtual void SetRigConfig(class URig* InRig) = 0;
+
+	/** Set a rig bone mapping */
+	virtual void SetRigBoneMapping(const FName& InNodeName, const FName& InBoneName) = 0;
+
+	/** Set multiple bone mappings */
+	virtual void SetRigBoneMappings(const TMap<FName, FName>& InMappings) = 0;
+	PRAGMA_ENABLE_DEPRECATION_WARNINGS
+
 	/** Remove any bones that are not used by any skeletal meshes */
 	virtual void RemoveUnusedBones() = 0;
 
diff --git a/Engine/Source/Editor/UnrealEd/Private/SSkeletonWidget.cpp b/Engine/Source/Editor/UnrealEd/Private/SSkeletonWidget.cpp
--- a/Engine/Source/Editor/UnrealEd/Private/SSkeletonWidget.cpp
+++ b/Engine/Source/Editor/UnrealEd/Private/SSkeletonWidget.cpp
@@ -17,6 +17,7 @@
 #include "Interfaces/IMainFrameModule.h"
 #include "ContentBrowserModule.h"
 #include "IDocumentation.h"
+#include "Animation/Rig.h"
 #include "AnimPreviewInstance.h"
 #include "AssetRegistry/AssetRegistryModule.h"
 #include "AnimationRuntime.h"
diff --git a/Engine/Source/Runtime/Engine/Classes/Animation/AnimSequence.h b/Engine/Source/Runtime/Engine/Classes/Animation/AnimSequence.h
--- a/Engine/Source/Runtime/Engine/Classes/Animation/AnimSequence.h
+++ b/Engine/Source/Runtime/Engine/Classes/Animation/AnimSequence.h
@@ -778,7 +778,9 @@ private:
 	ENGINE_API virtual void RemapTracksToNewSkeleton( USkeleton* NewSkeleton, bool bConvertSpaces ) override;
 
 	/** Retargeting functions */
-	ENGINE_API int32 GetSpaceBasedAnimationData(TArray< TArray<FTransform> > & AnimationDataInComponentSpace) const;
+	ENGINE_API bool ConvertAnimationDataToRiggingData(FAnimSequenceTrackContainer & RiggingAnimationData);
+	ENGINE_API bool ConvertRiggingDataToAnimationData(FAnimSequenceTrackContainer & RiggingAnimationData);
+	ENGINE_API int32 GetSpaceBasedAnimationData(TArray< TArray<FTransform> > & AnimationDataInComponentSpace, FAnimSequenceTrackContainer * RiggingAnimationData) const;
 #endif
 
 public:
diff --git a/Engine/Source/Runtime/Engine/Classes/Animation/Rig.h b/Engine/Source/Runtime/Engine/Classes/Animation/Rig.h
new file mode 100644
--- /dev/null
+++ b/Engine/Source/Runtime/Engine/Classes/Animation/Rig.h
@@ -0,0 +1,231 @@
+// Copyright Epic Games, Inc. All Rights Reserved.
+
+/** 
+ * This is the definition for a Rig that is used for retargeting animations
+ *
+ */
+
+#pragma once
+
+// @todo should I keep the hierarchy or not? - does it matter if hand is child of shoulder? - I think we should keep hierarchy
+// @todo should we support reset data 
+// @todo does it make sense to have "no constraint" on certain data? What does that mean? Just World? 
+
+#include "CoreMinimal.h"
+#include "UObject/ObjectMacros.h"
+#include "UObject/Object.h"
+#include "ReferenceSkeleton.h"
+#include "Animation/NodeMappingProviderInterface.h"
+#include "Rig.generated.h"
+
+class USkeleton;
+struct FPropertyChangedEvent;
+
+// January 2023 - Whole compilation unit is being deprecated, so disabled warnings for whole thing. 
+PRAGMA_DISABLE_DEPRECATION_WARNINGS
+
+/** Rig Controller for bone transform **/
+USTRUCT()
+struct UE_DEPRECATED(5.2, "URig retargeting is no longer supported. Use the IK Retargeter instead.") FNode
+{
+	GENERATED_USTRUCT_BODY()
+
+	/** Name of the original node. We don't allow to change this. This is used for identity.**/
+	UPROPERTY(VisibleAnywhere, Category="FNode")
+	FName Name;
+
+	/** We save Parent Node but if the parent node is removed, it will reset to root */
+	UPROPERTY(VisibleAnywhere, Category="FNode")
+	FName ParentName;
+
+	/** Absolute transform of the node. Hoping to use this data in the future to render*/
+	UPROPERTY()
+	FTransform Transform;
+
+	/** This is Display Name where it will be used to display in Retarget Manager. This name has to be unique. */
+	UPROPERTY(EditAnywhere, Category="FNode")
+	FString DisplayName;
+
+	UPROPERTY(EditAnywhere, Category="FNode")
+	bool bAdvanced;
+
+	FNode()
+		: Name(NAME_None)
+		, ParentName(NAME_None)
+		, bAdvanced(false)
+	{
+	}
+
+	FNode(FName InNodeName, FName InParentName, const FTransform& InTransform)
+		: Name(InNodeName)
+		, ParentName(InParentName)
+		, Transform(InTransform)
+		, DisplayName(InNodeName.ToString())
+		, bAdvanced(false)
+	{
+	}
+};
+
+/** Control Constraint Type */
+UENUM()
+namespace EControlConstraint
+{
+	enum UE_DEPRECATED(5.2, "URig retargeting is no longer supported. Use the IK Retargeter instead.") Type : int
+	{
+		/** Rotation constraint. */
+		Orientation,
+		/** Translation constraint. */
+		Translation,
+
+		/** Max Number. */
+		MAX
+	};
+}
+
+/** Constraint Transform Type. - currently unused */
+UENUM()
+namespace EConstraintTransform
+{
+	enum UE_DEPRECATED(5.2, "URig retargeting is no longer supported. Use the IK Retargeter instead.") Type : int
+	{
+		/** Absolute value. */
+		Absolute,
+		/** Apply relative transform from ref pose. */
+		Relative,
+	};
+}
+
+USTRUCT()
+struct UE_DEPRECATED(5.2, "URig retargeting is no longer supported. Use the IK Retargeter instead.") FRigTransformConstraint
+{
+	GENERATED_USTRUCT_BODY()
+	
+	/** What transform type **/
+	UPROPERTY(/*EditAnywhere, Category="FTransformBaseConstraint"*/)
+	TEnumAsByte<EConstraintTransform::Type>	TranformType = EConstraintTransform::Absolute;
+
+	/** Parent space that is defined **/
+	UPROPERTY(EditAnywhere, Category="FTransformBaseConstraint")
+	FName	ParentSpace;
+
+	/** Weight of the influence - for future*/
+	UPROPERTY(/*EditAnywhere, Category="FTransformBaseConstraint"*/)
+	float	Weight = 0.f;
+};
+
+/** This defines what constraint it is defined */
+USTRUCT()
+struct UE_DEPRECATED(5.2, "URig retargeting is no longer supported. Use the IK Retargeter instead.") FTransformBaseConstraint
+{
+	GENERATED_USTRUCT_BODY()
+	
+	/** What transform type **/
+	UPROPERTY(EditAnywhere, Category="FTransformBaseConstraint")
+	TArray<FRigTransformConstraint>			TransformConstraints;
+};
+
+
+/** This is a mapping table between bone in a particular skeletal mesh and bone of this skeleton set. */
+USTRUCT()
+struct UE_DEPRECATED(5.2, "URig retargeting is no longer supported. Use the IK Retargeter instead.") FTransformBase
+{
+	GENERATED_USTRUCT_BODY()
+
+	UPROPERTY(EditAnywhere, Category="FTransformBase")
+	FName					Node;
+
+	UPROPERTY(EditAnywhere, Category="FTransformBase")
+	FTransformBaseConstraint	Constraints[EControlConstraint::Type::MAX];
+};
+
+DECLARE_DELEGATE_RetVal_OneParam(int32, FGetParentIndex, FName);
+
+/**
+ *	URig : that has rigging data for skeleton
+ *		- used for retargeting
+ *		- support to share different animations
+ */
+UCLASS(hidecategories=Object, MinimalAPI)
+class UE_DEPRECATED(5.2, "URig retargeting is no longer supported. Use the IK Retargeter instead.") URig : public UObject, public INodeMappingProviderInterface
+{
+	GENERATED_UCLASS_BODY()
+
+private:
+
+	/** Skeleton bone tree - each contains name and parent index**/
+	UPROPERTY(EditAnywhere, Category = Rig, EditFixedSize)
+	TArray<FTransformBase> TransformBases;
+
+	/** Skeleton bone tree - each contains name and parent index**/
+	UPROPERTY(EditAnywhere, Category=Rig, EditFixedSize)
+	TArray<FNode> Nodes;
+
+#if WITH_EDITORONLY_DATA
+	// this is source skeleton it's created from. 
+	// since all node data can be modified after, 
+	// to figure out what was original source skeleton, this is better to have it for reset and getting original transform
+	struct FReferenceSkeleton SourceSkeleton;
+#endif
+public:
+
+	virtual void Serialize(FArchive& Ar) override;
+
+#if WITH_EDITOR
+	//~ Begin UObject Interface
+	virtual void PostEditChangeProperty(FPropertyChangedEvent& PropertyChangedEvent) override;
+	//~ End UObject Interface
+
+	// node related functions
+	int32 GetNodeNum() const;
+	const FNode* GetNode(int32 NodeIndex) const;
+	FName GetNodeName(int32 NodeIndex) const;
+	FName GetParentNode(FName& NodeName) const;
+	int32 FindNode(const FName& NodeName) const;
+	const TArray<FNode> & GetNodes() const { return Nodes; }
+
+	// create from skeleton
+	ENGINE_API void CreateFromSkeleton(const USkeleton* Skeleton, const TMap<int32, int32>& RequiredBones);
+	ENGINE_API void SetAllConstraintsToParents();
+	ENGINE_API void SetAllConstraintsToWorld();
+
+	// rig control related
+	int32 GetTransformBaseNum() const;
+	const TArray<FTransformBase> & GetTransformBases() const { return TransformBases; }
+	const FTransformBase* GetTransformBase(int32 TransformBaseIndex) const;
+	const FTransformBase* GetTransformBaseByNodeName(FName NodeName) const;
+	int32 FindTransformBaseByNodeName(FName NodeName) const;
+	int32 FindTransformParentNode(int32 NodeIndex, bool bTranslate, int32 Index=0) const;
+
+#if WITH_EDITORONLY_DATA
+	// source skeleton related, since this has been added later, it's possible 
+	// some skeletons don't have it
+	bool IsSourceReferenceSkeletonAvailable() const { return SourceSkeleton.GetRawBoneNum() > 0; }
+	const FReferenceSkeleton& GetSourceReferenceSkeleton() const { return SourceSkeleton;  }
+	ENGINE_API void SetSourceReferenceSkeleton(const FReferenceSkeleton& InSrcSkeleton);
+#endif // WITH_EDITORONLY_DATA
+#endif // WITH_EDITOR
+
+	// not sure if we'd like to keep this
+	ENGINE_API static FName WorldNodeName;
+
+private:
+
+#if WITH_EDITOR
+	// for now these are privates since we don't have good control yet
+	bool AddNode(FName Name, FName ParentNode, FTransform Transform);
+	bool DeleteNode(FName Name);
+	// rig constraint 
+	bool AddRigConstraint(FName NodeName, EControlConstraint::Type	ConstraintType, EConstraintTransform::Type	TranformType, FName ParentSpace, float Weight = 1.f);
+
+#endif // WITH_EDITOR
+
+	// INodeMappingProviderInterface
+	virtual void GetMappableNodeData(TArray<FName>& OutNames, TArray<FNodeItem>& OutNodeItems) const override;
+	// not useful so far
+//	void CalculateComponentSpace(int32 NodeIndex, const FTransform& LocalTransform, const TArray<FTransform> & TransformBuffer, const FGetParentIndex& DelegateToGetParentIndex, FTransform& OutComponentSpaceTransform) const;
+//	void CalculateLocalSpace(int32 NodeIndex, const FTransform& ComponentTransform, const TArray<FTransform> & TransformBuffer, const FGetParentIndex& DelegateToGetParentIndex, FTransform& OutLocalSpaceTransform) const;
+};
+
+// January 2023 - Whole compilation unit is being deprecated, so disabled warnings for whole thing. 
+PRAGMA_ENABLE_DEPRECATION_WARNINGS
+
diff --git a/Engine/Source/Runtime/Engine/Classes/Animation/Skeleton.h b/Engine/Source/Runtime/Engine/Classes/Animation/Skeleton.h
--- a/Engine/Source/Runtime/Engine/Classes/Animation/Skeleton.h
+++ b/Engine/Source/Runtime/Engine/Classes/Animation/Skeleton.h
@@ -216,6 +216,21 @@ struct FNameMapping
 	}
 };
 
+USTRUCT()
+struct FRigConfiguration
+{
+	GENERATED_USTRUCT_BODY()
+
+	PRAGMA_DISABLE_DEPRECATION_WARNINGS
+	UPROPERTY()
+	TObjectPtr<class URig>  Rig = nullptr;
+	PRAGMA_ENABLE_DEPRECATION_WARNINGS
+
+	// @todo in the future we can make this to be run-time data
+	UPROPERTY()
+	TArray<FNameMapping> BoneMappingTable;
+};
+
 USTRUCT()
 struct FAnimSlotGroup
 {
@@ -579,6 +594,11 @@ private:
 	UPROPERTY(duplicatetransient, AssetRegistrySearchable)
 	TSoftObjectPtr<class UDataAsset> AdditionalPreviewSkeletalMeshes;
 
+	PRAGMA_DISABLE_DEPRECATION_WARNINGS
+	UPROPERTY()
+	FRigConfiguration RigConfig;
+	PRAGMA_ENABLE_DEPRECATION_WARNINGS
+
 	/** rig property will be saved separately */
 	ENGINE_API virtual void GetAssetRegistryTags(TArray<FAssetRegistryTag>& OutTags) const override;
 
@@ -1002,6 +1022,21 @@ public:
 	ENGINE_API static const FName CompatibleSkeletonsNameTag;
 	ENGINE_API static const FString CompatibleSkeletonsTagDelimiter;
 
+	PRAGMA_DISABLE_DEPRECATION_WARNINGS
+	// rig Configs
+	ENGINE_API static const FName RigTag;
+	ENGINE_API void SetRigConfig(URig * Rig);
+	ENGINE_API FName GetRigBoneMapping(const FName& NodeName) const;
+	ENGINE_API bool SetRigBoneMapping(const FName& NodeName, FName BoneName);
+	ENGINE_API FName GetRigNodeNameFromBoneName(const FName& BoneName) const;
+	// this make sure it stays within the valid range
+	ENGINE_API int32 GetMappedValidNodes(TArray<FName> &OutValidNodeNames);
+	// verify if it has all latest data
+	ENGINE_API void RefreshRigConfig();
+	int32 FindRigBoneMapping(const FName& NodeName) const;
+	ENGINE_API URig * GetRig() const;
+	PRAGMA_ENABLE_DEPRECATION_WARNINGS
+
 #endif
 
 public:
diff --git a/Engine/Source/Runtime/Engine/Private/Animation/AnimSequence.cpp b/Engine/Source/Runtime/Engine/Private/Animation/AnimSequence.cpp
--- a/Engine/Source/Runtime/Engine/Private/Animation/AnimSequence.cpp
+++ b/Engine/Source/Runtime/Engine/Private/Animation/AnimSequence.cpp
@@ -25,6 +25,7 @@
 #include "Components/SkeletalMeshComponent.h"
 #include "Animation/AnimRootMotionProvider.h"
 #include "Animation/BlendSpace.h"
+#include "Animation/Rig.h"
 #include "Animation/AnimationSettings.h"
 #include "Animation/AnimBoneCompressionSettings.h"
 #include "Animation/AnimCurveCompressionCodec.h"
@@ -2191,6 +2192,114 @@ int32 FindMeshBoneIndexFromBoneName(USkeleton * Skeleton, const FName &BoneName)
 	return BoneIndex;
 }
 
+PRAGMA_DISABLE_DEPRECATION_WARNINGS
+void FillUpTransformBasedOnRig(USkeleton* Skeleton, TArray<FTransform>& NodeSpaceBases, TArray<FTransform> &Rotations, TArray<FVector>& Translations, TArray<bool>& TranslationParentFlags)
+{
+	TArray<FTransform> SpaceBases;
+	FAnimationRuntime::FillUpComponentSpaceTransformsRetargetBasePose(Skeleton, SpaceBases);
+
+	const URig* Rig = Skeleton->GetRig();
+
+	if (Rig)
+	{
+		// this one has to collect all Nodes in Rig data
+		// since we're comparing two of them together. 
+		int32 NodeNum = Rig->GetNodeNum();
+
+		if (NodeNum > 0)
+		{
+			NodeSpaceBases.Empty(NodeNum);
+			NodeSpaceBases.AddUninitialized(NodeNum);
+
+			Rotations.Empty(NodeNum);
+			Rotations.AddUninitialized(NodeNum);
+
+			Translations.Empty(NodeNum);
+			Translations.AddUninitialized(NodeNum);
+
+			TranslationParentFlags.Empty(Translations.Num());
+			TranslationParentFlags.AddZeroed(Translations.Num());
+
+			const USkeletalMesh* PreviewMesh = Skeleton->GetPreviewMesh();
+
+			for (int32 Index = 0; Index < NodeNum; ++Index)
+			{
+				const FName NodeName = Rig->GetNodeName(Index);
+				const FName& BoneName = Skeleton->GetRigBoneMapping(NodeName);
+				const int32& BoneIndex = FindMeshBoneIndexFromBoneName(Skeleton, BoneName);
+
+				if (BoneIndex == INDEX_NONE)
+				{
+					// add identity
+					NodeSpaceBases[Index].SetIdentity();
+					Rotations[Index].SetIdentity();
+					Translations[Index] = FVector::ZeroVector;
+				}
+				else
+				{
+					// initialize with SpaceBases - assuming World Based
+					NodeSpaceBases[Index] = SpaceBases[BoneIndex];
+					Rotations[Index] = SpaceBases[BoneIndex];
+					Translations[Index] = SpaceBases[BoneIndex].GetLocation();
+
+					const FTransformBase* TransformBase = Rig->GetTransformBaseByNodeName(NodeName);
+
+					if (TransformBase != nullptr)
+					{
+						// orientation constraint			
+						const auto& RotConstraint = TransformBase->Constraints[EControlConstraint::Type::Orientation];
+
+						if (RotConstraint.TransformConstraints.Num() > 0)
+						{
+							const FName& ParentBoneName = Skeleton->GetRigBoneMapping(RotConstraint.TransformConstraints[0].ParentSpace);
+							const int32& ParentBoneIndex = FindMeshBoneIndexFromBoneName(Skeleton, ParentBoneName);
+
+							if (ParentBoneIndex != INDEX_NONE)
+							{
+								Rotations[Index] = SpaceBases[BoneIndex].GetRelativeTransform(SpaceBases[ParentBoneIndex]);
+							}
+						}
+
+						// translation constraint
+						const auto& TransConstraint = TransformBase->Constraints[EControlConstraint::Type::Translation];
+
+						if (TransConstraint.TransformConstraints.Num() > 0)
+						{
+							const FName& ParentBoneName = Skeleton->GetRigBoneMapping(TransConstraint.TransformConstraints[0].ParentSpace);
+							const int32& ParentBoneIndex = FindMeshBoneIndexFromBoneName(Skeleton, ParentBoneName);
+
+							if (ParentBoneIndex != INDEX_NONE)
+							{
+								// I think translation has to include rotation, otherwise it won't work
+								Translations[Index] = SpaceBases[BoneIndex].GetLocation() - SpaceBases[ParentBoneIndex].GetLocation();
+								TranslationParentFlags[Index] = true;
+							}
+						}
+					}
+				}
+			}
+		}
+	}
+}
+
+int32 FindValidTransformParentTrack(const URig* Rig, int32 NodeIndex, bool bTranslate, const TArray<FName>& ValidNodeNames)
+{
+	int32 ParentIndex = Rig->FindTransformParentNode(NodeIndex, bTranslate);
+
+	// verify if it exists in ValidNodeNames
+	if (ParentIndex != INDEX_NONE)
+	{
+		FName NodeName = Rig->GetNodeName(ParentIndex);
+
+		return ValidNodeNames.Find(NodeName);
+	}
+
+	return INDEX_NONE;
+}
+
+PRAGMA_ENABLE_DEPRECATION_WARNINGS
+
+
 void UAnimSequence::RemapTracksToNewSkeleton( USkeleton* NewSkeleton, bool bConvertSpaces )
 {
 	// this is not cheap, so make sure it only happens in editor
@@ -2205,9 +2314,220 @@ void UAnimSequence::RemapTracksToNewSkeleton( USkeleton* NewSkeleton, bool bConv
 	if (bConvertSpaces)
 	{
 		USkeleton* OldSkeleton = GetSkeleton();
-		
-		// this will try to do bone to bone mapping
-		if(OldSkeleton && OldSkeleton != NewSkeleton)
+
+		PRAGMA_DISABLE_DEPRECATION_WARNINGS
+		// first check if both has same rig, if so, we'll retarget using it
+		if (OldSkeleton && OldSkeleton->GetRig() != nullptr && NewSkeleton->GetRig() == OldSkeleton->GetRig() && OldSkeleton->GetPreviewMesh() && NewSkeleton->GetPreviewMesh())
+		{
+			const URig* Rig = OldSkeleton->GetRig();
+
+			// we'll have to save the relative space bases transform from old ref pose to new refpose
+			TArray<FTransform> RelativeToNewSpaceBases;
+			// save the ratio of translation change
+			TArray<float> OldToNewTranslationRatio;
+			// create relative transform in component space between old skeleton and new skeleton
+			{
+				// first calculate component space ref pose to get the relative transform between
+				// two ref poses. It is very important update ref pose before getting here. 
+				TArray<FTransform> NewRotations, OldRotations, NewSpaceBases, OldSpaceBases;
+				TArray<FVector> NewTranslations, OldTranslations;
+				TArray<bool> NewTranslationParentFlags, OldTranslationParentFlags;
+				// get the spacebases transform
+				FillUpTransformBasedOnRig(NewSkeleton, NewSpaceBases, NewRotations, NewTranslations, NewTranslationParentFlags);
+				FillUpTransformBasedOnRig(OldSkeleton, OldSpaceBases, OldRotations, OldTranslations, OldTranslationParentFlags);
+
+				// now we'd like to get the relative transform from old to new ref pose in component space
+				// PK2*K2 = PK1*K1*theta where theta => P1*R1*theta = P2*R2 
+				// where	P1 - parent transform in component space for original skeleton
+				//			R1 - local space of the current bone for original skeleton
+				//			P2 - parent transform in component space for new skeleton
+				//			R2 - local space of the current bone for new skeleton
+				// what we're looking for is theta, so that we can apply that to animated transform
+				// this has to have all of nodes since comparing two skeletons, that might have different configuration
+				int32 NumNodes = Rig->GetNodeNum();
+				// saves the theta data per node
+				RelativeToNewSpaceBases.AddUninitialized(NumNodes);
+				// saves the translation conversion datao
+				OldToNewTranslationRatio.AddUninitialized(NumNodes);
+
+				const TArray<FNode>& Nodes = Rig->GetNodes();
+				// calculate the relative transform to new skeleton
+				// so that we can apply the delta in component space
+				for (int32 NodeIndex = 0; NodeIndex < NumNodes; ++NodeIndex)
+				{
+					// theta (RelativeToNewTransform) = (P1*R1)^(-1) * P2*R2 where theta => P1*R1*theta = P2*R2
+					RelativeToNewSpaceBases[NodeIndex] = NewSpaceBases[NodeIndex].GetRelativeTransform(OldSpaceBases[NodeIndex]); 
+
+					// also savees the translation difference between old to new
+					FVector OldTranslation = OldTranslations[NodeIndex];
+					FVector NewTranslation = NewTranslations[NodeIndex];
+
+					// skip root because we don't really have clear relative point to test with it
+					if (NodeIndex != 0 && NewTranslationParentFlags[NodeIndex] == OldTranslationParentFlags[NodeIndex])
+					{
+						// only do this if parent status matches, otherwise, you'll have invalid state 
+						// where one is based on shoulder, where the other is missing the shoulder node
+						float OldTranslationSize = OldTranslation.Size();
+						float NewTranslationSize = NewTranslation.Size();
+
+						OldToNewTranslationRatio[NodeIndex] = (FMath::IsNearlyZero(OldTranslationSize)) ? 1.f/*do not touch new translation size*/ : NewTranslationSize / OldTranslationSize;
+					}
+					else
+					{
+						OldToNewTranslationRatio[NodeIndex] = 1.f; // set to be 1, we don't know what it is
+					}
+
+					UE_LOG(LogAnimation, Verbose, TEXT("Retargeting (%s : %d) : OldtoNewTranslationRatio (%0.2f), Relative Transform (%s)"), *Nodes[NodeIndex].Name.ToString(), NodeIndex, 
+						OldToNewTranslationRatio[NodeIndex], *RelativeToNewSpaceBases[NodeIndex].ToString());
+					UE_LOG(LogAnimation, Verbose, TEXT("\tOldSpaceBase(%s), NewSpaceBase(%s)"), *OldSpaceBases[NodeIndex].ToString(), *NewSpaceBases[NodeIndex].ToString());
+				}
+			}
+
+			FAnimSequenceTrackContainer RiggingAnimationData;
+
+			// now convert animation data to rig data
+			ConvertAnimationDataToRiggingData(RiggingAnimationData);
+
+			// here we have to watch out the index
+			// The RiggingAnimationData will contain only the nodes that are mapped to source skeleton
+			// and here we convert everything that is in RiggingAnimationData which means based on source data
+			// when mapped back to new skeleton, it will discard results that are not mapped to target skeleton
+			
+			TArray<FName> SrcValidNodeNames;
+			int32 SrcNumTracks = OldSkeleton->GetMappedValidNodes(SrcValidNodeNames);
+
+			// now convert to space bases animation 
+			TArray< TArray<FTransform> > ComponentSpaceAnimations, ConvertedLocalSpaceAnimations, ConvertedSpaceAnimations;
+			ComponentSpaceAnimations.AddZeroed(SrcNumTracks);
+			ConvertedSpaceAnimations.AddZeroed(SrcNumTracks);
+			ConvertedLocalSpaceAnimations.AddZeroed(SrcNumTracks);
+
+			const int32 NumKeys = DataModelInterface->GetNumberOfKeys();
+			float Interval = DataModelInterface->GetFrameRate().AsInterval();
+
+			// allocate arrays
+			for(int32 SrcTrackIndex=0; SrcTrackIndex<SrcNumTracks; ++SrcTrackIndex)
+			{
+				ComponentSpaceAnimations[SrcTrackIndex].AddUninitialized(NumKeys);
+				ConvertedLocalSpaceAnimations[SrcTrackIndex].AddUninitialized(NumKeys);
+				ConvertedSpaceAnimations[SrcTrackIndex].AddUninitialized(NumKeys);
+			}
+
+			for (int32 SrcTrackIndex=0; SrcTrackIndex<SrcNumTracks; ++SrcTrackIndex)		
+			{
+				int32 NodeIndex = Rig->FindNode(SrcValidNodeNames[SrcTrackIndex]);
+				check (NodeIndex != INDEX_NONE);
+				auto& RawAnimation = RiggingAnimationData.AnimationTracks[SrcTrackIndex];
+
+				// find rotation parent node
+				int32 RotParentTrackIndex = FindValidTransformParentTrack(Rig, NodeIndex, false, SrcValidNodeNames);
+				int32 TransParentTrackIndex = FindValidTransformParentTrack(Rig, NodeIndex, true, SrcValidNodeNames);
+				// fill up keys - calculate PK1 * K1
+				for(int32 Key=0; Key<NumKeys; ++Key)
+				{
+					FTransform AnimatedLocalKey;
+					//ExtractBoneTransform(RiggingAnimationData.AnimationTracks[SrcTrackIndex], AnimatedLocalKey, Interval * Key);
+					FAnimationUtils::ExtractTransformFromTrack(RiggingAnimationData.AnimationTracks[SrcTrackIndex], Interval* Key, NumberOfSampledKeys, GetPlayLength(), Interpolation, AnimatedLocalKey);
+
+					AnimatedLocalKey.ScaleTranslation(OldToNewTranslationRatio[NodeIndex]);
+
+					if(RotParentTrackIndex != INDEX_NONE)
+					{
+						FQuat ComponentSpaceRotation = ComponentSpaceAnimations[RotParentTrackIndex][Key].GetRotation() * AnimatedLocalKey.GetRotation();
+						ComponentSpaceAnimations[SrcTrackIndex][Key].SetRotation(ComponentSpaceRotation);
+					}
+					else
+					{
+						ComponentSpaceAnimations[SrcTrackIndex][Key].SetRotation(AnimatedLocalKey.GetRotation());
+					}
+
+					if (TransParentTrackIndex != INDEX_NONE)
+					{
+						FVector ComponentSpaceTranslation = ComponentSpaceAnimations[TransParentTrackIndex][Key].TransformPosition(AnimatedLocalKey.GetTranslation());
+						ComponentSpaceAnimations[SrcTrackIndex][Key].SetTranslation(ComponentSpaceTranslation);
+						FVector ParentComponentSpaceScale3D = ComponentSpaceAnimations[TransParentTrackIndex][Key].GetScale3D();
+						ComponentSpaceAnimations[SrcTrackIndex][Key].SetScale3D(ParentComponentSpaceScale3D * AnimatedLocalKey.GetScale3D());
+					}
+					else
+					{
+						ComponentSpaceAnimations[SrcTrackIndex][Key].SetTranslation(AnimatedLocalKey.GetTranslation());
+						ComponentSpaceAnimations[SrcTrackIndex][Key].SetScale3D(AnimatedLocalKey.GetScale3D());
+					}
+				}
+			}
+
+			// now animation is converted to component space
+			TArray<struct FRawAnimSequenceTrack> NewRawAnimationData = RiggingAnimationData.AnimationTracks;
+			for (int32 SrcTrackIndex=0; SrcTrackIndex<SrcNumTracks; ++SrcTrackIndex)
+			{
+				int32 NodeIndex = Rig->FindNode(SrcValidNodeNames[SrcTrackIndex]);
+				// find rotation parent node
+				int32 RotParentTrackIndex = FindValidTransformParentTrack(Rig, NodeIndex, false, SrcValidNodeNames);
+				int32 TransParentTrackIndex = FindValidTransformParentTrack(Rig, NodeIndex, true, SrcValidNodeNames);
+
+				// clear translation;
+				RelativeToNewSpaceBases[NodeIndex].SetTranslation(FVector::ZeroVector);
+
+				for(int32 Key=0; Key<NumKeys; ++Key)
+				{
+					// now convert to the new space and save to local spaces
+					ConvertedSpaceAnimations[SrcTrackIndex][Key] = RelativeToNewSpaceBases[NodeIndex] * ComponentSpaceAnimations[SrcTrackIndex][Key];
+
+					if(RotParentTrackIndex != INDEX_NONE)
+					{
+						FQuat LocalRotation = ConvertedSpaceAnimations[RotParentTrackIndex][Key].GetRotation().Inverse() * ConvertedSpaceAnimations[SrcTrackIndex][Key].GetRotation();
+						ConvertedLocalSpaceAnimations[SrcTrackIndex][Key].SetRotation(LocalRotation);
+					}
+					else
+					{
+						ConvertedLocalSpaceAnimations[SrcTrackIndex][Key].SetRotation(ConvertedSpaceAnimations[SrcTrackIndex][Key].GetRotation());
+					}
+
+					if(TransParentTrackIndex != INDEX_NONE)
+					{
+						FTransform LocalTransform = ConvertedSpaceAnimations[SrcTrackIndex][Key].GetRelativeTransform(ConvertedSpaceAnimations[TransParentTrackIndex][Key]);
+						ConvertedLocalSpaceAnimations[SrcTrackIndex][Key].SetTranslation(LocalTransform.GetLocation());
+						ConvertedLocalSpaceAnimations[SrcTrackIndex][Key].SetScale3D(LocalTransform.GetScale3D());
+					}
+					else
+					{
+						ConvertedLocalSpaceAnimations[SrcTrackIndex][Key].SetTranslation(ConvertedSpaceAnimations[SrcTrackIndex][Key].GetTranslation());
+						ConvertedLocalSpaceAnimations[SrcTrackIndex][Key].SetScale3D(ConvertedSpaceAnimations[SrcTrackIndex][Key].GetScale3D());
+					}
+				}
+
+				auto& RawAnimation = NewRawAnimationData[SrcTrackIndex];
+				RawAnimation.PosKeys.Empty(NumKeys);
+				RawAnimation.PosKeys.AddUninitialized(NumKeys);
+				RawAnimation.RotKeys.Empty(NumKeys);
+				RawAnimation.RotKeys.AddUninitialized(NumKeys);
+				RawAnimation.ScaleKeys.Empty(NumKeys);
+				RawAnimation.ScaleKeys.AddUninitialized(NumKeys);
+
+				for(int32 Key=0; Key<NumKeys; ++Key)
+				{
+					RawAnimation.PosKeys[Key] = FVector3f(ConvertedLocalSpaceAnimations[SrcTrackIndex][Key].GetLocation());
+					RawAnimation.RotKeys[Key] = FQuat4f(ConvertedLocalSpaceAnimations[SrcTrackIndex][Key].GetRotation());
+					RawAnimation.ScaleKeys[Key] = FVector3f(ConvertedLocalSpaceAnimations[SrcTrackIndex][Key].GetScale3D());
+
+					// normalize rotation
+					RawAnimation.RotKeys[Key].Normalize();
+				}
+			}
+
+			RiggingAnimationData.AnimationTracks = MoveTemp(NewRawAnimationData);
+			RiggingAnimationData.TrackNames = MoveTemp(SrcValidNodeNames);
+
+			// set new skeleton
+			SetSkeleton(NewSkeleton);
+
+			// convert back to animated data with new skeleton
+			ConvertRiggingDataToAnimationData(RiggingAnimationData);
+		}
+		PRAGMA_ENABLE_DEPRECATION_WARNINGS
+		// @todo end rig testing
+		// @IMPORTANT: now otherwise this will try to do bone to bone mapping
+		else if(OldSkeleton && OldSkeleton != NewSkeleton)
 		{
 			// Validate animation tracks against the new skeleton, any tracks linked to bones that do not exist in the new hierarchy are removed
 			Controller->RemoveBoneTracksMissingFromSkeleton(NewSkeleton);
@@ -2493,7 +2813,17 @@ void UAnimSequence::ReplaceReferredAnimations(const TMap<UAnimationAsset*, UAnim
 	}
 }
 
-int32 UAnimSequence::GetSpaceBasedAnimationData(TArray< TArray<FTransform> >& AnimationDataInComponentSpace) const
+PRAGMA_DISABLE_DEPRECATION_WARNINGS
+int32 FindParentNodeIndex(URig* Rig, USkeleton* Skeleton, FName ParentNodeName)
+{
+	const int32& ParentNodeIndex = Rig->FindNode(ParentNodeName);
+	const FName& ParentBoneName = Skeleton->GetRigBoneMapping(ParentNodeName);
+	
+	return Skeleton->GetReferenceSkeleton().FindBoneIndex(ParentBoneName);
+}
+PRAGMA_ENABLE_DEPRECATION_WARNINGS
+
+int32 UAnimSequence::GetSpaceBasedAnimationData(TArray< TArray<FTransform> >& AnimationDataInComponentSpace, FAnimSequenceTrackContainer * RiggingAnimationData) const
 {
 	USkeleton* MySkeleton = GetSkeleton();
 
@@ -2515,53 +2845,433 @@ int32 UAnimSequence::GetSpaceBasedAnimationData(TArray< TArray<FTransform> >& An
 	{
 		AnimationDataInComponentSpace[BoneIndex].AddUninitialized(NumKeys);
 	}
-	
-	// now calculating old animated space bases
-	// this one calculates animated space per bones and per key
-	TArray<FTransform> BoneTransforms;
-	for (int32 BoneIndex = 0; BoneIndex < NumBones; ++BoneIndex)
+
+	PRAGMA_DISABLE_DEPRECATION_WARNINGS
+	if (RiggingAnimationData)
 	{
-		const FName BoneName = MySkeleton->GetReferenceSkeleton().GetBoneName(BoneIndex);
-		const int32 ParentBoneIndex = MySkeleton->GetReferenceSkeleton().GetParentIndex(BoneIndex);
+		const URig* Rig = MySkeleton->GetRig();
+
+		check(Rig);
+
+		// to fix the issue where parent of rig doesn't correspond to parent of this skeleton
+		// we do this in multiple iteration if needed. 
+		// this flag will be used to evaluate all of them until done
+		TArray<bool> BoneEvaluated;
+		BoneEvaluated.AddZeroed(NumBones);
+
+		bool bCompleted = false;
+		do
+		{
+			for(int32 BoneIndex = 0; BoneIndex < NumBones; ++BoneIndex)
+			{
+				if ( !BoneEvaluated[BoneIndex] )
+				{
+					const FName& BoneName = RefSkeleton.GetBoneName(BoneIndex);
+					const FName& NodeName = MySkeleton->GetRigNodeNameFromBoneName(BoneName);
+					const FTransformBase* TransformBase = Rig->GetTransformBaseByNodeName(NodeName);
+					const int32 NodeIndex = RiggingAnimationData->TrackNames.Find(NodeName);
+					if(NodeIndex != INDEX_NONE)
+					{
+						check(TransformBase);
+
+						// now calculate the component space
+						const TArray<FRigTransformConstraint>	& RotTransformConstraints = TransformBase->Constraints[EControlConstraint::Type::Orientation].TransformConstraints;
 
-		if (DataModelInterface->IsValidBoneTrackName(BoneName))
+						FQuat ComponentRotation;
+						FTransform ComponentTranslation;
+						FVector ComponentScale;
+
+						// rotation first
+						// this is easy since we just make sure it's evaluated or not
+						{
+							const FName& ParentNodeName = RotTransformConstraints[0].ParentSpace;
+							const FName& ParentBoneName = MySkeleton->GetRigBoneMapping(ParentNodeName);
+							const int32& ParentBoneIndex = RefSkeleton.FindBoneIndex(ParentBoneName);
+
+							if(ParentBoneIndex != INDEX_NONE)
+							{
+								if (BoneEvaluated[ParentBoneIndex])
+								{
+									for(int32 Key = 0; Key < NumKeys; ++Key)
+									{
+										ComponentRotation = AnimationDataInComponentSpace[ParentBoneIndex][Key].GetRotation() * FQuat(RiggingAnimationData->AnimationTracks[NodeIndex].RotKeys[Key]);
+										AnimationDataInComponentSpace[BoneIndex][Key].SetRotation(ComponentRotation);
+									}
+
+									BoneEvaluated[BoneIndex] = true;
+								}
+							}
+							else
+							{
+								for(int32 Key = 0; Key < NumKeys; ++Key)
+								{
+									ComponentRotation = FQuat(RiggingAnimationData->AnimationTracks[NodeIndex].RotKeys[Key]);
+									AnimationDataInComponentSpace[BoneIndex][Key].SetRotation(ComponentRotation);
+								}
+
+								BoneEvaluated[BoneIndex] = true;
+							}
+						}
+
+						const TArray<FRigTransformConstraint>	& PosTransformConstraints = TransformBase->Constraints[EControlConstraint::Type::Translation].TransformConstraints;
+
+						// now time to check translation
+						// this is a bit more complicated
+						// since we have to make sure if it's true to start with
+						// did we succeed on getting rotation?
+						if (BoneEvaluated[BoneIndex])
+						{
+							const FName& ParentNodeName = PosTransformConstraints[0].ParentSpace;
+							const FName& ParentBoneName = MySkeleton->GetRigBoneMapping(ParentNodeName);
+							const int32& ParentBoneIndex = RefSkeleton.FindBoneIndex(ParentBoneName);
+
+							if(ParentBoneIndex != INDEX_NONE)
+							{
+								// this has to be check
+								if (BoneEvaluated[ParentBoneIndex])
+								{
+									for(int32 Key = 0; Key < NumKeys; ++Key)
+									{
+										const FTransform& AnimCompSpace = AnimationDataInComponentSpace[ParentBoneIndex][Key];
+										ComponentTranslation = FTransform(FVector(RiggingAnimationData->AnimationTracks[NodeIndex].PosKeys[Key])) * AnimCompSpace;
+										AnimationDataInComponentSpace[BoneIndex][Key].SetTranslation(ComponentTranslation.GetTranslation());
+
+										ComponentScale = AnimCompSpace.GetScale3D() * FVector(RiggingAnimationData->AnimationTracks[NodeIndex].ScaleKeys[Key]);
+										AnimationDataInComponentSpace[BoneIndex][Key].SetScale3D(ComponentScale);
+									}
+								}
+								else
+								{
+									// if we failed to get parent clear the flag
+									// because if translation has been calculated, BoneEvaluated[BoneIndex] might be true
+									BoneEvaluated[BoneIndex] = false;
+								}
+							}
+							else
+							{
+								for(int32 Key = 0; Key < NumKeys; ++Key)
+								{
+									ComponentTranslation = FTransform(FVector(RiggingAnimationData->AnimationTracks[NodeIndex].PosKeys[Key]));
+									AnimationDataInComponentSpace[BoneIndex][Key].SetTranslation(ComponentTranslation.GetTranslation());
+									
+									ComponentScale = FVector(RiggingAnimationData->AnimationTracks[NodeIndex].ScaleKeys[Key]);
+									AnimationDataInComponentSpace[BoneIndex][Key].SetScale3D(ComponentScale);
+								}
+							}
+						}
+					}
+					else
+					{
+						int32 ParentIndex = RefSkeleton.GetParentIndex(BoneIndex);
+						const FTransform& LocalSpace = RefSkeleton.GetRefBonePose()[BoneIndex];
+						if(ParentIndex != INDEX_NONE)
+						{
+							// if parent is evaluated, do it
+							if (BoneEvaluated[ParentIndex])
+							{
+								for(int32 Key = 0; Key < NumKeys; ++Key)
+								{
+									AnimationDataInComponentSpace[BoneIndex][Key] = LocalSpace * AnimationDataInComponentSpace[ParentIndex][Key];
+								}
+
+								BoneEvaluated[BoneIndex] = true;
+							}
+						}
+						else
+						{
+							BoneEvaluated[BoneIndex] = true;
+
+							for(int32 Key = 0; Key < NumKeys; ++Key)
+							{
+								AnimationDataInComponentSpace[BoneIndex][Key] = LocalSpace;
+							}
+						}
+					}
+				}
+			}
+
+			bCompleted = true;
+			// see if we can get out, brute force for now
+			for(int32 BoneIndex = 0; BoneIndex < NumBones && bCompleted; ++BoneIndex)
+			{
+				bCompleted &= !!BoneEvaluated[BoneIndex];
+			}
+		} while (bCompleted == false);
+	}
+	PRAGMA_ENABLE_DEPRECATION_WARNINGS
+	else
+	{
+		// now calculating old animated space bases
+		// this one calculates aniamted space per bones and per key
+		TArray<FTransform> BoneTransforms;
+		for (int32 BoneIndex = 0; BoneIndex < NumBones; ++BoneIndex)
 		{
-			BoneTransforms.Reset();
-			DataModelInterface->GetBoneTrackTransforms(BoneName, BoneTransforms);
-			// fill up keys - calculate PK1 * K1
-			for (int32 Key = 0; Key < NumKeys; ++Key)
+			const FName BoneName = MySkeleton->GetReferenceSkeleton().GetBoneName(BoneIndex);
+			const int32 ParentBoneIndex = MySkeleton->GetReferenceSkeleton().GetParentIndex(BoneIndex);
+
+			if (DataModelInterface->IsValidBoneTrackName(BoneName))
 			{
-				FTransform AnimatedLocalKey = BoneTransforms[Key];
-				if (ParentBoneIndex != INDEX_NONE)
+				BoneTransforms.Reset();
+				DataModelInterface->GetBoneTrackTransforms(BoneName, BoneTransforms);
+				// fill up keys - calculate PK1 * K1
+				for (int32 Key = 0; Key < NumKeys; ++Key)
 				{
-					AnimationDataInComponentSpace[BoneIndex][Key] = AnimatedLocalKey * AnimationDataInComponentSpace[ParentBoneIndex][Key];
+					FTransform AnimatedLocalKey = BoneTransforms[Key];
+					if (ParentBoneIndex != INDEX_NONE)
+					{
+						AnimationDataInComponentSpace[BoneIndex][Key] = AnimatedLocalKey * AnimationDataInComponentSpace[ParentBoneIndex][Key];
+					}
+					else
+					{
+						AnimationDataInComponentSpace[BoneIndex][Key] = AnimatedLocalKey;
+					}
 				}
-				else
+			}
+			else
+			{
+				// get local spaces from refpose and use that to fill it up
+				FTransform LocalTransform = MySkeleton->GetReferenceSkeleton().GetRefBonePose()[BoneIndex];
+
+				for (int32 Key = 0; Key < NumKeys; ++Key)
+				{
+					if (ParentBoneIndex != INDEX_NONE)
+					{
+						AnimationDataInComponentSpace[BoneIndex][Key] = LocalTransform * AnimationDataInComponentSpace[ParentBoneIndex][Key];
+					}
+					else
+					{
+						AnimationDataInComponentSpace[BoneIndex][Key] = LocalTransform;
+					}
+				}
+			}	
+		}
+
+	}
+
+	return AnimationDataInComponentSpace.Num();
+}
+
+PRAGMA_DISABLE_DEPRECATION_WARNINGS
+bool UAnimSequence::ConvertAnimationDataToRiggingData(FAnimSequenceTrackContainer& RiggingAnimationData)
+{
+	USkeleton* MySkeleton = GetSkeleton();
+	if (MySkeleton && MySkeleton->GetRig())
+	{
+		const URig* Rig = MySkeleton->GetRig();
+		TArray<FName> ValidNodeNames;
+		int32 NumNodes = MySkeleton->GetMappedValidNodes(ValidNodeNames);
+		TArray< TArray<FTransform> > AnimationDataInComponentSpace;
+		int32 NumBones = GetSpaceBasedAnimationData(AnimationDataInComponentSpace, nullptr);
+
+		ValidateModel();
+		const int32 NumModelKeys = DataModelInterface->GetNumberOfKeys();
+
+		if (NumBones > 0)
+		{
+			RiggingAnimationData.Initialize(ValidNodeNames);
+
+			// first we copy all space bases back to it
+			for (int32 NodeIndex = 0; NodeIndex < NumNodes; ++NodeIndex)
+			{
+				struct FRawAnimSequenceTrack& Track = RiggingAnimationData.AnimationTracks[NodeIndex];
+				const FName& NodeName = ValidNodeNames[NodeIndex];
+				const FName& BoneName = MySkeleton->GetRigBoneMapping(NodeName);
+				const int32& BoneIndex = MySkeleton->GetReferenceSkeleton().FindBoneIndex(BoneName);
+
+				if (ensure(BoneIndex != INDEX_NONE))
 				{
-					AnimationDataInComponentSpace[BoneIndex][Key] = AnimatedLocalKey;
+					Track.PosKeys.Empty(NumModelKeys);
+					Track.RotKeys.Empty(NumModelKeys);
+					Track.ScaleKeys.Empty(NumModelKeys);
+					Track.PosKeys.AddUninitialized(NumModelKeys);
+					Track.RotKeys.AddUninitialized(NumModelKeys);
+					Track.ScaleKeys.AddUninitialized(NumModelKeys);
+
+					int32 RigConstraintIndex = Rig->FindTransformBaseByNodeName(NodeName);
+
+					if (RigConstraintIndex != INDEX_NONE)
+					{
+						const auto* RigConstraint = Rig->GetTransformBase(RigConstraintIndex);
+
+						// apply orientation - for now only one
+						const TArray<FRigTransformConstraint>& RotationTransformConstraint = RigConstraint->Constraints[EControlConstraint::Type::Orientation].TransformConstraints;
+
+						if (RotationTransformConstraint.Num() > 0)
+						{
+							const FName& ParentSpace = RotationTransformConstraint[0].ParentSpace;
+							const FName& ParentBoneName = MySkeleton->GetRigBoneMapping(ParentSpace);
+							const int32& ParentBoneIndex = MySkeleton->GetReferenceSkeleton().FindBoneIndex(ParentBoneName);
+							if (ParentBoneIndex != INDEX_NONE)
+							{
+								// if no rig control, component space is used
+								for (int32 KeyIndex = 0; KeyIndex < NumModelKeys; ++KeyIndex)
+								{
+									FTransform ParentTransform = AnimationDataInComponentSpace[ParentBoneIndex][KeyIndex];
+									FTransform RelativeTransform = AnimationDataInComponentSpace[BoneIndex][KeyIndex].GetRelativeTransform(ParentTransform);
+									Track.RotKeys[KeyIndex] = FQuat4f(RelativeTransform.GetRotation());
+								}
+							}
+							else
+							{
+								// if no rig control, component space is used
+								for (int32 KeyIndex = 0; KeyIndex < NumModelKeys; ++KeyIndex)
+								{
+									Track.RotKeys[KeyIndex] = FQuat4f(AnimationDataInComponentSpace[BoneIndex][KeyIndex].GetRotation());
+								}
+							}
+						}
+						else
+						{
+							// if no rig control, component space is used
+							for (int32 KeyIndex = 0; KeyIndex < NumModelKeys; ++KeyIndex)
+							{
+								Track.RotKeys[KeyIndex] = FQuat4f(AnimationDataInComponentSpace[BoneIndex][KeyIndex].GetRotation());
+							}
+						}
+
+						// apply translation - for now only one
+						const TArray<FRigTransformConstraint>& TranslationTransformConstraint = RigConstraint->Constraints[EControlConstraint::Type::Translation].TransformConstraints;
+
+						if (TranslationTransformConstraint.Num() > 0)
+						{
+							const FName& ParentSpace = TranslationTransformConstraint[0].ParentSpace;
+							const FName& ParentBoneName = MySkeleton->GetRigBoneMapping(ParentSpace);
+							const int32& ParentBoneIndex = MySkeleton->GetReferenceSkeleton().FindBoneIndex(ParentBoneName);
+							if (ParentBoneIndex != INDEX_NONE)
+							{
+								// if no rig control, component space is used
+								for (int32 KeyIndex = 0; KeyIndex < NumModelKeys; ++KeyIndex)
+								{
+									FTransform ParentTransform = AnimationDataInComponentSpace[ParentBoneIndex][KeyIndex];
+									FTransform RelativeTransform = AnimationDataInComponentSpace[BoneIndex][KeyIndex].GetRelativeTransform(ParentTransform);
+									Track.PosKeys[KeyIndex] = FVector3f(RelativeTransform.GetTranslation());
+									Track.ScaleKeys[KeyIndex] = FVector3f(RelativeTransform.GetScale3D());
+								}
+							}
+							else
+							{
+								for (int32 KeyIndex = 0; KeyIndex < NumModelKeys; ++KeyIndex)
+								{
+									Track.PosKeys[KeyIndex] = FVector3f(AnimationDataInComponentSpace[BoneIndex][KeyIndex].GetTranslation());
+									Track.ScaleKeys[KeyIndex] = FVector3f(AnimationDataInComponentSpace[BoneIndex][KeyIndex].GetScale3D());
+								}
+							}
+						}
+						else
+						{
+							for (int32 KeyIndex = 0; KeyIndex < NumModelKeys; ++KeyIndex)
+							{
+								Track.PosKeys[KeyIndex] = FVector3f(AnimationDataInComponentSpace[BoneIndex][KeyIndex].GetTranslation());
+								Track.ScaleKeys[KeyIndex] = FVector3f(AnimationDataInComponentSpace[BoneIndex][KeyIndex].GetScale3D());
+							}
+						}
+					}
+					else
+					{
+						// if no rig control, component space is used
+						for (int32 KeyIndex = 0; KeyIndex < NumModelKeys; ++KeyIndex)
+						{
+							Track.PosKeys[KeyIndex] = FVector3f(AnimationDataInComponentSpace[BoneIndex][KeyIndex].GetTranslation());
+							Track.RotKeys[KeyIndex] = FQuat4f(AnimationDataInComponentSpace[BoneIndex][KeyIndex].GetRotation());
+							Track.ScaleKeys[KeyIndex] = FVector3f(AnimationDataInComponentSpace[BoneIndex][KeyIndex].GetScale3D());
+						}
+					}
 				}
 			}
 		}
-		else
+
+		return true;
+	}
+
+	return false;
+}
+PRAGMA_ENABLE_DEPRECATION_WARNINGS
+
+bool UAnimSequence::ConvertRiggingDataToAnimationData(FAnimSequenceTrackContainer& RiggingAnimationData)
+{
+	if (RiggingAnimationData.GetNum() > 0)
+	{
+		ValidateModel();
+
+		IAnimationDataController::FScopedBracket ScopedBracket(Controller, LOCTEXT("ConvertRiggingDataToAnimationData_Bracket", "Converting Rigging Data to Animation Data"));
+
+		TArray< TArray<FTransform> > AnimationDataInComponentSpace;
+		int32 NumBones = GetSpaceBasedAnimationData(AnimationDataInComponentSpace, &RiggingAnimationData);
+
+		USkeleton* MySkeleton = GetSkeleton();
+		TArray<FName> ValidNodeNames;
+		MySkeleton->GetMappedValidNodes(ValidNodeNames);
+		// remove from ValidNodeNames if it doesn't belong to AnimationTrackNames
+		for (int32 NameIndex=0; NameIndex<ValidNodeNames.Num(); ++NameIndex)
+		{
+			if (RiggingAnimationData.TrackNames.Contains(ValidNodeNames[NameIndex]) == false)
+			{
+				ValidNodeNames.RemoveAt(NameIndex);
+				--NameIndex;
+			}
+		}
+
+		int32 ValidNumNodes = ValidNodeNames.Num();
+
+		// get local spaces
+		// add all tracks?
+		Controller->RemoveAllBoneTracks();
+
+		// Reused track data arrays
+		TArray<FVector3f> PosKeys;
+		TArray<FQuat4f> RotKeys;
+		TArray<FVector3f> ScaleKeys;
+
+		const int32 NumModelKeys = DataModelInterface->GetNumberOfKeys();
+		const FReferenceSkeleton& RefSkeleton = MySkeleton->GetReferenceSkeleton();
+		for (int32 NodeIndex = 0; NodeIndex < ValidNumNodes; ++NodeIndex)
 		{
-			// get local spaces from refpose and use that to fill it up
-			FTransform LocalTransform = MySkeleton->GetReferenceSkeleton().GetRefBonePose()[BoneIndex];
+			FName BoneName = MySkeleton->GetRigBoneMapping(ValidNodeNames[NodeIndex]);
+			int32 BoneIndex = RefSkeleton.FindBoneIndex(BoneName);
 
-			for (int32 Key = 0; Key < NumKeys; ++Key)
+			if (BoneIndex != INDEX_NONE)
 			{
-				if (ParentBoneIndex != INDEX_NONE)
+				// Add a new track by name
+				Controller->AddBoneCurve(BoneName);
+
+				PosKeys.SetNumUninitialized(NumModelKeys);
+				RotKeys.SetNumUninitialized(NumModelKeys);
+				ScaleKeys.SetNumUninitialized(NumModelKeys);
+
+				const int32& ParentBoneIndex = RefSkeleton.GetParentIndex(BoneIndex);
+
+				if(ParentBoneIndex != INDEX_NONE)
 				{
-					AnimationDataInComponentSpace[BoneIndex][Key] = LocalTransform * AnimationDataInComponentSpace[ParentBoneIndex][Key];
+					for(int32 KeyIndex = 0; KeyIndex < NumModelKeys; ++KeyIndex)
+					{
+						const FTransform LocalTransform = AnimationDataInComponentSpace[BoneIndex][KeyIndex].GetRelativeTransform(AnimationDataInComponentSpace[ParentBoneIndex][KeyIndex]);
+
+						PosKeys[KeyIndex] = FVector3f(LocalTransform.GetTranslation());
+						RotKeys[KeyIndex] = FQuat4f(LocalTransform.GetRotation());
+						ScaleKeys[KeyIndex] = FVector3f(LocalTransform.GetScale3D());
+					}
 				}
 				else
 				{
-					AnimationDataInComponentSpace[BoneIndex][Key] = LocalTransform;
+					for(int32 KeyIndex = 0; KeyIndex < NumModelKeys; ++KeyIndex)
+					{
+						const FTransform LocalTransform = AnimationDataInComponentSpace[BoneIndex][KeyIndex];
+
+						PosKeys[KeyIndex] = FVector3f(LocalTransform.GetTranslation());
+						RotKeys[KeyIndex] = FQuat4f(LocalTransform.GetRotation());
+						ScaleKeys[KeyIndex] = FVector3f(LocalTransform.GetScale3D());
+					}
 				}
+
+				// Set the track's key data
+				Controller->SetBoneTrackKeys(BoneName, PosKeys, RotKeys, ScaleKeys);
 			}
-		}	
+		}
+
+		return true;
 	}
 
-	return AnimationDataInComponentSpace.Num();
+	return false;
 }
 
 void UAnimSequence::AddKeyToSequence(float Time, const FName& BoneName, const FTransform& AdditiveTransform)
diff --git a/Engine/Source/Runtime/Engine/Private/Animation/Rig.cpp b/Engine/Source/Runtime/Engine/Private/Animation/Rig.cpp
new file mode 100644
--- /dev/null
+++ b/Engine/Source/Runtime/Engine/Private/Animation/Rig.cpp
@@ -0,0 +1,521 @@
+// Copyright Epic Games, Inc. All Rights Reserved.
+
+/*=============================================================================
+	URig.cpp: Rig functionality for sharing animations
+=============================================================================*/ 
+
+#include "Animation/Rig.h"
+#include "Animation/Skeleton.h"
+#include "UObject/FrameworkObjectVersion.h"
+#include "AnimationRuntime.h"
+
+#include UE_INLINE_GENERATED_CPP_BY_NAME(Rig)
+
+//@todo should move all this window stuff somewhere else. Persona?
+
+// January 2023 - Whole compilation unit is being deprecated, so disabled warnings for whole thing. 
+PRAGMA_DISABLE_DEPRECATION_WARNINGS
+
+#define LOCTEXT_NAMESPACE "Rig"
+
+FName URig::WorldNodeName(TEXT("World"));
+
+URig::URig(const FObjectInitializer& ObjectInitializer)
+	: Super(ObjectInitializer)
+{
+}
+
+#if WITH_EDITOR
+void URig::PostEditChangeProperty(FPropertyChangedEvent& PropertyChangedEvent)
+{
+	if (Nodes.Num() != TransformBases.Num())
+	{
+		int32 NodeNum = Nodes.Num();
+		int32 TransformBasesNum = TransformBases.Num();
+
+		// make sure to add name of nodes if not set
+		for(int32 NewId=TransformBasesNum; NewId <NodeNum; ++NewId)
+		{
+			if (Nodes[NewId].Name == NAME_None)
+			{
+				int32 UniqueIndex = 1;
+
+				while (1) 
+				{
+					// assign name of custom_# 
+					FName PotentialName = *FString::Printf(TEXT("Custom_%d"), UniqueIndex++);
+					if (FindNode(PotentialName) == INDEX_NONE)
+					{
+						Nodes[NewId].Name = PotentialName;
+						Nodes[NewId].ParentName = WorldNodeName;
+						Nodes[NewId].DisplayName = PotentialName.ToString();
+						break;
+					}
+				} 
+			}
+		}
+
+		// if their size doesn't match, make sure to add it
+		if (NodeNum < TransformBasesNum)
+		{
+			// just remove the last elements
+			TransformBases.RemoveAt(NodeNum, TransformBasesNum-NodeNum);
+		}
+		else
+		{
+			for (int32 NewId=TransformBasesNum; NewId <NodeNum; ++NewId)
+			{
+				AddRigConstraint(Nodes[NewId].Name, EControlConstraint::Translation, EConstraintTransform::Absolute, WorldNodeName, 1.f);
+				AddRigConstraint(Nodes[NewId].Name, EControlConstraint::Orientation, EConstraintTransform::Absolute, WorldNodeName, 1.f);
+			}
+		}
+	}
+
+	Super::PostEditChangeProperty(PropertyChangedEvent);
+}
+
+// for now these are privates since we don't have good control yet
+bool URig::AddNode(FName Name, FName ParentNode, FTransform Transform)
+{
+	int32 NodeIndex = FindNode(Name);
+
+	if (NodeIndex == INDEX_NONE)
+	{
+		Nodes.Add(FNode(Name, ParentNode, Transform));
+
+		return true;
+	}
+
+	return false;
+}
+
+bool URig::DeleteNode(FName Name)
+{
+	int32 NodeIndex = FindNode(Name);
+
+	if (NodeIndex != INDEX_NONE)
+	{
+		Nodes.RemoveAt(NodeIndex);
+		return true;
+	}
+
+	return false;
+}
+
+int32 URig::FindNode(const FName& NodeName) const
+{
+	int32 Index=0;
+	for (auto Node: Nodes)
+	{
+		if (Node.Name == NodeName)
+		{
+			return Index;
+		}
+
+		++Index;
+	}
+
+	return INDEX_NONE;
+}
+
+const FNode* URig::GetNode(int32 NodeIndex) const
+{
+	if ( Nodes.IsValidIndex(NodeIndex) )
+	{
+		return &Nodes[NodeIndex];
+	}
+
+	return NULL;
+}
+
+FName URig::GetNodeName(int32 NodeIndex) const
+{
+	if(Nodes.IsValidIndex(NodeIndex))
+	{
+		return Nodes[NodeIndex].Name;
+	}
+
+	return NAME_None;
+}
+
+FName URig::GetParentNode(FName& NodeName) const
+{
+	int32 NodeIndex = FindNode(NodeName);
+
+	if (NodeIndex != INDEX_NONE)
+	{
+		if (Nodes[NodeIndex].ParentName != NAME_None)
+		{
+			return Nodes[NodeIndex].ParentName;
+		}
+	}
+
+	return WorldNodeName;
+}
+
+bool URig::AddRigConstraint(FName NodeName, EControlConstraint::Type ConstraintType, EConstraintTransform::Type	TransformType, FName ParentSpace, float Weight /*= 1.f*/)
+{
+	if (ConstraintType == EControlConstraint::Type::MAX)
+	{
+		// invalid type
+		return false;
+	}
+
+	// make sure the ParentSpace is valid
+	int32 ParentIndex = FindNode(ParentSpace);
+	if (ParentIndex == INDEX_NONE)
+	{
+		// if parent is invalid, set it to World
+		ParentSpace = WorldNodeName;
+	}
+
+	int32 Index = FindTransformBaseByNodeName(NodeName);
+
+	if (Index == INDEX_NONE)
+	{
+		FTransformBase NewTransformBase;
+		NewTransformBase.Node = NodeName;
+
+		FRigTransformConstraint NewTransformConstraint;
+		NewTransformConstraint.TranformType = TransformType;
+		NewTransformConstraint.ParentSpace = ParentSpace;
+		NewTransformConstraint.Weight = Weight;
+
+		NewTransformBase.Constraints[ConstraintType].TransformConstraints.Add(NewTransformConstraint);
+
+		TransformBases.Add(NewTransformBase);
+	}
+	else
+	{
+		// it exists already, need to make sure we don't have different constraint types
+		FTransformBase & TransformBase = TransformBases[Index];
+
+		FTransformBaseConstraint & ControlConstraint = TransformBase.Constraints[ConstraintType];
+
+		FRigTransformConstraint NewTransformConstraint;
+		NewTransformConstraint.TranformType = TransformType;
+		NewTransformConstraint.ParentSpace = ParentSpace;
+		NewTransformConstraint.Weight = Weight;
+
+		// add new transform constraint
+		ControlConstraint.TransformConstraints.Add(NewTransformConstraint);
+	}
+
+	return true;
+}
+
+int32 URig::GetNodeNum() const
+{
+	return Nodes.Num();
+}
+
+int32 URig::GetTransformBaseNum() const
+{
+	return TransformBases.Num();
+}
+
+const FTransformBase* URig::GetTransformBase(int32 TransformBaseIndex) const
+{
+	if (TransformBases.IsValidIndex(TransformBaseIndex))
+	{
+		return &TransformBases[TransformBaseIndex];
+	}
+
+	return NULL;
+}
+
+const FTransformBase* URig::GetTransformBaseByNodeName(FName NodeName) const
+{
+	int32 TransformBaseIndex = FindTransformBaseByNodeName(NodeName);
+
+	if (TransformBaseIndex != INDEX_NONE)
+	{
+		return &TransformBases[TransformBaseIndex];
+	}
+
+	return NULL;
+}
+
+int32 URig::FindTransformParentNode(int32 NodeIndex, bool bTranslate, int32 Index/*=0*/) const
+{
+	if (Nodes.IsValidIndex(NodeIndex))
+	{
+		const FTransformBase* TransformBase = GetTransformBaseByNodeName(Nodes[NodeIndex].Name);
+
+		if(TransformBase)
+		{
+			FName ParentNodeName = NAME_None;
+			if(bTranslate)
+			{
+				ParentNodeName = TransformBase->Constraints[EControlConstraint::Type::Translation].TransformConstraints[Index].ParentSpace;
+			}
+			else
+			{
+				ParentNodeName = TransformBase->Constraints[EControlConstraint::Type::Orientation].TransformConstraints[Index].ParentSpace;
+			}
+
+			if(ParentNodeName != NAME_None)
+			{
+				return  FindNode(ParentNodeName);
+			}
+		}
+	}
+
+	return INDEX_NONE;
+}
+
+int32 URig::FindTransformBaseByNodeName(FName NodeName) const
+{
+	int32 Index=0;
+
+	for (auto TransformBase : TransformBases)
+	{
+		if (TransformBase.Node == NodeName)
+		{
+			return Index;
+		}
+
+		++Index;
+	}
+
+	return INDEX_NONE;
+}
+
+void URig::CreateFromSkeleton(const USkeleton* Skeleton, const TMap<int32, int32> & RequiredBones)
+{
+	// show dialog to choose which bone to add
+	if(RequiredBones.Num() >0)
+	{
+		const FReferenceSkeleton& RefSkeleton = Skeleton->GetReferenceSkeleton();
+		SourceSkeleton = RefSkeleton;
+		TArray<FTransform> SpaceBaseRefPose;
+
+		FAnimationRuntime::FillUpComponentSpaceTransformsRefPose(Skeleton, SpaceBaseRefPose);
+
+		// once selected, add node to the rig
+		for (auto It = RequiredBones.CreateConstIterator(); It; ++It)
+		{
+			int32 BoneIndex = (int32)(It.Key());
+			int32 ParentIndex = (int32)(It.Value());
+
+			check (BoneIndex != INDEX_NONE);
+
+			FName BoneName = RefSkeleton.GetBoneName(BoneIndex);
+			FName ParentBoneName = WorldNodeName;
+			if (ParentIndex != INDEX_NONE)
+			{
+				ParentBoneName = RefSkeleton.GetBoneName(ParentIndex);
+			}
+
+			const FTransform& Transform = SpaceBaseRefPose[BoneIndex];
+
+			AddNode(BoneName, ParentBoneName, Transform);
+		}
+
+		// add constraint to parent space and relative transform 
+		for(auto It = RequiredBones.CreateConstIterator(); It; ++It)
+		{
+			int32 BoneIndex = (int32)(It.Key());
+			int32 ParentIndex = (int32)(It.Value());
+
+			FName BoneName = RefSkeleton.GetBoneName(BoneIndex);
+
+			if (ParentIndex == INDEX_NONE)
+			{
+				AddRigConstraint(BoneName, EControlConstraint::Translation, EConstraintTransform::Absolute, WorldNodeName, 1.f);
+				AddRigConstraint(BoneName, EControlConstraint::Orientation, EConstraintTransform::Absolute, WorldNodeName, 1.f);
+			}
+			else
+			{
+				FName ParentBoneName = RefSkeleton.GetBoneName(ParentIndex);
+				AddRigConstraint(BoneName, EControlConstraint::Translation, EConstraintTransform::Absolute, ParentBoneName, 1.f);
+				AddRigConstraint(BoneName, EControlConstraint::Orientation, EConstraintTransform::Absolute, ParentBoneName, 1.f);
+			}
+		}
+	}
+}
+
+/*
+void URig::CalculateComponentSpace(int32 NodeIndex, const FTransform& LocalTransform, const TArray<FTransform> & TransformBuffer, const FGetParentIndex& DelegateToGetParentIndex, FTransform& OutComponentSpaceTransform) const
+{
+	int32 ConstraintIndex = FindTransformBaseByNodeName(Nodes[NodeIndex].Name);
+
+	if (ConstraintIndex != INDEX_NONE)
+	{
+		// now find transform constraint data
+		{
+			const FTransformBaseConstraint& Constraints = TransformBases[ConstraintIndex].Constraints[EControlConstraint::Type::Orientation];
+
+			// for now we only care for the first one
+			const FName ParentName = Constraints.TransformConstraints[0].ParentSpace;
+			int32 ParentIndex = INDEX_NONE;
+
+			if(DelegateToGetParentIndex.IsBound())
+			{
+				ParentIndex = DelegateToGetParentIndex.Execute(ParentName);
+			}
+			else
+			{
+				ParentIndex = FindTransformParentNode(NodeIndex, false);
+			}
+
+			if(ParentIndex != INDEX_NONE && TransformBuffer.IsValidIndex(ParentIndex))
+			{
+				FQuat ParentRotation = TransformBuffer[ParentIndex].GetRotation();
+				OutComponentSpaceTransform.SetRotation(ParentRotation * LocalTransform.GetRotation());
+			}
+		}
+
+		// same thing for translation
+		{
+			const FTransformBaseConstraint& Constraints = TransformBases[ConstraintIndex].Constraints[EControlConstraint::Type::Translation];
+
+			// for now we only care for the first one
+			const FName ParentName = Constraints.TransformConstraints[0].ParentSpace;
+			int32 ParentIndex = INDEX_NONE;
+
+			if(DelegateToGetParentIndex.IsBound())
+			{
+				ParentIndex = DelegateToGetParentIndex.Execute(ParentName);
+			}
+			else
+			{
+				ParentIndex = FindTransformParentNode(NodeIndex, false);
+			}
+
+			if(ParentIndex != INDEX_NONE && TransformBuffer.IsValidIndex(ParentIndex))
+			{
+				FTransform ParentTransform = TransformBuffer[ParentIndex];
+				OutComponentSpaceTransform.SetTranslation(ParentTransform.TransformVector(LocalTransform.GetTranslation()));
+			}
+		}
+		// @todo fix this
+		OutComponentSpaceTransform.SetScale3D(LocalTransform.GetScale3D());
+	}
+}
+
+void URig::CalculateLocalSpace(int32 NodeIndex, const FTransform& ComponentTransform, const TArray<FTransform> & TransformBuffer, const FGetParentIndex& DelegateToGetParentIndex, FTransform& OutLocalSpaceTransform) const
+{
+	int32 ConstraintIndex = FindTransformBaseByNodeName(Nodes[NodeIndex].Name);
+
+	if(ConstraintIndex != INDEX_NONE)
+	{
+		// now find transform constraint data
+		{
+			const FTransformBaseConstraint& Constraints = TransformBases[ConstraintIndex].Constraints[EControlConstraint::Type::Orientation];
+
+			// for now we only care for the first one
+			const FName ParentName = Constraints.TransformConstraints[0].ParentSpace;
+			int32 ParentIndex = INDEX_NONE;
+
+			if(DelegateToGetParentIndex.IsBound())
+			{
+				ParentIndex = DelegateToGetParentIndex.Execute(ParentName);
+			}
+			else
+			{
+				ParentIndex = FindTransformParentNode(NodeIndex, false);
+			}
+
+			if(ParentIndex != INDEX_NONE && TransformBuffer.IsValidIndex(ParentIndex))
+			{
+				FQuat ParentRotation = TransformBuffer[ParentIndex].GetRotation();
+				OutLocalSpaceTransform.SetRotation(ParentRotation.Inverse() * ComponentTransform.GetRotation());
+			}
+		}
+
+		// same thing for translation
+		{
+			const FTransformBaseConstraint& Constraints = TransformBases[ConstraintIndex].Constraints[EControlConstraint::Type::Translation];
+
+			// for now we only care for the first one
+			const FName ParentName = Constraints.TransformConstraints[0].ParentSpace;
+			int32 ParentIndex = INDEX_NONE;
+
+			if(DelegateToGetParentIndex.IsBound())
+			{
+				ParentIndex = DelegateToGetParentIndex.Execute(ParentName);
+			}
+			else
+			{
+				ParentIndex = FindTransformParentNode(NodeIndex, false);
+			}
+
+			if(ParentIndex != INDEX_NONE && TransformBuffer.IsValidIndex(ParentIndex))
+			{
+				FTransform LocalTransform = ComponentTransform.GetRelativeTransform(TransformBuffer[ParentIndex]);
+				OutLocalSpaceTransform.SetTranslation(LocalTransform.GetTranslation());
+			}
+		}
+
+		// @todo fix this
+		OutLocalSpaceTransform.SetScale3D(ComponentTransform.GetScale3D());
+	}
+}*/
+
+void URig::SetAllConstraintsToParents()
+{
+	for(auto & Control : TransformBases)
+	{
+		FName ParentNode = GetParentNode(Control.Node);
+
+		Control.Constraints[EControlConstraint::Type::Translation].TransformConstraints[0].ParentSpace = ParentNode;
+		Control.Constraints[EControlConstraint::Type::Orientation].TransformConstraints[0].ParentSpace = ParentNode;
+	}
+}
+void URig::SetAllConstraintsToWorld()
+{
+	for (auto & Control : TransformBases)
+	{
+		Control.Constraints[EControlConstraint::Type::Translation].TransformConstraints[0].ParentSpace = WorldNodeName;
+		Control.Constraints[EControlConstraint::Type::Orientation].TransformConstraints[0].ParentSpace = WorldNodeName;
+	}
+}
+
+#if WITH_EDITORONLY_DATA
+void URig::SetSourceReferenceSkeleton(const FReferenceSkeleton& InSrcSkeleton)
+{
+	SourceSkeleton = InSrcSkeleton;
+	MarkPackageDirty();
+}
+#endif // WITH_EDITORONLY_DATA
+#endif // WITH_EDITOR
+
+void URig::Serialize(FArchive& Ar)
+{
+	Super::Serialize(Ar);
+
+	Ar.UsingCustomVersion(FFrameworkObjectVersion::GUID);
+	if (Ar.CustomVer(FFrameworkObjectVersion::GUID) >= FFrameworkObjectVersion::AddSourceReferenceSkeletonToRig)
+	{
+#if WITH_EDITORONLY_DATA
+		Ar << SourceSkeleton;
+#else
+		FReferenceSkeleton Dummy;
+		Ar << Dummy;
+#endif 
+	}
+}
+
+void URig::GetMappableNodeData(TArray<FName>& OutNames, TArray<FNodeItem>& OutNodeItems) const
+{
+	const int32 NodeCount = Nodes.Num(); 
+	OutNames.Reset(NodeCount);
+	OutNodeItems.Reset(NodeCount);
+
+	if (NodeCount > 0)
+	{
+		OutNames.AddDefaulted(NodeCount);
+		OutNodeItems.AddDefaulted(NodeCount);
+
+		for (int32 Index = 0; Index < NodeCount; ++Index)
+		{
+			OutNames[Index] = Nodes[Index].Name;
+			OutNodeItems[Index] = FNodeItem(Nodes[Index].ParentName, Nodes[Index].Transform);
+		}
+	}
+}
+#undef LOCTEXT_NAMESPACE 
+
+// January 2023 - Whole compilation unit is being deprecated, so disabled warnings for whole thing. 
+PRAGMA_ENABLE_DEPRECATION_WARNINGS
\ No newline at end of file
diff --git a/Engine/Source/Runtime/Engine/Private/Animation/Skeleton.cpp b/Engine/Source/Runtime/Engine/Private/Animation/Skeleton.cpp
--- a/Engine/Source/Runtime/Engine/Private/Animation/Skeleton.cpp
+++ b/Engine/Source/Runtime/Engine/Private/Animation/Skeleton.cpp
@@ -17,6 +17,7 @@
 #include "AssetRegistry/AssetData.h"
 #include "AssetRegistry/ARFilter.h"
 #include "AssetRegistry/AssetRegistryModule.h"
+#include "Animation/Rig.h"
 #include "Animation/BlendProfile.h"
 #include "Engine/SkinnedAsset.h"
 #include "Logging/MessageLog.h"
@@ -52,6 +53,8 @@ const FString USkeleton::CurveTagDelimiter = TEXT(";");
 
 const FName USkeleton::CompatibleSkeletonsNameTag = FName(TEXT("CompatibleSkeletonList"));
 const FString USkeleton::CompatibleSkeletonsTagDelimiter = TEXT(";");
+
+const FName USkeleton::RigTag = FName(TEXT("Rig"));
 #endif 
 
 #if WITH_EDITORONLY_DATA
@@ -1951,9 +1954,169 @@ void USkeleton::HandleVirtualBoneChanges()
 
 #if WITH_EDITOR
 
+PRAGMA_DISABLE_DEPRECATION_WARNINGS
+
+void USkeleton::SetRigConfig(URig * Rig)
+{
+	if (RigConfig.Rig != Rig)
+	{
+		RigConfig.Rig = Rig;
+		RigConfig.BoneMappingTable.Empty();
+
+		if (Rig)
+		{
+			const FReferenceSkeleton& RefSkeleton = GetReferenceSkeleton();
+			const TArray<FNode> & Nodes = Rig->GetNodes();
+			// now add bone mapping table
+			for (auto Node: Nodes)
+			{
+				// if find same bone, use that bone for mapping
+				if (RefSkeleton.FindBoneIndex(Node.Name) != INDEX_NONE)
+				{
+					RigConfig.BoneMappingTable.Add(FNameMapping(Node.Name, Node.Name));
+				}
+				else
+				{
+					RigConfig.BoneMappingTable.Add(FNameMapping(Node.Name));
+				}
+			}
+		}
+	}
+}
+
+int32 USkeleton::FindRigBoneMapping(const FName& NodeName) const
+{
+	int32 Index=0;
+	for(const auto & NameMap : RigConfig.BoneMappingTable)
+	{
+		if(NameMap.NodeName == NodeName)
+		{
+			return Index;
+		}
+
+		++Index;
+	}
+
+	return INDEX_NONE;
+}
+
+FName USkeleton::GetRigBoneMapping(const FName& NodeName) const
+{
+	int32 Index = FindRigBoneMapping(NodeName);
+
+	if (Index != INDEX_NONE)
+	{
+		return RigConfig.BoneMappingTable[Index].BoneName;
+	}
+
+	return NAME_None;
+}
+
+FName USkeleton::GetRigNodeNameFromBoneName(const FName& BoneName) const
+{
+	for(const auto & NameMap : RigConfig.BoneMappingTable)
+	{
+		if(NameMap.BoneName == BoneName)
+		{
+			return NameMap.NodeName;
+		}
+	}
+
+	return NAME_None;
+}
+
+int32 USkeleton::GetMappedValidNodes(TArray<FName> &OutValidNodeNames)
+{
+	OutValidNodeNames.Empty();
+
+	for (auto Entry : RigConfig.BoneMappingTable)
+	{
+		if (Entry.BoneName != NAME_None)
+		{
+			OutValidNodeNames.Add(Entry.NodeName);
+		}
+	}
+
+	return OutValidNodeNames.Num();
+}
+
+bool USkeleton::SetRigBoneMapping(const FName& NodeName, FName BoneName)
+{
+	// make sure it's valid
+	int32 BoneIndex = ReferenceSkeleton.FindBoneIndex(BoneName);
+
+	// @todo we need to have validation phase where you can't set same bone for different nodes
+	// but it might be annoying to do that right now since the tool is ugly
+	// so for now it lets you set everything, but in the future
+	// we'll have to add verification
+	if ( BoneIndex == INDEX_NONE )
+	{
+		BoneName = NAME_None;
+	}
+
+	int32 Index = FindRigBoneMapping(NodeName);
+
+	if(Index != INDEX_NONE)
+	{
+		RigConfig.BoneMappingTable[Index].BoneName = BoneName;
+		return true;
+	}
+
+	return false;
+}
+
+void USkeleton::RefreshRigConfig()
+{
+	if (RigConfig.Rig != NULL)
+	{
+		if (RigConfig.BoneMappingTable.Num() > 0)
+		{
+			// verify if any missing bones or anything
+			// remove if removed
+			for ( int32 TableId=0; TableId<RigConfig.BoneMappingTable.Num(); ++TableId )
+			{
+				auto & BoneMapping = RigConfig.BoneMappingTable[TableId];
+
+				if ( RigConfig.Rig->FindNode(BoneMapping.NodeName) == INDEX_NONE)
+				{
+					// if not contains, remove it
+					RigConfig.BoneMappingTable.RemoveAt(TableId);
+					--TableId;
+				}
+			}
+
+			// if the count doesn't match, there is missing nodes. 
+			if (RigConfig.Rig->GetNodeNum() != RigConfig.BoneMappingTable.Num())
+			{
+				int32 NodeNum = RigConfig.Rig->GetNodeNum();
+				for(int32 NodeId=0; NodeId<NodeNum; ++NodeId)
+				{
+					const auto* Node = RigConfig.Rig->GetNode(NodeId);
+
+					if (FindRigBoneMapping(Node->Name) == INDEX_NONE)
+					{
+						RigConfig.BoneMappingTable.Add(FNameMapping(Node->Name));
+					}
+				}
+			}
+		}
+	}
+}
+
+URig * USkeleton::GetRig() const
+{
+	return RigConfig.Rig;
+}
+PRAGMA_ENABLE_DEPRECATION_WARNINGS
+
 void USkeleton::GetAssetRegistryTags(TArray<FAssetRegistryTag>& OutTags) const
 {
 	Super::GetAssetRegistryTags(OutTags);
+	PRAGMA_DISABLE_DEPRECATION_WARNINGS
+	FString RigFullName = (RigConfig.Rig)? RigConfig.Rig->GetFullName() : TEXT("");
+	PRAGMA_ENABLE_DEPRECATION_WARNINGS
+
+	OutTags.Add(FAssetRegistryTag(USkeleton::RigTag, RigFullName, FAssetRegistryTag::TT_Hidden));
 
 	TStringBuilder<256> CompatibleSkeletonsBuilder;
 	
