#ue5 - Added option for ProjectileMovementComponent to avoid interpolation transform updates on some frames when an object isn't as relevant. Default version is based on not recently rendered but can be extended.

#tests tested pickups, simple projectile weapons, complex projectiles. Tested networked with 1 and 2 players and with logging on to show frame skipping of interpolation when looking various directions.

diff --git a/Engine/Source/Runtime/Engine/Classes/GameFramework/ProjectileMovementComponent.h b/Engine/Source/Runtime/Engine/Classes/GameFramework/ProjectileMovementComponent.h
--- a/Engine/Source/Runtime/Engine/Classes/GameFramework/ProjectileMovementComponent.h
+++ b/Engine/Source/Runtime/Engine/Classes/GameFramework/ProjectileMovementComponent.h
@@ -115,9 +115,24 @@ class UProjectileMovementComponent : public UMovementComponent
 	UPROPERTY(EditAnywhere, BlueprintReadWrite, Category=ProjectileInterpolation)
 	uint8 bInterpRotation:1;
 
+	/**
+	 * If true, throttle interpolation when not relevant.
+	 * @see ThrottleInterpolationSkipFramesNotRecent, ThrottleInterpolationSkipFramesRecent, ThrottleInterpolationThresholdNotRenderedShortTime, ThrottleInterpolationThresholdNotRenderedLongTime
+	 */
+	UPROPERTY(EditAnywhere, BlueprintReadWrite, Category=ProjectileInterpolation)
+	uint8 bThrottleInterpolation:1;
+
 protected:
+	/**
+	 * True if interpolation has reached the target, false if there is more interpolation required.
+	 */
 	uint8 bInterpolationComplete:1;
 
+	/**
+	 * Tracks the number of frames since the last interpolation.
+	 */
+	int32 ThrottleInterpolationFramesSinceInterp;
+
 public:
 
 	/**
@@ -344,6 +359,30 @@ public:
 	UPROPERTY(EditAnywhere, BlueprintReadWrite, Category=ProjectileInterpolation, meta=(ClampMin="0"))
 	float InterpLocationSnapToTargetDistance;
 
+	/**
+	 * Time after not rendered recently when we consider throttling interpolation.
+	 */
+	UPROPERTY(EditAnywhere, BlueprintReadWrite, Category=ProjectileInterpolation, AdvancedDisplay)
+	float ThrottleInterpolationThresholdNotRenderedShortTime;
+
+	/**
+	 * Time after not rendered for a long time when we consider throttling interpolation.
+	 */
+	UPROPERTY(EditAnywhere, BlueprintReadWrite, Category=ProjectileInterpolation, AdvancedDisplay)
+	float ThrottleInterpolationThresholdNotRenderedLongTime;
+
+	/**
+	 * When recently relevant, skip this many frames of interpolation if throttling is enabled.
+	 */
+	UPROPERTY(EditAnywhere, BlueprintReadWrite, Category=ProjectileInterpolation, AdvancedDisplay, meta=(ClampMin="0", UIMin="0"))
+	int32 ThrottleInterpolationSkipFramesRecent;
+
+	/**
+	 * When not recently relevant, skip this many frames of interpolation if throttling is enabled.
+	 */
+	UPROPERTY(EditAnywhere, BlueprintReadWrite, Category=ProjectileInterpolation, AdvancedDisplay, meta=(ClampMin="0", UIMin="0"))
+	int32 ThrottleInterpolationSkipFramesNotRecent;
+
 	/**
 	 * Returns whether interpolation is complete because the target has been reached. True when interpolation is disabled.
 	 */
@@ -352,6 +391,24 @@ public:
 
 protected:
 
+	/**
+	 * Update interpolation throttling for this frame. Uses ComputeThrottleInterpolationMaxFrames() to determine the number of frames to allow to be skipped.
+	 * @return true if throttled this frame, false if there should be an update.
+	 */
+	bool UpdateThrottleInterpolation(float DeltaTime, USceneComponent* InterpComponent);
+
+	/**
+	 * Determine the number of frames to allow to skip when interpolating. Returning 0 means not to allow throttling.
+	 * Default implementation chooses between ThrottleInterpolationThresholdNotRenderedShortTime and ThrottleInterpolationThresholdNotRenderedLongTime based on WasRecentlyRendered(),
+	 * extend or override this for custom behavior.
+	 */
+	virtual int32 ComputeThrottleInterpolationMaxFrames(float DeltaTime, USceneComponent* InterpComponent);
+	
+	/**
+	 * Custom hook to reset throttle interpolation tracking when it was throttling previously.
+	 */
+	virtual void ResetThrottleInterpolation(float DeltaTime);
+
 	// Enum indicating how simulation should proceed after HandleBlockingHit() is called.
 	enum class EHandleBlockingHitResult
 	{
diff --git a/Engine/Source/Runtime/Engine/Private/Components/ProjectileMovementComponent.cpp b/Engine/Source/Runtime/Engine/Private/Components/ProjectileMovementComponent.cpp
--- a/Engine/Source/Runtime/Engine/Private/Components/ProjectileMovementComponent.cpp
+++ b/Engine/Source/Runtime/Engine/Private/Components/ProjectileMovementComponent.cpp
@@ -11,6 +11,7 @@
 
 CSV_DECLARE_CATEGORY_MODULE_EXTERN(CORE_API, Basic);
 DEFINE_LOG_CATEGORY_STATIC(LogProjectileMovement, Log, All);
+DEFINE_LOG_CATEGORY_STATIC(LogProjectileMovementInterpolation, Log, All);
 
 const float UProjectileMovementComponent::MIN_TICK_TIME = 1e-6f;
 
@@ -30,6 +31,12 @@ UProjectileMovementComponent::UProjectileMovementComponent(const FObjectInitiali
 	InterpRotationTime = 0.050f;
 	InterpLocationMaxLagDistance = 300.0f;
 	InterpLocationSnapToTargetDistance = 500.0f;
+	bThrottleInterpolation = false;
+	ThrottleInterpolationThresholdNotRenderedShortTime = 0.20f;
+	ThrottleInterpolationThresholdNotRenderedLongTime = 1.0f;
+	ThrottleInterpolationFramesSinceInterp = 0;
+	ThrottleInterpolationSkipFramesRecent = 1;
+	ThrottleInterpolationSkipFramesNotRecent = 2;
 
 	Velocity = FVector(1.f,0.f,0.f);
 
@@ -723,6 +730,9 @@ void UProjectileMovementComponent::SetInterpolatedComponent(USceneComponent* Com
 		InterpInitialLocationOffset = Component->GetRelativeLocation();
 		InterpInitialRotationOffset = Component->GetRelativeRotation().Quaternion();
 		bInterpolationComplete = false;
+
+		// Space out interpolation skipping to avoid objects spawned on single frame from always updating in sync
+		ThrottleInterpolationFramesSinceInterp = ThrottleInterpolationSkipFramesRecent > 0 ? FMath::RandRange(0, ThrottleInterpolationSkipFramesRecent) : 0;
 	}
 	else
 	{
@@ -812,12 +822,15 @@ void UProjectileMovementComponent::ResetInterpolation()
 {
 	if (USceneComponent* InterpComponent = GetInterpolatedComponent())
 	{
+		// Snap to original (non-interpolated) offset, we may be forcibly stopping interpolation and need to have it stay at the correct location.
 		InterpComponent->SetRelativeLocationAndRotation(InterpInitialLocationOffset, InterpInitialRotationOffset);
 	}
 
 	InterpLocationOffset = FVector::ZeroVector;
 	InterpRotationOffset = FQuat::Identity;
 	bInterpolationComplete = true;
+	
+	ThrottleInterpolationFramesSinceInterp = 0;
 }
 
 void UProjectileMovementComponent::TickInterpolation(float DeltaTime)
@@ -859,18 +872,32 @@ void UProjectileMovementComponent::TickInterpolation(float DeltaTime)
 
 			if (USceneComponent* InterpComponent = GetInterpolatedComponent())
 			{
-				// Apply result
-				if (UpdatedComponent)
+				const bool bShouldThrottleNow = UpdateThrottleInterpolation(DeltaTime, InterpComponent);
+				if (bShouldThrottleNow)
 				{
-					const FVector NewRelTranslation = UpdatedComponent->GetComponentToWorld().InverseTransformVectorNoScale(InterpLocationOffset) + InterpInitialLocationOffset;
-					if (bInterpRotation)
-					{
-						const FQuat NewRelRotation = InterpRotationOffset * InterpInitialRotationOffset;
-						InterpComponent->SetRelativeLocationAndRotation(NewRelTranslation, NewRelRotation);
-					}
-					else
+					UE_LOG(LogProjectileMovementInterpolation, Verbose, TEXT("--- Skip  Interpolation (%d frames : %s)"), ThrottleInterpolationFramesSinceInterp, *GetPathNameSafe(InterpComponent));
+					// Skip applying transform to InterpolatedComponent.
+					// Don't say we're done interpolating if we haven't applied the result yet, we need it to update next frame.
+					bInterpolationComplete = false;
+				}
+				else
+				{
+					ThrottleInterpolationFramesSinceInterp = 0;
+					UE_LOG(LogProjectileMovementInterpolation, Verbose, TEXT("+++ Apply Interpolation (%d frames : %s)"), ThrottleInterpolationFramesSinceInterp, *GetPathNameSafe(InterpComponent));
+
+					// Apply interpolation result
+					if (UpdatedComponent)
 					{
-						InterpComponent->SetRelativeLocation(NewRelTranslation);
+						const FVector NewRelTranslation = UpdatedComponent->GetComponentToWorld().InverseTransformVectorNoScale(InterpLocationOffset) + InterpInitialLocationOffset;
+						if (bInterpRotation)
+						{
+							const FQuat NewRelRotation = InterpRotationOffset * InterpInitialRotationOffset;
+							InterpComponent->SetRelativeLocationAndRotation(NewRelTranslation, NewRelRotation);
+						}
+						else
+						{
+							InterpComponent->SetRelativeLocation(NewRelTranslation);
+						}
 					}
 				}
 			}
@@ -889,3 +916,62 @@ void UProjectileMovementComponent::TickInterpolation(float DeltaTime)
 	}
 }
 
+
+bool UProjectileMovementComponent::UpdateThrottleInterpolation(float DeltaTime, USceneComponent* InterpComponent)
+{
+	bool bIsThrottlingThisFrame = false;
+
+	if (bThrottleInterpolation && bInterpMovement && InterpComponent)
+	{
+		const int32 ThrottleFrames = ComputeThrottleInterpolationMaxFrames(DeltaTime, InterpComponent);
+		if (ThrottleFrames > 0)
+		{
+			ThrottleInterpolationFramesSinceInterp += 1;
+			if (ThrottleInterpolationFramesSinceInterp <= ThrottleFrames)
+			{
+				bIsThrottlingThisFrame = true;
+			}
+		}
+	}
+
+	// Detect transition from throttled to not throttled.
+	if (!bIsThrottlingThisFrame && ThrottleInterpolationFramesSinceInterp > 0)
+	{
+		// Reset counter, not throttling this frame.
+		ThrottleInterpolationFramesSinceInterp = 0;
+
+		// Hook for custom reset logic
+		ResetThrottleInterpolation(DeltaTime);
+	}
+
+	return bIsThrottlingThisFrame;
+}
+
+int32 UProjectileMovementComponent::ComputeThrottleInterpolationMaxFrames(float DeltaTime, USceneComponent* InterpComponent)
+{
+	int32 ThrottleFrames = 0;
+	if (AActor* ActorOwner = InterpComponent->GetOwner())
+	{
+		// Not recently rendered?
+		if (!ActorOwner->WasRecentlyRendered(ThrottleInterpolationThresholdNotRenderedShortTime))
+		{
+			// Not rendered even a long time ago?
+			if (!ActorOwner->WasRecentlyRendered(ThrottleInterpolationThresholdNotRenderedLongTime))
+			{
+				ThrottleFrames = ThrottleInterpolationSkipFramesNotRecent;
+			}
+			else
+			{
+				ThrottleFrames = ThrottleInterpolationSkipFramesRecent;
+			}
+		}
+	}
+
+	return ThrottleFrames;
+}
+
+void UProjectileMovementComponent::ResetThrottleInterpolation(float DeltaTime)
+{
+	ThrottleInterpolationFramesSinceInterp = 0;
+}
+
