#ue5 - Add ENGINE_API to functions added in another branch where they weren't required.

diff --git a/Engine/Source/Runtime/Engine/Classes/GameFramework/ProjectileMovementComponent.h b/Engine/Source/Runtime/Engine/Classes/GameFramework/ProjectileMovementComponent.h
--- a/Engine/Source/Runtime/Engine/Classes/GameFramework/ProjectileMovementComponent.h
+++ b/Engine/Source/Runtime/Engine/Classes/GameFramework/ProjectileMovementComponent.h
@@ -395,19 +395,19 @@ protected:
 	 * Update interpolation throttling for this frame. Uses ComputeThrottleInterpolationMaxFrames() to determine the number of frames to allow to be skipped.
 	 * @return true if throttled this frame, false if there should be an update.
 	 */
-	bool UpdateThrottleInterpolation(float DeltaTime, USceneComponent* InterpComponent);
+	ENGINE_API bool UpdateThrottleInterpolation(float DeltaTime, USceneComponent* InterpComponent);
 
 	/**
 	 * Determine the number of frames to allow to skip when interpolating. Returning 0 means not to allow throttling.
 	 * Default implementation chooses between ThrottleInterpolationThresholdNotRenderedShortTime and ThrottleInterpolationThresholdNotRenderedLongTime based on WasRecentlyRendered(),
 	 * extend or override this for custom behavior.
 	 */
-	virtual int32 ComputeThrottleInterpolationMaxFrames(float DeltaTime, USceneComponent* InterpComponent);
+	ENGINE_API virtual int32 ComputeThrottleInterpolationMaxFrames(float DeltaTime, USceneComponent* InterpComponent);
 	
 	/**
 	 * Custom hook to reset throttle interpolation tracking when it was throttling previously.
 	 */
-	virtual void ResetThrottleInterpolation(float DeltaTime);
+	ENGINE_API virtual void ResetThrottleInterpolation(float DeltaTime);
 
 	// Enum indicating how simulation should proceed after HandleBlockingHit() is called.
 	enum class EHandleBlockingHitResult
