Added LLM tag for untagged video mem allocation in PhysicsField.

diff --git a/Engine/Source/Runtime/Engine/Private/PhysicsField/PhysicsFieldComponent.cpp b/Engine/Source/Runtime/Engine/Private/PhysicsField/PhysicsFieldComponent.cpp
--- a/Engine/Source/Runtime/Engine/Private/PhysicsField/PhysicsFieldComponent.cpp
+++ b/Engine/Source/Runtime/Engine/Private/PhysicsField/PhysicsFieldComponent.cpp
@@ -5,9 +5,9 @@
 #include "Engine/Engine.h"
 #include "Engine/World.h"
 #include "Field/FieldSystemNodes.h"
-
 #include "DataDrivenShaderPlatformInfo.h"
 #include "GlobalShader.h"
+#include "HAL/LowLevelMemTracker.h"
 #include "PipelineStateCache.h"
 #include "ShaderParameterUtils.h"
 #include "ProfilingDebugging/RealtimeGPUProfiler.h"
@@ -120,7 +120,10 @@ void InitInternalBuffer(FRHICommandListBase& RHICmdList, const uint32 ElementCou
 		const uint32 BufferCount = ElementCount * ElementSize;
 		const uint32 BufferBytes = sizeof(BufferType) * BufferCount;
 		
-		OutputBuffer.Initialize(RHICmdList, TEXT("FPhysicsFieldResource"), sizeof(BufferType), BufferCount, PixelFormat, BUF_Static);
+		{
+			LLM_SCOPE(ELLMTag::Physics);
+			OutputBuffer.Initialize(RHICmdList, TEXT("FPhysicsFieldResource"), sizeof(BufferType), BufferCount, PixelFormat, BUF_Static);
+		}
 
 		if (OutputBuffer.UAV)
 		{
@@ -148,9 +151,10 @@ void UpdateInternalBuffer(FRHICommandListBase& RHICmdList, const uint32 ElementC
 		const uint32 BufferCount = ElementCount * ElementSize;
 		const uint32 BufferBytes = sizeof(BufferType) * BufferCount;
 
-		if(bInitField)
+		if (bInitField)
 		{
-			OutputBuffer.Initialize(RHICmdList, TEXT("FPhysicsFieldResource"),sizeof(BufferType), BufferCount, PixelFormat, BUF_Static);
+			LLM_SCOPE(ELLMTag::Physics);
+			OutputBuffer.Initialize(RHICmdList, TEXT("FPhysicsFieldResource"), sizeof(BufferType), BufferCount, PixelFormat, BUF_Static);
 		}
 
 		void* OutputData = RHICmdList.LockBuffer(OutputBuffer.Buffer, 0, BufferBytes, RLM_WriteOnly);
