PLAY-7721: PO Interface GetRigidBody crash

Use Chaos::FConstPhysicsObjectHandle instead of Chaos::FPhysicsObjectHandle and Chaos::FPhysicsObject*
Also use GetParticle() from FWritePhysicsObjectInterface_Internal and then cast to FPBDRigidParticleHandle* under null-check locally.

diff --git a/Engine/Source/Runtime/Engine/Private/PhysicsEngine/PhysicsReplication.cpp b/Engine/Source/Runtime/Engine/Private/PhysicsEngine/PhysicsReplication.cpp
--- a/Engine/Source/Runtime/Engine/Private/PhysicsEngine/PhysicsReplication.cpp
+++ b/Engine/Source/Runtime/Engine/Private/PhysicsEngine/PhysicsReplication.cpp
@@ -151,7 +151,7 @@ void FPhysicsReplication::SetReplicatedTarget(UPrimitiveComponent* Component, FN
 		const bool bIsReplicatedAutonomous = OwnerRole == ROLE_AutonomousProxy && Component->bReplicatePhysicsToAutonomousProxy;
 		if (bIsSimulated || bIsReplicatedAutonomous)
 		{
-			Chaos::FPhysicsObjectHandle PhysicsObject = Component->GetPhysicsObjectByName(BoneName);
+			Chaos::FConstPhysicsObjectHandle PhysicsObject = Component->GetPhysicsObjectByName(BoneName);
 			SetReplicatedTarget(PhysicsObject, ReplicatedTarget, ServerFrame, Owner->GetPhysicsReplicationMode());
 			return;
 		}
@@ -195,7 +195,7 @@ void FPhysicsReplication::SetReplicatedTarget(UPrimitiveComponent* Component, FN
 	}
 }
 
-void FPhysicsReplication::SetReplicatedTarget(Chaos::FPhysicsObject* PhysicsObject, const FRigidBodyState& ReplicatedTarget, int32 ServerFrame, EPhysicsReplicationMode ReplicationMode)
+void FPhysicsReplication::SetReplicatedTarget(Chaos::FConstPhysicsObjectHandle PhysicsObject, const FRigidBodyState& ReplicatedTarget, int32 ServerFrame, EPhysicsReplicationMode ReplicationMode)
 {
 	UWorld* OwningWorld = GetOwningWorld();
 	if (OwningWorld == nullptr)
@@ -205,8 +205,7 @@ void FPhysicsReplication::SetReplicatedTarget(Chaos::FPhysicsObject* PhysicsObje
 
 	// TODO, Check if owning actor is ROLE_SimulatedProxy or ROLE_AutonomousProxy ?
 
-	FReplicatedPhysicsTarget Target;
-	Target.PhysicsObject = PhysicsObject;
+	FReplicatedPhysicsTarget Target(PhysicsObject);
 	Target.ReplicationMode = ReplicationMode;
 	Target.ServerFrame = ServerFrame;
 	Target.TargetState = ReplicatedTarget;
@@ -324,10 +323,9 @@ void FPhysicsReplication::OnTick(float DeltaSeconds, TMap<TWeakObjectPtr<UPrimit
 			const float PingSecondsOneWay = LocalPing * 0.5f * 0.001f;
 
 			// Queue up the target state for async replication
-			FPhysicsRepAsyncInputData AsyncInputData;
+			FPhysicsRepAsyncInputData AsyncInputData(PhysicsTarget.PhysicsObject);
 			AsyncInputData.TargetState = PhysicsTarget.TargetState;
 			AsyncInputData.Proxy = nullptr;
-			AsyncInputData.PhysicsObject = PhysicsTarget.PhysicsObject;
 			AsyncInputData.RepMode = PhysicsTarget.ReplicationMode;
 			AsyncInputData.ServerFrame = PhysicsTarget.ServerFrame;
 			AsyncInputData.FrameOffset = LocalFrameOffset;
@@ -584,12 +582,11 @@ bool FPhysicsReplication::ApplyRigidBodyState(float DeltaSeconds, FBodyInstance*
 			else
 			{
 				//If async is used, enqueue for callback
-				FPhysicsRepAsyncInputData AsyncInputData;
+				FPhysicsRepAsyncInputData AsyncInputData(nullptr);
 				AsyncInputData.TargetState = NewState;
 				AsyncInputData.TargetState.Position = IdealWorldTM.GetLocation();
 				AsyncInputData.TargetState.Quaternion = IdealWorldTM.GetRotation();
 				AsyncInputData.Proxy = static_cast<Chaos::FSingleParticlePhysicsProxy*>(BI->GetPhysicsActorHandle());
-				AsyncInputData.PhysicsObject = nullptr;
 				AsyncInputData.ErrorCorrection = { ErrorCorrection.LinearVelocityCoefficient, ErrorCorrection.AngularVelocityCoefficient, ErrorCorrection.PositionLerp, ErrorCorrection.AngleLerp };
 				AsyncInputData.LatencyOneWay = PingSeconds;
 
@@ -709,9 +706,7 @@ void FPhysicsReplicationAsync::UpdateRewindDataTarget(const FPhysicsRepAsyncInpu
 	}
 
 	Chaos::FReadPhysicsObjectInterface_Internal Interface = Chaos::FPhysicsObjectInternalInterface::GetRead();
-	Chaos::FPBDRigidParticleHandle* Handle = Interface.GetRigidParticle(Input.PhysicsObject);
-
-	if (Handle != nullptr)
+	if (Chaos::FGeometryParticleHandle* Handle = Interface.GetParticle(Input.PhysicsObject))
 	{
 		// Cache all target states inside RewindData
 		const int32 LocalFrame = Input.ServerFrame - Input.FrameOffset;
@@ -734,7 +729,7 @@ void FPhysicsReplicationAsync::UpdateAsyncTarget(const FPhysicsRepAsyncInputData
 		// First time we add a target, set it's previous and correction
 		// positions to the target position to avoid math with uninitialized
 		// memory.
-		Target = &ObjectToTarget.Add(Input.PhysicsObject);
+		Target = &ObjectToTarget.Add(Input.PhysicsObject, FReplicatedPhysicsTargetAsync(Input.PhysicsObject));
 		Target->PrevPos = Input.TargetState.Position;
 		Target->PrevPosTarget = Input.TargetState.Position;
 		Target->PrevRotTarget = Input.TargetState.Quaternion;
@@ -743,7 +738,6 @@ void FPhysicsReplicationAsync::UpdateAsyncTarget(const FPhysicsRepAsyncInputData
 
 	if (Input.ServerFrame > Target->ServerFrame)
 	{
-		Target->PhysicsObject = Input.PhysicsObject;
 		Target->PrevServerFrame = Target->ServerFrame;
 		Target->ServerFrame = Input.ServerFrame;
 		Target->TargetState = Input.TargetState;
@@ -791,30 +785,32 @@ void FPhysicsReplicationAsync::ApplyTargetStatesAsync(const float DeltaSeconds,
 		if (Target.PhysicsObject != nullptr)
 		{
 			Chaos::FWritePhysicsObjectInterface_Internal Interface = Chaos::FPhysicsObjectInternalInterface::GetWrite();
-			FPBDRigidParticleHandle* Handle = Interface.GetRigidParticle(Target.PhysicsObject);
-
-			if (Handle != nullptr)
+			if (FGeometryParticleHandle* Handle = Interface.GetParticle(Target.PhysicsObject))
 			{
-				// TODO, Remove the resim option from project settings, we only need the physics prediction one now
-				EPhysicsReplicationMode RepMode = Target.RepMode;
-				if (!Chaos::FPBDRigidsSolver::IsPhysicsResimulationEnabled() && RepMode == EPhysicsReplicationMode::Resimulation)
+				if (FPBDRigidParticleHandle* RigidHandle = Handle->CastToRigidParticle())
 				{
-					RepMode = EPhysicsReplicationMode::Default;
-				}
+					EPhysicsReplicationMode RepMode = Target.RepMode;
 
-				switch (RepMode)
-				{
-					case EPhysicsReplicationMode::Default:
-						bRemoveItr = DefaultReplication(Handle, Target, DeltaSeconds);
-						break;
+					// TODO, Remove the resim option from project settings, we only need the physics prediction one now
+					if (!Chaos::FPBDRigidsSolver::IsPhysicsResimulationEnabled() && RepMode == EPhysicsReplicationMode::Resimulation)
+					{
+						RepMode = EPhysicsReplicationMode::Default;
+					}
 
-					case EPhysicsReplicationMode::PredictiveInterpolation:
-						bRemoveItr = PredictiveInterpolation(Handle, Target, DeltaSeconds);
-						break;
+					switch (RepMode)
+					{
+						case EPhysicsReplicationMode::Default:
+							bRemoveItr = DefaultReplication(RigidHandle, Target, DeltaSeconds);
+							break;
 
-					case EPhysicsReplicationMode::Resimulation:
-						bRemoveItr = ResimulationReplication(Handle, Target, DeltaSeconds);
-						break;
+						case EPhysicsReplicationMode::PredictiveInterpolation:
+							bRemoveItr = PredictiveInterpolation(RigidHandle, Target, DeltaSeconds);
+							break;
+
+						case EPhysicsReplicationMode::Resimulation:
+							bRemoveItr = ResimulationReplication(RigidHandle, Target, DeltaSeconds);
+							break;
+					}
 				}
 			}
 		}
@@ -1194,13 +1190,6 @@ bool FPhysicsReplicationAsync::PredictiveInterpolation(Chaos::FPBDRigidParticleH
 	}
 	else // Velocity-based Replication
 	{
-		const Chaos::EObjectStateType ObjectState = Handle->ObjectState();
-		if (ObjectState != Chaos::EObjectStateType::Dynamic)
-		{
-			RigidsSolver->GetEvolution()->SetParticleObjectState(Handle, Chaos::EObjectStateType::Dynamic);
-		}
-
-
 		// --- Velocity Replication ---
 		// Get PosDiff
 		const FVector PosDiff = TargetPos - CurrentState.Position;
diff --git a/Engine/Source/Runtime/Engine/Public/PhysicsReplication.h b/Engine/Source/Runtime/Engine/Public/PhysicsReplication.h
--- a/Engine/Source/Runtime/Engine/Public/PhysicsReplication.h
+++ b/Engine/Source/Runtime/Engine/Public/PhysicsReplication.h
@@ -63,16 +63,16 @@ struct FPhysicsRepAsyncInputData
 {
 	FRigidBodyState TargetState;
 	Chaos::FSingleParticlePhysicsProxy* Proxy; // Used for legacy (BodyInstance) flow
-	Chaos::FPhysicsObject* PhysicsObject;
+	Chaos::FConstPhysicsObjectHandle PhysicsObject;
 	TOptional<FPhysicsRepErrorCorrectionData> ErrorCorrection;
 	EPhysicsReplicationMode RepMode;
 	int32 ServerFrame;
 	int32 FrameOffset;
 	float LatencyOneWay;
 
-	FPhysicsRepAsyncInputData()
+	FPhysicsRepAsyncInputData(Chaos::FConstPhysicsObjectHandle POHandle)
 		: Proxy(nullptr)
-		, PhysicsObject(nullptr)
+		, PhysicsObject(POHandle)
 		, RepMode(EPhysicsReplicationMode::Default)
 	{};
 };
@@ -90,10 +90,10 @@ struct FPhysicsReplicationAsyncInput : public Chaos::FSimCallbackInput
 
 struct FReplicatedPhysicsTargetAsync
 {
-	FReplicatedPhysicsTargetAsync()
+	FReplicatedPhysicsTargetAsync(Chaos::FConstPhysicsObjectHandle POHandle)
 		: AccumulatedErrorSeconds(0.0f)
 		, ServerFrame(0)
-		, PhysicsObject(nullptr)
+		, PhysicsObject(POHandle)
 	{ }
 
 	/** The target state replicated by server */
@@ -112,7 +112,7 @@ struct FReplicatedPhysicsTargetAsync
 	int32 FrameOffset;
 
 	/** Index of physics object on component */
-	Chaos::FPhysicsObject* PhysicsObject;
+	Chaos::FConstPhysicsObjectHandle PhysicsObject;
 
 	/** The replication mode this PhysicsObject should use */
 	EPhysicsReplicationMode RepMode;
@@ -146,7 +146,7 @@ class FPhysicsReplicationAsync : public Chaos::TSimCallbackObject<
 private:
 	float LatencyOneWay;
 	FRigidBodyErrorCorrection ErrorCorrectionDefault;
-	TMap<Chaos::FPhysicsObject*, FReplicatedPhysicsTargetAsync> ObjectToTarget;
+	TMap<Chaos::FConstPhysicsObjectHandle, FReplicatedPhysicsTargetAsync> ObjectToTarget;
 
 private:
 	void UpdateAsyncTarget(const FPhysicsRepAsyncInputData& Input);
@@ -165,11 +165,11 @@ public:
 
 struct FReplicatedPhysicsTarget
 {
-	FReplicatedPhysicsTarget()
+	FReplicatedPhysicsTarget(Chaos::FConstPhysicsObjectHandle POHandle = nullptr)
 		: ArrivedTimeSeconds(0.0f)
 		, AccumulatedErrorSeconds(0.0f)
 		, ServerFrame(0)
-		, PhysicsObject(nullptr)
+		, PhysicsObject(POHandle)
 	{ }
 
 	/** The target state replicated by server */
@@ -192,7 +192,7 @@ struct FReplicatedPhysicsTarget
 	int32 ServerFrame;
 
 	/** Index of physics object on component */
-	Chaos::FPhysicsObject* PhysicsObject;
+	Chaos::FConstPhysicsObjectHandle PhysicsObject;
 
 	/** The replication mode the target should be used with */
 	EPhysicsReplicationMode ReplicationMode;
@@ -219,7 +219,7 @@ public:
 	virtual void SetReplicatedTarget(UPrimitiveComponent* Component, FName BoneName, const FRigidBodyState& ReplicatedTarget) { SetReplicatedTarget(Component, BoneName, ReplicatedTarget, 0); }
 	ENGINE_API virtual void SetReplicatedTarget(UPrimitiveComponent* Component, FName BoneName, const FRigidBodyState& ReplicatedTarget, int32 ServerFrame) override;
 private:
-	ENGINE_API void SetReplicatedTarget(Chaos::FPhysicsObject* PhysicsObject, const FRigidBodyState& ReplicatedTarget, int32 ServerFrame, EPhysicsReplicationMode ReplicationMode = EPhysicsReplicationMode::Default);
+	ENGINE_API void SetReplicatedTarget(Chaos::FConstPhysicsObjectHandle PhysicsObject, const FRigidBodyState& ReplicatedTarget, int32 ServerFrame, EPhysicsReplicationMode ReplicationMode = EPhysicsReplicationMode::Default);
 
 public:
 	/** Remove the replicated target*/
