World Partition Location Volume: update class description comment.

diff --git a/Engine/Source/Runtime/Engine/Public/LocationVolume.h b/Engine/Source/Runtime/Engine/Public/LocationVolume.h
--- a/Engine/Source/Runtime/Engine/Public/LocationVolume.h
+++ b/Engine/Source/Runtime/Engine/Public/LocationVolume.h
@@ -11,7 +11,7 @@
 class FLoaderAdapterActor;
 
 /**
- * A volume representing a location in the world
+ * A volume representing a location in the world. Used for World Partition loading regions.
  */
 UCLASS(BlueprintType, Blueprintable, MinimalAPI)
 class ALocationVolume : public AVolume, public IWorldPartitionActorLoaderInterface
