APlayerController::IsPlayerMuted now works in the client.
Removed usage of EVoiceBlockReasons::MutedBy.

diff --git a/Engine/Source/Runtime/Engine/Classes/GameFramework/PlayerMuteList.h b/Engine/Source/Runtime/Engine/Classes/GameFramework/PlayerMuteList.h
--- a/Engine/Source/Runtime/Engine/Classes/GameFramework/PlayerMuteList.h
+++ b/Engine/Source/Runtime/Engine/Classes/GameFramework/PlayerMuteList.h
@@ -14,13 +14,12 @@ enum class EVoiceBlockReasons : uint8
 {
 	None		= 0,		// Communication with this client has no filter reasons
 	Muted		= 1 << 0,	// The owning player controller has explicitly muted the player
-	MutedBy		= 1 << 1,	// The owning player controller was explicitly muted by the player
-	Gameplay	= 1 << 2,	// The player was muted for gameplay reasons
-	Blocked		= 1 << 3,	// The owning player controller has blocked the player
-	BlockedBy	= 1 << 4,	// The owning player controller was blocked by the player
+	Gameplay	= 1 << 1,	// The player was muted for gameplay reasons
+	Blocked		= 1 << 2,	// The owning player controller has blocked the player
+	BlockedBy	= 1 << 3,	// The owning player controller was blocked by the player
 };
 ENUM_CLASS_FLAGS(EVoiceBlockReasons);
-ENUM_RANGE_BY_VALUES(EVoiceBlockReasons, EVoiceBlockReasons::Muted, EVoiceBlockReasons::MutedBy, EVoiceBlockReasons::Gameplay, EVoiceBlockReasons::Blocked, EVoiceBlockReasons::BlockedBy);
+ENUM_RANGE_BY_VALUES(EVoiceBlockReasons, EVoiceBlockReasons::Muted, EVoiceBlockReasons::Gameplay, EVoiceBlockReasons::Blocked, EVoiceBlockReasons::BlockedBy);
 
 /**
  * Container responsible for managing the mute state of a player controller
diff --git a/Engine/Source/Runtime/Engine/Private/PlayerController.cpp b/Engine/Source/Runtime/Engine/Private/PlayerController.cpp
--- a/Engine/Source/Runtime/Engine/Private/PlayerController.cpp
+++ b/Engine/Source/Runtime/Engine/Private/PlayerController.cpp
@@ -3747,6 +3747,11 @@ void APlayerController::ClientMutePlayer_Implementation(FUniqueNetIdRepl PlayerI
 	ULocalPlayer* LP = Cast<ULocalPlayer>(Player);
 	UWorld* World = GetWorld();
 
+	// @todo: As of now we don't have a proper way to inform the client of the specific voice block reason
+	// without changing the function signatures, therefore all server reasons are funneled into the client
+	// as "muted" for the time being.
+	MuteList.AddVoiceBlockReason(PlayerId.GetUniqueNetId(), EVoiceBlockReasons::Muted);
+
 	if (LP != NULL && World)
 	{
 		// Have the voice subsystem mute this player
@@ -3760,6 +3765,8 @@ void APlayerController::ClientUnmutePlayer_Implementation(FUniqueNetIdRepl Playe
 	ULocalPlayer* LP = Cast<ULocalPlayer>(Player);
 	UWorld* World = GetWorld();
 
+	MuteList.RemoveVoiceBlockReason(PlayerId.GetUniqueNetId(), EVoiceBlockReasons::Muted);
+
 	if (LP != NULL && World)
 	{
 		// Have the voice subsystem unmute this player
diff --git a/Engine/Source/Runtime/Engine/Private/PlayerMuteList.cpp b/Engine/Source/Runtime/Engine/Private/PlayerMuteList.cpp
--- a/Engine/Source/Runtime/Engine/Private/PlayerMuteList.cpp
+++ b/Engine/Source/Runtime/Engine/Private/PlayerMuteList.cpp
@@ -7,6 +7,8 @@
 
 #include UE_INLINE_GENERATED_CPP_BY_NAME(PlayerMuteList)
 
+constexpr EVoiceBlockReasons ClientVisibleVoiceFilterReasons = EVoiceBlockReasons::Muted | EVoiceBlockReasons::Gameplay | EVoiceBlockReasons::Blocked;
+
 FString LexToString(const EVoiceBlockReasons& Reason)
 {
 	FString Result;
@@ -27,9 +29,6 @@ FString LexToString(const EVoiceBlockReasons& Reason)
 			case EVoiceBlockReasons::Muted:
 				Result += TEXT("MUTED");
 				break;
-			case EVoiceBlockReasons::MutedBy:
-				Result += TEXT("MUTED BY");
-				break;
 			case EVoiceBlockReasons::Gameplay:
 				Result += TEXT("GAMEPLAY");
 				break;
@@ -57,7 +56,8 @@ bool FPlayerMuteList::AddVoiceBlockReason(const FUniqueNetIdPtr& PlayerId, EVoic
 	const EVoiceBlockReasons OldBlockReasons = BlockReasons;
 	EnumAddFlags(BlockReasons, VoiceBlockReason);
 	
-	return OldBlockReasons == EVoiceBlockReasons::None;
+	return (OldBlockReasons & ClientVisibleVoiceFilterReasons) == EVoiceBlockReasons::None && 
+		VoiceBlockReason != EVoiceBlockReasons::BlockedBy;
 }
 
 bool FPlayerMuteList::RemoveVoiceBlockReason(const FUniqueNetIdPtr& PlayerId, EVoiceBlockReasons VoiceBlockReason)
@@ -69,7 +69,8 @@ bool FPlayerMuteList::RemoveVoiceBlockReason(const FUniqueNetIdPtr& PlayerId, EV
 		const EVoiceBlockReasons OldBlockReasons = *BlockReasons;
 		EnumRemoveFlags(*BlockReasons, VoiceBlockReason);
 
-		return OldBlockReasons != EVoiceBlockReasons::None && *BlockReasons == EVoiceBlockReasons::None;
+		return (OldBlockReasons & ClientVisibleVoiceFilterReasons) != EVoiceBlockReasons::None && 
+			(*BlockReasons & ClientVisibleVoiceFilterReasons) == EVoiceBlockReasons::None;
 	}
 
 	return false;
@@ -83,19 +84,6 @@ void FPlayerMuteList::ServerMutePlayer(APlayerController* OwningPC, const FUniqu
 		// This is the first reason added, so we transitioned unmuted -> muted. Replicate mute state to client
 		OwningPC->ClientMutePlayer(MuteId);
 	}
-
-	// Find the muted player's player controller so it can be notified
-	APlayerController* OtherPC = OwningPC->GetPlayerControllerForMuting(MuteId);
-	const FUniqueNetIdRepl& OwningPlayerId = OwningPC->PlayerState->GetUniqueId();
-	if (OtherPC != NULL)
-	{
-		// Update their packet filter too
-		if (OtherPC->MuteList.AddVoiceBlockReason(OwningPlayerId.GetUniqueNetId(), EVoiceBlockReasons::MutedBy))
-		{
-			// This is the first reason added, so we transitioned unmuted -> muted. Tell the other PC to mute this one
-			OtherPC->ClientMutePlayer(OwningPlayerId);
-		}
-	}
 }
 
 void FPlayerMuteList::ServerUnmutePlayer(APlayerController* OwningPC, const FUniqueNetIdRepl& UnmuteId)
@@ -106,19 +94,6 @@ void FPlayerMuteList::ServerUnmutePlayer(APlayerController* OwningPC, const FUni
 		// We removed the last flag, transitioning from muted -> unmuted. Replicate mute state to client.
 		OwningPC->ClientUnmutePlayer(UnmuteId);
 	}
-
-	// Find the muted player's player controller so it can be notified
-	APlayerController* OtherPC = OwningPC->GetPlayerControllerForMuting(UnmuteId);
-	const FUniqueNetIdRepl& OwningPlayerId = OwningPC->PlayerState->GetUniqueId();
-	if (OtherPC != NULL)
-	{
-		// Update their packet filter list too
-		if (OtherPC->MuteList.RemoveVoiceBlockReason(OwningPlayerId.GetUniqueNetId(), EVoiceBlockReasons::MutedBy))
-		{
-			// We removed the last flag, transitioning from muted->unmuted.Replicate mute state to other PC.
-			OtherPC->ClientUnmutePlayer(OwningPlayerId);
-		}
-	}
 }
 
 void FPlayerMuteList::GameplayMutePlayer(APlayerController* OwningPC, const FUniqueNetIdRepl& MuteId)
