Backout remains of 25927248 that was properly fixed in 25927248.

diff --git a/Engine/Source/Runtime/Engine/Classes/GameFramework/Actor.h b/Engine/Source/Runtime/Engine/Classes/GameFramework/Actor.h
--- a/Engine/Source/Runtime/Engine/Classes/GameFramework/Actor.h
+++ b/Engine/Source/Runtime/Engine/Classes/GameFramework/Actor.h
@@ -2207,11 +2207,6 @@ public:
 
 	/** Does this actor supports external packaging? */
 	ENGINE_API virtual bool SupportsExternalPackaging() const;
-
-	/** Called when the we detach the actor from the external package. Temporary fix to allow the actor to fixup and rename any objects outered to the external package. */
-	virtual void OnDetachExternalPackage() {}
-	/** Called when the we reattach the actor to the external package. Temporary fix to allow the actor to fixup and rename any objects outered to the external package. */
-	virtual void OnReattachExternalPackage() {}
 #endif
 
 #if WITH_EDITOR
diff --git a/Engine/Source/Runtime/Engine/Private/Level.cpp b/Engine/Source/Runtime/Engine/Private/Level.cpp
--- a/Engine/Source/Runtime/Engine/Private/Level.cpp
+++ b/Engine/Source/Runtime/Engine/Private/Level.cpp
@@ -4064,7 +4064,6 @@ void ULevel::DetachAttachAllActorsPackages(bool bReattach)
 			if (Actor)
 			{
 				Actor->ReattachExternalPackage();
-				Actor->OnReattachExternalPackage();
 			}
 		}
 	}
@@ -4074,7 +4073,6 @@ void ULevel::DetachAttachAllActorsPackages(bool bReattach)
 		{
 			if (Actor)
 			{
-				Actor->OnDetachExternalPackage();
 				Actor->DetachExternalPackage();
 			}
 		}
