Fix compilation with ANIMNODE_STATS_VERBOSE enabled on targets where STATS is disabled

diff --git a/Engine/Source/Runtime/Engine/Public/Animation/AnimStats.h b/Engine/Source/Runtime/Engine/Public/Animation/AnimStats.h
--- a/Engine/Source/Runtime/Engine/Public/Animation/AnimStats.h
+++ b/Engine/Source/Runtime/Engine/Public/Animation/AnimStats.h
@@ -45,8 +45,12 @@ DECLARE_CYCLE_STAT_EXTERN(TEXT("UpdateCurvesPostEvaluation"), STAT_UpdateCurvesP
 DECLARE_CYCLE_STAT_EXTERN(TEXT("TickAssetPlayerInstances"), STAT_TickAssetPlayerInstances, STATGROUP_Anim, );
 DECLARE_CYCLE_STAT_EXTERN(TEXT("TickAssetPlayerInstance"), STAT_TickAssetPlayerInstance, STATGROUP_Anim, );
 
-// Uncomment for verbose anim node state
-#define ANIMNODE_STATS_VERBOSE 0
+#if STATS
+	// Set to 1 for verbose anim node profiling
+	#define ANIMNODE_STATS_VERBOSE 0
+#else
+	#define ANIMNODE_STATS_VERBOSE 0
+#endif 
 
 #define DO_ANIMSTAT_PROCESSING(StatName) DECLARE_CYCLE_STAT_EXTERN(TEXT(#StatName), STAT_ ## StatName, STATGROUP_Anim, ENGINE_API)
 #include "Animation/AnimMTStats.h"
