Content Bundle
* Move Content Bundle events into Engine (needed by [REDACTED_PERSON]).
* Add ContentBundleRemovedContent event.

diff --git a/Engine/Source/Editor/WorldPartitionEditor/Private/WorldPartition/ContentBundle/ContentBundleEditorSubsystem.cpp b/Engine/Source/Editor/WorldPartitionEditor/Private/WorldPartition/ContentBundle/ContentBundleEditorSubsystem.cpp
--- a/Engine/Source/Editor/WorldPartitionEditor/Private/WorldPartition/ContentBundle/ContentBundleEditorSubsystem.cpp
+++ b/Engine/Source/Editor/WorldPartitionEditor/Private/WorldPartition/ContentBundle/ContentBundleEditorSubsystem.cpp
@@ -288,6 +288,7 @@ void UContentBundleEditorSubsystem::NotifyContentBundleRemovedContent(const FCon
 		GetEditingSubmodule()->DeactivateCurrentContentBundleEditing();
 	}
 
+	OnContentBundleRemovedContent().Broadcast(ContentBundle);
 	NotifyContentBundleChanged(ContentBundle);
 }
 
diff --git a/Engine/Source/Editor/WorldPartitionEditor/Public/WorldPartition/ContentBundle/ContentBundleEditorSubsystem.h b/Engine/Source/Editor/WorldPartitionEditor/Public/WorldPartition/ContentBundle/ContentBundleEditorSubsystem.h
--- a/Engine/Source/Editor/WorldPartitionEditor/Public/WorldPartition/ContentBundle/ContentBundleEditorSubsystem.h
+++ b/Engine/Source/Editor/WorldPartitionEditor/Public/WorldPartition/ContentBundle/ContentBundleEditorSubsystem.h
@@ -137,22 +137,9 @@ public:
 	virtual void PushContentBundleEditing() override;
 	virtual void PopContentBundleEditing() override;
 
-	DECLARE_EVENT_OneParam(UContentBundleEditorSubsystem, FOnContentBundleChanged, const FContentBundleEditor*);
-	FOnContentBundleChanged& OnContentBundleChanged() { return ContentBundleChanged; }
-
-	DECLARE_EVENT_OneParam(UContentBundleEditorSubsystem, FOnContentBundleAdded, const FContentBundleEditor*);
-	FOnContentBundleAdded& OnContentBundleAdded() { return ContentBundleAdded; }
-
-	DECLARE_EVENT_OneParam(UContentBundleEditorSubsystem, FOnContentBundleRemoved, const FContentBundleEditor*);
-	FOnContentBundleRemoved& OnContentBundleRemoved() { return ContentBundleRemoved; }
-
 private:
 	void SelectActorsInternal(FContentBundleEditor& EditorContentBundle, bool bSelect);
 
 	UPROPERTY()
 	TObjectPtr<UContentBundleEditingSubmodule> ContentBundleEditingSubModule;
-
-	FOnContentBundleChanged ContentBundleChanged;
-	FOnContentBundleAdded ContentBundleAdded;
-	FOnContentBundleRemoved ContentBundleRemoved;
 };
\ No newline at end of file
diff --git a/Engine/Source/Runtime/Engine/Public/WorldPartition/ContentBundle/ContentBundleEditor.h b/Engine/Source/Runtime/Engine/Public/WorldPartition/ContentBundle/ContentBundleEditor.h
--- a/Engine/Source/Runtime/Engine/Public/WorldPartition/ContentBundle/ContentBundleEditor.h
+++ b/Engine/Source/Runtime/Engine/Public/WorldPartition/ContentBundle/ContentBundleEditor.h
@@ -52,6 +52,7 @@ public:
 	ENGINE_API void ReferenceAllActors();
 	ENGINE_API void UnreferenceAllActors();
 
+	const TWeakObjectPtr<UActorDescContainer>& GetActorDescContainer() const { return ActorDescContainer; };
 	ENGINE_API void GenerateStreaming(TArray<FString>* OutPackageToGenerate, bool bIsPIE);
 
 	URuntimeHashExternalStreamingObjectBase* GetStreamingObject() const { return ExternalStreamingObject; }
diff --git a/Engine/Source/Runtime/Engine/Public/WorldPartition/ContentBundle/ContentBundleEditorSubsystemInterface.h b/Engine/Source/Runtime/Engine/Public/WorldPartition/ContentBundle/ContentBundleEditorSubsystemInterface.h
--- a/Engine/Source/Runtime/Engine/Public/WorldPartition/ContentBundle/ContentBundleEditorSubsystemInterface.h
+++ b/Engine/Source/Runtime/Engine/Public/WorldPartition/ContentBundle/ContentBundleEditorSubsystemInterface.h
@@ -11,6 +11,8 @@ class IContentBundleEditorSubsystemInterface
 public:
 	static IContentBundleEditorSubsystemInterface* Get() { return Instance; }
 
+	virtual ~IContentBundleEditorSubsystemInterface() {};
+
 	virtual void NotifyContentBundleAdded(const FContentBundleEditor* ContentBundle) = 0;
 	virtual void NotifyContentBundleRemoved(const FContentBundleEditor* ContentBundle) = 0;
 	virtual void NotifyContentBundleInjectedContent(const FContentBundleEditor* ContentBundle) = 0;
@@ -25,8 +27,26 @@ public:
 	virtual void PushContentBundleEditing() = 0;
 	virtual void PopContentBundleEditing() = 0;
 
+	DECLARE_EVENT_OneParam(UContentBundleEditorSubsystem, FOnContentBundleAdded, const FContentBundleEditor*);
+	FOnContentBundleAdded& OnContentBundleAdded() { return ContentBundleAdded; }
+
+	DECLARE_EVENT_OneParam(UContentBundleEditorSubsystem, FOnContentBundleRemoved, const FContentBundleEditor*);
+	FOnContentBundleRemoved& OnContentBundleRemoved() { return ContentBundleRemoved; }
+
+	DECLARE_EVENT_OneParam(UContentBundleEditorSubsystem, FOnContentBundleRemovedContent, const FContentBundleEditor*);
+	FOnContentBundleRemovedContent& OnContentBundleRemovedContent() { return ContentBundleRemovedContent; }
+
+	DECLARE_EVENT_OneParam(UContentBundleEditorSubsystem, FOnContentBundleChanged, const FContentBundleEditor*);
+	FOnContentBundleChanged& OnContentBundleChanged() { return ContentBundleChanged; }
+
 protected:
 	static ENGINE_API void SetInstance(IContentBundleEditorSubsystemInterface* InInstance);
 	static ENGINE_API IContentBundleEditorSubsystemInterface* Instance;
+
+private:
+	FOnContentBundleAdded ContentBundleAdded;
+	FOnContentBundleRemoved ContentBundleRemoved;
+	FOnContentBundleRemovedContent ContentBundleRemovedContent;
+	FOnContentBundleChanged ContentBundleChanged;
 };
 #endif
