World Partition Hash Set: merge code used by different partition objects and add detection for duplicated cell names.

diff --git a/Engine/Source/Runtime/Engine/Private/WorldPartition/RuntimeHashSet/RuntimePartition.cpp b/Engine/Source/Runtime/Engine/Private/WorldPartition/RuntimeHashSet/RuntimePartition.cpp
--- a/Engine/Source/Runtime/Engine/Private/WorldPartition/RuntimeHashSet/RuntimePartition.cpp
+++ b/Engine/Source/Runtime/Engine/Private/WorldPartition/RuntimeHashSet/RuntimePartition.cpp
@@ -3,21 +3,58 @@
 #include "WorldPartition/RuntimeHashSet/RuntimePartition.h"
 
 #if WITH_EDITOR
-void URuntimePartition::PostEditChangeProperty(FPropertyChangedEvent& PropertyChangedEvent)
+void URuntimePartition::PostEditChangeProperty(FPropertyChangedEvent& InPropertyChangedEvent)
 {
-	const FName PropertyName = PropertyChangedEvent.GetPropertyName();
+	const FName PropertyName = InPropertyChangedEvent.GetPropertyName();
 
 	if (PropertyName == GET_MEMBER_NAME_CHECKED(URuntimePartition, LoadingRange))
 	{
 		LoadingRange = FMath::Max(LoadingRange, 0);
 	}
 
-	Super::PostEditChangeProperty(PropertyChangedEvent);
+	Super::PostEditChangeProperty(InPropertyChangedEvent);
 }
 
-bool URuntimePartition::PopulateCellActorInstances(const TArray<const IStreamingGenerationContext::FActorSetInstance*>& ActorSetInstances, bool bIsMainWorldPartition, bool bIsCellAlwaysLoaded, TArray<IStreamingGenerationContext::FActorInstance>& OutCellActorInstances)
+bool URuntimePartition::PopulateCellActorInstances(const TArray<const IStreamingGenerationContext::FActorSetInstance*>& InActorSetInstances, bool bInIsMainWorldPartition, bool bInIsCellAlwaysLoaded, TArray<IStreamingGenerationContext::FActorInstance>& OutCellActorInstances)
 {
 	UWorldPartitionRuntimeHash* RuntimeHash = GetTypedOuter<UWorldPartitionRuntimeHash>();
-	return RuntimeHash->PopulateCellActorInstances(ActorSetInstances, bIsMainWorldPartition, bIsCellAlwaysLoaded, OutCellActorInstances);
+	return RuntimeHash->PopulateCellActorInstances(InActorSetInstances, bInIsMainWorldPartition, bInIsCellAlwaysLoaded, OutCellActorInstances);
+}
+
+URuntimePartition::FCellDesc URuntimePartition::CreateCellDesc(const FString& InName, bool bInIsSpatiallyLoaded, const FGuid& InContentBundleID, int32 InLevel, const TArray<IStreamingGenerationContext::FActorInstance>& InActorInstances)
+{
+	FCellDesc CellDesc;
+
+	// Construct a unique name like this: PartitionName_CellName
+	TStringBuilder<512> StringBuilder;
+	StringBuilder += Name.ToString();
+	StringBuilder += TEXT("_");
+	StringBuilder += InName;
+	CellDesc.Name = *StringBuilder;
+
+	// Copy values coming from this partition
+	CellDesc.bBlockOnSlowStreaming = bBlockOnSlowStreaming;
+	CellDesc.bClientOnlyVisible = bClientOnlyVisible;
+	CellDesc.Priority = Priority;
+
+	// Set provided input values
+	CellDesc.bIsSpatiallyLoaded = bInIsSpatiallyLoaded;
+	CellDesc.ContentBundleID = InContentBundleID;
+	CellDesc.Level = InLevel;
+
+	// Add actor instances and update bounds
+	CellDesc.ActorInstances = InActorInstances;
+	for (const IStreamingGenerationContext::FActorInstance& ActorInstance : CellDesc.ActorInstances)
+	{
+		const FWorldPartitionActorDescView& ActorDescView = ActorInstance.GetActorDescView();
+		const FBox RuntimeBounds = ActorDescView.GetRuntimeBounds();
+
+		if (RuntimeBounds.IsValid)
+		{
+			CellDesc.Bounds += RuntimeBounds.TransformBy(ActorInstance.GetTransform());
+		}
+	}
+
+	return CellDesc;
 }
 #endif
\ No newline at end of file
diff --git a/Engine/Source/Runtime/Engine/Private/WorldPartition/RuntimeHashSet/RuntimePartitionLevelStreaming.cpp b/Engine/Source/Runtime/Engine/Private/WorldPartition/RuntimeHashSet/RuntimePartitionLevelStreaming.cpp
--- a/Engine/Source/Runtime/Engine/Private/WorldPartition/RuntimeHashSet/RuntimePartitionLevelStreaming.cpp
+++ b/Engine/Source/Runtime/Engine/Private/WorldPartition/RuntimeHashSet/RuntimePartitionLevelStreaming.cpp
@@ -59,26 +59,7 @@ bool URuntimePartitionLevelStreaming::GenerateStreaming(const TArray<const IStre
 
 		for (auto& [SubLevelName, SubLevelActorSetInstances] : SubLevelsActorInstances)
 		{
-			FCellDesc& CellDesc = OutRuntimeCellDescs.Emplace_GetRef();
-
-			CellDesc.Name = SubLevelName;
-			CellDesc.bIsSpatiallyLoaded = true;
-			CellDesc.ContentBundleID = SubLevelActorSetInstances[0].ActorSetInstance->ContentBundleID;
-			CellDesc.bBlockOnSlowStreaming = bBlockOnSlowStreaming;
-			CellDesc.bClientOnlyVisible = bClientOnlyVisible;
-			CellDesc.Priority = Priority;
-			CellDesc.Level = 0;
-			CellDesc.ActorInstances = SubLevelActorSetInstances;
-
-			for (const IStreamingGenerationContext::FActorInstance& ActorInstance : CellDesc.ActorInstances)
-			{
-				const FWorldPartitionActorDescView& ActorDescView = ActorInstance.GetActorDescView();
-				const FBox RuntimeBounds = ActorDescView.GetRuntimeBounds();
-				if (RuntimeBounds.IsValid)
-				{
-					CellDesc.Bounds += RuntimeBounds.TransformBy(ActorInstance.GetTransform());
-				}
-			}
+			OutRuntimeCellDescs.Emplace(CreateCellDesc(SubLevelName.ToString(), true, SubLevelActorSetInstances[0].ActorSetInstance->ContentBundleID, 0, SubLevelActorSetInstances));
 		}
 	}
 
diff --git a/Engine/Source/Runtime/Engine/Private/WorldPartition/RuntimeHashSet/RuntimePartitionPersistent.cpp b/Engine/Source/Runtime/Engine/Private/WorldPartition/RuntimeHashSet/RuntimePartitionPersistent.cpp
--- a/Engine/Source/Runtime/Engine/Private/WorldPartition/RuntimeHashSet/RuntimePartitionPersistent.cpp
+++ b/Engine/Source/Runtime/Engine/Private/WorldPartition/RuntimeHashSet/RuntimePartitionPersistent.cpp
@@ -14,15 +14,8 @@ bool URuntimePartitionPersistent::GenerateStreaming(const TArray<const IStreamin
 	TArray<IStreamingGenerationContext::FActorInstance> CellActorInstances;
 	if (PopulateCellActorInstances(ActorSetInstances, bIsMainWorldPartition, true, CellActorInstances))
 	{
-		FCellDesc& CellDesc = OutRuntimeCellDescs.Emplace_GetRef();
-
-		CellDesc.Name = NAME_PersistentLevel;
-		CellDesc.bIsSpatiallyLoaded = false;
-		CellDesc.ContentBundleID = CellActorInstances[0].ActorSetInstance->ContentBundleID;
-		CellDesc.bBlockOnSlowStreaming = bBlockOnSlowStreaming;
-		CellDesc.bClientOnlyVisible = bClientOnlyVisible;
-		CellDesc.Priority = Priority;
-		CellDesc.ActorInstances = CellActorInstances;
+		const FString CellName(TEXT("Persistent"));
+		OutRuntimeCellDescs.Emplace(CreateCellDesc(CellName, false, CellActorInstances[0].ActorSetInstance->ContentBundleID, 0, CellActorInstances));
 	}
 
 	return true;
diff --git a/Engine/Source/Runtime/Engine/Private/WorldPartition/RuntimeHashSet/WorldPartitionRuntimeHashSet.cpp b/Engine/Source/Runtime/Engine/Private/WorldPartition/RuntimeHashSet/WorldPartitionRuntimeHashSet.cpp
--- a/Engine/Source/Runtime/Engine/Private/WorldPartition/RuntimeHashSet/WorldPartitionRuntimeHashSet.cpp
+++ b/Engine/Source/Runtime/Engine/Private/WorldPartition/RuntimeHashSet/WorldPartitionRuntimeHashSet.cpp
@@ -462,6 +462,7 @@ bool UWorldPartitionRuntimeHashSet::GenerateStreaming(UWorldPartitionStreamingPo
 	//
 	TArray<FCellDescInstance> RuntimeCellDescsInstances;
 	{
+		TSet<FName> CellDescsNames;
 		for (auto [RuntimePartition, ActorSetInstances] : RuntimePartitionsToActorSetMap)
 		{
 			// Gather runtime partition cell descs
@@ -476,6 +477,10 @@ bool UWorldPartitionRuntimeHashSet::GenerateStreaming(UWorldPartitionStreamingPo
 			{
 				TMap<FDataLayersID, FCellDescInstance> RuntimeCellDescsInstancesSet;
 
+				bool bCellNameExists;
+				CellDescsNames.Add(RuntimeCellDesc.Name, &bCellNameExists);
+				check(!bCellNameExists);
+
 				for (const IStreamingGenerationContext::FActorInstance& ActorInstance : RuntimeCellDesc.ActorInstances)
 				{
 					const FDataLayersID DataLayersID(ActorInstance.ActorSetInstance->DataLayers);
diff --git a/Engine/Source/Runtime/Engine/Public/WorldPartition/RuntimeHashSet/RuntimePartition.h b/Engine/Source/Runtime/Engine/Public/WorldPartition/RuntimeHashSet/RuntimePartition.h
--- a/Engine/Source/Runtime/Engine/Public/WorldPartition/RuntimeHashSet/RuntimePartition.h
+++ b/Engine/Source/Runtime/Engine/Public/WorldPartition/RuntimeHashSet/RuntimePartition.h
@@ -14,7 +14,7 @@ class URuntimePartition : public UObject
 public:
 #if WITH_EDITOR
 	//~ Begin UObject Interface.
-	virtual void PostEditChangeProperty(FPropertyChangedEvent& PropertyChangedEvent) override;
+	virtual void PostEditChangeProperty(FPropertyChangedEvent& InPropertyChangedEvent) override;
 	//~ End UObject Interface.
 
 	struct FCellDesc
@@ -34,8 +34,8 @@ public:
 	};
 
 	virtual bool SupportsHLODs() const PURE_VIRTUAL(URuntimePartition::SupportsHLODs, return false;);
-	virtual bool IsValidGrid(FName GridName) const PURE_VIRTUAL(URuntimePartition::IsValidGrid, return false;);
-	virtual bool GenerateStreaming(const TArray<const IStreamingGenerationContext::FActorSetInstance*>& ActorSetInstances, TArray<FCellDesc>& OutRuntimeCellDescs) PURE_VIRTUAL(URuntimePartition::GenerateStreaming, return false;);
+	virtual bool IsValidGrid(FName InGridName) const PURE_VIRTUAL(URuntimePartition::IsValidGrid, return false;);
+	virtual bool GenerateStreaming(const TArray<const IStreamingGenerationContext::FActorSetInstance*>& InActorSetInstances, TArray<FCellDesc>& OutRuntimeCellDescs) PURE_VIRTUAL(URuntimePartition::GenerateStreaming, return false;);
 #endif
 
 #if WITH_EDITORONLY_DATA
@@ -60,6 +60,7 @@ public:
 
 protected:
 #if WITH_EDITOR
-	bool PopulateCellActorInstances(const TArray<const IStreamingGenerationContext::FActorSetInstance*>& ActorSetInstances, bool bIsMainWorldPartition, bool bIsCellAlwaysLoaded, TArray<IStreamingGenerationContext::FActorInstance>& OutCellActorInstances);
+	bool PopulateCellActorInstances(const TArray<const IStreamingGenerationContext::FActorSetInstance*>& InActorSetInstances, bool bInIsMainWorldPartition, bool bInIsCellAlwaysLoaded, TArray<IStreamingGenerationContext::FActorInstance>& OutCellActorInstances);
+	FCellDesc CreateCellDesc(const FString& InName, bool bInIsSpatiallyLoaded, const FGuid& InContentBundleID, int32 InLevel, const TArray<IStreamingGenerationContext::FActorInstance>& InActorInstances);
 #endif
 };
\ No newline at end of file
