Log actor & property name when RepLayout reads an out-of-range property handle. It's likely a bug in the serialization of that property that causes this error.

diff --git a/Engine/Source/Runtime/Engine/Private/RepLayout.cpp b/Engine/Source/Runtime/Engine/Private/RepLayout.cpp
--- a/Engine/Source/Runtime/Engine/Private/RepLayout.cpp
+++ b/Engine/Source/Runtime/Engine/Private/RepLayout.cpp
@@ -3736,6 +3736,12 @@ static bool ReceiveProperties_r(FReceivePropertiesSharedParams& Params, FReceive
 			// Read the next property handle to serialize.
 			// If we don't have any more properties, this could be a terminator.
 			ReadPropertyHandle(Params);
+
+			if (Params.ReadHandle != 0 && StackParams.CurrentHandle > Params.ReadHandle)
+			{
+				// Serialization of this property possibly has a bug and corrupted state, causing an invalid handle value to be read.
+				UE_LOG(LogRep, Error, TEXT("Replicated property %s has likely corrupted serialization of %s. Check its serialization code."), *Cmd.Property->GetFullName(), *GetFullNameSafe(Params.OwningObject));
+			}
 		}
 	}
 
