WorldMinimap: Prevent copy/pasting

#rnx

diff --git a/Engine/Source/Runtime/Engine/Public/WorldPartition/WorldPartitionMiniMap.h b/Engine/Source/Runtime/Engine/Public/WorldPartition/WorldPartitionMiniMap.h
--- a/Engine/Source/Runtime/Engine/Public/WorldPartition/WorldPartitionMiniMap.h
+++ b/Engine/Source/Runtime/Engine/Public/WorldPartition/WorldPartitionMiniMap.h
@@ -20,6 +20,7 @@ private:
 #if WITH_EDITOR
 	virtual bool IsDataLayerTypeSupported(TSubclassOf<UDataLayerInstance> DataLayerType) const final { return false; }
 	virtual bool ActorTypeIsMainWorldOnly() const override  { return true; }
+	virtual bool ShouldImport(FStringView ActorPropString, bool IsMovingLevel) override { return false; }
 #endif
 
 public:
