Introduce World Partition Loose Hierarchical Grid (LHGrid):
- Fixed base cell size with hierarchical levels, doubling in size until the last one.
- Cell placement algorithm:
  - Grid level is determined by actor�s bounds extent.
  - Cell is determined by actor�s bounds center.
- Once all actors have been placed, cells size are adjusted to represent the actors they contain, which means that adjusted cells can:
  - Be smaller by at most half the grid�s level cell size.
  - Overlap neighbor cells by at most half the grid�s level cell size.

This is how it avoids the promotion problem we had with the previous implementation.

diff --git a/Engine/Source/Runtime/Engine/Private/WorldPartition/RuntimeHashSet/RuntimePartitionLHGrid.cpp b/Engine/Source/Runtime/Engine/Private/WorldPartition/RuntimeHashSet/RuntimePartitionLHGrid.cpp
new file mode 100644
--- /dev/null
+++ b/Engine/Source/Runtime/Engine/Private/WorldPartition/RuntimeHashSet/RuntimePartitionLHGrid.cpp
@@ -0,0 +1,107 @@
+// Copyright Epic Games, Inc. All Rights Reserved.
+
+#include "WorldPartition/RuntimeHashSet/RuntimePartitionLHGrid.h"
+#include "Misc/HashBuilder.h"
+
+#if WITH_EDITOR
+struct FCellCoord
+{
+	FCellCoord(int64 InX, int64 InY, int64 InZ, int32 InLevel)
+		: X(InX)
+		, Y(InY)
+		, Z(InZ)
+		, Level(InLevel)
+	{}
+
+	int64 X;
+	int64 Y;
+	int64 Z;
+	int32 Level;
+
+	inline FString ToString() const
+	{
+		return FString::Printf(TEXT("%d_%d_%d_%d"), X, Y, Z, Level);
+	}
+
+	inline bool operator==(const FCellCoord& Other) const
+	{
+		return (X == Other.X) && (Y == Other.Y) && (Z == Other.Z) && (Level == Other.Level);
+	}
+
+	static inline int32 GetLevelForBox(const FBox& InBox, int32 InCellSize)
+	{
+		const FVector Extent = InBox.GetExtent();
+		const FVector::FReal MaxLength = Extent.GetMax() * 2.0;
+		return FMath::CeilToInt32(FMath::Max<FVector::FReal>(FMath::Log2(MaxLength / InCellSize), 0));
+	}
+
+	static inline FCellCoord GetCellCoords(const FVector& InPos, int32 InCellSize, int32 InLevel)
+	{
+		check(InLevel >= 0);
+		const int64 CellSizeForLevel = (int64)InCellSize * (1LL << InLevel);
+		return FCellCoord(
+			FMath::FloorToInt(InPos.X / CellSizeForLevel),
+			FMath::FloorToInt(InPos.Y / CellSizeForLevel),
+			FMath::FloorToInt(InPos.Z / CellSizeForLevel),
+			InLevel
+		);
+	}
+
+	friend uint32 GetTypeHash(const FCellCoord& CellCoord)
+	{
+		FHashBuilder HashBuilder;
+		HashBuilder << CellCoord.X << CellCoord.Y << CellCoord.Z << CellCoord.Level;
+		return HashBuilder.GetHash();
+	}
+};
+
+void URuntimePartitionLHGrid::PostEditChangeProperty(FPropertyChangedEvent& PropertyChangedEvent)
+{
+	const FName PropertyName = PropertyChangedEvent.GetPropertyName();
+
+	if (PropertyName == GET_MEMBER_NAME_CHECKED(URuntimePartitionLHGrid, CellSize))
+	{
+		CellSize = FMath::Max<int32>(CellSize, 1600);
+	}
+
+	Super::PostEditChangeProperty(PropertyChangedEvent);
+}
+
+bool URuntimePartitionLHGrid::SupportsHLODs() const
+{
+	return true;
+}
+
+bool URuntimePartitionLHGrid::IsValidGrid(FName GridName) const
+{
+	const TArray<FName> GridNameList = UWorldPartitionRuntimeHashSet::ParseGridName(GridName);
+	return GridNameList.Num() == 1;
+}
+
+bool URuntimePartitionLHGrid::GenerateStreaming(const TArray<const IStreamingGenerationContext::FActorSetInstance*>& ActorSetInstances, TArray<FCellDesc>& OutRuntimeCellDescs)
+{
+	UWorldPartition* WorldPartition = GetTypedOuter<UWorldPartition>();
+	UWorld* World = WorldPartition->GetWorld();
+	UWorld* OuterWorld = GetTypedOuter<UWorld>();
+	const bool bIsMainWorldPartition = (World == OuterWorld);
+
+	TArray<IStreamingGenerationContext::FActorInstance> CellActorInstances;
+	if (PopulateCellActorInstances(ActorSetInstances, bIsMainWorldPartition, false, CellActorInstances))
+	{
+		TMap<FCellCoord, TArray<IStreamingGenerationContext::FActorInstance>> SubLevelsActorInstances;
+		for (const IStreamingGenerationContext::FActorInstance& ActorInstance : CellActorInstances)
+		{
+			const int32 GridLevel = FCellCoord::GetLevelForBox(ActorInstance.GetBounds(), CellSize);
+			const FCellCoord CellCoord = FCellCoord::GetCellCoords(ActorInstance.GetBounds().GetCenter(), CellSize, GridLevel);
+			SubLevelsActorInstances.FindOrAdd(CellCoord).Add(ActorInstance);
+		}
+
+		for (auto& [CellCoord, SubLevelActorSetInstances] : SubLevelsActorInstances)
+		{
+			OutRuntimeCellDescs.Emplace(CreateCellDesc(CellCoord.ToString(), true, SubLevelActorSetInstances[0].ActorSetInstance->ContentBundleID, CellCoord.Level, SubLevelActorSetInstances));
+		}
+	}
+
+	return true;
+}
+#endif
\ No newline at end of file
diff --git a/Engine/Source/Runtime/Engine/Public/WorldPartition/RuntimeHashSet/RuntimePartitionLHGrid.h b/Engine/Source/Runtime/Engine/Public/WorldPartition/RuntimeHashSet/RuntimePartitionLHGrid.h
new file mode 100644
--- /dev/null
+++ b/Engine/Source/Runtime/Engine/Public/WorldPartition/RuntimeHashSet/RuntimePartitionLHGrid.h
@@ -0,0 +1,29 @@
+// Copyright Epic Games, Inc. All Rights Reserved.
+#pragma once
+
+#include "CoreMinimal.h"
+#include "WorldPartition/RuntimeHashSet/RuntimePartition.h"
+#include "RuntimePartitionLHGrid.generated.h"
+
+UCLASS()
+class URuntimePartitionLHGrid : public URuntimePartition
+{
+	GENERATED_BODY()
+
+public:
+#if WITH_EDITOR
+	//~ Begin UObject Interface.
+	virtual void PostEditChangeProperty(FPropertyChangedEvent& PropertyChangedEvent) override;
+	//~ End UObject Interface.
+
+	virtual bool SupportsHLODs() const override;
+	virtual bool IsValidGrid(FName GridName) const override;
+	virtual bool GenerateStreaming(const TArray<const IStreamingGenerationContext::FActorSetInstance*>& ActorSetInstances, TArray<FCellDesc>& OutRuntimeCellDescs) override;
+#endif
+
+private:
+#if WITH_EDITORONLY_DATA
+	UPROPERTY(EditAnywhere, Category = RuntimeSettings)
+	uint32 CellSize;
+#endif
+};
\ No newline at end of file
