World Partition
* Add new unsaved actors into the DirtyActors set. Useful to see unsaved LocationVolume on the WP Editor.
* Make sure we clean DirtyActors when CB content is removed.

diff --git a/Engine/Source/Editor/WorldPartitionEditor/Private/WorldPartition/SWorldPartitionEditorGrid2D.cpp b/Engine/Source/Editor/WorldPartitionEditor/Private/WorldPartition/SWorldPartitionEditorGrid2D.cpp
--- a/Engine/Source/Editor/WorldPartitionEditor/Private/WorldPartition/SWorldPartitionEditorGrid2D.cpp
+++ b/Engine/Source/Editor/WorldPartitionEditor/Private/WorldPartition/SWorldPartitionEditorGrid2D.cpp
@@ -1089,15 +1089,18 @@ void SWorldPartitionEditorGrid2D::Tick(const FGeometry& AllottedGeometry, const
 	// they will never get an actor descriptor so they will never appear in the world partition editor. Also include unsaved, newly created actors for convenience.
 	for (auto& [Reference, Actor] : GetWorldPartition()->GetDirtyActors())
 	{
-		if (!GetWorldPartition()->GetActorDesc(Actor->GetActorGuid()) && Actor->Implements<UWorldPartitionActorLoaderInterface>())
+		if (Actor.IsValid())
 		{
-			if (IWorldPartitionActorLoaderInterface::ILoaderAdapter* LoaderAdapter = Cast<IWorldPartitionActorLoaderInterface>(Actor)->GetLoaderAdapter())
+			if (!GetWorldPartition()->GetActorDesc(Actor->GetActorGuid()) && Actor->Implements<UWorldPartitionActorLoaderInterface>())
 			{
-				ShownLoaderInterfaces.Add(Actor);
+				if (IWorldPartitionActorLoaderInterface::ILoaderAdapter* LoaderAdapter = Cast<IWorldPartitionActorLoaderInterface>(Actor)->GetLoaderAdapter())
+				{
+					ShownLoaderInterfaces.Add(Actor.Get());
+				}
 			}
-		}
 
-		DirtyActorGuids.Add(Actor->GetActorGuid());
+			DirtyActorGuids.Add(Actor->GetActorGuid());
+		}
 	}
 
 	FLoaderInterfaceSet LastHoveredLoaderInterfaces = MoveTemp(HoveredLoaderInterfaces);
diff --git a/Engine/Source/Runtime/Engine/Private/WorldPartition/WorldPartition.cpp b/Engine/Source/Runtime/Engine/Private/WorldPartition/WorldPartition.cpp
--- a/Engine/Source/Runtime/Engine/Private/WorldPartition/WorldPartition.cpp
+++ b/Engine/Source/Runtime/Engine/Private/WorldPartition/WorldPartition.cpp
@@ -43,6 +43,8 @@
 #include "WorldPartition/LoaderAdapter/LoaderAdapterPinnedActors.h"
 #include "WorldPartition/HLOD/HLODLayer.h"
 #include "WorldPartition/Cook/WorldPartitionCookPackageContextInterface.h"
+#include "WorldPartition/ContentBundle/ContentBundleEditorSubsystemInterface.h"
+#include "WorldPartition/ContentBundle/ContentBundleEditor.h"
 #include "Modules/ModuleManager.h"
 #include "GameDelegates.h"
 #else
@@ -351,9 +353,14 @@ void UWorldPartition::OnPackageDirtyStateChanged(UPackage* Package)
 		{
 			if (Package->IsDirty())
 			{
-				DirtyActors.Add(ActorHandle.ToReference(), Actor);
+				DirtyActors.Add(FDirtyActor(ActorHandle.ToReference(), Actor));
 			}
 		}
+		else
+		{
+			// This is handling a new actor (unsaved).
+			DirtyActors.Add(FDirtyActor(Actor));
+		}
 	}
 }
 
@@ -897,6 +904,7 @@ void UWorldPartition::RegisterDelegates()
 			GEditor->OnPostBugItGoCalled().AddUObject(this, &UWorldPartition::OnPostBugItGoCalled);
 			GEditor->OnEditorClose().AddUObject(this, &UWorldPartition::SavePerUserSettings);
 			FWorldDelegates::OnPostWorldRename.AddUObject(this, &UWorldPartition::OnWorldRenamed);
+			IContentBundleEditorSubsystemInterface::Get()->OnContentBundleRemovedContent().AddUObject(this, &UWorldPartition::OnContentBundleRemovedContent);
 		}
 
 		if (!IsRunningCommandlet())
@@ -928,6 +936,11 @@ void UWorldPartition::UnregisterDelegates()
 	{
 		if (IsMainWorldPartition())
 		{
+			if (IContentBundleEditorSubsystemInterface* ContentBundleEditorSubsystem = IContentBundleEditorSubsystemInterface::Get())
+			{
+				ContentBundleEditorSubsystem->OnContentBundleRemovedContent().RemoveAll(this);
+			}
+
 			FWorldDelegates::OnPostWorldRename.RemoveAll(this);
 			FEditorDelegates::PreBeginPIE.RemoveAll(this);
 			FEditorDelegates::PrePIEEnded.RemoveAll(this);
@@ -1216,7 +1229,7 @@ void UWorldPartition::OnActorDescAdded(FWorldPartitionActorDesc* NewActorDesc)
 
 	if (AActor* NewActor = NewActorDesc->GetActor())
 	{
-		DirtyActors.Add(FWorldPartitionReference(NewActorDesc->GetContainer(), NewActorDesc->GetGuid()), NewActor);
+		DirtyActors.Add(FDirtyActor(FWorldPartitionReference(NewActorDesc->GetContainer(), NewActorDesc->GetGuid()), NewActor));
 	}
 
 	if (ForceLoadedActors)
@@ -1367,6 +1380,33 @@ void UWorldPartition::UnhashActorDesc(FWorldPartitionActorDesc* ActorDesc)
 	FWorldPartitionHandle ActorHandle(this, ActorDesc->GetGuid());
 	EditorHash->UnhashActor(ActorHandle);
 }
+
+void UWorldPartition::OnContentBundleRemovedContent(const FContentBundleEditor* ContentBundle)
+{
+	check(ContentBundle);
+	
+	const TWeakObjectPtr<UActorDescContainer>& ContentBundleActorDescContainer = ContentBundle->GetActorDescContainer();
+	if (!ContentBundleActorDescContainer.IsValid())
+	{
+		return;
+	}
+
+	if (Contains(ContentBundleActorDescContainer.Get()->GetContainerPackage()))
+	{
+		// This is handling a new actor (unsaved) from Content Bundle.
+		for (TSet<FDirtyActor>::TIterator DirtyActorIt(DirtyActors); DirtyActorIt; ++DirtyActorIt)
+		{
+			const FDirtyActor& DirtyActor = *DirtyActorIt;
+			if (!DirtyActor.WorldPartitionRef.IsSet())
+			{
+				if (DirtyActor.ActorPtr.IsValid() && DirtyActor.ActorPtr->GetContentBundleGuid() == ContentBundleActorDescContainer.Get()->GetContentBundleGuid())
+				{
+					DirtyActorIt.RemoveCurrent();
+				}
+			}
+		}
+	}
+}
 #endif
 
 void UWorldPartition::Serialize(FArchive& Ar)
@@ -1475,7 +1515,10 @@ void UWorldPartition::AddReferencedObjects(UObject* InThis, FReferenceCollector&
 		Collector.AllowEliminatingReferences(false);
 		for (auto& [ActorReference, Actor] : This->DirtyActors)
 		{
-			Collector.AddReferencedObject(Actor);
+			if (ActorReference.IsSet() || Actor.IsValid())
+			{
+				Collector.AddReferencedObject(Actor);
+			}
 		}
 		Collector.AllowEliminatingReferences(true);
 	}
@@ -1497,21 +1540,43 @@ void UWorldPartition::Tick(float DeltaSeconds)
 		EditorHash->Tick(DeltaSeconds);
 	}
 
-	for (TMap<FWorldPartitionReference, TObjectPtr<AActor>>::TIterator DirtyActorIt(DirtyActors); DirtyActorIt; ++DirtyActorIt)
+	for (TSet<FDirtyActor>::TIterator DirtyActorIt(DirtyActors); DirtyActorIt; ++DirtyActorIt)
 	{
-		if (!DirtyActorIt.Key().IsValid())
+		const FDirtyActor& DirtyActor = *DirtyActorIt;
+		if (DirtyActor.WorldPartitionRef.IsSet())
 		{
-			DirtyActorIt.RemoveCurrent();
+			const FWorldPartitionReference& Ref = DirtyActor.WorldPartitionRef.GetValue();
+
+			if (!Ref.IsValid())
+			{
+				DirtyActorIt.RemoveCurrent();
+			}
+			else if (DirtyActor.ActorPtr.IsValid() && !DirtyActor.ActorPtr->GetPackage()->IsDirty())
+			{
+				// Transfer ownership of the last ref if actor can be pinned
+				if (Ref->GetHardRefCount() <= 1 && PinnedActors && FLoaderAdapterPinnedActors::SupportsPinning(Ref.Get()))
+				{
+					PinnedActors->AddActors({ Ref.ToHandle() });
+					DirtyActorIt.RemoveCurrent();
+				} // Clean up if a loader took a reference on the actor (ILoaderAdapter::RefreshLoadedState was called since)
+				else if (Ref->GetHardRefCount() > 1)
+				{
+					DirtyActorIt.RemoveCurrent();
+				}
+			}
 		}
-		else if (!DirtyActorIt.Value()->GetPackage()->IsDirty())
+		else
 		{
-			// Transfer ownership of the last ref if actor can be pinned
-			if (DirtyActorIt.Key()->GetHardRefCount() <= 1 && PinnedActors && FLoaderAdapterPinnedActors::SupportsPinning(DirtyActorIt.Key().Get()))
+			// This is handling a new actor (unsaved).
+			if (DirtyActor.ActorPtr.IsValid())
 			{
-				PinnedActors->AddActors({ DirtyActorIt.Key().ToHandle() });
-				DirtyActorIt.RemoveCurrent();
-			} // Clean up if a loader took a reference on the actor (ILoaderAdapter::RefreshLoadedState was called since)
-			else if (DirtyActorIt.Key()->GetHardRefCount() > 1)
+				if (!DirtyActor.ActorPtr->GetPackage()->IsDirty())
+				{
+					DirtyActorIt.RemoveCurrent();
+				}
+			}
+			// In this case, we know that the DirtyActor is not in the transaction buffer anymore and this is fine removing it.
+			else if ((!DirtyActor.ActorPtr.IsValid(true)))
 			{
 				DirtyActorIt.RemoveCurrent();
 			}
@@ -1943,11 +2008,15 @@ bool UWorldPartition::UnregisterActorDescContainer(UActorDescContainer* InActorD
 			{
 				ActorGuids.Add(It->GetGuid());
 
-				for (TMap<FWorldPartitionReference, TObjectPtr<AActor>>::TIterator DirtyActorIt(DirtyActors); DirtyActorIt; ++DirtyActorIt)
+				for (TSet<FDirtyActor>::TIterator DirtyActorIt(DirtyActors); DirtyActorIt; ++DirtyActorIt)
 				{
-					if (DirtyActorIt.Key() == ActorHandle)
+					const FDirtyActor& DirtyActor = *DirtyActorIt;
+					if (DirtyActor.WorldPartitionRef.IsSet())
 					{
-						DirtyActorIt.RemoveCurrent();
+						if (DirtyActor.WorldPartitionRef.GetValue() == ActorHandle)
+						{
+							DirtyActorIt.RemoveCurrent();
+						}
 					}
 				}
 			}
diff --git a/Engine/Source/Runtime/Engine/Public/WorldPartition/WorldPartition.h b/Engine/Source/Runtime/Engine/Public/WorldPartition/WorldPartition.h
--- a/Engine/Source/Runtime/Engine/Public/WorldPartition/WorldPartition.h
+++ b/Engine/Source/Runtime/Engine/Public/WorldPartition/WorldPartition.h
@@ -99,6 +99,42 @@ public:
 };
 #endif
 
+struct ENGINE_API FDirtyActor
+{
+	FDirtyActor()
+		: ActorPtr(nullptr)
+	{}
+
+	FDirtyActor(AActor* InActor)
+		: ActorPtr(InActor)
+	{
+	}
+
+	FDirtyActor(const FWorldPartitionReference& InWorldPartitionRef, AActor* InActor)
+		: WorldPartitionRef(InWorldPartitionRef)
+		, ActorPtr(InActor)
+	{
+	}
+
+	TOptional<FWorldPartitionReference> WorldPartitionRef;
+	TWeakObjectPtr<AActor>	ActorPtr;		// TWeakObjectPtr is for undo support.
+
+	bool operator == (const FDirtyActor& InDirtyActor) const
+	{
+		return WorldPartitionRef == InDirtyActor.WorldPartitionRef && ActorPtr == InDirtyActor.ActorPtr;
+	}
+
+	friend uint32 GetTypeHash(const FDirtyActor& InDirtyActor)
+	{
+		uint32 Hash = GetTypeHash(InDirtyActor.ActorPtr);
+		if (InDirtyActor.WorldPartitionRef.IsSet())
+		{
+			Hash = HashCombine(Hash, GetTypeHash(InDirtyActor.WorldPartitionRef.GetValue()));
+		}
+		return Hash;
+	}
+};
+
 UCLASS(AutoExpandCategories=(WorldPartition), MinimalAPI)
 class UWorldPartition final : public UObject, public FActorDescContainerCollection, public IWorldPartitionCookPackageGenerator
 {
@@ -319,7 +355,7 @@ public:
 
 	bool IsEnablingStreamingJustified() const { return bEnablingStreamingJustified; }
 
-	const TMap<FWorldPartitionReference, AActor*>& GetDirtyActors() const { return ObjectPtrDecay(DirtyActors); }
+	const TSet<FDirtyActor>& GetDirtyActors() const { return DirtyActors; }
 #endif
 
 public:
@@ -435,7 +471,7 @@ private:
 
 	TArray<FWorldPartitionReference> LoadedSubobjects;
 
-	TMap<FWorldPartitionReference, TObjectPtr<AActor>> DirtyActors;
+	TSet<FDirtyActor> DirtyActors;
 
 	TSet<FString> GeneratedStreamingPackageNames;
 
@@ -497,8 +533,7 @@ private:
 #if WITH_EDITOR
 	ENGINE_API void HashActorDesc(FWorldPartitionActorDesc* ActorDesc);
 	ENGINE_API void UnhashActorDesc(FWorldPartitionActorDesc* ActorDesc);
-	ENGINE_API void HashActorDescContainer(UActorDescContainer* ActorDescContainer);
-	ENGINE_API void UnhashActorDescContainer(UActorDescContainer* ActorDescContainer);
+	void OnContentBundleRemovedContent(const FContentBundleEditor* ContentBundle);
 
 public:
 	// Editor loader adapters management
