Fix failed check on converting spline to polyline when spline has duplicate constant points -- handle 0-distance segments by directly adding the segment's point

diff --git a/Engine/Source/Runtime/Engine/Private/Components/SplineComponent.cpp b/Engine/Source/Runtime/Engine/Private/Components/SplineComponent.cpp
--- a/Engine/Source/Runtime/Engine/Private/Components/SplineComponent.cpp
+++ b/Engine/Source/Runtime/Engine/Private/Components/SplineComponent.cpp
@@ -1678,8 +1678,9 @@ bool USplineComponent::ConvertSplineSegmentToPolyLine(int32 SplinePointStartInde
 	double SubstepSize = Dist / NumLines;
 	if (SubstepSize == 0.0)
 	{
-		// Make sure there's at least 1 sub-step so that we get a single spline vertex for constant interpolation : 
-		SubstepSize = StopDist + 1.0;
+		// There is no distance to cover, so handle the segment with a single point
+		OutPoints.Add(GetLocationAtDistanceAlongSpline(StopDist, CoordinateSpace));
+		return true;
 	}
 
 	double SubstepStartDist = StartDist;
