Recompile shader jobs locally in distributed shader compiler on failed task results (such as OOM error).

#rnx

diff --git a/Engine/Source/Runtime/Engine/Private/ShaderCompiler/ShaderCompiler.cpp b/Engine/Source/Runtime/Engine/Private/ShaderCompiler/ShaderCompiler.cpp
--- a/Engine/Source/Runtime/Engine/Private/ShaderCompiler/ShaderCompiler.cpp
+++ b/Engine/Source/Runtime/Engine/Private/ShaderCompiler/ShaderCompiler.cpp
@@ -1276,26 +1276,6 @@ static void ReissueShaderCompileJobs(const TArray<FShaderCommonCompileJob*>& Sou
 // Make functions so the crash reporter can disambiguate the actual error because of the different callstacks
 namespace ShaderCompileWorkerError
 {
-	static bool TryReissueShaderCompileJobs(const TArray<FShaderCommonCompileJobPtr>& SourceJobs)
-	{
-		// Reschedule all queued shader compile jobs when a worker crashed due to running out of memory.
-		TArray<FShaderCommonCompileJob*> ReissueSourceJobs;
-		ReissueSourceJobs.Reserve(SourceJobs.Num());
-		for (const FShaderCommonCompileJobPtr& Job : SourceJobs)
-		{
-			// Switch to local thread. If the worker already is a local thread, we cannot recover from this error.
-			FShaderCommonCompileJob* CurrentJob = Job.GetReference();
-			if (CurrentJob->CurrentWorker == EShaderCompilerWorkerType::LocalThread)
-			{
-				return false;
-			}
-			CurrentJob->CurrentWorker = EShaderCompilerWorkerType::LocalThread;
-			ReissueSourceJobs.Add(CurrentJob);
-		}
-		ReissueShaderCompileJobs(ReissueSourceJobs);
-		return true;
-	}
-
 	void HandleGeneralCrash(const TCHAR* ExceptionInfo, const TCHAR* Callstack)
 	{
 		GLog->Panic();
@@ -1382,10 +1362,10 @@ namespace ShaderCompileWorkerError
 			MemoryStats.PeakUsedVirtual
 		);
 
-		if (TryReissueShaderCompileJobs(QueuedJobs))
+		if (GShaderCompilingManager->IsRemoteCompilingEnabled())
 		{
-			// We recovered from this error
-			UE_LOG(LogShaderCompilers, Warning, TEXT("%s\nReissue %d shader compile %s to remaining workers"), *ErrorReport, QueuedJobs.Num(), (QueuedJobs.Num() == 1 ? TEXT("job") : TEXT("jobs")));
+			// Remote shader compiler supports re-compiling jobs on local machine
+			UE_LOG(LogShaderCompilers, Warning, TEXT("%s\nRecompile %d shader compile %s locally"), *ErrorReport, QueuedJobs.Num(), (QueuedJobs.Num() == 1 ? TEXT("job") : TEXT("jobs")));
 			return true;
 		}
 		else
@@ -2072,8 +2052,9 @@ struct FSCWOutputFileContext
 	}
 };
 
-// Process results from Worker Process
-void FShaderCompileUtilities::DoReadTaskResults(const TArray<FShaderCommonCompileJobPtr>& QueuedJobs, FArchive& OutputFile)
+// Process results from Worker Process.
+// Returns false if reading the tasks failed but we were able to recover from handing a crash report. In this case, all jobs must be submitted/processed again.
+FSCWErrorCode::ECode FShaderCompileUtilities::DoReadTaskResults(const TArray<FShaderCommonCompileJobPtr>& QueuedJobs, FArchive& OutputFile)
 {
 	FSCWOutputFileContext OutputFileContext(OutputFile);
 
@@ -2119,7 +2100,7 @@ void FShaderCompileUtilities::DoReadTaskResults(const TArray<FShaderCommonCompil
 		if (HandleWorkerCrash(QueuedJobs, OutputFile, OutputVersion, OutputFileContext.FileSize, (FSCWErrorCode::ECode)ErrorCode, NumProcessedJobs, CallstackLength, ExceptionInfoLength, HostnameLength))
 		{
 			FSCWErrorCode::Reset();
-			return;
+			return (FSCWErrorCode::ECode)ErrorCode;
 		}
 	}
 
@@ -2212,6 +2193,8 @@ void FShaderCompileUtilities::DoReadTaskResults(const TArray<FShaderCommonCompil
 	
 	// Requeue any jobs we wish to run again
 	ReissueShaderCompileJobs(ReissueSourceJobs);
+
+	return FSCWErrorCode::Success;
 }
 
 #if WITH_EDITOR
diff --git a/Engine/Source/Runtime/Engine/Private/ShaderCompiler/ShaderCompilerDistributed.cpp b/Engine/Source/Runtime/Engine/Private/ShaderCompiler/ShaderCompilerDistributed.cpp
--- a/Engine/Source/Runtime/Engine/Private/ShaderCompiler/ShaderCompilerDistributed.cpp
+++ b/Engine/Source/Runtime/Engine/Private/ShaderCompiler/ShaderCompilerDistributed.cpp
@@ -322,18 +322,21 @@ int32 FShaderCompileDistributedThreadRunnable_Interface::CompilingLoop()
 		if (Result.bCompleted)
 		{
 			// Check the output file exists. If it does, attempt to open it and serialize in the completed jobs.
+			bool bCompileJobsSucceeded = false;
+
 			if (IFileManager::Get().FileExists(*Task->OutputFilePath))
 			{
-				FArchive* OutputFileAr = IFileManager::Get().CreateFileReader(*Task->OutputFilePath, FILEREAD_Silent);
-				if (OutputFileAr)
+				if (TUniquePtr<FArchive> OutputFileAr = TUniquePtr<FArchive>(IFileManager::Get().CreateFileReader(*Task->OutputFilePath, FILEREAD_Silent)))
 				{
 					bOutputFileReadFailed = false;
-					FShaderCompileUtilities::DoReadTaskResults(Task->ShaderJobs, *OutputFileAr);
-					delete OutputFileAr;
+					if (FShaderCompileUtilities::DoReadTaskResults(Task->ShaderJobs, *OutputFileAr) == FSCWErrorCode::Success)
+					{
+						bCompileJobsSucceeded = true;
+					}
 				}
 			}
 
-			if (bOutputFileReadFailed)
+			if (!bCompileJobsSucceeded)
 			{
 				// Reading result from XGE job failed, so recompile shaders in current job batch locally
 				UE_LOG(LogShaderCompilers, Log, TEXT("Rescheduling shader compilation to run locally after distributed job failed: %s"), *Task->OutputFilePath);
diff --git a/Engine/Source/Runtime/Engine/Public/ShaderCompiler.h b/Engine/Source/Runtime/Engine/Public/ShaderCompiler.h
--- a/Engine/Source/Runtime/Engine/Public/ShaderCompiler.h
+++ b/Engine/Source/Runtime/Engine/Public/ShaderCompiler.h
@@ -367,7 +367,7 @@ class FShaderCompileUtilities
 {
 public:
 	static bool DoWriteTasks(const TArray<FShaderCommonCompileJobPtr>& QueuedJobs, FArchive& TransferFile, IDistributedBuildController* BuildDistributionController = nullptr, bool bUseRelativePaths = false, bool bCompressTaskFile = false);
-	static void DoReadTaskResults(const TArray<FShaderCommonCompileJobPtr>& QueuedJobs, FArchive& OutputFile);
+	static FSCWErrorCode::ECode DoReadTaskResults(const TArray<FShaderCommonCompileJobPtr>& QueuedJobs, FArchive& OutputFile);
 
 	/** Execute the specified (single or pipeline) shader compile job. */
 	static void ExecuteShaderCompileJob(FShaderCommonCompileJob& Job);
@@ -771,6 +771,16 @@ public:
 		return GetNumOutstandingJobs() > 0 || HasShaderJobs() || GetNumPendingJobs() > 0 || NumExternalJobs > 0;
 	}
 
+	/**
+	 * Returns whether remote compiling is enabled. Otherwise, only the local machine is used for shader compilation.
+	 * This depends on command line argument '-NoRemoteShaderCompile', CVar 'r.ShaderCompiler.AllowDistributedCompilation',
+	 * and whether the target platform supports remote compilin; see ITargetPlatform::CanSupportRemoteShaderCompile().
+	 */
+	bool IsRemoteCompilingEnabled() const
+	{
+		return BuildDistributionController != nullptr;
+	}
+
 	/**
 	 * Returns whether normal throttling settings should be ignored because shader compilation is at the moment the only action blocking the critical path.
 	 * An example of such situation is startup compilation of global shaders and default materials, but there may be more cases like that.
