[Backout] - CL26287803 for CIS error.
#rnx
Original CL Desc
-----------------------------------------------------------------
World Partition
* Add new unsaved actors into the DirtyActors set. Useful to see unsaved LocationVolume on the WP Editor.
* Make sure we clean DirtyActors when CB content is removed.

diff --git a/Engine/Source/Editor/WorldPartitionEditor/Private/WorldPartition/SWorldPartitionEditorGrid2D.cpp b/Engine/Source/Editor/WorldPartitionEditor/Private/WorldPartition/SWorldPartitionEditorGrid2D.cpp
--- a/Engine/Source/Editor/WorldPartitionEditor/Private/WorldPartition/SWorldPartitionEditorGrid2D.cpp
+++ b/Engine/Source/Editor/WorldPartitionEditor/Private/WorldPartition/SWorldPartitionEditorGrid2D.cpp
@@ -1089,18 +1089,15 @@ void SWorldPartitionEditorGrid2D::Tick(const FGeometry& AllottedGeometry, const
 	// they will never get an actor descriptor so they will never appear in the world partition editor. Also include unsaved, newly created actors for convenience.
 	for (auto& [Reference, Actor] : GetWorldPartition()->GetDirtyActors())
 	{
-		if (Actor.IsValid())
+		if (!GetWorldPartition()->GetActorDesc(Actor->GetActorGuid()) && Actor->Implements<UWorldPartitionActorLoaderInterface>())
 		{
-			if (!GetWorldPartition()->GetActorDesc(Actor->GetActorGuid()) && Actor->Implements<UWorldPartitionActorLoaderInterface>())
+			if (IWorldPartitionActorLoaderInterface::ILoaderAdapter* LoaderAdapter = Cast<IWorldPartitionActorLoaderInterface>(Actor)->GetLoaderAdapter())
 			{
-				if (IWorldPartitionActorLoaderInterface::ILoaderAdapter* LoaderAdapter = Cast<IWorldPartitionActorLoaderInterface>(Actor)->GetLoaderAdapter())
-				{
-					ShownLoaderInterfaces.Add(Actor.Get());
-				}
+				ShownLoaderInterfaces.Add(Actor);
 			}
-
-			DirtyActorGuids.Add(Actor->GetActorGuid());
 		}
+
+		DirtyActorGuids.Add(Actor->GetActorGuid());
 	}
 
 	FLoaderInterfaceSet LastHoveredLoaderInterfaces = MoveTemp(HoveredLoaderInterfaces);
diff --git a/Engine/Source/Runtime/Engine/Private/WorldPartition/WorldPartition.cpp b/Engine/Source/Runtime/Engine/Private/WorldPartition/WorldPartition.cpp
--- a/Engine/Source/Runtime/Engine/Private/WorldPartition/WorldPartition.cpp
+++ b/Engine/Source/Runtime/Engine/Private/WorldPartition/WorldPartition.cpp
@@ -43,8 +43,6 @@
 #include "WorldPartition/LoaderAdapter/LoaderAdapterPinnedActors.h"
 #include "WorldPartition/HLOD/HLODLayer.h"
 #include "WorldPartition/Cook/WorldPartitionCookPackageContextInterface.h"
-#include "WorldPartition/ContentBundle/ContentBundleEditorSubsystemInterface.h"
-#include "WorldPartition/ContentBundle/ContentBundleEditor.h"
 #include "Modules/ModuleManager.h"
 #include "GameDelegates.h"
 #else
@@ -353,14 +351,9 @@ void UWorldPartition::OnPackageDirtyStateChanged(UPackage* Package)
 		{
 			if (Package->IsDirty())
 			{
-				DirtyActors.Add(FDirtyActor(ActorHandle.ToReference(), Actor));
+				DirtyActors.Add(ActorHandle.ToReference(), Actor);
 			}
 		}
-		else
-		{
-			// This is handling a new actor (unsaved).
-			DirtyActors.Add(FDirtyActor(Actor));
-		}
 	}
 }
 
@@ -904,7 +897,6 @@ void UWorldPartition::RegisterDelegates()
 			GEditor->OnPostBugItGoCalled().AddUObject(this, &UWorldPartition::OnPostBugItGoCalled);
 			GEditor->OnEditorClose().AddUObject(this, &UWorldPartition::SavePerUserSettings);
 			FWorldDelegates::OnPostWorldRename.AddUObject(this, &UWorldPartition::OnWorldRenamed);
-			IContentBundleEditorSubsystemInterface::Get()->OnContentBundleRemovedContent().AddUObject(this, &UWorldPartition::OnContentBundleRemovedContent);
 		}
 
 		if (!IsRunningCommandlet())
@@ -936,11 +928,6 @@ void UWorldPartition::UnregisterDelegates()
 	{
 		if (IsMainWorldPartition())
 		{
-			if (IContentBundleEditorSubsystemInterface* ContentBundleEditorSubsystem = IContentBundleEditorSubsystemInterface::Get())
-			{
-				ContentBundleEditorSubsystem->OnContentBundleRemovedContent().RemoveAll(this);
-			}
-
 			FWorldDelegates::OnPostWorldRename.RemoveAll(this);
 			FEditorDelegates::PreBeginPIE.RemoveAll(this);
 			FEditorDelegates::PrePIEEnded.RemoveAll(this);
@@ -1229,7 +1216,7 @@ void UWorldPartition::OnActorDescAdded(FWorldPartitionActorDesc* NewActorDesc)
 
 	if (AActor* NewActor = NewActorDesc->GetActor())
 	{
-		DirtyActors.Add(FDirtyActor(FWorldPartitionReference(NewActorDesc->GetContainer(), NewActorDesc->GetGuid()), NewActor));
+		DirtyActors.Add(FWorldPartitionReference(NewActorDesc->GetContainer(), NewActorDesc->GetGuid()), NewActor);
 	}
 
 	if (ForceLoadedActors)
@@ -1380,33 +1367,6 @@ void UWorldPartition::UnhashActorDesc(FWorldPartitionActorDesc* ActorDesc)
 	FWorldPartitionHandle ActorHandle(this, ActorDesc->GetGuid());
 	EditorHash->UnhashActor(ActorHandle);
 }
-
-void UWorldPartition::OnContentBundleRemovedContent(const FContentBundleEditor* ContentBundle)
-{
-	check(ContentBundle);
-	
-	const TWeakObjectPtr<UActorDescContainer>& ContentBundleActorDescContainer = ContentBundle->GetActorDescContainer();
-	if (!ContentBundleActorDescContainer.IsValid())
-	{
-		return;
-	}
-
-	if (Contains(ContentBundleActorDescContainer.Get()->GetContainerPackage()))
-	{
-		// This is handling a new actor (unsaved) from Content Bundle.
-		for (TSet<FDirtyActor>::TIterator DirtyActorIt(DirtyActors); DirtyActorIt; ++DirtyActorIt)
-		{
-			const FDirtyActor& DirtyActor = *DirtyActorIt;
-			if (!DirtyActor.WorldPartitionRef.IsSet())
-			{
-				if (DirtyActor.ActorPtr.IsValid() && DirtyActor.ActorPtr->GetContentBundleGuid() == ContentBundleActorDescContainer.Get()->GetContentBundleGuid())
-				{
-					DirtyActorIt.RemoveCurrent();
-				}
-			}
-		}
-	}
-}
 #endif
 
 void UWorldPartition::Serialize(FArchive& Ar)
@@ -1515,10 +1475,7 @@ void UWorldPartition::AddReferencedObjects(UObject* InThis, FReferenceCollector&
 		Collector.AllowEliminatingReferences(false);
 		for (auto& [ActorReference, Actor] : This->DirtyActors)
 		{
-			if (ActorReference.IsSet() || Actor.IsValid())
-			{
-				Collector.AddReferencedObject(Actor);
-			}
+			Collector.AddReferencedObject(Actor);
 		}
 		Collector.AllowEliminatingReferences(true);
 	}
@@ -1540,43 +1497,21 @@ void UWorldPartition::Tick(float DeltaSeconds)
 		EditorHash->Tick(DeltaSeconds);
 	}
 
-	for (TSet<FDirtyActor>::TIterator DirtyActorIt(DirtyActors); DirtyActorIt; ++DirtyActorIt)
+	for (TMap<FWorldPartitionReference, TObjectPtr<AActor>>::TIterator DirtyActorIt(DirtyActors); DirtyActorIt; ++DirtyActorIt)
 	{
-		const FDirtyActor& DirtyActor = *DirtyActorIt;
-		if (DirtyActor.WorldPartitionRef.IsSet())
+		if (!DirtyActorIt.Key().IsValid())
 		{
-			const FWorldPartitionReference& Ref = DirtyActor.WorldPartitionRef.GetValue();
-
-			if (!Ref.IsValid())
-			{
-				DirtyActorIt.RemoveCurrent();
-			}
-			else if (DirtyActor.ActorPtr.IsValid() && !DirtyActor.ActorPtr->GetPackage()->IsDirty())
-			{
-				// Transfer ownership of the last ref if actor can be pinned
-				if (Ref->GetHardRefCount() <= 1 && PinnedActors && FLoaderAdapterPinnedActors::SupportsPinning(Ref.Get()))
-				{
-					PinnedActors->AddActors({ Ref.ToHandle() });
-					DirtyActorIt.RemoveCurrent();
-				} // Clean up if a loader took a reference on the actor (ILoaderAdapter::RefreshLoadedState was called since)
-				else if (Ref->GetHardRefCount() > 1)
-				{
-					DirtyActorIt.RemoveCurrent();
-				}
-			}
+			DirtyActorIt.RemoveCurrent();
 		}
-		else
+		else if (!DirtyActorIt.Value()->GetPackage()->IsDirty())
 		{
-			// This is handling a new actor (unsaved).
-			if (DirtyActor.ActorPtr.IsValid())
+			// Transfer ownership of the last ref if actor can be pinned
+			if (DirtyActorIt.Key()->GetHardRefCount() <= 1 && PinnedActors && FLoaderAdapterPinnedActors::SupportsPinning(DirtyActorIt.Key().Get()))
 			{
-				if (!DirtyActor.ActorPtr->GetPackage()->IsDirty())
-				{
-					DirtyActorIt.RemoveCurrent();
-				}
-			}
-			// In this case, we know that the DirtyActor is not in the transaction buffer anymore and this is fine removing it.
-			else if ((!DirtyActor.ActorPtr.IsValid(true)))
+				PinnedActors->AddActors({ DirtyActorIt.Key().ToHandle() });
+				DirtyActorIt.RemoveCurrent();
+			} // Clean up if a loader took a reference on the actor (ILoaderAdapter::RefreshLoadedState was called since)
+			else if (DirtyActorIt.Key()->GetHardRefCount() > 1)
 			{
 				DirtyActorIt.RemoveCurrent();
 			}
@@ -2008,15 +1943,11 @@ bool UWorldPartition::UnregisterActorDescContainer(UActorDescContainer* InActorD
 			{
 				ActorGuids.Add(It->GetGuid());
 
-				for (TSet<FDirtyActor>::TIterator DirtyActorIt(DirtyActors); DirtyActorIt; ++DirtyActorIt)
+				for (TMap<FWorldPartitionReference, TObjectPtr<AActor>>::TIterator DirtyActorIt(DirtyActors); DirtyActorIt; ++DirtyActorIt)
 				{
-					const FDirtyActor& DirtyActor = *DirtyActorIt;
-					if (DirtyActor.WorldPartitionRef.IsSet())
+					if (DirtyActorIt.Key() == ActorHandle)
 					{
-						if (DirtyActor.WorldPartitionRef.GetValue() == ActorHandle)
-						{
-							DirtyActorIt.RemoveCurrent();
-						}
+						DirtyActorIt.RemoveCurrent();
 					}
 				}
 			}
diff --git a/Engine/Source/Runtime/Engine/Public/WorldPartition/WorldPartition.h b/Engine/Source/Runtime/Engine/Public/WorldPartition/WorldPartition.h
--- a/Engine/Source/Runtime/Engine/Public/WorldPartition/WorldPartition.h
+++ b/Engine/Source/Runtime/Engine/Public/WorldPartition/WorldPartition.h
@@ -99,42 +99,6 @@ public:
 };
 #endif
 
-struct ENGINE_API FDirtyActor
-{
-	FDirtyActor()
-		: ActorPtr(nullptr)
-	{}
-
-	FDirtyActor(AActor* InActor)
-		: ActorPtr(InActor)
-	{
-	}
-
-	FDirtyActor(const FWorldPartitionReference& InWorldPartitionRef, AActor* InActor)
-		: WorldPartitionRef(InWorldPartitionRef)
-		, ActorPtr(InActor)
-	{
-	}
-
-	TOptional<FWorldPartitionReference> WorldPartitionRef;
-	TWeakObjectPtr<AActor>	ActorPtr;		// TWeakObjectPtr is for undo support.
-
-	bool operator == (const FDirtyActor& InDirtyActor) const
-	{
-		return WorldPartitionRef == InDirtyActor.WorldPartitionRef && ActorPtr == InDirtyActor.ActorPtr;
-	}
-
-	friend uint32 GetTypeHash(const FDirtyActor& InDirtyActor)
-	{
-		uint32 Hash = GetTypeHash(InDirtyActor.ActorPtr);
-		if (InDirtyActor.WorldPartitionRef.IsSet())
-		{
-			Hash = HashCombine(Hash, GetTypeHash(InDirtyActor.WorldPartitionRef.GetValue()));
-		}
-		return Hash;
-	}
-};
-
 UCLASS(AutoExpandCategories=(WorldPartition), MinimalAPI)
 class UWorldPartition final : public UObject, public FActorDescContainerCollection, public IWorldPartitionCookPackageGenerator
 {
@@ -355,7 +319,7 @@ public:
 
 	bool IsEnablingStreamingJustified() const { return bEnablingStreamingJustified; }
 
-	const TSet<FDirtyActor>& GetDirtyActors() const { return DirtyActors; }
+	const TMap<FWorldPartitionReference, AActor*>& GetDirtyActors() const { return ObjectPtrDecay(DirtyActors); }
 #endif
 
 public:
@@ -471,7 +435,7 @@ private:
 
 	TArray<FWorldPartitionReference> LoadedSubobjects;
 
-	TSet<FDirtyActor> DirtyActors;
+	TMap<FWorldPartitionReference, TObjectPtr<AActor>> DirtyActors;
 
 	TSet<FString> GeneratedStreamingPackageNames;
 
@@ -533,7 +497,8 @@ private:
 #if WITH_EDITOR
 	ENGINE_API void HashActorDesc(FWorldPartitionActorDesc* ActorDesc);
 	ENGINE_API void UnhashActorDesc(FWorldPartitionActorDesc* ActorDesc);
-	void OnContentBundleRemovedContent(const FContentBundleEditor* ContentBundle);
+	ENGINE_API void HashActorDescContainer(UActorDescContainer* ActorDescContainer);
+	ENGINE_API void UnhashActorDescContainer(UActorDescContainer* ActorDescContainer);
 
 public:
 	// Editor loader adapters management
