Don't recreate scene proxy for empty ISMs

diff --git a/Engine/Source/Runtime/Engine/Private/HierarchicalInstancedStaticMesh.cpp b/Engine/Source/Runtime/Engine/Private/HierarchicalInstancedStaticMesh.cpp
--- a/Engine/Source/Runtime/Engine/Private/HierarchicalInstancedStaticMesh.cpp
+++ b/Engine/Source/Runtime/Engine/Private/HierarchicalInstancedStaticMesh.cpp
@@ -3316,21 +3316,23 @@ FPrimitiveSceneProxy* UHierarchicalInstancedStaticMeshComponent::CreateSceneProx
 		!GetStaticMesh()->IsCompiling() &&
 		GetStaticMesh()->HasValidRenderData(false);
 
-	if (bMeshIsValid)
+	if (!bMeshIsValid)
 	{
-		check(InstancingRandomSeed != 0);
-
-		// If instance data was modified, update GPU copy.
-		// If InstanceBuffer was initialized with RequireCPUAccess (always true in editor).
-		if (InstanceUpdateCmdBuffer.NumInlineCommands() > 0 && PerInstanceRenderData->InstanceBuffer.RequireCPUAccess)
-		{
-			PerInstanceRenderData->UpdateFromCommandBuffer(InstanceUpdateCmdBuffer);
-		}
+		return nullptr;
+	}
+		
+	check(InstancingRandomSeed != 0);
 
-		ProxySize = PerInstanceRenderData->ResourceSize;
-		INC_DWORD_STAT_BY(STAT_FoliageInstanceBuffers, ProxySize);
+	// If instance data was modified, update GPU copy.
+	// If InstanceBuffer was initialized with RequireCPUAccess (always true in editor).
+	if (InstanceUpdateCmdBuffer.NumInlineCommands() > 0 && PerInstanceRenderData->InstanceBuffer.RequireCPUAccess)
+	{
+		PerInstanceRenderData->UpdateFromCommandBuffer(InstanceUpdateCmdBuffer);
 	}
 
+	ProxySize = PerInstanceRenderData->ResourceSize;
+	INC_DWORD_STAT_BY(STAT_FoliageInstanceBuffers, ProxySize);
+
 	// NOTE: Purposefully skipping UInstancedStaticMeshComponent implementation
 	return UStaticMeshComponent::CreateSceneProxy();
 }
diff --git a/Engine/Source/Runtime/Engine/Private/InstancedStaticMesh.cpp b/Engine/Source/Runtime/Engine/Private/InstancedStaticMesh.cpp
--- a/Engine/Source/Runtime/Engine/Private/InstancedStaticMesh.cpp
+++ b/Engine/Source/Runtime/Engine/Private/InstancedStaticMesh.cpp
@@ -2760,19 +2760,21 @@ FPrimitiveSceneProxy* UInstancedStaticMeshComponent::CreateSceneProxy()
 		GetStaticMesh()->IsCompiling() == false &&
 		GetStaticMesh()->HasValidRenderData();
 
-	if (bMeshIsValid)
+	if (!bMeshIsValid)
 	{
-		check(InstancingRandomSeed != 0);
-		
-		// if instance data was modified, update GPU copy
-		// generally happens only in editor 
-		if (InstanceUpdateCmdBuffer.NumTotalCommands() != 0)
-		{
-			FlushInstanceUpdateCommands(true);
-		}
+		return nullptr;
+	}
+
+	check(InstancingRandomSeed != 0);
 		
-		ProxySize = PerInstanceRenderData->ResourceSize;
+	// if instance data was modified, update GPU copy
+	// generally happens only in editor 
+	if (InstanceUpdateCmdBuffer.NumTotalCommands() != 0)
+	{
+		FlushInstanceUpdateCommands(true);
 	}
+		
+	ProxySize = PerInstanceRenderData->ResourceSize;
 
 	return Super::CreateSceneProxy();
 }
