cluster unions should be unbreakable on the client to help prevent server-client desync

diff --git a/Engine/Source/Runtime/Engine/Private/PhysicsEngine/ClusterUnionComponent.cpp b/Engine/Source/Runtime/Engine/Private/PhysicsEngine/ClusterUnionComponent.cpp
--- a/Engine/Source/Runtime/Engine/Private/PhysicsEngine/ClusterUnionComponent.cpp
+++ b/Engine/Source/Runtime/Engine/Private/PhysicsEngine/ClusterUnionComponent.cpp
@@ -354,14 +354,19 @@ void UClusterUnionComponent::OnCreatePhysicsState()
 
 	FChaosUserData::Set<UPrimitiveComponent>(&PhysicsUserData, this);
 
+	const bool bHasAuthority = GetOwner()->HasAuthority();
+
 	Chaos::FClusterUnionInitData InitData;
 	InitData.UserData = static_cast<void*>(&PhysicsUserData);
 	InitData.ActorId = GetOwner()->GetUniqueID();
 	InitData.ComponentId = GetUniqueID();
-	InitData.bNeedsClusterXRInitialization = GetOwner()->HasAuthority();
+	InitData.bNeedsClusterXRInitialization = bHasAuthority;
+
+	// Client needs to be set to unbreakable so the server is authoritative.
+	InitData.bUnbreakable = !bHasAuthority;
 
 	// Only need to check connectivity on the server and have the client rely on replication to get the memo on when to release from cluster union.
-	InitData.bCheckConnectivity = GetOwner()->HasAuthority();
+	InitData.bCheckConnectivity = bHasAuthority;
 
 	bHasReceivedTransform = false;
 	PhysicsProxy = new Chaos::FClusterUnionPhysicsProxy{ this, Parameters, InitData };
diff --git a/Engine/Source/Runtime/Experimental/Chaos/Private/PhysicsProxy/ClusterUnionPhysicsProxy.cpp b/Engine/Source/Runtime/Experimental/Chaos/Private/PhysicsProxy/ClusterUnionPhysicsProxy.cpp
--- a/Engine/Source/Runtime/Experimental/Chaos/Private/PhysicsProxy/ClusterUnionPhysicsProxy.cpp
+++ b/Engine/Source/Runtime/Experimental/Chaos/Private/PhysicsProxy/ClusterUnionPhysicsProxy.cpp
@@ -122,6 +122,7 @@ namespace Chaos
 			Particle_Internal = ClusterUnion->InternalCluster;
 			Particle_Internal->SetPhysicsProxy(this);
 			Particle_Internal->GTGeometryParticle() = Particle_External.Get();
+			Particle_Internal->SetUnbreakable(InitData.bUnbreakable);
 
 #if !(UE_BUILD_SHIPPING || UE_BUILD_TEST)
 			Particle_Internal->SetDebugName(MakeShared<FString, ESPMode::ThreadSafe>(FString::Printf(TEXT("%s"), *GetOwner()->GetName())));
diff --git a/Engine/Source/Runtime/Experimental/Chaos/Public/PhysicsProxy/ClusterUnionPhysicsProxy.h b/Engine/Source/Runtime/Experimental/Chaos/Public/PhysicsProxy/ClusterUnionPhysicsProxy.h
--- a/Engine/Source/Runtime/Experimental/Chaos/Public/PhysicsProxy/ClusterUnionPhysicsProxy.h
+++ b/Engine/Source/Runtime/Experimental/Chaos/Public/PhysicsProxy/ClusterUnionPhysicsProxy.h
@@ -32,6 +32,7 @@ namespace Chaos
 		uint32 ComponentId = INDEX_NONE;
 		bool bNeedsClusterXRInitialization = true;
 		bool bCheckConnectivity = true;
+		bool bUnbreakable = false;
 	};
 
 	/**
