fix issue where replicated cluster union particles would not accurately account for particles that were removed from the cluster union

diff --git a/Engine/Source/Runtime/Engine/Classes/PhysicsEngine/ClusterUnionComponent.h b/Engine/Source/Runtime/Engine/Classes/PhysicsEngine/ClusterUnionComponent.h
--- a/Engine/Source/Runtime/Engine/Classes/PhysicsEngine/ClusterUnionComponent.h
+++ b/Engine/Source/Runtime/Engine/Classes/PhysicsEngine/ClusterUnionComponent.h
@@ -132,6 +132,9 @@ public:
 	UFUNCTION(BlueprintCallable, Category = "Cluster Union")
 	ENGINE_API void RemoveComponentFromCluster(UPrimitiveComponent* InComponent);
 
+	UFUNCTION(BlueprintCallable, Category = "Cluster Union")
+	ENGINE_API void RemoveComponentBonesFromCluster(UPrimitiveComponent* InComponent, const TArray<int32>& BoneIds);
+
 	UFUNCTION(BlueprintCallable, Category = "Cluster Union")
 	ENGINE_API TArray<UPrimitiveComponent*> GetPrimitiveComponents();
 
diff --git a/Engine/Source/Runtime/Engine/Classes/PhysicsEngine/ClusterUnionReplicatedProxyComponent.h b/Engine/Source/Runtime/Engine/Classes/PhysicsEngine/ClusterUnionReplicatedProxyComponent.h
--- a/Engine/Source/Runtime/Engine/Classes/PhysicsEngine/ClusterUnionReplicatedProxyComponent.h
+++ b/Engine/Source/Runtime/Engine/Classes/PhysicsEngine/ClusterUnionReplicatedProxyComponent.h
@@ -78,6 +78,7 @@ private:
 
 	UPROPERTY(ReplicatedUsing=OnRep_ParticleBoneIds)
 	TArray<int32> ParticleBoneIds;
+	TSet<int32> LastSyncedBoneIds;
 
 	UPROPERTY()
 	bool bNetUpdateParticleBoneIds;
diff --git a/Engine/Source/Runtime/Engine/Private/PhysicsEngine/ClusterUnionComponent.cpp b/Engine/Source/Runtime/Engine/Private/PhysicsEngine/ClusterUnionComponent.cpp
--- a/Engine/Source/Runtime/Engine/Private/PhysicsEngine/ClusterUnionComponent.cpp
+++ b/Engine/Source/Runtime/Engine/Private/PhysicsEngine/ClusterUnionComponent.cpp
@@ -229,6 +229,107 @@ void UClusterUnionComponent::RemoveComponentFromCluster(UPrimitiveComponent* InC
 	}
 }
 
+// TODO: Can this merge with RemoveComponentFromCluster?
+void UClusterUnionComponent::RemoveComponentBonesFromCluster(UPrimitiveComponent* InComponent, const TArray<int32>& BoneIds)
+{
+	if (!InComponent || !PhysicsProxy)
+	{
+		return;
+	}
+	TSet<int32> RemoveBoneIdsSet{ BoneIds };
+	// If we're removing bones from a component and it hasn't been added into the component yet, then we want to make sure 
+	// the set we actually want to add doesn't have the remove bones in it.
+	if (FClusterUnionPendingAddData* PendingAdd = PendingComponentsToAdd.Find(InComponent))
+	{
+		for (int32 BoneId : RemoveBoneIdsSet)
+		{
+			PendingAdd->BoneIds.Remove(BoneId);
+		}
+
+		if (PendingAdd->BoneIds.IsEmpty())
+		{
+			PendingComponentsToAdd.Remove(InComponent);
+		}
+
+		// Component is pending add so we can be confident we don't need to do any more work.
+		return;
+	}
+
+	// If we're still waiting on the physics sync for the added component, we still need to
+	// remove its acceleration handles from the acceleration structure since the payloads have
+	// already been added and will not yet be stored in PerComponentData.
+	if (FClusterUnionPendingAddData* PendingData = PendingComponentSync.Find(InComponent))
+	{
+		if (AccelerationStructure)
+		{
+			for (int32 BoneId : RemoveBoneIdsSet)
+			{
+				FExternalSpatialAccelerationPayload Payload;
+				Payload.Initialize(InComponent, BoneId);
+				PendingData->AccelerationPayloads.Remove(Payload);
+				AccelerationStructure->RemoveElement(Payload);
+			}
+		}
+
+		for (int32 BoneId : RemoveBoneIdsSet)
+		{
+			PendingData->BoneIds.Remove(BoneId);
+		}
+
+		if (PendingData->BoneIds.IsEmpty())
+		{
+			PendingComponentSync.Remove(InComponent);
+		}
+	}
+
+	TSet<Chaos::FPhysicsObjectHandle> PhysicsObjectsToRemove;
+	if (FClusteredComponentData* ComponentData = PerComponentData.Find(InComponent))
+	{
+		for (int32 BoneId : RemoveBoneIdsSet)
+		{
+			if (AccelerationStructure)
+			{
+				FExternalSpatialAccelerationPayload Payload;
+				Payload.Initialize(InComponent, BoneId);
+				ComponentData->CachedAccelerationPayloads.Remove(Payload);
+				AccelerationStructure->RemoveElement(Payload);
+			}
+			ComponentData->BoneIds.Remove(BoneId);
+		}
+
+		if (InComponent->HasValidPhysicsState())
+		{
+			PhysicsObjectsToRemove = TSet<Chaos::FPhysicsObjectHandle>{ GetAllPhysicsObjectsById(InComponent, RemoveBoneIdsSet.Array()) };
+		}
+
+		if (ComponentData->BoneIds.IsEmpty())
+		{
+			// We need to mark the replicated proxy as pending deletion.
+			// This way anyone who tries to use the replicated proxy component knows that it
+			// doesn't actually denote a meaningful cluster union relationship.
+			ComponentData->bPendingDeletion = true;
+
+			if (IsAuthority())
+			{
+				if (UClusterUnionReplicatedProxyComponent* Component = ComponentData->ReplicatedProxyComponent.Get())
+				{
+					Component->MarkPendingDeletion();
+				}
+			}
+
+			PerComponentData.Remove(InComponent);
+		}
+	}
+
+	// If PhysicsObjectsToRemove is empty, it either means that BoneIds is empty OR the component's physics state is already destroyed.
+	// In the case of the latter, we rely on the physics thread to cleanup the cluster union manager properly and to sync back.
+	if (!PhysicsObjectsToRemove.IsEmpty())
+	{
+		PhysicsProxy->RemovePhysicsObjects_External(PhysicsObjectsToRemove);
+		ForceRebuildGTParticleGeometry();
+	}
+}
+
 void UClusterUnionComponent::ForceRebuildGTParticleGeometry()
 {
 	if (!PhysicsProxy)
diff --git a/Engine/Source/Runtime/Engine/Private/PhysicsEngine/ClusterUnionReplicatedProxyComponent.cpp b/Engine/Source/Runtime/Engine/Private/PhysicsEngine/ClusterUnionReplicatedProxyComponent.cpp
--- a/Engine/Source/Runtime/Engine/Private/PhysicsEngine/ClusterUnionReplicatedProxyComponent.cpp
+++ b/Engine/Source/Runtime/Engine/Private/PhysicsEngine/ClusterUnionReplicatedProxyComponent.cpp
@@ -155,7 +155,42 @@ void UClusterUnionReplicatedProxyComponent::DeferAddComponentToClusterHandleUnti
 
 	if (ParentClusterUnion->HasReceivedTransform())
 	{
-		ParentClusterUnion->AddComponentToCluster(ChildClusteredComponent.Get(), ParticleBoneIds);
+		// Need to check if we're *losing* bones instead and handle that situation as well as adding new bones into the cluster union.
+		// This extra check also does a bit of prevention on the client side from adding duplicate particles into the cluster union.
+		TSet<int32> NewBoneIdSet{ ParticleBoneIds };
+
+		TArray<int32> ToAdd;
+		ToAdd.Reserve(NewBoneIdSet.Num());
+
+		for (int32 BoneId : NewBoneIdSet)
+		{
+			if (!LastSyncedBoneIds.Contains(BoneId))
+			{
+				ToAdd.Add(BoneId);
+			}
+		}
+
+		if (!ToAdd.IsEmpty())
+		{
+			ParentClusterUnion->AddComponentToCluster(ChildClusteredComponent.Get(), ToAdd);
+		}
+
+		TArray<int32> ToRemove;
+		ToRemove.Reserve(LastSyncedBoneIds.Num());
+		for (int32 BoneId : LastSyncedBoneIds)
+		{
+			if (!NewBoneIdSet.Contains(BoneId))
+			{
+				ToRemove.Add(BoneId);
+			}
+		}
+
+		if (!ToRemove.IsEmpty())
+		{
+			ParentClusterUnion->RemoveComponentBonesFromCluster(ChildClusteredComponent.Get(), ToRemove);
+		}
+
+		LastSyncedBoneIds = NewBoneIdSet;
 	}
 	else if (AActor* Owner = GetOwner())
 	{
