Substrate - fix for simplfication being request for material that hsould not.

diff --git a/Engine/Source/Runtime/Engine/Private/Materials/HLSLMaterialTranslator.cpp b/Engine/Source/Runtime/Engine/Private/Materials/HLSLMaterialTranslator.cpp
--- a/Engine/Source/Runtime/Engine/Private/Materials/HLSLMaterialTranslator.cpp
+++ b/Engine/Source/Runtime/Engine/Private/Materials/HLSLMaterialTranslator.cpp
@@ -11711,8 +11711,11 @@ bool FHLSLMaterialTranslator::FStrataCompilationContext::StrataGenerateDerivedMa
 				const bool bRootOfParameterBlendingSubTree		= bCurrentOpRequestParameterBlending && !bInsideParameterBlendingSubTree;
 				const bool bUseParameterBlending				= bCurrentOpRequestParameterBlending || bInsideParameterBlendingSubTree;
 
-				CurrentOperator.bUseParameterBlending			= bUseParameterBlending;
-				CurrentOperator.bRootOfParameterBlendingSubTree	= bRootOfParameterBlendingSubTree;
+				if (CurrentOperator.BSDFType == STRATA_BSDF_TYPE_SLAB)
+				{
+					CurrentOperator.bUseParameterBlending = bUseParameterBlending;
+					CurrentOperator.bRootOfParameterBlendingSubTree = bRootOfParameterBlendingSubTree;
+				}
 
 				switch (CurrentOperator.OperatorType)
 				{
