World Partition HashSet:
- Runtime Partition will now set default values when created in the property grid.
- Converted HashSet to use streaming datas instead of external streaming objects for internal data.

diff --git a/Engine/Source/Runtime/Engine/Private/WorldPartition/RuntimeHashSet/RuntimePartition.cpp b/Engine/Source/Runtime/Engine/Private/WorldPartition/RuntimeHashSet/RuntimePartition.cpp
--- a/Engine/Source/Runtime/Engine/Private/WorldPartition/RuntimeHashSet/RuntimePartition.cpp
+++ b/Engine/Source/Runtime/Engine/Private/WorldPartition/RuntimeHashSet/RuntimePartition.cpp
@@ -3,6 +3,16 @@
 #include "WorldPartition/RuntimeHashSet/RuntimePartition.h"
 
 #if WITH_EDITOR
+void URuntimePartition::SetDefaultValues()
+{
+	Name = GetClass()->GetFName();
+	bBlockOnSlowStreaming = false;
+	bClientOnlyVisible = false;
+	Priority = 0;
+	LoadingRange = 25600;
+	bIsHLODSetup = false;
+}
+
 void URuntimePartition::PostEditChangeProperty(FPropertyChangedEvent& InPropertyChangedEvent)
 {
 	const FName PropertyName = InPropertyChangedEvent.GetPropertyName();
diff --git a/Engine/Source/Runtime/Engine/Private/WorldPartition/RuntimeHashSet/RuntimePartitionLHGrid.cpp b/Engine/Source/Runtime/Engine/Private/WorldPartition/RuntimeHashSet/RuntimePartitionLHGrid.cpp
--- a/Engine/Source/Runtime/Engine/Private/WorldPartition/RuntimeHashSet/RuntimePartitionLHGrid.cpp
+++ b/Engine/Source/Runtime/Engine/Private/WorldPartition/RuntimeHashSet/RuntimePartitionLHGrid.cpp
@@ -67,6 +67,12 @@ void URuntimePartitionLHGrid::PostEditChangeProperty(FPropertyChangedEvent& Prop
 	Super::PostEditChangeProperty(PropertyChangedEvent);
 }
 
+void URuntimePartitionLHGrid::SetDefaultValues()
+{
+	Super::SetDefaultValues();
+	CellSize = LoadingRange / 2;
+}
+
 bool URuntimePartitionLHGrid::SupportsHLODs() const
 {
 	return true;
diff --git a/Engine/Source/Runtime/Engine/Private/WorldPartition/RuntimeHashSet/WorldPartitionRuntimeHashSet.cpp b/Engine/Source/Runtime/Engine/Private/WorldPartition/RuntimeHashSet/WorldPartitionRuntimeHashSet.cpp
--- a/Engine/Source/Runtime/Engine/Private/WorldPartition/RuntimeHashSet/WorldPartitionRuntimeHashSet.cpp
+++ b/Engine/Source/Runtime/Engine/Private/WorldPartition/RuntimeHashSet/WorldPartitionRuntimeHashSet.cpp
@@ -91,21 +91,31 @@ void FRuntimePartitionDesc::UpdateHLODPartitionLayers()
 }
 #endif
 
-void URuntimeHashSetExternalStreamingObject::CreatePartitionsSpatialIndex() const
+void FRuntimePartitionStreamingData::CreatePartitionsSpatialIndex() const
 {
-	for (const FRuntimePartitionStreamingData& StreamingData : RuntimeStreamingData)
+	if (!SpatialIndex)
 	{
-		if (!StreamingData.SpatialIndex)
+		SpatialIndex = MakeUnique<FStaticSpatialIndexType>();
+
+		TArray<TPair<FBox, UWorldPartitionRuntimeCell*>> PartitionsElements;
+		Algo::Transform(StreamingCells, PartitionsElements, [](UWorldPartitionRuntimeCell* Cell)
 		{
-			StreamingData.SpatialIndex = MakeUnique<FStaticSpatialIndexType>();
+			return TPair<FBox, UWorldPartitionRuntimeCell*>(Cell->GetContentBounds(), Cell);
+		});
+		SpatialIndex->Init(PartitionsElements);
+	}
+}
 
-			TArray<TPair<FBox, UWorldPartitionRuntimeCell*>> PartitionsElements;
-			Algo::Transform(StreamingData.RuntimeCells, PartitionsElements, [](UWorldPartitionRuntimeCell* RuntimeCell)
-			{
-				return TPair<FBox, UWorldPartitionRuntimeCell*>(RuntimeCell->GetContentBounds(), RuntimeCell);
-			});
-			StreamingData.SpatialIndex->Init(PartitionsElements);
-		}
+void FRuntimePartitionStreamingData::DestroyPartitionsSpatialIndex() const
+{
+	SpatialIndex.Reset();
+}
+
+void URuntimeHashSetExternalStreamingObject::CreatePartitionsSpatialIndex() const
+{
+	for (const FRuntimePartitionStreamingData& StreamingData : RuntimeStreamingData)
+	{
+		StreamingData.CreatePartitionsSpatialIndex();
 	}
 }
 
@@ -113,7 +123,7 @@ void URuntimeHashSetExternalStreamingObject::DestroyPartitionsSpatialIndex() con
 {
 	for (const FRuntimePartitionStreamingData& StreamingData : RuntimeStreamingData)
 	{
-		StreamingData.SpatialIndex.Reset();
+		StreamingData.DestroyPartitionsSpatialIndex();
 	}
 }
 
@@ -149,9 +159,9 @@ void UWorldPartitionRuntimeHashSet::PostLoad()
 	else
 #endif
 	{
-		ForEachStreamingObject([](const URuntimeHashSetExternalStreamingObject* CurStreamingObject)
+		ForEachStreamingData([](const FRuntimePartitionStreamingData& StreamingData)
 		{
-			CurStreamingObject->CreatePartitionsSpatialIndex();
+			StreamingData.CreatePartitionsSpatialIndex();
 		});
 	}
 }
@@ -190,13 +200,13 @@ bool UWorldPartitionRuntimeHashSet::Draw2D(FWorldPartitionDraw2DContext& DrawCon
 	}
 
 	TMap<FName, TArray<const FRuntimePartitionStreamingData*>> FilteredStreamingObjects;
-	ForEachStreamingObject([&FilteredStreamingObjects](const URuntimeHashSetExternalStreamingObject* CurStreamingObject)
+	ForEachStreamingData([&FilteredStreamingObjects](const FRuntimePartitionStreamingData& StreamingData)
 	{
-		for (const FRuntimePartitionStreamingData& RuntimeStreamingData : CurStreamingObject->RuntimeStreamingData)
+		if (StreamingData.StreamingCells.Num())
 		{
-			if (FWorldPartitionDebugHelper::IsDebugRuntimeHashGridShown(RuntimeStreamingData.Name))
+			if (FWorldPartitionDebugHelper::IsDebugRuntimeHashGridShown(StreamingData.Name))
 			{
-				FilteredStreamingObjects.FindOrAdd(RuntimeStreamingData.Name).Add(&RuntimeStreamingData);
+				FilteredStreamingObjects.FindOrAdd(StreamingData.Name).Add(&StreamingData);
 			}
 		}
 	});
@@ -263,7 +273,7 @@ bool UWorldPartitionRuntimeHashSet::Draw2D(FWorldPartitionDraw2DContext& DrawCon
 			{
 				const FVector2D CellBoundsSize = FVector2D(Cell->GetCellBounds().GetSize());
 				const FVector2D CellBoundsMin = FVector2D(Cell->GetCellBounds().Min);
-				DrawContext.LocalDrawTile(GridScreenBounds, CellBoundsMin, CellBoundsSize, Cell->GetDebugColor(VisualizeMode).CopyWithNewOpacity(0.25f), WorldToScreen);
+				DrawContext.LocalDrawTile(GridScreenBounds, CellBoundsMin, CellBoundsSize, Cell->GetDebugColor(VisualizeMode).CopyWithNewOpacity(0.25f / FMath::Max<float>(GShowRuntimeHashSetDebugDisplayLevelCount, 1)), WorldToScreen);
 				DrawContext.LocalDrawBox(GridScreenBounds, CellBoundsMin, CellBoundsSize, FLinearColor::Black, 1, WorldToScreen);
 			}
 		}
@@ -503,8 +513,7 @@ bool UWorldPartitionRuntimeHashSet::GenerateStreaming(UWorldPartitionStreamingPo
 	//
 	// Create and populate streaming object
 	//
-	check(!StreamingObject);
-	StreamingObject = CreateExternalStreamingObject<URuntimeHashSetExternalStreamingObject>(this, MakeUniqueObjectName(this, URuntimeHashExternalStreamingObjectBase::StaticClass()));
+	check(RuntimeStreamingData.IsEmpty());
 
 	// Generate runtime cells
 	auto CreateRuntimeCellFromCellDesc = [this](const URuntimePartition::FCellDesc& CellDesc, const TArray<const UDataLayerInstance*>& DataLayers, TSubclassOf<UWorldPartitionRuntimeCell> CellClass, TSubclassOf<UWorldPartitionRuntimeCellData> CellDataClass)
@@ -558,7 +567,7 @@ bool UWorldPartitionRuntimeHashSet::GenerateStreaming(UWorldPartitionStreamingPo
 			check(CellGuid.IsValid());
 		}
 
-		UWorldPartitionRuntimeCell* RuntimeCell = Super::CreateRuntimeCell(CellClass, CellDataClass, CellObjectName, TEXT(""), StreamingObject);
+		UWorldPartitionRuntimeCell* RuntimeCell = Super::CreateRuntimeCell(CellClass, CellDataClass, CellObjectName, TEXT(""));
 
 		RuntimeCell->SetIsAlwaysLoaded(!CellDesc.bIsSpatiallyLoaded);
 		RuntimeCell->SetDataLayers(DataLayers);
@@ -589,17 +598,19 @@ bool UWorldPartitionRuntimeHashSet::GenerateStreaming(UWorldPartitionStreamingPo
 			RuntimeCell->RuntimeCellData->ContentBounds = CellDescInstance.Bounds;
 		}
 
+		// Create partition streaming data
+		FRuntimePartitionStreamingData& StreamingData = RuntimePartitionsStreamingData.FindOrAdd(CellDescInstance.SourcePartition);
+
+		StreamingData.Name = CellDescInstance.SourcePartition->Name;
+		StreamingData.LoadingRange = CellDescInstance.SourcePartition->LoadingRange;
+
 		if (RuntimeCell->IsAlwaysLoaded())
 		{
-			StreamingObject->NonSpatiallyLoadedRuntimeCells.Add(RuntimeCell);
+			StreamingData.NonStreamingCells.Add(RuntimeCell);
 		}
 		else
 		{
-			FRuntimePartitionStreamingData& StreamingData = RuntimePartitionsStreamingData.FindOrAdd(CellDescInstance.SourcePartition);
-
-			StreamingData.Name = CellDescInstance.SourcePartition->Name;
-			StreamingData.LoadingRange = CellDescInstance.SourcePartition->LoadingRange;
-			StreamingData.RuntimeCells.Add(RuntimeCell);
+			StreamingData.StreamingCells.Add(RuntimeCell);
 		}
 	}
 
@@ -608,7 +619,7 @@ bool UWorldPartitionRuntimeHashSet::GenerateStreaming(UWorldPartitionStreamingPo
 	//
 	for (auto& [Partition, StreamingData] : RuntimePartitionsStreamingData)
 	{
-		StreamingObject->RuntimeStreamingData.Emplace(MoveTemp(StreamingData));
+		RuntimeStreamingData.Emplace(MoveTemp(StreamingData));
 	}
 
 	//
@@ -644,7 +655,7 @@ void UWorldPartitionRuntimeHashSet::FlushStreaming()
 	PersistentPartitionDesc.Name = NAME_None;
 	PersistentPartitionDesc.MainLayer = nullptr;
 
-	StreamingObject = nullptr;
+	RuntimeStreamingData.Empty();
 }
 
 bool UWorldPartitionRuntimeHashSet::IsValidGrid(FName GridName) const
@@ -673,12 +684,12 @@ bool UWorldPartitionRuntimeHashSet::IsValidGrid(FName GridName) const
 
 TArray<UWorldPartitionRuntimeCell*> UWorldPartitionRuntimeHashSet::GetAlwaysLoadedCells() const
 {
-	TArray<UWorldPartitionRuntimeCell*> NonSpatiallyLoadedRuntimeCells;
-	ForEachStreamingObject([&NonSpatiallyLoadedRuntimeCells](const URuntimeHashSetExternalStreamingObject* CurStreamingObject)
+	TArray<UWorldPartitionRuntimeCell*> AlwaysLoadedCells;
+	ForEachStreamingData([&AlwaysLoadedCells](const FRuntimePartitionStreamingData& StreamingData)
 	{
-		NonSpatiallyLoadedRuntimeCells.Append(CurStreamingObject->NonSpatiallyLoadedRuntimeCells);
+		AlwaysLoadedCells.Append(StreamingData.NonStreamingCells);
 	});
-	return NonSpatiallyLoadedRuntimeCells;
+	return AlwaysLoadedCells;
 }
 
 void UWorldPartitionRuntimeHashSet::DumpStateLog(FHierarchicalLogArchive& Ar) const
@@ -718,11 +729,25 @@ TArray<FName> UWorldPartitionRuntimeHashSet::ParseGridName(FName GridName)
 
 URuntimeHashExternalStreamingObjectBase* UWorldPartitionRuntimeHashSet::StoreToExternalStreamingObject(UObject* StreamingObjectOuter, FName StreamingObjectName)
 {
-	check(StreamingObject);
-	StreamingObject->Rename(*StreamingObjectName.ToString(), StreamingObjectOuter, REN_DoNotDirty | REN_ForceNoResetLoaders);
-	URuntimeHashExternalStreamingObjectBase* Result = StreamingObject;
-	StreamingObject = nullptr;
-	return Result;
+	check(!RuntimeStreamingData.IsEmpty());
+
+	URuntimeHashSetExternalStreamingObject* NewStreamingObject = CreateExternalStreamingObject<URuntimeHashSetExternalStreamingObject>(this, MakeUniqueObjectName(this, URuntimeHashExternalStreamingObjectBase::StaticClass()));
+	NewStreamingObject->RuntimeStreamingData = MoveTemp(RuntimeStreamingData);
+
+	for (FRuntimePartitionStreamingData& StreamingData : NewStreamingObject->RuntimeStreamingData)
+	{
+		for (UWorldPartitionRuntimeCell* Cell : StreamingData.StreamingCells)
+		{
+			Cell->Rename(nullptr, NewStreamingObject,  REN_DoNotDirty | REN_ForceNoResetLoaders);
+		}
+
+		for (UWorldPartitionRuntimeCell* Cell : StreamingData.NonStreamingCells)
+		{
+			Cell->Rename(nullptr, NewStreamingObject,  REN_DoNotDirty | REN_ForceNoResetLoaders);
+		}
+	}
+
+	return NewStreamingObject;
 }
 #endif
 
@@ -753,23 +778,9 @@ bool UWorldPartitionRuntimeHashSet::RemoveExternalStreamingObject(URuntimeHashEx
 // Streaming interface
 void UWorldPartitionRuntimeHashSet::ForEachStreamingCells(TFunctionRef<bool(const UWorldPartitionRuntimeCell*)> Func) const
 {
-	auto ForEachStreamingCells = [this, &Func](FStaticSpatialIndexType* InSpatialIndex)
-	{
-		if (InSpatialIndex)
-		{
-			InSpatialIndex->ForEachElement([this, &Func](const UWorldPartitionRuntimeCell* RuntimeCell)
-			{
-				if (IsCellRelevantFor(RuntimeCell->GetClientOnlyVisible()))
-				{
-					Func(RuntimeCell);
-				}
-			});
-		}
-	};
-
-	auto ForEachNonStreamingCells = [this, &Func](TArray<TObjectPtr<UWorldPartitionRuntimeCell>> InNonSpatiallyLoadedRuntimeCells)
+	auto ForEachCells = [this, &Func](const TArray<TObjectPtr<UWorldPartitionRuntimeCell>>& InCells)
 	{
-		for (UWorldPartitionRuntimeCell* Cell : InNonSpatiallyLoadedRuntimeCells)
+		for (UWorldPartitionRuntimeCell* Cell : InCells)
 		{
 			if (IsCellRelevantFor(Cell->GetClientOnlyVisible()))
 			{
@@ -778,14 +789,10 @@ void UWorldPartitionRuntimeHashSet::ForEachStreamingCells(TFunctionRef<bool(cons
 		}
 	};
 
-	ForEachStreamingObject([&ForEachStreamingCells, &ForEachNonStreamingCells](const URuntimeHashSetExternalStreamingObject* CurStreamingObject)
+	ForEachStreamingData([&ForEachCells](const FRuntimePartitionStreamingData& StreamingData)
 	{
-		for (const FRuntimePartitionStreamingData& StreamingData : CurStreamingObject->RuntimeStreamingData)
-		{
-			ForEachStreamingCells(StreamingData.SpatialIndex.Get());
-		}
-
-		ForEachNonStreamingCells(CurStreamingObject->NonSpatiallyLoadedRuntimeCells);
+		ForEachCells(StreamingData.StreamingCells);
+		ForEachCells(StreamingData.NonStreamingCells);
 	});
 }
 
@@ -825,9 +832,9 @@ void UWorldPartitionRuntimeHashSet::ForEachStreamingCellsQuery(const FWorldParti
 		}
 	};
 
-	auto ForEachNonStreamingCells = [&ShouldAddCell, &QuerySource, &Func](TArray<TObjectPtr<UWorldPartitionRuntimeCell>> InNonSpatiallyLoadedRuntimeCells)
+	auto ForEachNonStreamingCells = [&ShouldAddCell, &QuerySource, &Func](TArray<TObjectPtr<UWorldPartitionRuntimeCell>> InNonStreamingCells)
 	{
-		for (UWorldPartitionRuntimeCell* Cell : InNonSpatiallyLoadedRuntimeCells)
+		for (UWorldPartitionRuntimeCell* Cell : InNonStreamingCells)
 		{
 			if (ShouldAddCell(Cell, QuerySource))
 			{
@@ -836,14 +843,10 @@ void UWorldPartitionRuntimeHashSet::ForEachStreamingCellsQuery(const FWorldParti
 		}
 	};
 
-	ForEachStreamingObject([&ForEachStreamingCells, &ForEachNonStreamingCells](const URuntimeHashSetExternalStreamingObject* CurStreamingObject)
+	ForEachStreamingData([&ForEachStreamingCells, &ForEachNonStreamingCells](const FRuntimePartitionStreamingData& StreamingData)
 	{
-		for (const FRuntimePartitionStreamingData& StreamingData : CurStreamingObject->RuntimeStreamingData)
-		{
-			ForEachStreamingCells(StreamingData.SpatialIndex.Get());
-		}
-
-		ForEachNonStreamingCells(CurStreamingObject->NonSpatiallyLoadedRuntimeCells);
+		ForEachStreamingCells(StreamingData.SpatialIndex.Get());
+		ForEachNonStreamingCells(StreamingData.NonStreamingCells);
 	});
 }
 
@@ -890,9 +893,9 @@ void UWorldPartitionRuntimeHashSet::ForEachStreamingCellsSources(const TArray<FW
 		}
 	};
 
-	auto ForEachNonStreamingCells = [this, &ActivateStreamingSourceCells, &LoadStreamingSourceCells](TArray<TObjectPtr<UWorldPartitionRuntimeCell>> InNonSpatiallyLoadedRuntimeCells)
+	auto ForEachNonStreamingCells = [this, &ActivateStreamingSourceCells, &LoadStreamingSourceCells](TArray<TObjectPtr<UWorldPartitionRuntimeCell>> InNonStreamingCells)
 	{
-		for (UWorldPartitionRuntimeCell* Cell : InNonSpatiallyLoadedRuntimeCells)
+		for (UWorldPartitionRuntimeCell* Cell : InNonStreamingCells)
 		{
 			if (IsCellRelevantFor(Cell->GetClientOnlyVisible()))
 			{
@@ -908,14 +911,10 @@ void UWorldPartitionRuntimeHashSet::ForEachStreamingCellsSources(const TArray<FW
 		}
 	};
 
-	ForEachStreamingObject([&ForEachStreamingCells, &ForEachNonStreamingCells](const URuntimeHashSetExternalStreamingObject* CurStreamingObject)
+	ForEachStreamingData([&ForEachStreamingCells, &ForEachNonStreamingCells](const FRuntimePartitionStreamingData& StreamingData)
 	{
-		for (const FRuntimePartitionStreamingData& StreamingData : CurStreamingObject->RuntimeStreamingData)
-		{
-			ForEachStreamingCells(StreamingData.SpatialIndex.Get(), StreamingData.LoadingRange, StreamingData.Name);
-		}
-
-		ForEachNonStreamingCells(CurStreamingObject->NonSpatiallyLoadedRuntimeCells);
+		ForEachStreamingCells(StreamingData.SpatialIndex.Get(), StreamingData.LoadingRange, StreamingData.Name);
+		ForEachNonStreamingCells(StreamingData.NonStreamingCells);
 	});
 
 	auto ExecuteFuncOnCells = [Func](const TSet<const UWorldPartitionRuntimeCell*>& Cells, EStreamingSourceTargetState TargetState)
@@ -954,7 +953,7 @@ void UWorldPartitionRuntimeHashSet::PostEditChangeChainProperty(FPropertyChanged
 		{
 			RuntimePartitionDesc.Name = RuntimePartitionDesc.Class->GetFName();
 			RuntimePartitionDesc.MainLayer = NewObject<URuntimePartition>(this, RuntimePartitionDesc.Class, NAME_None);
-			RuntimePartitionDesc.MainLayer->Name = RuntimePartitionDesc.Name;
+			RuntimePartitionDesc.MainLayer->SetDefaultValues();
 
 			RuntimePartitionDesc.UpdateHLODPartitionLayers();
 		}
@@ -1005,18 +1004,23 @@ void UWorldPartitionRuntimeHashSet::UpdateHLODPartitionLayers()
 }
 #endif
 
-void UWorldPartitionRuntimeHashSet::ForEachStreamingObject(TFunctionRef<void(const URuntimeHashSetExternalStreamingObject*)> Func) const
+void UWorldPartitionRuntimeHashSet::ForEachStreamingData(TFunctionRef<void(const FRuntimePartitionStreamingData&)> Func) const
 {
-	if (StreamingObject)
+	for (const FRuntimePartitionStreamingData& StreamingData : RuntimeStreamingData)
 	{
-		Func(StreamingObject.Get());
+		Func(StreamingData);
 	}
 
 	for (const TWeakObjectPtr<URuntimeHashExternalStreamingObjectBase>& InjectedExternalStreamingObject : InjectedExternalStreamingObjects)
 	{
 		if (InjectedExternalStreamingObject.IsValid())
 		{
-			Func(CastChecked<URuntimeHashSetExternalStreamingObject>(InjectedExternalStreamingObject.Get()));
+			URuntimeHashSetExternalStreamingObject* ExternalStreamingObject = CastChecked<URuntimeHashSetExternalStreamingObject>(InjectedExternalStreamingObject.Get());
+			
+			for (const FRuntimePartitionStreamingData& StreamingData : ExternalStreamingObject->RuntimeStreamingData)
+			{
+				Func(StreamingData);
+			}
 		}
 	}
 }
\ No newline at end of file
diff --git a/Engine/Source/Runtime/Engine/Public/WorldPartition/RuntimeHashSet/RuntimePartition.h b/Engine/Source/Runtime/Engine/Public/WorldPartition/RuntimeHashSet/RuntimePartition.h
--- a/Engine/Source/Runtime/Engine/Public/WorldPartition/RuntimeHashSet/RuntimePartition.h
+++ b/Engine/Source/Runtime/Engine/Public/WorldPartition/RuntimeHashSet/RuntimePartition.h
@@ -33,6 +33,7 @@ public:
 		TArray<IStreamingGenerationContext::FActorInstance> ActorInstances;
 	};
 
+	virtual void SetDefaultValues();
 	virtual bool SupportsHLODs() const PURE_VIRTUAL(URuntimePartition::SupportsHLODs, return false;);
 	virtual bool IsValidGrid(FName InGridName) const PURE_VIRTUAL(URuntimePartition::IsValidGrid, return false;);
 	virtual bool GenerateStreaming(const TArray<const IStreamingGenerationContext::FActorSetInstance*>& InActorSetInstances, TArray<FCellDesc>& OutRuntimeCellDescs) PURE_VIRTUAL(URuntimePartition::GenerateStreaming, return false;);
diff --git a/Engine/Source/Runtime/Engine/Public/WorldPartition/RuntimeHashSet/RuntimePartitionLHGrid.h b/Engine/Source/Runtime/Engine/Public/WorldPartition/RuntimeHashSet/RuntimePartitionLHGrid.h
--- a/Engine/Source/Runtime/Engine/Public/WorldPartition/RuntimeHashSet/RuntimePartitionLHGrid.h
+++ b/Engine/Source/Runtime/Engine/Public/WorldPartition/RuntimeHashSet/RuntimePartitionLHGrid.h
@@ -16,9 +16,12 @@ public:
 	virtual void PostEditChangeProperty(FPropertyChangedEvent& PropertyChangedEvent) override;
 	//~ End UObject Interface.
 
+	//~ Begin URuntimePartition interface
+	virtual void SetDefaultValues() override;
 	virtual bool SupportsHLODs() const override;
 	virtual bool IsValidGrid(FName GridName) const override;
 	virtual bool GenerateStreaming(const TArray<const IStreamingGenerationContext::FActorSetInstance*>& ActorSetInstances, TArray<FCellDesc>& OutRuntimeCellDescs) override;
+	//~ End URuntimePartition interface
 #endif
 
 private:
diff --git a/Engine/Source/Runtime/Engine/Public/WorldPartition/RuntimeHashSet/RuntimePartitionLevelStreaming.h b/Engine/Source/Runtime/Engine/Public/WorldPartition/RuntimeHashSet/RuntimePartitionLevelStreaming.h
--- a/Engine/Source/Runtime/Engine/Public/WorldPartition/RuntimeHashSet/RuntimePartitionLevelStreaming.h
+++ b/Engine/Source/Runtime/Engine/Public/WorldPartition/RuntimeHashSet/RuntimePartitionLevelStreaming.h
@@ -12,9 +12,11 @@ class URuntimePartitionLevelStreaming : public URuntimePartition
 
 public:
 #if WITH_EDITOR
+	//~ Begin URuntimePartition interface
 	virtual bool SupportsHLODs() const override;
 	virtual bool IsValidGrid(FName GridName) const override;
 	virtual bool GenerateStreaming(const TArray<const IStreamingGenerationContext::FActorSetInstance*>& ActorSetInstances, TArray<FCellDesc>& OutRuntimeCellDescs) override;
+	//~ End URuntimePartition interface
 #endif
 
 #if WITH_EDITORONLY_DATA
diff --git a/Engine/Source/Runtime/Engine/Public/WorldPartition/RuntimeHashSet/RuntimePartitionPersistent.h b/Engine/Source/Runtime/Engine/Public/WorldPartition/RuntimeHashSet/RuntimePartitionPersistent.h
--- a/Engine/Source/Runtime/Engine/Public/WorldPartition/RuntimeHashSet/RuntimePartitionPersistent.h
+++ b/Engine/Source/Runtime/Engine/Public/WorldPartition/RuntimeHashSet/RuntimePartitionPersistent.h
@@ -12,8 +12,10 @@ class URuntimePartitionPersistent : public URuntimePartition
 
 public:
 #if WITH_EDITOR
+	//~ Begin URuntimePartition interface
 	virtual bool SupportsHLODs() const override { return false; }
 	virtual bool IsValidGrid(FName GridName) const override { return true; }
 	virtual bool GenerateStreaming(const TArray<const IStreamingGenerationContext::FActorSetInstance*>& ActorSetInstances, TArray<FCellDesc>& OutRuntimeCellDescs) override;
+	//~ End URuntimePartition interface
 #endif
 };
diff --git a/Engine/Source/Runtime/Engine/Public/WorldPartition/RuntimeHashSet/WorldPartitionRuntimeHashSet.h b/Engine/Source/Runtime/Engine/Public/WorldPartition/RuntimeHashSet/WorldPartitionRuntimeHashSet.h
--- a/Engine/Source/Runtime/Engine/Public/WorldPartition/RuntimeHashSet/WorldPartitionRuntimeHashSet.h
+++ b/Engine/Source/Runtime/Engine/Public/WorldPartition/RuntimeHashSet/WorldPartitionRuntimeHashSet.h
@@ -66,6 +66,9 @@ struct FRuntimePartitionStreamingData
 {
 	GENERATED_USTRUCT_BODY()
 
+	void CreatePartitionsSpatialIndex() const;
+	void DestroyPartitionsSpatialIndex() const;
+
 	/** Name of the runtime partition, currently maps to target grids. */
 	UPROPERTY()
 	FName Name;
@@ -74,7 +77,10 @@ struct FRuntimePartitionStreamingData
 	int32 LoadingRange = 0;
 
 	UPROPERTY()
-	TArray<TObjectPtr<UWorldPartitionRuntimeCell>> RuntimeCells;
+	TArray<TObjectPtr<UWorldPartitionRuntimeCell>> StreamingCells;
+
+	UPROPERTY()
+	TArray<TObjectPtr<UWorldPartitionRuntimeCell>> NonStreamingCells;
 
 	// Transient
 	mutable TUniquePtr<FStaticSpatialIndexType> SpatialIndex;
@@ -102,9 +108,6 @@ public:
 	void CreatePartitionsSpatialIndex() const;
 	void DestroyPartitionsSpatialIndex() const;
 
-	UPROPERTY()
-	TArray<TObjectPtr<UWorldPartitionRuntimeCell>> NonSpatiallyLoadedRuntimeCells;
-
 	UPROPERTY()
 	TArray<FRuntimePartitionStreamingData> RuntimeStreamingData;
 };
@@ -166,7 +169,7 @@ private:
 	void UpdateHLODPartitionLayers();
 #endif
 
-	ENGINE_API void ForEachStreamingObject(TFunctionRef<void(const URuntimeHashSetExternalStreamingObject*)> Func) const;
+	ENGINE_API void ForEachStreamingData(TFunctionRef<void(const FRuntimePartitionStreamingData&)> Func) const;
 
 public:
 #if WITH_EDITORONLY_DATA
@@ -180,5 +183,5 @@ public:
 	TArray<FRuntimePartitionDesc> RuntimePartitions;
 
 	UPROPERTY()
-	TObjectPtr<URuntimeHashSetExternalStreamingObject> StreamingObject;
+	TArray<FRuntimePartitionStreamingData> RuntimeStreamingData;
 };
\ No newline at end of file
