Uniform Expression Cache:  fix for crashes caused by OnVirtualTextureDestroyed callbacks updating the uniform expression cache, overlapped with FSceneRenderer::UpdateScene, which reads from the uniform expression cache in UpdateAllPrimitiveSceneInfos.  Virtual Texture pending callbacks are moved from the start of FVirtualTextureSystem::BeginUpdate to a separate function, which gets called at the beginning of FSceneRenderer::UpdateScene.

#rnx

diff --git a/Engine/Source/Runtime/Engine/Private/Materials/MaterialRenderProxy.cpp b/Engine/Source/Runtime/Engine/Private/Materials/MaterialRenderProxy.cpp
--- a/Engine/Source/Runtime/Engine/Private/Materials/MaterialRenderProxy.cpp
+++ b/Engine/Source/Runtime/Engine/Private/Materials/MaterialRenderProxy.cpp
@@ -365,7 +365,11 @@ void FMaterialRenderProxy::EvaluateUniformExpressions(FUniformExpressionCache& O
 	FMaterialShaderMap* ShaderMap = Context.Material.GetRenderingThreadShaderMap();
 	const FUniformExpressionSet& UniformExpressionSet = ShaderMap->GetUniformExpressionSet();
 
-	OutUniformExpressionCache.CachedUniformExpressionShaderMap = ShaderMap;
+	// Initialize this to null, and set it to its final value (ShaderMap) at the end of the function.  The goal is to increase the timing window where
+	// we can detect thread safety bugs where the uniform expression cache is accessed while the expressions are being evaluated.  Bugs will typically
+	// manifest as an assert in FMaterialShader::GetShaderBindings, called from various mesh draw command generation logic.
+	OutUniformExpressionCache.CachedUniformExpressionShaderMap = nullptr;
+
 	OutUniformExpressionCache.ResetAllocatedVTs();
 	OutUniformExpressionCache.AllocatedVTs.Empty(UniformExpressionSet.VTStacks.Num());
 	OutUniformExpressionCache.OwnedAllocatedVTs.Empty(UniformExpressionSet.VTStacks.Num());
@@ -436,6 +440,9 @@ void FMaterialRenderProxy::EvaluateUniformExpressions(FUniformExpressionCache& O
 
 	OutUniformExpressionCache.ParameterCollections = UniformExpressionSet.ParameterCollections;
 
+	// Deliberately set this last, see comment above where it's initialized to nullptr
+	OutUniformExpressionCache.CachedUniformExpressionShaderMap = ShaderMap;
+
 	++UniformExpressionCacheSerialNumber;
 }
 
diff --git a/Engine/Source/Runtime/Renderer/Private/SceneRendering.cpp b/Engine/Source/Runtime/Renderer/Private/SceneRendering.cpp
--- a/Engine/Source/Runtime/Renderer/Private/SceneRendering.cpp
+++ b/Engine/Source/Runtime/Renderer/Private/SceneRendering.cpp
@@ -2836,6 +2836,10 @@ inline EUpdateAllPrimitiveSceneInfosAsyncOps GetUpdateAllPrimitiveSceneInfosAsyn
 
 IVisibilityTaskData* FSceneRenderer::UpdateScene(FRDGBuilder& GraphBuilder, FGlobalDynamicBuffers GlobalDynamicBuffers)
 {
+	// Needs to run before UpdateAllPrimitiveSceneInfos, as it may call OnVirtualTextureDestroyedCB, which modifies the uniform
+	// expression cache. UpdateAllPrimitiveSceneInfos generates mesh draw commands, which use the uniform expression cache.
+	FVirtualTextureSystem::Get().CallPendingCallbacks();
+
 	/**
 	  * UpdateStaticMeshes removes and re-creates cached FMeshDrawCommands.  If there are multiple scene renderers being run together,
 	  * we need allocated pipeline state IDs not to change, in case async tasks related to prior scene renderers are still in flight
diff --git a/Engine/Source/Runtime/Renderer/Private/VT/VirtualTextureProducer.cpp b/Engine/Source/Runtime/Renderer/Private/VT/VirtualTextureProducer.cpp
--- a/Engine/Source/Runtime/Renderer/Private/VT/VirtualTextureProducer.cpp
+++ b/Engine/Source/Runtime/Renderer/Private/VT/VirtualTextureProducer.cpp
@@ -175,6 +175,11 @@ void FVirtualTextureProducerCollection::CallPendingCallbacks()
 
 }
 
+bool FVirtualTextureProducerCollection::HasPendingCallbacks() const
+{
+	return NumPendingCallbacks != 0;
+}
+
 void FVirtualTextureProducerCollection::AddDestroyedCallback(const FVirtualTextureProducerHandle& Handle, FVTProducerDestroyedFunction* Function, void* Baton)
 {
 	check(IsInRenderingThread());
diff --git a/Engine/Source/Runtime/Renderer/Private/VT/VirtualTextureProducer.h b/Engine/Source/Runtime/Renderer/Private/VT/VirtualTextureProducer.h
--- a/Engine/Source/Runtime/Renderer/Private/VT/VirtualTextureProducer.h
+++ b/Engine/Source/Runtime/Renderer/Private/VT/VirtualTextureProducer.h
@@ -62,6 +62,7 @@ public:
 	void AddDestroyedCallback(const FVirtualTextureProducerHandle& Handle, FVTProducerDestroyedFunction* Function, void* Baton);
 	uint32 RemoveAllCallbacks(const void* Baton);
 	void CallPendingCallbacks();
+	bool HasPendingCallbacks() const;
 
 	/** Notify producers marked as "continous notify" that all requests have been completed. */
 	void NotifyRequestsCompleted();
diff --git a/Engine/Source/Runtime/Renderer/Private/VT/VirtualTextureSystem.cpp b/Engine/Source/Runtime/Renderer/Private/VT/VirtualTextureSystem.cpp
--- a/Engine/Source/Runtime/Renderer/Private/VT/VirtualTextureSystem.cpp
+++ b/Engine/Source/Runtime/Renderer/Private/VT/VirtualTextureSystem.cpp
@@ -2647,10 +2647,19 @@ void FVirtualTextureSystem::BeginUpdate(FRDGBuilder& GraphBuilder, FVirtualTextu
 	}, UE::Tasks::ETaskPriority::High, Updater->bAsyncTaskAllowed);
 }
 
+void FVirtualTextureSystem::CallPendingCallbacks()
+{
+	TRACE_CPUPROFILER_EVENT_SCOPE(FVirtualTextureSystem::CallPendingCallbacks);
+	SCOPE_CYCLE_COUNTER(STAT_VirtualTextureSystem_Update);
+
+	Producers.CallPendingCallbacks();
+}
+
 TUniquePtr<FVirtualTextureUpdater> FVirtualTextureSystem::BeginUpdate(FRDGBuilder& GraphBuilder, ERHIFeatureLevel::Type FeatureLevel, FScene* Scene, const FVirtualTextureUpdateSettings& Settings)
 {
 	check(IsInRenderingThread());
 	check(!bUpdating);
+	checkf(Producers.HasPendingCallbacks() == false, TEXT("FVirtualTextureSystem::CallPendingCallbacks(), typically called in FSceneRenderer::UpdateScene(), must run before FVirtualTextureSystem::BeginUpdate()"));
 
 	AllocateResources(GraphBuilder);
 
@@ -2659,8 +2668,6 @@ TUniquePtr<FVirtualTextureUpdater> FVirtualTextureSystem::BeginUpdate(FRDGBuilde
 	SCOPE_CYCLE_COUNTER(STAT_VirtualTextureSystem_Update);
 	RDG_GPU_STAT_SCOPE(GraphBuilder, VirtualTexture);
 
-	Producers.CallPendingCallbacks();
-
 	if (!Settings.bEnablePageRequests)
 	{
 		for (uint32 ID = 0; ID < MaxSpaces; ID++)
diff --git a/Engine/Source/Runtime/Renderer/Private/VT/VirtualTextureSystem.h b/Engine/Source/Runtime/Renderer/Private/VT/VirtualTextureSystem.h
--- a/Engine/Source/Runtime/Renderer/Private/VT/VirtualTextureSystem.h
+++ b/Engine/Source/Runtime/Renderer/Private/VT/VirtualTextureSystem.h
@@ -103,6 +103,9 @@ public:
 
 	void Update(FRDGBuilder& GraphBuilder, ERHIFeatureLevel::Type FeatureLevel, FScene* Scene, const FVirtualTextureUpdateSettings& Settings);
 
+	// Called in FSceneRenderer::UpdateScene, before FVirtualTextureSystem::BeginUpdate -- see comments there for more info
+	void CallPendingCallbacks();
+
 	TUniquePtr<FVirtualTextureUpdater> BeginUpdate(FRDGBuilder& GraphBuilder, ERHIFeatureLevel::Type FeatureLevel, FScene* Scene, const FVirtualTextureUpdateSettings& Settings);
 	void WaitForTasks(FVirtualTextureUpdater* Updater);
 	void EndUpdate(FRDGBuilder& GraphBuilder, TUniquePtr<FVirtualTextureUpdater>&& Updater, ERHIFeatureLevel::Type FeatureLevel);
