Trimmed down some new and verbose shader defines when using  r.DumpShaderDebugShortNames

#rnx

diff --git a/Engine/Source/Runtime/Engine/Private/ShaderCompiler/ShaderCompiler.cpp b/Engine/Source/Runtime/Engine/Private/ShaderCompiler/ShaderCompiler.cpp
--- a/Engine/Source/Runtime/Engine/Private/ShaderCompiler/ShaderCompiler.cpp
+++ b/Engine/Source/Runtime/Engine/Private/ShaderCompiler/ShaderCompiler.cpp
@@ -6063,6 +6063,9 @@ void GlobalBeginCompileShader(
 			Input.DebugGroupName.ReplaceInline(TEXT("PROPAGATE_AO"), TEXT("AO"));
 			Input.DebugGroupName.ReplaceInline(TEXT("PROPAGATE_SECONDARY_OCCLUSION"), TEXT("SEC_OCC"));
 			Input.DebugGroupName.ReplaceInline(TEXT("PROPAGATE_MULTIPLE_BOUNCES"), TEXT("MULT_BOUNC"));
+			Input.DebugGroupName.ReplaceInline(TEXT("LOCAL_LIGHTS_DISABLED"), TEXT("NoLL"));
+			Input.DebugGroupName.ReplaceInline(TEXT("LOCAL_LIGHTS_ENABLED"), TEXT("LL"));
+			Input.DebugGroupName.ReplaceInline(TEXT("LOCAL_LIGHTS_PREPASS_ENABLED"), TEXT("LLPP"));
 			Input.DebugGroupName.ReplaceInline(TEXT("PostProcess"), TEXT("Post"));
 			Input.DebugGroupName.ReplaceInline(TEXT("AntiAliasing"), TEXT("AA"));
 			Input.DebugGroupName.ReplaceInline(TEXT("Mobile"), TEXT("Mob"));
