Adding another category for CSV reports to move PostProcess time out of ViewportMisc
#RNX

diff --git a/Engine/Source/Runtime/Engine/Private/LocalPlayer.cpp b/Engine/Source/Runtime/Engine/Private/LocalPlayer.cpp
--- a/Engine/Source/Runtime/Engine/Private/LocalPlayer.cpp
+++ b/Engine/Source/Runtime/Engine/Private/LocalPlayer.cpp
@@ -869,6 +869,8 @@ FSceneView* ULocalPlayer::CalcSceneView( class FSceneViewFamily* ViewFamily,
 	ViewFamily->Views.Add(View);
 
 	{
+		QUICK_SCOPE_CYCLE_COUNTER(STAT_PostprocessSettings);
+		CSV_SCOPED_TIMING_STAT_EXCLUSIVE(PostProcessSettings);
 		View->StartFinalPostprocessSettings(ViewInfo.Location);
 
 		TArray<FPostProcessSettings> const* CameraAnimPPSettings = nullptr;
