Only sets r.Mobile.FloatPrecisionMode on mobile renderer's shaders

diff --git a/Engine/Source/Runtime/Engine/Private/ShaderCompiler/ShaderCompiler.cpp b/Engine/Source/Runtime/Engine/Private/ShaderCompiler/ShaderCompiler.cpp
--- a/Engine/Source/Runtime/Engine/Private/ShaderCompiler/ShaderCompiler.cpp
+++ b/Engine/Source/Runtime/Engine/Private/ShaderCompiler/ShaderCompiler.cpp
@@ -6264,6 +6264,7 @@ void GlobalBeginCompileShader(
 		Input.Environment.CompilerFlags.Add(CFLAG_NoFastMath);
 	}
     
+	if (bUsingMobileRenderer)
     {
         static IConsoleVariable* CVar = IConsoleManager::Get().FindConsoleVariable(TEXT("r.Mobile.FloatPrecisionMode"));
         Input.Environment.FullPrecisionInPS = CVar ? (CVar->GetInt() == EMobileFloatPrecisionMode::Full) : false;
