SplineComponent: Fix crash in ConvertSplineToPolyLine when spline has no points.

diff --git a/Engine/Source/Runtime/Engine/Private/Components/SplineComponent.cpp b/Engine/Source/Runtime/Engine/Private/Components/SplineComponent.cpp
--- a/Engine/Source/Runtime/Engine/Private/Components/SplineComponent.cpp
+++ b/Engine/Source/Runtime/Engine/Private/Components/SplineComponent.cpp
@@ -1042,7 +1042,7 @@ int32 USplineComponent::GetNumberOfSplinePoints() const
 int32 USplineComponent::GetNumberOfSplineSegments() const
 {
 	const int32 NumPoints = SplineCurves.Position.Points.Num();
-	return (bClosedLoop ? NumPoints : NumPoints - 1);
+	return (bClosedLoop ? NumPoints : FMath::Max(0, NumPoints - 1));
 }
 
 float USplineComponent::GetInputKeyValueAtSplinePoint(int32 PointIndex) const
