PR #10597: Fix FTickableGameObject crash

diff --git a/Engine/Source/Runtime/Engine/Private/Subsystems/WorldSubsystem.cpp b/Engine/Source/Runtime/Engine/Private/Subsystems/WorldSubsystem.cpp
--- a/Engine/Source/Runtime/Engine/Private/Subsystems/WorldSubsystem.cpp
+++ b/Engine/Source/Runtime/Engine/Private/Subsystems/WorldSubsystem.cpp
@@ -60,7 +60,8 @@ bool UTickableWorldSubsystem::IsAllowedToTick() const
 	// No matter what IsTickable says, don't let CDOs or uninitialized world subsystems tick :
 	// Note: even if GetTickableTickType was overridden by the child class and returns something else than ETickableTickType::Never for CDOs, 
 	//  it's probably a mistake, so by default, don't allow ticking. If the child class really intends its CDO to tick, it can always override IsAllowedToTick...
-	return !IsTemplate() && bInitialized;
+	// NOTE: `bInitialized` must be checked first as `IsTemplate()` might access a dangling `Outer` if we are awaiting GC but `Outer` has already been deleted.
+	return bInitialized && !IsTemplate();
 }
 
 void UTickableWorldSubsystem::Tick(float DeltaTime)
diff --git a/Engine/Source/Runtime/Engine/Private/Tickable.cpp b/Engine/Source/Runtime/Engine/Private/Tickable.cpp
--- a/Engine/Source/Runtime/Engine/Private/Tickable.cpp
+++ b/Engine/Source/Runtime/Engine/Private/Tickable.cpp
@@ -136,9 +136,9 @@ void FTickableGameObject::TickObjects(UWorld* World, const int32 InTickType, con
 			if (FTickableGameObject* TickableObject = static_cast<FTickableGameObject*>(TickableEntry.TickableObject))
 			{
 				// If it is tickable and in this world
-				if (((TickableEntry.TickType == ETickableTickType::Always) || TickableObject->IsTickable()) 
-					&& (TickableObject->GetTickableGameObjectWorld() == World)
-					&& TickableObject->IsAllowedToTick())
+				if (TickableObject->IsAllowedToTick()
+					&& ((TickableEntry.TickType == ETickableTickType::Always) || TickableObject->IsTickable())
+					&& (TickableObject->GetTickableGameObjectWorld() == World))
 				{
 					const bool bIsGameWorld = InTickType == LEVELTICK_All || (World && World->IsGameWorld());
 					// If we are in editor and it is editor tickable, always tick
