PR #10521: Adding DeterminesOutputType to camera modifier fn's in PlayerCameraManager

diff --git a/Engine/Source/Runtime/Engine/Classes/Camera/PlayerCameraManager.h b/Engine/Source/Runtime/Engine/Classes/Camera/PlayerCameraManager.h
--- a/Engine/Source/Runtime/Engine/Classes/Camera/PlayerCameraManager.h
+++ b/Engine/Source/Runtime/Engine/Classes/Camera/PlayerCameraManager.h
@@ -594,14 +594,14 @@ public:
 	 * @param ModifierClass - The class of camera modifier to create.
 	 * @return Returns the newly created camera modifier.
 	 */
-	UFUNCTION(BlueprintCallable, Category = "Game|Player")
+	UFUNCTION(BlueprintCallable, Category = "Game|Player", meta = (DeterminesOutputType = "ModifierClass"))
 	ENGINE_API virtual UCameraModifier* AddNewCameraModifier(TSubclassOf<UCameraModifier> ModifierClass);
 
 	/** 
 	 * Returns camera modifier for this camera of the given class, if it exists. 
 	 * Exact class match only. If there are multiple modifiers of the same class, the first one is returned.
 	 */
-	UFUNCTION(BlueprintCallable, Category = "Game|Player")
+	UFUNCTION(BlueprintCallable, Category = "Game|Player", meta = (DeterminesOutputType = "ModifierClass"))
 	ENGINE_API virtual UCameraModifier* FindCameraModifierByClass(TSubclassOf<UCameraModifier> ModifierClass);
 
 	/** 
