World Partition: fix runtime spatial hash debug display scaling regression.

diff --git a/Engine/Source/Runtime/Engine/Private/WorldPartition/WorldPartitionRuntimeSpatialHash.cpp b/Engine/Source/Runtime/Engine/Private/WorldPartition/WorldPartitionRuntimeSpatialHash.cpp
--- a/Engine/Source/Runtime/Engine/Private/WorldPartition/WorldPartitionRuntimeSpatialHash.cpp
+++ b/Engine/Source/Runtime/Engine/Private/WorldPartition/WorldPartitionRuntimeSpatialHash.cpp
@@ -1824,8 +1824,20 @@ bool UWorldPartitionRuntimeSpatialHash::Draw2D(FWorldPartitionDraw2DContext& Dra
 			Source.ForEachShape(StreamingGrid->GetLoadingRange(), StreamingGrid->GridName, StreamingGrid->HLODLayer, /*bProjectIn2D*/ true, [&GridsShapeBounds](const FSphericalSector& Shape) { GridsShapeBounds += Shape.CalcBounds(); });
 		}
 
-		const FVector2D GridReferenceWorldPos = FVector2D(WorldRegion.GetCenter());
-		const FVector2D WorldRegionExtent = FVector2D(WorldRegion.GetExtent().GetMax());
+		FVector2D GridReferenceWorldPos;
+		FVector2D WorldRegionExtent;
+
+		if (DrawContext.IsDetailedMode())
+		{
+			GridReferenceWorldPos = FVector2D(GridsShapeBounds.GetCenter());
+			WorldRegionExtent = FVector2D(GridsShapeBounds.ExpandBy(GridsShapeBounds.GetExtent() * 0.1f).GetExtent().GetMax());
+		}
+		else
+		{
+			GridReferenceWorldPos = FVector2D(WorldRegion.GetCenter());
+			WorldRegionExtent = FVector2D(WorldRegion.GetExtent().GetMax());
+		}
+
 		const FVector2D GridScreenOffset = GridScreenInitialOffset + ((float)GridIndex * FVector2D(GridMaxScreenWidth, 0.f)) + GridScreenHalfExtent + FVector2D(GridScreenWidthShrinkSize * 0.5f);
 		const FVector2D WorldToScreenScale = GridScreenHalfExtent / WorldRegionExtent;
 		const FBox2D GridScreenBounds(GridScreenOffset - GridScreenHalfExtent, GridScreenOffset + GridScreenHalfExtent);
