World Partition Class Descriptor Registry: don't prefetch classes that don't exists on disk yet when prefetching class hierarchy.
Avoids warning "Package does not exist, will not scan" when saving a new Blueprint actor BP class.

diff --git a/Engine/Source/Runtime/Engine/Private/WorldPartition/WorldPartitionClassDescRegistry.cpp b/Engine/Source/Runtime/Engine/Private/WorldPartition/WorldPartitionClassDescRegistry.cpp
--- a/Engine/Source/Runtime/Engine/Private/WorldPartition/WorldPartitionClassDescRegistry.cpp
+++ b/Engine/Source/Runtime/Engine/Private/WorldPartition/WorldPartitionClassDescRegistry.cpp
@@ -420,7 +420,11 @@ void FWorldPartitionClassDescRegistry::PrefetchClassDesc(UClass* InClass)
 						
 		if (!ClassByPath.Contains(ParentClassPath))
 		{
-			PrefetchClassDescs({ ParentClassPath });
+			// Only prefetch classes that exists on disk
+			if (!ParentClass->GetPackage()->HasAnyPackageFlags(PKG_NewlyCreated))
+			{
+				PrefetchClassDescs({ ParentClassPath });
+			}
 
 			if (ClassByPath.Contains(ParentClassPath))
 			{
