Fix crash in cook when serializing textures with no mips, as happens when a texture build fails

#rnx

diff --git a/Engine/Source/Runtime/Engine/Private/TextureDerivedData.cpp b/Engine/Source/Runtime/Engine/Private/TextureDerivedData.cpp
--- a/Engine/Source/Runtime/Engine/Private/TextureDerivedData.cpp
+++ b/Engine/Source/Runtime/Engine/Private/TextureDerivedData.cpp
@@ -2502,7 +2502,7 @@ static void SerializePlatformData(
 
 		#if WITH_EDITORONLY_DATA
 			static bool bDisableOptionalMips = FParse::Param(FCommandLine::Get(), TEXT("DisableOptionalMips"));
-			if (!bDisableOptionalMips)
+			if (!bDisableOptionalMips && NumMips > 0 )
 			{
 				const int32 LODGroup = Texture->LODGroup;
 				const int32 FirstMipWidth  = PlatformData->Mips[FirstMipToSerialize].SizeX;
