World Partition Runtime HashSet: added debug display view modes for data layers and content bundles.

diff --git a/Engine/Source/Runtime/Engine/Private/WorldPartition/RuntimeHashSet/WorldPartitionRuntimeHashSet.cpp b/Engine/Source/Runtime/Engine/Private/WorldPartition/RuntimeHashSet/WorldPartitionRuntimeHashSet.cpp
--- a/Engine/Source/Runtime/Engine/Private/WorldPartition/RuntimeHashSet/WorldPartitionRuntimeHashSet.cpp
+++ b/Engine/Source/Runtime/Engine/Private/WorldPartition/RuntimeHashSet/WorldPartitionRuntimeHashSet.cpp
@@ -31,15 +31,21 @@ FAutoConsoleCommand WorldPartitionRuntimeHashSetEnable(
 
 static int32 GShowRuntimeHashSetDebugDisplayLevel = 0;
 static FAutoConsoleVariableRef CVarShowRuntimeHashSetDebugDisplayLevel(
-	TEXT("wp.Runtime.HashSet.ShowtDebugDisplayLevel"),
+	TEXT("wp.Runtime.HashSet.ShowDebugDisplayLevel"),
 	GShowRuntimeHashSetDebugDisplayLevel,
-	TEXT("Used to choose which level to display when showing world partition partitions."));
+	TEXT("Used to choose which level to display when showing runtime partitions."));
 
 static int32 GShowRuntimeHashSetDebugDisplayLevelCount = 1;
 static FAutoConsoleVariableRef CVarShowRuntimeHashSetDebugDisplayLevelCount(
 	TEXT("wp.Runtime.HashSet.ShowDebugDisplayLevelCount"),
 	GShowRuntimeHashSetDebugDisplayLevelCount,
-	TEXT("Used to choose how many levels to display when showing world partition runtime partitions."));
+	TEXT("Used to choose how many levels to display when showing runtime partitions."));
+
+static int32 GShowRuntimeHashSetDebugDisplayMode = 0;
+static FAutoConsoleVariableRef CVarShowRuntimeHashSetDebugDisplayMode(
+	TEXT("wp.Runtime.HashSet.ShowDebugDisplayMode"),
+	GShowRuntimeHashSetDebugDisplayMode,
+	TEXT("Used to choose what mode to display when showing runtime partitions (0=Level Streaming State, 1=Data Layers, 2=Content Bundles)."));
 
 #if WITH_EDITOR
 void FRuntimePartitionDesc::UpdateHLODPartitionLayers()
@@ -273,7 +279,55 @@ bool UWorldPartitionRuntimeHashSet::Draw2D(FWorldPartitionDraw2DContext& DrawCon
 			{
 				const FVector2D CellBoundsSize = FVector2D(Cell->GetCellBounds().GetSize());
 				const FVector2D CellBoundsMin = FVector2D(Cell->GetCellBounds().Min);
-				DrawContext.LocalDrawTile(GridScreenBounds, CellBoundsMin, CellBoundsSize, Cell->GetDebugColor(VisualizeMode).CopyWithNewOpacity(0.25f / FMath::Max<float>(GShowRuntimeHashSetDebugDisplayLevelCount, 1)), WorldToScreen);
+				
+				float CellOpacity;
+				TArray<FLinearColor> CellColors;
+
+				switch (GShowRuntimeHashSetDebugDisplayMode)
+				{
+				case 0:
+					CellColors.Add(Cell->GetDebugColor(VisualizeMode));
+					CellOpacity = 0.25f / FMath::Max<float>(GShowRuntimeHashSetDebugDisplayLevelCount, 1);
+					break;
+				case 1:
+					if (DataLayerDebugColors.Num() && Cell->GetDataLayers().Num())
+					{
+						for (const FName& DataLayer : Cell->GetDataLayers())
+						{
+							CellColors.Add(DataLayerDebugColors[DataLayer]);
+						}
+						CellOpacity = 0.67f;
+					}
+					break;
+				case 2:
+					if (ContentBundleManager && Cell->GetContentBundleID().IsValid())
+					{
+						if (const FContentBundleBase* ContentBundle = ContentBundleManager->GetContentBundle(OwningWorld, Cell->GetContentBundleID()))
+						{
+							check(ContentBundle->GetDescriptor());
+							CellColors.Add(ContentBundle->GetDescriptor()->GetDebugColor());
+							CellOpacity = 0.67f;
+						}
+					}
+					break;
+				}
+
+				if (CellColors.IsEmpty())
+				{
+					CellColors.Add(FLinearColor::White);
+					CellOpacity = 0.1f;
+				}
+
+				FVector2D::FReal BoundsOffsetX = 0;
+				const FVector2D::FReal BoundsOffsetStepX = CellBoundsSize.X / CellColors.Num();
+				for (const FLinearColor& CellColor : CellColors)
+				{
+					const FVector2D EffectiveCellBoundsMin(CellBoundsMin.X + BoundsOffsetX, CellBoundsMin.Y);
+					const FVector2D EffectiveCellBoundsSize(BoundsOffsetStepX, CellBoundsSize.Y);
+					DrawContext.LocalDrawTile(GridScreenBounds, EffectiveCellBoundsMin, EffectiveCellBoundsSize, CellColor.CopyWithNewOpacity(CellOpacity), WorldToScreen);
+					BoundsOffsetX += BoundsOffsetStepX;
+				}
+				
 				DrawContext.LocalDrawBox(GridScreenBounds, CellBoundsMin, CellBoundsSize, FLinearColor::Black, 1, WorldToScreen);
 			}
 		}
