Fix groom view mode in cook build.

diff --git a/Engine/Source/Runtime/Engine/Private/GroomVisualizationData.cpp b/Engine/Source/Runtime/Engine/Private/GroomVisualizationData.cpp
--- a/Engine/Source/Runtime/Engine/Private/GroomVisualizationData.cpp
+++ b/Engine/Source/Runtime/Engine/Private/GroomVisualizationData.cpp
@@ -30,51 +30,62 @@ EGroomViewMode GetGroomViewMode(const FSceneView& View)
 	EGroomViewMode Out = EGroomViewMode::None;
 	if (IsGroomEnabled())
 	{
-		static const auto CVarGroomViewMode = IConsoleManager::Get().FindConsoleVariable(FGroomVisualizationData::GetVisualizeConsoleCommandName());
-		const uint32 ViewMode = CVarGroomViewMode && CVarGroomViewMode->AsVariableInt() ? CVarGroomViewMode->AsVariableInt()->GetValueOnRenderThread() : 0;
-		switch (ViewMode)
+		static TConsoleVariableData<int32>* CVarGroomViewMode = nullptr; 
+		if (CVarGroomViewMode == nullptr)
 		{
-		case 1:  return EGroomViewMode::MacroGroups;
-		case 2:  return EGroomViewMode::LightBounds;
-		case 3:  return EGroomViewMode::MacroGroupScreenRect;
-		case 4:  return EGroomViewMode::DeepOpacityMaps;
-		case 5:  return EGroomViewMode::SamplePerPixel;
-		case 6:  return EGroomViewMode::TAAResolveType;
-		case 7:  return EGroomViewMode::CoverageType;
-		case 8:  return EGroomViewMode::VoxelsDensity;
-		case 9:	 return EGroomViewMode::None;
-		case 10: return EGroomViewMode::None;
-		case 11: return EGroomViewMode::None;
-		case 12: return EGroomViewMode::MeshProjection;
-		case 13: return EGroomViewMode::Coverage;
-		case 14: return EGroomViewMode::MaterialDepth;
-		case 15: return EGroomViewMode::MaterialBaseColor;
-		case 16: return EGroomViewMode::MaterialRoughness;
-		case 17: return EGroomViewMode::MaterialSpecular;
-		case 18: return EGroomViewMode::MaterialTangent;
-		case 19: return EGroomViewMode::Tile;
-		case 20: return EGroomViewMode::None;
-		case 21: return EGroomViewMode::SimHairStrands;
-		case 22: return EGroomViewMode::RenderHairStrands;
-		case 23: return EGroomViewMode::RootUV;
-		case 24: return EGroomViewMode::RootUDIM;
-		case 25: return EGroomViewMode::UV;
-		case 26: return EGroomViewMode::Seed;
-		case 27: return EGroomViewMode::Dimension;
-		case 28: return EGroomViewMode::RadiusVariation;
-		case 29: return EGroomViewMode::Color;
-		case 30: return EGroomViewMode::Roughness;
-		case 31: return EGroomViewMode::Cluster;
-		case 32: return EGroomViewMode::ClusterAABB;
-		case 33: return EGroomViewMode::Tangent;
-		case 34: return EGroomViewMode::ControlPoints;
-		case 35: return EGroomViewMode::Group;
-		case 36: return EGroomViewMode::LODColoration;
-		case 37: return EGroomViewMode::CardGuides;
-		case 38: return EGroomViewMode::AO;
-		case 39: return EGroomViewMode::ClumpID;
-		case 40: return EGroomViewMode::Memory;
-		default: break;
+			if (IConsoleVariable* ConsoleVariable = IConsoleManager::Get().FindConsoleVariable(FGroomVisualizationData::GetVisualizeConsoleCommandName()))
+			{
+				CVarGroomViewMode = ConsoleVariable->AsVariableInt();
+			}
+		}
+
+		if (CVarGroomViewMode)
+		{
+			const uint32 ViewMode = CVarGroomViewMode->GetValueOnRenderThread();
+			switch (ViewMode)
+			{
+			case 1:  return EGroomViewMode::MacroGroups;
+			case 2:  return EGroomViewMode::LightBounds;
+			case 3:  return EGroomViewMode::MacroGroupScreenRect;
+			case 4:  return EGroomViewMode::DeepOpacityMaps;
+			case 5:  return EGroomViewMode::SamplePerPixel;
+			case 6:  return EGroomViewMode::TAAResolveType;
+			case 7:  return EGroomViewMode::CoverageType;
+			case 8:  return EGroomViewMode::VoxelsDensity;
+			case 9:	 return EGroomViewMode::None;
+			case 10: return EGroomViewMode::None;
+			case 11: return EGroomViewMode::None;
+			case 12: return EGroomViewMode::MeshProjection;
+			case 13: return EGroomViewMode::Coverage;
+			case 14: return EGroomViewMode::MaterialDepth;
+			case 15: return EGroomViewMode::MaterialBaseColor;
+			case 16: return EGroomViewMode::MaterialRoughness;
+			case 17: return EGroomViewMode::MaterialSpecular;
+			case 18: return EGroomViewMode::MaterialTangent;
+			case 19: return EGroomViewMode::Tile;
+			case 20: return EGroomViewMode::None;
+			case 21: return EGroomViewMode::SimHairStrands;
+			case 22: return EGroomViewMode::RenderHairStrands;
+			case 23: return EGroomViewMode::RootUV;
+			case 24: return EGroomViewMode::RootUDIM;
+			case 25: return EGroomViewMode::UV;
+			case 26: return EGroomViewMode::Seed;
+			case 27: return EGroomViewMode::Dimension;
+			case 28: return EGroomViewMode::RadiusVariation;
+			case 29: return EGroomViewMode::Color;
+			case 30: return EGroomViewMode::Roughness;
+			case 31: return EGroomViewMode::Cluster;
+			case 32: return EGroomViewMode::ClusterAABB;
+			case 33: return EGroomViewMode::Tangent;
+			case 34: return EGroomViewMode::ControlPoints;
+			case 35: return EGroomViewMode::Group;
+			case 36: return EGroomViewMode::LODColoration;
+			case 37: return EGroomViewMode::CardGuides;
+			case 38: return EGroomViewMode::AO;
+			case 39: return EGroomViewMode::ClumpID;
+			case 40: return EGroomViewMode::Memory;
+			default: break;
+			}
 		}
 
 		const FGroomVisualizationData& VisualizationData = GetGroomVisualizationData();
