Updated ARM64EC to use the win64 vorbis dll.

diff --git a/Engine/Source/Runtime/Engine/Private/VorbisAudioInfo.cpp b/Engine/Source/Runtime/Engine/Private/VorbisAudioInfo.cpp
--- a/Engine/Source/Runtime/Engine/Private/VorbisAudioInfo.cpp
+++ b/Engine/Source/Runtime/Engine/Private/VorbisAudioInfo.cpp
@@ -784,7 +784,7 @@ void LoadVorbisLibraries()
 		PlatformString = TEXT("HoloLens");
 #endif
 
-#if PLATFORM_CPU_ARM_FAMILY
+#if PLATFORM_CPU_ARM_FAMILY && !defined(_M_ARM64EC)
 #if PLATFORM_64BITS
 		FString RootOggPath = FPaths::EngineDir() / TEXT("Binaries/ThirdParty/Ogg/") / PlatformString / VSVersion / TEXT("arm64/");
 		FString RootVorbisPath = FPaths::EngineDir() / TEXT("Binaries/ThirdParty/Vorbis/") / PlatformString / VSVersion / TEXT("arm64/");
