Replaced operator new TArray calls with emplacement.

diff --git a/Engine/Source/Runtime/Core/Private/Stats/StatsData.cpp b/Engine/Source/Runtime/Core/Private/Stats/StatsData.cpp
--- a/Engine/Source/Runtime/Core/Private/Stats/StatsData.cpp
+++ b/Engine/Source/Runtime/Core/Private/Stats/StatsData.cpp
@@ -436,8 +436,8 @@ void FRawStatStackNode::Encode(TArray<FStatMessage>& OutStats) const
 			FRawStatStackNode const* Child = It.Value();
 			Child->Encode(OutStats);
 		}
-		FStatMessage* EndStat = new (OutStats) FStatMessage(Meta);
-		EndStat->NameAndInfo.SetField<EStatOperation>(EStatOperation::ChildrenEnd);
+		FStatMessage& EndStat = OutStats.Emplace_GetRef(Meta);
+		EndStat.NameAndInfo.SetField<EStatOperation>(EStatOperation::ChildrenEnd);
 	}
 	else
 	{
diff --git a/Engine/Source/Runtime/Engine/Private/Animation/AnimNode_StateMachine.cpp b/Engine/Source/Runtime/Engine/Private/Animation/AnimNode_StateMachine.cpp
--- a/Engine/Source/Runtime/Engine/Private/Animation/AnimNode_StateMachine.cpp
+++ b/Engine/Source/Runtime/Engine/Private/Animation/AnimNode_StateMachine.cpp
@@ -230,12 +230,12 @@ void FAnimNode_StateMachine::Initialize_AnyThread(const FAnimationInitializeCont
 			for (int32 StateIndex = 0; StateIndex < Machine->States.Num(); ++StateIndex)
 			{
 				const FBakedAnimationState& State = Machine->States[StateIndex];
-				FPoseLink* StatePoseLink = new (StatePoseLinks) FPoseLink();
+				FPoseLink& StatePoseLink = StatePoseLinks.AddDefaulted_GetRef();
 
 				// because conduits don't contain bound graphs, this link is no longer guaranteed to be valid
 				if (State.StateRootNodeIndex != INDEX_NONE)
 				{
-					StatePoseLink->LinkID = AnimBlueprintClass->GetAnimNodeProperties().Num() - 1 - State.StateRootNodeIndex; //@TODO: Crazysauce
+					StatePoseLink.LinkID = AnimBlueprintClass->GetAnimNodeProperties().Num() - 1 - State.StateRootNodeIndex; //@TODO: Crazysauce
 				}
 
 				// also initialize transitions
@@ -1225,59 +1225,56 @@ void FAnimNode_StateMachine::TransitionToState(const FAnimationUpdateContext& Co
 		const float ExistingWeightOfNextState = GetStateWeight(NextState);
 
 		// Push the transition onto the stack
-		FAnimationActiveTransitionEntry* NewTransition = new (ActiveTransitionArray) FAnimationActiveTransitionEntry(NextState, ExistingWeightOfNextState, PreviousState, TransitionInfo, CrossFadeTimeAdjustment);
-		if (NewTransition)
+		FAnimationActiveTransitionEntry& NewTransition = ActiveTransitionArray.Emplace_GetRef(NextState, ExistingWeightOfNextState, PreviousState, TransitionInfo, CrossFadeTimeAdjustment);
+		if ((TransitionInfo.LogicType == ETransitionLogicType::TLT_Custom) && BakedExitTransition)
 		{
-			if ((TransitionInfo.LogicType == ETransitionLogicType::TLT_Custom) && BakedExitTransition)
-			{
-				NewTransition->InitializeCustomGraphLinks(Context, *BakedExitTransition);
-			}
+			NewTransition.InitializeCustomGraphLinks(Context, *BakedExitTransition);
+		}
 
-			// Initialize blend data if necessary
-			if (NewTransition->BlendProfile)
-			{
-				NewTransition->StateBlendData.AddZeroed(2);
-				NewTransition->StateBlendData[0].PerBoneBlendData.AddZeroed(NewTransition->BlendProfile->GetNumBlendEntries());
-				NewTransition->StateBlendData[1].PerBoneBlendData.AddZeroed(NewTransition->BlendProfile->GetNumBlendEntries());
-			}
+		// Initialize blend data if necessary
+		if (NewTransition.BlendProfile)
+		{
+			NewTransition.StateBlendData.AddZeroed(2);
+			NewTransition.StateBlendData[0].PerBoneBlendData.AddZeroed(NewTransition.BlendProfile->GetNumBlendEntries());
+			NewTransition.StateBlendData[1].PerBoneBlendData.AddZeroed(NewTransition.BlendProfile->GetNumBlendEntries());
+		}
 
-			if (TransitionInfo.LogicType == ETransitionLogicType::TLT_Inertialization)
+		if (TransitionInfo.LogicType == ETransitionLogicType::TLT_Inertialization)
+		{
+			UE::Anim::IInertializationRequester* InertializationRequester = Context.GetMessage<UE::Anim::IInertializationRequester>();
+			if (InertializationRequester)
 			{
-				UE::Anim::IInertializationRequester* InertializationRequester = Context.GetMessage<UE::Anim::IInertializationRequester>();
-				if (InertializationRequester)
-				{
-					FInertializationRequest Request;
-					Request.Duration = TransitionInfo.CrossfadeDuration;
-					Request.BlendProfile = TransitionInfo.BlendProfile;
-					Request.bUseBlendMode = true;
-					Request.BlendMode = TransitionInfo.BlendMode;
-					Request.CustomBlendCurve = TransitionInfo.CustomCurve;
+				FInertializationRequest Request;
+				Request.Duration = TransitionInfo.CrossfadeDuration;
+				Request.BlendProfile = TransitionInfo.BlendProfile;
+				Request.bUseBlendMode = true;
+				Request.BlendMode = TransitionInfo.BlendMode;
+				Request.CustomBlendCurve = TransitionInfo.CustomCurve;
 #if ANIM_TRACE_ENABLED
-					Request.Description = FText::Format(LOCTEXT("InertializationRequestDescription", 
-						"\"{0}\" Transition from \"{1}\" to \"{2}\""), 
-						FText::FromName(GetMachineDescription()->MachineName),
-						FText::FromName(GetStateInfo(TransitionInfo.PreviousState).StateName),
-						FText::FromName(GetStateInfo(TransitionInfo.NextState).StateName));
-					Request.NodeId = Context.GetCurrentNodeId();
-					Request.AnimInstance = Context.AnimInstanceProxy->GetAnimInstanceObject();
+				Request.Description = FText::Format(LOCTEXT("InertializationRequestDescription", 
+					"\"{0}\" Transition from \"{1}\" to \"{2}\""), 
+					FText::FromName(GetMachineDescription()->MachineName),
+					FText::FromName(GetStateInfo(TransitionInfo.PreviousState).StateName),
+					FText::FromName(GetStateInfo(TransitionInfo.NextState).StateName));
+				Request.NodeId = Context.GetCurrentNodeId();
+				Request.AnimInstance = Context.AnimInstanceProxy->GetAnimInstanceObject();
 #endif
 
-					InertializationRequester->RequestInertialization(Request);
-					InertializationRequester->AddDebugRecord(*Context.AnimInstanceProxy, Context.GetCurrentNodeId());
-				}
-				else
-				{
-					LogInertializationRequestError(Context, PreviousState, NextState);
-				}
+				InertializationRequester->RequestInertialization(Request);
+				InertializationRequester->AddDebugRecord(*Context.AnimInstanceProxy, Context.GetCurrentNodeId());
 			}
-
-			NewTransition->SourceTransitionIndices = SourceTransitionIndices;
-
-			if (!bFirstUpdate || (bFirstUpdate && !bSkipFirstUpdateTransition))
+			else
 			{
-				Context.AnimInstanceProxy->AddAnimNotifyFromGeneratedClass(NewTransition->StartNotify);
+				LogInertializationRequestError(Context, PreviousState, NextState);
 			}
 		}
+
+		NewTransition.SourceTransitionIndices = SourceTransitionIndices;
+
+		if (!bFirstUpdate || (bFirstUpdate && !bSkipFirstUpdateTransition))
+		{
+			Context.AnimInstanceProxy->AddAnimNotifyFromGeneratedClass(NewTransition.StartNotify);
+		}
 	}
 
 	ConsumeMarkedTransitionEvents();
diff --git a/Engine/Source/Runtime/Engine/Private/Animation/AnimSync.cpp b/Engine/Source/Runtime/Engine/Private/Animation/AnimSync.cpp
--- a/Engine/Source/Runtime/Engine/Private/Animation/AnimSync.cpp
+++ b/Engine/Source/Runtime/Engine/Private/Animation/AnimSync.cpp
@@ -464,8 +464,8 @@ FAnimTickRecord& FAnimSync::CreateUninitializedTickRecord(FAnimGroupInstance*& O
 	}
 
 	// Create the record
-	FAnimTickRecord* TickRecord = new ((OutSyncGroupPtr != nullptr) ? OutSyncGroupPtr->ActivePlayers : UngroupedActivePlayerArrays[GetSyncGroupWriteIndex()]) FAnimTickRecord();
-	return *TickRecord;
+	FAnimTickRecord& TickRecord = ((OutSyncGroupPtr != nullptr) ? OutSyncGroupPtr->ActivePlayers : UngroupedActivePlayerArrays[GetSyncGroupWriteIndex()]).AddDefaulted_GetRef();
+	return TickRecord;
 }
 
 FAnimTickRecord& FAnimSync::CreateUninitializedTickRecordInScope(FAnimInstanceProxy& InProxy, FAnimGroupInstance*& OutSyncGroupPtr, FName GroupName, EAnimSyncGroupScope Scope)
diff --git a/Engine/Source/Runtime/Engine/Private/BatchedElements.cpp b/Engine/Source/Runtime/Engine/Private/BatchedElements.cpp
--- a/Engine/Source/Runtime/Engine/Private/BatchedElements.cpp
+++ b/Engine/Source/Runtime/Engine/Private/BatchedElements.cpp
@@ -144,29 +144,29 @@ void FBatchedElements::AddLine(const FVector& Start, const FVector& End, const F
 	{
 		if (DepthBias == 0.0f)
 		{
-			new(LineVertices) FSimpleElementVertex(Start, FVector2D::ZeroVector, OpaqueColor, HitProxyId);
-			new(LineVertices) FSimpleElementVertex(End, FVector2D::ZeroVector, OpaqueColor, HitProxyId);
+			LineVertices.Emplace(Start, FVector2D::ZeroVector, OpaqueColor, HitProxyId);
+			LineVertices.Emplace(End, FVector2D::ZeroVector, OpaqueColor, HitProxyId);
 		}
 		else
 		{
 			// Draw degenerate triangles in wireframe mode to support depth bias (d3d11 and opengl3 don't support depth bias on line primitives, but do on wireframes)
-			FBatchedWireTris* WireTri = new(WireTris) FBatchedWireTris();
-			WireTri->DepthBias = DepthBias;
-			new(WireTriVerts) FSimpleElementVertex(Start, FVector2D::ZeroVector, OpaqueColor, HitProxyId);
-			new(WireTriVerts) FSimpleElementVertex(End, FVector2D::ZeroVector, OpaqueColor, HitProxyId);
-			new(WireTriVerts) FSimpleElementVertex(End, FVector2D::ZeroVector, OpaqueColor, HitProxyId);
+			FBatchedWireTris& WireTri = WireTris.AddDefaulted_GetRef();
+			WireTri.DepthBias = DepthBias;
+			WireTriVerts.Emplace(Start, FVector2D::ZeroVector, OpaqueColor, HitProxyId);
+			WireTriVerts.Emplace(End, FVector2D::ZeroVector, OpaqueColor, HitProxyId);
+			WireTriVerts.Emplace(End, FVector2D::ZeroVector, OpaqueColor, HitProxyId);
 		}
 	}
 	else
 	{
-		FBatchedThickLines* ThickLine = new(ThickLines) FBatchedThickLines;
-		ThickLine->Start = Start;
-		ThickLine->End = End;
-		ThickLine->Thickness = Thickness;
-		ThickLine->Color = OpaqueColor;
-		ThickLine->HitProxyColor = HitProxyId.GetColor();
-		ThickLine->DepthBias = DepthBias;
-		ThickLine->bScreenSpace = bScreenSpace;
+		FBatchedThickLines& ThickLine = ThickLines.AddDefaulted_GetRef();
+		ThickLine.Start = Start;
+		ThickLine.End = End;
+		ThickLine.Thickness = Thickness;
+		ThickLine.Color = OpaqueColor;
+		ThickLine.HitProxyColor = HitProxyId.GetColor();
+		ThickLine.DepthBias = DepthBias;
+		ThickLine.bScreenSpace = bScreenSpace;
 	}
 }
 
@@ -176,29 +176,29 @@ void FBatchedElements::AddTranslucentLine(const FVector& Start, const FVector& E
 	{
 		if (DepthBias == 0.0f)
 		{
-			new(LineVertices) FSimpleElementVertex(Start, FVector2D::ZeroVector, Color, HitProxyId);
-			new(LineVertices) FSimpleElementVertex(End, FVector2D::ZeroVector, Color, HitProxyId);
+			LineVertices.Emplace(Start, FVector2D::ZeroVector, Color, HitProxyId);
+			LineVertices.Emplace(End, FVector2D::ZeroVector, Color, HitProxyId);
 		}
 		else
 		{
 			// Draw degenerate triangles in wireframe mode to support depth bias (d3d11 and opengl3 don't support depth bias on line primitives, but do on wireframes)
-			FBatchedWireTris* WireTri = new(WireTris) FBatchedWireTris();
-			WireTri->DepthBias = DepthBias;
-			new(WireTriVerts) FSimpleElementVertex(Start, FVector2D::ZeroVector, Color, HitProxyId);
-			new(WireTriVerts) FSimpleElementVertex(End, FVector2D::ZeroVector, Color, HitProxyId);
-			new(WireTriVerts) FSimpleElementVertex(End, FVector2D::ZeroVector, Color, HitProxyId);
+			FBatchedWireTris& WireTri = WireTris.AddDefaulted_GetRef();
+			WireTri.DepthBias = DepthBias;
+			WireTriVerts.Emplace(Start, FVector2D::ZeroVector, Color, HitProxyId);
+			WireTriVerts.Emplace(End, FVector2D::ZeroVector, Color, HitProxyId);
+			WireTriVerts.Emplace(End, FVector2D::ZeroVector, Color, HitProxyId);
 		}
 	}
 	else
 	{
-		FBatchedThickLines* ThickLine = new(ThickLines) FBatchedThickLines;
-		ThickLine->Start = Start;
-		ThickLine->End = End;
-		ThickLine->Thickness = Thickness;
-		ThickLine->Color = Color;
-		ThickLine->HitProxyColor = HitProxyId.GetColor();
-		ThickLine->DepthBias = DepthBias;
-		ThickLine->bScreenSpace = bScreenSpace;
+		FBatchedThickLines& ThickLine = ThickLines.AddDefaulted_GetRef();
+		ThickLine.Start = Start;
+		ThickLine.End = End;
+		ThickLine.Thickness = Thickness;
+		ThickLine.Color = Color;
+		ThickLine.HitProxyColor = HitProxyId.GetColor();
+		ThickLine.DepthBias = DepthBias;
+		ThickLine.bScreenSpace = bScreenSpace;
 
 	}
 }
@@ -209,24 +209,24 @@ void FBatchedElements::AddPoint(const FVector& Position,float Size,const FLinear
 	FLinearColor OpaqueColor(Color);
 	OpaqueColor.A = 1;
 
-	FBatchedPoint* Point = new(Points) FBatchedPoint;
-	Point->Position = Position;
-	Point->Size = Size;
-	Point->Color = OpaqueColor.ToFColor(true);
-	Point->HitProxyColor = HitProxyId.GetColor();
+	FBatchedPoint& Point = Points.AddDefaulted_GetRef();
+	Point.Position = Position;
+	Point.Size = Size;
+	Point.Color = OpaqueColor.ToFColor(true);
+	Point.HitProxyColor = HitProxyId.GetColor();
 }
 
 int32 FBatchedElements::AddVertex(const FVector4& InPosition, const FVector2D& InTextureCoordinate, const FLinearColor& InColor, FHitProxyId HitProxyId)
 {
 	int32 VertexIndex = MeshVertices.Num();
-	new(MeshVertices) FSimpleElementVertex(InPosition, InTextureCoordinate, InColor, HitProxyId);
+	MeshVertices.Emplace(InPosition, InTextureCoordinate, InColor, HitProxyId);
 	return VertexIndex;
 }
 
 int32 FBatchedElements::AddVertexf(const FVector4f& InPosition,const FVector2f& InTextureCoordinate,const FLinearColor& InColor,FHitProxyId HitProxyId)
 {
 	int32 VertexIndex = MeshVertices.Num();
-	new(MeshVertices) FSimpleElementVertex(InPosition,InTextureCoordinate,InColor,HitProxyId);
+	MeshVertices.Emplace(InPosition,InTextureCoordinate,InColor,HitProxyId);
 	return VertexIndex;
 }
 
@@ -311,7 +311,7 @@ void FBatchedElements::AddTriangleExtensive(int32 V0,int32 V1,int32 V2,FBatchedE
 		else
 		{
 			// Create a new mesh element for the texture if this is the first triangle encountered using it.
-			MeshElement = new(MeshElements) FBatchedMeshElement;
+			MeshElement = &MeshElements.AddDefaulted_GetRef();
 			MeshElement->Texture = Texture;
 			MeshElement->BatchedElementParameters = BatchedElementParameters;
 			MeshElement->BlendMode = BlendMode;
@@ -434,19 +434,19 @@ void FBatchedElements::AddSprite(
 {
 	check(Texture);
 
-	FBatchedSprite* Sprite = new(Sprites) FBatchedSprite;
-	Sprite->Position = Position;
-	Sprite->SizeX = SizeX;
-	Sprite->SizeY = SizeY;
-	Sprite->Texture = Texture;
-	Sprite->Color = Color;
-	Sprite->HitProxyColor = HitProxyId.GetColor();
-	Sprite->U = U;
-	Sprite->UL = UL == 0.f ? Texture->GetSizeX() : UL;
-	Sprite->V = V;
-	Sprite->VL = VL == 0.f ? Texture->GetSizeY() : VL;
-	Sprite->OpacityMaskRefVal = OpacityMaskRefVal;
-	Sprite->BlendMode = BlendMode;
+	FBatchedSprite& Sprite = Sprites.AddDefaulted_GetRef();
+	Sprite.Position = Position;
+	Sprite.SizeX = SizeX;
+	Sprite.SizeY = SizeY;
+	Sprite.Texture = Texture;
+	Sprite.Color = Color;
+	Sprite.HitProxyColor = HitProxyId.GetColor();
+	Sprite.U = U;
+	Sprite.UL = UL == 0.f ? Texture->GetSizeX() : UL;
+	Sprite.V = V;
+	Sprite.VL = VL == 0.f ? Texture->GetSizeY() : VL;
+	Sprite.OpacityMaskRefVal = OpacityMaskRefVal;
+	Sprite.BlendMode = BlendMode;
 }
 
 /** Translates a ESimpleElementBlendMode into a RHI state change for rendering a mesh with the blend mode normally. */
diff --git a/Engine/Source/Runtime/Engine/Private/Components/SplineMeshComponent.cpp b/Engine/Source/Runtime/Engine/Private/Components/SplineMeshComponent.cpp
--- a/Engine/Source/Runtime/Engine/Private/Components/SplineMeshComponent.cpp
+++ b/Engine/Source/Runtime/Engine/Private/Components/SplineMeshComponent.cpp
@@ -1282,7 +1282,7 @@ void USplineMeshComponent::RecreateCollision()
 			// Convert boxes to convex hulls to better respect distortion
 			for (FKBoxElem& BoxElem : BodySetup->AggGeom.BoxElems)
 			{
-				FKConvexElem& ConvexElem = *new(BodySetup->AggGeom.ConvexElems) FKConvexElem();
+				FKConvexElem& ConvexElem = BodySetup->AggGeom.ConvexElems.AddDefaulted_GetRef();
 
 				const FVector Radii = FVector(BoxElem.X / 2, BoxElem.Y / 2, BoxElem.Z / 2).ComponentMax(FVector(1.0f));
 				const FTransform ElementTM = BoxElem.GetTransform();
diff --git a/Engine/Source/Runtime/Engine/Private/Components/SpotLightComponent.cpp b/Engine/Source/Runtime/Engine/Private/Components/SpotLightComponent.cpp
--- a/Engine/Source/Runtime/Engine/Private/Components/SpotLightComponent.cpp
+++ b/Engine/Source/Runtime/Engine/Private/Components/SpotLightComponent.cpp
@@ -115,7 +115,7 @@ public:
 	 */
 	virtual bool GetWholeSceneProjectedShadowInitializer(const FSceneViewFamily& ViewFamily, TArray<FWholeSceneProjectedShadowInitializer, TInlineAllocator<6> >& OutInitializers) const override
 	{
-		FWholeSceneProjectedShadowInitializer& OutInitializer = *new(OutInitializers) FWholeSceneProjectedShadowInitializer;
+		FWholeSceneProjectedShadowInitializer& OutInitializer = OutInitializers.AddDefaulted_GetRef();
 		OutInitializer.PreShadowTranslation = -GetLightToWorld().GetOrigin();
 		OutInitializer.WorldToLight = GetWorldToLight().RemoveTranslation();
 		OutInitializer.Scales = FVector2D(InvTanOuterCone,InvTanOuterCone);
diff --git a/Engine/Source/Runtime/Engine/Private/Streaming/TextureStreamingBuild.cpp b/Engine/Source/Runtime/Engine/Private/Streaming/TextureStreamingBuild.cpp
--- a/Engine/Source/Runtime/Engine/Private/Streaming/TextureStreamingBuild.cpp
+++ b/Engine/Source/Runtime/Engine/Private/Streaming/TextureStreamingBuild.cpp
@@ -687,7 +687,7 @@ void FStreamingTextureLevelContext::ProcessMaterial(const FBoxSphereBounds& Comp
 		{
 			if (*BuildDataIndex != INDEX_NONE)
 			{
-				FStreamingRenderAssetPrimitiveInfo& Info = *new(OutStreamingTextures) FStreamingRenderAssetPrimitiveInfo();
+				FStreamingRenderAssetPrimitiveInfo& Info = OutStreamingTextures.AddDefaulted_GetRef();
 				const FStreamingTextureBuildInfo& BuildInfo = (*ComponentBuildData)[*BuildDataIndex];
 
 				Info.RenderAsset = Texture;
@@ -712,7 +712,7 @@ void FStreamingTextureLevelContext::ProcessMaterial(const FBoxSphereBounds& Comp
 
 			if (TextureDensity)
 			{
-				FStreamingRenderAssetPrimitiveInfo& Info = *new(OutStreamingTextures) FStreamingRenderAssetPrimitiveInfo();
+				FStreamingRenderAssetPrimitiveInfo& Info = OutStreamingTextures.AddDefaulted_GetRef();
 
 				Info.RenderAsset = Texture;
 				Info.TexelFactor = TextureDensity * ComponentScaling;
diff --git a/Engine/Source/Runtime/Engine/Private/TextureDerivedDataTask.cpp b/Engine/Source/Runtime/Engine/Private/TextureDerivedDataTask.cpp
--- a/Engine/Source/Runtime/Engine/Private/TextureDerivedDataTask.cpp
+++ b/Engine/Source/Runtime/Engine/Private/TextureDerivedDataTask.cpp
@@ -188,11 +188,11 @@ void FTextureSourceData::Init(UTexture& InTexture, TextureMipGenSettings InMipGe
 	Layers.Reserve(NumLayers);
 	for (int LayerIndex = 0; LayerIndex < NumLayers; ++LayerIndex)
 	{
-		FTextureSourceLayerData* LayerData = new(Layers) FTextureSourceLayerData();
+		FTextureSourceLayerData& LayerData = Layers.AddDefaulted_GetRef();
 
-		LayerData->ImageFormat = FImageCoreUtils::ConvertToRawImageFormat( InTexture.Source.GetFormat(LayerIndex) );
+		LayerData.ImageFormat = FImageCoreUtils::ConvertToRawImageFormat( InTexture.Source.GetFormat(LayerIndex) );
 
-		LayerData->SourceGammaSpace = InTexture.Source.GetGammaSpace(LayerIndex);
+		LayerData.SourceGammaSpace = InTexture.Source.GetGammaSpace(LayerIndex);
 	}
 
 	Blocks.Reserve(NumBlocks);
@@ -203,25 +203,25 @@ void FTextureSourceData::Init(UTexture& InTexture, TextureMipGenSettings InMipGe
 
 		if (SourceBlock.NumMips > 0 && SourceBlock.NumSlices > 0)
 		{
-			FTextureSourceBlockData* BlockData = new(Blocks) FTextureSourceBlockData();
-			BlockData->BlockX = SourceBlock.BlockX;
-			BlockData->BlockY = SourceBlock.BlockY;
-			BlockData->SizeX = SourceBlock.SizeX;
-			BlockData->SizeY = SourceBlock.SizeY;
-			BlockData->NumMips = SourceBlock.NumMips;
-			BlockData->NumSlices = SourceBlock.NumSlices;
+			FTextureSourceBlockData& BlockData = Blocks.AddDefaulted_GetRef();
+			BlockData.BlockX = SourceBlock.BlockX;
+			BlockData.BlockY = SourceBlock.BlockY;
+			BlockData.SizeX = SourceBlock.SizeX;
+			BlockData.SizeY = SourceBlock.SizeY;
+			BlockData.NumMips = SourceBlock.NumMips;
+			BlockData.NumSlices = SourceBlock.NumSlices;
 
 			if (InMipGenSettings != TMGS_LeaveExistingMips)
 			{
-				BlockData->NumMips = 1;
+				BlockData.NumMips = 1;
 			}
 
 			if (!bInCubeMap && !bInTextureArray && !bInVolumeTexture)
 			{
-				BlockData->NumSlices = 1;
+				BlockData.NumSlices = 1;
 			}
 
-			BlockData->MipsPerLayer.SetNum(NumLayers);
+			BlockData.MipsPerLayer.SetNum(NumLayers);
 
 			SizeInBlocksX = FMath::Max(SizeInBlocksX, SourceBlock.BlockX + 1);
 			SizeInBlocksY = FMath::Max(SizeInBlocksY, SourceBlock.BlockY + 1);
@@ -291,13 +291,13 @@ void FTextureSourceData::GetSourceMips(FTextureSource& Source, IImageWrapperModu
 						int32 MipSizeZ = SourceBlock.NumSlices;
 						for (int32 MipIndex = 0; MipIndex < BlockData.NumMips; ++MipIndex)
 						{
-							FImage* SourceMip = new(BlockData.MipsPerLayer[LayerIndex]) FImage(
+							FImage& SourceMip = BlockData.MipsPerLayer[LayerIndex].Emplace_GetRef(
 								MipSizeX, MipSizeY, MipSizeZ,
 								LayerData.ImageFormat,
 								LayerData.SourceGammaSpace
 							);
 
-							if (!ScopedMipData.GetMipData(SourceMip->RawData, BlockIndex, LayerIndex, MipIndex))
+							if (!ScopedMipData.GetMipData(SourceMip.RawData, BlockIndex, LayerIndex, MipIndex))
 							{
 								UE_LOG(LogTexture, Warning, TEXT("Cannot retrieve source data for mip %d of %s"), MipIndex, *TextureFullName);
 								ReleaseMemory();
diff --git a/Engine/Source/Runtime/Engine/Private/TextureLODSettings.cpp b/Engine/Source/Runtime/Engine/Private/TextureLODSettings.cpp
--- a/Engine/Source/Runtime/Engine/Private/TextureLODSettings.cpp
+++ b/Engine/Source/Runtime/Engine/Private/TextureLODSettings.cpp
@@ -118,7 +118,7 @@ TArray<FString> UTextureLODSettings::GetTextureGroupNames()
 	// TEXTUREGROUP_MAX is not actually Max, it's count
 	TextureGroupNames.Reserve(TEXTUREGROUP_MAX);
 
-#define GROUPNAMES(g) new(TextureGroupNames) FString(TEXT(#g));
+#define GROUPNAMES(g) TextureGroupNames.Emplace(TEXT(#g));
 	FOREACH_ENUM_TEXTUREGROUP(GROUPNAMES)
 #undef GROUPNAMES
 
diff --git a/Engine/Source/Runtime/Engine/Private/TickTaskManager.cpp b/Engine/Source/Runtime/Engine/Private/TickTaskManager.cpp
--- a/Engine/Source/Runtime/Engine/Private/TickTaskManager.cpp
+++ b/Engine/Source/Runtime/Engine/Private/TickTaskManager.cpp
@@ -455,7 +455,7 @@ public:
 		{
 			TickTasks[StartTickGroup][EndTickGroup].Add(Task);
 		}
-		new (TickCompletionEvents[EndTickGroup]) FGraphEventRef(Task->GetCompletionEvent());
+		TickCompletionEvents[EndTickGroup].Add(Task->GetCompletionEvent());
 	}
 	/** Add a completion handle to a tick group, parallel version **/
 	FORCEINLINE void AddTickTaskCompletionParallel(ETickingGroup StartTickGroup, ETickingGroup EndTickGroup, TGraphTask<FTickFunctionTask>* Task, bool bHiPri)
@@ -2244,7 +2244,7 @@ static void AddTestTickFunctions(const TArray<FString>& Args, UWorld* InWorld)
 	TestTickFunctions.Reserve(NumTestTickFunctions);
 	for (int32 Index = 0; Index < NumTestTickFunctions; Index++)
 	{
-		(new (TestTickFunctions) FTestTickFunction())->RegisterTickFunction(Level);
+		TestTickFunctions.AddDefaulted_GetRef().RegisterTickFunction(Level);
 	}
 }
 
diff --git a/Engine/Source/Runtime/Engine/Private/URL.cpp b/Engine/Source/Runtime/Engine/Private/URL.cpp
--- a/Engine/Source/Runtime/Engine/Private/URL.cpp
+++ b/Engine/Source/Runtime/Engine/Private/URL.cpp
@@ -170,7 +170,7 @@ FURL::FURL( FURL* Base, const TCHAR* TextURL, ETravelType Type )
 		check(Base);
 		for( int32 i=0; i<Base->Op.Num(); i++ )
 		{
-			new(Op)FString(Base->Op[i]);
+			Op.Add(Base->Op[i]);
 		}
 	}
 
@@ -657,7 +657,7 @@ void FURL::AddOption( const TCHAR* Str )
 
 	if (i == Op.Num())
 	{
-		new(Op) FString(Str);
+		Op.Emplace(Str);
 	}
 	else
 	{
diff --git a/Engine/Source/Runtime/Engine/Private/UserInterface/Canvas.cpp b/Engine/Source/Runtime/Engine/Private/UserInterface/Canvas.cpp
--- a/Engine/Source/Runtime/Engine/Private/UserInterface/Canvas.cpp
+++ b/Engine/Source/Runtime/Engine/Private/UserInterface/Canvas.cpp
@@ -316,7 +316,7 @@ void FCanvas::Construct()
 
 	const FIntPoint RenderTargetSizeXY = RenderTarget ? RenderTarget->GetSizeXY() : FIntPoint(1, 1);
 
-	new(TransformStack) FTransformEntry( 
+	TransformStack.Emplace( 
 		FMatrix( FScaleMatrix(GetDPIScale()) * CalcBaseTransform2D(RenderTargetSizeXY.X, RenderTargetSizeXY.Y) )
 		);
 
@@ -339,7 +339,7 @@ void FCanvas::SetBaseTransform(const FMatrix& Transform)
 	}
 	else
 	{
-		new(TransformStack) FTransformEntry(Transform);
+		TransformStack.Emplace(Transform);
 	}
 }
 
@@ -600,7 +600,7 @@ FCanvas::FCanvasSortElement& FCanvas::GetSortElement(int32 DepthSortKey)
 	// if it doesn't exist then add a new entry (no duplicates allowed)
 	else
 	{
-		new(SortedElements) FCanvasSortElement(DepthSortKey);
+		SortedElements.Emplace(DepthSortKey);
 		ElementIdx = SortedElements.Num()-1;
 		// keep track of newly added array index for later lookup
 		SortedElementLookupMap.Add( DepthSortKey, ElementIdx );
diff --git a/Engine/Source/Runtime/RenderCore/Private/DumpGPU.cpp b/Engine/Source/Runtime/RenderCore/Private/DumpGPU.cpp
--- a/Engine/Source/Runtime/RenderCore/Private/DumpGPU.cpp
+++ b/Engine/Source/Runtime/RenderCore/Private/DumpGPU.cpp
@@ -2037,11 +2037,11 @@ FString FRDGBuilder::BeginResourceDump(const TCHAR* Cmd)
 		{
 			if (**NextToken == TCHAR('-'))
 			{
-				new(Switches) FString(NextToken.Mid(1));
+				Switches.Add(NextToken.Mid(1));
 			}
 			else
 			{
-				new(Tokens) FString(NextToken);
+				Tokens.Add(NextToken);
 			}
 		}
 
diff --git a/Engine/Source/Runtime/RenderCore/Private/ProfilingDebugging/RealtimeGPUProfiler.cpp b/Engine/Source/Runtime/RenderCore/Private/ProfilingDebugging/RealtimeGPUProfiler.cpp
--- a/Engine/Source/Runtime/RenderCore/Private/ProfilingDebugging/RealtimeGPUProfiler.cpp
+++ b/Engine/Source/Runtime/RenderCore/Private/ProfilingDebugging/RealtimeGPUProfiler.cpp
@@ -486,7 +486,7 @@ public:
 
 			if (MaxNumQueries < 0 || QueryCount < (uint32)MaxNumQueries)
 			{
-				new (GpuProfilerEvents) FRealtimeGPUProfilerEvent(*RenderQueryPool);
+				GpuProfilerEvents.Emplace(*RenderQueryPool);
 				QueryCount += FRealtimeGPUProfilerEvent::GetNumRHIQueriesPerEvent();
 			}
 			else
diff --git a/Engine/Source/Runtime/Slate/Private/Framework/Layout/InertialScrollManager.cpp b/Engine/Source/Runtime/Slate/Private/Framework/Layout/InertialScrollManager.cpp
--- a/Engine/Source/Runtime/Slate/Private/Framework/Layout/InertialScrollManager.cpp
+++ b/Engine/Source/Runtime/Slate/Private/Framework/Layout/InertialScrollManager.cpp
@@ -13,7 +13,7 @@ FInertialScrollManager::FInertialScrollManager(double InSampleTimeout)
 
 void FInertialScrollManager::AddScrollSample(float Delta, double CurrentTime)
 {
-	new( ScrollSamples ) FScrollSample(Delta, CurrentTime);
+	ScrollSamples.Emplace(Delta, CurrentTime);
 
 	float Total = 0;
 	double OldestTime = 0;
diff --git a/Engine/Source/Runtime/Slate/Private/Widgets/Input/SVirtualJoystick.cpp b/Engine/Source/Runtime/Slate/Private/Widgets/Input/SVirtualJoystick.cpp
--- a/Engine/Source/Runtime/Slate/Private/Widgets/Input/SVirtualJoystick.cpp
+++ b/Engine/Source/Runtime/Slate/Private/Widgets/Input/SVirtualJoystick.cpp
@@ -542,8 +542,8 @@ void SVirtualJoystick::SetJoystickVisibility(const bool bInVisible, const bool b
 
 void SVirtualJoystick::AddControl(const FControlInfo& Control)
 {
-	FControlData* ControlData = new (Controls) FControlData;
-	ControlData->Info = Control;
+	FControlData& ControlData = Controls.AddDefaulted_GetRef();
+	ControlData.Info = Control;
 }
 
 void SVirtualJoystick::ClearControls()
