[Backout] - CL26284902
Original CL Desc
-----------------------------------------------------------------
[Backout] - CL26281988
Original CL Desc
-----------------------------------------------------------------
Unshelved from pending changelist '26281904':

Quartz Subscriber Optimization

During the creation of this class, it would subscribe itself to the Quartz Subsystem, so that it can receive regular tick updates. The unsubscribe was happening exclusively in the destructor. Therefore for managed objects, they weren't being unsubscribed until the object was fully deleted by the garbage collector.

I wrapped up the subscription code inside a new QuartzSubscriptionToken. This class adds a few more checks for sanity, and prevents things from being subscribed multiple times, or unsubscribed when its not even subscribed. The Actor/Component objects that derive from this class now unsubscribe in the EndPlay function call - they subscribe in BeginPlay. This prevents the constant build up in subscribed objects until the garbage collector runs.

diff --git a/Engine/Source/Runtime/Engine/Classes/Sound/QuartzSubscription.cpp b/Engine/Source/Runtime/Engine/Classes/Sound/QuartzSubscription.cpp
--- a/Engine/Source/Runtime/Engine/Classes/Sound/QuartzSubscription.cpp
+++ b/Engine/Source/Runtime/Engine/Classes/Sound/QuartzSubscription.cpp
@@ -22,13 +22,13 @@ namespace Audio
 	}
 } // namespace Audio
 
+FQuartzTickableObject::FQuartzTickableObject()
+{}
+
 PRAGMA_DISABLE_DEPRECATION_WARNINGS
 FQuartzTickableObject::~FQuartzTickableObject()
 {
-	if(const TSharedPtr<FQuartzTickableObjectsManager> ObjManagerPtr = TickableObjectManagerPtr.Pin())
-	{
-		ObjManagerPtr->UnsubscribeFromQuartzTick(this);
-	}
+	QuartzUnsubscribe();
 }
 PRAGMA_ENABLE_DEPRECATION_WARNINGS
 
@@ -48,19 +48,16 @@ FQuartzTickableObject* FQuartzTickableObject::Init(UWorld* InWorldPtr)
 	PRAGMA_ENABLE_DEPRECATION_WARNINGS
 
 	UQuartzSubsystem* QuartzSubsystemPtr = UQuartzSubsystem::Get(InWorldPtr);
-	if(ensure(QuartzSubsystemPtr))
-	{
-		TickableObjectManagerPtr = QuartzSubsystemPtr->GetTickableObjectManager();
-
-		if(TSharedPtr<FQuartzTickableObjectsManager> ObjManagerPtr = TickableObjectManagerPtr.Pin())
-		{
-			ObjManagerPtr->SubscribeToQuartzTick(this);
-		}
-	}
+	QuartzSubscriptionToken.Subscribe(this, QuartzSubsystemPtr);
 
 	return this;
 }
 
+void FQuartzTickableObject::QuartzUnsubscribe()
+{
+	QuartzSubscriptionToken.Unsubscribe();
+}
+
 int32 FQuartzTickableObject::AddCommandDelegate(const FOnQuartzCommandEventBP& InDelegate)
 {
 	const int32 Num = QuantizedCommandDelegates.Num();
@@ -90,8 +87,8 @@ UQuartzSubsystem* FQuartzTickableObject::GetQuartzSubsystem() const
 void FQuartzTickableObject::ExecCommand(const Audio::FQuartzQuantizedCommandDelegateData& Data)
 {
 	checkSlow(Data.DelegateSubType < EQuartzCommandDelegateSubType::Count);
-
-	if(const TSharedPtr<FQuartzTickableObjectsManager> ObjManagerPtr = TickableObjectManagerPtr.Pin())
+	
+	if(const TSharedPtr<FQuartzTickableObjectsManager> ObjManagerPtr = QuartzSubscriptionToken.GetTickableObjectManager())
 	{
 		ObjManagerPtr->PushLatencyTrackerResult(Data.RequestRecieved());
 	}
@@ -132,7 +129,7 @@ void FQuartzTickableObject::ExecCommand(const Audio::FQuartzQuantizedCommandDele
 
 void FQuartzTickableObject::ExecCommand(const Audio::FQuartzMetronomeDelegateData& Data)
 {
-	if(const TSharedPtr<FQuartzTickableObjectsManager> ObjManagerPtr = TickableObjectManagerPtr.Pin())
+	if(const TSharedPtr<FQuartzTickableObjectsManager> ObjManagerPtr = QuartzSubscriptionToken.GetTickableObjectManager())
 	{
 		ObjManagerPtr->PushLatencyTrackerResult(Data.RequestRecieved());
 	}
diff --git a/Engine/Source/Runtime/Engine/Classes/Sound/QuartzSubscription.h b/Engine/Source/Runtime/Engine/Classes/Sound/QuartzSubscription.h
--- a/Engine/Source/Runtime/Engine/Classes/Sound/QuartzSubscription.h
+++ b/Engine/Source/Runtime/Engine/Classes/Sound/QuartzSubscription.h
@@ -2,6 +2,7 @@
 
 #pragma once
 
+#include "QuartzSubscriptionToken.h"
 #include "Sound/QuartzQuantizationUtilities.h"
 #include "Containers/ConsumeAllMpmcQueue.h"
 
@@ -126,7 +127,7 @@ class FQuartzTickableObject
 {
 public:
 	// ctor
-	FQuartzTickableObject() {}
+	ENGINE_API FQuartzTickableObject();
 
 	// explicitly defaulted ctors (to disable deprecation warnings in compiler-generated functions)
     PRAGMA_DISABLE_DEPRECATION_WARNINGS
@@ -149,8 +150,7 @@ public:
 
 	ENGINE_API void AddMetronomeBpDelegate(EQuartzCommandQuantization InQuantizationBoundary, const FOnQuartzMetronomeEventBP& OnQuantizationEvent);
 
-
-	bool IsInitialized() const { return TickableObjectManagerPtr.IsValid(); }
+	bool IsInitialized() const { return QuartzSubscriptionToken.IsSubscribed(); }
 
 	ENGINE_API Audio::FQuartzGameThreadSubscriber GetQuartzSubscriber();
 
@@ -178,6 +178,8 @@ protected:
 	ENGINE_API void SetNotificationAnticipationAmountMilliseconds(const double Milliseconds);
 	ENGINE_API void SetNotificationAnticipationAmountMusicalDuration(const EQuartzCommandQuantization Duration,  const double Multiplier);
 
+	ENGINE_API void QuartzUnsubscribe();
+
 private:
 	struct FMetronomeDelegateGameThreadData
 	{
@@ -203,7 +205,7 @@ public:
 
 private:
 	Audio::FQuartzOffset NotificationOffset;
-	TWeakPtr<FQuartzTickableObjectsManager> TickableObjectManagerPtr;
 
+	FQuartzSubscriptionToken QuartzSubscriptionToken;
 }; // class FQuartzTickableObject
 
diff --git a/Engine/Source/Runtime/Engine/Classes/Sound/QuartzSubscriptionToken.cpp b/Engine/Source/Runtime/Engine/Classes/Sound/QuartzSubscriptionToken.cpp
new file mode 100644
--- /dev/null
+++ b/Engine/Source/Runtime/Engine/Classes/Sound/QuartzSubscriptionToken.cpp
@@ -0,0 +1,76 @@
+﻿// Copyright Epic Games, Inc. All Rights Reserved.
+
+#include "QuartzSubscriptionToken.h"
+
+#include "Quartz/QuartzSubsystem.h"
+#include "Sound/QuartzQuantizationUtilities.h"
+
+FQuartzSubscriptionToken::FQuartzSubscriptionToken()
+	: SubscribingObject(nullptr)
+	{}
+
+FQuartzSubscriptionToken::~FQuartzSubscriptionToken()
+{
+	Unsubscribe();
+}
+
+void FQuartzSubscriptionToken::Subscribe(FQuartzTickableObject* Subscriber, UQuartzSubsystem* QuartzSubsystem)
+{
+	if(IsSubscribed())
+	{
+		return;
+	}
+
+	if(ensure(Subscriber))
+	{
+		SubscribingObject = Subscriber;
+	}
+	else
+	{
+		UE_LOG(LogAudioQuartz, Error, TEXT("Invalid subscribing object."));
+		return;
+	}
+	
+	if(ensure(QuartzSubsystem))
+	{
+		TickableObjectManagerPtr = QuartzSubsystem->GetTickableObjectManager();
+	}
+	else
+	{
+		UE_LOG(LogAudioQuartz, Error, TEXT("Unable to obtain Tickable Object Manager."));
+		return;
+	}
+	
+	TSharedPtr<FQuartzTickableObjectsManager> ObjManagerPtr = GetTickableObjectManager();
+	if(ObjManagerPtr.IsValid())
+	{
+		ObjManagerPtr->SubscribeToQuartzTick(SubscribingObject);
+	}
+	else
+	{
+		UE_LOG(LogAudioQuartz, Warning, TEXT("Unable to subscribe to Quartz tick. Tickable Object Manager pointer is not valid."));
+	}
+}
+
+void FQuartzSubscriptionToken::Unsubscribe()
+{
+	const TSharedPtr<FQuartzTickableObjectsManager> ObjManagerPtr = GetTickableObjectManager();
+	if(ObjManagerPtr.IsValid())
+	{
+		ObjManagerPtr->UnsubscribeFromQuartzTick(SubscribingObject);
+	}
+
+	SubscribingObject = nullptr;
+	TickableObjectManagerPtr = nullptr;
+}
+
+TSharedPtr<FQuartzTickableObjectsManager> FQuartzSubscriptionToken::GetTickableObjectManager()
+{
+	TSharedPtr<FQuartzTickableObjectsManager> ObjManagerPtr = TickableObjectManagerPtr.Pin();
+	return ObjManagerPtr;
+}
+
+bool FQuartzSubscriptionToken::IsSubscribed() const
+{
+	return TickableObjectManagerPtr.IsValid();
+}
\ No newline at end of file
diff --git a/Engine/Source/Runtime/Engine/Classes/Sound/QuartzSubscriptionToken.h b/Engine/Source/Runtime/Engine/Classes/Sound/QuartzSubscriptionToken.h
new file mode 100644
--- /dev/null
+++ b/Engine/Source/Runtime/Engine/Classes/Sound/QuartzSubscriptionToken.h
@@ -0,0 +1,24 @@
+﻿// Copyright Epic Games, Inc. All Rights Reserved.
+
+#pragma once
+#include "Templates/SharedPointer.h"
+
+class UQuartzSubsystem;
+struct FQuartzTickableObjectsManager;
+
+class FQuartzSubscriptionToken
+{
+public:
+	ENGINE_API FQuartzSubscriptionToken();
+	ENGINE_API ~FQuartzSubscriptionToken();
+	
+	ENGINE_API void Subscribe(FQuartzTickableObject* Subscriber, UQuartzSubsystem* QuartzSubsystem);
+	ENGINE_API void Unsubscribe();
+
+	ENGINE_API TSharedPtr<FQuartzTickableObjectsManager> GetTickableObjectManager();
+	ENGINE_API bool IsSubscribed() const;
+	
+private:
+	FQuartzTickableObject* SubscribingObject;
+	TWeakPtr<FQuartzTickableObjectsManager> TickableObjectManagerPtr;
+};
diff --git a/Engine/Source/Runtime/Engine/Private/Components/AudioComponent.cpp b/Engine/Source/Runtime/Engine/Private/Components/AudioComponent.cpp
--- a/Engine/Source/Runtime/Engine/Private/Components/AudioComponent.cpp
+++ b/Engine/Source/Runtime/Engine/Private/Components/AudioComponent.cpp
@@ -308,6 +308,7 @@ void UAudioComponent::OnUnregister()
 void UAudioComponent::EndPlay(const EEndPlayReason::Type EndPlayReason)
 {
 	Super::EndPlay(EndPlayReason);
+	QuartzUnsubscribe();
 }
 
 const UObject* UAudioComponent::AdditionalStatObject() const
