Compile error fix for Quartz system

diff --git a/Engine/Source/Runtime/Engine/Classes/Sound/QuartzSubscriptionToken.h b/Engine/Source/Runtime/Engine/Classes/Sound/QuartzSubscriptionToken.h
--- a/Engine/Source/Runtime/Engine/Classes/Sound/QuartzSubscriptionToken.h
+++ b/Engine/Source/Runtime/Engine/Classes/Sound/QuartzSubscriptionToken.h
@@ -4,6 +4,7 @@
 #include "Templates/SharedPointer.h"
 
 class UQuartzSubsystem;
+class FQuartzTickableObject;
 struct FQuartzTickableObjectsManager;
 
 class FQuartzSubscriptionToken
