Added a 'default' overlay material option to a skeletal mesh asset.
When assinged this overlay material will be used if component does not override it

diff --git a/Engine/Source/Runtime/Engine/Classes/Components/MeshComponent.h b/Engine/Source/Runtime/Engine/Classes/Components/MeshComponent.h
--- a/Engine/Source/Runtime/Engine/Classes/Components/MeshComponent.h
+++ b/Engine/Source/Runtime/Engine/Classes/Components/MeshComponent.h
@@ -66,6 +66,10 @@ public:
 	UFUNCTION(BlueprintCallable, Category="Rendering|Material")
 	ENGINE_API void SetOverlayMaterial(class UMaterialInterface* NewOverlayMaterial);
 
+	/** Get the overlay material used by this instance */
+	UFUNCTION(BlueprintCallable, Category="Rendering|Material")
+	float GetOverlayMaterialMaxDrawDistance() const;
+	
 	/** Change the overlay material max draw distance used by this instance */
 	UFUNCTION(BlueprintCallable, Category="Rendering|Material")
 	ENGINE_API void SetOverlayMaterialMaxDrawDistance(float InMaxDrawDistance);
@@ -173,6 +177,12 @@ public:
 		return (ParameterCache ? ParameterCache->ScalarParameterDefaultValue : 0.f);
 	}
 protected:
+	/** Get the default overlay material used by a mesh */
+	virtual UMaterialInterface* GetDefaultOverlayMaterial() const { return nullptr; };
+	
+	/** Get the default overlay material max draw distance */
+	virtual float GetDefaultOverlayMaterialMaxDrawDistance() const { return 0.f; };
+
 	/** Retrieves all the (scalar/vector-)parameters from within the used materials on the SkeletalMesh, and stores material index vs parameter names */
 	ENGINE_API void CacheMaterialParameterNameIndices();
 
diff --git a/Engine/Source/Runtime/Engine/Classes/Components/SkinnedMeshComponent.h b/Engine/Source/Runtime/Engine/Classes/Components/SkinnedMeshComponent.h
--- a/Engine/Source/Runtime/Engine/Classes/Components/SkinnedMeshComponent.h
+++ b/Engine/Source/Runtime/Engine/Classes/Components/SkinnedMeshComponent.h
@@ -1534,6 +1534,13 @@ protected:
 	UE_DEPRECATED(5.1, "This method has been deprecated. Please use RemoveFollowerPoseComponent instead.")
 	virtual void RemoveSlavePoseComponent(USkinnedMeshComponent* SkinnedMeshComponent) { RemoveFollowerPoseComponent(SkinnedMeshComponent); }
 
+protected:
+	/** Get the default overlay material used by a mesh */
+	virtual UMaterialInterface* GetDefaultOverlayMaterial() const override;
+
+	/** Get the default overlay material max draw distance */
+	virtual float GetDefaultOverlayMaterialMaxDrawDistance() const override;
+
 public:
 	/** 
 	 * Refresh Follower Components if exists
diff --git a/Engine/Source/Runtime/Engine/Classes/Engine/SkeletalMesh.h b/Engine/Source/Runtime/Engine/Classes/Engine/SkeletalMesh.h
--- a/Engine/Source/Runtime/Engine/Classes/Engine/SkeletalMesh.h
+++ b/Engine/Source/Runtime/Engine/Classes/Engine/SkeletalMesh.h
@@ -135,6 +135,8 @@ enum class ESkeletalMeshAsyncProperties : uint64
 	RayTracingMinLOD = 1llu << 54,
 	ClothLODBiasMode = 1llu << 55,
 	DefaultMeshDeformer = 1llu << 56,
+	OverlayMaterial = 1llu << 57,
+	OverlayMaterialMaxDrawDistance = 1llu << 58,
 	All = MAX_uint64
 };
 
@@ -2877,6 +2879,57 @@ protected:
 	UPROPERTY(EditAnywhere, BlueprintReadOnly, Category = "Deformer", Meta=(Filter="/Script/Engine.SkinnedMeshComponent"))
 	TObjectPtr<UMeshDeformer> DefaultMeshDeformer;
 
+	/** Default translucent material to blend on top of this mesh. Mesh will be rendered twice - once with a base material and once with overlay material */
+	UE_DEPRECATED(5.3, "This must be protected for async build, always use the accessors even internally.")
+	UPROPERTY(EditAnywhere, BlueprintReadOnly, Category=Rendering)
+	TObjectPtr<class UMaterialInterface> OverlayMaterial;
+
+	/** Default max draw distance for overlay material. A distance of 0 indicates that overlay will be culled using primitive max distance. */
+	UE_DEPRECATED(5.3, "This must be protected for async build, always use the accessors even internally.")
+	UPROPERTY(EditAnywhere, BlueprintReadOnly, Category=Rendering)
+	float OverlayMaterialMaxDrawDistance;
+
+public:
+	/** Get the default overlay material used by this mesh */
+	UFUNCTION(BlueprintCallable, Category="Rendering|Material")
+	virtual class UMaterialInterface* GetOverlayMaterial() const override
+	{
+		WaitUntilAsyncPropertyReleased(ESkeletalMeshAsyncProperties::OverlayMaterial, ESkinnedAssetAsyncPropertyLockType::ReadOnly);
+		PRAGMA_DISABLE_DEPRECATION_WARNINGS
+		return OverlayMaterial;
+		PRAGMA_ENABLE_DEPRECATION_WARNINGS
+	}
+
+	/** Change the default overlay material used by this mesh */
+	UFUNCTION(BlueprintCallable, Category="Rendering|Material")
+	void SetOverlayMaterial(class UMaterialInterface* NewOverlayMaterial)
+	{
+		WaitUntilAsyncPropertyReleased(ESkeletalMeshAsyncProperties::OverlayMaterial);
+		PRAGMA_DISABLE_DEPRECATION_WARNINGS
+		OverlayMaterial = NewOverlayMaterial;
+		PRAGMA_ENABLE_DEPRECATION_WARNINGS
+	}
+	
+	/** Get the default overlay material max draw distance used by this mesh */
+	UFUNCTION(BlueprintCallable, Category="Rendering|Material")
+	virtual float GetOverlayMaterialMaxDrawDistance() const override
+	{
+		WaitUntilAsyncPropertyReleased(ESkeletalMeshAsyncProperties::OverlayMaterialMaxDrawDistance, ESkinnedAssetAsyncPropertyLockType::ReadOnly);
+		PRAGMA_DISABLE_DEPRECATION_WARNINGS
+		return OverlayMaterialMaxDrawDistance;
+		PRAGMA_ENABLE_DEPRECATION_WARNINGS
+	}
+
+	/** Change the default overlay material max draw distance used by this mesh */
+	UFUNCTION(BlueprintCallable, Category="Rendering|Material")
+	void SetOverlayMaterialMaxDrawDistance(float InMaxDrawDistance)
+	{
+		WaitUntilAsyncPropertyReleased(ESkeletalMeshAsyncProperties::OverlayMaterialMaxDrawDistance);
+		PRAGMA_DISABLE_DEPRECATION_WARNINGS
+		OverlayMaterialMaxDrawDistance = InMaxDrawDistance;
+		PRAGMA_ENABLE_DEPRECATION_WARNINGS
+	}
+
 private:
 	/**
 	 * Initial step for the Post Load process - Can't be done in parallel. USkinnedAsset Interface.
diff --git a/Engine/Source/Runtime/Engine/Classes/Engine/SkinnedAsset.h b/Engine/Source/Runtime/Engine/Classes/Engine/SkinnedAsset.h
--- a/Engine/Source/Runtime/Engine/Classes/Engine/SkinnedAsset.h
+++ b/Engine/Source/Runtime/Engine/Classes/Engine/SkinnedAsset.h
@@ -176,6 +176,11 @@ public:
 	ENGINE_API virtual UMeshDeformer* GetDefaultMeshDeformer() const
 	PURE_VIRTUAL(USkinnedAsset::GetDefaultMeshDeformer, return nullptr;);
 
+	virtual class UMaterialInterface* GetOverlayMaterial() const
+	PURE_VIRTUAL(USkinnedAsset::GetOverlayMaterial, return nullptr;);
+	virtual float GetOverlayMaterialMaxDrawDistance() const
+	PURE_VIRTUAL(USkinnedAsset::GetOverlayMaterialMaxDrawDistance, return 0.f;);
+	
 	/** Return true if given index's LOD is valid */
 	ENGINE_API virtual bool IsValidLODIndex(int32 Index) const
 	PURE_VIRTUAL(USkinnedAsset::IsValidLODIndex, return false;);
diff --git a/Engine/Source/Runtime/Engine/Private/Components/MeshComponent.cpp b/Engine/Source/Runtime/Engine/Private/Components/MeshComponent.cpp
--- a/Engine/Source/Runtime/Engine/Private/Components/MeshComponent.cpp
+++ b/Engine/Source/Runtime/Engine/Private/Components/MeshComponent.cpp
@@ -137,9 +137,10 @@ FMaterialRelevance UMeshComponent::GetMaterialRelevance(ERHIFeatureLevel::Type I
 		Result |= MaterialInterface->GetRelevance_Concurrent(InFeatureLevel);
 	}
 
-	if (OverlayMaterial != nullptr)
+	UMaterialInterface const* OverlayMaterialInterface = GetOverlayMaterial();
+	if (OverlayMaterialInterface != nullptr)
 	{
-		Result |= OverlayMaterial->GetRelevance_Concurrent(InFeatureLevel);
+		Result |= OverlayMaterialInterface->GetRelevance_Concurrent(InFeatureLevel);
 	}
 
 	return Result;
@@ -213,15 +214,23 @@ void UMeshComponent::GetUsedMaterials(TArray<UMaterialInterface*>& OutMaterials,
 		}
 	}
 
-	if (OverlayMaterial != nullptr)
+	UMaterialInterface* OverlayMaterialInterface = GetOverlayMaterial();
+	if (OverlayMaterialInterface != nullptr)
 	{
-		OutMaterials.Add(OverlayMaterial);
+		OutMaterials.Add(OverlayMaterialInterface);
 	}
 }
 
 UMaterialInterface* UMeshComponent::GetOverlayMaterial() const
 {
-	return OverlayMaterial;
+	if (OverlayMaterial)
+	{ 
+		return OverlayMaterial;
+	}
+	else
+	{
+		return GetDefaultOverlayMaterial();
+	}
 }
 
 void UMeshComponent::SetOverlayMaterial(UMaterialInterface* NewOverlayMaterial)
@@ -235,6 +244,18 @@ void UMeshComponent::SetOverlayMaterial(UMaterialInterface* NewOverlayMaterial)
 	}
 }
 
+float UMeshComponent::GetOverlayMaterialMaxDrawDistance() const
+{
+	if (OverlayMaterialMaxDrawDistance != 0.f)
+	{ 
+		return OverlayMaterialMaxDrawDistance;
+	}
+	else
+	{
+		return GetDefaultOverlayMaterialMaxDrawDistance();
+	}
+}
+
 void UMeshComponent::SetOverlayMaterialMaxDrawDistance(float InMaxDrawDistance)
 {
 	if (OverlayMaterialMaxDrawDistance != InMaxDrawDistance)
diff --git a/Engine/Source/Runtime/Engine/Private/Components/SkinnedMeshComponent.cpp b/Engine/Source/Runtime/Engine/Private/Components/SkinnedMeshComponent.cpp
--- a/Engine/Source/Runtime/Engine/Private/Components/SkinnedMeshComponent.cpp
+++ b/Engine/Source/Runtime/Engine/Private/Components/SkinnedMeshComponent.cpp
@@ -623,12 +623,13 @@ void USkinnedMeshComponent::CollectPSOPrecacheData(const FPSOPrecacheParams& Bas
 		ComponentParams.PSOPrecacheParams = PrecachePSOParams;
 	}
 
-	if (OverlayMaterial && VFsPerMaterials.Num() != 0)
+	UMaterialInterface* OverlayMaterialInterface = GetOverlayMaterial();
+	if (OverlayMaterialInterface && VFsPerMaterials.Num() != 0)
 	{
 		// Overlay is rendered with the same set of VFs
 		FComponentPSOPrecacheParams& ComponentParams = OutParams[OutParams.AddDefaulted()];
 
-		ComponentParams.MaterialInterface = OverlayMaterial;
+		ComponentParams.MaterialInterface = OverlayMaterialInterface;
 		ComponentParams.VertexFactoryDataList = VFsPerMaterials[0].VertexFactoryDataList;
 		ComponentParams.PSOPrecacheParams = PrecachePSOParams;
 		ComponentParams.PSOPrecacheParams.bCastShadow = false;
@@ -2527,6 +2528,24 @@ void USkinnedMeshComponent::SetForceWireframe(bool InForceWireframe)
 	}
 }
 
+UMaterialInterface* USkinnedMeshComponent::GetDefaultOverlayMaterial() const
+{
+	if (GetSkinnedAsset())
+	{
+		return GetSkinnedAsset()->GetOverlayMaterial();
+	}
+	return nullptr;
+}
+
+float USkinnedMeshComponent::GetDefaultOverlayMaterialMaxDrawDistance() const
+{
+	if (GetSkinnedAsset())
+	{
+		return GetSkinnedAsset()->GetOverlayMaterialMaxDrawDistance();
+	}
+	return 0.f;
+}
+
 #if WITH_EDITOR
 void USkinnedMeshComponent::SetSectionPreview(int32 InSectionIndexPreview)
 {
@@ -3250,9 +3269,10 @@ void USkinnedMeshComponent::GetUsedMaterials( TArray<UMaterialInterface*>& OutMa
 			OutMaterials.Add( MaterialInterface );
 		}
 
-		if (OverlayMaterial != nullptr)
+		UMaterialInterface* OverlayMaterialInterface = GetOverlayMaterial();
+		if (OverlayMaterialInterface != nullptr)
 		{
-			OutMaterials.Add(OverlayMaterial);
+			OutMaterials.Add(OverlayMaterialInterface);
 		}
 	}
 
diff --git a/Engine/Source/Runtime/Engine/Private/Components/StaticMeshComponent.cpp b/Engine/Source/Runtime/Engine/Private/Components/StaticMeshComponent.cpp
--- a/Engine/Source/Runtime/Engine/Private/Components/StaticMeshComponent.cpp
+++ b/Engine/Source/Runtime/Engine/Private/Components/StaticMeshComponent.cpp
@@ -1621,12 +1621,13 @@ void UStaticMeshComponent::CollectPSOPrecacheDataImpl(
 		ComponentParams.PSOPrecacheParams = PrecachePSOParams;
 	}
 
-	if (OverlayMaterial && VFTypesPerMaterialIndex.Num() != 0)
+	UMaterialInterface* OverlayMaterialInterface = GetOverlayMaterial();
+	if (OverlayMaterialInterface && VFTypesPerMaterialIndex.Num() != 0)
 	{
 		// Overlay is rendered with the same set of VFs
 		FComponentPSOPrecacheParams& ComponentParams = OutParams[OutParams.AddDefaulted()];
 		
-		ComponentParams.MaterialInterface = OverlayMaterial;
+		ComponentParams.MaterialInterface = OverlayMaterialInterface;
 		ComponentParams.VertexFactoryDataList = VFTypesPerMaterialIndex[0].VertexFactoryDataList;
 		ComponentParams.PSOPrecacheParams = PrecachePSOParams;
 		ComponentParams.PSOPrecacheParams.bCastShadow = false;
@@ -2850,9 +2851,10 @@ void UStaticMeshComponent::GetUsedMaterials(TArray<UMaterialInterface*>& OutMate
 				OutMaterials.Add(GetMaterial(MaterialIndex));
 			}
 
-			if (OverlayMaterial != nullptr)
+			UMaterialInterface* OverlayMaterialInterface = GetOverlayMaterial();
+			if (OverlayMaterialInterface != nullptr)
 			{
-				OutMaterials.Add(OverlayMaterial);
+				OutMaterials.Add(OverlayMaterialInterface);
 			}
 		}
 	}
diff --git a/Engine/Source/Runtime/Engine/Private/SkeletalMesh.cpp b/Engine/Source/Runtime/Engine/Private/SkeletalMesh.cpp
--- a/Engine/Source/Runtime/Engine/Private/SkeletalMesh.cpp
+++ b/Engine/Source/Runtime/Engine/Private/SkeletalMesh.cpp
@@ -5564,8 +5564,8 @@ FSkeletalMeshSceneProxy::FSkeletalMeshSceneProxy(const USkinnedMeshComponent* Co
 		,	SkeletalMeshRenderData(InSkelMeshRenderData)
 		,	SkeletalMeshForDebug(Component->GetSkinnedAsset())
 		,	PhysicsAssetForDebug(Component->GetPhysicsAsset())
-		,	OverlayMaterial(Component->OverlayMaterial)
-		,	OverlayMaterialMaxDrawDistance(Component->OverlayMaterialMaxDrawDistance)
+		,	OverlayMaterial(Component->GetOverlayMaterial())
+		,	OverlayMaterialMaxDrawDistance(Component->GetOverlayMaterialMaxDrawDistance())
 #if RHI_RAYTRACING
 		,	bAnySegmentUsesWorldPositionOffset(false)
 #endif
diff --git a/Engine/Source/Runtime/Engine/Private/StaticMeshRender.cpp b/Engine/Source/Runtime/Engine/Private/StaticMeshRender.cpp
--- a/Engine/Source/Runtime/Engine/Private/StaticMeshRender.cpp
+++ b/Engine/Source/Runtime/Engine/Private/StaticMeshRender.cpp
@@ -200,8 +200,8 @@ static TAutoConsoleVariable<float> CVarRayTracingStaticMeshesWPOCullingRadius(
 FStaticMeshSceneProxy::FStaticMeshSceneProxy(UStaticMeshComponent* InComponent, bool bForceLODsShareStaticLighting)
 	: FPrimitiveSceneProxy(InComponent, InComponent->GetStaticMesh()->GetFName())
 	, RenderData(InComponent->GetStaticMesh()->GetRenderData())
-	, OverlayMaterial(InComponent->OverlayMaterial)
-	, OverlayMaterialMaxDrawDistance(InComponent->OverlayMaterialMaxDrawDistance)
+	, OverlayMaterial(InComponent->GetOverlayMaterial())
+	, OverlayMaterialMaxDrawDistance(InComponent->GetOverlayMaterialMaxDrawDistance())
 	, ForcedLodModel(InComponent->ForcedLodModel)
 	, bCastShadow(InComponent->CastShadow)
 	, bReverseCulling(InComponent->bReverseCulling)
