Fixed linking error from CL# 26405350

diff --git a/Engine/Source/Runtime/Engine/Classes/Components/SkinnedMeshComponent.h b/Engine/Source/Runtime/Engine/Classes/Components/SkinnedMeshComponent.h
--- a/Engine/Source/Runtime/Engine/Classes/Components/SkinnedMeshComponent.h
+++ b/Engine/Source/Runtime/Engine/Classes/Components/SkinnedMeshComponent.h
@@ -1536,10 +1536,10 @@ protected:
 
 protected:
 	/** Get the default overlay material used by a mesh */
-	virtual UMaterialInterface* GetDefaultOverlayMaterial() const override;
+	ENGINE_API virtual UMaterialInterface* GetDefaultOverlayMaterial() const override;
 
 	/** Get the default overlay material max draw distance */
-	virtual float GetDefaultOverlayMaterialMaxDrawDistance() const override;
+	ENGINE_API virtual float GetDefaultOverlayMaterialMaxDrawDistance() const override;
 
 public:
 	/** 
