Level Instance: Allow child edit when parent has unsaved changes (prompt user for save)

#rnx

diff --git a/Engine/Source/Runtime/Engine/Private/LevelInstance/LevelInstanceSubsystem.cpp b/Engine/Source/Runtime/Engine/Private/LevelInstance/LevelInstanceSubsystem.cpp
--- a/Engine/Source/Runtime/Engine/Private/LevelInstance/LevelInstanceSubsystem.cpp
+++ b/Engine/Source/Runtime/Engine/Private/LevelInstance/LevelInstanceSubsystem.cpp
@@ -631,25 +631,11 @@ bool ULevelInstanceSubsystem::OnExitEditorModeInternal(bool bForceExit)
 	if (LevelInstanceEdit)
 	{
 		TGuardValue<bool> CommitScope(bIsCommittingLevelInstance, true);
-		ILevelInstanceInterface* LevelInstance = GetEditingLevelInstance();
-
 		bool bDiscard = false;
-		bool bIsDirty = IsLevelInstanceEditDirty(LevelInstanceEdit.Get());
-		if (bIsDirty && CanCommitLevelInstance(LevelInstance, /*bDiscardEdits=*/true))
-		{
-			// if bForceExit we can't cancel the exiting of the mode so the user needs to decide between saving or discarding
-			EAppReturnType::Type Ret = FMessageDialog::Open(
-				bForceExit ? EAppMsgType::YesNo : EAppMsgType::YesNoCancel, LOCTEXT("CommitOrDiscardChangesMsg", "Unsaved Level changes will get discarded. Do you want to save them now?"), 
-				LOCTEXT("CommitOrDiscardChangesTitle", "Save changes?"));
-			if (Ret == EAppReturnType::Cancel && !bForceExit)
-			{
-				return false;
-			}
-
-			bDiscard = (Ret != EAppReturnType::Yes);
+		if (PromptUserForCommit(LevelInstanceEdit.Get(), bDiscard, bForceExit))
+		{
+			return CommitLevelInstanceInternal(LevelInstanceEdit, bDiscard, /*bDiscardOnFailure=*/bForceExit);
 		}
-
-		return CommitLevelInstanceInternal(LevelInstanceEdit, bDiscard, /*bDiscardOnFailure=*/bForceExit);
 	}
 
 	return false;
@@ -1702,6 +1688,34 @@ ILevelInstanceInterface* ULevelInstanceSubsystem::GetEditingLevelInstance() cons
 	return nullptr;
 }
 
+bool ULevelInstanceSubsystem::PromptUserForCommit(const FLevelInstanceEdit* InLevelInstanceEdit, bool& bOutDiscard, bool bForceCommit) const
+{
+	bOutDiscard = false;
+	// Can commit no pending changes
+	if (!IsLevelInstanceEditDirty(InLevelInstanceEdit)) 
+	{
+		return true;
+	}
+
+	// If changes can be discarded prompt user
+	if (CanCommitLevelInstance(InLevelInstanceEdit->GetLevelInstance(), /*bDiscardEdits=*/true))
+	{
+		// if bForceExit we can't cancel the exiting of the mode so the user needs to decide between saving or discarding
+		EAppReturnType::Type Ret = FMessageDialog::Open(
+			bForceCommit ? EAppMsgType::YesNo : EAppMsgType::YesNoCancel, LOCTEXT("CommitOrDiscardChangesMsg", "Unsaved Level changes will get discarded. Do you want to save them now?"),
+			LOCTEXT("CommitOrDiscardChangesTitle", "Save changes?"));
+		if (Ret == EAppReturnType::Cancel && !bForceCommit)
+		{
+			return false;
+		}
+
+		bOutDiscard = (Ret != EAppReturnType::Yes);
+	}
+
+	// Can commit but can't discard changes
+	return true;
+}
+
 bool ULevelInstanceSubsystem::CanEditLevelInstance(const ILevelInstanceInterface* LevelInstance, FText* OutReason) const
 {
 	// Only allow Editing in Editor World
@@ -1720,15 +1734,6 @@ bool ULevelInstanceSubsystem::CanEditLevelInstance(const ILevelInstanceInterface
 			}
 			return false;
 		}
-
-		if (IsLevelInstanceEditDirty(LevelInstanceEdit.Get()))
-		{
-			if (OutReason)
-			{
-				*OutReason = FText::Format(LOCTEXT("CanEditLevelInstanceDirtyEdit", "Current Level Instance has unsaved changes and needs to be committed first ({0})."), FText::FromString(GetEditingLevelInstance()->GetWorldAssetPackage()));
-			}
-			return false;
-		}
 	}
 	
 	if (LevelInstance->IsWorldAssetValid())
@@ -1792,6 +1797,13 @@ bool ULevelInstanceSubsystem::EditLevelInstanceInternal(ILevelInstanceInterface*
 {
 	check(CanEditLevelInstance(LevelInstance));
 		
+	// If there is a current edit and it is dirty, offer the user a chance to Save/Discard/Cancel
+	bool bDiscard = false;
+	if (LevelInstanceEdit && !PromptUserForCommit(LevelInstanceEdit.Get(), bDiscard))
+	{
+		return false;
+	}
+
 	FScopedSlowTask SlowTask(0, LOCTEXT("BeginEditLevelInstance", "Loading Level Instance for edit..."), !GetWorld()->IsGameWorld());
 	SlowTask.MakeDialog();
 
@@ -1842,8 +1854,7 @@ bool ULevelInstanceSubsystem::EditLevelInstanceInternal(ILevelInstanceInterface*
 		// Make sure to keep the top level instance actor loaded when we commit the current one
 		FWorldPartitionReference CurrentEditLevelInstanceActorRef = CurrentEditLevelInstanceActor;
 		
-		check(!IsLevelInstanceEditDirty(LevelInstanceEdit.Get()));
-		CommitLevelInstanceInternal(LevelInstanceEdit);
+		CommitLevelInstanceInternal(LevelInstanceEdit, bDiscard);
 
 		ILevelInstanceInterface* LevelInstanceToEdit = GetLevelInstance(PendingEditId);
 		check(LevelInstanceToEdit);
diff --git a/Engine/Source/Runtime/Engine/Public/LevelInstance/LevelInstanceSubsystem.h b/Engine/Source/Runtime/Engine/Public/LevelInstance/LevelInstanceSubsystem.h
--- a/Engine/Source/Runtime/Engine/Public/LevelInstance/LevelInstanceSubsystem.h
+++ b/Engine/Source/Runtime/Engine/Public/LevelInstance/LevelInstanceSubsystem.h
@@ -191,7 +191,8 @@ private:
 	
 	ENGINE_API const FLevelInstanceEdit* GetLevelInstanceEdit(const ILevelInstanceInterface* LevelInstance) const;
 	ENGINE_API bool IsLevelInstanceEditDirty(const FLevelInstanceEdit* LevelInstanceEdit) const;
-	
+	ENGINE_API bool PromptUserForCommit(const FLevelInstanceEdit* InLevelInstanceEdit, bool& bOutDiscard, bool bForceCommit = false) const;
+
 	struct FLevelsToRemoveScope
 	{
 		FLevelsToRemoveScope(ULevelInstanceSubsystem* InOwner);
