Fixed some uninitialized FBox

diff --git a/Engine/Plugins/Experimental/Dataflow/Source/DataflowEnginePlugin/Private/Dataflow/DataflowEngineRendering.cpp b/Engine/Plugins/Experimental/Dataflow/Source/DataflowEnginePlugin/Private/Dataflow/DataflowEngineRendering.cpp
--- a/Engine/Plugins/Experimental/Dataflow/Source/DataflowEnginePlugin/Private/Dataflow/DataflowEngineRendering.cpp
+++ b/Engine/Plugins/Experimental/Dataflow/Source/DataflowEnginePlugin/Private/Dataflow/DataflowEngineRendering.cpp
@@ -207,7 +207,7 @@ namespace Dataflow
 				{
 					FName PrimaryOutput = State.GetRenderOutputs()[0]; // "Box"
 
-					FBox Default;
+					FBox Default(ForceInit);
 					const FBox& Box = State.GetValue<FBox>(PrimaryOutput, Default);
 
 					const int32 NumVertices = 8;
diff --git a/Engine/Plugins/Experimental/GPULightmass/Source/GPULightmass/Private/LightmapRenderer.cpp b/Engine/Plugins/Experimental/GPULightmass/Source/GPULightmass/Private/LightmapRenderer.cpp
--- a/Engine/Plugins/Experimental/GPULightmass/Source/GPULightmass/Private/LightmapRenderer.cpp
+++ b/Engine/Plugins/Experimental/GPULightmass/Source/GPULightmass/Private/LightmapRenderer.cpp
@@ -3194,7 +3194,7 @@ bool MortonCompare(FUintVector A, FUintVector B)
 
 void FSceneRenderState::BuildMortonSortedLightmapRefList()
 {
-	FBox SceneBounds;
+	FBox SceneBounds(ForceInit);
 
 	for (FLightmapRenderState& Lightmap : LightmapRenderStates.Elements)
 	{
diff --git a/Engine/Source/Runtime/Engine/Private/Collision/KAggregateGeom.cpp b/Engine/Source/Runtime/Engine/Private/Collision/KAggregateGeom.cpp
--- a/Engine/Source/Runtime/Engine/Private/Collision/KAggregateGeom.cpp
+++ b/Engine/Source/Runtime/Engine/Private/Collision/KAggregateGeom.cpp
@@ -495,7 +495,7 @@ FBox FKTaperedCapsuleElem::CalcAABB(const FTransform& BoneTM, float Scale) const
 
 FBox FKLevelSetElem::CalcAABB(const FTransform& BoneTM, const FVector& Scale3D) const
 {
-	FBox Box;
+	FBox Box(ForceInit);
 
 	if (LevelSet.IsValid())
 	{
@@ -509,7 +509,7 @@ FBox FKLevelSetElem::CalcAABB(const FTransform& BoneTM, const FVector& Scale3D)
 
 FBox FKLevelSetElem::UntransformedAABB() const
 {
-	FBox Box;
+	FBox Box(ForceInit);
 
 	if (LevelSet.IsValid())
 	{
@@ -608,7 +608,7 @@ void FKLevelSetElem::GetZeroIsosurfaceGridCellFaces(TArray<FVector3f>& Vertices,
 
 FBox FKSkinnedLevelSetElem::CalcAABB(const FTransform& BoneTM, const FVector& Scale3D) const
 {
-	FBox Box;
+	FBox Box(ForceInit);
 
 	if (WeightedLevelSet.IsValid())
 	{
diff --git a/Engine/Source/Runtime/Engine/Private/Elements/Framework/TypedElementCommonActions.cpp b/Engine/Source/Runtime/Engine/Private/Elements/Framework/TypedElementCommonActions.cpp
--- a/Engine/Source/Runtime/Engine/Private/Elements/Framework/TypedElementCommonActions.cpp
+++ b/Engine/Source/Runtime/Engine/Private/Elements/Framework/TypedElementCommonActions.cpp
@@ -321,7 +321,7 @@ TArray<FTypedElementHandle> UTypedElementCommonActions::PasteNormalizedElements(
 	// Paste at implementation
 	if (PasteOptions.bPasteAtLocation)
 	{
-		FBox BoundingBox;
+		FBox BoundingBox(ForceInit);
 
 		TArray<TTypedElement<ITypedElementWorldInterface>> WorldElements;
 		WorldElements.Reserve(PastedHandles.Num());
diff --git a/Engine/Source/Runtime/Engine/Private/InstancedStaticMesh.cpp b/Engine/Source/Runtime/Engine/Private/InstancedStaticMesh.cpp
--- a/Engine/Source/Runtime/Engine/Private/InstancedStaticMesh.cpp
+++ b/Engine/Source/Runtime/Engine/Private/InstancedStaticMesh.cpp
@@ -4712,7 +4712,7 @@ void UInstancedStaticMeshComponent::InitPerInstanceRenderData(bool InitializeFro
 
 	bool KeepInstanceBufferCPUAccess = UseGPUScene(GetFeatureLevelShaderPlatform(FeatureLevel), FeatureLevel) || GIsEditor || InRequireCPUAccess || ComponentRequestsCPUAccess(this, FeatureLevel);
 
-	FBox LocalBounds;
+	FBox LocalBounds(ForceInit);
 	if (GetStaticMesh())
 	{
 		FVector BoundsMin, BoundsMax;
