Fix ENGINE_API & similar typos in audio code

diff --git a/Engine/Source/Runtime/AudioMixer/Public/Components/SynthComponent.h b/Engine/Source/Runtime/AudioMixer/Public/Components/SynthComponent.h
--- a/Engine/Source/Runtime/AudioMixer/Public/Components/SynthComponent.h
+++ b/Engine/Source/Runtime/AudioMixer/Public/Components/SynthComponent.h
@@ -200,7 +200,7 @@ public:
 	* @return The set of of Modulators applied to this component for the given Destination.
 	*/
 	UFUNCTION(BlueprintPure, Category = "Audio|Components|Audio", DisplayName = "Get Modulators")
-	UPARAM(DisplayName = "ModAUDIOMIXER_API ulators") TSet<USoundModulatorBase*> GetModulators(const EModulationDestination Destination);
+	AUDIOMIXER_API UPARAM(DisplayName = "Modulators") TSet<USoundModulatorBase*> GetModulators(const EModulationDestination Destination);
 
 	/** Auto destroy this component on completion */
 	UPROPERTY()
diff --git a/Engine/Source/Runtime/Engine/Classes/Components/AudioComponent.h b/Engine/Source/Runtime/Engine/Classes/Components/AudioComponent.h
--- a/Engine/Source/Runtime/Engine/Classes/Components/AudioComponent.h
+++ b/Engine/Source/Runtime/Engine/Classes/Components/AudioComponent.h
@@ -745,7 +745,7 @@ public:
 	* @return The set of of Modulators applied to this component for the given Destination.
 	*/
 	UFUNCTION(BlueprintPure, Category = "Audio|Components|Audio", DisplayName = "Get Modulators")
-	ENGINE_API UPARAM(DisplayName = "ModENGINE_API ulators") TSet<USoundModulatorBase*> GetModulators(const EModulationDestination Destination);
+	ENGINE_API UPARAM(DisplayName = "Modulators") TSet<USoundModulatorBase*> GetModulators(const EModulationDestination Destination);
 
 	static ENGINE_API void PlaybackCompleted(uint64 AudioComponentID, bool bFailedToStart);
 
diff --git a/Engine/Source/Runtime/Engine/Public/Audio/SoundParameterControllerInterface.h b/Engine/Source/Runtime/Engine/Public/Audio/SoundParameterControllerInterface.h
--- a/Engine/Source/Runtime/Engine/Public/Audio/SoundParameterControllerInterface.h
+++ b/Engine/Source/Runtime/Engine/Public/Audio/SoundParameterControllerInterface.h
@@ -69,32 +69,32 @@ class UAudioParameterConversionStatics : public UBlueprintFunctionLibrary
 
 public:
 	UFUNCTION(BlueprintPure, Category = "Audio|Parameter", meta = (Keywords = "make construct convert create"))
-	static UPARAM(DisplayName = "ParENGINE_API ameter") FAudioParameter BooleanToAudioParameter(FName Name, bool Bool);
+	static UPARAM(DisplayName = "Parameter") FAudioParameter BooleanToAudioParameter(FName Name, bool Bool);
 
 	UFUNCTION(BlueprintPure, Category = "Audio|Parameter", meta = (Keywords = "make construct convert create"))
-	static UPARAM(DisplayName = "ParENGINE_API ameter") FAudioParameter FloatToAudioParameter(FName Name, float Float);
+	static UPARAM(DisplayName = "Parameter") FAudioParameter FloatToAudioParameter(FName Name, float Float);
 
 	UFUNCTION(BlueprintPure, Category = "Audio|Parameter", meta = (Keywords = "make construct convert create"))
-	static UPARAM(DisplayName = "ParENGINE_API ameter") FAudioParameter IntegerToAudioParameter(FName Name, int32 Integer);
+	static UPARAM(DisplayName = "Parameter") FAudioParameter IntegerToAudioParameter(FName Name, int32 Integer);
 
 	UFUNCTION(BlueprintPure, Category = "Audio|Parameter", meta = (Keywords = "make construct convert create"))
-	static UPARAM(DisplayName = "ParENGINE_API ameter") FAudioParameter StringToAudioParameter(FName Name, FString String);
+	static UPARAM(DisplayName = "Parameter") FAudioParameter StringToAudioParameter(FName Name, FString String);
 
 	UFUNCTION(BlueprintPure, Category = "Audio|Parameter", meta = (Keywords = "make construct convert create"))
-	static UPARAM(DisplayName = "ParENGINE_API ameter") FAudioParameter ObjectToAudioParameter(FName Name, UObject* Object);
+	static UPARAM(DisplayName = "Parameter") FAudioParameter ObjectToAudioParameter(FName Name, UObject* Object);
 
 	UFUNCTION(BlueprintPure, Category = "Audio|Parameter", meta = (Keywords = "make construct convert create"))
-	static UPARAM(DisplayName = "ParENGINE_API ameter") FAudioParameter BooleanArrayToAudioParameter(FName Name, TArray<bool> Bools);
+	static UPARAM(DisplayName = "Parameter") FAudioParameter BooleanArrayToAudioParameter(FName Name, TArray<bool> Bools);
 
 	UFUNCTION(BlueprintPure, Category = "Audio|Parameter", meta = (Keywords = "make construct convert create"))
-	static UPARAM(DisplayName = "ParENGINE_API ameter") FAudioParameter FloatArrayToAudioParameter(FName Name, TArray<float> Floats);
+	static UPARAM(DisplayName = "Parameter") FAudioParameter FloatArrayToAudioParameter(FName Name, TArray<float> Floats);
 
 	UFUNCTION(BlueprintPure, Category = "Audio|Parameter", meta = (Keywords = "make construct convert create"))
-	static UPARAM(DisplayName = "ParENGINE_API ameter") FAudioParameter IntegerArrayToAudioParameter(FName Name, TArray<int32> Integers);
+	static UPARAM(DisplayName = "Parameter") FAudioParameter IntegerArrayToAudioParameter(FName Name, TArray<int32> Integers);
 
 	UFUNCTION(BlueprintPure, Category = "Audio|Parameter", meta = (Keywords = "make construct convert create"))
-	static UPARAM(DisplayName = "ParENGINE_API ameter") FAudioParameter StringArrayToAudioParameter(FName Name, TArray<FString> Strings);
+	static UPARAM(DisplayName = "Parameter") FAudioParameter StringArrayToAudioParameter(FName Name, TArray<FString> Strings);
 
 	UFUNCTION(BlueprintPure, Category = "Audio|Parameter", meta = (Keywords = "make construct convert create"))
-	static UPARAM(DisplayName = "ParENGINE_API ameter") FAudioParameter ObjectArrayToAudioParameter(FName Name, TArray<UObject*> Objects);
+	static UPARAM(DisplayName = "Parameter") FAudioParameter ObjectArrayToAudioParameter(FName Name, TArray<UObject*> Objects);
 };
