Shader Compiler:  Job cache data structure refactoring, in prepation for taskifying SubmitJobs().  The goal was to make the structures amenable to lock free approaches, to reduce the amount of time expensive write locks are required on the Job Cache.  Also, all structures have been moved out of the header (using TPimplPtr on the job cache structure), so we can do refactors of the job cache at will, without affecting public headers.  Significant testing was done, including stress testing the memory budget logic, and cancelling in-flight jobs, including duplicate jobs.  I discovered a bug in AddJobOutput, where it was overwriting existing map entries rather than adding references to them, causing asserts in memory budget logic (in addition to a more minor issue of the memory statistic being incorrect).

* Instead of having three maps by input hash (in-flight, duplicate, and completed jobs), there is one map that contains all three.
* Memory for items in the map is block allocated, and persistent (only removed during rare over budget cleanup case), allowing jobs to maintain direct back references to their map entry, avoiding map lookup.
* Find followed by Add replaced with a combined FindOrAdd operation.
* EvictionQueue was removed, and instead over budget cleanup eliminates items from the front of the map.  Block allocation makes it extremely cheap to remove items from the front.
* The previous four items combine to reduce the number of unique container insert/remove operations per typical job from six (find existing before, insert in-flight, remove in-flight, find existing after, insert output, insert eviction queue) to one (find or add).
* Over budget cleanup has been made significantly more rare, by reducing memory to a percentage of the budget, rather than eliminating potentially a single item at a time.
* Duplicate tracking uses an array rather than a map, with back references to the duplicate entry, avoiding map lookup for that case as well.

#rnx

diff --git a/Engine/Source/Runtime/Engine/Private/ShaderCompiler/ShaderCompiler.cpp b/Engine/Source/Runtime/Engine/Private/ShaderCompiler/ShaderCompiler.cpp
--- a/Engine/Source/Runtime/Engine/Private/ShaderCompiler/ShaderCompiler.cpp
+++ b/Engine/Source/Runtime/Engine/Private/ShaderCompiler/ShaderCompiler.cpp
@@ -111,6 +111,14 @@ static FAutoConsoleVariableRef CVarShaderCompilerMaxJobCacheMemoryPercent(
 	ECVF_Default
 );
 
+int32 GShaderCompilerJobCacheOverflowReducePercent = 80;
+static FAutoConsoleVariableRef CVarShaderCompilerJobCacheOverflowReducePercent(
+	TEXT("r.ShaderCompiler.JobCacheOverflowReducePercent"),
+	GShaderCompilerJobCacheOverflowReducePercent,
+	TEXT("When shader compiler job cache memory overflows, reduce memory to this percentage of the maximum.  Reduces overhead relative to cleaning up items one at a time when at max budget."),
+	ECVF_Default
+);
+
 static TAutoConsoleVariable<bool> CVarPreprocessedJobCache(
 	TEXT("r.ShaderCompiler.PreprocessedJobCache"),
 	false,
@@ -232,6 +240,522 @@ static TAutoConsoleVariable<bool> CVarShaderCompilerDumpDDCKeys(
 	ECVF_Default
 );
 
+int32 GShaderCompilerDebugDiscardCacheOutputs = 0;
+static FAutoConsoleVariableRef CVarShaderCompilerDebugDiscardCacheOutputs(
+	TEXT("r.ShaderCompiler.DebugDiscardCacheOutputs"),
+	GShaderCompilerDebugDiscardCacheOutputs,
+	TEXT("if != 0, cache outputs are discarded (not added to the output map) for debugging purposes.\nEliminates usefulness of the cache, but allows repeated triggering of the same jobs for stress testing (for example, rapid undo/redo in the Material editor)."),
+	ECVF_Default
+);
+
+
+/** Map element type for job cache */
+struct FShaderJobData
+{
+	using FJobInputHash = FBlake3Hash;
+	using FJobOutputHash = FBlake3Hash;
+
+	FJobInputHash InputHash;
+
+	/** Output hash will be zero if output data has not been written yet, or can be cleared if output data has been removed */
+	FJobOutputHash OutputHash;
+
+	/** In-flight job with the given input hash */
+	FShaderCommonCompileJob* JobInFlight = nullptr;
+
+	/** Head of a linked list of duplicate jobs */
+	FShaderCommonCompileJob* DuplicateJobsWaitList = nullptr;
+
+	bool IsEmpty() const
+	{
+		return OutputHash.IsZero() && JobInFlight == nullptr && DuplicateJobsWaitList == nullptr;
+	}
+
+	FORCEINLINE bool HasOutput() const
+	{
+		return OutputHash.IsZero() == false;
+	}
+};
+
+/** Block of map elements for job cache */
+struct FShaderJobDataBlock
+{
+	static const int32 BlockSize = 512;
+	static_assert(FMath::IsPowerOfTwo(BlockSize));
+
+	FShaderJobData Data[BlockSize];
+};
+
+class FShaderJobDataMap
+{
+public:
+	FShaderJobDataMap()
+	{
+		// Reserve so we don't need a special case for an empty HashTable array
+		Reserve(FShaderJobDataBlock::BlockSize);
+	}
+
+	FShaderJobData* Find(const FShaderJobData::FJobInputHash& Key);
+	FShaderJobCacheRef FindOrAdd(const FShaderJobData::FJobInputHash& Key);
+
+	FORCEINLINE int32 Num() const
+	{
+		return NumItems;
+	}
+
+	FORCEINLINE FShaderJobData& operator[](int32 Index)
+	{
+		check(Index < NumItems);
+		return DataBlocks[Index / FShaderJobDataBlock::BlockSize].Data[Index & (FShaderJobDataBlock::BlockSize - 1)];
+	}
+
+	FORCEINLINE const FShaderJobData& operator[](int32 Index) const
+	{
+		return (*const_cast<FShaderJobDataMap*>(this))[Index];
+	}
+
+	uint64 GetAllocatedSize() const
+	{
+		return DataBlocks.GetAllocatedSize() + HashTable.GetAllocatedSize();
+	}
+
+	void RemoveLeadingBlocks(int32 BlocksToRemove)
+	{
+		check(BlocksToRemove <= DataBlocks.Num() && BlocksToRemove > 0);
+		DataBlocks.RemoveAt(0, BlocksToRemove);
+		NumItems -= BlocksToRemove * FShaderJobDataBlock::BlockSize;
+		check(NumItems >= 0);
+
+		if (NumItems == 0)
+		{
+			// If we happened to remove ALL the items, reserve again, as done in the constructor
+			Reserve(FShaderJobDataBlock::BlockSize);
+		}
+		else
+		{
+			// Otherwise, we need to rehash, as all item indices will have changed
+			ReHash(GetDesiredHashTableSize());
+		}
+	}
+
+private:
+	void ReHash(int32 HashTableSize);
+	void Reserve(int32 NumReserve);
+
+	int32 GetDesiredHashTableSize() const
+	{
+		return FMath::RoundUpToPowerOfTwo(DataBlocks.Num() * FShaderJobDataBlock::BlockSize * 2);
+	}
+
+	/** An indirect array of blocks is used, so data elements never move in memory when the table grows */
+	TIndirectArray<FShaderJobDataBlock> DataBlocks;
+	int32 NumItems = 0;
+
+	/** Power of two hash table with linear probing */
+	TArray<uint32> HashTable;
+	uint32 HashTableMask = 0;
+};
+
+
+struct FShaderJobCacheStoredOutput
+{
+private:
+	/** How many times this output is referenced by the cached jobs */
+	int32 NumReferences = 0;
+
+public:
+
+	/** How many times this output has been returned as a cached result, no matter the input hash */
+	int32 NumHits = 0;
+
+	/** Canned output */
+	FSharedBuffer JobOutput;
+
+	/** Path to where the cached debug info is stored. */
+	FString CachedDebugInfoPath;
+
+	/** Similar to FRefCountBase AddRef, but not atomic */
+	int32 AddRef()
+	{
+		++NumReferences;
+
+		return NumReferences;
+	}
+
+	int32 GetNumReferences() const
+	{
+		return NumReferences;
+	}
+
+	/** Similar to FRefCountBase Release, but not atomic */
+	int32 Release()
+	{
+		checkf(NumReferences >= 0, TEXT("Attempting to release shader job cache output that was already released"));
+
+		--NumReferences;
+
+		const int32 RemainingNumReferences = NumReferences;
+
+		if (RemainingNumReferences == 0)
+		{
+			delete this;
+		}
+
+		return RemainingNumReferences;
+	}
+
+	uint64 GetAllocatedSize() const
+	{
+		return static_cast<uint64>(JobOutput.GetSize() + sizeof(*this));
+	}
+};
+
+
+/** Private implementation class for FShaderCompileJobCollection */
+class FShaderJobCache
+{
+public:
+	FShaderJobCache();
+	~FShaderJobCache();
+
+	FShaderCompileJob* PrepareJob(uint32 InId, const FShaderCompileJobKey& InKey, EShaderCompileJobPriority InPriority);
+	FShaderPipelineCompileJob* PrepareJob(uint32 InId, const FShaderPipelineCompileJobKey& InKey, EShaderCompileJobPriority InPriority);
+	void RemoveJob(FShaderCommonCompileJob* InJob);
+
+	int32 RemoveAllPendingJobsWithId(uint32 InId);
+
+	void SubmitJobs(const TArray<FShaderCommonCompileJobPtr>& InJobs);
+
+	/** This is an entry point for all jobs that have finished the compilation (whether real or cached). Can be called from multiple threads.*/
+	void ProcessFinishedJob(FShaderCommonCompileJob* FinishedJob, bool bWasCached);
+
+	/** Adds the job to cache. */
+	void AddToCacheAndProcessPending(FShaderCommonCompileJob* FinishedJob);
+
+	/** Log caching statistics.*/
+	void LogCachingStats();
+
+	/** Gather statistics to send to analytics */
+	void GatherAnalytics(const FString& BaseName, TArray<FAnalyticsEventAttribute>& Attributes) const;
+
+	int32 GetNumPendingJobs(EShaderCompileJobPriority InPriority) const;
+
+	int32 GetNumOutstandingJobs() const;
+
+	int32 GetNumPendingJobs() const;
+
+	int32 GetPendingJobs(EShaderCompilerWorkerType InWorkerType, EShaderCompileJobPriority InPriority, int32 MinNumJobs, int32 MaxNumJobs, TArray<FShaderCommonCompileJobPtr>& OutJobs);
+
+private:
+	using FJobInputHash = FShaderCommonCompileJob::FInputHash;
+	using FJobOutputHash = FBlake3Hash;
+	using FJobCachedOutput = FSharedBuffer;
+	using FStoredOutput = FShaderJobCacheStoredOutput;
+
+	// cannot allow managing this from outside as the caching logic is not exposed
+	inline int32 InternalSubtractNumOutstandingJobs(int32 Value)
+	{
+		const int32 PrevNumOutstandingJobs = NumOutstandingJobs.Subtract(Value);
+		check(PrevNumOutstandingJobs >= Value);
+		return PrevNumOutstandingJobs - Value;
+	}
+
+	template<typename JobType, typename KeyType>
+	int32 InternalFindJobIndex(uint32 InJobHash, uint32 InJobId, const KeyType& InKey) const
+	{
+		const int32 TypeIndex = (int32)JobType::Type;
+		uint32 CurrentPriorityIndex = 0u;
+		int32 CurrentIndex = INDEX_NONE;
+		for (int32 Index = JobHash[TypeIndex].First(InJobHash); JobHash[TypeIndex].IsValid(Index); Index = JobHash[TypeIndex].Next(Index))
+		{
+			const FShaderCommonCompileJob* Job = Jobs[TypeIndex][Index].GetReference();
+			check(Job->Type == JobType::Type);
+
+			// We find the job that matches the key with the highest priority
+			if (Job->Id == InJobId &&
+				(uint32)Job->Priority >= CurrentPriorityIndex &&
+				static_cast<const JobType*>(Job)->Key == InKey)
+			{
+				CurrentPriorityIndex = (uint32)Job->Priority;
+				CurrentIndex = Index;
+			}
+		}
+		return CurrentIndex;
+	}
+
+	template<typename JobType, typename KeyType>
+	JobType* InternalFindJob(uint32 InJobHash, uint32 InJobId, const KeyType& InKey) const
+	{
+		const int32 TypeIndex = (int32)JobType::Type;
+		const int32 JobIndex = InternalFindJobIndex<JobType>(InJobHash, InJobId, InKey);
+		return JobIndex != INDEX_NONE ? static_cast<JobType*>(Jobs[TypeIndex][JobIndex].GetReference()) : nullptr;
+	}
+
+	template<typename JobType, typename KeyType>
+	JobType* InternalPrepareJob(uint32 InId, const KeyType& InKey, EShaderCompileJobPriority InPriority)
+	{
+		const uint32 Hash = InKey.MakeHash(InId);
+		JobType* PrevJob = nullptr;
+		{
+			FReadScopeLock Locker(JobLock);
+			PrevJob = InternalFindJob<JobType>(Hash, InId, InKey);
+		}
+
+		JobType* NewJob = nullptr;
+		if (PrevJob == nullptr || (uint32)InPriority > (uint32)PrevJob->Priority)
+		{
+			FWriteScopeLock Locker(JobLock);
+			if (PrevJob == nullptr)
+			{
+				PrevJob = InternalFindJob<JobType>(Hash, InId, InKey);
+			}
+			if (PrevJob == nullptr)
+			{
+				NewJob = new JobType(Hash, InId, InPriority, InKey);
+				InternalAddJob(NewJob);
+			}
+			else if ((uint32)InPriority > (uint32)PrevJob->Priority)
+			{
+				InternalSetPriority(PrevJob, InPriority);
+			}
+		}
+
+		return NewJob;
+	}
+
+	void InternalAddJob(FShaderCommonCompileJob* Job);
+	void InternalRemoveJob(FShaderCommonCompileJob* InJob);
+	void InternalSetPriority(FShaderCommonCompileJob* Job, EShaderCompileJobPriority InPriority);
+
+	/** Looks for or adds an entry for the given hash in the cache.  Returns cached output if it exists. */
+	FShaderJobCacheRef FindOrAdd(const FJobInputHash& Hash, const bool bCheckDDC, FJobCachedOutput*& OutCachedOutput);
+
+	/** Find an existing item in the cache. */
+	FShaderJobData* Find(const FJobInputHash& Hash);
+
+	/** Add a reference to a duplicate job (to the DuplicateJobs array) */
+	void AddDuplicateJob(FShaderCommonCompileJob* DuplicateJob);
+
+	/** Remove a reference to a duplicate job (from the DuplicateJobs array)  */
+	void RemoveDuplicateJob(FShaderCommonCompileJob* DuplicateJob);
+
+	/** Adds a job output to the cache */
+	void AddJobOutput(FShaderJobData& JobData, const FShaderCommonCompileJob* FinishedJob, const FJobInputHash& Hash, const FJobCachedOutput& Contents, int32 InitialHitCount, const bool bAddToDDC);
+
+	/** Returns memory used by the cache*/
+	uint64 GetAllocatedMemory() const;
+
+	/** Compute memory used by the cache from scratch.  Should match GetAllocatedMemory() if CurrentlyAllocateMemory is being properly updated (useful for validation). */
+	uint64 ComputeAllocatedMemory() const;
+
+	/** Logs out the statistics */
+	void LogStats();
+
+	/** Calculates current memory budget, in bytes */
+	uint64 GetCurrentMemoryBudget() const;
+
+	/** Cleans up oldest outputs to fit in the given memory budget */
+	void CullOutputsToMemoryBudget(uint64 TargetBudgetBytes);
+
+	/** Guards access to the structure */
+	mutable FRWLock JobLock;
+
+	/** Queue of tasks that haven't been assigned to a worker yet. */
+	FShaderCommonCompileJob* PendingJobs[NumShaderCompileJobPriorities];
+	int32 NumPendingJobs[NumShaderCompileJobPriorities];
+
+	/** Number of jobs currently being compiled.  This includes PendingJobs and any jobs that have been assigned to workers but aren't complete yet. */
+	FThreadSafeCounter NumOutstandingJobs;
+
+	TArray<FShaderCommonCompileJobPtr> Jobs[NumShaderCompileJobTypes];
+	TArray<int32> FreeIndices[NumShaderCompileJobTypes];
+	FHashTable JobHash[NumShaderCompileJobTypes];
+
+	/* A lot of outputs can be duplicated, so they are deduplicated before storing */
+	TMap<FJobOutputHash, FStoredOutput*> Outputs;
+
+	/** Map of input hashes to job data (in flight jobs and output) */
+	FShaderJobDataMap InputHashToJobData;
+
+	/** List of duplicate jobs */
+	TArray<FShaderCommonCompileJob*> DuplicateJobs;
+
+	/** Statistics - total number of times we tried to Find() some input hash */
+	uint64 TotalSearchAttempts = 0;
+
+	/** Statistics - total number of times we succeded in Find()ing output for some input hash */
+	uint64 TotalCacheHits = 0;
+
+	/** Statistics - allocated memory. If the number is non-zero, we can trust it as accurate. Otherwise, recalculate. */
+	uint64 CurrentlyAllocatedMemory = 0;
+};
+
+static FShaderJobData& GetShaderJobData(const FShaderJobCacheRef& CacheRef)
+{
+	check(CacheRef.Block);
+	return CacheRef.Block->Data[CacheRef.IndexInBlock];
+}
+
+FShaderJobData* FShaderJobDataMap::Find(const FShaderJobData::FJobInputHash& Key)
+{
+	// Search for key with linear probing
+	for (uint32 TableIndex = GetTypeHash(Key) & HashTableMask; HashTable[TableIndex] != INDEX_NONE; TableIndex = (TableIndex + 1) & HashTableMask)
+	{
+		if ((*this)[TableIndex].InputHash == Key)
+		{
+			return &(*this)[TableIndex];
+		}
+	}
+	return nullptr;
+}
+
+FShaderJobCacheRef FShaderJobDataMap::FindOrAdd(const FShaderJobData::FJobInputHash& Key)
+{
+	// Search for key with linear probing
+	uint32 TableIndex;
+	for (TableIndex = GetTypeHash(Key) & HashTableMask; HashTable[TableIndex] != INDEX_NONE; TableIndex = (TableIndex + 1) & HashTableMask)
+	{
+		int32 ItemIndex = HashTable[TableIndex];
+		if ((*this)[ItemIndex].InputHash == Key)
+		{
+			return FShaderJobCacheRef({ &DataBlocks[ItemIndex / FShaderJobDataBlock::BlockSize], ItemIndex & (FShaderJobDataBlock::BlockSize - 1), INDEX_NONE });
+		}
+	}
+
+	// Ensure there is space for item
+	Reserve(NumItems + 1);
+
+	// Initialize allocated item
+	int32 AllocatedIndex = NumItems++;
+	FShaderJobCacheRef AllocatedItem({ &DataBlocks[AllocatedIndex / FShaderJobDataBlock::BlockSize], AllocatedIndex & (FShaderJobDataBlock::BlockSize - 1), INDEX_NONE });
+	GetShaderJobData(AllocatedItem).InputHash = Key;
+
+	// Add to empty spot in hash table
+	HashTable[TableIndex] = AllocatedIndex;
+
+	return AllocatedItem;
+}
+
+void FShaderJobDataMap::ReHash(int32 HashTableSize)
+{
+	// Resize table and rehash
+	HashTable.SetNumUninitialized(HashTableSize);
+	memset(HashTable.GetData(), 0xff, HashTable.Num() * HashTable.GetTypeSize());
+	HashTableMask = HashTableSize - 1;
+
+	for (int32 OuterIndex = 0; OuterIndex < DataBlocks.Num(); OuterIndex++)
+	{
+		FShaderJobData* Data = DataBlocks[OuterIndex].Data;
+
+		for (int32 InnerIndex = 0; InnerIndex < FShaderJobDataBlock::BlockSize; InnerIndex++)
+		{
+			int32 Index = OuterIndex * FShaderJobDataBlock::BlockSize + InnerIndex;
+			if (Index >= NumItems)
+			{
+				OuterIndex = DataBlocks.Num();
+				break;
+			}
+
+			// Find table entry for key -- keys will be unique when rehashing, so we don't need to check for existing keys
+			for (uint32 TableIndex = GetTypeHash(Data[InnerIndex].InputHash) & HashTableMask;; TableIndex = (TableIndex + 1) & HashTableMask)
+			{
+				if (HashTable[TableIndex] == INDEX_NONE)
+				{
+					HashTable[TableIndex] = Index;
+					break;
+				}
+			}
+		}
+	}
+}
+
+void FShaderJobDataMap::Reserve(int32 NumReserve)
+{
+	if (NumReserve > DataBlocks.Num() * FShaderJobDataBlock::BlockSize)
+	{
+		while (NumReserve > DataBlocks.Num() * FShaderJobDataBlock::BlockSize)
+		{
+			DataBlocks.Add(new FShaderJobDataBlock);
+		}
+
+		int32 HashTableSize = GetDesiredHashTableSize();
+		if (HashTableSize != HashTable.Num())
+		{
+			ReHash(HashTableSize);
+		}
+	}
+}
+
+void FShaderJobCache::CullOutputsToMemoryBudget(uint64 TargetBudgetBytes)
+{
+	// Track consecutive empty items.  We can delete empty blocks from the front of the map at the end.
+	int32 ConsecutiveEmptyItems = 0;
+	uint64 EmptyBlockSavings = 0;
+
+	// We don't cull items from the last block
+
+	for (int32 ItemIndex = 0; ItemIndex < InputHashToJobData.Num(); ItemIndex++)
+	{
+		FShaderJobData& JobData = InputHashToJobData[ItemIndex];
+
+		// Check if we are in budget yet
+		if (CurrentlyAllocatedMemory - EmptyBlockSavings <= TargetBudgetBytes)
+		{
+			break;
+		}
+
+		// We can only free this output if there is no in-flight job
+		if (JobData.JobInFlight == nullptr)
+		{
+			// Empty this item out (if not already empty), by removing the reference to the output and zeroing it out
+			if (!JobData.OutputHash.IsZero())
+			{
+				FStoredOutput** FoundStoredOutput = Outputs.Find(JobData.OutputHash);
+
+				if (FoundStoredOutput)
+				{
+					FStoredOutput* StoredOutput = *FoundStoredOutput;
+					checkf(StoredOutput, TEXT("Invalid entry found in FShaderJobCache Output hash table. All values are expected to be valid pointers."));
+
+					const uint64 OutputSize = StoredOutput->GetAllocatedSize();
+
+					// Decrement reference count and remove cached object if it's no longer referenced by any input hashes
+					if (StoredOutput->Release() == 0)
+					{
+						Outputs.Remove(JobData.OutputHash);
+						CurrentlyAllocatedMemory -= OutputSize;
+					}
+				}
+
+				JobData.OutputHash.Reset();
+			}
+
+			// Track if this is another consecutive empty item
+			if (ItemIndex == ConsecutiveEmptyItems)
+			{
+				ConsecutiveEmptyItems++;
+
+				// Take into account that we will be removing empty job data blocks at the end, by adding the savings when we reach a full block
+				if ((ConsecutiveEmptyItems & (FShaderJobDataBlock::BlockSize - 1)) == 0)
+				{
+					EmptyBlockSavings += sizeof(FShaderJobDataBlock);
+				}
+			}
+		}
+	}
+
+	int32 ConsecutiveEmptyBlocks = ConsecutiveEmptyItems / FShaderJobDataBlock::BlockSize;
+	if (ConsecutiveEmptyBlocks > 0)
+	{
+		uint64 InputHashToJobDataOriginalSize = InputHashToJobData.GetAllocatedSize();
+
+		InputHashToJobData.RemoveLeadingBlocks(ConsecutiveEmptyBlocks);
+
+		CurrentlyAllocatedMemory += InputHashToJobData.GetAllocatedSize() - InputHashToJobDataOriginalSize;
+	}
+}
+
+
 #if WITH_EDITOR
 static FDelayedAutoRegisterHelper GKickOffShaderAutoGen(EDelayedRegisterRunPhase::ShaderTypesReady, []
 {
@@ -363,18 +887,73 @@ namespace ShaderCompiler
 
 FShaderCompileJobCollection::FShaderCompileJobCollection()
 {
-	FMemory::Memzero(PendingJobs);
-	FMemory::Memzero(NumPendingJobs);
-
 	PrintStatsCmd = IConsoleManager::Get().RegisterConsoleCommand(
 		TEXT("r.ShaderCompiler.PrintStats"),
 		TEXT("Prints out to the log the stats for the shader compiler."),
 		FConsoleCommandDelegate::CreateRaw(this, &FShaderCompileJobCollection::HandlePrintStats),
 		ECVF_Default
 	);
+
+	JobsCache = MakePimpl<FShaderJobCache>();
 }
 
-void FShaderCompileJobCollection::InternalAddJob(FShaderCommonCompileJob* InJob)
+
+// Pass through functions to inner FShaderJobCache implementation class
+FShaderCompileJob* FShaderCompileJobCollection::PrepareJob(uint32 InId, const FShaderCompileJobKey& InKey, EShaderCompileJobPriority InPriority)
+{
+	return JobsCache->PrepareJob(InId, InKey, InPriority);
+}
+FShaderPipelineCompileJob* FShaderCompileJobCollection::PrepareJob(uint32 InId, const FShaderPipelineCompileJobKey& InKey, EShaderCompileJobPriority InPriority)
+{
+	return JobsCache->PrepareJob(InId, InKey, InPriority);
+}
+void FShaderCompileJobCollection::RemoveJob(FShaderCommonCompileJob* InJob)
+{
+	JobsCache->RemoveJob(InJob);
+}
+int32 FShaderCompileJobCollection::RemoveAllPendingJobsWithId(uint32 InId)
+{
+	return JobsCache->RemoveAllPendingJobsWithId(InId);
+}
+void FShaderCompileJobCollection::SubmitJobs(const TArray<FShaderCommonCompileJobPtr>& InJobs)
+{
+	JobsCache->SubmitJobs(InJobs);
+}
+void FShaderCompileJobCollection::ProcessFinishedJob(FShaderCommonCompileJob* FinishedJob, bool bWasCached)
+{
+	JobsCache->ProcessFinishedJob(FinishedJob, bWasCached);
+}
+void FShaderCompileJobCollection::AddToCacheAndProcessPending(FShaderCommonCompileJob* FinishedJob)
+{
+	JobsCache->AddToCacheAndProcessPending(FinishedJob);
+}
+void FShaderCompileJobCollection::LogCachingStats()
+{
+	JobsCache->LogCachingStats();
+}
+void FShaderCompileJobCollection::GatherAnalytics(const FString& BaseName, TArray<FAnalyticsEventAttribute>& Attributes) const
+{
+	JobsCache->GatherAnalytics(BaseName, Attributes);
+}
+int32 FShaderCompileJobCollection::GetNumPendingJobs(EShaderCompileJobPriority InPriority) const
+{
+	return JobsCache->GetNumPendingJobs(InPriority);
+}
+int32 FShaderCompileJobCollection::GetNumOutstandingJobs() const
+{
+	return JobsCache->GetNumOutstandingJobs();
+}
+int32 FShaderCompileJobCollection::GetNumPendingJobs() const
+{
+	return JobsCache->GetNumPendingJobs();
+}
+int32 FShaderCompileJobCollection::GetPendingJobs(EShaderCompilerWorkerType InWorkerType, EShaderCompileJobPriority InPriority, int32 MinNumJobs, int32 MaxNumJobs, TArray<FShaderCommonCompileJobPtr>& OutJobs)
+{
+	return JobsCache->GetPendingJobs(InWorkerType, InPriority, MinNumJobs, MaxNumJobs, OutJobs);
+}
+
+
+void FShaderJobCache::InternalAddJob(FShaderCommonCompileJob* InJob)
 {
 	const int32 TypeIndex = (int32)InJob->Type;
 
@@ -444,7 +1023,7 @@ static FShaderCommonCompileJob* CloneJob(const FShaderCommonCompileJob* SrcJob)
 	}
 }
 
-void FShaderCompileJobCollection::InternalSetPriority(FShaderCommonCompileJob* Job, EShaderCompileJobPriority InPriority)
+void FShaderJobCache::InternalSetPriority(FShaderCommonCompileJob* Job, EShaderCompileJobPriority InPriority)
 {
 	const int32 PriorityIndex = (int32)InPriority;
 
@@ -490,7 +1069,7 @@ void FShaderCompileJobCollection::InternalSetPriority(FShaderCommonCompileJob* J
 	}
 }
 
-void FShaderCompileJobCollection::InternalRemoveJob(FShaderCommonCompileJob* InJob)
+void FShaderJobCache::InternalRemoveJob(FShaderCommonCompileJob* InJob)
 {
 	const int32 TypeIndex = (int32)InJob->Type;
 	const int32 JobIndex = InJob->JobIndex;
@@ -506,17 +1085,17 @@ void FShaderCompileJobCollection::InternalRemoveJob(FShaderCommonCompileJob* InJ
 }
 
 
-void FShaderCompileJobCollection::RemoveJob(FShaderCommonCompileJob* InJob)
+void FShaderJobCache::RemoveJob(FShaderCommonCompileJob* InJob)
 {
-	FWriteScopeLock Locker(Lock);
+	FWriteScopeLock Locker(JobLock);
 	InternalRemoveJob(InJob);
 }
 
-int32 FShaderCompileJobCollection::RemoveAllPendingJobsWithId(uint32 InId)
+int32 FShaderJobCache::RemoveAllPendingJobsWithId(uint32 InId)
 {
 	int32 NumRemoved = 0;
 	{
-		FWriteScopeLock Locker(Lock);
+		FWriteScopeLock Locker(JobLock);
 		for (int32 PriorityIndex = 0; PriorityIndex < NumShaderCompileJobPriorities; ++PriorityIndex)
 		{
 			for (FShaderCommonCompileJob::TIterator It(PendingJobs[PriorityIndex]); It;)
@@ -528,7 +1107,13 @@ int32 FShaderCompileJobCollection::RemoveAllPendingJobsWithId(uint32 InId)
 				{
 					if (ShaderCompiler::IsJobCacheEnabled())
 					{
-						JobsInFlight.Remove(Job.GetInputHash());
+						if (Job.JobCacheRef.Block)
+						{
+							FShaderJobData& JobData = GetShaderJobData(Job.JobCacheRef);
+
+							check(JobData.JobInFlight == &Job);
+							JobData.JobInFlight = nullptr;
+						}
 					}
 
 					check(NumPendingJobs[PriorityIndex] > 0);
@@ -543,43 +1128,34 @@ int32 FShaderCompileJobCollection::RemoveAllPendingJobsWithId(uint32 InId)
 
 		if (ShaderCompiler::IsJobCacheEnabled())
 		{
-			// Also look into the jobs that are cached
-			// Since each entry in DuplicateJobsWaitList is a list, and the head node can be removed, we essentially have to rebuild it
-			for (TMap<FShaderCommonCompileJob::FInputHash, FShaderCommonCompileJob*>::TIterator Iter(DuplicateJobsWaitList); Iter; ++Iter)
+			// Also look into duplicate jobs that are cached -- we don't increment in the "for" loop because the current item may be deleted
+			for (int32 DuplicateIndex = 0; DuplicateIndex < DuplicateJobs.Num();)
 			{
-				FShaderCommonCompileJob* ListHead = Iter.Value();
-				FShaderCommonCompileJob* NewListHead = ListHead;
+				FShaderCommonCompileJob* DuplicateJob = DuplicateJobs[DuplicateIndex];
+				check(DuplicateJob->JobCacheRef.DuplicateIndex == DuplicateIndex);
 
-				// each entry in DJWL is a linked list of jobs that share the same ihash
-				for (FShaderCommonCompileJob::TIterator It(ListHead); It;)
+				if (DuplicateJob->Id == InId)
 				{
-					FShaderCommonCompileJob& Job = *It;
-					It.Next();
+					FShaderJobData& JobData = GetShaderJobData(DuplicateJob->JobCacheRef);
 
-					if (Job.Id == InId)
+					// if we're removing the list head, we need to update it to the next
+					if (JobData.DuplicateJobsWaitList == DuplicateJob)
 					{
-						// if we're removing the list head, we need to update the next
-						if (NewListHead == &Job)
-						{
-							NewListHead = Job.Next();
-						}
-
-						Job.Unlink();
-						Job.PendingPriority = EShaderCompileJobPriority::None;
-						InternalRemoveJob(&Job);
-						++NumRemoved;
+						JobData.DuplicateJobsWaitList = JobData.DuplicateJobsWaitList->Next();
 					}
-				}
 
-				if (NewListHead == nullptr)
-				{
-					// we removed the last job for this hash
-					Iter.RemoveCurrent();
+					DuplicateJob->Unlink();
+					DuplicateJob->PendingPriority = EShaderCompileJobPriority::None;
+					InternalRemoveJob(DuplicateJob);
+					++NumRemoved;
+
+					// This removes the current job (at DuplicateIndex), so we don't increment in this case
+					RemoveDuplicateJob(DuplicateJob);
 				}
-				else if (NewListHead != ListHead)
+				else
 				{
-					// update the mapping
-					Iter.Value() = NewListHead;
+					// Didn't remove a job, increment!
+					DuplicateIndex++;
 				}
 			}
 		}
@@ -590,7 +1166,7 @@ int32 FShaderCompileJobCollection::RemoveAllPendingJobsWithId(uint32 InId)
 	return NumRemoved;
 }
 
-void FShaderCompileJobCollection::SubmitJobs(const TArray<FShaderCommonCompileJobPtr>& InJobs)
+void FShaderJobCache::SubmitJobs(const TArray<FShaderCommonCompileJobPtr>& InJobs)
 {
 	if (InJobs.Num() > 0)
 	{
@@ -612,7 +1188,7 @@ void FShaderCompileJobCollection::SubmitJobs(const TArray<FShaderCommonCompileJo
 				}, EParallelForFlags::Unbalanced);
 			}
 
-			FWriteScopeLock Locker(Lock);
+			FWriteScopeLock Locker(JobLock);
 
 			// Optimization: Only search the linked list once to prevent O(n^2) behavior
 			FShaderCommonCompileJob* ExistingJobs[NumShaderCompileJobPriorities] = { nullptr };
@@ -632,8 +1208,11 @@ void FShaderCompileJobCollection::SubmitJobs(const TArray<FShaderCommonCompileJo
 
 					const bool bCheckDDC = !(Job->bIsDefaultMaterial || Job->bIsGlobalShader);
 
-					// see if we can find the job in the cache first
-					if (FSharedBuffer* ExistingOutput = CompletedJobsCache.Find(InputHash, bCheckDDC))
+					FSharedBuffer* ExistingOutput;
+					FShaderJobCacheRef JobCacheRef = FindOrAdd(InputHash, bCheckDDC, ExistingOutput);
+
+					// see if there are already cached results for this job
+					if (ExistingOutput)
 					{
 						UE_LOG(LogShaderCompilers, UE_SHADERCACHE_LOG_LEVEL, TEXT("There is already a cached job with the ihash %s, processing the new one immediately."), *LexToString(InputHash));
 						FMemoryReaderView MemReader(*ExistingOutput);
@@ -644,27 +1223,34 @@ void FShaderCompileJobCollection::SubmitJobs(const TArray<FShaderCommonCompileJo
 
 						continue;
 					}
-					// see if another job with the same input hash is being worked on
-					else if (FShaderCommonCompileJob** DuplicateInFlight = JobsInFlight.Find(InputHash))
+					else
 					{
-						UE_LOG(LogShaderCompilers, UE_SHADERCACHE_LOG_LEVEL, TEXT("There is an outstanding job with the ihash %s, not submitting another one (adding to wait list)."), *LexToString(InputHash));
+						FShaderJobData& JobData = GetShaderJobData(JobCacheRef);
+						Job->JobCacheRef = JobCacheRef;
 
-						// because of the cloned jobs, we need to maintain a separate mapping
-						FShaderCommonCompileJob** WaitListHead = DuplicateJobsWaitList.Find(InputHash);
-						if (WaitListHead)
+						// see if another job with the same input hash is being worked on
+						if (JobData.JobInFlight)
 						{
-							Job->LinkAfter(*WaitListHead);
+							UE_LOG(LogShaderCompilers, UE_SHADERCACHE_LOG_LEVEL, TEXT("There is an outstanding job with the ihash %s, not submitting another one (adding to wait list)."), *LexToString(InputHash));
+
+							// because of the cloned jobs, we need to maintain a separate mapping
+							FShaderCommonCompileJob** WaitListHead = &JobData.DuplicateJobsWaitList;
+							if (*WaitListHead)
+							{
+								Job->LinkAfter(*WaitListHead);
+							}
+							else
+							{
+								*WaitListHead = Job;
+							}
+							AddDuplicateJob(Job);
+							bNewJob = false;
 						}
 						else
 						{
-							DuplicateJobsWaitList.Add(InputHash, Job);
+							// track new jobs so we can dedupe them
+							JobData.JobInFlight = Job;
 						}
-						bNewJob = false;
-					}
-					else
-					{
-						// track new jobs so we can dedupe them
-						JobsInFlight.Add(InputHash, Job);
 					}
 				}
 
@@ -732,7 +1318,7 @@ void FShaderCompileJobCollection::HandlePrintStats()
 }
 
 TRACE_DECLARE_INT_COUNTER(Shaders_Compiled, TEXT("Shaders/Compiled"));
-void FShaderCompileJobCollection::ProcessFinishedJob(FShaderCommonCompileJob* FinishedJob, bool bWasCached)
+void FShaderJobCache::ProcessFinishedJob(FShaderCommonCompileJob* FinishedJob, bool bWasCached)
 {
 	FinishedJob->OnComplete();
 
@@ -757,7 +1343,7 @@ void FShaderCompileJobCollection::ProcessFinishedJob(FShaderCommonCompileJob* Fi
 	}
 }
 
-void FShaderCompileJobCollection::AddToCacheAndProcessPending(FShaderCommonCompileJob* FinishedJob)
+void FShaderJobCache::AddToCacheAndProcessPending(FShaderCommonCompileJob* FinishedJob)
 {
 	if (!ShaderCompiler::IsJobCacheEnabled())
 	{
@@ -772,64 +1358,71 @@ void FShaderCompileJobCollection::AddToCacheAndProcessPending(FShaderCommonCompi
 	FinishedJob->SerializeOutput(Writer);
 
 	FSharedBuffer Buffer = MakeSharedBufferFromArray(MoveTemp(Output));
+
 	// TODO: reduce the scope - e.g. SerializeOutput and processing finished jobs can be moved out of it
-	FWriteScopeLock JobLocker(Lock);
+	FWriteScopeLock JobLocker(JobLock);
+
+	FShaderJobData& JobData = GetShaderJobData(FinishedJob->JobCacheRef);
 
 	// see if there are outstanding jobs that also need to be resolved
 	int32 NumOutstandingJobsWithSameHash = 0;
-	if (FShaderCommonCompileJob** WaitList = DuplicateJobsWaitList.Find(InputHash))
+	FShaderCommonCompileJob* CurHead = JobData.DuplicateJobsWaitList;
+	while (CurHead)
 	{
-		FShaderCommonCompileJob* CurHead = *WaitList;
-		while (CurHead)
-		{
-			checkf(CurHead != FinishedJob, TEXT("Job that is being added to cache was also on a waiting list! Error in bookkeeping."));
+		checkf(CurHead != FinishedJob, TEXT("Job that is being added to cache was also on a waiting list! Error in bookkeeping."));
 
-			FMemoryReaderView MemReader(Buffer);
-			CurHead->SerializeOutput(MemReader);
-			checkf(CurHead->bSucceeded == FinishedJob->bSucceeded, TEXT("Different success status for the job with the same ihash"));
+		FMemoryReaderView MemReader(Buffer);
+		CurHead->SerializeOutput(MemReader);
+		checkf(CurHead->bSucceeded == FinishedJob->bSucceeded, TEXT("Different success status for the job with the same ihash"));
 
-			// finish the job instantly
-			ProcessFinishedJob(CurHead, true);
-			++NumOutstandingJobsWithSameHash;
+		// finish the job instantly
+		ProcessFinishedJob(CurHead, true);
+		++NumOutstandingJobsWithSameHash;
 
-			CurHead = CurHead->Next();
-		}
+		RemoveDuplicateJob(CurHead);
 
-		// remove the waitlist head
-		DuplicateJobsWaitList.Remove(InputHash);
+		CurHead = CurHead->Next();
+	}
 
-		if (NumOutstandingJobsWithSameHash > 0)
-		{
-			UE_LOG(LogShaderCompilers, UE_SHADERCACHE_LOG_LEVEL, TEXT("Processed %d outstanding jobs with the same ihash %s."), NumOutstandingJobsWithSameHash, *LexToString(InputHash));
-		}
+	if (NumOutstandingJobsWithSameHash > 0)
+	{
+		// Clear the duplicate job list if we processed any jobs
+		JobData.DuplicateJobsWaitList = nullptr;
+
+		UE_LOG(LogShaderCompilers, UE_SHADERCACHE_LOG_LEVEL, TEXT("Processed %d outstanding jobs with the same ihash %s."), NumOutstandingJobsWithSameHash, *LexToString(InputHash));
 	}
 
 	if (FinishedJob->bSucceeded)
 	{
 		const bool bAddToDDC = !(FinishedJob->bIsDefaultMaterial || FinishedJob->bIsGlobalShader);
 		// we only cache jobs that succeded
-		CompletedJobsCache.AddJobOutput(FinishedJob, InputHash, Buffer, NumOutstandingJobsWithSameHash, bAddToDDC);
+		AddJobOutput(JobData, FinishedJob, InputHash, Buffer, NumOutstandingJobsWithSameHash, bAddToDDC);
 	}
 
-	// remove ourselves from the jobs in flight, if we were there (if this job is a cloned job it might not have been)
-	JobsInFlight.Remove(InputHash);
+	// remove ourselves from the jobs in flight
+	JobData.JobInFlight = nullptr;
+	FinishedJob->JobCacheRef.Clear();
 }
 
-void FShaderCompileJobCollection::LogCachingStats()
+void FShaderJobCache::LogCachingStats()
 {
-	FWriteScopeLock Locker(Lock);	// write lock because logging actually changes the cache state (in a minor way - updating the memory used - but still).
-	CompletedJobsCache.LogStats();
+	FWriteScopeLock Locker(JobLock);	// write lock because logging actually changes the cache state (in a minor way - updating the memory used - but still).
+	LogStats();
 }
 
-void FShaderCompileJobCollection::GatherAnalytics(const FString& BaseName, TArray<FAnalyticsEventAttribute>& Attributes) const
+int32 FShaderJobCache::GetNumPendingJobs(EShaderCompileJobPriority InPriority) const
 {
-	FWriteScopeLock Locker(Lock);
-	CompletedJobsCache.GatherAnalytics(BaseName, Attributes);
+	return NumPendingJobs[(int32)InPriority];
 }
 
-int32 FShaderCompileJobCollection::GetNumPendingJobs() const
+int32 FShaderJobCache::GetNumOutstandingJobs() const
+{
+	return NumOutstandingJobs.GetValue();
+}
+
+int32 FShaderJobCache::GetNumPendingJobs() const
 {
-	FReadScopeLock Locker(Lock);
+	FReadScopeLock Locker(JobLock);
 	int32 NumJobs = 0;
 	for (int32 i = 0; i < NumShaderCompileJobPriorities; ++i)
 	{
@@ -838,7 +1431,7 @@ int32 FShaderCompileJobCollection::GetNumPendingJobs() const
 	return NumJobs;
 }
 
-int32 FShaderCompileJobCollection::GetPendingJobs(EShaderCompilerWorkerType InWorkerType, EShaderCompileJobPriority InPriority, int32 MinNumJobs, int32 MaxNumJobs, TArray<FShaderCommonCompileJobPtr>& OutJobs)
+int32 FShaderJobCache::GetPendingJobs(EShaderCompilerWorkerType InWorkerType, EShaderCompileJobPriority InPriority, int32 MinNumJobs, int32 MaxNumJobs, TArray<FShaderCommonCompileJobPtr>& OutJobs)
 {
 	check(InWorkerType != EShaderCompilerWorkerType::None);
 	check(InPriority != EShaderCompileJobPriority::None);
@@ -846,7 +1439,7 @@ int32 FShaderCompileJobCollection::GetPendingJobs(EShaderCompilerWorkerType InWo
 	const int32 PriorityIndex = (int32)InPriority;
 	int32 NumPendingJobsOfPriority = 0;
 	{
-		FReadScopeLock Locker(Lock);
+		FReadScopeLock Locker(JobLock);
 		NumPendingJobsOfPriority = NumPendingJobs[PriorityIndex];
 	}
 
@@ -856,7 +1449,7 @@ int32 FShaderCompileJobCollection::GetPendingJobs(EShaderCompilerWorkerType InWo
 		return 0;
 	}
 
-	FWriteScopeLock Locker(Lock);
+	FWriteScopeLock Locker(JobLock);
 
 	// there was a time window before we checked and then acquired the write lock - make sure the number is still sufficient
 	NumPendingJobsOfPriority = NumPendingJobs[PriorityIndex];
@@ -910,12 +1503,12 @@ int32 FShaderCompileJobCollection::GetPendingJobs(EShaderCompilerWorkerType InWo
 	return NumJobs;
 }
 
-FShaderCompileJob* FShaderCompileJobCollection::PrepareJob(uint32 InId, const FShaderCompileJobKey& InKey, EShaderCompileJobPriority InPriority)
+FShaderCompileJob* FShaderJobCache::PrepareJob(uint32 InId, const FShaderCompileJobKey& InKey, EShaderCompileJobPriority InPriority)
 {
 	return InternalPrepareJob<FShaderCompileJob>(InId, InKey, InPriority);
 }
 
-FShaderPipelineCompileJob* FShaderCompileJobCollection::PrepareJob(uint32 InId, const FShaderPipelineCompileJobKey& InKey, EShaderCompileJobPriority InPriority)
+FShaderPipelineCompileJob* FShaderJobCache::PrepareJob(uint32 InId, const FShaderPipelineCompileJobKey& InKey, EShaderCompileJobPriority InPriority)
 {
 	return InternalPrepareJob<FShaderPipelineCompileJob>(InId, InKey, InPriority);
 }
@@ -8755,149 +9348,139 @@ namespace
 }
 #endif
 
-struct FShaderJobCacheStoredOutput
+FShaderJobCacheRef FShaderJobCache::FindOrAdd(const FJobInputHash& Hash, const bool bCheckDDC, FJobCachedOutput*& OutCachedOutput)
 {
-private:
-	/** How many times this output is referenced by the cached jobs */
-	int32 NumReferences = 0;
+	check(ShaderCompiler::IsJobCacheEnabled());
 
-public:
-
-	/** How many times this output has been returned as a cached result, no matter the input hash */
-	int32 NumHits = 0;
+	++TotalSearchAttempts;
+	TRACE_COUNTER_INCREMENT(Shaders_JobCacheSearchAttempts);
+	TRACE_CPUPROFILER_EVENT_SCOPE(FShaderJobCache::Find);
 
-	/** Canned output */
-	FSharedBuffer JobOutput;
+	OutCachedOutput = nullptr;
 
-	/** Path to where the cached debug info is stored. */
-	FString CachedDebugInfoPath;
+	uint64 InputHashToJobDataSize = InputHashToJobData.GetAllocatedSize();
 
-	/** Similar to FRefCountBase AddRef, but not atomic */
-	int32 AddRef()
-	{
-		++NumReferences;
+	FShaderJobCacheRef JobCacheRef = InputHashToJobData.FindOrAdd(Hash);
+	FShaderJobData& JobData = GetShaderJobData(JobCacheRef);
 
-		return NumReferences;
-	}
+	CurrentlyAllocatedMemory += InputHashToJobData.GetAllocatedSize() - InputHashToJobDataSize;
 
-	/** Similar to FRefCountBase Release, but not atomic */
-	int32 Release()
+	if (JobData.HasOutput())
 	{
-		checkf(NumReferences >= 0, TEXT("Attempting to release shader job cache output that was already released"));
-
-		--NumReferences;
-
-		const int32 RemainingNumReferences = NumReferences;
-
-		if (RemainingNumReferences == 0)
-		{
-			delete this;
-		}
+		++TotalCacheHits;
+		TRACE_COUNTER_INCREMENT(Shaders_JobCacheHits);
 
-		return RemainingNumReferences;
-	}
+		FStoredOutput** CannedOutput = Outputs.Find(JobData.OutputHash);
+		// we should not allow a dangling input to output mapping to exist
+		checkf(CannedOutput != nullptr, TEXT("Inconsistency in FShaderJobCache - cache record for ihash %s exists, but output cannot be found."), *LexToString(Hash));
+		// update the output hit count
+		(*CannedOutput)->NumHits++;
 
-	uint64 GetAllocatedSize() const
-	{
-		return static_cast<uint64>(JobOutput.GetSize() + sizeof(*this));
+		OutCachedOutput = &(*CannedOutput)->JobOutput;
 	}
-};
-
-FShaderJobCache::FJobCachedOutput* FShaderJobCache::Find(const FJobInputHash& Hash, const bool bCheckDDC)
-{
-	++TotalSearchAttempts;
-	TRACE_COUNTER_INCREMENT(Shaders_JobCacheSearchAttempts);
-	TRACE_CPUPROFILER_EVENT_SCOPE(FShaderJobCache::Find);
-
-	if (ShaderCompiler::IsJobCacheEnabled())
+#if WITH_EDITOR
+	else
 	{
-		FJobOutputHash* OutputHash = InputHashToOutput.Find(Hash);
-		if (OutputHash)
+		// If we didn't find it in memory search the DDC if it's enabled.
+		const bool bCachePerShaderDDC = IsShaderJobCacheDDCEnabled() && bCheckDDC;
+		if (bCachePerShaderDDC)
 		{
-			++TotalCacheHits;
-			TRACE_COUNTER_INCREMENT(Shaders_JobCacheHits);
+			FSharedBuffer Results;
 
-			FStoredOutput** CannedOutput = Outputs.Find(*OutputHash);
-			// we should not allow a dangling input to output mapping to exist
-			checkf(CannedOutput != nullptr, TEXT("Inconsistency in FShaderJobCache - cache record for ihash %s exists, but output cannot be found."), *LexToString(Hash));
-			// update the output hit count
-			(*CannedOutput)->NumHits++;
-			return &(*CannedOutput)->JobOutput;
-		}
-#if WITH_EDITOR
-		else
-		{
-			// If we didn't find it in memory search the DDC if it's enabled.
-			const bool bCachePerShaderDDC = IsShaderJobCacheDDCEnabled() && bCheckDDC;
-			if (bCachePerShaderDDC)
-			{
-				FSharedBuffer Results;
+			TRACE_COUNTER_INCREMENT(Shaders_JobCacheDDCRequests);
 
-				TRACE_COUNTER_INCREMENT(Shaders_JobCacheDDCRequests);
+			UE::DerivedData::FRequestOwner RequestOwner(UE::DerivedData::EPriority::Blocking);
 
-				UE::DerivedData::FRequestOwner RequestOwner(UE::DerivedData::EPriority::Blocking);
-
-				UE::DerivedData::FCacheGetRequest Request;
-				Request.Name = TEXT("FShaderJobCache");
-				// Create key.
-				Request.Key.Bucket = ShaderJobCacheDDCBucket;
-				Request.Key.Hash = Hash;
-				Request.Policy = IsShaderJobCacheDDCRemotePolicyEnabled() ? UE::DerivedData::ECachePolicy::Default : UE::DerivedData::ECachePolicy::Local;
-
-				UE::DerivedData::GetCache().Get(
-					{ Request },
-					RequestOwner,
-					[&Results](UE::DerivedData::FCacheGetResponse&& Response)
+			UE::DerivedData::FCacheGetRequest Request;
+			Request.Name = TEXT("FShaderJobCache");
+			// Create key.
+			Request.Key.Bucket = ShaderJobCacheDDCBucket;
+			Request.Key.Hash = Hash;
+			Request.Policy = IsShaderJobCacheDDCRemotePolicyEnabled() ? UE::DerivedData::ECachePolicy::Default : UE::DerivedData::ECachePolicy::Local;
+
+			UE::DerivedData::GetCache().Get(
+				{ Request },
+				RequestOwner,
+				[&Results](UE::DerivedData::FCacheGetResponse&& Response)
+				{
+					switch (Response.Status)
 					{
-						switch (Response.Status)
+						case UE::DerivedData::EStatus::Ok:
 						{
-							case UE::DerivedData::EStatus::Ok:
-							{
-								Results = Response.Record.GetValue(ShaderJobCacheId).GetData().Decompress();
+							Results = Response.Record.GetValue(ShaderJobCacheId).GetData().Decompress();
 
-								TRACE_COUNTER_INCREMENT(Shaders_JobCacheDDCHits);
+							TRACE_COUNTER_INCREMENT(Shaders_JobCacheDDCHits);
 
-								break;
-							}
-							case UE::DerivedData::EStatus::Error:
-							{
-								break;
-							}
-							case UE::DerivedData::EStatus::Canceled:
-							{
-								break;
-							}
+							break;
 						}
-					});
+						case UE::DerivedData::EStatus::Error:
+						{
+							break;
+						}
+						case UE::DerivedData::EStatus::Canceled:
+						{
+							break;
+						}
+					}
+				});
 
-				RequestOwner.Wait();
-				if (!Results.IsNull())
-				{
-					// Create a new entry to store in the FShaderJobCache
-					FStoredOutput* NewStoredOutput = new FStoredOutput();
-					NewStoredOutput->JobOutput = Results;
+			RequestOwner.Wait();
+			if (!Results.IsNull())
+			{
+				// Create a new entry to store in the FShaderJobCache
+				FStoredOutput* NewStoredOutput = new FStoredOutput();
+				NewStoredOutput->JobOutput = Results;
 
-					TRACE_COUNTER_ADD(Shaders_JobCacheDDCBytesReceived, Results.GetSize());
+				TRACE_COUNTER_ADD(Shaders_JobCacheDDCBytesReceived, Results.GetSize());
 
-					// Generate an output hash and cache the result in the FShaderJobCache
-					FJobOutputHash NewOutputHash = FBlake3::HashBuffer(NewStoredOutput->JobOutput.GetData(), NewStoredOutput->JobOutput.GetSize());
+				// Generate an output hash and cache the result in the FShaderJobCache
+				FJobOutputHash NewOutputHash = FBlake3::HashBuffer(NewStoredOutput->JobOutput.GetData(), NewStoredOutput->JobOutput.GetSize());
 
-					NewStoredOutput->AddRef();
-					Outputs.Add(NewOutputHash, NewStoredOutput);
-					InputHashToOutput.Add(Hash, NewOutputHash);
-					EvictionQueue.PushLast(Hash);
+				NewStoredOutput->AddRef();
+				Outputs.Add(NewOutputHash, NewStoredOutput);
+				JobData.OutputHash = NewOutputHash;
 
-					CurrentlyAllocatedMemory += NewStoredOutput->GetAllocatedSize();
+				CurrentlyAllocatedMemory += NewStoredOutput->GetAllocatedSize();
 
-					// return the results.
-					return &NewStoredOutput->JobOutput;
-				}
+				// return the results.
+				OutCachedOutput = &NewStoredOutput->JobOutput;
 			}
 		}
-#endif
 	}
+#endif
 
-	return nullptr;
+	return JobCacheRef;
+}
+
+FShaderJobData* FShaderJobCache::Find(const FJobInputHash& Hash)
+{
+	check(ShaderCompiler::IsJobCacheEnabled());
+
+	return InputHashToJobData.Find(Hash);
+}
+
+/** Adds a reference to a duplicate job (to the DuplicateJobs array) */
+void FShaderJobCache::AddDuplicateJob(FShaderCommonCompileJob* DuplicateJob)
+{
+	check(DuplicateJob->JobCacheRef.DuplicateIndex == INDEX_NONE);
+
+	DuplicateJob->JobCacheRef.DuplicateIndex = DuplicateJobs.Add(DuplicateJob);
+}
+
+/** Removes a reference to a duplicate job (from the DuplicateJobs array)  */
+void FShaderJobCache::RemoveDuplicateJob(FShaderCommonCompileJob* DuplicateJob)
+{
+	int32 DuplicateIndex = DuplicateJob->JobCacheRef.DuplicateIndex;
+	check(DuplicateIndex >= 0 && DuplicateIndex < DuplicateJobs.Num() && DuplicateJobs[DuplicateIndex] == DuplicateJob);
+	DuplicateJob->JobCacheRef.DuplicateIndex = INDEX_NONE;
+
+	DuplicateJobs.RemoveAtSwap(DuplicateIndex);
+
+	// After removing, we need to update the cached index of the job we swapped
+	if (DuplicateIndex < DuplicateJobs.Num())
+	{
+		DuplicateJobs[DuplicateIndex]->JobCacheRef.DuplicateIndex = DuplicateIndex;
+	}
 }
 
 uint64 FShaderJobCache::GetCurrentMemoryBudget() const
@@ -8909,7 +9492,10 @@ uint64 FShaderJobCache::GetCurrentMemoryBudget() const
 
 FShaderJobCache::FShaderJobCache()
 {
-	CurrentlyAllocatedMemory = sizeof(*this) + InputHashToOutput.GetAllocatedSize() + Outputs.GetAllocatedSize();
+	FMemory::Memzero(PendingJobs);
+	FMemory::Memzero(NumPendingJobs);
+
+	CurrentlyAllocatedMemory = sizeof(*this) + InputHashToJobData.GetAllocatedSize() + Outputs.GetAllocatedSize();
 }
 
 FShaderJobCache::~FShaderJobCache()
@@ -8920,7 +9506,7 @@ FShaderJobCache::~FShaderJobCache()
 	}
 }
 
-void FShaderJobCache::AddJobOutput(const FShaderCommonCompileJob* FinishedJob, const FJobInputHash& Hash, const FJobCachedOutput& Contents, int32 InitialHitCount, const bool bAddToDDC)
+void FShaderJobCache::AddJobOutput(FShaderJobData& JobData, const FShaderCommonCompileJob* FinishedJob, const FJobInputHash& Hash, const FJobCachedOutput& Contents, int32 InitialHitCount, const bool bAddToDDC)
 {
 	TRACE_CPUPROFILER_EVENT_SCOPE(FShaderJobCache::Add);
 
@@ -8929,17 +9515,16 @@ void FShaderJobCache::AddJobOutput(const FShaderCommonCompileJob* FinishedJob, c
 		return;
 	}
 
-	FJobOutputHash* ExistingOutputHash = InputHashToOutput.Find(Hash);
-	if (ExistingOutputHash && !ShaderCompiler::IsJobCacheDebugValidateEnabled())
+	if (JobData.HasOutput() && !ShaderCompiler::IsJobCacheDebugValidateEnabled())
 	{
 		return;
 	}
 
 	FJobOutputHash OutputHash = FBlake3::HashBuffer(Contents.GetData(), Contents.GetSize());
 
-	if (ExistingOutputHash && ShaderCompiler::IsJobCacheDebugValidateEnabled())
+	if (JobData.HasOutput() && ShaderCompiler::IsJobCacheDebugValidateEnabled())
 	{
-		if (OutputHash != *ExistingOutputHash)
+		if (OutputHash != JobData.OutputHash)
 		{
 			FString OriginalFilename = 
 				GShaderCompilingManager->GetAbsoluteShaderDebugInfoDirectory() / 
@@ -8960,7 +9545,7 @@ void FShaderJobCache::AddJobOutput(const FShaderCommonCompileJob* FinishedJob, c
 				TArrayView<const uint8>((const uint8*)Contents.GetData(), Contents.GetSize()), 
 				*NewFilename);
 
-			FStoredOutput** StoredOutput = Outputs.Find(*ExistingOutputHash);
+			FStoredOutput** StoredOutput = Outputs.Find(JobData.OutputHash);
 			check(StoredOutput);
 
 			FFileHelper::SaveArrayToFile(
@@ -8987,21 +9572,30 @@ void FShaderJobCache::AddJobOutput(const FShaderCommonCompileJob* FinishedJob, c
 		}
 	}
 
-	// add the record
-	const uint64 InputHashToOutputOriginalSize = InputHashToOutput.GetAllocatedSize();
-
-	InputHashToOutput.Add(Hash, OutputHash);
-	EvictionQueue.PushLast(Hash);
+	// Cache this value for thread safety
+	bool bDiscardCacheOutputs = GShaderCompilerDebugDiscardCacheOutputs != 0;
 
-	const bool bCachePerShaderDDC = IsShaderJobCacheDDCEnabled() && bAddToDDC;
+	// add the record
+	if (UNLIKELY(bDiscardCacheOutputs == false))
+	{
+		JobData.OutputHash = OutputHash;
+	}
 
 	FStoredOutput** CannedOutput = Outputs.Find(OutputHash);
-	if (CannedOutput && (bCachePerShaderDDC == false))
+	if (CannedOutput)
 	{
 		// update the output hit count
-		const int32 NumRef = (*CannedOutput)->AddRef();
+		int32 NumRef;
+		if (UNLIKELY(bDiscardCacheOutputs == false))
+		{
+			NumRef = (*CannedOutput)->AddRef();
+		}
+		else
+		{
+			NumRef = (*CannedOutput)->GetNumReferences();
+		}
 
-		if (bDumpCachedDebugInfo)
+		if (UNLIKELY(bDumpCachedDebugInfo))
 		{
 			// Write cache hit debug file
 			const FString& CachedDebugInfoPath = (*CannedOutput)->CachedDebugInfoPath;
@@ -9015,109 +9609,93 @@ void FShaderJobCache::AddJobOutput(const FShaderCommonCompileJob* FinishedJob, c
 	}
 	else
 	{
-		const uint64 OutputsOriginalSize = Outputs.GetAllocatedSize();
+		if (UNLIKELY(bDiscardCacheOutputs == false))
+		{
+			const uint64 OutputsOriginalSize = Outputs.GetAllocatedSize();
 
-		FStoredOutput* NewStoredOutput = new FStoredOutput();
-		NewStoredOutput->NumHits = InitialHitCount;
-		NewStoredOutput->JobOutput = Contents;
-		NewStoredOutput->CachedDebugInfoPath = InputDebugInfoPath;
-		NewStoredOutput->AddRef();
-		Outputs.Add(OutputHash, NewStoredOutput);
+			FStoredOutput* NewStoredOutput = new FStoredOutput();
+			NewStoredOutput->NumHits = InitialHitCount;
+			NewStoredOutput->JobOutput = Contents;
+			NewStoredOutput->CachedDebugInfoPath = InputDebugInfoPath;
+			NewStoredOutput->AddRef();
+			Outputs.Add(OutputHash, NewStoredOutput);
 
-		CurrentlyAllocatedMemory += NewStoredOutput->GetAllocatedSize();
+			CurrentlyAllocatedMemory += NewStoredOutput->GetAllocatedSize() + Outputs.GetAllocatedSize() - OutputsOriginalSize;
+		}
 
-		if (bDumpCachedDebugInfo)
+		if (UNLIKELY(bDumpCachedDebugInfo))
 		{
 			// Write new allocated cache file
 			if (!InputDebugInfoPath.IsEmpty())
 			{
 				const FString CacheFilename = FString::Printf(TEXT("%s/%s.bytecode"), *InputDebugInfoPath, *InputSourceFilename);
-				FFileHelper::SaveArrayToFile(TArrayView<const uint8>((const uint8*)NewStoredOutput->JobOutput.GetData(), NewStoredOutput->JobOutput.GetSize()), *CacheFilename);
+				FFileHelper::SaveArrayToFile(TArrayView<const uint8>((const uint8*)Contents.GetData(), Contents.GetSize()), *CacheFilename);
 			}
 		}
 
-#if WITH_EDITOR
-		if (bCachePerShaderDDC)
-		{
-			// Create key.
-			UE::DerivedData::FCacheKey Key;
-			Key.Bucket = ShaderJobCacheDDCBucket;
-			Key.Hash = Hash;
-			UE::DerivedData::FCacheRecordBuilder RecordBuilder(Key);
-
-			RecordBuilder.AddValue(ShaderJobCacheId, FSharedBuffer::MakeView(Contents));
-
-			TRACE_COUNTER_ADD(Shaders_JobCacheDDCBytesSent, Contents.GetSize());
-
-			UE::DerivedData::FRequestOwner RequestOwner(UE::DerivedData::EPriority::Normal);
-			UE::DerivedData::FRequestBarrier RequestBarrier(RequestOwner);
-			RequestOwner.KeepAlive();
-			UE::DerivedData::GetCache().Put(
-				{ {{TEXT("FShaderJobCache")}, RecordBuilder.Build(), IsShaderJobCacheDDCRemotePolicyEnabled() ? UE::DerivedData::ECachePolicy::Default : UE::DerivedData::ECachePolicy::Local } },
-				RequestOwner
-			);
-		}
-#endif
-
-		// delete the previous cache entries if we have a budget
-		const uint64 MemoryBudgetBytes = GetCurrentMemoryBudget();
+		// delete oldest cache entries if we exceed the budget
+		uint64 MemoryBudgetBytes = GetCurrentMemoryBudget();
 		if (MemoryBudgetBytes)
 		{
 			if (CurrentlyAllocatedMemory > MemoryBudgetBytes)
 			{
 				TRACE_CPUPROFILER_EVENT_SCOPE(FShaderJobCache::Trim);
-			
-				FJobInputHash KeyToRemove = {};
-				while (CurrentlyAllocatedMemory > MemoryBudgetBytes && EvictionQueue.TryPopFirst(KeyToRemove))
-				{
-					RemoveByInputHash(KeyToRemove);
-				}
+
+				uint64 TargetBudgetBytes = MemoryBudgetBytes * FMath::Clamp(GShaderCompilerJobCacheOverflowReducePercent, 0, 100) / 100;
+				uint64 MemoryBefore = CurrentlyAllocatedMemory;
+
+				// Cull outputs to reach the budget target
+				CullOutputsToMemoryBudget(TargetBudgetBytes);
+
+				UE_LOG(LogShaderCompilers, Log, TEXT("Memory overflow, reduced from %.1lf to %.1lf MB."), (double)MemoryBefore / (1024 * 1024), (double)CurrentlyAllocatedMemory / (1024 * 1024));
 			}
 		}
-
-		CurrentlyAllocatedMemory += Outputs.GetAllocatedSize();
-		CurrentlyAllocatedMemory -= OutputsOriginalSize;
 	}
 
-	CurrentlyAllocatedMemory += InputHashToOutput.GetAllocatedSize();
-	CurrentlyAllocatedMemory -= InputHashToOutputOriginalSize;
-}
-
-void FShaderJobCache::RemoveByInputHash(const FJobInputHash& InputHash)
-{
-	const FJobOutputHash* FoundOutputHash = InputHashToOutput.Find(InputHash);
+#if WITH_EDITOR
+	const bool bCachePerShaderDDC = IsShaderJobCacheDDCEnabled() && bAddToDDC;
 
-	if (FoundOutputHash)
+	if (bCachePerShaderDDC)
 	{
-		const FJobOutputHash& OutputHash = *FoundOutputHash;
-		FStoredOutput** FoundStoredOutput = Outputs.Find(OutputHash);
-
-		if (FoundStoredOutput)
-		{
-			FStoredOutput* StoredOutput = *FoundStoredOutput;
-			checkf(StoredOutput, TEXT("Invalid entry found in FShaderJobCache Output hash table. All values are expected to be valid pointers."));
+		// Create key.
+		UE::DerivedData::FCacheKey Key;
+		Key.Bucket = ShaderJobCacheDDCBucket;
+		Key.Hash = Hash;
+		UE::DerivedData::FCacheRecordBuilder RecordBuilder(Key);
 
-			const uint64 OutputSize = StoredOutput->GetAllocatedSize();
+		RecordBuilder.AddValue(ShaderJobCacheId, FSharedBuffer::MakeView(Contents));
 
-			// Decrement reference count and remove cached object if it's no longer referenced by any input hashes
-			if (StoredOutput->Release() == 0)
-			{
-				Outputs.Remove(OutputHash);
-				CurrentlyAllocatedMemory -= OutputSize;
-			}
-		}
+		TRACE_COUNTER_ADD(Shaders_JobCacheDDCBytesSent, Contents.GetSize());
 
-		InputHashToOutput.Remove(InputHash);
+		UE::DerivedData::FRequestOwner RequestOwner(UE::DerivedData::EPriority::Normal);
+		UE::DerivedData::FRequestBarrier RequestBarrier(RequestOwner);
+		RequestOwner.KeepAlive();
+		UE::DerivedData::GetCache().Put(
+			{ {{TEXT("FShaderJobCache")}, RecordBuilder.Build(), IsShaderJobCacheDDCRemotePolicyEnabled() ? UE::DerivedData::ECachePolicy::Default : UE::DerivedData::ECachePolicy::Local } },
+			RequestOwner
+		);
 	}
+#endif
 }
 
-/** Calculates memory used by the cache*/
+#include "Math/UnitConversion.h"
+
+/** Returns memory used by the cache*/
 uint64 FShaderJobCache::GetAllocatedMemory() const
 {
 	return CurrentlyAllocatedMemory;
 }
 
-#include "Math/UnitConversion.h"
+/** Compute memory used by the cache from scratch.  Should match GetAllocatedMemory() if CurrentlyAllocateMemory is being properly updated. */
+uint64 FShaderJobCache::ComputeAllocatedMemory() const
+{
+	uint64 AllocatedSize = sizeof(FShaderJobCache) + InputHashToJobData.GetAllocatedSize() + Outputs.GetAllocatedSize();
+	for (auto OutputIter : Outputs)
+	{
+		AllocatedSize += OutputIter.Value->GetAllocatedSize();
+	}
+	return AllocatedSize;
+}
 
 /** Logs out the statistics */
 void FShaderJobCache::LogStats()
@@ -9129,9 +9707,9 @@ void FShaderJobCache::LogStats()
 		(TotalSearchAttempts > 0) ? 100.0 * static_cast<double>(TotalCacheHits) / static_cast<double>(TotalSearchAttempts) : 0.0);
 
 	UE_LOG(LogShaderCompilers, Display, TEXT("Tracking %s distinct input hashes that result in %s distinct outputs (%.2f%%)"),
-		*FormatNumber(InputHashToOutput.Num()), 
+		*FormatNumber(InputHashToJobData.Num()), 
 		*FormatNumber(Outputs.Num()),
-		(InputHashToOutput.Num() > 0) ? 100.0 * static_cast<double>(Outputs.Num()) / static_cast<double>(InputHashToOutput.Num()) : 0.0);
+		(InputHashToJobData.Num() > 0) ? 100.0 * static_cast<double>(Outputs.Num()) / static_cast<double>(InputHashToJobData.Num()) : 0.0);
 
 	static const FNumberFormattingOptions SizeFormattingOptions = FNumberFormattingOptions().SetMinimumFractionalDigits(2).SetMaximumFractionalDigits(2);
 
@@ -9152,6 +9730,7 @@ void FShaderJobCache::LogStats()
 
 void FShaderJobCache::GatherAnalytics(const FString& BaseName, TArray<FAnalyticsEventAttribute>& Attributes) const
 {
+	FWriteScopeLock Locker(JobLock);
 	const FString ChildName = TEXT("JobCache_");
 
 	{
diff --git a/Engine/Source/Runtime/Engine/Public/ShaderCompiler.h b/Engine/Source/Runtime/Engine/Public/ShaderCompiler.h
--- a/Engine/Source/Runtime/Engine/Public/ShaderCompiler.h
+++ b/Engine/Source/Runtime/Engine/Public/ShaderCompiler.h
@@ -45,61 +45,7 @@ struct FAnalyticsEventAttribute;
 extern ENGINE_API bool IsShaderJobCacheDDCEnabled();
 
 struct FShaderJobCacheStoredOutput;
-
-class FShaderJobCache
-{
-public:
-
-	FShaderJobCache();
-	~FShaderJobCache();
-
-	using FJobInputHash = FShaderCommonCompileJob::FInputHash;
-	using FJobCachedOutput = FSharedBuffer;
-
-	/** Looks for the job in the cache, returns null if not found */
-	FJobCachedOutput* Find(const FJobInputHash& Hash, const bool bCheckDDC);
-
-	/** Adds a job output to the cache */
-	void AddJobOutput(const FShaderCommonCompileJob* FinishedJob, const FJobInputHash& Hash, const FJobCachedOutput& Contents, int InitialHitCount, const bool bAddToDDC);
-
-	/** Calculates memory used by the cache*/
-	uint64 GetAllocatedMemory() const;
-
-	/** Logs out the statistics */
-	void LogStats();
-
-	/** Gather statistics to send to analytics */
-	void GatherAnalytics(const FString& BaseName, TArray<FAnalyticsEventAttribute>& Attributes) const;
-
-	/** Calculates current memory budget, in bytes */
-	uint64 GetCurrentMemoryBudget() const;
-
-private:
-
-	using FJobOutputHash = FBlake3Hash;
-	using FStoredOutput = FShaderJobCacheStoredOutput;
-
-	void RemoveByInputHash(const FJobInputHash& InputHash);
-
-	/* A lot of outputs can be duplicated, so they are deduplicated before storing */
-	TMap<FJobOutputHash, FStoredOutput*> Outputs;
-
-	/** Map of input hashes to output hashes */
-	TMap<FJobInputHash, FJobOutputHash> InputHashToOutput;
-
-	/** Queue to evict oldest elements when memory budget is exceeded */
-	TDeque<FJobInputHash> EvictionQueue;
-
-	/** Statistics - total number of times we tried to Find() some input hash */
-	uint64 TotalSearchAttempts = 0;
-
-	/** Statistics - total number of times we succeded in Find()ing output for some input hash */
-	uint64 TotalCacheHits = 0;
-
-	/** Statistics - allocated memory. If the number is non-zero, we can trust it as accurate. Otherwise, recalculate. */
-	uint64 CurrentlyAllocatedMemory = 0;
-};
-
+class FShaderJobCache;
 
 class FShaderCompileJobCollection
 {
@@ -125,119 +71,20 @@ public:
 
 	void GatherAnalytics(const FString& BaseName, TArray<FAnalyticsEventAttribute>& Attributes) const;
 
-	inline int32 GetNumPendingJobs(EShaderCompileJobPriority InPriority) const
-	{
-		return NumPendingJobs[(int32)InPriority];
-	}
+	int32 GetNumPendingJobs(EShaderCompileJobPriority InPriority) const;
 
-	inline int32 GetNumOutstandingJobs() const
-	{
-		return NumOutstandingJobs.GetValue();
-	}
+	int32 GetNumOutstandingJobs() const;
 
 	int32 GetNumPendingJobs() const;
 
 	int32 GetPendingJobs(EShaderCompilerWorkerType InWorkerType, EShaderCompileJobPriority InPriority, int32 MinNumJobs, int32 MaxNumJobs, TArray<FShaderCommonCompileJobPtr>& OutJobs);
 
 private:
-	void InternalAddJob(FShaderCommonCompileJob* Job);
-	void InternalRemoveJob(FShaderCommonCompileJob* InJob);
-	void InternalSetPriority(FShaderCommonCompileJob* Job, EShaderCompileJobPriority InPriority);
-	// cannot allow managing this from outside as the caching logic is not exposed
-	inline int32 InternalSubtractNumOutstandingJobs(int32 Value)
-	{
-		const int32 PrevNumOutstandingJobs = NumOutstandingJobs.Subtract(Value);
-		check(PrevNumOutstandingJobs >= Value);
-		return PrevNumOutstandingJobs - Value;
-	}
-
-	template<typename JobType, typename KeyType>
-	int32 InternalFindJobIndex(uint32 InJobHash, uint32 InJobId, const KeyType& InKey) const
-	{
-		const int32 TypeIndex = (int32)JobType::Type;
-		uint32 CurrentPriorityIndex = 0u;
-		int32 CurrentIndex = INDEX_NONE;
-		for (int32 Index = JobHash[TypeIndex].First(InJobHash); JobHash[TypeIndex].IsValid(Index); Index = JobHash[TypeIndex].Next(Index))
-		{
-			const FShaderCommonCompileJob* Job = Jobs[TypeIndex][Index].GetReference();
-			check(Job->Type == JobType::Type);
-
-			// We find the job that matches the key with the highest priority
-			if (Job->Id == InJobId &&
-				(uint32)Job->Priority >= CurrentPriorityIndex &&
-				static_cast<const JobType*>(Job)->Key == InKey)
-			{
-				CurrentPriorityIndex = (uint32)Job->Priority;
-				CurrentIndex = Index;
-			}
-		}
-		return CurrentIndex;
-	}
-
-	template<typename JobType, typename KeyType>
-	JobType* InternalFindJob(uint32 InJobHash, uint32 InJobId, const KeyType& InKey) const
-	{
-		const int32 TypeIndex = (int32)JobType::Type;
-		const int32 JobIndex = InternalFindJobIndex<JobType>(InJobHash, InJobId, InKey);
-		return JobIndex != INDEX_NONE ? static_cast<JobType*>(Jobs[TypeIndex][JobIndex].GetReference()) : nullptr;
-	}
-
-	template<typename JobType, typename KeyType>
-	JobType* InternalPrepareJob(uint32 InId, const KeyType& InKey, EShaderCompileJobPriority InPriority)
-	{
-		const uint32 Hash = InKey.MakeHash(InId);
-		JobType* PrevJob = nullptr;
-		{
-			FReadScopeLock Locker(Lock);
-			PrevJob = InternalFindJob<JobType>(Hash, InId, InKey);
-		}
-
-		JobType* NewJob = nullptr;
-		if (PrevJob == nullptr || (uint32)InPriority > (uint32)PrevJob->Priority)
-		{
-			FWriteScopeLock Locker(Lock);
-			if (PrevJob == nullptr)
-			{
-				PrevJob = InternalFindJob<JobType>(Hash, InId, InKey);
-			}
-			if (PrevJob == nullptr)
-			{
-				NewJob = new JobType(Hash, InId, InPriority, InKey);
-				InternalAddJob(NewJob);
-			}
-			else if ((uint32)InPriority > (uint32)PrevJob->Priority)
-			{
-				InternalSetPriority(PrevJob, InPriority);
-			}
-		}
-
-		return NewJob;
-	}
-
 	/** Handles the console command to log shader compiler stats */
 	void HandlePrintStats();
 
-	/** Queue of tasks that haven't been assigned to a worker yet. */
-	FShaderCommonCompileJob* PendingJobs[NumShaderCompileJobPriorities];
-	int32 NumPendingJobs[NumShaderCompileJobPriorities];
-
-	/** Number of jobs currently being compiled.  This includes PendingJobs and any jobs that have been assigned to workers but aren't complete yet. */
-	FThreadSafeCounter NumOutstandingJobs;
-
-	TArray<FShaderCommonCompileJobPtr> Jobs[NumShaderCompileJobTypes];
-	TArray<int32> FreeIndices[NumShaderCompileJobTypes];
-	FHashTable JobHash[NumShaderCompileJobTypes];
-	/** Guards access to the above job storage and also the cache structures below - JobsInFlight, WaitList and the Cache itself */
-	mutable FRWLock Lock;
-
-	/** Map of input hash to the jobs that we decided to execute. Note that mapping will miss cloned jobs (to avoid being a multimap). */
-	TMap<FShaderCommonCompileJob::FInputHash, FShaderCommonCompileJob*> JobsInFlight;
-
-	/** Map of input hash to the jobs that we delayed because a job with the same hash was executing. Each job is a head of a linked list of jobs with the same input hash (ihash) */
-	TMap<FShaderCommonCompileJob::FInputHash, FShaderCommonCompileJob*> DuplicateJobsWaitList;
-
-	/** Cache for the completed jobs.*/
-	FShaderJobCache CompletedJobsCache;
+	/** Cache for in flight and completed jobs.*/
+	TPimplPtr<FShaderJobCache> JobsCache;
 
 	/** Debugging - console command to print stats. */
 	class IConsoleObject* PrintStatsCmd;
diff --git a/Engine/Source/Runtime/RenderCore/Public/ShaderCompilerJobTypes.h b/Engine/Source/Runtime/RenderCore/Public/ShaderCompilerJobTypes.h
--- a/Engine/Source/Runtime/RenderCore/Public/ShaderCompilerJobTypes.h
+++ b/Engine/Source/Runtime/RenderCore/Public/ShaderCompilerJobTypes.h
@@ -41,6 +41,34 @@ struct FPendingShaderMapCompileResults
 using FPendingShaderMapCompileResultsPtr = TRefCountPtr<FPendingShaderMapCompileResults>;
 
 
+/**
+ * Cached reference to the location of an in-flight job's FShaderJobData in the FShaderJobDataMap, used by the private FShaderJobCache class.
+ *
+ * Caching the reference avoids the need to do additional map lookups to find the entry again, potentially avoiding a lock of the container for
+ * the lookup.  Heap allocation of blocks is used by the cache to allow map entries to have a persistent location in memory.  The persistent
+ * memory allows modifications of map entry data for a given job, without needing locks to protect against container resizing.
+ *
+ * In-flight jobs and their duplicates reference the same FShaderJobData.  Client code should treat this structure as opaque.
+ */
+struct FShaderJobCacheRef
+{
+	/** Pointer to block the private FShaderJobData is stored in */
+	struct FShaderJobDataBlock* Block = nullptr;
+
+	/** Index of FShaderJobData in the block */
+	int32 IndexInBlock = INDEX_NONE;
+
+	/** If job is a duplicate, index of pointer to job in DuplicateJobs array in FShaderJobCache, used for clearing the pointer when the in-flight job completes */
+	int32 DuplicateIndex = INDEX_NONE;
+
+	void Clear()
+	{
+		Block = nullptr;
+		IndexInBlock = INDEX_NONE;
+		DuplicateIndex = INDEX_NONE;
+	}
+};
+
 /** Stores all of the common information used to compile a shader or pipeline. */
 class FShaderCommonCompileJob : public TIntrusiveLinkedList<FShaderCommonCompileJob>
 {
@@ -86,6 +114,8 @@ public:
 	/** In-engine timestamp of job being completed. Encompasses the compile time. Not set for jobs that are satisfied from the jobs cache */
 	double TimeExecutionCompleted = 0.0;
 
+	FShaderJobCacheRef JobCacheRef;
+
 	uint32 AddRef() const
 	{
 		return uint32(NumRefs.Increment());
