Fixed issue where ExplicitTangent property on Nanite meshes would not be saved properly because of missing UPROPERTY annotations

diff --git a/Engine/Source/Runtime/Engine/Classes/Engine/EngineTypes.h b/Engine/Source/Runtime/Engine/Classes/Engine/EngineTypes.h
--- a/Engine/Source/Runtime/Engine/Classes/Engine/EngineTypes.h
+++ b/Engine/Source/Runtime/Engine/Classes/Engine/EngineTypes.h
@@ -2814,6 +2814,7 @@ struct FMeshNaniteSettings
 	uint8 bPreserveArea : 1;
 
 	/** Whether to store explicit tangents instead of using the implicitly derived ones. */
+	UPROPERTY(BlueprintReadWrite, EditAnywhere, Category = NaniteSettings)
 	uint8 bExplicitTangents : 1;
 
 	/** Position Precision. Step size is 2^(-PositionPrecision) cm. MIN_int32 is auto. */
