Disable rendering of overlay material if max draw distance set to a negative value

diff --git a/Engine/Source/Runtime/Engine/Private/SkeletalMesh.cpp b/Engine/Source/Runtime/Engine/Private/SkeletalMesh.cpp
--- a/Engine/Source/Runtime/Engine/Private/SkeletalMesh.cpp
+++ b/Engine/Source/Runtime/Engine/Private/SkeletalMesh.cpp
@@ -6355,10 +6355,11 @@ void FSkeletalMeshSceneProxy::GetDynamicElementsSection(const TArray<const FScen
 			INC_DWORD_STAT_BY(STAT_SkelMeshTriangles,Mesh.GetNumPrimitives());
 			INC_DWORD_STAT(STAT_SkelMeshDrawCalls);
 
-			if (OverlayMaterial != nullptr)
+			// negative cull distance disables overlay rendering
+			if (OverlayMaterial != nullptr && OverlayMaterialMaxDrawDistance >= 0.f)
 			{
 				const bool bHasOverlayCullDistance = 
-					OverlayMaterialMaxDrawDistance > 1.f && 
+					OverlayMaterialMaxDrawDistance > 0.f && 
 					OverlayMaterialMaxDrawDistance != FLT_MAX && 
 					!ViewFamily.EngineShowFlags.DistanceCulledPrimitives;
 				
diff --git a/Engine/Source/Runtime/Engine/Private/StaticMeshRender.cpp b/Engine/Source/Runtime/Engine/Private/StaticMeshRender.cpp
--- a/Engine/Source/Runtime/Engine/Private/StaticMeshRender.cpp
+++ b/Engine/Source/Runtime/Engine/Private/StaticMeshRender.cpp
@@ -1366,7 +1366,8 @@ void FStaticMeshSceneProxy::DrawStaticElements(FStaticPrimitiveDrawInterface* PD
 								PDI->DrawMesh(MeshBatch, ScreenSize);
 							}
 
-							if (OverlayMaterial != nullptr)
+							// negative cull distance disables overlay rendering
+							if (OverlayMaterial != nullptr && OverlayMaterialMaxDrawDistance >= 0.f)
 							{
 								FMeshBatch OverlayMeshBatch(BaseMeshBatch);
 								OverlayMeshBatch.bOverlayMaterial = true;
diff --git a/Engine/Source/Runtime/Renderer/Private/SceneVisibility.cpp b/Engine/Source/Runtime/Renderer/Private/SceneVisibility.cpp
--- a/Engine/Source/Runtime/Renderer/Private/SceneVisibility.cpp
+++ b/Engine/Source/Runtime/Renderer/Private/SceneVisibility.cpp
@@ -1335,7 +1335,7 @@ void FRelevancePacket::ComputeRelevance(FDynamicPrimitiveIndexList& DynamicPrimi
 				{
 					// Overlay mesh can have its own cull distance that is shorter than primitive cull distance
 					float OverlayMaterialMaxDrawDistance = StaticMeshRelevance.ScreenSize;
-					if (OverlayMaterialMaxDrawDistance > 1.f && OverlayMaterialMaxDrawDistance != FLT_MAX)
+					if (OverlayMaterialMaxDrawDistance > 0.f && OverlayMaterialMaxDrawDistance != FLT_MAX)
 					{
 						if (DistanceSquared > FMath::Square(OverlayMaterialMaxDrawDistance * MaxDrawDistanceScale))
 						{
