Chaos Cloth - Fixed the cloth deformer when using a mirrored transform.

#opkforgithub public

diff --git a/Engine/Plugins/Experimental/ChaosCloth/Source/Chaos/Private/ChaosCloth/ChaosClothingSimulation.cpp b/Engine/Plugins/Experimental/ChaosCloth/Source/Chaos/Private/ChaosCloth/ChaosClothingSimulation.cpp
--- a/Engine/Plugins/Experimental/ChaosCloth/Source/Chaos/Private/ChaosCloth/ChaosClothingSimulation.cpp
+++ b/Engine/Plugins/Experimental/ChaosCloth/Source/Chaos/Private/ChaosCloth/ChaosClothingSimulation.cpp
@@ -655,6 +655,7 @@ PRAGMA_ENABLE_DEPRECATION_WARNINGS
 		// Get the reference transform used in the current animation pose
 		FTransform ReferenceBoneTransform = ComponentSpaceTransforms[ReferenceBoneIndex];
 		ReferenceBoneTransform *= OwnerTransform;
+		const bool bIsMirrored = (ReferenceBoneTransform.GetDeterminant() < 0);
 		ReferenceBoneTransform.SetScale3D(FVector(1.0f));  // Scale is already baked in the cloth mesh
 
 		// Set the world space transform to be this cloth's reference bone
@@ -698,6 +699,16 @@ PRAGMA_ENABLE_DEPRECATION_WARNINGS
 			}
 		}
 
+		// Invert normals in mirrored setups
+		if (bIsMirrored)
+		{
+			using FNormalsType = decltype(Data.Normals)::ElementType;
+			for (FNormalsType& Normal : Data.Normals)
+			{
+				Normal = -Normal;
+			}
+		}
+
 		// Set the current LOD these data apply to, so that the correct deformer mappings can be applied
 PRAGMA_DISABLE_DEPRECATION_WARNINGS  // TODO: CHAOS_IS_CLOTHINGSIMULATIONMESH_ABSTRACT
 		Data.LODIndex = Cloth->GetMesh()->GetOwnerLODIndex(LODIndex);  // The owner component LOD index can be different to the cloth mesh LOD index
diff --git a/Engine/Shaders/Private/GpuSkinCacheComputeShader.usf b/Engine/Shaders/Private/GpuSkinCacheComputeShader.usf
--- a/Engine/Shaders/Private/GpuSkinCacheComputeShader.usf
+++ b/Engine/Shaders/Private/GpuSkinCacheComputeShader.usf
@@ -389,7 +389,9 @@ void SkinCacheUpdateBatchCS(uint3 GroupID : SV_GroupID,
     int NumInfluences = 0;
 	float SumWeights = 0.0;
 	float ClothWeight = 0.0;
-	
+
+	const float3 WorldScaleAbs = abs(WorldScale);  // World scale can't be used mirrored to calculate the clothing positions and tangents since the cloth normals are then reversed
+
 	UNROLL_N(NUM_INFLUENCES_PER_VERTEX)
     for (int i = 0; i < ClothNumInfluencesPerVertex; ++i )
     {
@@ -423,21 +425,21 @@ void SkinCacheUpdateBatchCS(uint3 GroupID : SV_GroupID,
                 SumWeights = 1.0f;	   
 			}
 
-			NormalPosition += Weight * (Influence.NormalBaryCoordsAndDist.x * (A + NA * Influence.NormalBaryCoordsAndDist.w * WorldScale.x)
-								      + Influence.NormalBaryCoordsAndDist.y * (B + NB * Influence.NormalBaryCoordsAndDist.w * WorldScale.y)
-								      + Influence.NormalBaryCoordsAndDist.z * (C + NC * Influence.NormalBaryCoordsAndDist.w * WorldScale.z));
+			NormalPosition += Weight * (Influence.NormalBaryCoordsAndDist.x * (A + NA * Influence.NormalBaryCoordsAndDist.w * WorldScaleAbs.x)
+								      + Influence.NormalBaryCoordsAndDist.y * (B + NB * Influence.NormalBaryCoordsAndDist.w * WorldScaleAbs.y)
+								      + Influence.NormalBaryCoordsAndDist.z * (C + NC * Influence.NormalBaryCoordsAndDist.w * WorldScaleAbs.z));
 
-		    TangentPosition += Weight * (Influence.TangentBaryCoordsAndDist.x * (A + NA * Influence.TangentBaryCoordsAndDist.w * WorldScale.x)
-								       + Influence.TangentBaryCoordsAndDist.y * (B + NB * Influence.TangentBaryCoordsAndDist.w * WorldScale.y)
-								       + Influence.TangentBaryCoordsAndDist.z * (C + NC * Influence.TangentBaryCoordsAndDist.w * WorldScale.z));
+		    TangentPosition += Weight * (Influence.TangentBaryCoordsAndDist.x * (A + NA * Influence.TangentBaryCoordsAndDist.w * WorldScaleAbs.x)
+								       + Influence.TangentBaryCoordsAndDist.y * (B + NB * Influence.TangentBaryCoordsAndDist.w * WorldScaleAbs.y)
+								       + Influence.TangentBaryCoordsAndDist.z * (C + NC * Influence.TangentBaryCoordsAndDist.w * WorldScaleAbs.z));
 
 			float3 TriangleBary = float3(Influence.PositionBaryCoordsAndDist.x,
                                          Influence.PositionBaryCoordsAndDist.y,
 										 1.0f - Influence.PositionBaryCoordsAndDist.x - Influence.PositionBaryCoordsAndDist.y);
 
-            float3 SimPosition = TriangleBary.x*(A+NA*Influence.PositionBaryCoordsAndDist.w * WorldScale.x)
-					           + TriangleBary.y*(B+NB*Influence.PositionBaryCoordsAndDist.w * WorldScale.y)
-					           + TriangleBary.z*(C+NC*Influence.PositionBaryCoordsAndDist.w * WorldScale.z);
+            float3 SimPosition = TriangleBary.x*(A+NA*Influence.PositionBaryCoordsAndDist.w * WorldScaleAbs.x)
+					           + TriangleBary.y*(B+NB*Influence.PositionBaryCoordsAndDist.w * WorldScaleAbs.y)
+					           + TriangleBary.z*(C+NC*Influence.PositionBaryCoordsAndDist.w * WorldScaleAbs.z);
 
 	        Intermediates.SimulatedPosition += Weight * SimPosition;
 	    }
diff --git a/Engine/Shaders/Private/GpuSkinVertexFactory.ush b/Engine/Shaders/Private/GpuSkinVertexFactory.ush
--- a/Engine/Shaders/Private/GpuSkinVertexFactory.ush
+++ b/Engine/Shaders/Private/GpuSkinVertexFactory.ush
@@ -421,6 +421,8 @@ float4 ClothingPosition(FClothVertex ClothVertexInfluences[NUM_INFLUENCES_PER_VE
 	float SimulWeight = 0.0f; 
 	float3 AveragedSimulatedPosition = float3(0, 0, 0);
 
+	const float3 WorldScaleAbs = abs(WorldScale);  // World scale can't be used mirrored to calculate the clothing positions and tangents since the cloth normals are then reversed
+
 	UNROLL_N(NUM_INFLUENCES_PER_VERTEX)
     for (int i = 0; i < ClothNumInfluencesPerVertex; ++i )
     {
@@ -457,9 +459,9 @@ float4 ClothingPosition(FClothVertex ClothVertexInfluences[NUM_INFLUENCES_PER_VE
                                          Influence.BaryCoordPos.y,
 										 1.0f - Influence.BaryCoordPos.x - Influence.BaryCoordPos.y);
 
-			float3 SimPosition = TriangleBary.x * (A + NA * Influence.BaryCoordPos.w * WorldScale.x)
-					           + TriangleBary.y * (B + NB * Influence.BaryCoordPos.w * WorldScale.y)
-					           + TriangleBary.z * (C + NC * Influence.BaryCoordPos.w * WorldScale.z);
+			float3 SimPosition = TriangleBary.x * (A + NA * Influence.BaryCoordPos.w * WorldScaleAbs.x)
+					           + TriangleBary.y * (B + NB * Influence.BaryCoordPos.w * WorldScaleAbs.y)
+					           + TriangleBary.z * (C + NC * Influence.BaryCoordPos.w * WorldScaleAbs.z);
 
 			AveragedSimulatedPosition += Weight * SimPosition;
 	    }
@@ -756,6 +758,8 @@ float3x3 SkinTangents( FVertexFactoryInput Input, FVertexFactoryIntermediates In
 	    float SumWeights = 0.0;
 		float SimulWeight = 0.0;
 
+		const float3 WorldScaleAbs = abs(WorldScale);  // World scale can't be used mirrored to calculate the clothing positions and tangents since the cloth normals are then reversed
+
         for (int i = 0; i < ClothNumInfluencesPerVertex; ++i )
         {
             const FClothVertex Influence = Intermediates.ClothVertexInfluences[i];
@@ -789,13 +793,13 @@ float3x3 SkinTangents( FVertexFactoryInput Input, FVertexFactoryIntermediates In
                     SumWeights = 1.0f;	   
 			    }
 
-			    NormalPosition += Weight * (Influence.BaryCoordNormal.x * (A + NA * Influence.BaryCoordNormal.w * WorldScale.x)
-								          + Influence.BaryCoordNormal.y * (B + NB * Influence.BaryCoordNormal.w * WorldScale.y)
-								          + Influence.BaryCoordNormal.z * (C + NC * Influence.BaryCoordNormal.w * WorldScale.z));
+			    NormalPosition += Weight * (Influence.BaryCoordNormal.x * (A + NA * Influence.BaryCoordNormal.w * WorldScaleAbs.x)
+								          + Influence.BaryCoordNormal.y * (B + NB * Influence.BaryCoordNormal.w * WorldScaleAbs.y)
+								          + Influence.BaryCoordNormal.z * (C + NC * Influence.BaryCoordNormal.w * WorldScaleAbs.z));
 
-		        TangentPosition += Weight * (Influence.BaryCoordTangent.x * (A + NA * Influence.BaryCoordTangent.w * WorldScale.x)
-								           + Influence.BaryCoordTangent.y * (B + NB * Influence.BaryCoordTangent.w * WorldScale.y)
-								           + Influence.BaryCoordTangent.z * (C + NC * Influence.BaryCoordTangent.w * WorldScale.z));
+		        TangentPosition += Weight * (Influence.BaryCoordTangent.x * (A + NA * Influence.BaryCoordTangent.w * WorldScaleAbs.x)
+								           + Influence.BaryCoordTangent.y * (B + NB * Influence.BaryCoordTangent.w * WorldScaleAbs.y)
+								           + Influence.BaryCoordTangent.z * (C + NC * Influence.BaryCoordTangent.w * WorldScaleAbs.z));
 	        }
         }
 
diff --git a/Engine/Source/Runtime/Engine/Private/SkeletalRenderCPUSkin.cpp b/Engine/Source/Runtime/Engine/Private/SkeletalRenderCPUSkin.cpp
--- a/Engine/Source/Runtime/Engine/Private/SkeletalRenderCPUSkin.cpp
+++ b/Engine/Source/Runtime/Engine/Private/SkeletalRenderCPUSkin.cpp
@@ -712,7 +712,7 @@ static void SkinVertexSection(
 	const FClothSimulData* ClothSimData, 
 	float ClothBlendWeight, 
 	const FMatrix& WorldToLocal,
-	const FVector& WorldScale )
+	const FVector& WorldScaleAbs )
 {
 	static constexpr VectorRegister VECTOR_INV_65535 = MakeVectorRegisterDoubleConstant(1.0 / 65535, 1.0 / 65535, 1.0 / 65535, 1.0 / 65535);
 	
@@ -955,14 +955,14 @@ static void SkinVertexSection(
 						return FVector(0, 0, 1);
 					}
 
-					FORCEINLINE static FVector ClothingPosition(const FMeshToMeshVertData& InClothVertData, const FClothSimulData& InClothSimData, const FVector& InWorldScale)
+					FORCEINLINE static FVector ClothingPosition(const FMeshToMeshVertData& InClothVertData, const FClothSimulData& InClothSimData, const FVector& InWorldScaleAbs)
 					{
-						return    InClothVertData.PositionBaryCoordsAndDist.X * (GetClothSimulPosition(InClothSimData, InClothVertData.SourceMeshVertIndices[0]) + GetClothSimulNormal(InClothSimData, InClothVertData.SourceMeshVertIndices[0]) * InClothVertData.PositionBaryCoordsAndDist.W * InWorldScale.X)
-								+ InClothVertData.PositionBaryCoordsAndDist.Y * (GetClothSimulPosition(InClothSimData, InClothVertData.SourceMeshVertIndices[1]) + GetClothSimulNormal(InClothSimData, InClothVertData.SourceMeshVertIndices[1]) * InClothVertData.PositionBaryCoordsAndDist.W * InWorldScale.Y)
-								+ InClothVertData.PositionBaryCoordsAndDist.Z * (GetClothSimulPosition(InClothSimData, InClothVertData.SourceMeshVertIndices[2]) + GetClothSimulNormal(InClothSimData, InClothVertData.SourceMeshVertIndices[2]) * InClothVertData.PositionBaryCoordsAndDist.W * InWorldScale.Z);
+						return    InClothVertData.PositionBaryCoordsAndDist.X * (GetClothSimulPosition(InClothSimData, InClothVertData.SourceMeshVertIndices[0]) + GetClothSimulNormal(InClothSimData, InClothVertData.SourceMeshVertIndices[0]) * InClothVertData.PositionBaryCoordsAndDist.W * InWorldScaleAbs.X)
+								+ InClothVertData.PositionBaryCoordsAndDist.Y * (GetClothSimulPosition(InClothSimData, InClothVertData.SourceMeshVertIndices[1]) + GetClothSimulNormal(InClothSimData, InClothVertData.SourceMeshVertIndices[1]) * InClothVertData.PositionBaryCoordsAndDist.W * InWorldScaleAbs.Y)
+								+ InClothVertData.PositionBaryCoordsAndDist.Z * (GetClothSimulPosition(InClothSimData, InClothVertData.SourceMeshVertIndices[2]) + GetClothSimulNormal(InClothSimData, InClothVertData.SourceMeshVertIndices[2]) * InClothVertData.PositionBaryCoordsAndDist.W * InWorldScaleAbs.Z);
 					}
 
-					FORCEINLINE static void ClothingTangents(const FMeshToMeshVertData& InClothVertData, const FClothSimulData& InClothSimData, const FVector& InSimulatedPosition, const FMatrix& InWorldToLocal, const FVector& InWorldScale, FVector& OutTangentX, FVector& OutTangentZ)
+					FORCEINLINE static void ClothingTangents(const FMeshToMeshVertData& InClothVertData, const FClothSimulData& InClothSimData, const FVector& InSimulatedPosition, const FMatrix& InWorldToLocal, const FVector& InWorldScaleAbs, FVector& OutTangentX, FVector& OutTangentZ)
 					{
 						FVector A = GetClothSimulPosition(InClothSimData, InClothVertData.SourceMeshVertIndices[0]);
 						FVector B = GetClothSimulPosition(InClothSimData, InClothVertData.SourceMeshVertIndices[1]);
@@ -972,13 +972,13 @@ static void SkinVertexSection(
 						FVector NB = GetClothSimulNormal(InClothSimData, InClothVertData.SourceMeshVertIndices[1]);
 						FVector NC = GetClothSimulNormal(InClothSimData, InClothVertData.SourceMeshVertIndices[2]);
 
-						FVector NormalPosition = InClothVertData.NormalBaryCoordsAndDist.X*(A + NA*InClothVertData.NormalBaryCoordsAndDist.W * InWorldScale.X)
-												+ InClothVertData.NormalBaryCoordsAndDist.Y*(B + NB*InClothVertData.NormalBaryCoordsAndDist.W * InWorldScale.Y)
-												+ InClothVertData.NormalBaryCoordsAndDist.Z*(C + NC*InClothVertData.NormalBaryCoordsAndDist.W * InWorldScale.Z);
+						FVector NormalPosition = InClothVertData.NormalBaryCoordsAndDist.X*(A + NA*InClothVertData.NormalBaryCoordsAndDist.W * InWorldScaleAbs.X)
+												+ InClothVertData.NormalBaryCoordsAndDist.Y*(B + NB*InClothVertData.NormalBaryCoordsAndDist.W * InWorldScaleAbs.Y)
+												+ InClothVertData.NormalBaryCoordsAndDist.Z*(C + NC*InClothVertData.NormalBaryCoordsAndDist.W * InWorldScaleAbs.Z);
 
-						FVector TangentPosition = InClothVertData.TangentBaryCoordsAndDist.X*(A + NA*InClothVertData.TangentBaryCoordsAndDist.W * InWorldScale.X)
-												+ InClothVertData.TangentBaryCoordsAndDist.Y*(B + NB*InClothVertData.TangentBaryCoordsAndDist.W * InWorldScale.Y)
-												+ InClothVertData.TangentBaryCoordsAndDist.Z*(C + NC*InClothVertData.TangentBaryCoordsAndDist.W * InWorldScale.Z);
+						FVector TangentPosition = InClothVertData.TangentBaryCoordsAndDist.X*(A + NA*InClothVertData.TangentBaryCoordsAndDist.W * InWorldScaleAbs.X)
+												+ InClothVertData.TangentBaryCoordsAndDist.Y*(B + NB*InClothVertData.TangentBaryCoordsAndDist.W * InWorldScaleAbs.Y)
+												+ InClothVertData.TangentBaryCoordsAndDist.Z*(C + NC*InClothVertData.TangentBaryCoordsAndDist.W * InWorldScaleAbs.Z);
 
 						OutTangentX = (TangentPosition - InSimulatedPosition).GetUnsafeNormal();
 						OutTangentZ = (NormalPosition - InSimulatedPosition).GetUnsafeNormal();
@@ -990,7 +990,7 @@ static void SkinVertexSection(
 				};
 
 				// build sim position (in world space)
-				FVector SimulatedPositionWorld = ClothCPU::ClothingPosition(*ClothVertData, *ClothSimData, WorldScale);
+				FVector SimulatedPositionWorld = ClothCPU::ClothingPosition(*ClothVertData, *ClothSimData, WorldScaleAbs);
 
 				// transform back to local space
 				FVector3f SimulatedPosition = (FVector4f)WorldToLocal.TransformPosition(SimulatedPositionWorld);
@@ -1003,7 +1003,7 @@ static void SkinVertexSection(
 				// recompute tangent & normal
 				FVector TangentX;
 				FVector TangentZ;
-				ClothCPU::ClothingTangents(*ClothVertData, *ClothSimData, SimulatedPositionWorld, WorldToLocal, WorldScale, TangentX, TangentZ);
+				ClothCPU::ClothingTangents(*ClothVertData, *ClothSimData, SimulatedPositionWorld, WorldToLocal, WorldScaleAbs, TangentX, TangentZ);
 
 				// Lerp between skinned and simulated tangents
 				FVector SkinnedTangentX = DestVertex->TangentX.ToFVector();
@@ -1057,13 +1057,15 @@ static void SkinVertices(
 
 	int32 VertexBufferBaseIndex=0;
 
+	const FVector WorldScaleAbs = WorldScale.GetAbs();  // World scale can't be used mirrored to calculate the clothing positions and tangents since the cloth normals are then reversed
+
 	for(int32 SectionIndex= 0;SectionIndex< LOD.RenderSections.Num();SectionIndex++)
 	{
 		FSkelMeshRenderSection& Section = LOD.RenderSections[SectionIndex];
 
 		const FClothSimulData* ClothSimData = ClothSimulUpdateData.Find(Section.CorrespondClothAssetIndex);
 
-		SkinVertexSection<VertexType, NumberOfUVs>(DestVertex, MorphEvalInfos, MorphTargetWeights, Section, LOD, WeightBuffer, VertexBufferBaseIndex, NumValidMorphs, CurBaseVertIdx, LODIndex, ReferenceToLocal, ClothSimData, ClothBlendWeight, WorldToLocal, WorldScale);
+		SkinVertexSection<VertexType, NumberOfUVs>(DestVertex, MorphEvalInfos, MorphTargetWeights, Section, LOD, WeightBuffer, VertexBufferBaseIndex, NumValidMorphs, CurBaseVertIdx, LODIndex, ReferenceToLocal, ClothSimData, ClothBlendWeight, WorldToLocal, WorldScaleAbs);
 	}
 
 	VectorSetControlRegister( StatusRegister );
