Re-enable disabling of InGamePerformanceTrackers for un-initialized worlds.

diff --git a/Engine/Source/Runtime/Engine/Private/InGamePerformanceTracker.cpp b/Engine/Source/Runtime/Engine/Private/InGamePerformanceTracker.cpp
--- a/Engine/Source/Runtime/Engine/Private/InGamePerformanceTracker.cpp
+++ b/Engine/Source/Runtime/Engine/Private/InGamePerformanceTracker.cpp
@@ -80,7 +80,6 @@ void FInGamePerformanceTracker::ExitTimedSection()
 FInGameScopedCycleCounter::FInGameScopedCycleCounter(class UWorld* InWorld, EInGamePerfTrackers Tracker, EInGamePerfTrackerThreads TrackerThread, bool bEnabled)
 : FInGameCycleCounter(InWorld && bEnabled && InWorld->PerfTrackers ? &InWorld->PerfTrackers->GetInGamePerformanceTracker(Tracker, TrackerThread) : nullptr)
 {
-	check(!InWorld || InWorld->PerfTrackers);//UE-38057 - Crash potentially caused by dereferencing null here. Though this shouldn't be possible.
 	Begin();
 }
 
diff --git a/Engine/Source/Runtime/Engine/Private/LevelTick.cpp b/Engine/Source/Runtime/Engine/Private/LevelTick.cpp
--- a/Engine/Source/Runtime/Engine/Private/LevelTick.cpp
+++ b/Engine/Source/Runtime/Engine/Private/LevelTick.cpp
@@ -1302,10 +1302,13 @@ void UWorld::Tick( ELevelTick TickType, float DeltaSeconds )
 	FWorldDelegates::OnWorldTickStart.Broadcast(this, TickType, DeltaSeconds);
 
 	//Tick game and other thread trackers.
-	for (int32 Tracker = 0; Tracker < (int32)EInGamePerfTrackers::Num; ++Tracker)
+	if (PerfTrackers)
 	{
-		PerfTrackers->GetInGamePerformanceTracker((EInGamePerfTrackers)Tracker, EInGamePerfTrackerThreads::GameThread).Tick();
-		PerfTrackers->GetInGamePerformanceTracker((EInGamePerfTrackers)Tracker, EInGamePerfTrackerThreads::OtherThread).Tick();
+		for (int32 Tracker = 0; Tracker < (int32)EInGamePerfTrackers::Num; ++Tracker)
+		{
+			PerfTrackers->GetInGamePerformanceTracker((EInGamePerfTrackers)Tracker, EInGamePerfTrackerThreads::GameThread).Tick();
+			PerfTrackers->GetInGamePerformanceTracker((EInGamePerfTrackers)Tracker, EInGamePerfTrackerThreads::OtherThread).Tick();
+		}
 	}
 
 #if LOG_DETAILED_PATHFINDING_STATS
@@ -1783,9 +1786,12 @@ void UWorld::Tick( ELevelTick TickType, float DeltaSeconds )
 		[WorldParam](FRHICommandList& RHICmdList)
 		{
 			//Tick game and other thread trackers.
-			for (int32 Tracker = 0; Tracker < (int32)EInGamePerfTrackers::Num; ++Tracker)
+			if (WorldParam->PerfTrackers)
 			{
-				WorldParam->PerfTrackers->GetInGamePerformanceTracker((EInGamePerfTrackers)Tracker, EInGamePerfTrackerThreads::RenderThread).Tick();
+				for (int32 Tracker = 0; Tracker < (int32)EInGamePerfTrackers::Num; ++Tracker)
+				{
+					WorldParam->PerfTrackers->GetInGamePerformanceTracker((EInGamePerfTrackers)Tracker, EInGamePerfTrackerThreads::RenderThread).Tick();
+				}
 			}
 		});
 
diff --git a/Engine/Source/Runtime/Engine/Private/World.cpp b/Engine/Source/Runtime/Engine/Private/World.cpp
--- a/Engine/Source/Runtime/Engine/Private/World.cpp
+++ b/Engine/Source/Runtime/Engine/Private/World.cpp
@@ -144,7 +144,7 @@ CSV_DEFINE_CATEGORY(LevelStreamingPendingPurge, (!UE_BUILD_SHIPPING));
 static int32 bDisableRemapScriptActors = 0;
 FAutoConsoleVariableRef CVarDisableRemapScriptActors(TEXT("net.DisableRemapScriptActors"), bDisableRemapScriptActors, TEXT("When set, disables name remapping of compiled script actors (for networking)"));
 
-static bool bDisableInGamePerfTrackersForUninitializedWorlds = false;
+static bool bDisableInGamePerfTrackersForUninitializedWorlds = true;
 FAutoConsoleVariableRef CVarDisableInGamePerfTrackersForUninitializedWorlds(TEXT("s.World.SkipPerfTrackerForUninitializedWorlds"), bDisableInGamePerfTrackersForUninitializedWorlds, TEXT("When set, disables allocation of InGamePerformanceTrackers for Worlds that aren't initialized."));
 
 
