Telemetry: Add Language and Locale for the editor. This will be used by Localization team to identify the distribution of prefered language used in Unreal Editor.

diff --git a/Engine/Source/Runtime/Engine/Private/EngineAnalytics.cpp b/Engine/Source/Runtime/Engine/Private/EngineAnalytics.cpp
--- a/Engine/Source/Runtime/Engine/Private/EngineAnalytics.cpp
+++ b/Engine/Source/Runtime/Engine/Private/EngineAnalytics.cpp
@@ -17,6 +17,8 @@
 #include "StudioAnalytics.h"
 #include "UObject/Class.h"
 #include "Containers/Set.h"
+#include "Internationalization/Culture.h"
+#include "Internationalization/Internationalization.h"
 
 #if WITH_EDITOR
 #include "AnalyticsSessionSummaryManager.h"
@@ -284,6 +286,11 @@ void FEngineAnalytics::AppendMachineStats(TArray<FAnalyticsEventAttribute>& Even
 	EventAttributes.Emplace(TEXT("RenderingGPUAdapter"), GRHIAdapterName);
 	EventAttributes.Emplace(TEXT("CPUVendor"), FPlatformMisc::GetCPUVendor());
 	EventAttributes.Emplace(TEXT("CPUBrand"), FPlatformMisc::GetCPUBrand());
+
+	// Send Internationalization setting of the editor
+	EventAttributes.Emplace(TEXT("Internationalization.Language"), FInternationalization::Get().GetCurrentLanguage()->GetName());
+	EventAttributes.Emplace(TEXT("Internationalization.Locale"), FInternationalization::Get().GetCurrentLocale()->GetName());
+
 #if WITH_EDITOR
 	EventAttributes.Emplace(TEXT("Horde.TemplateID"), FHorde::GetTemplateId());
 	EventAttributes.Emplace(TEXT("Horde.TemplateName"), FHorde::GetTemplateName());
