FPhysScene_Chaos::HandleGlobalCollisionEvent should not try to grab the proxy's body instance if its owning component is no longer valid

diff --git a/Engine/Source/Runtime/Engine/Private/PhysicsEngine/Experimental/PhysScene_Chaos.cpp b/Engine/Source/Runtime/Engine/Private/PhysicsEngine/Experimental/PhysScene_Chaos.cpp
--- a/Engine/Source/Runtime/Engine/Private/PhysicsEngine/Experimental/PhysScene_Chaos.cpp
+++ b/Engine/Source/Runtime/Engine/Private/PhysicsEngine/Experimental/PhysScene_Chaos.cpp
@@ -1050,19 +1050,25 @@ void FPhysScene_Chaos::HandleGlobalCollisionEvent(Chaos::FCollisionDataArray con
 			const Chaos::FChaosPhysicsMaterial* InternalMat2 = CollisionItem.Mat2.Get();
 			CollisionEvent.Body1.PhysMaterial = InternalMat1 ? FPhysicsUserData::Get<UPhysicalMaterial>(InternalMat1->UserData) : nullptr;
 			CollisionEvent.Body2.PhysMaterial = InternalMat2 ? FPhysicsUserData::Get<UPhysicalMaterial>(InternalMat2->UserData) : nullptr;
-
-			const FBodyInstance* BodyInst1 = GetBodyInstanceFromProxyAndShape(CollisionItem.Proxy1, CollisionItem.ShapeIndex1);
-			if (BodyInst1 != nullptr)
+			
+			if (CollisionEvent.Body1.Component.IsValid())
 			{
-				CollisionEvent.Body1.BodyIndex = BodyInst1->InstanceBodyIndex;
-				CollisionEvent.Body1.BoneName = BodyInst1->BodySetup.IsValid() ? BodyInst1->BodySetup->BoneName : NAME_None;
+				const FBodyInstance* BodyInst1 = GetBodyInstanceFromProxyAndShape(CollisionItem.Proxy1, CollisionItem.ShapeIndex1);
+				if (BodyInst1 != nullptr)
+				{
+					CollisionEvent.Body1.BodyIndex = BodyInst1->InstanceBodyIndex;
+					CollisionEvent.Body1.BoneName = BodyInst1->BodySetup.IsValid() ? BodyInst1->BodySetup->BoneName : NAME_None;
+				}
 			}
 
-			const FBodyInstance* BodyInst2 = GetBodyInstanceFromProxyAndShape(CollisionItem.Proxy2, CollisionItem.ShapeIndex2);
-			if (BodyInst2 != nullptr)
+			if (CollisionEvent.Body2.Component.IsValid())
 			{
-				CollisionEvent.Body2.BodyIndex = BodyInst2->InstanceBodyIndex;
-				CollisionEvent.Body2.BoneName = BodyInst2->BodySetup.IsValid() ? BodyInst2->BodySetup->BoneName : NAME_None;
+				const FBodyInstance* BodyInst2 = GetBodyInstanceFromProxyAndShape(CollisionItem.Proxy2, CollisionItem.ShapeIndex2);
+				if (BodyInst2 != nullptr)
+				{
+					CollisionEvent.Body2.BodyIndex = BodyInst2->InstanceBodyIndex;
+					CollisionEvent.Body2.BoneName = BodyInst2->BodySetup.IsValid() ? BodyInst2->BodySetup->BoneName : NAME_None;
+				}
 			}
 		}
 	}
