ColorManagement: Fixed DDC2 texture build for advanced source color/encoding settings, added dependency on OpenColorIO library version for key invalidation.

diff --git a/Engine/Source/Developer/TextureCompressor/Private/TextureCompressorModule.cpp b/Engine/Source/Developer/TextureCompressor/Private/TextureCompressorModule.cpp
--- a/Engine/Source/Developer/TextureCompressor/Private/TextureCompressorModule.cpp
+++ b/Engine/Source/Developer/TextureCompressor/Private/TextureCompressorModule.cpp
@@ -3211,6 +3211,11 @@ void FTextureBuildSettings::GetEncodedTextureDescriptionWithPixelFormat(FEncoded
 	}
 }
 
+uint32 FTextureBuildSettings::GetOpenColorIOVersion()
+{
+	return OpenColorIOWrapper::GetVersionHex();
+}
+
 void FTextureBuildSettings::GetEncodedTextureDescription(FEncodedTextureDescription* OutTextureDescription, const ITextureFormat* InTextureFormat, int32 InEncodedMip0SizeX, int32 InEncodedMip0SizeY, int32 InEncodedMip0NumSlices, int32 InMipCount, bool bInImageHasAlphaChannel) const
 {
 	EPixelFormat EncodedPixelFormat = InTextureFormat->GetEncodedPixelFormat(*this, bInImageHasAlphaChannel);
diff --git a/Engine/Source/Developer/TextureCompressor/Public/TextureCompressorModule.h b/Engine/Source/Developer/TextureCompressor/Public/TextureCompressorModule.h
--- a/Engine/Source/Developer/TextureCompressor/Public/TextureCompressorModule.h
+++ b/Engine/Source/Developer/TextureCompressor/Public/TextureCompressorModule.h
@@ -349,6 +349,9 @@ struct FTextureBuildSettings
 	*/
 	TEXTURECOMPRESSOR_API void GetEncodedTextureDescription(FEncodedTextureDescription* OutTextureDescription, const ITextureFormat* InTextureFormat, int32 InEncodedMip0SizeX, int32 InEncodedMip0SizeY, int32 InEncodedMip0NumSlices, int32 InMipCount, bool bInImageHasAlphaChannel) const;
 	TEXTURECOMPRESSOR_API void GetEncodedTextureDescriptionWithPixelFormat(FEncodedTextureDescription* OutTextureDescription, EPixelFormat InEncodedPixelFormat, int32 InEncodedMip0SizeX, int32 InEncodedMip0SizeY, int32 InEncodedMip0NumSlices, int32 InMipCount) const;
+
+	/* Obtain the OpenColorIO library version, primarily used for DDC invalidation. */
+	TEXTURECOMPRESSOR_API static uint32 GetOpenColorIOVersion();
 };
 
 /**
diff --git a/Engine/Source/Runtime/ColorManagement/Private/ColorSpace.cpp b/Engine/Source/Runtime/ColorManagement/Private/ColorSpace.cpp
--- a/Engine/Source/Runtime/ColorManagement/Private/ColorSpace.cpp
+++ b/Engine/Source/Runtime/ColorManagement/Private/ColorSpace.cpp
@@ -75,9 +75,17 @@ FColorSpace::FColorSpace(const FVector2d& InRed, const FVector2d& InGreen, const
 }
 
 FColorSpace::FColorSpace(EColorSpace ColorSpaceType)
-	: Chromaticities(InPlace, FVector2d::Zero())
+	: Chromaticities(MakeChromaticities(ColorSpaceType))
 	, bIsSRGB(ColorSpaceType == EColorSpace::sRGB)
 {
+	RgbToXYZ = CalcRgbToXYZ();
+	XYZToRgb = RgbToXYZ.Inverse();
+}
+
+TStaticArray<FVector2d, 4> FColorSpace::MakeChromaticities(UE::Color::EColorSpace ColorSpaceType)
+{
+	TStaticArray<FVector2d, 4> Chromaticities(InPlace, FVector2d::Zero());
+
 	switch (ColorSpaceType)
 	{
 	case EColorSpace::None:
@@ -172,8 +180,7 @@ FColorSpace::FColorSpace(EColorSpace ColorSpaceType)
 		break;
 	}
 
-	RgbToXYZ = CalcRgbToXYZ();
-	XYZToRgb = RgbToXYZ.Inverse();
+	return Chromaticities;
 }
 
 FMatrix44d FColorSpace::CalcRgbToXYZ() const
diff --git a/Engine/Source/Runtime/ColorManagement/Public/ColorSpace.h b/Engine/Source/Runtime/ColorManagement/Public/ColorSpace.h
--- a/Engine/Source/Runtime/ColorManagement/Public/ColorSpace.h
+++ b/Engine/Source/Runtime/ColorManagement/Public/ColorSpace.h
@@ -169,6 +169,13 @@ public:
 	FColorSpace& operator=(FColorSpace&&) = default;
 	FColorSpace& operator=(const FColorSpace&) = default;
 
+	/**
+	* Make the chromaticities of the color space type.
+	* 
+	* @return TStaticArray of four chromaticities
+	*/
+	COLORMANAGEMENT_API static TStaticArray<FVector2d, 4> MakeChromaticities(UE::Color::EColorSpace ColorSpaceType);
+
 	/**
 	* Getter for the color space chromaticity coordinates.
 	* 
diff --git a/Engine/Source/Runtime/Engine/Private/TextureDerivedData.cpp b/Engine/Source/Runtime/Engine/Private/TextureDerivedData.cpp
--- a/Engine/Source/Runtime/Engine/Private/TextureDerivedData.cpp
+++ b/Engine/Source/Runtime/Engine/Private/TextureDerivedData.cpp
@@ -166,6 +166,11 @@ static void SerializeForKey(FArchive& Ar, const FTextureBuildSettings& Settings)
 		TempByte = Settings.ChromaticAdaptationMethod; Ar << TempByte;
 	}
 
+	if (Settings.SourceEncodingOverride != 0 || Settings.bHasColorSpaceDefinition)
+	{
+		TempUint32 = FTextureBuildSettings::GetOpenColorIOVersion(); Ar << TempUint32;
+	}
+
 	TempByte = Settings.bPreserveBorder; Ar << TempByte;
 
 	// bDitherMipMapAlpha was removed from Texture
diff --git a/Engine/Source/Runtime/Engine/Private/TextureDerivedDataBuildUtils.cpp b/Engine/Source/Runtime/Engine/Private/TextureDerivedDataBuildUtils.cpp
--- a/Engine/Source/Runtime/Engine/Private/TextureDerivedDataBuildUtils.cpp
+++ b/Engine/Source/Runtime/Engine/Private/TextureDerivedDataBuildUtils.cpp
@@ -3,6 +3,7 @@
 #include "TextureDerivedDataBuildUtils.h"
 
 #if WITH_EDITOR
+#include "ColorSpace.h"
 #include "DerivedDataBuild.h"
 #include "DerivedDataBuildFunctionRegistry.h"
 #include "DerivedDataSharedString.h"
@@ -136,13 +137,57 @@ static void WriteBuildSettings(FCbWriter& Writer, const FTextureBuildSettings& B
 	WriteCbFieldWithDefault<bool>(Writer, "bVolume", BuildSettings.bVolume, DefaultSettings.bVolume);
 	WriteCbFieldWithDefault<bool>(Writer, "bLongLatSource", BuildSettings.bLongLatSource, DefaultSettings.bLongLatSource);
 	WriteCbFieldWithDefault<bool>(Writer, "bSRGB", BuildSettings.bSRGB, DefaultSettings.bSRGB);
-	WriteCbFieldWithDefault(Writer, "SourceEncodingOverride", BuildSettings.SourceEncodingOverride, DefaultSettings.SourceEncodingOverride);
+
+	if (BuildSettings.SourceEncodingOverride != 0 /*UE::Color::EEncoding::None*/)
+	{
+		WriteCbField<uint32>(Writer, "EncodingOverrideVersion", UE::Color::ENCODING_TYPES_VER);
+		WriteCbFieldWithDefault(Writer, "SourceEncodingOverride", BuildSettings.SourceEncodingOverride, DefaultSettings.SourceEncodingOverride);
+	}
+	else // @todo SerializeForKey : remove else-case when overall key changes
+	{
+		WriteCbFieldWithDefault(Writer, "SourceEncodingOverride", BuildSettings.SourceEncodingOverride, DefaultSettings.SourceEncodingOverride);
+	}
+
 	WriteCbFieldWithDefault<bool>(Writer, "bHasColorSpaceDefinition", BuildSettings.bHasColorSpaceDefinition, DefaultSettings.bHasColorSpaceDefinition);
-	WriteCbFieldWithDefault(Writer, "RedChromaticityCoordinate", BuildSettings.RedChromaticityCoordinate, DefaultSettings.RedChromaticityCoordinate);
-	WriteCbFieldWithDefault(Writer, "GreenChromaticityCoordinate", BuildSettings.GreenChromaticityCoordinate, DefaultSettings.GreenChromaticityCoordinate);
-	WriteCbFieldWithDefault(Writer, "BlueChromaticityCoordinate", BuildSettings.BlueChromaticityCoordinate, DefaultSettings.BlueChromaticityCoordinate);
-	WriteCbFieldWithDefault(Writer, "WhiteChromaticityCoordinate", BuildSettings.WhiteChromaticityCoordinate, DefaultSettings.WhiteChromaticityCoordinate);
-	WriteCbFieldWithDefault(Writer, "ChromaticAdaptationMethod", BuildSettings.ChromaticAdaptationMethod, DefaultSettings.ChromaticAdaptationMethod);
+	if (BuildSettings.bHasColorSpaceDefinition)
+	{
+		WriteCbField<uint32>(Writer, "ColorSpaceVersion", UE::Color::COLORSPACE_VER);
+		/*
+		* The texture color transform depends on the chromaticities of both the source color space
+		* and the destination (working) color space, per its project setting. We therefore include
+		* the working color space chromaticities to incur a texture rebuild if/when changed.
+		*/
+		const TStaticArray<FVector2d, 4> DefaultChromaticities = UE::Color::FColorSpace::MakeChromaticities(UE::Color::EColorSpace::sRGB);
+		const UE::Color::FColorSpace& WorkingColorSpace = UE::Color::FColorSpace::GetWorking();
+		WriteCbFieldWithDefault(Writer, "WorkingRedChromaticity", FVector2f(WorkingColorSpace.GetRedChromaticity()), FVector2f(DefaultChromaticities[0]));
+		WriteCbFieldWithDefault(Writer, "WorkingGreenChromaticity", FVector2f(WorkingColorSpace.GetGreenChromaticity()), FVector2f(DefaultChromaticities[1]));
+		WriteCbFieldWithDefault(Writer, "WorkingBlueChromaticity", FVector2f(WorkingColorSpace.GetBlueChromaticity()), FVector2f(DefaultChromaticities[2]));
+		WriteCbFieldWithDefault(Writer, "WorkingWhiteChromaticity", FVector2f(WorkingColorSpace.GetWhiteChromaticity()), FVector2f(DefaultChromaticities[3]));
+
+		WriteCbFieldWithDefault(Writer, "RedChromaticityCoordinate", BuildSettings.RedChromaticityCoordinate, DefaultSettings.RedChromaticityCoordinate);
+		WriteCbFieldWithDefault(Writer, "GreenChromaticityCoordinate", BuildSettings.GreenChromaticityCoordinate, DefaultSettings.GreenChromaticityCoordinate);
+		WriteCbFieldWithDefault(Writer, "BlueChromaticityCoordinate", BuildSettings.BlueChromaticityCoordinate, DefaultSettings.BlueChromaticityCoordinate);
+		WriteCbFieldWithDefault(Writer, "WhiteChromaticityCoordinate", BuildSettings.WhiteChromaticityCoordinate, DefaultSettings.WhiteChromaticityCoordinate);
+		WriteCbFieldWithDefault(Writer, "ChromaticAdaptationMethod", BuildSettings.ChromaticAdaptationMethod, DefaultSettings.ChromaticAdaptationMethod);
+	}
+	else // @todo SerializeForKey : remove else-case when overall key changes
+	{
+		WriteCbFieldWithDefault(Writer, "RedChromaticityCoordinate", BuildSettings.RedChromaticityCoordinate, DefaultSettings.RedChromaticityCoordinate);
+		WriteCbFieldWithDefault(Writer, "GreenChromaticityCoordinate", BuildSettings.GreenChromaticityCoordinate, DefaultSettings.GreenChromaticityCoordinate);
+		WriteCbFieldWithDefault(Writer, "BlueChromaticityCoordinate", BuildSettings.BlueChromaticityCoordinate, DefaultSettings.BlueChromaticityCoordinate);
+		WriteCbFieldWithDefault(Writer, "WhiteChromaticityCoordinate", BuildSettings.WhiteChromaticityCoordinate, DefaultSettings.WhiteChromaticityCoordinate);
+		WriteCbFieldWithDefault(Writer, "ChromaticAdaptationMethod", BuildSettings.ChromaticAdaptationMethod, DefaultSettings.ChromaticAdaptationMethod);
+	}
+
+	if (BuildSettings.SourceEncodingOverride != 0 || BuildSettings.bHasColorSpaceDefinition)
+	{
+		/*
+		 * If any advanced color settings are in use, we link the OpenColorIO library version to the key
+		 * since updates could cause changes in processing.
+		 */ 
+		WriteCbField<uint32>(Writer, "OpenColorIOVersion", FTextureBuildSettings::GetOpenColorIOVersion());
+	}
+
 	WriteCbFieldWithDefault<bool>(Writer, "bUseLegacyGamma", BuildSettings.bUseLegacyGamma, DefaultSettings.bUseLegacyGamma);
 	WriteCbFieldWithDefault<bool>(Writer, "bPreserveBorder", BuildSettings.bPreserveBorder, DefaultSettings.bPreserveBorder);
 	WriteCbFieldWithDefault<bool>(Writer, "bForceNoAlphaChannel", BuildSettings.bForceNoAlphaChannel, DefaultSettings.bForceNoAlphaChannel);
@@ -197,14 +242,7 @@ static void WriteBuildSettings(FCbWriter& Writer, const FTextureBuildSettings& B
 	{
 		WriteCbField<bool>(Writer, "VTPow22_ForceNewDDcKey", true); 
 	}
-	// Force a new DDC key when utilizing color space transforms due to the switch to
-	// OpenColorIO. Remove when the build version is updated and all textures get rebuilt
-	// with it.
-	if (BuildSettings.SourceEncodingOverride != 0 /*UE::Color::EEncoding::None*/ ||
-		BuildSettings.bHasColorSpaceDefinition)
-	{
-		WriteCbField<bool>(Writer, "OCIO_ForceNewKey", true);
-	}
+	// @todo: see further above, remove the two else-cases per comments.
 
 	Writer.EndObject();
 }
