debug rendering : temporary fix for debug primitive not drawing on top of normal rendering when depth priority is set other than default

diff --git a/Engine/Source/Runtime/Engine/Private/Components/LineBatchComponent.cpp b/Engine/Source/Runtime/Engine/Private/Components/LineBatchComponent.cpp
--- a/Engine/Source/Runtime/Engine/Private/Components/LineBatchComponent.cpp
+++ b/Engine/Source/Runtime/Engine/Private/Components/LineBatchComponent.cpp
@@ -67,7 +67,10 @@ void FLineBatcherSceneProxy::GetDynamicMeshElements(const TArray<const FSceneVie
 
 #if UE_ENABLE_DEBUG_DRAWING
 			//If we are in a debug draw build, this line batcher should write out to the debug only PDI
-			FPrimitiveDrawInterface* PDI = Collector.GetDebugPDI(ViewIndex);
+			//FPrimitiveDrawInterface* PDI = Collector.GetDebugPDI(ViewIndex);
+
+			// temporary change to fix the problem around depth priority not being accounted for when drawing debug primitive
+			FPrimitiveDrawInterface* PDI = Collector.GetPDI(ViewIndex);
 #else
 			FPrimitiveDrawInterface* PDI = Collector.GetPDI(ViewIndex);
 #endif
