Fixed unsafe usage of GEngine->DeferredCommands. Speculative fix for PLAY-8626

diff --git a/Engine/Source/Runtime/Engine/Private/UnrealEngine.cpp b/Engine/Source/Runtime/Engine/Private/UnrealEngine.cpp
--- a/Engine/Source/Runtime/Engine/Private/UnrealEngine.cpp
+++ b/Engine/Source/Runtime/Engine/Private/UnrealEngine.cpp
@@ -2049,24 +2049,22 @@ void UEngine::Init(IEngineLoop* InEngineLoop)
 #if !(UE_BUILD_SHIPPING) || ENABLE_PGO_PROFILE
 	// Optionally Exec an exec file
 	FString Temp;
-	if( FParse::Value(FCommandLine::Get(), TEXT("EXEC="), Temp) )
+	if (FParse::Value(FCommandLine::Get(), TEXT("EXEC="), Temp))
 	{
-		new(GEngine->DeferredCommands) FString(FString(TEXT("exec ")) + Temp);
+		GEngine->DeferredCommands.Add(FString(TEXT("exec ")) + Temp);
 	}
 
 	// Optionally exec commands passed in the command line.
 	ParseExecCommands::QueueDeferredCommands(ParseExecCommands::ParseExecCmdsFromCommandLine(TEXT("ExecCmds")));
 
 	// optionally set the vsync console variable
-	if( FParse::Param(FCommandLine::Get(), TEXT("vsync")) )
+	if (FParse::Param(FCommandLine::Get(), TEXT("vsync")))
 	{
-		new(GEngine->DeferredCommands) FString(TEXT("r.vsync 1"));
+		GEngine->DeferredCommands.Add(TEXT("r.vsync 1"));
 	}
-
-	// optionally set the vsync console variable
-	if( FParse::Param(FCommandLine::Get(), TEXT("novsync")) )
+	else if (FParse::Param(FCommandLine::Get(), TEXT("novsync")))
 	{
-		new(GEngine->DeferredCommands) FString(TEXT("r.vsync 0"));
+		GEngine->DeferredCommands.Add(TEXT("r.vsync 0"));
 	}
 #endif // !(UE_BUILD_SHIPPING) || ENABLE_PGO_PROFILE
 
@@ -5256,7 +5254,7 @@ bool UEngine::HandleCrackURLCommand( const TCHAR* Cmd, FOutputDevice& Ar )
 
 bool UEngine::HandleDeferCommand( const TCHAR* Cmd, FOutputDevice& Ar )
 {
-	new(DeferredCommands)FString(Cmd);
+	DeferredCommands.Add(Cmd);
 	return 1;
 }
 
diff --git a/Engine/Source/Runtime/Launch/Private/IOS/IOSAppDelegateConsoleHandling.cpp b/Engine/Source/Runtime/Launch/Private/IOS/IOSAppDelegateConsoleHandling.cpp
--- a/Engine/Source/Runtime/Launch/Private/IOS/IOSAppDelegateConsoleHandling.cpp
+++ b/Engine/Source/Runtime/Launch/Private/IOS/IOSAppDelegateConsoleHandling.cpp
@@ -99,7 +99,7 @@ extern bool GShowSplashScreen;
 			Ch.AddZeroed([ConsoleCommand length]+1);
 
 			FPlatformString::CFStringToTCHAR((CFStringRef)ConsoleCommand, Ch.GetData());
-			new(GEngine->DeferredCommands) FString(Ch.GetData());
+			GEngine->DeferredCommands.Add(FString(Ch.GetData()));
 		}
 #if !PLATFORM_TVOS
 		NSUInteger ExistingCommand = [self.ConsoleHistoryValues indexOfObjectPassingTest:
