WorldPartition: Add invalid package mappings to InstancingContext to avoid loading non-instanced actor worlds

#rnx

diff --git a/Engine/Source/Runtime/Engine/Private/WorldPartition/WorldPartition.cpp b/Engine/Source/Runtime/Engine/Private/WorldPartition/WorldPartition.cpp
--- a/Engine/Source/Runtime/Engine/Private/WorldPartition/WorldPartition.cpp
+++ b/Engine/Source/Runtime/Engine/Private/WorldPartition/WorldPartition.cpp
@@ -610,6 +610,16 @@ void UWorldPartition::Initialize(UWorld* InWorld, const FTransform& InTransform)
 		ActorDescContainer = RegisterActorDescContainer(ContainerInitParams);
 
 		{
+			// If a Valid Actor references an Invalid Actor:
+			// Make sure Invalid Actors do not load their imports (ex: outer non instanced world).
+			if (bIsInstanced)
+			{
+				for (const FAssetData& InvalidActor : ActorDescContainer->InvalidActors)
+				{
+					InstancingContext.AddPackageMapping(InvalidActor.PackageName, NAME_None);
+				}
+			}
+
 			TRACE_CPUPROFILER_EVENT_SCOPE(UActorDescContainer::Hash);
 			for (FActorDescContainerCollection::TIterator<> ActorDescIterator(this); ActorDescIterator; ++ActorDescIterator)
 			{
