Fix GetLinkedAnimLayerInstanceByGroup not working for non-implemented layers with no instance assigned

diff --git a/Engine/Source/Runtime/Engine/Private/Animation/AnimInstance.cpp b/Engine/Source/Runtime/Engine/Private/Animation/AnimInstance.cpp
--- a/Engine/Source/Runtime/Engine/Private/Animation/AnimInstance.cpp
+++ b/Engine/Source/Runtime/Engine/Private/Animation/AnimInstance.cpp
@@ -3385,26 +3385,31 @@ UAnimInstance* UAnimInstance::GetLinkedAnimLayerInstanceByGroup(FName InGroup) c
 		{
 			const FAnimNode_LinkedAnimLayer* Layer = LayerNodeProperty->ContainerPtrToValuePtr<FAnimNode_LinkedAnimLayer>(this);
 
-			UClass* ClassForGroups;
-			if (UClass* InstanceClass = Layer->InstanceClass.Get())
+			TArray<UClass*, TInlineAllocator<4>> ClassesForGroups;
+			if (UAnimInstance* LayerInstance = Layer->GetTargetInstance<UAnimInstance>())
 			{
-				ClassForGroups = InstanceClass;
+				ClassesForGroups.Add(LayerInstance->GetClass());
 			}
-			else if(UClass* InterfaceClass = Layer->Interface.Get())
+			if (UClass* InstanceClass = Layer->InstanceClass.Get())
 			{
-				ClassForGroups = InterfaceClass;
+				ClassesForGroups.Add(InstanceClass);
 			}
-			else
+			if(UClass* InterfaceClass = Layer->Interface.Get())
 			{
-				ClassForGroups = GetClass();
+				ClassesForGroups.Add(InterfaceClass);
 			}
 
-			IAnimClassInterface* AnimClassInterfaceForGroups = IAnimClassInterface::GetFromClass(ClassForGroups);
-			if(const FAnimBlueprintFunction* FoundFunction = IAnimClassInterface::FindAnimBlueprintFunction(AnimClassInterfaceForGroups, Layer->Layer))
+			ClassesForGroups.Add(GetClass());
+
+			for(UClass* ClassForGroups : ClassesForGroups)
 			{
-				if(InGroup == FoundFunction->Group)
+				IAnimClassInterface* AnimClassInterfaceForGroups = IAnimClassInterface::GetFromClass(ClassForGroups);
+				if(const FAnimBlueprintFunction* FoundFunction = IAnimClassInterface::FindAnimBlueprintFunction(AnimClassInterfaceForGroups, Layer->Layer))
 				{
-					return Layer->GetTargetInstance<UAnimInstance>();
+					if(InGroup == FoundFunction->Group)
+					{
+						return Layer->GetTargetInstance<UAnimInstance>();
+					}
 				}
 			}
 		}
@@ -3421,23 +3426,32 @@ void UAnimInstance::GetLinkedAnimLayerInstancesByGroup(FName InGroup, TArray<UAn
 		{
 			const FAnimNode_LinkedAnimLayer* Layer = LayerNodeProperty->ContainerPtrToValuePtr<FAnimNode_LinkedAnimLayer>(this);
 
-			UClass* ClassForGroups;
-			if (UClass* InterfaceClass = Layer->Interface.Get())
+			TArray<UClass*, TInlineAllocator<4>> ClassesForGroups;
+			if (UAnimInstance* LayerInstance = Layer->GetTargetInstance<UAnimInstance>())
 			{
-				ClassForGroups = InterfaceClass;
+				ClassesForGroups.Add(LayerInstance->GetClass());
 			}
-			else
+			if (UClass* InstanceClass = Layer->InstanceClass.Get())
 			{
-				ClassForGroups = GetClass();
+				ClassesForGroups.Add(InstanceClass);
 			}
+			if (UClass* InterfaceClass = Layer->Interface.Get())
+			{
+				ClassesForGroups.Add(InterfaceClass);
+			}
+
+			ClassesForGroups.Add(GetClass());
 
-			IAnimClassInterface* AnimClassInterfaceForGroups = IAnimClassInterface::GetFromClass(ClassForGroups);
-			if (const FAnimBlueprintFunction* FoundFunction = IAnimClassInterface::FindAnimBlueprintFunction(AnimClassInterfaceForGroups, Layer->Layer))
+			for (UClass* ClassForGroups : ClassesForGroups)
 			{
-				UAnimInstance* TargetInstance = Layer->GetTargetInstance<UAnimInstance>();
-				if (InGroup == FoundFunction->Group && !OutLinkedInstances.Contains(TargetInstance))
+				IAnimClassInterface* AnimClassInterfaceForGroups = IAnimClassInterface::GetFromClass(ClassForGroups);
+				if (const FAnimBlueprintFunction* FoundFunction = IAnimClassInterface::FindAnimBlueprintFunction(AnimClassInterfaceForGroups, Layer->Layer))
 				{
-					OutLinkedInstances.Add(TargetInstance);
+					UAnimInstance* TargetInstance = Layer->GetTargetInstance<UAnimInstance>();
+					if (InGroup == FoundFunction->Group && !OutLinkedInstances.Contains(TargetInstance))
+					{
+						OutLinkedInstances.Add(TargetInstance);
+					}
 				}
 			}
 		}
