Attempt to detect and recover from corrupted mesh streaming data in shipping. Simply crash in other build configs. In either case, emit a detailed error message

diff --git a/Engine/Source/Runtime/Engine/Private/SkeletalMeshLODRenderData.cpp b/Engine/Source/Runtime/Engine/Private/SkeletalMeshLODRenderData.cpp
--- a/Engine/Source/Runtime/Engine/Private/SkeletalMeshLODRenderData.cpp
+++ b/Engine/Source/Runtime/Engine/Private/SkeletalMeshLODRenderData.cpp
@@ -738,7 +738,7 @@ void FSkeletalMeshLODRenderData::SerializeStreamedData(FArchive& Ar, USkinnedAss
 	Ar << SkinWeightProfilesData;
 	SkinWeightProfilesData.Init(&SkinWeightVertexBuffer);
 
-	if (Ar.IsLoading() && Owner)
+	if (Ar.IsLoading() && !Ar.IsError() && Owner)
 	{
 		Owner->SetSkinWeightProfilesData(LODIdx, SkinWeightProfilesData);
 	}
diff --git a/Engine/Source/Runtime/Engine/Private/StaticMesh.cpp b/Engine/Source/Runtime/Engine/Private/StaticMesh.cpp
--- a/Engine/Source/Runtime/Engine/Private/StaticMesh.cpp
+++ b/Engine/Source/Runtime/Engine/Private/StaticMesh.cpp
@@ -505,14 +505,6 @@ void FStaticMeshLODResources::SerializeBuffers(FArchive& Ar, UStaticMesh* OwnerS
 		}
 	}
 
-	bHasRayTracingGeometry = false;
-	bHasWireframeIndices = false;
-	bHasDepthOnlyIndices = false;
-	bHasReversedIndices = false;
-	bHasReversedDepthOnlyIndices = false;
-	bHasColorVertexData = false;
-	DepthOnlyNumTriangles = 0;
-
 	FStripDataFlags StripFlags(Ar, InStripFlags);
 
 	VertexBuffers.PositionVertexBuffer.Serialize(Ar, bNeedsCPUAccess);
@@ -576,7 +568,6 @@ void FStaticMeshLODResources::SerializeBuffers(FArchive& Ar, UStaticMesh* OwnerS
 	{
 		SerializedAdditionalIndexBuffers->WireframeIndexBuffer.Serialize(Ar, bNeedsCPUAccess);
 		AccumIndexBufferSize(SerializedAdditionalIndexBuffers->WireframeIndexBuffer, OutBuffersSize.SerializedBuffersSize);
-		bHasWireframeIndices = AdditionalIndexBuffers && SerializedAdditionalIndexBuffers->WireframeIndexBuffer.GetNumIndices() != 0;
 	}
 
 	if (Ar.IsLoading() && Ar.CustomVer(FUE5ReleaseStreamObjectVersion::GUID) < FUE5ReleaseStreamObjectVersion::RemovingTessellation && !StripFlags.IsClassDataStripped(CDSF_AdjacencyData_DEPRECATED))
@@ -594,22 +585,26 @@ void FStaticMeshLODResources::SerializeBuffers(FArchive& Ar, UStaticMesh* OwnerS
 			// Immediately release serialized offline BLAS data if it won't be used anyway due to rendering settings.
 			RayTracingGeometry.RawData.Discard();
 		}
-		bHasRayTracingGeometry = RayTracingGeometry.RawData.Num() != 0;
 	}
 
-	// Needs to be done now because on cooked platform, indices are discarded after RHIInit.
-	bHasDepthOnlyIndices = DepthOnlyIndexBuffer.GetNumIndices() != 0;
-	bHasReversedIndices = AdditionalIndexBuffers && bSerializeReversedIndexBuffer && SerializedAdditionalIndexBuffers->ReversedIndexBuffer.GetNumIndices() != 0;
-	bHasReversedDepthOnlyIndices = AdditionalIndexBuffers && bSerializeReversedIndexBuffer && SerializedAdditionalIndexBuffers->ReversedDepthOnlyIndexBuffer.GetNumIndices() != 0;
-	bHasColorVertexData = VertexBuffers.ColorVertexBuffer.GetNumVertices() > 0;
-	DepthOnlyNumTriangles = DepthOnlyIndexBuffer.GetNumIndices() / 3;
-
 	AreaWeightedSectionSamplers.SetNum(Sections.Num());
 	for (FStaticMeshSectionAreaWeightedTriangleSampler& Sampler : AreaWeightedSectionSamplers)
 	{
 		Sampler.Serialize(Ar);
 	}
 	AreaWeightedSampler.Serialize(Ar);
+
+	// Update metadata but only if serialization was successful. This needs to be done now because on cooked platform, indices are discarded after RHIInit.
+	if (!Ar.IsError())
+	{
+		bHasRayTracingGeometry = bSerializeRayTracingGeometry && RayTracingGeometry.RawData.Num() != 0;
+		bHasWireframeIndices = AdditionalIndexBuffers && bSerializeWireframeIndexBuffer && SerializedAdditionalIndexBuffers->WireframeIndexBuffer.GetNumIndices() != 0;
+		bHasDepthOnlyIndices = DepthOnlyIndexBuffer.GetNumIndices() != 0;
+		bHasReversedIndices = AdditionalIndexBuffers && bSerializeReversedIndexBuffer && SerializedAdditionalIndexBuffers->ReversedIndexBuffer.GetNumIndices() != 0;
+		bHasReversedDepthOnlyIndices = AdditionalIndexBuffers && bSerializeReversedIndexBuffer && SerializedAdditionalIndexBuffers->ReversedDepthOnlyIndexBuffer.GetNumIndices() != 0;
+		bHasColorVertexData = VertexBuffers.ColorVertexBuffer.GetNumVertices() > 0;
+		DepthOnlyNumTriangles = DepthOnlyIndexBuffer.GetNumIndices() / 3;
+	}
 }
 
 void FStaticMeshLODResources::SerializeAvailabilityInfo(FArchive& Ar)
diff --git a/Engine/Source/Runtime/Engine/Private/Streaming/SkeletalMeshUpdate.cpp b/Engine/Source/Runtime/Engine/Private/Streaming/SkeletalMeshUpdate.cpp
--- a/Engine/Source/Runtime/Engine/Private/Streaming/SkeletalMeshUpdate.cpp
+++ b/Engine/Source/Runtime/Engine/Private/Streaming/SkeletalMeshUpdate.cpp
@@ -157,15 +157,6 @@ void FSkeletalMeshStreamIn::CreateBuffers_Internal(const FContext& Context)
 		{
 			FSkeletalMeshLODRenderData& LODResource = *Context.LODResourcesView[LODIndex];
 
-#if ALLOW_RENDER_ASSET_STREAMING_BREADCRUMB
-			const FStaticMeshVertexBuffer& TangentTexCoordBuffers = LODResource.StaticVertexBuffers.StaticMeshVertexBuffer;
-			if (GStreamingEnableRenderAssetUpdateBreadcrumb != 0 && TangentTexCoordBuffers.GetNumVertices() > 0 && !TangentTexCoordBuffers.GetTangentData())
-			{
-				FRenderAssetUpdateBreadcrumb* Breadcrumb = (FRenderAssetUpdateBreadcrumb*)alloca(sizeof(FRenderAssetUpdateBreadcrumb));
-				Breadcrumb->SetDebugData(Mesh->GetPathName(), LODIndex, TangentTexCoordBuffers.GetNumVertices(), false);
-				*(uint32*)3 = 0xed000001;
-			}
-#endif
 			if (bRenderThread)
 			{
 				IntermediateBuffersArray[LODIndex].CreateFromCPUData_RenderThread(LODResource);
@@ -545,15 +536,32 @@ void FSkeletalMeshStreamIn_IO::SerializeLODData(const FContext& Context)
 			constexpr uint8 DummyStripFlags = 0;
 			LODResource.SerializeStreamedData(Ar, const_cast<USkeletalMesh*>(Mesh), LODIndex + Context.AssetLODBias, DummyStripFlags, bNeedsCPUAccess, bForceKeepCPUResources);
 
-#if ALLOW_RENDER_ASSET_STREAMING_BREADCRUMB
-			const FStaticMeshVertexBuffer& TangentTexCoordBuffers = LODResource.StaticVertexBuffers.StaticMeshVertexBuffer;
-			if (GStreamingEnableRenderAssetUpdateBreadcrumb != 0 && TangentTexCoordBuffers.GetNumVertices() > 0 && !TangentTexCoordBuffers.GetTangentData())
+			// Attempt to recover from possibly corrupted data
+			if (Ar.IsError())
 			{
-				FRenderAssetUpdateBreadcrumb* Breadcrumb = (FRenderAssetUpdateBreadcrumb*)alloca(sizeof(FRenderAssetUpdateBreadcrumb));
-				Breadcrumb->SetDebugData(Mesh->GetPathName(), LODIndex, TangentTexCoordBuffers.GetNumVertices(), Ar.GetError());
-				*(uint32*)3 = 0xed000001;
-			}
+				UE_LOG(LogContentStreaming, Error,
+					TEXT("[%s] SkeletalMesh stream in failed due to possibly corrupted data. LOD %d %d-%d. BulkData %#x offset %lld size %lld flags %#x. bForceKeepCPUResources %d. bNeedsCPUAccess %d."),
+					*Mesh->GetPathName(),
+					LODIndex,
+					PendingFirstLODIdx,
+					CurrentFirstLODIdx - 1,
+					LODResource.StreamingBulkData.GetIoFilenameHash(),
+					LODResource.StreamingBulkData.GetBulkDataOffsetInFile(),
+					LODResource.StreamingBulkData.GetBulkDataSize(),
+					LODResource.StreamingBulkData.GetBulkDataFlags(),
+					bForceKeepCPUResources,
+					bNeedsCPUAccess);
+
+#if STREAMING_RETRY_ON_DESERIALIZATION_ERROR
+				bFailedOnIOError = true;
+				MarkAsCancelled();
+				break;
+#else
+				GLog->FlushThreadedLogs();
+				GLog->Flush();
+				UE_LOG(LogContentStreaming, Fatal, TEXT("Possibly corrupted skeletal mesh LOD data detected."));
 #endif
+			}
 		}
 
 		BulkData = FIoBuffer();
diff --git a/Engine/Source/Runtime/Engine/Private/Streaming/StaticMeshUpdate.cpp b/Engine/Source/Runtime/Engine/Private/Streaming/StaticMeshUpdate.cpp
--- a/Engine/Source/Runtime/Engine/Private/Streaming/StaticMeshUpdate.cpp
+++ b/Engine/Source/Runtime/Engine/Private/Streaming/StaticMeshUpdate.cpp
@@ -173,15 +173,6 @@ void FStaticMeshStreamIn::CreateBuffers_Internal(const FContext& Context)
 		{
 			FStaticMeshLODResources& LODResource = *Context.LODResourcesView[LODIdx];
 
-#if ALLOW_RENDER_ASSET_STREAMING_BREADCRUMB
-			const FStaticMeshVertexBuffer& TangentTexCoordBuffers = LODResource.VertexBuffers.StaticMeshVertexBuffer;
-			if (GStreamingEnableRenderAssetUpdateBreadcrumb != 0 && TangentTexCoordBuffers.GetNumVertices() > 0 && !TangentTexCoordBuffers.GetTangentData() && Context.Mesh)
-			{
-				FRenderAssetUpdateBreadcrumb* Breadcrumb = (FRenderAssetUpdateBreadcrumb*)alloca(sizeof(FRenderAssetUpdateBreadcrumb));
-				Breadcrumb->SetDebugData(Context.Mesh->GetPathName(), LODIdx, TangentTexCoordBuffers.GetNumVertices(), false);
-				*(uint32*)3 = 0xed000002;
-			}
-#endif
 			if (bRenderThread)
 			{
 				IntermediateBuffersArray[LODIdx].CreateFromCPUData_RenderThread(LODResource);
@@ -588,17 +579,31 @@ void FStaticMeshStreamIn_IO::SerializeLODData(const FContext& Context)
 			constexpr uint8 DummyStripFlags = 0;
 			typename FStaticMeshLODResources::FStaticMeshBuffersSize DummyBuffersSize;
 			LODResource.SerializeBuffers(Ar, const_cast<UStaticMesh*>(Mesh), DummyStripFlags, DummyBuffersSize);
-			check(DummyBuffersSize.CalcBuffersSize() == LODResource.BuffersSize);
 
-#if ALLOW_RENDER_ASSET_STREAMING_BREADCRUMB
-			const FStaticMeshVertexBuffer& TangentTexCoordBuffers = LODResource.VertexBuffers.StaticMeshVertexBuffer;
-			if (GStreamingEnableRenderAssetUpdateBreadcrumb != 0 && TangentTexCoordBuffers.GetNumVertices() > 0 && !TangentTexCoordBuffers.GetTangentData())
+			// Attempt to recover from possibly corrupted data if allowed
+			if (Ar.IsError())
 			{
-				FRenderAssetUpdateBreadcrumb* Breadcrumb = (FRenderAssetUpdateBreadcrumb*)alloca(sizeof(FRenderAssetUpdateBreadcrumb));
-				Breadcrumb->SetDebugData(Mesh->GetPathName(), LODIdx, TangentTexCoordBuffers.GetNumVertices(), Ar.GetError());
-				*(uint32*)3 = 0xed000002;
-			}
+				UE_LOG(LogContentStreaming, Error,
+					TEXT("[%s] StaticMesh stream in failed due to possibly corrupted data. LOD %d %d-%d. BulkData %#x offset %lld size %lld flags %#x."),
+					*Mesh->GetPathName(),
+					LODIdx,
+					PendingFirstLODIdx,
+					CurrentFirstLODIdx - 1,
+					LODResource.StreamingBulkData.GetIoFilenameHash(),
+					LODResource.StreamingBulkData.GetBulkDataOffsetInFile(),
+					LODResource.StreamingBulkData.GetBulkDataSize(),
+					LODResource.StreamingBulkData.GetBulkDataFlags());
+
+#if STREAMING_RETRY_ON_DESERIALIZATION_ERROR
+				bFailedOnIOError = true;
+				MarkAsCancelled();
+				break;
+#else
+				GLog->FlushThreadedLogs();
+				GLog->Flush();
+				UE_LOG(LogContentStreaming, Fatal, TEXT("Possibly corrupted static mesh LOD data detected."));
 #endif
+			}
 		}
 
 		BulkData = FIoBuffer();
diff --git a/Engine/Source/Runtime/Engine/Private/Streaming/TextureStreamingHelpers.cpp b/Engine/Source/Runtime/Engine/Private/Streaming/TextureStreamingHelpers.cpp
--- a/Engine/Source/Runtime/Engine/Private/Streaming/TextureStreamingHelpers.cpp
+++ b/Engine/Source/Runtime/Engine/Private/Streaming/TextureStreamingHelpers.cpp
@@ -299,28 +299,6 @@ static TAutoConsoleVariable<int32> CVarStreamingVRAMPercentageClamp(
 	TEXT("This avoids reserving too much space for systems with a lot of VRAM. (default=1024)"),
 	ECVF_Cheat);
 
-#if ALLOW_RENDER_ASSET_STREAMING_BREADCRUMB
-int32 GStreamingEnableRenderAssetUpdateBreadcrumb = 1;
-static FAutoConsoleVariableRef CVarStreamingEnableRenderAssetUpdateBreadcrumb(
-	TEXT("r.Streaming.EnableRenderAssetUpdateBreadcrumb"),
-	GStreamingEnableRenderAssetUpdateBreadcrumb,
-	TEXT("Whether to detect invalid streaming states, crash when detected, and store debug info onto stack"),
-	ECVF_Default);
-
-void FRenderAssetUpdateBreadcrumb::SetDebugData(const FString& InAssetName, int32 InLODIndex, uint32 InNumVertices, bool bInIsArError)
-{
-	FMemory::Memzero(this, sizeof(FRenderAssetUpdateBreadcrumb));
-	const int32 CopySize = FMath::Min<int32>(InAssetName.Len(), UE_ARRAY_COUNT(AssetName) - 1);
-	if (CopySize > 0)
-	{
-		FCStringAnsi::Strncpy(AssetName, TCHAR_TO_ANSI(*InAssetName), CopySize);
-	}
-	LODIndex = InLODIndex;
-	NumVertices = InNumVertices;
-	bIsArError = bInIsArError;
-}
-#endif
-
 void FRenderAssetStreamingSettings::Update()
 {
 	MaxEffectiveScreenSize = CVarStreamingScreenSizeEffectiveMax.GetValueOnAnyThread();
diff --git a/Engine/Source/Runtime/Engine/Private/Streaming/TextureStreamingHelpers.h b/Engine/Source/Runtime/Engine/Private/Streaming/TextureStreamingHelpers.h
--- a/Engine/Source/Runtime/Engine/Private/Streaming/TextureStreamingHelpers.h
+++ b/Engine/Source/Runtime/Engine/Private/Streaming/TextureStreamingHelpers.h
@@ -12,8 +12,8 @@ TextureStreamingHelpers.h: Definitions of classes used for texture streaming.
 #include "Misc/MemStack.h"
 #include "Engine/TextureDefines.h"
 
-#ifndef ALLOW_RENDER_ASSET_STREAMING_BREADCRUMB
-#define ALLOW_RENDER_ASSET_STREAMING_BREADCRUMB 0
+#ifndef STREAMING_RETRY_ON_DESERIALIZATION_ERROR
+#define STREAMING_RETRY_ON_DESERIALIZATION_ERROR UE_BUILD_SHIPPING
 #endif
 
 class AActor;
@@ -191,17 +191,3 @@ struct FRenderAssetStreamingStats
 
 // Helper to access the level bStaticComponentsRegisteredInStreamingManager flag.
 extern bool OwnerLevelHasRegisteredStaticComponentsInStreamingManager(const class AActor* Owner);
-
-#if ALLOW_RENDER_ASSET_STREAMING_BREADCRUMB
-struct FRenderAssetUpdateBreadcrumb
-{
-	ANSICHAR AssetName[256];
-	int32 LODIndex;
-	uint32 NumVertices;
-	bool bIsArError;
-
-	FORCENOINLINE void SetDebugData(const FString& InAssetName, int32 InLODIndex, uint32 InNumVertices, bool bInIsArError);
-};
-
-extern int32 GStreamingEnableRenderAssetUpdateBreadcrumb;
-#endif
