Fix tracking of per-vertex-factory stats for PSO precaching (missing dereference).

#pf skip

diff --git a/Engine/Source/Runtime/Engine/Private/PSOPrecache.cpp b/Engine/Source/Runtime/Engine/Private/PSOPrecache.cpp
--- a/Engine/Source/Runtime/Engine/Private/PSOPrecache.cpp
+++ b/Engine/Source/Runtime/Engine/Private/PSOPrecache.cpp
@@ -612,7 +612,7 @@ void PSOCollectorStats::FPrecacheUsageData::UpdateStats(uint32 MeshPassType, con
 		if (VFType != nullptr)
 		{
 			uint32* Value = PerVertexFactoryCount.FindOrAdd(VFType, 0);
-			Value++;
+			(*Value)++;
 		}
 	}
 }
