ShowFlag improvements to Heterogeneous Volume Component

1) Add "Heterogeneous Volumes" show flag.
2) Respect DrawBounds show flag.

diff --git a/Engine/Source/Runtime/Engine/Private/Components/HeterogeneousVolumeComponent.cpp b/Engine/Source/Runtime/Engine/Private/Components/HeterogeneousVolumeComponent.cpp
--- a/Engine/Source/Runtime/Engine/Private/Components/HeterogeneousVolumeComponent.cpp
+++ b/Engine/Source/Runtime/Engine/Private/Components/HeterogeneousVolumeComponent.cpp
@@ -145,7 +145,7 @@ FPrimitiveViewRelevance FHeterogeneousVolumeSceneProxy::GetViewRelevance(const F
 		MaterialRelevance.SetPrimitiveViewRelevance(Result);
 	}
 
-	Result.bDrawRelevance = IsShown(View);
+	Result.bDrawRelevance = IsShown(View) && View->Family->EngineShowFlags.HeterogeneousVolumes;
 	Result.bOpaque = false;
 	Result.bStaticRelevance = false;
 	Result.bDynamicRelevance = true;
@@ -195,6 +195,18 @@ void FHeterogeneousVolumeSceneProxy::GetDynamicMeshElements(
 
 		Collector.AddMesh(0, Mesh);
 	}
+
+	// Draw bounds 
+	if (ViewFamily.EngineShowFlags.HeterogeneousVolumes)
+	{
+		for (int32 ViewIndex = 0; ViewIndex < Views.Num(); ViewIndex++)
+		{
+			if (VisibilityMap & (1 << ViewIndex))
+			{
+				RenderBounds(Collector.GetPDI(ViewIndex), ViewFamily.EngineShowFlags, GetBounds(), IsSelected());
+			}
+		}
+	}
 }
 
 /*=============================================================================
@@ -353,13 +365,15 @@ AHeterogeneousVolume::AHeterogeneousVolume(const FObjectInitializer& ObjectIniti
 			}
 		};
 		static FConstructorStatics ConstructorStatics;
-
 		if (GetSpriteComponent())
 		{
 			GetSpriteComponent()->SetRelativeScale3D(FVector(0.5f, 0.5f, 0.5f));
+			GetSpriteComponent()->bHiddenInGame = true;
+			GetSpriteComponent()->bIsScreenSizeScaled = true;
 			GetSpriteComponent()->SpriteInfo.Category = ConstructorStatics.ID_HeterogeneousVolume;
 			GetSpriteComponent()->SpriteInfo.DisplayName = ConstructorStatics.NAME_HeterogeneousVolume;
 			GetSpriteComponent()->SetupAttachment(HeterogeneousVolumeComponent);
+			GetSpriteComponent()->bReceivesDecals = false;
 		}
 	}
 #endif // WITH_EDITORONLY_DATA
diff --git a/Engine/Source/Runtime/Engine/Public/ShowFlagsValues.inl b/Engine/Source/Runtime/Engine/Public/ShowFlagsValues.inl
--- a/Engine/Source/Runtime/Engine/Public/ShowFlagsValues.inl
+++ b/Engine/Source/Runtime/Engine/Public/ShowFlagsValues.inl
@@ -181,6 +181,8 @@ SHOWFLAG_ALWAYS_ACCESSIBLE(DynamicShadows, SFG_LightingComponents, NSLOCTEXT("Un
 SHOWFLAG_ALWAYS_ACCESSIBLE(Particles, SFG_Normal, NSLOCTEXT("UnrealEd", "ParticlesSF", "Particle Sprites"))
 /** Niagara, for now SHOWFLAG_ALWAYS_ACCESSIBLE because it's exposed in SceneCapture */
 SHOWFLAG_FIXED_IN_SHIPPING(1, Niagara, SFG_Advanced, NSLOCTEXT("UnrealEd", "NiagaraSF", "Niagara Renderers"))
+/** HeterogeneousVolumes, for now SHOWFLAG_ALWAYS_ACCESSIBLE because it's exposed in SceneCapture */
+SHOWFLAG_FIXED_IN_SHIPPING(1, HeterogeneousVolumes, SFG_Advanced, NSLOCTEXT("UnrealEd", "HeterogeneousVolumesSF", "Heterogeneous Volumes"))
 /** if SkeletalMeshes are getting rendered, for now SHOWFLAG_ALWAYS_ACCESSIBLE because it's exposed in SceneCapture */
 SHOWFLAG_ALWAYS_ACCESSIBLE(SkeletalMeshes, SFG_Normal, NSLOCTEXT("UnrealEd", "SkeletalMeshesSF", "Skeletal Meshes"))
 /** if the builder brush (editor) is getting rendered */
