Add FWaitForEngineFramesCommand
`FWaitForEngineFramesCommand` runs for a specified number of frames to achieve the "wait for X frames" functionality

diff --git a/Engine/Source/Runtime/Engine/Private/Tests/AutomationCommon.cpp b/Engine/Source/Runtime/Engine/Private/Tests/AutomationCommon.cpp
--- a/Engine/Source/Runtime/Engine/Private/Tests/AutomationCommon.cpp
+++ b/Engine/Source/Runtime/Engine/Private/Tests/AutomationCommon.cpp
@@ -704,6 +704,19 @@ bool FWaitForNextEngineFrameCommand::Update()
 	return LastFrame != GFrameCounter;
 }
 
+FWaitForEngineFramesCommand::FWaitForEngineFramesCommand(int32 InFramesToWait)
+{
+	// Assert that InFramesToWait is at least 1
+	check(InFramesToWait > 0);
+	FramesToWait = InFramesToWait;
+}
+
+bool FWaitForEngineFramesCommand::Update()
+{
+	FrameCounter += 1;
+	return FrameCounter == FramesToWait;
+}
+
 ///////////////////////////////////////////////////////////////////////
 // Common Latent commands which are used across test type. I.e. Engine, Network, etc...
 
diff --git a/Engine/Source/Runtime/Engine/Public/Tests/AutomationCommon.h b/Engine/Source/Runtime/Engine/Public/Tests/AutomationCommon.h
--- a/Engine/Source/Runtime/Engine/Public/Tests/AutomationCommon.h
+++ b/Engine/Source/Runtime/Engine/Public/Tests/AutomationCommon.h
@@ -232,6 +232,22 @@ private:
 	uint64 LastFrame = 0;
 };
 
+/**
+ * Latent command to wait for a given number of engine frames
+ */
+class FWaitForEngineFramesCommand : public IAutomationLatentCommand
+{
+public:
+	ENGINE_API explicit FWaitForEngineFramesCommand(int32 InFramesToWait = 1);
+
+	ENGINE_API bool Update() override;
+
+private:
+	int32 FrameCounter = 0;
+	int32 FramesToWait = 1;
+};
+
+
 /**
 * Request an Image Comparison and queue the result to the test report
 * @param InImageName	Name used to identify the comparison
