Fix reinstancing crash in linked anim layers/graphs

diff --git a/Engine/Source/Runtime/Engine/Private/Animation/AnimNode_LinkedAnimGraph.cpp b/Engine/Source/Runtime/Engine/Private/Animation/AnimNode_LinkedAnimGraph.cpp
--- a/Engine/Source/Runtime/Engine/Private/Animation/AnimNode_LinkedAnimGraph.cpp
+++ b/Engine/Source/Runtime/Engine/Private/Animation/AnimNode_LinkedAnimGraph.cpp
@@ -517,6 +517,8 @@ void FAnimNode_LinkedAnimGraph::HandleObjectsReinstanced_Impl(UObject* InSourceO
 		USkeletalMeshComponent* MeshComponent = SourceAnimInstance->GetSkelMeshComponent();
 		if(MeshComponent && MeshComponent->GetSkeletalMeshAsset())
 		{
+			SourceAnimInstance->RecalcRequiredBones();
+
 			FAnimationInitializeContext Context(&SourceProxy);
 			InitializeSubGraph_AnyThread(Context);
 		}
diff --git a/Engine/Source/Runtime/Engine/Private/Animation/AnimNode_LinkedAnimLayer.cpp b/Engine/Source/Runtime/Engine/Private/Animation/AnimNode_LinkedAnimLayer.cpp
--- a/Engine/Source/Runtime/Engine/Private/Animation/AnimNode_LinkedAnimLayer.cpp
+++ b/Engine/Source/Runtime/Engine/Private/Animation/AnimNode_LinkedAnimLayer.cpp
@@ -235,6 +235,8 @@ void FAnimNode_LinkedAnimLayer::HandleObjectsReinstanced_Impl(UObject* InSourceO
 		USkeletalMeshComponent* MeshComponent = SourceAnimInstance->GetSkelMeshComponent();
 		if(MeshComponent && MeshComponent->GetSkeletalMeshAsset())
 		{
+			SourceAnimInstance->RecalcRequiredBones();
+
 			FAnimationInitializeContext Context(&SourceProxy);
 			InitializeSubGraph_AnyThread(Context);
 		}
