Move support of InstanceLODDistanceScale from HISM up to ISM.
The value is used to modify ScreenSize stored in MeshElement/MeshRelevance.
Note that we keep the convention from it's use in HISM that smaller values drive higher LODs.
That's the same convention as with foliage.LODDistanceScale but the inverse of r.StaticMeshLODDistanceScale.

diff --git a/Engine/Source/Runtime/Engine/Classes/Components/HierarchicalInstancedStaticMeshComponent.h b/Engine/Source/Runtime/Engine/Classes/Components/HierarchicalInstancedStaticMeshComponent.h
--- a/Engine/Source/Runtime/Engine/Classes/Components/HierarchicalInstancedStaticMeshComponent.h
+++ b/Engine/Source/Runtime/Engine/Classes/Components/HierarchicalInstancedStaticMeshComponent.h
@@ -183,10 +183,6 @@ class UHierarchicalInstancedStaticMeshComponent : public UInstancedStaticMeshCom
 	// Current value of density scaling applied to this component
 	float CurrentDensityScaling;
 
-	/** Scale applied to change the computation of LOD distances when using the StaticMesh screen sizes. */
-	UPROPERTY(EditAnywhere, BlueprintReadOnly, Category = Culling)
-	float InstanceLODDistanceScale;
-
 	// The number of nodes in the occlusion layer
 	UPROPERTY()
 	int32 OcclusionLayerNumNodes;
@@ -252,13 +248,6 @@ public:
 	ENGINE_API virtual void PreAllocateInstancesMemory(int32 AddedInstanceCount) override;
 	virtual bool SupportsRemoveSwap() const override { return true; }
 
-	/** Sets the fading start and culling end distances for this component. */
-	UFUNCTION(BlueprintCallable, Category = "Components|HierarchicalInstancedStaticMesh")
-	ENGINE_API void SetLODDistanceScale(float InLODDistanceScale);
-
-	UFUNCTION(BlueprintCallable, Category = "Components|HierarchicalInstancedStaticMesh")
-	float GetLODDistanceScale() const { return InstanceLODDistanceScale; }
-
 	/** Get the number of instances that overlap a given sphere */
 	ENGINE_API int32 GetOverlappingSphereCount(const FSphere& Sphere) const;
 	/** Get the number of instances that overlap a given box */
diff --git a/Engine/Source/Runtime/Engine/Classes/Components/InstancedStaticMeshComponent.h b/Engine/Source/Runtime/Engine/Classes/Components/InstancedStaticMeshComponent.h
--- a/Engine/Source/Runtime/Engine/Classes/Components/InstancedStaticMeshComponent.h
+++ b/Engine/Source/Runtime/Engine/Classes/Components/InstancedStaticMeshComponent.h
@@ -124,6 +124,13 @@ class UInstancedStaticMeshComponent : public UStaticMeshComponent, public ISMIns
 	UPROPERTY()
 	TArray<FInstancedStaticMeshRandomSeed> AdditionalRandomSeeds;
 
+	/** 
+	 * Scale applied to change the computation of LOD distances when using the StaticMesh screen sizes. 
+	 * Smaller values make LODs transition earlier.
+	 */
+	UPROPERTY(EditAnywhere, BlueprintReadOnly, Category = Culling)
+	float InstanceLODDistanceScale = 1.f;
+
 	/** Distance from camera at which each instance begins to fade out. */
 	UPROPERTY(EditAnywhere, BlueprintReadOnly, Category=Culling)
 	int32 InstanceStartCullDistance;
@@ -180,6 +187,14 @@ class UInstancedStaticMeshComponent : public UStaticMeshComponent, public ISMIns
 	UFUNCTION(BlueprintCallable, Category = "Components|InstancedStaticMesh")
 	ENGINE_API bool GetInstanceTransform(int32 InstanceIndex, FTransform& OutInstanceTransform, bool bWorldSpace = false) const;
 	
+	/** Gets the current LOD scale. */
+	UFUNCTION(BlueprintCallable, Category = "Components|InstancedStaticMesh")
+	float GetLODDistanceScale() const { return InstanceLODDistanceScale; }
+
+	/** Sets the LOD scale. */
+	UFUNCTION(BlueprintCallable, Category = "Components|InstancedStaticMesh")
+	ENGINE_API void SetLODDistanceScale(float InLODDistanceScale);
+
 	// TODO: KevinO cleanup
 	/** Get the prev transform for the instance specified. Only works if PerInstancePrevTransform has been setup and updated through BatchUpdateInstancesTransforms */
 	ENGINE_API bool GetInstancePrevTransform(int32 InstanceIndex, FTransform& OutInstanceTransform, bool bWorldSpace = false) const;
diff --git a/Engine/Source/Runtime/Engine/Classes/Engine/InstancedStaticMesh.h b/Engine/Source/Runtime/Engine/Classes/Engine/InstancedStaticMesh.h
--- a/Engine/Source/Runtime/Engine/Classes/Engine/InstancedStaticMesh.h
+++ b/Engine/Source/Runtime/Engine/Classes/Engine/InstancedStaticMesh.h
@@ -501,6 +501,7 @@ public:
 #if WITH_EDITOR
 	,	bHasSelectedInstances(false)
 #endif
+	,	InstanceLODDistanceScale(InComponent->InstanceLODDistanceScale)
 #if RHI_RAYTRACING
 	,	CachedRayTracingLOD(-1)
 #endif
@@ -594,6 +595,7 @@ public:
 
 	ENGINE_API virtual bool GetInstanceDrawDistanceMinMax(FVector2f& OutDistanceMinMax) const override;
 
+	virtual float GetLodScreenSizeScale() const override;
 	virtual float GetGpuLodInstanceRadius() const override;
 
 	virtual bool IsDetailMesh() const override { return true; }
@@ -619,6 +621,9 @@ protected:
 	FInstancingUserData UserData_SelectedInstances;
 	FInstancingUserData UserData_DeselectedInstances;
 
+	/** LOD distance scale from component. */
+	float InstanceLODDistanceScale;
+
 #if RHI_RAYTRACING
 	struct FRayTracingDynamicData
 	{
diff --git a/Engine/Source/Runtime/Engine/Private/HierarchicalInstancedStaticMesh.cpp b/Engine/Source/Runtime/Engine/Private/HierarchicalInstancedStaticMesh.cpp
--- a/Engine/Source/Runtime/Engine/Private/HierarchicalInstancedStaticMesh.cpp
+++ b/Engine/Source/Runtime/Engine/Private/HierarchicalInstancedStaticMesh.cpp
@@ -2082,7 +2082,6 @@ UHierarchicalInstancedStaticMeshComponent::UHierarchicalInstancedStaticMeshCompo
 	, UnbuiltInstanceBounds(ForceInit)
 	, bEnableDensityScaling(false)
 	, CurrentDensityScaling(1.0f)
-	, InstanceLODDistanceScale(1.0f)
 	, OcclusionLayerNumNodes(0)
 	, InstanceCountToRender(0)
 	, bIsAsyncBuilding(false)
@@ -2727,15 +2726,6 @@ void UHierarchicalInstancedStaticMeshComponent::ClearInstances()
 	FNavigationSystem::UpdateComponentData(*this);
 }
 
-void UHierarchicalInstancedStaticMeshComponent::SetLODDistanceScale(float InLODDistanceScale)
-{
-	if (InstanceLODDistanceScale != InLODDistanceScale)
-	{
-		InstanceLODDistanceScale = InLODDistanceScale;
-		MarkRenderStateDirty();
-	}
-}
-
 bool UHierarchicalInstancedStaticMeshComponent::ShouldCreatePhysicsState() const
 {
 	if (bDisableCollision)
diff --git a/Engine/Source/Runtime/Engine/Private/InstancedStaticMesh.cpp b/Engine/Source/Runtime/Engine/Private/InstancedStaticMesh.cpp
--- a/Engine/Source/Runtime/Engine/Private/InstancedStaticMesh.cpp
+++ b/Engine/Source/Runtime/Engine/Private/InstancedStaticMesh.cpp
@@ -2153,6 +2153,11 @@ bool FInstancedStaticMeshSceneProxy::GetInstanceDrawDistanceMinMax(FVector2f& Ou
 	}
 }
 
+float FInstancedStaticMeshSceneProxy::GetLodScreenSizeScale() const
+{
+	return InstanceLODDistanceScale > 0.f ? 1.f / InstanceLODDistanceScale : 1.f;
+}
+
 float FInstancedStaticMeshSceneProxy::GetGpuLodInstanceRadius() const
 {
 	// Note that StaticMeshBounds.SphereRadius is a better fit, but doesn't match the value on the GPU used for LOD culling.
@@ -5326,6 +5331,15 @@ void UInstancedStaticMeshComponent::BeginDestroy()
 	Super::BeginDestroy();
 }
 
+void UInstancedStaticMeshComponent::SetLODDistanceScale(float InLODDistanceScale)
+{
+	if (InstanceLODDistanceScale != InLODDistanceScale)
+	{
+		InstanceLODDistanceScale = InLODDistanceScale;
+		MarkRenderStateDirty();
+	}
+}
+
 #if WITH_EDITOR
 void UInstancedStaticMeshComponent::PostEditChangeChainProperty(FPropertyChangedChainEvent& PropertyChangedEvent)
 {
diff --git a/Engine/Source/Runtime/Engine/Private/StaticMeshRender.cpp b/Engine/Source/Runtime/Engine/Private/StaticMeshRender.cpp
--- a/Engine/Source/Runtime/Engine/Private/StaticMeshRender.cpp
+++ b/Engine/Source/Runtime/Engine/Private/StaticMeshRender.cpp
@@ -2411,7 +2411,7 @@ FLightInteraction FStaticMeshSceneProxy::FLODInfo::GetInteraction(const FLightSc
 
 float FStaticMeshSceneProxy::GetScreenSize( int32 LODIndex ) const
 {
-	return RenderData->ScreenSize[LODIndex].GetValue();
+	return RenderData->ScreenSize[LODIndex].GetValue() * GetLodScreenSizeScale();
 }
 
 /**
@@ -2446,7 +2446,8 @@ int32 FStaticMeshSceneProxy::GetLOD(const FSceneView* View) const
 #endif
 
 	const FBoxSphereBounds& ProxyBounds = GetBounds();
-	return ComputeStaticMeshLOD(RenderData, ProxyBounds.Origin, ProxyBounds.SphereRadius, *View, ClampedMinLOD);
+	const float LODScale = GetCachedScalabilityCVars().StaticMeshLODDistanceScale * GetLodScreenSizeScale();
+	return ComputeStaticMeshLOD(RenderData, ProxyBounds.Origin, ProxyBounds.SphereRadius, *View, ClampedMinLOD, LODScale);
 }
 
 FLODMask FStaticMeshSceneProxy::GetLODMask(const FSceneView* View) const
@@ -2507,8 +2508,7 @@ FLODMask FStaticMeshSceneProxy::GetLODMask(const FSceneView* View) const
 			}
 
 			FCachedSystemScalabilityCVars CachedSystemScalabilityCVars = GetCachedScalabilityCVars();
-
-			const float LODScale = CachedSystemScalabilityCVars.StaticMeshLODDistanceScale;
+			const float LODScale = CachedSystemScalabilityCVars.StaticMeshLODDistanceScale * GetLodScreenSizeScale();
 
 			if (bUseDithered)
 			{
diff --git a/Engine/Source/Runtime/Engine/Public/PrimitiveSceneProxy.h b/Engine/Source/Runtime/Engine/Public/PrimitiveSceneProxy.h
--- a/Engine/Source/Runtime/Engine/Public/PrimitiveSceneProxy.h
+++ b/Engine/Source/Runtime/Engine/Public/PrimitiveSceneProxy.h
@@ -1109,6 +1109,9 @@ public:
 
 	virtual uint8 GetCurrentFirstLODIdx_RenderThread() const { return 0; }
 
+	/** Returns a scale to apply to ScreenSize used in LOD calculation. */
+	virtual float GetLodScreenSizeScale() const { return 1.f; }
+
 	/** Returns the instance radius to use for per instance GPU LOD calculation. Returns 0.f if GPU LOD isn't enabled on the primitive. */
 	virtual float GetGpuLodInstanceRadius() const { return 0.f; }
 
diff --git a/Engine/Source/Runtime/Experimental/GeometryCollectionEngine/Private/GeometryCollection/GeometryCollectionISMPoolComponent.cpp b/Engine/Source/Runtime/Experimental/GeometryCollectionEngine/Private/GeometryCollection/GeometryCollectionISMPoolComponent.cpp
--- a/Engine/Source/Runtime/Experimental/GeometryCollectionEngine/Private/GeometryCollection/GeometryCollectionISMPoolComponent.cpp
+++ b/Engine/Source/Runtime/Experimental/GeometryCollectionEngine/Private/GeometryCollection/GeometryCollectionISMPoolComponent.cpp
@@ -107,13 +107,9 @@ FGeometryCollectionISM::FGeometryCollectionISM(AActor* InOwningActor, const FGeo
 	ISMC->SetCollisionEnabled(ECollisionEnabled::NoCollision);
 	ISMC->bOverrideMinLOD = MeshInstance.Desc.MinLod > 0;	
 	ISMC->MinLOD = MeshInstance.Desc.MinLod;
+	ISMC->SetLODDistanceScale(MeshInstance.Desc.LodScale);
 	ISMC->ComponentTags.Append(MeshInstance.Desc.Tags);
 
-	if (HISMC)
-	{
-		HISMC->SetLODDistanceScale(MeshInstance.Desc.LodScale);
-	}
-
 	InOwningActor->AddInstanceComponent(ISMC);
 	ISMC->RegisterComponent();
 	ISMC->SetVisibility(false);
