Deprecates UGameUserSettings::MinResolutionScale that causes sg.ResolutionQuality!=0 at start up of game build ending up with very low rendering resolution

diff --git a/Engine/Source/Runtime/Engine/Classes/GameFramework/GameUserSettings.h b/Engine/Source/Runtime/Engine/Classes/GameFramework/GameUserSettings.h
--- a/Engine/Source/Runtime/Engine/Classes/GameFramework/GameUserSettings.h
+++ b/Engine/Source/Runtime/Engine/Classes/GameFramework/GameUserSettings.h
@@ -421,6 +421,7 @@ protected:
 	float FrameRateLimit;
 
 	/** Min resolution scale we allow in current display mode */
+	UE_DEPRECATED(5.3, "MinResolutionScale is now always Scalability::MinResolutionScale=0 that fallbacks to default project wide behavior defined by r.ScreenPercentage.Default")
 	float MinResolutionScale;
 
 	/** Desired screen width used to calculate the resolution scale when user changes display mode */
diff --git a/Engine/Source/Runtime/Engine/Private/GameUserSettings.cpp b/Engine/Source/Runtime/Engine/Private/GameUserSettings.cpp
--- a/Engine/Source/Runtime/Engine/Private/GameUserSettings.cpp
+++ b/Engine/Source/Runtime/Engine/Private/GameUserSettings.cpp
@@ -218,7 +218,9 @@ void UGameUserSettings::SetToDefaults()
 	WindowPosY = GetDefaultWindowPosition().Y;
 	FullscreenMode = GetDefaultWindowMode();
 	FrameRateLimit = 0.0f;
+	PRAGMA_DISABLE_DEPRECATION_WARNINGS
 	MinResolutionScale = Scalability::MinResolutionScale;
+	PRAGMA_ENABLE_DEPRECATION_WARNINGS
 	DesiredScreenWidth = 1280;
 	DesiredScreenHeight = 720;
 	LastUserConfirmedDesiredScreenWidth = DesiredScreenWidth;
@@ -261,7 +263,9 @@ void UGameUserSettings::UpdateResolutionQuality()
 	const int32 MinHeight = UKismetSystemLibrary::GetMinYResolutionFor3DView();
 	const int32 ScreenWidth = (FullscreenMode == EWindowMode::WindowedFullscreen) ? GetDesktopResolution().X : ResolutionSizeX;
 	const int32 ScreenHeight = (FullscreenMode == EWindowMode::WindowedFullscreen) ? GetDesktopResolution().Y : ResolutionSizeY;
-	MinResolutionScale = FMath::Max<float>(Scalability::MinResolutionScale, ((float)MinHeight / (float)ScreenHeight) * 100.0f);
+	PRAGMA_DISABLE_DEPRECATION_WARNINGS
+	MinResolutionScale = Scalability::MinResolutionScale;
+	PRAGMA_ENABLE_DEPRECATION_WARNINGS
 
 	if (bUseDesiredScreenHeight)
 	{
@@ -269,7 +273,7 @@ void UGameUserSettings::UpdateResolutionQuality()
 	}
 	else
 	{
-		ScalabilityQuality.ResolutionQuality = FMath::Max(ScalabilityQuality.ResolutionQuality, MinResolutionScale);
+		ScalabilityQuality.ResolutionQuality = FMath::Max(ScalabilityQuality.ResolutionQuality, Scalability::MinResolutionScale);
 	}
 }
 
@@ -281,14 +285,14 @@ float UGameUserSettings::GetDefaultResolutionScale()
 	const int32 ClampedHeight = (ScreenHeight > 0 && DesiredScreenHeight > ScreenHeight) ? ScreenHeight : DesiredScreenHeight;
 
 	const float DesiredResQuality = FindResolutionQualityForScreenSize(ClampedWidth, ClampedHeight);
-	return FMath::Max(DesiredResQuality, MinResolutionScale);
+	return FMath::Max(DesiredResQuality, Scalability::MinResolutionScale);
 }
 
 float UGameUserSettings::GetRecommendedResolutionScale()
 {
 	const float RecommendedResQuality = FindResolutionQualityForScreenSize(LastRecommendedScreenWidth, LastRecommendedScreenHeight);
 
-	return FMath::Max(RecommendedResQuality, MinResolutionScale);
+	return FMath::Max(RecommendedResQuality, Scalability::MinResolutionScale);
 }
 
 float UGameUserSettings::FindResolutionQualityForScreenSize(float Width, float Height)
@@ -760,7 +764,7 @@ int32 UGameUserSettings::GetOverallScalabilityLevel() const
 void UGameUserSettings::GetResolutionScaleInformationEx(float& CurrentScaleNormalized, float& CurrentScaleValue, float& MinScaleValue, float& MaxScaleValue) const
 {
 	CurrentScaleValue = ScalabilityQuality.ResolutionQuality;
-	MinScaleValue = MinResolutionScale;
+	MinScaleValue = Scalability::MinResolutionScale;
 	MaxScaleValue = Scalability::MaxResolutionScale;
 	CurrentScaleNormalized = ((float)CurrentScaleValue - (float)MinScaleValue) / (float)(MaxScaleValue - MinScaleValue);
 }
@@ -775,7 +779,7 @@ float UGameUserSettings::GetResolutionScaleNormalized() const
 
 void UGameUserSettings::SetResolutionScaleValueEx(float NewScaleValue)
 {
-	ScalabilityQuality.ResolutionQuality = FMath::Clamp(NewScaleValue, MinResolutionScale, Scalability::MaxResolutionScale);
+	ScalabilityQuality.ResolutionQuality = FMath::Clamp(NewScaleValue, Scalability::MinResolutionScale, Scalability::MaxResolutionScale);
 	const int32 ScreenWidth = (FullscreenMode == EWindowMode::WindowedFullscreen) ? GetDesktopResolution().X : ResolutionSizeX;
 	const int32 ScreenHeight = (FullscreenMode == EWindowMode::WindowedFullscreen) ? GetDesktopResolution().Y : ResolutionSizeY;
 	DesiredScreenWidth = ScreenWidth * ScalabilityQuality.ResolutionQuality / 100.0f;
@@ -784,7 +788,7 @@ void UGameUserSettings::SetResolutionScaleValueEx(float NewScaleValue)
 
 void UGameUserSettings::SetResolutionScaleNormalized(float NewScaleNormalized)
 {
-	const float RemappedValue = FMath::Lerp((float)MinResolutionScale, (float)Scalability::MaxResolutionScale, NewScaleNormalized);
+	const float RemappedValue = FMath::Lerp((float)Scalability::MinResolutionScale, (float)Scalability::MaxResolutionScale, NewScaleNormalized);
 	SetResolutionScaleValueEx(RemappedValue);
 }
 
diff --git a/Engine/Source/Runtime/Engine/Public/Scalability.h b/Engine/Source/Runtime/Engine/Public/Scalability.h
--- a/Engine/Source/Runtime/Engine/Public/Scalability.h
+++ b/Engine/Source/Runtime/Engine/Public/Scalability.h
@@ -213,7 +213,7 @@ namespace Scalability
 	ENGINE_API FQualityLevels GetQualityLevelCounts();
 
 	/** Minimum single axis scale for render resolution */
-	static const float MinResolutionScale = 10.0f;
+	static const float MinResolutionScale = 0.0f;
 
 	/** Maximum single axis scale for render resolution */
 	static const float MaxResolutionScale = 100.0f;
