NetDriver
* Now collecting the amount of packets sent & received by each connection each frame.
* Outputting in csv profiler the average and max per-connection packets received and sent every frame.
* Added net.ForceOnePacketPerBunch cheat to force every packet to contain a single bunch.  This can be used to debug rare ordering issues when two actors are dependant but should not.

diff --git a/Engine/Source/Runtime/Engine/Classes/Engine/NetConnection.h b/Engine/Source/Runtime/Engine/Classes/Engine/NetConnection.h
--- a/Engine/Source/Runtime/Engine/Classes/Engine/NetConnection.h
+++ b/Engine/Source/Runtime/Engine/Classes/Engine/NetConnection.h
@@ -519,6 +519,10 @@ public:
 	int32 InTotalBytes, OutTotalBytes;
 	/** packets sent/received on this connection (accumulated during a StatPeriod) */
 	int32 InPackets, OutPackets;
+	/** Packets received in the current tick */
+	int32 InPacketsThisFrame = 0;
+	/** Packets sent in the current tick */
+	int32 OutPacketsThisFrame = 0;
 	/** total packets sent/received on this connection */
 	int32 InTotalPackets, OutTotalPackets;
 	/** bytes sent/received on this connection (per second) - these are from previous StatPeriod interval */
diff --git a/Engine/Source/Runtime/Engine/Private/NetConnection.cpp b/Engine/Source/Runtime/Engine/Private/NetConnection.cpp
--- a/Engine/Source/Runtime/Engine/Private/NetConnection.cpp
+++ b/Engine/Source/Runtime/Engine/Private/NetConnection.cpp
@@ -1858,6 +1858,7 @@ void UNetConnection::ReceivedRawPacket( void* InData, int32 Count )
 	InBytes += PacketBytes;
 	InTotalBytes += PacketBytes;
 	++InPackets;
+	++InPacketsThisFrame;
 	++InTotalPackets;
 
 	if (Driver)
@@ -2161,9 +2162,8 @@ void UNetConnection::FlushNet(bool bIgnoreSimulation)
 		// Make sure that we always push an ChannelRecordEntry for each transmitted packet even if it is empty
 		FChannelRecordImpl::PushPacketId(ChannelRecord, OutPacketId);
 
-		
-
 		++OutPackets;
+		++OutPacketsThisFrame;
 		++OutTotalPackets;
 		Driver->OutPackets++;
 		Driver->OutTotalPackets++;
@@ -5628,6 +5628,29 @@ FAutoConsoleCommandWithWorldAndArgs AddimulatedConnectionsCmd(TEXT("net.Simulate
 
 FAutoConsoleCommandWithWorldAndArgs RemoveSimulatedConnectionsCmd(TEXT("net.DisconnectSimulatedConnections"), TEXT("Disconnects some simulated connections (0 = all)"), FConsoleCommandWithWorldAndArgsDelegate::CreateStatic(RemoveSimulatedNetConnections));
 
+FAutoConsoleCommandWithWorldAndArgs ForceOnePacketPerBunch(TEXT("net.ForceOnePacketPerBunch"), 
+	TEXT("When set to true it will enable AutoFlush on all connections and force a packet to be sent for every bunch we create. This forces one packet per replicated actor and can help find rare ordering bugs"), 
+	FConsoleCommandWithWorldAndArgsDelegate::CreateLambda([](const TArray<FString>& Args, UWorld* World)
+{
+	bool bEnable = true;
+	if (Args.Num() > 0)
+	{
+		LexTryParseString<bool>(bEnable, *Args[0]);
+	}
+
+	UE_LOG(LogNet, Display, TEXT("ForceOnePacketPerBunch set to %s"), bEnable?TEXT("true"):TEXT("false"));
+
+	for (TObjectIterator<UNetConnection> It; It; ++It)
+	{
+		// Only set autoflush for GameNetDriver connections
+		if (It->Driver == World->NetDriver)
+		{
+			It->SetAutoFlush(bEnable);
+		}
+	}
+}));
+
+
 // ----------------------------------------------------------------
 
 
diff --git a/Engine/Source/Runtime/Engine/Private/NetDriver.cpp b/Engine/Source/Runtime/Engine/Private/NetDriver.cpp
--- a/Engine/Source/Runtime/Engine/Private/NetDriver.cpp
+++ b/Engine/Source/Runtime/Engine/Private/NetDriver.cpp
@@ -5183,11 +5183,14 @@ struct FScopedNetDriverStats
 
 			int32 NumOpenChannels = 0;
 			int32 NumTickingChannels = 0;
+			int32 MaxOutPackets = 0;
+			int32 MaxInPackets = 0;
 
 			for (const UNetConnection* ClientConnection : NetDriver->ClientConnections)
 			{
 				NumOpenChannels += ClientConnection->OpenChannels.Num();
 				NumTickingChannels += ClientConnection->GetNumTickingChannels();
+				MaxOutPackets = FMath::Max(MaxOutPackets, ClientConnection->OutPackets);
 			}
 
 			CSV_CUSTOM_STAT(Replication, NumOpenChannels, (float)NumOpenChannels, ECsvCustomStatOp::Set);
@@ -7103,8 +7106,6 @@ void UNetDriver::UpdateNetworkStats()
 {
 	using namespace UE::Net::Private;
 
-	const double CurrentRealtimeSeconds = FPlatformTime::Seconds();
-
 	bool bCollectServerStats = false;
 #if USE_SERVER_PERF_COUNTERS || STATS
 	bCollectServerStats = true;
@@ -7113,6 +7114,8 @@ void UNetDriver::UpdateNetworkStats()
 	if (bCollectNetStats || bCollectServerStats)
 	{
 		SCOPE_CYCLE_COUNTER(STAT_NetTickFlushGatherStats);
+
+		const double CurrentRealtimeSeconds = FPlatformTime::Seconds();
 		// Update network stats (only main game net driver for now) if stats or perf counters are used
 		if (NetDriverName == NAME_GameNetDriver &&
 			CurrentRealtimeSeconds - StatUpdateTime > StatPeriod)
@@ -7123,12 +7126,16 @@ void UNetDriver::UpdateNetworkStats()
 			int32 ClientInPacketsMax = 0;
 			int32 ClientInPacketsMin = 0;
 			int32 ClientInPacketsAvg = 0;
+			int32 ClientsInPacketsThisFrameAvg = 0;
+			int32 ClientsInPacketsThisFrameMax = 0;
 			int32 ClientOutBytesMax = 0;
 			int32 ClientOutBytesMin = 0;
 			int32 ClientOutBytesAvg = 0;
 			int32 ClientOutPacketsMax = 0;
 			int32 ClientOutPacketsMin = 0;
 			int32 ClientOutPacketsAvg = 0;
+			int32 ClientsOutPacketsThisFrameAvg = 0;
+			int32 ClientsOutPacketsThisFrameMax = 0;
 			int NumClients = 0;
 			int32 MaxPacketOverhead = 0;
 
@@ -7151,9 +7158,16 @@ void UNetDriver::UpdateNetworkStats()
 					UpdatePerClientMinMaxAvg(OutPackets);
 
 					MaxPacketOverhead = FMath::Max(Client->PacketOverhead, MaxPacketOverhead);
-
 #undef UpdatePerClientMinMaxAvg
 
+					ClientsInPacketsThisFrameAvg += Client->InPacketsThisFrame;
+					ClientsOutPacketsThisFrameAvg += Client->OutPacketsThisFrame;
+					ClientsInPacketsThisFrameMax = FMath::Max(ClientsInPacketsThisFrameMax, Client->InPacketsThisFrame);
+					ClientsOutPacketsThisFrameMax = FMath::Max(ClientsOutPacketsThisFrameMax, Client->OutPacketsThisFrame);
+
+					Client->InPacketsThisFrame = 0;
+					Client->OutPacketsThisFrame = 0;
+
 					++NumClients;
 				}
 			}
@@ -7164,6 +7178,8 @@ void UNetDriver::UpdateNetworkStats()
 				ClientInPacketsAvg /= NumClients;
 				ClientOutBytesAvg /= NumClients;
 				ClientOutPacketsAvg /= NumClients;
+				ClientsInPacketsThisFrameAvg  /= NumClients;
+				ClientsOutPacketsThisFrameAvg /= NumClients;
 			}
 
 			int32 Ping = 0;
@@ -7472,6 +7488,13 @@ void UNetDriver::UpdateNetworkStats()
 			}
 #endif // USE_SERVER_PERF_COUNTERS
 
+#if CSV_PROFILER
+			CSV_CUSTOM_STAT(Replication, InPacketsClientAvg, ClientsInPacketsThisFrameAvg, ECsvCustomStatOp::Set);
+			CSV_CUSTOM_STAT(Replication, InPacketsClientMax, ClientsInPacketsThisFrameMax, ECsvCustomStatOp::Set);
+			CSV_CUSTOM_STAT(Replication, OutPacketsClientAvg, ClientsOutPacketsThisFrameAvg, ECsvCustomStatOp::Set);
+			CSV_CUSTOM_STAT(Replication, OutPacketsClientMax, ClientsOutPacketsThisFrameMax, ECsvCustomStatOp::Set);
+#endif
+
 			// Reset everything
 			InBytes = 0;
 			OutBytes = 0;
@@ -7491,7 +7514,68 @@ void UNetDriver::UpdateNetworkStats()
 			VoiceOutPercent = 0;
 			StatUpdateTime = CurrentRealtimeSeconds;
 		}
+		else
+		{
+#if CSV_PROFILER
+			// CSV stats need to be collected every frame
+
+			int32 ClientsInPacketsThisFrameAvg = 0;
+			int32 ClientsInPacketsThisFrameMax = 0;
+			int32 ClientsOutPacketsThisFrameAvg = 0;
+			int32 ClientsOutPacketsThisFrameMax = 0;
+
+			int32 NumClients = 0;
+
+			for (UNetConnection* Client : ClientConnections)
+			{
+				if (Client)
+				{
+					++NumClients;
+
+					ClientsInPacketsThisFrameAvg += Client->InPacketsThisFrame;
+					ClientsOutPacketsThisFrameAvg += Client->OutPacketsThisFrame;
+					ClientsInPacketsThisFrameMax = FMath::Max(ClientsInPacketsThisFrameMax, Client->InPacketsThisFrame);
+					ClientsOutPacketsThisFrameMax = FMath::Max(ClientsOutPacketsThisFrameMax, Client->OutPacketsThisFrame);
+
+					Client->InPacketsThisFrame = 0;
+					Client->OutPacketsThisFrame = 0;
+				}
+			}
+
+			if (NumClients > 1)
+			{
+				ClientsInPacketsThisFrameAvg /= NumClients;
+				ClientsOutPacketsThisFrameAvg /= NumClients;
+			}
+
+			CSV_CUSTOM_STAT(Replication, InPacketsClientAvg, ClientsInPacketsThisFrameAvg, ECsvCustomStatOp::Set);
+			CSV_CUSTOM_STAT(Replication, InPacketsClientMax, ClientsInPacketsThisFrameMax, ECsvCustomStatOp::Set);
+			CSV_CUSTOM_STAT(Replication, OutPacketsClientAvg, ClientsOutPacketsThisFrameAvg, ECsvCustomStatOp::Set);
+			CSV_CUSTOM_STAT(Replication, OutPacketsClientMax, ClientsOutPacketsThisFrameMax, ECsvCustomStatOp::Set);
+#else
+			// Reset the per-frame stats here
+			for (UNetConnection* Client : ClientConnections)
+			{
+				if (Client)
+				{
+					Client->InPacketsThisFrame = 0;
+					Client->OutPacketsThisFrame = 0;
+				}
+			}
+#endif
+		}
 	} // bCollectNetStats ||(USE_SERVER_PERF_COUNTERS) || STATS
+	else
+	{
+		for (UNetConnection* Client : ClientConnections)
+		{
+			if (Client)
+			{
+				Client->InPacketsThisFrame = 0;
+				Client->OutPacketsThisFrame = 0;
+			}
+		}
+	}
 }
 
 #if NET_DEBUG_RELEVANT_ACTORS
