Shader Compiler:  fix for crash processing cloned jobs, which don't have cache entries allocated.  Skip caching logic for that case.

#rnx

diff --git a/Engine/Source/Runtime/Engine/Private/ShaderCompiler/ShaderCompiler.cpp b/Engine/Source/Runtime/Engine/Private/ShaderCompiler/ShaderCompiler.cpp
--- a/Engine/Source/Runtime/Engine/Private/ShaderCompiler/ShaderCompiler.cpp
+++ b/Engine/Source/Runtime/Engine/Private/ShaderCompiler/ShaderCompiler.cpp
@@ -1345,7 +1345,9 @@ void FShaderJobCache::ProcessFinishedJob(FShaderCommonCompileJob* FinishedJob, b
 
 void FShaderJobCache::AddToCacheAndProcessPending(FShaderCommonCompileJob* FinishedJob)
 {
-	if (!ShaderCompiler::IsJobCacheEnabled())
+	// Cloned jobs won't include an entry in the job cache, so skip the caching logic.  The non-cloned version of the same
+	// job will handle adding data to the cache when it completes.
+	if (!ShaderCompiler::IsJobCacheEnabled() || !FinishedJob->JobCacheRef.Block)
 	{
 		return;
 	}
