Make a cast to int32 explicit to avoid warnings in calling code.

diff --git a/Engine/Source/Runtime/Engine/Classes/PhysicsEngine/PhysicsSettings.h b/Engine/Source/Runtime/Engine/Classes/PhysicsEngine/PhysicsSettings.h
--- a/Engine/Source/Runtime/Engine/Classes/PhysicsEngine/PhysicsSettings.h
+++ b/Engine/Source/Runtime/Engine/Classes/PhysicsEngine/PhysicsSettings.h
@@ -258,7 +258,7 @@ public:
 	UFUNCTION(BlueprintCallable, Category = "Physics")
 	int32 GetPhysicsHistoryCount() const
 	{
-		return FMath::CeilToInt(1.f / AsyncFixedTimeStepSize) * (PhysicsPrediction.MaxSupportedLatencyPrediction / 1000.f);
+		return (int32) (FMath::CeilToInt(1.f / AsyncFixedTimeStepSize) * (PhysicsPrediction.MaxSupportedLatencyPrediction / 1000.f));
 	}
 
 	ENGINE_API virtual void PostInitProperties() override;
