UWorld: Logging: Downgrade CleanupWorld logmessage to verbose to make cookcommandlet logs easier to read.
#rnx

diff --git a/Engine/Source/Runtime/Engine/Private/World.cpp b/Engine/Source/Runtime/Engine/Private/World.cpp
--- a/Engine/Source/Runtime/Engine/Private/World.cpp
+++ b/Engine/Source/Runtime/Engine/Private/World.cpp
@@ -5314,7 +5314,19 @@ void UWorld::CleanupWorldInternal(bool bSessionEnded, bool bCleanupResources, bo
 	}
 	CleanupWorldTag = CleanupWorldGlobalTag;
 
-	UE_LOG(LogWorld, Log, TEXT("UWorld::CleanupWorld for %s, bSessionEnded=%s, bCleanupResources=%s"), *GetName(), bSessionEnded ? TEXT("true") : TEXT("false"), bCleanupResources ? TEXT("true") : TEXT("false"));
+	// Downgrade verbosity when running commandlet to reduce spam; this message isn't as important in commandlets
+#if !NO_LOGGING
+	FString Message = FString::Printf(TEXT("UWorld::CleanupWorld for %s, bSessionEnded=%s, bCleanupResources=%s"),
+		*GetName(), bSessionEnded ? TEXT("true") : TEXT("false"), bCleanupResources ? TEXT("true") : TEXT("false"));
+	if (IsRunningCommandlet())
+	{
+		UE_LOG(LogWorld, Verbose, TEXT("%s"), *Message);
+	}
+	else
+	{
+		UE_LOG(LogWorld, Log, TEXT("%s"), *Message);
+	}
+#endif
 
 	check(IsVisibilityRequestPending() == false);
 	
