Fix route gamepad to second window. Gamepad input was never being routed to the next viewport client during a PIE session with multiple viewports.

diff --git a/Engine/Source/Runtime/Engine/Private/GameViewportClient.cpp b/Engine/Source/Runtime/Engine/Private/GameViewportClient.cpp
--- a/Engine/Source/Runtime/Engine/Private/GameViewportClient.cpp
+++ b/Engine/Source/Runtime/Engine/Private/GameViewportClient.cpp
@@ -629,6 +629,32 @@ bool UGameViewportClient::InputKey(const FInputKeyEventArgs& InEventArgs)
 
 	RemapControllerInput(EventArgs);
 
+
+#if WITH_EDITOR
+	if (EventArgs.Key.IsGamepadKey())
+	{
+		/** For PIE, since this is gamepad, check if we want to route gamepad to second window.
+		 * Let the next PIE window handle the input (this allows people to use a controller for one window and kbm for the other).
+		 */
+		const FViewportClient* InViewportClient = InEventArgs.Viewport != nullptr ? InEventArgs.Viewport->GetClient() : nullptr;
+		if (InViewportClient == this)
+		{
+			const ULevelEditorPlaySettings* PlayInSettings = GetDefault<ULevelEditorPlaySettings>();
+			const bool CanRouteGamepadToSecondWindow = [&PlayInSettings] { bool RouteGamepadToSecondWindow(false); return (PlayInSettings->GetRouteGamepadToSecondWindow(RouteGamepadToSecondWindow) && RouteGamepadToSecondWindow); }();
+			const bool CanRunUnderOneProcess = [&PlayInSettings] { bool RunUnderOneProcess(false); return (PlayInSettings->GetRunUnderOneProcess(RunUnderOneProcess) && RunUnderOneProcess); }();
+			if (CanRouteGamepadToSecondWindow && CanRunUnderOneProcess && InEventArgs.Viewport->IsPlayInEditorViewport())
+			{
+				if (UGameViewportClient* NextViewport = GEngine->GetNextPIEViewport(this))
+				{
+					const bool bResult = NextViewport->InputKey(InEventArgs);
+					return false;
+				}
+			}
+		}	
+	}
+#endif
+	
+
 	if (IgnoreInput())
 	{
 		return ViewportConsole ? ViewportConsole->InputKey(EventArgs.InputDevice, EventArgs.Key, EventArgs.Event, EventArgs.AmountDepressed, EventArgs.IsGamepad()) : false;
@@ -698,6 +724,32 @@ bool UGameViewportClient::InputAxis(FViewport* InViewport, FInputDeviceId InputD
 	// Handle mapping controller id and key if needed
 	FInputKeyEventArgs EventArgs(InViewport, InputDevice, Key, IE_Axis);
 
+	RemapControllerInput(EventArgs);
+
+#if WITH_EDITOR
+	if (bGamepad)
+	{
+		/** For PIE, since this is gamepad, check if we want to route gamepad to second window.
+	 * Let the next PIE window handle the input (this allows people to use a controller for one window and kbm for the other).
+	 */
+		const FViewportClient* InViewportClient = InViewport != nullptr ? InViewport->GetClient() : nullptr;
+		if (InViewportClient == this)
+		{
+			const ULevelEditorPlaySettings* PlayInSettings = GetDefault<ULevelEditorPlaySettings>();
+			const bool CanRouteGamepadToSecondWindow = [&PlayInSettings] { bool RouteGamepadToSecondWindow(false); return (PlayInSettings->GetRouteGamepadToSecondWindow(RouteGamepadToSecondWindow) && RouteGamepadToSecondWindow); }();
+			const bool CanRunUnderOneProcess = [&PlayInSettings] { bool RunUnderOneProcess(false); return (PlayInSettings->GetRunUnderOneProcess(RunUnderOneProcess) && RunUnderOneProcess); }();
+			if (CanRouteGamepadToSecondWindow && CanRunUnderOneProcess && InViewport->IsPlayInEditorViewport())
+			{
+				if (UGameViewportClient* NextViewport = GEngine->GetNextPIEViewport(this))
+				{
+					const bool bResult = NextViewport->InputAxis(InViewport, InputDevice, Key, Delta, DeltaTime, NumSamples, bGamepad);
+					return false;
+				}
+			}
+		}
+	}
+#endif
+	
 	OnInputAxisEvent.Broadcast(InViewport, EventArgs.ControllerId, EventArgs.Key, Delta, DeltaTime, NumSamples, EventArgs.IsGamepad());
 	
 	bool bResult = false;
