Fix for audio not previewing or playing in sequences when cooked audio is referenced in current level.

diff --git a/Engine/Source/Runtime/Engine/Private/SoundWave.cpp b/Engine/Source/Runtime/Engine/Private/SoundWave.cpp
--- a/Engine/Source/Runtime/Engine/Private/SoundWave.cpp
+++ b/Engine/Source/Runtime/Engine/Private/SoundWave.cpp
@@ -1012,7 +1012,6 @@ void USoundWave::Serialize( FArchive& Ar )
 	FByteBulkData TempOldBulkData;
 	bool bBulkDataConverted = false;
 
-	bLoadedFromCookedData = Ar.IsLoading() && bCooked;
 	if (bVirtualizeWhenSilent_DEPRECATED)
 	{
 		bVirtualizeWhenSilent_DEPRECATED = 0;
@@ -1021,6 +1020,9 @@ void USoundWave::Serialize( FArchive& Ar )
 
 	if (Ar.IsLoading())
 	{
+		// Modify bLoadedFromCookedData state only when loading from archive
+		bLoadedFromCookedData = bCooked;
+
 		ModulationSettings.VersionModulators();
 	}
 
