WorldPartition: Fix invalid bounds when saving an unregistered actor (Fixup redirector is one way of triggering this path). This was leading to a failed check in WorldPartition streaming generation.

#rnx

diff --git a/Engine/Source/Runtime/Engine/Public/WorldPartition/WorldPartitionActorDesc.h b/Engine/Source/Runtime/Engine/Public/WorldPartition/WorldPartitionActorDesc.h
--- a/Engine/Source/Runtime/Engine/Public/WorldPartition/WorldPartitionActorDesc.h
+++ b/Engine/Source/Runtime/Engine/Public/WorldPartition/WorldPartitionActorDesc.h
@@ -278,6 +278,7 @@ protected:
 	{
 		BoundsLocation = From->BoundsLocation;
 		BoundsExtent = From->BoundsExtent;
+		bIsBoundsValid = From->bIsBoundsValid;
 	}
 
 	virtual uint32 GetSizeOf() const { return sizeof(FWorldPartitionActorDesc); }
