Fix static analysis warnings on the Game Viewport client, which were fallout of the fix for UE-169979.

#rnx

diff --git a/Engine/Source/Runtime/Engine/Private/GameViewportClient.cpp b/Engine/Source/Runtime/Engine/Private/GameViewportClient.cpp
--- a/Engine/Source/Runtime/Engine/Private/GameViewportClient.cpp
+++ b/Engine/Source/Runtime/Engine/Private/GameViewportClient.cpp
@@ -727,12 +727,12 @@ bool UGameViewportClient::InputAxis(FViewport* InViewport, FInputDeviceId InputD
 	RemapControllerInput(EventArgs);
 
 #if WITH_EDITOR
-	if (bGamepad)
+	if (bGamepad && InViewport)
 	{
 		/** For PIE, since this is gamepad, check if we want to route gamepad to second window.
 		 * Let the next PIE window handle the input (this allows people to use a controller for one window and kbm for the other).
 		 */
-		const FViewportClient* InViewportClient = InViewport != nullptr ? InViewport->GetClient() : nullptr;
+		const FViewportClient* InViewportClient = InViewport->GetClient();
 		if (InViewportClient == this)
 		{
 			const ULevelEditorPlaySettings* PlayInSettings = GetDefault<ULevelEditorPlaySettings>();
@@ -757,33 +757,36 @@ bool UGameViewportClient::InputAxis(FViewport* InViewport, FInputDeviceId InputD
 	// Don't allow mouse/joystick input axes while in PIE and the console has forced the cursor to be visible.  It's
 	// just distracting when moving the mouse causes mouse look while you are trying to move the cursor over a button
 	// in the editor!
-	if( !( InViewport->IsSlateViewport() && InViewport->IsPlayInEditorViewport() ) || ViewportConsole == nullptr || !ViewportConsole->ConsoleActive() )
+	if (InViewport)
 	{
-		// route to subsystems that care
-		if (ViewportConsole != nullptr)
+		if( !( InViewport->IsSlateViewport() && InViewport->IsPlayInEditorViewport() ) || ViewportConsole == nullptr || !ViewportConsole->ConsoleActive() )
 		{
-			bResult = ViewportConsole->InputAxis(EventArgs.InputDevice, EventArgs.Key, Delta, DeltaTime, NumSamples, EventArgs.IsGamepad());
-		}
+			// route to subsystems that care
+			if (ViewportConsole != nullptr)
+			{
+				bResult = ViewportConsole->InputAxis(EventArgs.InputDevice, EventArgs.Key, Delta, DeltaTime, NumSamples, EventArgs.IsGamepad());
+			}
 		
-		// Try the override callback, this may modify event args
-		if (!bResult && OnOverrideInputAxisEvent.IsBound())
-		{
-			bResult = OnOverrideInputAxisEvent.Execute(EventArgs, Delta, DeltaTime, NumSamples);
-		}
+			// Try the override callback, this may modify event args
+			if (!bResult && OnOverrideInputAxisEvent.IsBound())
+			{
+				bResult = OnOverrideInputAxisEvent.Execute(EventArgs, Delta, DeltaTime, NumSamples);
+			}
 
-		if (!bResult)
-		{
-			ULocalPlayer* const TargetPlayer = GEngine->GetLocalPlayerFromInputDevice(this, EventArgs.InputDevice);
-			if (TargetPlayer && TargetPlayer->PlayerController)
+			if (!bResult)
 			{
-				bResult = TargetPlayer->PlayerController->InputKey(FInputKeyParams(EventArgs.Key, (double)Delta, DeltaTime, NumSamples, EventArgs.IsGamepad(), EventArgs.InputDevice));
+				ULocalPlayer* const TargetPlayer = GEngine->GetLocalPlayerFromInputDevice(this, EventArgs.InputDevice);
+				if (TargetPlayer && TargetPlayer->PlayerController)
+				{
+					bResult = TargetPlayer->PlayerController->InputKey(FInputKeyParams(EventArgs.Key, (double)Delta, DeltaTime, NumSamples, EventArgs.IsGamepad(), EventArgs.InputDevice));
+				}
 			}
-		}
 
-		if( InViewport->IsSlateViewport() && InViewport->IsPlayInEditorViewport() )
-		{
-			// Absorb all keys so game input events are not routed to the Slate editor frame
-			bResult = true;
+			if( InViewport->IsSlateViewport() && InViewport->IsPlayInEditorViewport() )
+			{
+				// Absorb all keys so game input events are not routed to the Slate editor frame
+				bResult = true;
+			}
 		}
 	}
 
