Iris Fix for crash in CharacterMovementComponent if PackageMap cannot be found

diff --git a/Engine/Source/Runtime/Engine/Private/Components/CharacterMovementComponent.cpp b/Engine/Source/Runtime/Engine/Private/Components/CharacterMovementComponent.cpp
--- a/Engine/Source/Runtime/Engine/Private/Components/CharacterMovementComponent.cpp
+++ b/Engine/Source/Runtime/Engine/Private/Components/CharacterMovementComponent.cpp
@@ -9339,6 +9339,12 @@ void UCharacterMovementComponent::ServerMovePacked_ServerReceive(const FCharacte
 		ServerMoveBitReader.PackageMap = PackedBits.GetPackageMap();
 	}
 
+	if (ServerMoveBitReader.PackageMap == nullptr)
+	{
+		devCode(UE_LOG(LogNetPlayerMovement, Error, TEXT("ServerMovePacked_ServerReceive: Failed to find PackageMap for data serialization!")));
+		return;
+	}
+
 	// Deserialize bits to move data struct.
 	// We had to wait until now and use the temp bit stream because the RPC doesn't know about the virtual overrides on the possibly custom struct that is our data container.
 	FCharacterNetworkMoveDataContainer& MoveDataContainer = GetNetworkMoveDataContainer();
