[shaders]
- move responsibility for logging shader stats out of ShaderCompiler and into a tickable editor/cook object
- add aggregation of shader stats for multiprocess cooks
- add an option to report stats individually on all cook processes driven by cvar (the aggregated stats will always be reported)

#rnx

diff --git a/Engine/Source/Editor/UnrealEd/Private/CookOnTheFlyServer.cpp b/Engine/Source/Editor/UnrealEd/Private/CookOnTheFlyServer.cpp
--- a/Engine/Source/Editor/UnrealEd/Private/CookOnTheFlyServer.cpp
+++ b/Engine/Source/Editor/UnrealEd/Private/CookOnTheFlyServer.cpp
@@ -124,6 +124,8 @@
 #include "Settings/ProjectPackagingSettings.h"
 #include "ShaderCodeLibrary.h"
 #include "ShaderCompiler.h"
+#include "ShaderStats.h"
+#include "ShaderStatsCollector.h"
 #include "ShaderLibraryChunkDataGenerator.h"
 #include "PipelineCacheChunkDataGenerator.h"
 #include "String/Find.h"
@@ -9028,6 +9030,15 @@ void UCookOnTheFlyServer::BeginCookStartShaderCodeLibrary(FBeginCookContext& Beg
 		}
 	}
 
+	if (CookDirector)
+	{
+		CookDirector->Register(new FShaderStatsAggregator(FShaderStatsAggregator::EMode::Director));
+	}
+	else if (CookWorkerClient)
+	{
+		CookWorkerClient->Register(new FShaderStatsAggregator(FShaderStatsAggregator::EMode::Worker));
+	}
+
 	if (!IsCookWorkerMode())
 	{
 		CleanShaderCodeLibraries();
@@ -9833,9 +9844,11 @@ void UCookOnTheFlyServer::PrintFinishStats()
 
 void UCookOnTheFlyServer::PrintDetailedCookStats()
 {
-	if (GShaderCompilerStats)
+	// Stats are aggregated on the director, so writing the stats CSVs is only needed on the director 
+	// (or single cook process in the sp cook case)
+	if (!IsCookWorkerMode())
 	{
-		GShaderCompilerStats->WriteStats();
+		FShaderStatsFunctions::WriteShaderStats();
 	}
 
 	COOK_STAT(
diff --git a/Engine/Source/Editor/UnrealEd/Private/ShaderStatsCollector.cpp b/Engine/Source/Editor/UnrealEd/Private/ShaderStatsCollector.cpp
new file mode 100644
--- /dev/null
+++ b/Engine/Source/Editor/UnrealEd/Private/ShaderStatsCollector.cpp
@@ -0,0 +1,164 @@
+// Copyright Epic Games, Inc. All Rights Reserved.
+
+#include "ShaderStatsCollector.h"
+#include "ShaderStats.h"
+
+static TAutoConsoleVariable<int> CVarShaderCompilerStatsPrintoutInterval(
+	TEXT("r.ShaderCompiler.StatsPrintoutInterval"),
+	180,
+	TEXT("Minimum interval (in seconds) between printing out debugging stats (by default, no closer than once per minute)."),
+	ECVF_Default
+);
+
+static TAutoConsoleVariable<bool> CVarShaderCompilerPrintIndividualProcessStats(
+	TEXT("r.ShaderCompiler.PrintIndividualProcessStats"),
+	false,
+	TEXT("If true, in a multiprocess cook, stats will be printed for each individual process (as well as aggregated across all)."),
+	ECVF_Default
+);
+
+class FShaderStatsReporter : public FTickableEditorObject, FTickableCookObject
+{
+public:
+	FShaderStatsReporter() : FTickableEditorObject(), FTickableCookObject() {}
+	virtual void Tick(float DeltaTime) override;
+	virtual void TickCook(float DeltaTime, bool bCookComplete) override;
+	virtual bool IsTickable() const override;
+	virtual TStatId GetStatId() const override;
+	
+	inline void SetAggregator(FShaderStatsAggregator* InAggregator)
+	{
+		Aggregator = InAggregator;
+	}
+
+private:
+	friend class FShaderStatsFunctions;
+	void InternalTick(float DeltaTime, bool bForceLog = false);
+	void LogStats();
+	void AggregateStats(FShaderCompilerStats& OutStats) const;
+
+	FShaderStatsAggregator* Aggregator;
+	double LastLogTime = 0.0f;
+	double LastSynchTime = 0.0f;
+};
+
+static FShaderStatsReporter GShaderStatsReporter;
+
+void FShaderStatsFunctions::GatherShaderAnalytics(TArray<FAnalyticsEventAttribute>& Attributes)
+{
+	FShaderCompilerStats AggregatedCompilerStats;
+	GShaderCompilingManager->GetLocalStats(AggregatedCompilerStats);
+	GShaderStatsReporter.AggregateStats(AggregatedCompilerStats);
+	AggregatedCompilerStats.GatherAnalytics(TEXT("Shaders_"), Attributes);
+}
+
+void FShaderStatsFunctions::WriteShaderStats()
+{
+	FShaderCompilerStats AggregatedCompilerStats;
+	GShaderCompilingManager->GetLocalStats(AggregatedCompilerStats);
+	GShaderStatsReporter.AggregateStats(AggregatedCompilerStats);
+	AggregatedCompilerStats.WriteStats();
+}
+
+void FShaderStatsReporter::Tick(float DeltaTime)
+{
+	InternalTick(DeltaTime);
+}
+
+void FShaderStatsReporter::TickCook(float DeltaTime, bool bCookComplete)
+{
+	// force log output at completion of cook
+	InternalTick(DeltaTime, bCookComplete);
+}
+
+void FShaderStatsReporter::InternalTick(float DeltaTime, bool bForceLog)
+{
+	const int PrintInterval = CVarShaderCompilerStatsPrintoutInterval.GetValueOnAnyThread();
+	if (bForceLog || ((PrintInterval > 0) && (FPlatformTime::Seconds() - LastLogTime) >= PrintInterval))
+	{
+		LastLogTime = FPlatformTime::Seconds();
+		LogStats();
+	}
+}
+
+bool FShaderStatsReporter::IsTickable() const
+{
+	// the stats should be reported (and so stats reporter tickable) in the case where we're the cook director process in
+	// a multiprocess cook, running a single process cook/editor, or if the "print individual process stats" cvar is enabled.
+	// the Aggregator will have Mode==Director the first case, and be unset in the second.
+	return !Aggregator || (Aggregator->Mode == FShaderStatsAggregator::EMode::Director) || CVarShaderCompilerPrintIndividualProcessStats.GetValueOnAnyThread();
+}
+
+TStatId FShaderStatsReporter::GetStatId() const
+{
+	RETURN_QUICK_DECLARE_CYCLE_STAT(FShaderStatsReporter, STATGROUP_Tickables);
+}
+
+void FShaderStatsReporter::LogStats()
+{
+	// only ever called on director process in MP cook, or single process in normal cook/editor
+	FShaderCompilerStats LocalStats;
+	GShaderCompilingManager->GetLocalStats(LocalStats);
+
+	if (Aggregator && Aggregator->Mode == FShaderStatsAggregator::EMode::Director && CVarShaderCompilerPrintIndividualProcessStats.GetValueOnAnyThread())
+	{
+		// if we're the director in a MP cook, and the "individual process stats" are requested, print them before aggregating and printing the aggregated results
+		LocalStats.WriteStatSummary();
+	}
+	
+	AggregateStats(LocalStats);
+	LocalStats.WriteStatSummary();
+}
+
+void FShaderStatsReporter::AggregateStats(FShaderCompilerStats& OutStats) const
+{
+	if (Aggregator && Aggregator->Mode == FShaderStatsAggregator::EMode::Director)
+	{
+		OutStats.SetMultiProcessAggregated();
+		for (FShaderCompilerStats& WorkerStats : Aggregator->WorkerCompilerStats)
+		{
+			OutStats.Aggregate(WorkerStats);
+		}
+	}
+}
+
+FGuid FShaderStatsAggregator::MessageType(TEXT("F8D2E291D9A44F999D79E3839091BF25"));
+
+FShaderStatsAggregator::FShaderStatsAggregator(EMode Mode) : Mode(Mode)
+{
+	GShaderStatsReporter.SetAggregator(this);
+}
+
+void FShaderStatsAggregator::ClientTick(UE::Cook::FMPCollectorClientTickContext& Context)
+{
+	const int PrintInterval = CVarShaderCompilerStatsPrintoutInterval.GetValueOnAnyThread();
+	// synch stats twice as often as they are printed to ensure each stats print on the cook director
+	// has at least one update from all workers.
+	const int SynchInterval = PrintInterval / 2;
+
+	if (Context.IsFlush() || ((SynchInterval > 0) && (FPlatformTime::Seconds() - LastSynchTime) >= PrintInterval))
+	{
+		LastSynchTime = FPlatformTime::Seconds();
+		FCbWriter Writer;
+		Writer.BeginObject();
+
+		FShaderCompilerStats LocalStats;
+		GShaderCompilingManager->GetLocalStats(LocalStats);
+		LocalStats.WriteToCompactBinary(Writer);
+
+		Writer.EndObject();
+
+		Context.AddMessage(Writer.Save().AsObject());
+	}
+}
+
+void FShaderStatsAggregator::ServerReceiveMessage(UE::Cook::FMPCollectorServerMessageContext& Context, FCbObjectView Message)
+{
+	int32 WorkerId = Context.GetProfileId();
+
+	if (WorkerCompilerStats.Num() < (WorkerId + 1))
+	{
+		WorkerCompilerStats.SetNum(WorkerId + 1);
+	}
+	WorkerCompilerStats[WorkerId].ReadFromCompactBinary(Message);
+}
diff --git a/Engine/Source/Editor/UnrealEd/Private/ShaderStatsCollector.h b/Engine/Source/Editor/UnrealEd/Private/ShaderStatsCollector.h
new file mode 100644
--- /dev/null
+++ b/Engine/Source/Editor/UnrealEd/Private/ShaderStatsCollector.h
@@ -0,0 +1,30 @@
+// Copyright Epic Games, Inc. All Rights Reserved.
+
+#pragma once
+
+#include "Cooker/CookMPCollector.h"
+#include "TickableEditorObject.h"
+
+class FShaderStatsAggregator : public UE::Cook::IMPCollector{
+public:
+	enum class EMode
+	{
+		Worker,
+		Director,
+	};
+	FShaderStatsAggregator(EMode Mode);
+	virtual FGuid GetMessageType() const override { return MessageType; }
+	virtual const TCHAR* GetDebugName() const override { return TEXT("FShaderStatsAggregator"); }
+
+	virtual void ClientTick(UE::Cook::FMPCollectorClientTickContext& Context) override;
+	virtual void ServerReceiveMessage(UE::Cook::FMPCollectorServerMessageContext& Context, FCbObjectView Message) override;
+
+	static FGuid MessageType;
+
+private:
+	friend class FShaderStatsFunctions;
+	friend class FShaderStatsReporter;
+	EMode Mode;
+	TArray<FShaderCompilerStats> WorkerCompilerStats;
+	float LastSynchTime = 0.0f;
+};
diff --git a/Engine/Source/Editor/UnrealEd/Private/UnrealEdMisc.cpp b/Engine/Source/Editor/UnrealEd/Private/UnrealEdMisc.cpp
--- a/Engine/Source/Editor/UnrealEd/Private/UnrealEdMisc.cpp
+++ b/Engine/Source/Editor/UnrealEd/Private/UnrealEdMisc.cpp
@@ -549,7 +549,7 @@ void FUnrealEdMisc::OnInit()
 
 	FEditorDelegates::OnEditorInitialized.Broadcast(TotalEditorStartupTime);
 	
-	GShaderCompilingManager->PrintStats(true);
+	GShaderCompilingManager->PrintStats();
 }
 
 FString FUnrealEdMisc::FindMapFileFromPartialName(const FString& PartialMapName)
diff --git a/Engine/Source/Editor/UnrealEd/Public/ShaderStats.h b/Engine/Source/Editor/UnrealEd/Public/ShaderStats.h
new file mode 100644
--- /dev/null
+++ b/Engine/Source/Editor/UnrealEd/Public/ShaderStats.h
@@ -0,0 +1,14 @@
+// Copyright Epic Games, Inc. All Rights Reserved.
+
+#pragma once
+
+#include "Containers/Array.h"
+
+struct FAnalyticsEventAttribute;
+
+class FShaderStatsFunctions
+{
+public:
+	UNREALED_API static void GatherShaderAnalytics(TArray<FAnalyticsEventAttribute>& Attributes);
+	static void WriteShaderStats();
+};
diff --git a/Engine/Source/Runtime/Engine/Private/ShaderCompiler/ShaderCompiler.cpp b/Engine/Source/Runtime/Engine/Private/ShaderCompiler/ShaderCompiler.cpp
--- a/Engine/Source/Runtime/Engine/Private/ShaderCompiler/ShaderCompiler.cpp
+++ b/Engine/Source/Runtime/Engine/Private/ShaderCompiler/ShaderCompiler.cpp
@@ -43,6 +43,7 @@
 #include "RenderUtils.h"
 #include "SceneInterface.h"
 #include "SceneManagement.h"
+#include "Serialization/CompactBinaryWriter.h"
 #include "Serialization/NameAsStringProxyArchive.h"
 #include "ShaderCodeLibrary.h"
 #include "ShaderPlatformCachedIniValue.h"
@@ -192,14 +193,6 @@ bool IsShaderJobCacheDDCEnabled()
 	return false;
 }
 
-int32 GShaderCompilerCacheStatsPrintoutInterval = 180;
-static FAutoConsoleVariableRef CVarShaderCompilerCacheStatsPrintoutInterval(
-	TEXT("r.ShaderCompiler.CacheStatsPrintoutInterval"),
-	GShaderCompilerCacheStatsPrintoutInterval,
-	TEXT("Minimum interval (in seconds) between printing out debugging stats (by default, no closer than each 3 minutes)."),
-	ECVF_Default
-);
-
 int32 GShaderCompilerAllowDistributedCompilation = 1;
 static FAutoConsoleVariableRef CVarShaderCompilerAllowDistributedCompilation(
 	TEXT("r.ShaderCompiler.AllowDistributedCompilation"),
@@ -432,11 +425,8 @@ public:
 	/** Adds the job to cache. */
 	void AddToCacheAndProcessPending(FShaderCommonCompileJob* FinishedJob);
 
-	/** Log caching statistics.*/
-	void LogCachingStats();
-
-	/** Gather statistics to send to analytics */
-	void GatherAnalytics(const FString& BaseName, TArray<FAnalyticsEventAttribute>& Attributes) const;
+	/** Populates caching stats in the given compiler stats struct. */
+	void GetStats(FShaderCompilerStats& OutStats) const;
 
 	int32 GetNumPendingJobs(EShaderCompileJobPriority InPriority) const;
 
@@ -548,9 +538,6 @@ private:
 	/** Compute memory used by the cache from scratch.  Should match GetAllocatedMemory() if CurrentlyAllocateMemory is being properly updated (useful for validation). */
 	uint64 ComputeAllocatedMemory() const;
 
-	/** Logs out the statistics */
-	void LogStats();
-
 	/** Calculates current memory budget, in bytes */
 	uint64 GetCurrentMemoryBudget() const;
 
@@ -927,13 +914,9 @@ void FShaderCompileJobCollection::AddToCacheAndProcessPending(FShaderCommonCompi
 {
 	JobsCache->AddToCacheAndProcessPending(FinishedJob);
 }
-void FShaderCompileJobCollection::LogCachingStats()
-{
-	JobsCache->LogCachingStats();
-}
-void FShaderCompileJobCollection::GatherAnalytics(const FString& BaseName, TArray<FAnalyticsEventAttribute>& Attributes) const
+void FShaderCompileJobCollection::GetCachingStats(FShaderCompilerStats& OutStats) const
 {
-	JobsCache->GatherAnalytics(BaseName, Attributes);
+	JobsCache->GetStats(OutStats);
 }
 int32 FShaderCompileJobCollection::GetNumPendingJobs(EShaderCompileJobPriority InPriority) const
 {
@@ -1314,7 +1297,7 @@ void FShaderJobCache::SubmitJobs(const TArray<FShaderCommonCompileJobPtr>& InJob
 
 void FShaderCompileJobCollection::HandlePrintStats()
 {
-	GShaderCompilingManager->PrintStats(true);
+	GShaderCompilingManager->PrintStats();
 }
 
 TRACE_DECLARE_INT_COUNTER(Shaders_Compiled, TEXT("Shaders/Compiled"));
@@ -1406,12 +1389,6 @@ void FShaderJobCache::AddToCacheAndProcessPending(FShaderCommonCompileJob* Finis
 	FinishedJob->JobCacheRef.Clear();
 }
 
-void FShaderJobCache::LogCachingStats()
-{
-	FWriteScopeLock Locker(JobLock);	// write lock because logging actually changes the cache state (in a minor way - updating the memory used - but still).
-	LogStats();
-}
-
 int32 FShaderJobCache::GetNumPendingJobs(EShaderCompileJobPriority InPriority) const
 {
 	return NumPendingJobs[(int32)InPriority];
@@ -3821,13 +3798,8 @@ void FShaderCompilerStats::WriteStats(FOutputDevice* Ar)
 	{
 		FlushRenderingCommands();
 
-		uint32 MultiprocessId = 0;
-		FParse::Value(FCommandLine::Get(), TEXT("-MultiprocessId="), MultiprocessId);
 		FString FileName = FPaths::Combine(*FPaths::ProjectSavedDir(),
-			FString::Printf(TEXT("MaterialStats/Stats-%s%s.csv"),
-				*FDateTime::Now().ToString(),
-				(MultiprocessId == 0 ? TEXT("") : *FString::Printf(TEXT("-%d"), MultiprocessId))
-			));
+			FString::Printf(TEXT("MaterialStats/Stats-%s.csv"), *FDateTime::Now().ToString()));
 		auto DebugWriter = IFileManager::Get().CreateFileWriter(*FileName);
 		FDiagnosticTableWriterCSV StatWriter(DebugWriter);
 		const TSparseArray<ShaderCompilerStats>& PlatformStats = GetShaderCompilerStats();
@@ -3911,12 +3883,7 @@ void FShaderCompilerStats::WriteStats(FOutputDevice* Ar)
 	}
 	{
 
-		uint32 MultiprocessId = 0;
-		FParse::Value(FCommandLine::Get(), TEXT("-MultiprocessId="), MultiprocessId);
-		FString FileName = FString::Printf(TEXT("%s/MaterialStatsDebug/StatsDebug-%s%s.csv"),
-			*FPaths::ProjectSavedDir(), *FDateTime::Now().ToString(),
-			(MultiprocessId == 0 ? TEXT("") : *FString::Printf(TEXT("-%d"), MultiprocessId))
-		);
+		FString FileName = FString::Printf(TEXT("%s/MaterialStatsDebug/StatsDebug-%s.csv"), *FPaths::ProjectSavedDir(), *FDateTime::Now().ToString());
 		auto DebugWriter = IFileManager::Get().CreateFileWriter(*FileName);
 		FDiagnosticTableWriterCSV StatWriter(DebugWriter);
 		const TSparseArray<ShaderCompilerStats>& PlatformStats = GetShaderCompilerStats();
@@ -3995,70 +3962,109 @@ static FString PrintJobsCompletedPercentageToString(int64 JobsAssigned, int64 Jo
 void FShaderCompilerStats::WriteStatSummary()
 {
 	const uint32 TotalCompiled = GetTotalShadersCompiled();
+	if (TotalCompiled == 0)
+	{
+		// early out if we haven't done anything yet
+		return;
+	}
+
+	UE_LOG(LogShaderCompilers, Display, TEXT("================================================"));
+
+	const TCHAR* AggregatedSuffix = bMultiProcessAggregated ? TEXT(" (aggregated across all cook processes)") : TEXT("");
+
+	// Only log cache stats if the cache has been queried at least once (this will always be 0 if the job cache is disabled)
+	if (Counters.TotalCacheSearchAttempts > 0)
+	{
+		UE_LOG(LogShaderCompilers, Display, TEXT("=== FShaderJobCache stats%s ==="), AggregatedSuffix);
+		UE_LOG(LogShaderCompilers, Display, TEXT("Total job queries %s, among them cache hits %s (%.2f%%)"),
+			*FormatNumber(Counters.TotalCacheSearchAttempts),
+			*FormatNumber(Counters.TotalCacheHits),
+			(Counters.TotalCacheSearchAttempts > 0) ? 100.0 * static_cast<double>(Counters.TotalCacheHits) / static_cast<double>(Counters.TotalCacheSearchAttempts) : 0.0);
+
+		UE_LOG(LogShaderCompilers, Display, TEXT("Tracking %s distinct input hashes that result in %s distinct outputs (%.2f%%)"),
+			*FormatNumber(Counters.UniqueCacheInputHashes),
+			*FormatNumber(Counters.UniqueCacheOutputs),
+			(Counters.UniqueCacheInputHashes > 0) ? 100.0 * static_cast<double>(Counters.UniqueCacheOutputs) / static_cast<double>(Counters.UniqueCacheInputHashes) : 0.0);
+
+		static const FNumberFormattingOptions SizeFormattingOptions = FNumberFormattingOptions().SetMinimumFractionalDigits(2).SetMaximumFractionalDigits(2);
+
+		if (Counters.CacheMemBudget > 0)
+		{
+			UE_LOG(LogShaderCompilers, Display, TEXT("RAM used: %s of %s budget. Usage: %.2f%%"),
+				*FText::AsMemory(Counters.CacheMemUsed, &SizeFormattingOptions, nullptr, EMemoryUnitStandard::IEC).ToString(),
+				*FText::AsMemory(Counters.CacheMemBudget, &SizeFormattingOptions, nullptr, EMemoryUnitStandard::IEC).ToString(),
+				100.0 * Counters.CacheMemUsed / Counters.CacheMemBudget);
+		}
+		else
+		{
+			UE_LOG(LogShaderCompilers, Display, TEXT("RAM used: %s, no memory limit set"), *FText::AsMemory(Counters.CacheMemUsed, &SizeFormattingOptions, nullptr, EMemoryUnitStandard::IEC).ToString());
+		}
+	}
+
 	const double TotalTimeAtLeastOneJobWasInFlight = GetTimeShaderCompilationWasActive();
 
-	UE_LOG(LogShaderCompilers, Display, TEXT("=== Shader Compilation stats ==="));
+	UE_LOG(LogShaderCompilers, Display, TEXT("=== Shader Compilation stats%s ==="), AggregatedSuffix);
 	UE_LOG(LogShaderCompilers, Display, TEXT("Shaders Compiled: %s"), *FormatNumber(TotalCompiled));
 
 	FScopeLock Lock(&CompileStatsLock);	// make a local copy for all the stats?
 	UE_LOG(LogShaderCompilers, Display, TEXT("Jobs assigned %s, completed %s (%s)"), 
-		*FormatNumber(JobsAssigned),
-		*FormatNumber(JobsCompleted),
-		*PrintJobsCompletedPercentageToString(JobsAssigned, JobsCompleted));
+		*FormatNumber(Counters.JobsAssigned),
+		*FormatNumber(Counters.JobsCompleted),
+		*PrintJobsCompletedPercentageToString(Counters.JobsAssigned, Counters.JobsCompleted));
 
-	if (TimesLocalWorkersWereIdle > 0.0)
+	if (Counters.TimesLocalWorkersWereIdle > 0.0)
 	{
-		UE_LOG(LogShaderCompilers, Display, TEXT("Average time worker was idle: %.2f s"), AccumulatedLocalWorkerIdleTime / TimesLocalWorkersWereIdle);
+		UE_LOG(LogShaderCompilers, Display, TEXT("Average time worker was idle: %.2f s"), Counters.AccumulatedLocalWorkerIdleTime / Counters.TimesLocalWorkersWereIdle);
 	}
 
-	if (JobsAssigned > 0)
+	if (Counters.JobsAssigned > 0)
 	{
-		UE_LOG(LogShaderCompilers, Display, TEXT("Time job spent in pending queue: average %.2f s, longest %.2f s"), AccumulatedPendingTime / (double)JobsAssigned, MaxPendingTime);
+		UE_LOG(LogShaderCompilers, Display, TEXT("Time job spent in pending queue: average %.2f s, longest %.2f s"), Counters.AccumulatedPendingTime / (double)Counters.JobsAssigned, Counters.MaxPendingTime);
 	}
 
-	if (JobsCompleted > 0)
+	if (Counters.JobsCompleted > 0)
 	{
-		UE_LOG(LogShaderCompilers, Display, TEXT("Job execution time: average %.2f s, max %.2f s"), AccumulatedJobExecutionTime / (double)JobsCompleted, MaxJobExecutionTime);
-		UE_LOG(LogShaderCompilers, Display, TEXT("Job life time (pending + execution): average %.2f s, max %.2f"), AccumulatedJobLifeTime / (double)JobsCompleted, MaxJobLifeTime);
+		UE_LOG(LogShaderCompilers, Display, TEXT("Job execution time: average %.2f s, max %.2f s"), Counters.AccumulatedJobExecutionTime / (double)Counters.JobsCompleted, Counters.MaxJobExecutionTime);
+		UE_LOG(LogShaderCompilers, Display, TEXT("Job life time (pending + execution): average %.2f s, max %.2f"), Counters.AccumulatedJobLifeTime / (double)Counters.JobsCompleted, Counters.MaxJobLifeTime);
 	}
 
-	if (NumAccumulatedShaderCodes > 0)
+	if (Counters.NumAccumulatedShaderCodes > 0)
 	{
-		const FString AvgCodeSizeStr = FText::AsMemory((uint64)((double)AccumulatedShaderCodeSize / (double)NumAccumulatedShaderCodes)).ToString();
-		const FString MinCodeSizeStr = FText::AsMemory((uint64)MinShaderCodeSize).ToString();
-		const FString MaxCodeSizeStr = FText::AsMemory((uint64)MaxShaderCodeSize).ToString();
+		const FString AvgCodeSizeStr = FText::AsMemory((uint64)((double)Counters.AccumulatedShaderCodeSize / (double)Counters.NumAccumulatedShaderCodes)).ToString();
+		const FString MinCodeSizeStr = FText::AsMemory((uint64)Counters.MinShaderCodeSize).ToString();
+		const FString MaxCodeSizeStr = FText::AsMemory((uint64)Counters.MaxShaderCodeSize).ToString();
 		UE_LOG(LogShaderCompilers, Display, TEXT("Shader code size: average %s, min %s, max %s"), *AvgCodeSizeStr, *MinCodeSizeStr, *MaxCodeSizeStr);
 	}
 
 	UE_LOG(LogShaderCompilers, Display, TEXT("Time at least one job was in flight (either pending or executed): %.2f s"), TotalTimeAtLeastOneJobWasInFlight);
 
 	// print stats about the batches
-	if (LocalJobBatchesSeen > 0 && DistributedJobBatchesSeen > 0)
+	if (Counters.LocalJobBatchesSeen > 0 && Counters.DistributedJobBatchesSeen > 0)
 	{
-		int64 JobBatchesSeen = LocalJobBatchesSeen + DistributedJobBatchesSeen;
-		double TotalJobsReportedInJobBatches = TotalJobsReportedInLocalJobBatches + TotalJobsReportedInDistributedJobBatches;
+		int64 JobBatchesSeen = Counters.LocalJobBatchesSeen + Counters.DistributedJobBatchesSeen;
+		double TotalJobsReportedInJobBatches = Counters.TotalJobsReportedInLocalJobBatches + Counters.TotalJobsReportedInDistributedJobBatches;
 
 		UE_LOG(LogShaderCompilers, Display, TEXT("Jobs were issued in %s batches (%s local, %s distributed), average %.2f jobs/batch (%.2f jobs/local batch. %.2f jobs/distributed batch)"),
-			*FormatNumber(JobBatchesSeen), *FormatNumber(LocalJobBatchesSeen), *FormatNumber(DistributedJobBatchesSeen),
+			*FormatNumber(JobBatchesSeen), *FormatNumber(Counters.LocalJobBatchesSeen), *FormatNumber(Counters.DistributedJobBatchesSeen),
 			static_cast<double>(TotalJobsReportedInJobBatches) / static_cast<double>(JobBatchesSeen),
-			static_cast<double>(TotalJobsReportedInLocalJobBatches) / static_cast<double>(LocalJobBatchesSeen),
-			static_cast<double>(TotalJobsReportedInDistributedJobBatches) / static_cast<double>(DistributedJobBatchesSeen)
+			static_cast<double>(Counters.TotalJobsReportedInLocalJobBatches) / static_cast<double>(Counters.LocalJobBatchesSeen),
+			static_cast<double>(Counters.TotalJobsReportedInDistributedJobBatches) / static_cast<double>(Counters.DistributedJobBatchesSeen)
 		);
 	}
-	else if (LocalJobBatchesSeen > 0)
+	else if (Counters.LocalJobBatchesSeen > 0)
 	{
 		UE_LOG(LogShaderCompilers, Display, TEXT("Jobs were issued in %s batches (only local compilation was used), average %.2f jobs/batch"), 
-			*FormatNumber(LocalJobBatchesSeen), static_cast<double>(TotalJobsReportedInLocalJobBatches) / static_cast<double>(LocalJobBatchesSeen));
+			*FormatNumber(Counters.LocalJobBatchesSeen), static_cast<double>(Counters.TotalJobsReportedInLocalJobBatches) / static_cast<double>(Counters.LocalJobBatchesSeen));
 	}
-	else if (DistributedJobBatchesSeen > 0)
+	else if (Counters.DistributedJobBatchesSeen > 0)
 	{
 		UE_LOG(LogShaderCompilers, Display, TEXT("Jobs were issued in %s batches (only distributed compilation was used), average %.2f jobs/batch"),
-			*FormatNumber(DistributedJobBatchesSeen), static_cast<double>(TotalJobsReportedInDistributedJobBatches) / static_cast<double>(DistributedJobBatchesSeen));
+			*FormatNumber(Counters.DistributedJobBatchesSeen), static_cast<double>(Counters.TotalJobsReportedInDistributedJobBatches) / static_cast<double>(Counters.DistributedJobBatchesSeen));
 	}
 
 	if (TotalTimeAtLeastOneJobWasInFlight > 0.0)
 	{
-		UE_LOG(LogShaderCompilers, Display, TEXT("Average processing rate: %.2f jobs/sec"), (double)JobsCompleted / TotalTimeAtLeastOneJobWasInFlight);
+		UE_LOG(LogShaderCompilers, Display, TEXT("Average processing rate: %.2f jobs/sec"), (double)Counters.JobsCompleted / TotalTimeAtLeastOneJobWasInFlight);
 	}
 
 	if (ShaderTimings.Num())
@@ -4079,7 +4085,7 @@ void FShaderCompilerStats::WriteStatSummary()
 		if (TotalTimeAtLeastOneJobWasInFlight > 0.0)
 		{
 			double EffectiveParallelization = TotalThreadTimeForAllShaders / TotalTimeAtLeastOneJobWasInFlight;
-			if (DistributedJobBatchesSeen == 0)
+			if (Counters.DistributedJobBatchesSeen == 0)
 			{
 				UE_LOG(LogShaderCompilers, Display, TEXT("Effective parallelization: %.2f (times faster than compiling all shaders on one thread). Compare with number of workers: %d"), EffectiveParallelization, GShaderCompilingManager->GetNumLocalWorkers());
 			}
@@ -4126,6 +4132,8 @@ void FShaderCompilerStats::WriteStatSummary()
 			}
 		}
 	}
+
+	UE_LOG(LogShaderCompilers, Display, TEXT("================================================"));
 }
 
 void FShaderCompilerStats::GatherAnalytics(const FString& BaseName, TArray<FAnalyticsEventAttribute>& Attributes)
@@ -4136,7 +4144,7 @@ void FShaderCompilerStats::GatherAnalytics(const FString& BaseName, TArray<FAnal
 
 	{
 		FString AttrName = BaseName + TEXT("ShadersCompiled");
-		Attributes.Emplace(MoveTemp(AttrName), JobsCompleted);
+		Attributes.Emplace(MoveTemp(AttrName), Counters.JobsCompleted);
 	}
 
 	if (ShaderTimings.Num())
@@ -4165,19 +4173,324 @@ void FShaderCompilerStats::GatherAnalytics(const FString& BaseName, TArray<FAnal
 			Attributes.Emplace(MoveTemp(AttrName), EffectiveParallelization);
 		}
 	}
+
+	if (Counters.TotalCacheSearchAttempts)
+	{
+		const FString ChildName = TEXT("JobCache_");
+
+		{
+			FString AttrName = BaseName + ChildName + TEXT("Queries");
+			Attributes.Emplace(MoveTemp(AttrName), Counters.TotalCacheSearchAttempts);
+		}
+
+		{
+			FString AttrName = BaseName + ChildName + TEXT("Hits");
+			Attributes.Emplace(MoveTemp(AttrName), Counters.TotalCacheHits);
+		}
+
+		{
+			FString AttrName = BaseName + ChildName + TEXT("NumInputs");
+			Attributes.Emplace(MoveTemp(AttrName), Counters.UniqueCacheInputHashes);
+		}
+
+		{
+			FString AttrName = BaseName + ChildName + TEXT("NumOutputs");
+			Attributes.Emplace(MoveTemp(AttrName), Counters.UniqueCacheOutputs);
+		}
+
+		{
+			FString AttrName = BaseName + ChildName + TEXT("MemUsed");
+			Attributes.Emplace(MoveTemp(AttrName), Counters.CacheMemUsed);
+		}
+
+		{
+			FString AttrName = BaseName + ChildName + TEXT("MemBudget");
+			Attributes.Emplace(MoveTemp(AttrName), Counters.CacheMemBudget);
+		}
+	}
 }
 
 uint32 FShaderCompilerStats::GetTotalShadersCompiled()
 {
 	FScopeLock Lock(&CompileStatsLock);
-	return (uint32)FMath::Max(0ll, JobsCompleted);
+	return (uint32)FMath::Max(0ll, Counters.JobsCompleted);
+}
+
+void AddToInterval(TArray<TInterval<double>>& Accumulator, const TInterval<double>& NewInterval)
+{
+	bool bFoundOverlap = false;
+	TInterval<double> New = NewInterval;
+	int32 Idx = 0;
+	do
+	{
+		bFoundOverlap = false;
+		for (; Idx < Accumulator.Num(); ++Idx)
+		{
+			const TInterval<double>& Existing = Accumulator[Idx];
+			if (Existing.Max < New.Min)
+			{
+				continue;	// no overlap but the new interval starts after this one ends, keep searching
+			}
+
+			if (New.Max < Existing.Min)
+			{
+				break;		// no overlap, but the new interval ends before this one starts, insert here
+			}
+
+			// if fully contained within existing interval, just ignore
+			if (Existing.Min <= New.Min && New.Max <= Existing.Max)
+			{
+				return;
+			}
+
+			bFoundOverlap = true;
+			// if there's an overlap, remove the existing interval, merge with the new one and attempt to add again
+			TInterval<double> Merged(FMath::Min(Existing.Min, New.Min), FMath::Max(Existing.Max, New.Max));
+			check(Merged.Size() >= Existing.Size());
+			check(Merged.Size() >= New.Size());
+			Accumulator.RemoveAt(Idx);
+			New = Merged;
+			break;
+		}
+	} while (bFoundOverlap);
+
+	// if we arrived here without an overlap, we have a new one; insert in the appropriate place
+	if (!bFoundOverlap)
+	{
+		Accumulator.Insert(New, Idx);
+	}
+}
+
+void FShaderCompilerStats::Aggregate(FShaderCompilerStats& Other)
+{
+	// note: intentionally not taking local lock as this should only ever be called on a local copy of the stats object
+	FScopeLock Lock(&Other.CompileStatsLock);
+	Counters += Other.Counters;
+
+	for (TSparseArray<ShaderCompilerStats>::TConstIterator It(Other.CompileStats); It; ++It)
+	{
+		if (!CompileStats.IsValidIndex(It.GetIndex()))
+		{
+			CompileStats.EmplaceAt(It.GetIndex());
+		}
+
+		ShaderCompilerStats& Stats = CompileStats[It.GetIndex()];
+		for (const TPair<FString, FShaderStats>& StatsKeyValue : *It)
+		{
+			FShaderStats* Current = Stats.Find(StatsKeyValue.Key);
+			if (Current)
+			{
+				*Current += StatsKeyValue.Value;
+			}
+			else
+			{
+				Stats.Add(StatsKeyValue);
+			}
+		}
+	}
+
+	// note: this is suboptimal (O(n^2)) but there aren't a lot of these in practice
+	for (const TInterval<double>& Interval : Other.JobLifeTimeIntervals)
+	{
+		AddToInterval(JobLifeTimeIntervals, Interval);
+	}
+
+	for (const TPair<FString, FShaderTimings>& TimingsKeyValue : Other.ShaderTimings)
+	{
+		FShaderTimings* Current = ShaderTimings.Find(TimingsKeyValue.Key);
+		if (Current)
+		{
+			*Current += TimingsKeyValue.Value;
+		}
+		else
+		{
+			ShaderTimings.Add(TimingsKeyValue);
+		}
+	}
+}
+
+void FShaderCompilerStats::WriteToCompactBinary(FCbWriter& Writer)
+{
+	FScopeLock Lock(&CompileStatsLock);
+	Writer.AddBinary("Counters", &Counters, sizeof(Counters));
+
+	Writer.BeginArray("CompileStatIndices");	
+	// Write the array of valid indices this worker has in the compile stats sparse array
+	for (TSparseArray<ShaderCompilerStats>::TConstIterator It(CompileStats); It; ++It)
+	{
+		if (CompileStats.IsValidIndex(It.GetIndex()))
+		{
+			Writer << It.GetIndex();
+		}
+	}
+	Writer.EndArray();
+
+	Writer.BeginArray("CompileStats");
+		// Then write the actual compile stats maps in the same order as the above indices
+	for (TSparseArray<ShaderCompilerStats>::TConstIterator It(CompileStats); It; ++It)
+	{
+		if (!CompileStats.IsValidIndex(It.GetIndex()))
+		{
+			continue;
+		}
+
+		Writer.BeginObject();
+		Writer.BeginArray("CompileStatsKeys");
+		for (TPair<FString, FShaderStats> Pair : *It)
+		{
+			Writer << Pair.Key;
+		}
+		Writer.EndArray();
+
+		Writer.BeginArray("CompileStatsValues");
+		for (const TPair<FString, FShaderStats>& Pair : *It)
+		{
+			Writer.BeginObject();
+			Writer << "Compiled" << Pair.Value.Compiled;
+			Writer << "CompiledDouble" << Pair.Value.CompiledDouble;
+			Writer << "CompileTime" << Pair.Value.CompileTime;
+			Writer << "Cooked" << Pair.Value.Cooked;
+			Writer << "CookedDouble" << Pair.Value.CookedDouble;
+			Writer.BeginArray("PermutationCompilations");
+			for (const FShaderCompilerSinglePermutationStat& Stat : Pair.Value.PermutationCompilations)
+			{
+				Writer.BeginObject();
+				Writer << "Compiled" << Stat.Compiled;
+				Writer << "CompiledDouble" << Stat.CompiledDouble;
+				Writer << "Cooked" << Stat.Cooked;
+				Writer << "CookedDouble" << Stat.CookedDouble;
+				Writer << "PermutationString" << Stat.PermutationString;
+				Writer.EndObject();
+			}
+			Writer.EndArray();
+			Writer.EndObject();
+		}
+		Writer.EndArray();
+		Writer.EndObject();
+	}
+	Writer.EndArray();
+
+	Writer.BeginArray("JobLifeTimeIntervals");
+	for (const TInterval<double>& Interval : JobLifeTimeIntervals)
+	{
+		Writer.AddBinary(&Interval, sizeof(TInterval<double>));
+	}
+	Writer.EndArray();
+
+	Writer.BeginArray("ShaderTimingsKeys");
+	for (const TPair<FString, FShaderTimings>& TimingPair : ShaderTimings)
+	{
+		Writer << TimingPair.Key;
+	}
+	Writer.EndArray();
+
+	Writer.BeginArray("ShaderTimingsValues");
+	for (const TPair<FString, FShaderTimings>& TimingPair : ShaderTimings)
+	{
+		Writer.AddBinary(&TimingPair.Value, sizeof(FShaderTimings));
+	}
+	Writer.EndArray();
+}
+
+void FShaderCompilerStats::ReadFromCompactBinary(FCbObjectView& Reader)
+{
+	FScopeLock Lock(&CompileStatsLock);
+	FMemoryView CountersMem = Reader["Counters"].AsBinaryView();
+	check(CountersMem.GetSize() == sizeof(FCounters));
+	Counters = *reinterpret_cast<const FCounters*>(CountersMem.GetData());
+
+	FCbArrayView CompileStatIndicesView = Reader["CompileStatIndices"].AsArrayView();
+	FCbArrayView CompileStatsView = Reader["CompileStats"].AsArrayView();
+	check(CompileStatIndicesView.Num() == CompileStatsView.Num());
+
+	FCbFieldViewIterator IndexIt = CompileStatIndicesView.CreateViewIterator();
+	FCbFieldViewIterator StatsIt = CompileStatsView.CreateViewIterator();
+
+	while (IndexIt && StatsIt)
+	{
+		if (!CompileStats.IsValidIndex(IndexIt.AsUInt32()))
+		{
+			FSparseArrayAllocationInfo AllocInfo = CompileStats.InsertUninitialized(IndexIt.AsUInt32());
+			new(AllocInfo) ShaderCompilerStats();
+		}
+		ShaderCompilerStats& Stats = CompileStats[IndexIt.AsUInt32()];
+
+		FCbObjectView PlatformStatsObject = StatsIt->AsObjectView();
+		FCbArrayView StatsKeysView = PlatformStatsObject["CompileStatsKeys"].AsArrayView();
+		FCbArrayView StatsValuesView = PlatformStatsObject["CompileStatsValues"].AsArrayView();
+		check(StatsKeysView.Num() == StatsValuesView.Num());
+
+		Stats.Reserve(StatsKeysView.Num());
+
+		FCbFieldViewIterator KeysIt = StatsKeysView.CreateViewIterator();
+		FCbFieldViewIterator ValuesIt = StatsValuesView.CreateViewIterator();
+
+		while (KeysIt && ValuesIt)
+		{
+			FCbObjectView ShaderStatsObject = ValuesIt->AsObjectView();
+			FShaderStats& ShaderStats = Stats.Add(FString(KeysIt->AsString()));
+			ShaderStats.Compiled = ShaderStatsObject["Compiled"].AsUInt32();
+			ShaderStats.CompiledDouble = ShaderStatsObject["CompiledDouble"].AsUInt32();
+			ShaderStats.CompileTime = ShaderStatsObject["CompileTime"].AsFloat();
+			ShaderStats.Cooked = ShaderStatsObject["Cooked"].AsUInt32();
+			ShaderStats.CookedDouble = ShaderStatsObject["CookedDouble"].AsUInt32();
+
+			FCbArrayView PermutationsArrayView = ShaderStatsObject["PermutationCompilations"].AsArrayView();
+			ShaderStats.PermutationCompilations.Reset(PermutationsArrayView.Num());
+			for (FCbFieldView CompilationField : PermutationsArrayView)
+			{
+				FCbObjectView PermutationObject = CompilationField.AsObjectView();
+				uint32 Index = ShaderStats.PermutationCompilations.Emplace
+				(
+					FString(PermutationObject["PermutationString"].AsString()),
+					PermutationObject["Compiled"].AsUInt32(),
+					PermutationObject["Cooked"].AsUInt32()
+				);
+				ShaderStats.PermutationCompilations[Index].CompiledDouble = PermutationObject["CompiledDouble"].AsUInt32();
+				ShaderStats.PermutationCompilations[Index].CookedDouble = PermutationObject["CookedDouble"].AsUInt32();
+			}
+
+			++ValuesIt;
+			++KeysIt;
+		}
+
+		++IndexIt;
+		++StatsIt;
+	}
+
+	FCbArrayView JobLifeTimeIntervalsView = Reader["JobLifeTimeIntervals"].AsArrayView();
+	JobLifeTimeIntervals.Reset(JobLifeTimeIntervalsView.Num());
+	for (FCbFieldView JobLifeTimeField : JobLifeTimeIntervalsView)
+	{
+		FMemoryView IntervalObj = JobLifeTimeField.AsBinaryView();
+		check(IntervalObj.GetSize() == sizeof(TInterval<double>));
+		JobLifeTimeIntervals.Add(*reinterpret_cast<const TInterval<double>*>(IntervalObj.GetData()));
+	}
+
+	FCbArrayView TimingsKeysView = Reader["ShaderTimingsKeys"].AsArrayView();
+	FCbArrayView TimingsValuesView = Reader["ShaderTimingsValues"].AsArrayView();
+	check(TimingsKeysView.Num() == TimingsValuesView.Num());
+
+	ShaderTimings.Reserve(TimingsKeysView.Num());
+
+	FCbFieldViewIterator TimingsKeysIt = TimingsKeysView.CreateViewIterator();
+	FCbFieldViewIterator TimingsValuesIt = TimingsValuesView.CreateViewIterator();
+
+	while (TimingsKeysIt && TimingsValuesIt)
+	{
+		FMemoryView TimingsValuesBinary = TimingsValuesIt->AsBinaryView();
+		check(TimingsValuesBinary.GetSize() == sizeof(FShaderTimings));
+		ShaderTimings.Add(FString(TimingsKeysIt->AsString()), *reinterpret_cast<const FShaderTimings*>(TimingsValuesBinary.GetData()));
+		++TimingsKeysIt;
+		++TimingsValuesIt;
+	}
 }
 
 void FShaderCompilerStats::RegisterLocalWorkerIdleTime(double IdleTime)
 {
 	FScopeLock Lock(&CompileStatsLock);
-	AccumulatedLocalWorkerIdleTime += IdleTime;
-	TimesLocalWorkersWereIdle++;
+	Counters.AccumulatedLocalWorkerIdleTime += IdleTime;
+	Counters.TimesLocalWorkersWereIdle++;
 }
 
 void FShaderCompilerStats::RegisterNewPendingJob(FShaderCommonCompileJob& Job)
@@ -4195,10 +4508,10 @@ void FShaderCompilerStats::RegisterAssignedJob(FShaderCommonCompileJob& Job)
 	Job.TimeAssignedToExecution = FPlatformTime::Seconds();
 
 	FScopeLock Lock(&CompileStatsLock);
-	JobsAssigned++;
+	Counters.JobsAssigned++;
 	double TimeSpendPending = (Job.TimeAssignedToExecution - Job.TimeAddedToPendingQueue);
-	AccumulatedPendingTime += TimeSpendPending;
-	MaxPendingTime = FMath::Max(TimeSpendPending, MaxPendingTime);
+	Counters.AccumulatedPendingTime += TimeSpendPending;
+	Counters.MaxPendingTime = FMath::Max(TimeSpendPending, Counters.MaxPendingTime);
 }
 
 void FShaderCompilerStats::RegisterFinishedJob(FShaderCommonCompileJob& Job)
@@ -4207,54 +4520,15 @@ void FShaderCompilerStats::RegisterFinishedJob(FShaderCommonCompileJob& Job)
 	Job.TimeExecutionCompleted = FPlatformTime::Seconds();
 
 	FScopeLock Lock(&CompileStatsLock);
-	JobsCompleted++;
+	Counters.JobsCompleted++;
 	
 	double ExecutionTime = (Job.TimeExecutionCompleted - Job.TimeAssignedToExecution);
-	AccumulatedJobExecutionTime += ExecutionTime;
-	MaxJobExecutionTime = FMath::Max(ExecutionTime, MaxJobExecutionTime);
+	Counters.AccumulatedJobExecutionTime += ExecutionTime;
+	Counters.MaxJobExecutionTime = FMath::Max(ExecutionTime, Counters.MaxJobExecutionTime);
 
 	double LifeTime = (Job.TimeExecutionCompleted - Job.TimeAddedToPendingQueue);
-	AccumulatedJobLifeTime += LifeTime;
-	MaxJobLifeTime = FMath::Max(LifeTime, MaxJobLifeTime);
-
-	auto AddToInterval = [](TArray<TInterval<double>>& Accumulator, const TInterval<double>& NewInterval)
-	{
-		bool bFoundOverlap = false;
-		TInterval<double> New = NewInterval;
-		do
-		{
-			bFoundOverlap = false;
-			for (int32 Idx = 0; Idx < Accumulator.Num(); ++Idx)
-			{
-				const TInterval<double>& Existing = Accumulator[Idx];
-				if (Existing.Max < New.Min || New.Max < Existing.Min)
-				{
-					continue;	// no overlap
-				}
-
-				// if fully contained within existing interval, just ignore
-				if (Existing.Min <= New.Min && New.Max <= Existing.Max)
-				{
-					return;
-				}
-
-				bFoundOverlap = true;
-				// if there's an overlap, remove the existing interval, merge with the new one and attempt to add again
-				TInterval<double> Merged(FMath::Min(Existing.Min, New.Min), FMath::Max(Existing.Max, New.Max));
-				check(Merged.Size() >= Existing.Size());
-				check(Merged.Size() >= New.Size());
-				Accumulator.RemoveAt(Idx);
-				New = Merged;
-				break;
-			}
-		} while (bFoundOverlap);
-
-		// if we arrived here without an overlap, we have a new one
-		if (!bFoundOverlap)
-		{
-			Accumulator.Add(New);
-		}
-	};
+	Counters.AccumulatedJobLifeTime += LifeTime;
+	Counters.MaxJobLifeTime = FMath::Max(LifeTime, Counters.MaxJobLifeTime);
 
 	// estimate lifetime without an overlap
 	ensure(Job.TimeAddedToPendingQueue != 0.0 && Job.TimeAddedToPendingQueue <= Job.TimeExecutionCompleted);
@@ -4266,10 +4540,10 @@ void FShaderCompilerStats::RegisterFinishedJob(FShaderCommonCompileJob& Job)
 		const int32 ShaderCodeSize = SingleJob.Output.ShaderCode.GetShaderCodeSize();
 		if (ShaderCodeSize > 0)
 		{
-			MinShaderCodeSize = (MinShaderCodeSize > 0 ? FMath::Min(MinShaderCodeSize, ShaderCodeSize) : ShaderCodeSize);
-			MaxShaderCodeSize = (MaxShaderCodeSize > 0 ? FMath::Max(MaxShaderCodeSize, ShaderCodeSize) : ShaderCodeSize);
-			AccumulatedShaderCodeSize += (uint64)ShaderCodeSize;
-			++NumAccumulatedShaderCodes;
+			Counters.MinShaderCodeSize = (Counters.MinShaderCodeSize > 0 ? FMath::Min(Counters.MinShaderCodeSize, ShaderCodeSize) : ShaderCodeSize);
+			Counters.MaxShaderCodeSize = (Counters.MaxShaderCodeSize > 0 ? FMath::Max(Counters.MaxShaderCodeSize, ShaderCodeSize) : ShaderCodeSize);
+			Counters.AccumulatedShaderCodeSize += (uint64)ShaderCodeSize;
+			++Counters.NumAccumulatedShaderCodes;
 		}
 
 		const FString ShaderName(SingleJob.Key.ShaderType->GetName());
@@ -4304,14 +4578,14 @@ void FShaderCompilerStats::RegisterJobBatch(int32 NumJobs, EExecutionType ExecTy
 	if (ExecType == EExecutionType::Local)
 	{
 		FScopeLock Lock(&CompileStatsLock);
-		++LocalJobBatchesSeen;
-		TotalJobsReportedInLocalJobBatches += NumJobs;
+		++Counters.LocalJobBatchesSeen;
+		Counters.TotalJobsReportedInLocalJobBatches += NumJobs;
 	}
 	else if (ExecType == EExecutionType::Distributed)
 	{
 		FScopeLock Lock(&CompileStatsLock);
-		++DistributedJobBatchesSeen;
-		TotalJobsReportedInDistributedJobBatches += NumJobs;
+		++Counters.DistributedJobBatchesSeen;
+		Counters.TotalJobsReportedInDistributedJobBatches += NumJobs;
 	}
 	else
 	{
@@ -4400,22 +4674,22 @@ void FShaderCompilerStats::RegisterCompiledShaders(uint32 NumCompiled, EShaderPl
 
 void FShaderCompilerStats::AddDDCMiss(uint32 NumMisses)
 {
-	ShaderMapDDCMisses += NumMisses;
+	Counters.ShaderMapDDCMisses += NumMisses;
 }
 
 uint32 FShaderCompilerStats::GetDDCMisses() const
 {
-	return ShaderMapDDCMisses;
+	return Counters.ShaderMapDDCMisses;
 }
 
 void FShaderCompilerStats::AddDDCHit(uint32 NumHits)
 {
-	ShaderMapDDCHits += NumHits;
+	Counters.ShaderMapDDCHits += NumHits;
 }
 
 uint32 FShaderCompilerStats::GetDDCHits() const
 {
-	return ShaderMapDDCHits;
+	return Counters.ShaderMapDDCHits;
 }
 
 double FShaderCompilerStats::GetTimeShaderCompilationWasActive()
@@ -4667,8 +4941,6 @@ FShaderCompilingManager::~FShaderCompilingManager()
 		return;
 	}
 
-	PrintStats(true);
-
 	for (const auto& Thread : Threads)
 	{
 		Thread->Stop();
@@ -4861,18 +5133,6 @@ int32 FShaderCompilingManager::GetNumOutstandingJobs() const
 	return AllJobs.GetNumOutstandingJobs();
 }
 
-void FShaderCompilingManager::GatherAnalytics(TArray<FAnalyticsEventAttribute>& Attributes) const
-{
-	const FString BaseName = TEXT("Shaders_");
-
-	if (ShaderCompiler::IsJobCacheEnabled())
-	{
-		AllJobs.GatherAnalytics(BaseName, Attributes);
-	}
-
-	GShaderCompilerStats->GatherAnalytics(BaseName, Attributes);
-}
-
 FShaderCompilingManager::EDumpShaderDebugInfo FShaderCompilingManager::GetDumpShaderDebugInfo() const
 {
 	if (GDumpShaderDebugInfo < EDumpShaderDebugInfo::Never || GDumpShaderDebugInfo > EDumpShaderDebugInfo::OnErrorOrWarning)
@@ -5752,8 +6012,6 @@ void FShaderCompilingManager::ProcessCompiledShaderMaps(
 		}
 	}
 
-	GShaderCompilingManager->PrintStats();
-
 	UpdateNumRemainingAssets();
 #endif // WITH_EDITOR
 }
@@ -5804,33 +6062,20 @@ void FShaderCompilingManager::Shutdown()
 	// be longer than other asset compilers, otherwise niagara compilations might get stuck.
 }
 
-void FShaderCompilingManager::PrintStats(bool bForceLogIgnoringTimeInverval)
+void FShaderCompilingManager::PrintStats()
 {
-	static double LastTimeStatsPrinted = FPlatformTime::Seconds();
-	// do not print if
-	//  - job cache is disabled
-	//  - not enough time passed since the previous time and we're not forced to print
-	if ((!bForceLogIgnoringTimeInverval && GShaderCompilerCacheStatsPrintoutInterval > 0 && FPlatformTime::Seconds() - LastTimeStatsPrinted < GShaderCompilerCacheStatsPrintoutInterval))
-	{
-		return;
-	}
-	if (!AllowShaderCompiling())
-	{
-		return;
-	}
-
-	UE_LOG(LogShaderCompilers, Display, TEXT("================================================"));
+	FShaderCompilerStats LocalStats;
+	GetLocalStats(LocalStats);
+	LocalStats.WriteStatSummary();
+}
 
-	if (ShaderCompiler::IsJobCacheEnabled())
+void FShaderCompilingManager::GetLocalStats(FShaderCompilerStats& OutStats) const
+{
+	if (GShaderCompilerStats)
 	{
-		AllJobs.LogCachingStats();
+		OutStats.Aggregate(*GShaderCompilerStats);
+		AllJobs.GetCachingStats(OutStats);
 	}
-
-	GShaderCompilerStats->WriteStatSummary();
-
-	UE_LOG(LogShaderCompilers, Display, TEXT("================================================"));
-
-	LastTimeStatsPrinted = FPlatformTime::Seconds();
 }
 
 bool FShaderCompilingManager::HandlePotentialRetryOnError(TMap<int32, FShaderMapFinalizeResults>& CompletedShaderMaps)
@@ -9711,71 +9956,15 @@ uint64 FShaderJobCache::ComputeAllocatedMemory() const
 	return AllocatedSize;
 }
 
-/** Logs out the statistics */
-void FShaderJobCache::LogStats()
+void FShaderJobCache::GetStats(FShaderCompilerStats& OutStats) const
 {
-	UE_LOG(LogShaderCompilers, Display, TEXT("=== FShaderJobCache stats ==="), this);
-	UE_LOG(LogShaderCompilers, Display, TEXT("Total job queries %s, among them cache hits %s (%.2f%%)"),
-		*FormatNumber(TotalSearchAttempts),
-		*FormatNumber(TotalCacheHits),
-		(TotalSearchAttempts > 0) ? 100.0 * static_cast<double>(TotalCacheHits) / static_cast<double>(TotalSearchAttempts) : 0.0);
-
-	UE_LOG(LogShaderCompilers, Display, TEXT("Tracking %s distinct input hashes that result in %s distinct outputs (%.2f%%)"),
-		*FormatNumber(InputHashToJobData.Num()), 
-		*FormatNumber(Outputs.Num()),
-		(InputHashToJobData.Num() > 0) ? 100.0 * static_cast<double>(Outputs.Num()) / static_cast<double>(InputHashToJobData.Num()) : 0.0);
-
-	static const FNumberFormattingOptions SizeFormattingOptions = FNumberFormattingOptions().SetMinimumFractionalDigits(2).SetMaximumFractionalDigits(2);
-
-	const uint64 MemUsed = GetAllocatedMemory();
-	const uint64 MemBudget = GetCurrentMemoryBudget();
-	if (MemBudget > 0)
-	{
-		UE_LOG(LogShaderCompilers, Display, TEXT("RAM used: %s of %s budget. Usage: %.2f%%"), 
-			*FText::AsMemory(MemUsed, &SizeFormattingOptions, nullptr, EMemoryUnitStandard::IEC).ToString(),
-			*FText::AsMemory(MemBudget, &SizeFormattingOptions, nullptr, EMemoryUnitStandard::IEC).ToString(),
-			100.0 * MemUsed / MemBudget);
-	}
-	else
-	{
-		UE_LOG(LogShaderCompilers, Display, TEXT("RAM used: %s, no memory limit set"), *FText::AsMemory(MemUsed, &SizeFormattingOptions, nullptr, EMemoryUnitStandard::IEC).ToString());
-	}
-}
-
-void FShaderJobCache::GatherAnalytics(const FString& BaseName, TArray<FAnalyticsEventAttribute>& Attributes) const
-{
-	FWriteScopeLock Locker(JobLock);
-	const FString ChildName = TEXT("JobCache_");
-
-	{
-		FString AttrName = BaseName + ChildName + TEXT("Queries");
-		Attributes.Emplace(MoveTemp(AttrName), TotalSearchAttempts);
-	}
-
-	{
-		FString AttrName = BaseName + ChildName + TEXT("Hits");
-		Attributes.Emplace(MoveTemp(AttrName), TotalCacheHits);
-	}
-
-	{
-		FString AttrName = BaseName + ChildName + TEXT("NumInputs");
-		Attributes.Emplace(MoveTemp(AttrName), Outputs.Num());
-	}
-
-	{
-		FString AttrName = BaseName + ChildName + TEXT("NumOutputs");
-		Attributes.Emplace(MoveTemp(AttrName), Outputs.Num());
-	}
-
-	{
-		FString AttrName = BaseName + ChildName + TEXT("MemUsed");
-		Attributes.Emplace(MoveTemp(AttrName), GetAllocatedMemory());
-	}
-
-	{
-		FString AttrName = BaseName + ChildName + TEXT("MemBudget");
-		Attributes.Emplace(MoveTemp(AttrName), GetCurrentMemoryBudget());
-	}
+	FReadScopeLock Locker(JobLock);
+	OutStats.Counters.TotalCacheSearchAttempts = TotalSearchAttempts;
+	OutStats.Counters.TotalCacheHits = TotalCacheHits;
+	OutStats.Counters.UniqueCacheInputHashes = InputHashToJobData.Num();
+	OutStats.Counters.UniqueCacheOutputs = Outputs.Num();
+	OutStats.Counters.CacheMemUsed = GetAllocatedMemory();
+	OutStats.Counters.CacheMemBudget = GetCurrentMemoryBudget();
 }
 
 #undef LOCTEXT_NAMESPACE
diff --git a/Engine/Source/Runtime/Engine/Public/ShaderCompiler.h b/Engine/Source/Runtime/Engine/Public/ShaderCompiler.h
--- a/Engine/Source/Runtime/Engine/Public/ShaderCompiler.h
+++ b/Engine/Source/Runtime/Engine/Public/ShaderCompiler.h
@@ -30,6 +30,8 @@
 #include "Hash/Blake3.h"
 #include "SceneTypes.h"
 
+class FCbObjectView;
+class FCbWriter;
 class FVertexFactoryType;
 class IDistributedBuildController;
 class FMaterialShaderMap;
@@ -37,6 +39,7 @@ class FMaterialShaderMap;
 DECLARE_LOG_CATEGORY_EXTERN(LogShaderCompilers, Log, All);
 
 class FShaderCompileJob;
+class FShaderCompilerStats;
 class FShaderPipelineCompileJob;
 struct FAnalyticsEventAttribute;
 
@@ -66,10 +69,8 @@ public:
 	/** Adds the job to cache. */
 	void AddToCacheAndProcessPending(FShaderCommonCompileJob* FinishedJob);
 
-	/** Log caching statistics.*/
-	void LogCachingStats();
-
-	void GatherAnalytics(const FString& BaseName, TArray<FAnalyticsEventAttribute>& Attributes) const;
+	/** Retrieve caching statistics. */
+	void GetCachingStats(FShaderCompilerStats& OutStats) const;
 
 	int32 GetNumPendingJobs(EShaderCompileJobPriority InPriority) const;
 
@@ -303,6 +304,35 @@ public:
 		uint32 CompiledDouble = 0;
 		uint32 CookedDouble = 0;
 		float CompileTime = 0.f;
+
+		FShaderStats& operator+=(const FShaderStats& Other)
+		{
+			if (Compiled)
+			{
+				CompiledDouble += Other.Compiled;
+			}
+			else
+			{
+				Compiled += Other.Compiled;
+			}
+
+			if (Cooked)
+			{
+				CookedDouble += Other.Cooked;
+			}
+			else
+			{
+				Cooked += Other.Cooked;
+			}
+
+			CompiledDouble += Other.CompiledDouble;
+			CookedDouble += Other.CookedDouble;
+			CompileTime += Other.CompileTime;
+
+			PermutationCompilations.Append(Other.PermutationCompilations);
+
+			return *this;
+		}
 	};
 	using ShaderCompilerStats = TMap<FString, FShaderStats>;
 
@@ -315,6 +345,18 @@ public:
 		float TotalPreprocessTime = 0.0f;
 		int32 NumCompiled = 0;
 		float AverageCompileTime = 0.0f;	// stored explicitly as an optimization
+
+		FShaderTimings& operator+=(const FShaderTimings& Other)
+		{
+			MinCompileTime = FMath::Min(MinCompileTime, Other.MinCompileTime);
+			MaxCompileTime = FMath::Max(MaxCompileTime, Other.MaxCompileTime);
+			TotalCompileTime += Other.TotalCompileTime;
+			TotalPreprocessTime += Other.TotalPreprocessTime;
+			NumCompiled += Other.NumCompiled;
+			AverageCompileTime = (TotalCompileTime + Other.TotalCompileTime) / static_cast<float>(NumCompiled + Other.NumCompiled);
+
+			return *this;
+		}
 	};
 
 	ENGINE_API void RegisterCookedShaders(uint32 NumCooked, float CompileTime, EShaderPlatform Platform, const FString MaterialPath, FString PermutationString = FString(""));
@@ -324,6 +366,11 @@ public:
 	ENGINE_API void WriteStatSummary();
 	ENGINE_API uint32 GetTotalShadersCompiled();
 
+	ENGINE_API void Aggregate(FShaderCompilerStats& Other);
+	ENGINE_API void WriteToCompactBinary(FCbWriter& Writer);
+	ENGINE_API void ReadFromCompactBinary(FCbObjectView& Reader);
+	inline void SetMultiProcessAggregated() { bMultiProcessAggregated = true; }
+
 	void AddDDCMiss(uint32 NumMisses);
 	uint32 GetDDCMisses() const;
 	void AddDDCHit(uint32 NumHits);
@@ -351,68 +398,135 @@ public:
 	/** Informs statistics about a new job batch, so we can tally up batches. */
 	void RegisterJobBatch(int32 NumJobs, EExecutionType ExecType);
 
-	void GatherAnalytics(const FString& BaseName, TArray<FAnalyticsEventAttribute>& Attributes);
+	ENGINE_API void GatherAnalytics(const FString& BaseName, TArray<FAnalyticsEventAttribute>& Attributes);
 
 private:
+	friend class FShaderJobCache;
 	FCriticalSection CompileStatsLock;
 	TSparseArray<ShaderCompilerStats> CompileStats;
 
-	/** This tracks accumulated wait time from local workers during the lifetime of the stats.
-	 *
-	 * Wait time is only counted for local workers that are alive and not between their invocations
-	 */
-	double AccumulatedLocalWorkerIdleTime = 0.0;
+	struct FCounters
+	{
+		/** This tracks accumulated wait time from local workers during the lifetime of the stats.
+		 *
+		 * Wait time is only counted for local workers that are alive and not between their invocations
+		 */
+		double AccumulatedLocalWorkerIdleTime = 0.0;
+
+		/** How many times we registered idle time? */
+		double TimesLocalWorkersWereIdle = 0;
+
+		/** Number of jobs assigned to workers, no matter if they completed or not - used to average pending time. */
+		int64 JobsAssigned = 0;
+
+		/** Total number jobs completed. */
+		int64 JobsCompleted = 0;
+
+		/** Amount of time a job had to spent in pending queue (i.e. waiting to be assigned to a worker). */
+		double AccumulatedPendingTime = 0;
 
-	/** How many times we registered idle time? */
-	double TimesLocalWorkersWereIdle = 0;
+		/** Max amount of time any single job was pending (waiting to be assigned to a worker). */
+		double MaxPendingTime = 0;
 
-	/** Number of jobs assigned to workers, no matter if they completed or not - used to average pending time. */
-	int64 JobsAssigned = 0;
+		/** Amount of time job spent being processed by the worker. */
+		double AccumulatedJobExecutionTime = 0;
 
-	/** Total number jobs completed. */
-	int64 JobsCompleted = 0;
+		/** Max amount of time any single job spent being processed by the worker. */
+		double MaxJobExecutionTime = 0;
 
-	/** Amount of time a job had to spent in pending queue (i.e. waiting to be assigned to a worker). */
-	double AccumulatedPendingTime = 0;
+		/** Amount of time job spent being processed overall. */
+		double AccumulatedJobLifeTime = 0;
 
-	/** Max amount of time any single job was pending (waiting to be assigned to a worker). */
-	double MaxPendingTime = 0;
+		/** Max amount of time any single job spent being processed overall. */
+		double MaxJobLifeTime = 0;
 
-	/** Amount of time job spent being processed by the worker. */
-	double AccumulatedJobExecutionTime = 0;
+		/** Number of local job batches seen. */
+		int64 LocalJobBatchesSeen = 0;
 
-	/** Max amount of time any single job spent being processed by the worker. */
-	double MaxJobExecutionTime = 0;
+		/** Total jobs in local job batches. */
+		int64 TotalJobsReportedInLocalJobBatches = 0;
 
-	/** Amount of time job spent being processed overall. */
-	double AccumulatedJobLifeTime = 0;
+		/** Number of distributed job batches seen. */
+		int64 DistributedJobBatchesSeen = 0;
 
-	/** Max amount of time any single job spent being processed overall. */
-	double MaxJobLifeTime = 0;
+		/** Total jobs in local job batches. */
+		int64 TotalJobsReportedInDistributedJobBatches = 0;
 
-	/** Number of local job batches seen. */
-	int64 LocalJobBatchesSeen = 0;
+		/** Size of the smallest output shader code. */
+		int32 MinShaderCodeSize = 0;
 
-	/** Total jobs in local job batches. */
-	int64 TotalJobsReportedInLocalJobBatches = 0;
+		/** Size of the largest output shader code. */
+		int32 MaxShaderCodeSize = 0;
 
-	/** Number of distributed job batches seen. */
-	int64 DistributedJobBatchesSeen = 0;
+		/** Total accumulated size of all output shader codes. */
+		uint64 AccumulatedShaderCodeSize = 0;
 
-	/** Total jobs in local job batches. */
-	int64 TotalJobsReportedInDistributedJobBatches = 0;
+		/** Number of accumulated output shader codes. */
+		uint64 NumAccumulatedShaderCodes = 0;
 
-	/** Size of the smallest output shader code. */
-	int32 MinShaderCodeSize = 0;
+		/** Total number of DDC misses on shader maps. */
+		uint32 ShaderMapDDCMisses = 0;
 
-	/** Size of the largest output shader code. */
-	int32 MaxShaderCodeSize = 0;
+		/** Total number of DDC hits on shader maps. */
+		uint32 ShaderMapDDCHits = 0;
 
-	/** Total accumulated size of all output shader codes. */
-	uint64 AccumulatedShaderCodeSize = 0;
+		/** Total number of job cache query attempts. */
+		uint64 TotalCacheSearchAttempts = 0;
 
-	/** Number of accumulated output shader codes. */
-	uint64 NumAccumulatedShaderCodes = 0;
+		/** Total number of hits in the job cache (i.e. input hashes seen >1 time) */
+		uint64 TotalCacheHits = 0;
+
+		/** Total number of unique input hashes seen in job cache queries */
+		uint64 UniqueCacheInputHashes = 0;
+
+		/** Total number of unique job outputs stored in the cache.
+		  * Outputs are deduplicated based on a content hash so this number is in practice smaller than UniqueCacheInputHashes.
+		  */
+		uint64 UniqueCacheOutputs = 0;
+
+		/** Total amount of memory currently used by the job cache */
+		uint64 CacheMemUsed = 0;
+
+		/** Memory budget allocated for the job cache */
+		uint64 CacheMemBudget = 0;
+
+		FCounters& operator+=(const FCounters& Other)
+		{
+			AccumulatedLocalWorkerIdleTime += Other.AccumulatedLocalWorkerIdleTime;
+			TimesLocalWorkersWereIdle += Other.TimesLocalWorkersWereIdle;
+			JobsAssigned += Other.JobsAssigned;
+			JobsCompleted += Other.JobsCompleted;
+			AccumulatedPendingTime += Other.AccumulatedPendingTime;
+			MaxPendingTime = FMath::Max(Other.MaxPendingTime, MaxPendingTime);
+			AccumulatedJobExecutionTime += Other.AccumulatedJobExecutionTime;
+			MaxJobExecutionTime = FMath::Max(Other.MaxJobExecutionTime, MaxJobExecutionTime);
+			AccumulatedJobLifeTime += Other.AccumulatedJobLifeTime;
+			MaxJobLifeTime = FMath::Max(Other.MaxJobLifeTime, MaxJobLifeTime);
+			LocalJobBatchesSeen += Other.LocalJobBatchesSeen;
+			TotalJobsReportedInLocalJobBatches += Other.TotalJobsReportedInLocalJobBatches;
+			DistributedJobBatchesSeen += Other.DistributedJobBatchesSeen;
+			TotalJobsReportedInDistributedJobBatches += Other.TotalJobsReportedInDistributedJobBatches;
+			if (Other.MinShaderCodeSize > 0)
+			{
+				MinShaderCodeSize = (MinShaderCodeSize > 0 ? FMath::Min(MinShaderCodeSize, Other.MinShaderCodeSize) : Other.MinShaderCodeSize);
+			}
+			MaxShaderCodeSize = FMath::Max(Other.MaxShaderCodeSize, MaxShaderCodeSize);
+			AccumulatedShaderCodeSize += Other.AccumulatedShaderCodeSize;
+			NumAccumulatedShaderCodes += Other.NumAccumulatedShaderCodes;
+			ShaderMapDDCMisses += Other.ShaderMapDDCMisses;
+			ShaderMapDDCHits += Other.ShaderMapDDCHits;
+			TotalCacheSearchAttempts += Other.TotalCacheSearchAttempts;
+			TotalCacheHits += Other.TotalCacheHits;
+			UniqueCacheInputHashes += Other.UniqueCacheInputHashes;
+			UniqueCacheOutputs += Other.UniqueCacheOutputs;
+			CacheMemUsed += Other.CacheMemUsed;
+			CacheMemBudget += Other.CacheMemBudget;
+
+			return *this;
+		}
+	};
+
+	FCounters Counters;
 
 	/** Accumulates the job lifetimes without overlaps */
 	TArray<TInterval<double>> JobLifeTimeIntervals;
@@ -420,11 +534,7 @@ private:
 	/** Map of shader names to their compilation timings */
 	TMap<FString, FShaderTimings> ShaderTimings;
 
-	/** Total number of DDC misses on shader maps. */
-	uint32 ShaderMapDDCMisses = 0;
-
-	/** Total number of DDC hits on shader maps. */
-	uint32 ShaderMapDDCHits = 0;
+	bool bMultiProcessAggregated = false;
 };
 
 
@@ -592,7 +702,8 @@ public:
 	ENGINE_API int32 GetNumPendingJobs() const;
 	ENGINE_API int32 GetNumOutstandingJobs() const;
 
-	ENGINE_API void GatherAnalytics(TArray<FAnalyticsEventAttribute>& Attributes) const;
+	UE_DEPRECATED(5.4, "FShaderCompilingManager::GatherAnalytics is deprecated; use GatherShaderAnalytics function from ShaderAnalytics.h instead")
+	inline void GatherAnalytics(TArray<FAnalyticsEventAttribute>& Attributes) const {}
 
 	/** 
 	 * Returns whether to display a notification that shader compiling is happening in the background. 
@@ -758,9 +869,16 @@ public:
 
 	/**
 	 * Prints stats related to shader compilation to the log.
-	 * @param bForceLogIgnoringTimeInverval - this function is called often, so not every invocation normally will actually log the stats. This parameter being true bypasses this pacing.
 	 */
-	ENGINE_API void PrintStats(bool bForceLogIgnoringTimeInverval = false);
+	ENGINE_API void PrintStats();
+
+	UE_DEPRECATED(5.4, "PrintStats no longer accepts a 'force' boolean since printing on an interval has become the responsibility of FShaderStatsCollector. Call the no-argument version instead.")
+	inline void PrintStats(bool bForceLogIgnoringTimeInterval) { PrintStats(); }
+
+	/** Retrieve compiler statistics for all compilation done in this process.
+	  *	Note that this will not include stats for any compilation that occurs in worker processes in a multiprocess cook. 
+	  */
+	ENGINE_API void GetLocalStats(FShaderCompilerStats & OutStats) const;
 
 	/** 
 	 * Processes completed asynchronous shader maps, and assigns them to relevant materials.
