WorldPartition: Avoid Loading/Saving EditorPerProjectUserSettings when initializing world with EWorldType::Inactive (Content Browser operations)

#rnx

diff --git a/Engine/Source/Runtime/Engine/Private/WorldPartition/WorldPartitionEditorPerProjectUserSettings.cpp b/Engine/Source/Runtime/Engine/Private/WorldPartition/WorldPartitionEditorPerProjectUserSettings.cpp
--- a/Engine/Source/Runtime/Engine/Private/WorldPartition/WorldPartitionEditorPerProjectUserSettings.cpp
+++ b/Engine/Source/Runtime/Engine/Private/WorldPartition/WorldPartitionEditorPerProjectUserSettings.cpp
@@ -87,6 +87,11 @@ TArray<FName> UWorldPartitionEditorPerProjectUserSettings::GetWorldDataLayersLoa
 
 const FWorldPartitionPerWorldSettings* UWorldPartitionEditorPerProjectUserSettings::GetWorldPartitionPerWorldSettings(UWorld* InWorld) const
 {
+	if (!ShouldLoadSettings(InWorld))
+	{
+		return nullptr;
+	}
+
 	if (const FWorldPartitionPerWorldSettings* ExistingPerWorldSettings = PerWorldEditorSettings.Find(TSoftObjectPtr<UWorld>(InWorld)))
 	{
 		return ExistingPerWorldSettings;
@@ -101,7 +106,12 @@ const FWorldPartitionPerWorldSettings* UWorldPartitionEditorPerProjectUserSettin
 
 bool UWorldPartitionEditorPerProjectUserSettings::ShouldSaveSettings(const UWorld* InWorld) const
 {
-	return InWorld && !InWorld->IsGameWorld() && FPackageName::DoesPackageExist(InWorld->GetPackage()->GetName());
+	return InWorld && !InWorld->IsGameWorld() && (InWorld->WorldType != EWorldType::Inactive) && FPackageName::DoesPackageExist(InWorld->GetPackage()->GetName());
+}
+
+bool UWorldPartitionEditorPerProjectUserSettings::ShouldLoadSettings(const UWorld* InWorld) const
+{
+	return InWorld && (InWorld->WorldType != EWorldType::Inactive);
 }
 
 #endif
diff --git a/Engine/Source/Runtime/Engine/Public/WorldPartition/WorldPartitionEditorPerProjectUserSettings.h b/Engine/Source/Runtime/Engine/Public/WorldPartition/WorldPartitionEditorPerProjectUserSettings.h
--- a/Engine/Source/Runtime/Engine/Public/WorldPartition/WorldPartitionEditorPerProjectUserSettings.h
+++ b/Engine/Source/Runtime/Engine/Public/WorldPartition/WorldPartitionEditorPerProjectUserSettings.h
@@ -109,7 +109,7 @@ public:
 	ENGINE_API void SetWorldDataLayersNonDefaultEditorLoadStates(UWorld* InWorld, const TArray<FName>& InDataLayersLoadedInEditor, const TArray<FName>& InDataLayersNotLoadedInEditor);
 
 private:
-	ENGINE_API const FWorldPartitionPerWorldSettings* GetWorldPartitionPerWorldSettings(UWorld* InWorld) const;
+	const FWorldPartitionPerWorldSettings* GetWorldPartitionPerWorldSettings(UWorld* InWorld) const;
 #endif
 
 #if WITH_EDITORONLY_DATA
@@ -143,7 +143,9 @@ public:
 	uint32 bHideLevelInstanceContent : 1;
 
 private:
-	ENGINE_API bool ShouldSaveSettings(const UWorld* InWorld) const;
+	bool ShouldSaveSettings(const UWorld* InWorld) const;
+
+	bool ShouldLoadSettings(const UWorld* InWorld) const;
 
 	UPROPERTY(config)
 	uint32 bDisableLoadingOfLastLoadedRegions : 1;
