Fix Light transmission doesn't bleed light onto subsurface materials. Add CVar r.SSProfiles.Transmission.UseLegacy (default to 1) to use the legacy transmission profile.

diff --git a/Engine/Source/Runtime/Engine/Private/Rendering/BurleyNormalizedSSS.cpp b/Engine/Source/Runtime/Engine/Private/Rendering/BurleyNormalizedSSS.cpp
--- a/Engine/Source/Runtime/Engine/Private/Rendering/BurleyNormalizedSSS.cpp
+++ b/Engine/Source/Runtime/Engine/Private/Rendering/BurleyNormalizedSSS.cpp
@@ -8,6 +8,14 @@ BurleyNormalizedSSS.cpp: Compute the transmition profile for Burley normalized S
 #include "Rendering/BurleyNormalizedSSS.h"
 #include "Math/Vector.h"
 
+static TAutoConsoleVariable<int32> CVarSSProfilesTransmissionUseLegacy(
+	TEXT("r.SSProfiles.Transmission.UseLegacy"),
+	1,
+	TEXT("0. Use more physically correct transmission profile.\n")
+	TEXT("1. Use legacy transmission profile (default)."),
+	ECVF_RenderThreadSafe
+);
+
 // estimated from the sampling interval, 1/TargetBufferSize(1/32) and MaxTransmissionProfileDistance. If any is changed, this parameter should be re-estimated.
 const float ProfileRadiusOffset = 0.06;
 
@@ -255,6 +263,12 @@ void ComputeTransmissionProfileBurley(FLinearColor* TargetBuffer, uint32 TargetB
 
 	const float SubsurfaceScatteringUnitInCm = 0.1f;
 	const float UnitScale = WorldUnitScale / SubsurfaceScatteringUnitInCm;
+	float InvUnitScale = 1.0 / UnitScale; // Scaling the unit is equivalent to inverse scaling of the profile.
+
+	if (CVarSSProfilesTransmissionUseLegacy.GetValueOnAnyThread() == 1)
+	{
+		InvUnitScale *= 0.1f; // Backward compatibility
+	}
 
 	static float MaxTransmissionProfileDistance = 5.0f; // See SSSS_MAX_TRANSMISSION_PROFILE_DISTANCE in TransmissionCommon.ush
 	static float CmToMm = 10.0f;
@@ -265,8 +279,8 @@ void ComputeTransmissionProfileBurley(FLinearColor* TargetBuffer, uint32 TargetB
 
 	for (uint32 i = 0; i < TargetBufferSize; ++i)
 	{
-		const float DistanceInMm = i * InvSize * (MaxTransmissionProfileDistance * CmToMm) * UnitScale;
-		const float OffsetInMm = (ProfileRadiusOffset * CmToMm) * UnitScale;
+		const float DistanceInMm = i * InvSize * (MaxTransmissionProfileDistance * CmToMm) * InvUnitScale;
+		const float OffsetInMm = (ProfileRadiusOffset * CmToMm) * InvUnitScale;
 
 
 		FLinearColor TransmissionProfile = Burley_TransmissionProfile(DistanceInMm + OffsetInMm, SurfaceAlbedo, ScalingFactor, DiffuseMeanFreePathInMm);
