Fix UonUnity Win64 build for BurleyNormalizedSSS.cpp

diff --git a/Engine/Source/Runtime/Engine/Private/Rendering/BurleyNormalizedSSS.cpp b/Engine/Source/Runtime/Engine/Private/Rendering/BurleyNormalizedSSS.cpp
--- a/Engine/Source/Runtime/Engine/Private/Rendering/BurleyNormalizedSSS.cpp
+++ b/Engine/Source/Runtime/Engine/Private/Rendering/BurleyNormalizedSSS.cpp
@@ -7,6 +7,7 @@ BurleyNormalizedSSS.cpp: Compute the transmition profile for Burley normalized S
 
 #include "Rendering/BurleyNormalizedSSS.h"
 #include "Math/Vector.h"
+#include "HAL/IConsoleManager.h"
 
 static TAutoConsoleVariable<int32> CVarSSProfilesTransmissionUseLegacy(
 	TEXT("r.SSProfiles.Transmission.UseLegacy"),
