Added AActor::FindComponentByTag to avoid allocating and copying arrays of components around

diff --git a/Engine/Source/Runtime/Engine/Classes/GameFramework/Actor.h b/Engine/Source/Runtime/Engine/Classes/GameFramework/Actor.h
--- a/Engine/Source/Runtime/Engine/Classes/GameFramework/Actor.h
+++ b/Engine/Source/Runtime/Engine/Classes/GameFramework/Actor.h
@@ -3516,6 +3516,10 @@ public:
 	UFUNCTION(BlueprintCallable, Category = "Actor", meta = (ComponentClass = "/Script/Engine.ActorComponent", DisplayName = "Get Components By Class", ScriptName = "GetComponentsByClass", DeterminesOutputType = "ComponentClass"))
 	ENGINE_API TArray<UActorComponent*> K2_GetComponentsByClass(TSubclassOf<UActorComponent> ComponentClass) const;
 
+	/** Searches components array and returns first encountered component with a given tag. */
+	UFUNCTION(BlueprintCallable, Category = "Actor", meta = (ComponentClass = "/Script/Engine.ActorComponent"), meta = (DeterminesOutputType = "ComponentClass"))
+	UActorComponent* FindComponentByTag(TSubclassOf<UActorComponent> ComponentClass, FName Tag) const;
+
 	/** Gets all the components that inherit from the given class with a given tag. */
 	UFUNCTION(BlueprintCallable, Category = "Actor", meta = (ComponentClass = "/Script/Engine.ActorComponent"), meta = (DeterminesOutputType = "ComponentClass"))
 	ENGINE_API TArray<UActorComponent*> GetComponentsByTag(TSubclassOf<UActorComponent> ComponentClass, FName Tag) const;
@@ -3535,6 +3539,15 @@ public:
 
 		return (T*)FindComponentByClass(T::StaticClass());
 	}
+	
+	/** Templatized version of FindComponentByTag that handles casting for you */
+	template<class T>
+	T* FindComponentByTag(const FName& Tag) const
+	{
+		static_assert(TPointerIsConvertibleFromTo<T, const UActorComponent>::Value, "'T' template parameter to FindComponentByTag must be derived from UActorComponent");
+
+		return (T*)FindComponentByTag(T::StaticClass(), Tag);
+	}
 
 	/** Templatized version of FindComponentByInterface that handles casting for you */
 	template<class T>
diff --git a/Engine/Source/Runtime/Engine/Private/Actor.cpp b/Engine/Source/Runtime/Engine/Private/Actor.cpp
--- a/Engine/Source/Runtime/Engine/Private/Actor.cpp
+++ b/Engine/Source/Runtime/Engine/Private/Actor.cpp
@@ -3502,6 +3502,19 @@ TArray<UActorComponent*> AActor::K2_GetComponentsByClass(TSubclassOf<UActorCompo
 	return MoveTemp(Components);
 }
 
+UActorComponent* AActor::FindComponentByTag(TSubclassOf<UActorComponent> ComponentClass, FName Tag) const
+{
+	for (UActorComponent* Component : GetComponents())
+	{
+		if (Component && Component->IsA(ComponentClass) && Component->ComponentHasTag(Tag))
+		{
+			return Component;
+		}
+	}
+
+	return nullptr;
+}
+
 TArray<UActorComponent*> AActor::GetComponentsByTag(TSubclassOf<UActorComponent> ComponentClass, FName Tag) const
 {
 	TInlineComponentArray<UActorComponent*> ComponentsByClass;
