Added missing ENGINE_API to AActor::FindComponentByTag

#rnx

diff --git a/Engine/Source/Runtime/Engine/Classes/GameFramework/Actor.h b/Engine/Source/Runtime/Engine/Classes/GameFramework/Actor.h
--- a/Engine/Source/Runtime/Engine/Classes/GameFramework/Actor.h
+++ b/Engine/Source/Runtime/Engine/Classes/GameFramework/Actor.h
@@ -3518,7 +3518,7 @@ public:
 
 	/** Searches components array and returns first encountered component with a given tag. */
 	UFUNCTION(BlueprintCallable, Category = "Actor", meta = (ComponentClass = "/Script/Engine.ActorComponent"), meta = (DeterminesOutputType = "ComponentClass"))
-	UActorComponent* FindComponentByTag(TSubclassOf<UActorComponent> ComponentClass, FName Tag) const;
+	ENGINE_API UActorComponent* FindComponentByTag(TSubclassOf<UActorComponent> ComponentClass, FName Tag) const;
 
 	/** Gets all the components that inherit from the given class with a given tag. */
 	UFUNCTION(BlueprintCallable, Category = "Actor", meta = (ComponentClass = "/Script/Engine.ActorComponent"), meta = (DeterminesOutputType = "ComponentClass"))
