Substrate - removed invalid comment

diff --git a/Engine/Source/Runtime/Engine/Classes/Materials/MaterialExpressionStrata.h b/Engine/Source/Runtime/Engine/Classes/Materials/MaterialExpressionStrata.h
--- a/Engine/Source/Runtime/Engine/Classes/Materials/MaterialExpressionStrata.h
+++ b/Engine/Source/Runtime/Engine/Classes/Materials/MaterialExpressionStrata.h
@@ -1000,7 +1000,7 @@ class UMaterialExpressionStrataTransmittanceToMFP : public UMaterialExpressionSt
 	FExpressionInput TransmittanceColor;
 
 	/**
-	* The desired thickness in centimeter. This can be set lower than 0.1mm (= 0.01cm) to enable the Thin lighting model on the slab node for instance.
+	* The desired thickness in centimeter.
 	* Another use case example: this node output called thickness can be modulated before it is plugged in a slab node, this can be used to achieve simple scattering/transmittance variation of the same material.
 	*/
 	UPROPERTY()
