Made some SkyAtmosphere properties appear close to others instead of down the detail panel.

diff --git a/Engine/Source/Runtime/Engine/Classes/Components/VolumetricCloudComponent.h b/Engine/Source/Runtime/Engine/Classes/Components/VolumetricCloudComponent.h
--- a/Engine/Source/Runtime/Engine/Classes/Components/VolumetricCloudComponent.h
+++ b/Engine/Source/Runtime/Engine/Classes/Components/VolumetricCloudComponent.h
@@ -122,18 +122,17 @@ class UVolumetricCloudComponent : public USceneComponent
 	UPROPERTY(EditAnywhere, BlueprintReadOnly, Category = "Cloud Tracing", meta = (UIMin = "0.0", UIMax = "1.0", ClampMin = "0.0", ClampMax = "1.0", SliderExponent = 5.0))
 	float StopTracingTransmittanceThreshold;
 	
-	
 	/** Specify the aerial perspective start distance on cloud for Rayleigh scattering only. (kilometers) */
-	UPROPERTY(EditAnywhere, BlueprintReadOnly, interp, Category = "Art Direction", AdvancedDisplay, meta = (UIMin = 0.0, UIMax = 100.0, ClampMin = 0.0, SliderExponent = 2.0))
+	UPROPERTY(EditAnywhere, BlueprintReadOnly, interp, Category = "Art Direction", meta = (UIMin = 0.0, UIMax = 100.0, ClampMin = 0.0, SliderExponent = 2.0))
 	float AerialPespectiveRayleighScatteringStartDistance;
-	/** Specify the distance over which the Rayleigh scattering will linearly ramp up to ful leffect. (kilometers) */
-	UPROPERTY(EditAnywhere, BlueprintReadOnly, interp, Category = "Art Direction", AdvancedDisplay, meta = (UIMin = 0.0, UIMax = 100.0, ClampMin = 0.0, SliderExponent = 2.0))
+	/** Specify the distance over which the Rayleigh scattering will linearly ramp up to full effect. (kilometers) */
+	UPROPERTY(EditAnywhere, BlueprintReadOnly, interp, Category = "Art Direction", meta = (UIMin = 0.0, UIMax = 100.0, ClampMin = 0.0, SliderExponent = 2.0))
 	float AerialPespectiveRayleighScatteringFadeDistance;
 	/** Specify the aerial perspective start distance on cloud for Mie scattering only. (kilometers) */
-	UPROPERTY(EditAnywhere, BlueprintReadOnly, interp, Category = "Art Direction", AdvancedDisplay, meta = (UIMin = 0.0, UIMax = 100.0, ClampMin = 0.0, SliderExponent = 2.0))
+	UPROPERTY(EditAnywhere, BlueprintReadOnly, interp, Category = "Art Direction", meta = (UIMin = 0.0, UIMax = 100.0, ClampMin = 0.0, SliderExponent = 2.0))
 	float AerialPespectiveMieScatteringStartDistance;
-	/** Specify the distance over which the Rayleigh scattering will linearly ramp up to ful leffect. (kilometers) */
-	UPROPERTY(EditAnywhere, BlueprintReadOnly, interp, Category = "Art Direction", AdvancedDisplay, meta = (UIMin = 0.0, UIMax = 100.0, ClampMin = 0.0, SliderExponent = 2.0))
+	/** Specify the distance over which the Rayleigh scattering will linearly ramp up to full effect. (kilometers) */
+	UPROPERTY(EditAnywhere, BlueprintReadOnly, interp, Category = "Art Direction", meta = (UIMin = 0.0, UIMax = 100.0, ClampMin = 0.0, SliderExponent = 2.0))
 	float AerialPespectiveMieScatteringFadeDistance;
 
 
