Add option to log stack when GC is forced

diff --git a/Engine/Source/Runtime/Engine/Private/UnrealEngine.cpp b/Engine/Source/Runtime/Engine/Private/UnrealEngine.cpp
--- a/Engine/Source/Runtime/Engine/Private/UnrealEngine.cpp
+++ b/Engine/Source/Runtime/Engine/Private/UnrealEngine.cpp
@@ -1582,6 +1582,19 @@ static FAutoConsoleVariableRef CVarLowMemoryTimeBetweenPurgingPendingKillObjects
 	ECVF_Default
 );
 
+#ifndef UE_ENABLE_LOG_STACK_ON_FORCE_GC
+#define UE_ENABLE_LOG_STACK_ON_FORCE_GC 0
+#endif
+#if UE_ENABLE_LOG_STACK_ON_FORCE_GC
+static bool GLogStackOnForceGC = false;
+static FAutoConsoleVariableRef CVarLogStackOnForceGC(
+	TEXT("gc.LogStackOnForceGC"),
+	GLogStackOnForceGC,
+	TEXT("Whether to log a stack trace when ForceGarbageCollection is called."),
+	ECVF_Default
+);
+#endif
+
 // see also: s.ForceGCAfterLevelStreamedOut, s.ContinuouslyIncrementalGCWhileLevelsPendingPurge
 static float GLowMemoryTimeBetweenPurgingPendingLevels = 15.0f;
 static FAutoConsoleVariableRef CVarLowMemoryTimeBetweenPurgingPendingLevels(
@@ -1712,6 +1725,13 @@ void UEngine::ForceGarbageCollection(bool bForcePurge/*=false*/)
 {
 	TimeSinceLastPendingKillPurge = 1.0f + GetTimeBetweenGarbageCollectionPasses();
 	bFullPurgeTriggered = bFullPurgeTriggered || bForcePurge;
+
+#if UE_ENABLE_LOG_STACK_ON_FORCE_GC
+	if (GLogStackOnForceGC)
+	{
+		FDebug::DumpStackTraceToLog(ELogVerbosity::Log);
+	}
+#endif
 }
 
 void UEngine::DelayGarbageCollection()
