Fix for issue where cooking multiple platforms at once could cause some of them to have empty Nanite bulkdata

diff --git a/Engine/Source/Runtime/Engine/Private/StaticMesh.cpp b/Engine/Source/Runtime/Engine/Private/StaticMesh.cpp
--- a/Engine/Source/Runtime/Engine/Private/StaticMesh.cpp
+++ b/Engine/Source/Runtime/Engine/Private/StaticMesh.cpp
@@ -4969,18 +4969,20 @@ void UStaticMesh::WillNeverCacheCookedPlatformDataAgain()
 
 void UStaticMesh::ClearCachedCookedPlatformData(const ITargetPlatform* TargetPlatform)
 {
-	FStaticMeshRenderData& PlatformRenderData = GetPlatformStaticMeshRenderData(this, TargetPlatform);
-	PlatformRenderData.NaniteResourcesPtr->DropBulkData();
+	if (!IsRunningCookCommandlet())
+	{
+		// Drop bulk data after serialization as editor streams from DDC and doesn't need it to be resident.
+		// When running the cook commandlet multiple platforms might share the same FStaticMeshRenderData,
+		// so we defer the dropping to ClearAllCachedCookedPlatformData.
+		FStaticMeshRenderData& PlatformRenderData = GetPlatformStaticMeshRenderData(this, TargetPlatform);
+		PlatformRenderData.NaniteResourcesPtr->DropBulkData();
+	}
 }
 
 void UStaticMesh::ClearAllCachedCookedPlatformData()
 {
-	for (ITargetPlatform* TargetPlatform : GetTargetPlatformManagerRef().GetActiveTargetPlatforms())
-	{
-		ClearCachedCookedPlatformData(TargetPlatform);
-	}
-
 	GetRenderData()->NextCachedRenderData.Reset();
+	GetRenderData()->NaniteResourcesPtr->DropBulkData();
 }
 
 void UStaticMesh::BeginCacheForCookedPlatformData(const ITargetPlatform* TargetPlatform)
