Control Rig: Fix mapping a rig to a skel mesh with out of wack base ref pose

diff --git a/Engine/Source/Runtime/Engine/Private/SkeletalMesh.cpp b/Engine/Source/Runtime/Engine/Private/SkeletalMesh.cpp
--- a/Engine/Source/Runtime/Engine/Private/SkeletalMesh.cpp
+++ b/Engine/Source/Runtime/Engine/Private/SkeletalMesh.cpp
@@ -5192,13 +5192,19 @@ void USkeletalMesh::GetMappableNodeData(TArray<FName>& OutNames, TArray<FNodeIte
 	// hasn't tested this route, but we don't have retarget base pose if not editor, wonder we should to non-editor soon
 	ensure(false);
 	FAnimationRuntime::FillUpComponentSpaceTransforms(GetRefSkeleton(), GetRefSkeleton().GetRefBonePose(), ComponentSpaceRefPose);
-#endif // 
+#endif //
 
 	const int32 NumJoint = GetRefSkeleton().GetNum();
 	// allocate buffer
 	OutNames.Reset(NumJoint);
 	OutNodeItems.Reset(NumJoint);
 
+	if(ComponentSpaceRefPose.Num() < NumJoint)
+	{
+		// if the mesh's RetargetBasePose is out of whack we should rely on the ref skeleton
+		FAnimationRuntime::FillUpComponentSpaceTransforms(GetRefSkeleton(), GetRefSkeleton().GetRefBonePose(), ComponentSpaceRefPose);
+	}
+
 	if (NumJoint > 0)
 	{
 		OutNames.AddDefaulted(NumJoint);
