Add a way to treat composite members as complete values when generating the custom property list for the Blueprint instance post-construct initialization fast path.

Change summary:
- Added UBlueprintGeneratedClass::RequiresCompleteValueForPostConstruction (config). These members will not be further decomposed into sub property delta lists.
- Added a CVar to allow the BP property initialization fast path to more easily be disabled for testing (similar to the CVar that exists for fast path component instancing).
- Revised property list generation to avoid storing extra "placeholder" entries for array values that exceed the size of the native parent default container's value (perfmem).

diff --git a/Engine/Config/BaseEngine.ini b/Engine/Config/BaseEngine.ini
--- a/Engine/Config/BaseEngine.ini
+++ b/Engine/Config/BaseEngine.ini
@@ -3309,3 +3309,9 @@ bEnableDiskOverflowStore=true
 
 [LowLevelTestsSettings]
 bUpdateBuildGraphPropertiesFile=false
+
+[/Script/Engine.BlueprintGeneratedClass]
+; Members of this struct property are inter-dependent and their initialization can be deferred, which
+; means native defaults can match BP defaults at BP compile time while still requiring post-construct
+; initialization, if other members that they are dependent on differ in Blueprint class default data.
++RequiresCompleteValueForPostConstruction=/Script/Engine.BodyInstance:CollisionResponses
\ No newline at end of file
diff --git a/Engine/Source/Runtime/Engine/Classes/Engine/BlueprintGeneratedClass.h b/Engine/Source/Runtime/Engine/Classes/Engine/BlueprintGeneratedClass.h
--- a/Engine/Source/Runtime/Engine/Classes/Engine/BlueprintGeneratedClass.h
+++ b/Engine/Source/Runtime/Engine/Classes/Engine/BlueprintGeneratedClass.h
@@ -627,7 +627,7 @@ struct FBPComponentClassOverride
 	}
 };
 
-UCLASS(NeedsDeferredDependencyLoading, MinimalAPI)
+UCLASS(NeedsDeferredDependencyLoading, MinimalAPI, config=Engine)
 class UBlueprintGeneratedClass : public UClass, public IBlueprintPropertyGuidProvider
 {
 	GENERATED_UCLASS_BODY()
@@ -923,6 +923,9 @@ protected:
 	ENGINE_API bool BuildCustomArrayPropertyListForPostConstruction(FArrayProperty* ArrayProperty, FCustomPropertyListNode*& InPropertyList, const uint8* DataPtr, const uint8* DefaultDataPtr, int32 StartIndex = 0);
 
 private:
+	/** List of native class-owned properties that require complete values for comparison when generating the post-construction property list due to how they are used. These properties cannot be reduced into a list of subfields (e.g. structs). */
+	UPROPERTY(config)
+	TArray<FString> RequiresCompleteValueForPostConstruction;
 	/** List of native class-owned properties that differ from defaults. This is used to optimize property initialization during post-construction by minimizing the number of native class-owned property values that get copied to the new instance. */
 	TIndirectArray<FCustomPropertyListNode> CustomPropertyListForPostConstruction;
 	/** In some cases UObject::ConditionalPostLoad() code calls PostLoadDefaultObject() on a class that's still being serialized. */
diff --git a/Engine/Source/Runtime/Engine/Private/BlueprintGeneratedClass.cpp b/Engine/Source/Runtime/Engine/Private/BlueprintGeneratedClass.cpp
--- a/Engine/Source/Runtime/Engine/Private/BlueprintGeneratedClass.cpp
+++ b/Engine/Source/Runtime/Engine/Private/BlueprintGeneratedClass.cpp
@@ -44,6 +44,14 @@ extern UNREALED_API class UEditorEngine* GEditor;
 DEFINE_STAT(STAT_PersistentUberGraphFrameMemory);
 DEFINE_STAT(STAT_BPCompInstancingFastPathMemory);
 
+static int32 GBlueprintNativePropertyInitFastPathDisabled = 0;
+static FAutoConsoleVariableRef CVarBlueprintNativePropertyInitFastPathDisabled(
+	TEXT("bp.NativePropertyInitFastPathDisabled"),
+	GBlueprintNativePropertyInitFastPathDisabled,
+	TEXT("Disable the native property initialization fast path."),
+	ECVF_Default
+);
+
 int32 GBlueprintComponentInstancingFastPathDisabled = 0;
 static FAutoConsoleVariableRef CVarBlueprintComponentInstancingFastPathDisabled(
 	TEXT("bp.ComponentInstancingFastPathDisabled"),
@@ -93,6 +101,38 @@ void WriteBPGCBreadcrumbs(FCrashContextExtendedWriter& Writer, const BPGCBreadcr
 }
 #endif // WITH_ADDITIONAL_CRASH_CONTEXTS
 
+namespace UE::Runtime::Engine::Private
+{
+	struct FBlueprintGeneratedClassUtils
+	{
+		static UClass* FindFirstNativeClassInHierarchy(UClass* InClass)
+		{
+			UClass* CurrentClass = InClass;
+			while (CurrentClass && !CurrentClass->HasAnyClassFlags(CLASS_Native | CLASS_Intrinsic))
+			{
+				CurrentClass = CurrentClass->GetSuperClass();
+			}
+
+			return CurrentClass;
+		}
+
+		static bool ShouldInitializePropertyDuringPostConstruction(const FProperty& Property)
+		{
+			const UClass* OwnerClass = Property.GetOwnerClass();
+			const bool bIsConfigProperty = Property.HasAnyPropertyFlags(CPF_Config) && !(OwnerClass && OwnerClass->HasAnyClassFlags(CLASS_PerObjectConfig));
+			const bool bIsTransientProperty = Property.HasAnyPropertyFlags(CPF_Transient | CPF_DuplicateTransient | CPF_NonPIEDuplicateTransient);
+
+			// Skip config properties as they're already in the PostConstructLink chain. Also skip transient properties if they contain a reference to an instanced subobjects (as those should not be initialized from defaults).
+			if (!bIsConfigProperty && (!bIsTransientProperty || !Property.ContainsInstancedObjectProperty()))
+			{
+				return true;
+			}
+
+			return false;
+		}
+	};
+}
+
 PRAGMA_DISABLE_DEPRECATION_WARNINGS
 UBlueprintGeneratedClass::UBlueprintGeneratedClass(const FObjectInitializer& ObjectInitializer)
 	: Super(ObjectInitializer)
@@ -714,80 +754,91 @@ void UBlueprintGeneratedClass::ConformSparseClassData(UObject* Object)
 
 bool UBlueprintGeneratedClass::BuildCustomPropertyListForPostConstruction(FCustomPropertyListNode*& InPropertyList, UStruct* InStruct, const uint8* DataPtr, const uint8* DefaultDataPtr)
 {
+	using namespace UE::Runtime::Engine::Private;
+
 	const UClass* OwnerClass = Cast<UClass>(InStruct);
 	FCustomPropertyListNode** CurrentNodePtr = &InPropertyList;
 
 	for (FProperty* Property = InStruct->PropertyLink; Property; Property = Property->PropertyLinkNext)
 	{
-		const bool bIsConfigProperty = Property->HasAnyPropertyFlags(CPF_Config) && !(OwnerClass && OwnerClass->HasAnyClassFlags(CLASS_PerObjectConfig));
-		const bool bIsTransientProperty = Property->HasAnyPropertyFlags(CPF_Transient | CPF_DuplicateTransient | CPF_NonPIEDuplicateTransient);
-
-		// Skip config properties as they're already in the PostConstructLink chain. Also skip transient properties if they contain a reference to an instanced subobjects (as those should not be initialized from defaults).
-		if (!bIsConfigProperty && (!bIsTransientProperty || !Property->ContainsInstancedObjectProperty()))
+		if (FBlueprintGeneratedClassUtils::ShouldInitializePropertyDuringPostConstruction(*Property))
 		{
+			// Some properties require a full value comparison; check for those cases here.
+			const bool bAlwaysUseCompleteValue = RequiresCompleteValueForPostConstruction.Contains(Property->GetPathName());
+
 			for (int32 Idx = 0; Idx < Property->ArrayDim; Idx++)
 			{
 				const uint8* PropertyValue = Property->ContainerPtrToValuePtr<uint8>(DataPtr, Idx);
 				const uint8* DefaultPropertyValue = Property->ContainerPtrToValuePtrForDefaults<uint8>(InStruct, DefaultDataPtr, Idx);
 
-				// If this is a struct property, recurse to pull out any fields that differ from the native CDO.
-				if (FStructProperty* StructProperty = CastField<FStructProperty>(Property))
+				bool bUseCompleteValue = bAlwaysUseCompleteValue;
+				if (!bUseCompleteValue)
 				{
-					// Create a new node for the struct property.
-					*CurrentNodePtr = new FCustomPropertyListNode(Property, Idx);
-					CustomPropertyListForPostConstruction.Add(*CurrentNodePtr);
-
-					UScriptStruct::ICppStructOps* CppStructOps = nullptr;
-					if (StructProperty->Struct)
+					// If this is a struct property, recurse to pull out any fields that differ from the native CDO.
+					if (FStructProperty* StructProperty = CastField<FStructProperty>(Property))
 					{
-						CppStructOps = StructProperty->Struct->GetCppStructOps();
-					}
+						// Create a new node for the struct property.
+						*CurrentNodePtr = new FCustomPropertyListNode(Property, Idx);
+						CustomPropertyListForPostConstruction.Add(*CurrentNodePtr);
 
-					// Check if we should initialize using the full value (e.g. a USTRUCT with one or more non-reflected fields).
-					bool bIsIdentical = false;
-					const uint32 PortFlags = 0;
-					if(!CppStructOps || !CppStructOps->HasIdentical() || !CppStructOps->Identical(PropertyValue, DefaultPropertyValue, PortFlags, bIsIdentical))
-					{
-						// Recursively gather up all struct fields that differ and assign to the current node's sub property list.
-						bIsIdentical = !BuildCustomPropertyListForPostConstruction((*CurrentNodePtr)->SubPropertyList, StructProperty->Struct, PropertyValue, DefaultPropertyValue);
-					}
+						UScriptStruct::ICppStructOps* CppStructOps = nullptr;
+						if (StructProperty->Struct)
+						{
+							CppStructOps = StructProperty->Struct->GetCppStructOps();
+						}
 
-					if (!bIsIdentical)
-					{
-						// Advance to the next node in the list.
-						CurrentNodePtr = &(*CurrentNodePtr)->PropertyListNext;
+						// Check if we should initialize using the full value (e.g. a USTRUCT with one or more non-reflected fields).
+						bool bIsIdentical = false;
+						const uint32 PortFlags = 0;
+						if (!CppStructOps || !CppStructOps->HasIdentical() || !CppStructOps->Identical(PropertyValue, DefaultPropertyValue, PortFlags, bIsIdentical))
+						{
+							// Recursively gather up all struct fields that differ and assign to the current node's sub property list.
+							bIsIdentical = !BuildCustomPropertyListForPostConstruction((*CurrentNodePtr)->SubPropertyList, StructProperty->Struct, PropertyValue, DefaultPropertyValue);
+						}
+
+						if (!bIsIdentical)
+						{
+							// Advance to the next node in the list.
+							CurrentNodePtr = &(*CurrentNodePtr)->PropertyListNext;
+						}
+						else
+						{
+							// Remove the node for the struct property since it does not differ from the native CDO.
+							CustomPropertyListForPostConstruction.RemoveAt(CustomPropertyListForPostConstruction.Num() - 1);
+
+							// Clear the current node ptr since the array will have freed up the memory it referenced.
+							*CurrentNodePtr = nullptr;
+						}
 					}
-					else
+					else if (FArrayProperty* ArrayProperty = CastField<FArrayProperty>(Property))
 					{
-						// Remove the node for the struct property since it does not differ from the native CDO.
-						CustomPropertyListForPostConstruction.RemoveAt(CustomPropertyListForPostConstruction.Num() - 1);
+						// Create a new node for the array property.
+						*CurrentNodePtr = new FCustomPropertyListNode(Property, Idx);
+						CustomPropertyListForPostConstruction.Add(*CurrentNodePtr);
 
-						// Clear the current node ptr since the array will have freed up the memory it referenced.
-						*CurrentNodePtr = nullptr;
-					}
-				}
-				else if (FArrayProperty* ArrayProperty = CastField<FArrayProperty>(Property))
-				{
-					// Create a new node for the array property.
-					*CurrentNodePtr = new FCustomPropertyListNode(Property, Idx);
-					CustomPropertyListForPostConstruction.Add(*CurrentNodePtr);
+						// Recursively gather up all array item indices that differ and assign to the current node's sub property list.
+						if (BuildCustomArrayPropertyListForPostConstruction(ArrayProperty, (*CurrentNodePtr)->SubPropertyList, PropertyValue, DefaultPropertyValue))
+						{
+							// Advance to the next node in the list.
+							CurrentNodePtr = &(*CurrentNodePtr)->PropertyListNext;
+						}
+						else
+						{
+							// Remove the node for the array property since it does not differ from the native CDO.
+							CustomPropertyListForPostConstruction.RemoveAt(CustomPropertyListForPostConstruction.Num() - 1);
 
-					// Recursively gather up all array item indices that differ and assign to the current node's sub property list.
-					if (BuildCustomArrayPropertyListForPostConstruction(ArrayProperty, (*CurrentNodePtr)->SubPropertyList, PropertyValue, DefaultPropertyValue))
-					{
-						// Advance to the next node in the list.
-						CurrentNodePtr = &(*CurrentNodePtr)->PropertyListNext;
+							// Clear the current node ptr since the array will have freed up the memory it referenced.
+							*CurrentNodePtr = nullptr;
+						}
 					}
 					else
 					{
-						// Remove the node for the array property since it does not differ from the native CDO.
-						CustomPropertyListForPostConstruction.RemoveAt(CustomPropertyListForPostConstruction.Num() - 1);
-
-						// Clear the current node ptr since the array will have freed up the memory it referenced.
-						*CurrentNodePtr = nullptr;
+						// Not explicitly handled above; fall back to using a full value comparison and emit this property if anything differs.
+						bUseCompleteValue = true;
 					}
 				}
-				else if (!Property->Identical(PropertyValue, DefaultPropertyValue))
+				
+				if (bUseCompleteValue && !Property->Identical(PropertyValue, DefaultPropertyValue))
 				{
 					// Create a new node, link it into the chain and add it into the array.
 					*CurrentNodePtr = new FCustomPropertyListNode(Property, Idx);
@@ -873,25 +924,9 @@ bool UBlueprintGeneratedClass::BuildCustomArrayPropertyListForPostConstruction(F
 		}
 		else
 		{
-			// Create a temp default array as a placeholder to compare against the remaining elements in the value.
-			FScriptArray TempDefaultArray;
-			const int32 Count = ArrayValueHelper.Num() - DefaultArrayValueHelper.Num();
-			TempDefaultArray.Add(Count, ArrayProperty->Inner->ElementSize, ArrayProperty->Inner->GetMinAlignment());
-			uint8 *Dest = (uint8*)TempDefaultArray.GetData();
-			if (ArrayProperty->Inner->PropertyFlags & CPF_ZeroConstructor)
-			{
-				FMemory::Memzero(Dest, Count * ArrayProperty->Inner->ElementSize);
-			}
-			else
-			{
-				for (int32 i = 0; i < Count; i++, Dest += ArrayProperty->Inner->ElementSize)
-				{
-					ArrayProperty->Inner->InitializeValue(Dest);
-				}
-			}
-
-			// Recursively fill out the property list for the remainder of the elements in the value that extend beyond the size of the default value.
-			BuildCustomArrayPropertyListForPostConstruction(ArrayProperty, *CurrentArrayNodePtr, DataPtr, (uint8*)&TempDefaultArray, ArrayValueIndex);
+			// NULL signals the end of the array value change at the current index.
+			*CurrentArrayNodePtr = new FCustomPropertyListNode(nullptr, ArrayValueIndex);
+			CustomPropertyListForPostConstruction.Add(*CurrentArrayNodePtr);
 
 			// Don't need to record anything else.
 			break;
@@ -904,18 +939,14 @@ bool UBlueprintGeneratedClass::BuildCustomArrayPropertyListForPostConstruction(F
 
 void UBlueprintGeneratedClass::UpdateCustomPropertyListForPostConstruction()
 {
+	using namespace UE::Runtime::Engine::Private;
+
 	// Empty the current list.
 	CustomPropertyListForPostConstruction.Reset();
 	bCustomPropertyListForPostConstructionInitialized = false;
 
 	// Find the first native antecedent. All non-native decendant properties are attached to the PostConstructLink chain (see UStruct::Link), so we only need to worry about properties owned by native super classes here.
-	UClass* SuperClass = GetSuperClass();
-	while (SuperClass && !SuperClass->HasAnyClassFlags(CLASS_Native | CLASS_Intrinsic))
-	{
-		SuperClass = SuperClass->GetSuperClass();
-	}
-
-	if (SuperClass)
+	if (UClass* SuperClass = FBlueprintGeneratedClassUtils::FindFirstNativeClassInHierarchy(GetSuperClass()))
 	{
 		check(ClassDefaultObject != nullptr);
 
@@ -940,11 +971,36 @@ void UBlueprintGeneratedClass::SetupObjectInitializer(FObjectInitializer& Object
 void UBlueprintGeneratedClass::InitPropertiesFromCustomList(uint8* DataPtr, const uint8* DefaultDataPtr)
 {
 	FScopeLock SerializeAndPostLoadLock(&SerializeAndPostLoadCritical);
-	checkf(bCustomPropertyListForPostConstructionInitialized, TEXT("Custom Property List Not Initialized for %s"), *GetPathNameSafe(this)); // Something went wrong, probably a race condition
 
-	if (const FCustomPropertyListNode* CustomPropertyList = GetCustomPropertyListForPostConstruction())
+	if (GBlueprintNativePropertyInitFastPathDisabled
+		|| !ensureMsgf(bCustomPropertyListForPostConstructionInitialized, TEXT("Custom Property List Not Initialized for %s"), *GetPathNameSafe(this))) // Something went wrong, probably a race condition
 	{
-		InitPropertiesFromCustomList(CustomPropertyList, this, DataPtr, DefaultDataPtr);
+		// Slow path - initialize all inherited native properties, regardless of whether they have been modified away from natively-initialized values.
+		using namespace UE::Runtime::Engine::Private;
+		if (const UClass* NativeParentClass = FBlueprintGeneratedClassUtils::FindFirstNativeClassInHierarchy(GetSuperClass()))
+		{
+			for (TFieldIterator<FProperty> It(NativeParentClass); It; ++It)
+			{
+				const FProperty* Property = *It;
+				check(Property);
+
+				if (FBlueprintGeneratedClassUtils::ShouldInitializePropertyDuringPostConstruction(*Property))
+				{
+					Property->CopyCompleteValue_InContainer(DataPtr, DefaultDataPtr);
+				}
+			}
+		}
+	}
+	else
+	{
+		// Note: It is valid to have a NULL custom property list when the 'initialized' flag is also set - this
+		// implies that no inherited native properties have been modified and thus do not need to be initialized.
+		if (const FCustomPropertyListNode* CustomPropertyList = GetCustomPropertyListForPostConstruction())
+		{
+			// Fast path - will initialize only the subset of inherited native properties that have been modified. The rest
+			// have already been initialized by the native class ctor, and initializing them again here would be redundant.
+			InitPropertiesFromCustomList(CustomPropertyList, this, DataPtr, DefaultDataPtr);
+		}
 	}
 }
 
@@ -980,9 +1036,15 @@ void UBlueprintGeneratedClass::InitArrayPropertyFromCustomList(const FArrayPrope
 		DstArrayValueHelper.RemoveValues(SrcNum, DstNum - SrcNum);
 	}
 
+	int32 ArrayIndex = 0;
 	for (const FCustomPropertyListNode* CustomArrayPropertyListNode = InPropertyList; CustomArrayPropertyListNode; CustomArrayPropertyListNode = CustomArrayPropertyListNode->PropertyListNext)
 	{
-		int32 ArrayIndex = CustomArrayPropertyListNode->ArrayIndex;
+		ArrayIndex = CustomArrayPropertyListNode->ArrayIndex;
+		if (CustomArrayPropertyListNode->Property == nullptr)
+		{
+			// This signals the end of the default value change.
+			break;
+		}
 
 		uint8* DstArrayItemValue = DstArrayValueHelper.GetRawPtr(ArrayIndex);
 		const uint8* SrcArrayItemValue = SrcArrayValueHelper.GetRawPtr(ArrayIndex);
@@ -1008,6 +1070,35 @@ void UBlueprintGeneratedClass::InitArrayPropertyFromCustomList(const FArrayPrope
 			ArrayProperty->Inner->CopyCompleteValue(DstArrayItemValue, SrcArrayItemValue);
 		}
 	}
+
+	// If necessary, copy the remainder of the source array value to the destination. It's possible for the number of elements in the
+	// derived Blueprint class default object to exceed the number of elements in the first antecedent native superclass default object,
+	// since the custom property list that's generated for the array will cover only the subset of elements that are common to both sides.
+	if (ArrayIndex < SrcNum)
+	{
+		if (ArrayProperty->Inner->HasAnyPropertyFlags(CPF_IsPlainOldData))
+		{
+			uint8* DstArrayItemValue = DstArrayValueHelper.GetRawPtr(ArrayIndex);
+			const uint8* SrcArrayItemValue = SrcArrayValueHelper.GetRawPtr(ArrayIndex);
+
+			FMemory::Memcpy(DstArrayItemValue, SrcArrayItemValue, (SrcNum - ArrayIndex) * ArrayProperty->Inner->ElementSize);
+		}
+		else
+		{
+			for (; ArrayIndex < SrcNum; ++ArrayIndex)
+			{
+				uint8* DstArrayItemValue = DstArrayValueHelper.GetRawPtr(ArrayIndex);
+				const uint8* SrcArrayItemValue = SrcArrayValueHelper.GetRawPtr(ArrayIndex);
+
+				if (DstArrayItemValue == nullptr && SrcArrayItemValue == nullptr)
+				{
+					continue;
+				}
+
+				ArrayProperty->Inner->CopyCompleteValue(DstArrayItemValue, SrcArrayItemValue);
+			}
+		}
+	}
 }
 
 bool UBlueprintGeneratedClass::InitPropertyFromSubPropertyList(const FProperty* Property, const FCustomPropertyListNode* SubPropertyList, uint8* PropertyValue, const uint8* DefaultPropertyValue)
