WorldPartition: Replaced check with ensure as the source of assert is an invalid (NaN) provided streaming source position.

diff --git a/Engine/Source/Runtime/Engine/Private/WorldPartition/WorldPartitionSubsystem.cpp b/Engine/Source/Runtime/Engine/Private/WorldPartition/WorldPartitionSubsystem.cpp
--- a/Engine/Source/Runtime/Engine/Private/WorldPartition/WorldPartitionSubsystem.cpp
+++ b/Engine/Source/Runtime/Engine/Private/WorldPartition/WorldPartitionSubsystem.cpp
@@ -1530,7 +1530,7 @@ float FStreamingSourceVelocity::GetAverageVelocity(const FVector& NewPosition, c
 	}
 
 	// Compute velocity (m/s)
-	check(Distance < MAX_flt);
+	ensureMsgf(Distance < MAX_flt, TEXT("Invalid distance %lf computed using position %s and position %s"), Distance, *NewPosition.ToString(), *LastPosition.ToString());
 	const float Velocity = (float)Distance / DeltaSeconds;
 	// Update velocity history buffer and sum
 	LastIndex = (LastIndex + 1) % VELOCITY_HISTORY_SAMPLE_COUNT;
