WorldPartition: Fixed regression on server side when cell unloading is postponed and wp.Runtime.UpdateStreaming.EnableOptimization is enabled.

diff --git a/Engine/Source/Runtime/Engine/Private/WorldPartition/WorldPartitionStreamingPolicy.cpp b/Engine/Source/Runtime/Engine/Private/WorldPartition/WorldPartitionStreamingPolicy.cpp
--- a/Engine/Source/Runtime/Engine/Private/WorldPartition/WorldPartitionStreamingPolicy.cpp
+++ b/Engine/Source/Runtime/Engine/Private/WorldPartition/WorldPartitionStreamingPolicy.cpp
@@ -307,7 +307,7 @@ void UWorldPartitionStreamingPolicy::UpdateStreamingState()
 	// Update new streaming sources hash
 	UpdateStreamingHash = NewUpdateStreamingHash;
 
-	bool bUpdateEpoch = false;
+	bool bUpdateServerEpoch = false;
 
 	check(FrameActivateCells.IsEmpty());
 	check(FrameLoadCells.IsEmpty());
@@ -372,7 +372,7 @@ void UWorldPartitionStreamingPolicy::UpdateStreamingState()
 				return; 
 			}
 
-			bUpdateEpoch = true;
+			bUpdateServerEpoch = true;
 
 			const UDataLayerManager* DataLayerManager = WorldPartition->GetDataLayerManager();
 			TSet<FName> EffectiveActiveDataLayerNames = DataLayerManager->GetEffectiveActiveDataLayerNames();
@@ -452,15 +452,15 @@ void UWorldPartitionStreamingPolicy::UpdateStreamingState()
 	}
 
 	const TSet<FName>& ServerClientsVisibleLevelNames = GetWorld()->GetSubsystem<UWorldPartitionSubsystem>()->ServerClientsVisibleLevelNames;
-	auto ShouldWaitForClientVisibility = [bIsServer, this, &bUpdateEpoch, &ServerClientsVisibleLevelNames](const UWorldPartitionRuntimeCell* Cell)
+	auto ShouldWaitForClientVisibility = [bIsServer, this, &bUpdateServerEpoch, &ServerClientsVisibleLevelNames](const UWorldPartitionRuntimeCell* Cell)
 	{
 		check(bIsServer);
 		if (ULevel* Level = Cell->GetLevel())
 		{
 			if (ServerClientsVisibleLevelNames.Contains(Cell->GetLevel()->GetPackage()->GetFName()))
 			{
-				UE_CLOG(bUpdateEpoch, LogWorldPartition, Verbose, TEXT("Server epoch update delayed by client visibility"));
-				bUpdateEpoch = false;
+				UE_CLOG(bUpdateServerEpoch, LogWorldPartition, Verbose, TEXT("Server epoch update delayed by client visibility"));
+				bUpdateServerEpoch = false;
 				return true;
 			}
 		}
@@ -592,12 +592,21 @@ void UWorldPartitionStreamingPolicy::UpdateStreamingState()
 	UpdateStreamingPerformance(FrameActivateCells);
 	
 	// Update Epoch if we aren't waiting for clients anymore
-	if (bUpdateEpoch)
+	if (bIsServer)
 	{
-		ServerDataLayersStatesEpoch = AWorldDataLayers::GetDataLayersStateEpoch();
-		ServerStreamingStateEpoch = WorldPartition->GetStreamingStateEpoch();
-		ServerStreamingEnabledEpoch = NewServerStreamingEnabledEpoch;
-		UE_LOG(LogWorldPartition, Verbose, TEXT("Server epoch updated"));
+		if (bUpdateServerEpoch)
+		{
+			ServerDataLayersStatesEpoch = AWorldDataLayers::GetDataLayersStateEpoch();
+			ServerStreamingStateEpoch = WorldPartition->GetStreamingStateEpoch();
+			ServerStreamingEnabledEpoch = NewServerStreamingEnabledEpoch;
+			UE_LOG(LogWorldPartition, Verbose, TEXT("Server epoch updated"));
+		}
+		else
+		{
+			// Invalidate UpdateStreamingHash as it was built with latest epochs 
+			// and is now also used to optimize server's update streaming.
+			UpdateStreamingHash = 0;
+		}
 	}
 }
 
