minor cleanup of usage of FObjectPtrProperty.
casting to FObjectPtrProperty is unceseecary in this case as the Cast can avoid resolving
https://p4-swarm.epicgames.net/reviews/26548107

diff --git a/Engine/Source/Runtime/Engine/Private/Actor.cpp b/Engine/Source/Runtime/Engine/Private/Actor.cpp
--- a/Engine/Source/Runtime/Engine/Private/Actor.cpp
+++ b/Engine/Source/Runtime/Engine/Private/Actor.cpp
@@ -817,20 +817,8 @@ void AActor::Serialize(FArchive& Ar)
 				// Ignore transient properties since they won't be serialized
 				if(!ObjProp->HasAnyPropertyFlags(CPF_Transient))
 				{
-					UActorComponent* ActorComponent = nullptr;
-					if (FObjectPtrProperty* ObjPtrProp = CastField<FObjectPtrProperty>(ObjProp))
-					{
-						FObjectPtr& ObjectPtr = (FObjectPtr&)ObjPtrProp->GetPropertyValue_InContainer(this);
-						if (ObjectPtr && ObjectPtr.IsA<UActorComponent>())
-						{
-							ActorComponent = Cast<UActorComponent>(ObjectPtr.Get());
-						}
-					}
-					else
-					{
-						ActorComponent = Cast<UActorComponent>(ObjProp->GetObjectPropertyValue_InContainer(this));
-					}
-
+					const TObjectPtr<UObject>& ObjectPtr = ObjProp->GetPropertyValue_InContainer(this);
+					UActorComponent* ActorComponent = Cast<UActorComponent>(ObjectPtr);
 					if(ActorComponent != nullptr && ActorComponent->CreationMethod == EComponentCreationMethod::Native)
 					{
 						NativeConstructedComponentToPropertyMap.Add(ActorComponent->GetFName(), ObjProp);
