Added cache bucket display names

#rnx

diff --git a/Engine/Plugins/FX/Niagara/Source/Niagara/Niagara.Build.cs b/Engine/Plugins/FX/Niagara/Source/Niagara/Niagara.Build.cs
--- a/Engine/Plugins/FX/Niagara/Source/Niagara/Niagara.Build.cs
+++ b/Engine/Plugins/FX/Niagara/Source/Niagara/Niagara.Build.cs
@@ -59,6 +59,7 @@ public class Niagara : ModuleRules
             PrivateDependencyModuleNames.AddRange(
                 new string[] {
                 "TargetPlatform",
+                "DerivedDataCache",
 				"EditorFramework",
                 "UnrealEd",
 				"SlateCore",
diff --git a/Engine/Plugins/FX/Niagara/Source/Niagara/Private/NiagaraScript.cpp b/Engine/Plugins/FX/Niagara/Source/Niagara/Private/NiagaraScript.cpp
--- a/Engine/Plugins/FX/Niagara/Source/Niagara/Private/NiagaraScript.cpp
+++ b/Engine/Plugins/FX/Niagara/Source/Niagara/Private/NiagaraScript.cpp
@@ -31,6 +31,7 @@
 #include "Interfaces/ITargetPlatform.h"
 #if WITH_EDITOR
 	#include "DerivedDataCacheInterface.h"
+	#include "DerivedDataCacheKey.h"
 	#include "Interfaces/ITargetPlatform.h"
 	#include "NiagaraSettings.h"
 	#include "Internationalization/Regex.h"
@@ -1003,6 +1004,8 @@ FString UNiagaraScript::BuildNiagaraDDCKeyString(const FNiagaraVMExecutableDataI
 	KeyString.AppendChar(TCHAR('_'));
 
 	CompileId.AppendKeyString(KeyString);
+
+	static UE::DerivedData::FCacheBucket LegacyBucket(TEXTVIEW("LegacyNiagaraScriptDerivedData"), TEXTVIEW("NiagaraScript"));
 	return FDerivedDataCacheInterface::BuildCacheKey(TEXT("NiagaraScriptDerivedData"), NIAGARASCRIPT_DERIVEDDATA_VER, *KeyString);
 }
 
diff --git a/Engine/Plugins/FX/Niagara/Source/NiagaraShader/NiagaraShader.Build.cs b/Engine/Plugins/FX/Niagara/Source/NiagaraShader/NiagaraShader.Build.cs
--- a/Engine/Plugins/FX/Niagara/Source/NiagaraShader/NiagaraShader.Build.cs
+++ b/Engine/Plugins/FX/Niagara/Source/NiagaraShader/NiagaraShader.Build.cs
@@ -36,7 +36,8 @@ public class NiagaraShader : ModuleRules
 
             PrivateDependencyModuleNames.AddRange(
                 new string[] {
-            });
+                    "DerivedDataCache",
+                });
         }
 
         PublicIncludePathModuleNames.AddRange(
diff --git a/Engine/Plugins/FX/Niagara/Source/NiagaraShader/Private/NiagaraShader.cpp b/Engine/Plugins/FX/Niagara/Source/NiagaraShader/Private/NiagaraShader.cpp
--- a/Engine/Plugins/FX/Niagara/Source/NiagaraShader/Private/NiagaraShader.cpp
+++ b/Engine/Plugins/FX/Niagara/Source/NiagaraShader/Private/NiagaraShader.cpp
@@ -14,6 +14,7 @@
 #if WITH_EDITOR
 	#include "Interfaces/ITargetPlatformManagerModule.h"
 	#include "DerivedDataCacheInterface.h"
+	#include "DerivedDataCacheKey.h"
 	#include "Interfaces/ITargetPlatformManagerModule.h"
 #endif
 #include "ProfilingDebugging/CookStats.h"
@@ -459,6 +460,7 @@ static FString GetNiagaraShaderMapKeyString(const FNiagaraShaderMapId& ShaderMap
 	NiagaraShaderMapAppendKeyString(Platform, ShaderMapKeyString);
 	ShaderMapAppendKeyString(Platform, ShaderMapKeyString);
 	ShaderMapId.AppendKeyString(ShaderMapKeyString);
+	static UE::DerivedData::FCacheBucket LegacyBucket(TEXTVIEW("LegacyNIAGARASM"), TEXTVIEW("NiagaraShader"));
 	return FDerivedDataCacheInterface::BuildCacheKey(TEXT("NIAGARASM"), *NIAGARASHADERMAP_DERIVEDDATA_VER, *ShaderMapKeyString);
 }
 
diff --git a/Engine/Source/Runtime/Engine/Private/AudioDerivedData.cpp b/Engine/Source/Runtime/Engine/Private/AudioDerivedData.cpp
--- a/Engine/Source/Runtime/Engine/Private/AudioDerivedData.cpp
+++ b/Engine/Source/Runtime/Engine/Private/AudioDerivedData.cpp
@@ -12,6 +12,7 @@
 #include "Async/Async.h"
 #include "SoundWaveCompiler.h"
 #include "DerivedDataCacheInterface.h"
+#include "DerivedDataCacheKey.h"
 #include "ProfilingDebugging/CookStats.h"
 #include "AudioResampler.h"
 #include "AudioCompressionSettingsUtils.h"
@@ -256,6 +257,7 @@ static void GetStreamedAudioDerivedDataKeySuffix(
  */
 static void GetStreamedAudioDerivedDataKeyFromSuffix(const FString& KeySuffix, FString& OutKey)
 {
+	static UE::DerivedData::FCacheBucket LegacyBucket(TEXTVIEW("LegacySTREAMEDAUDIO"), TEXTVIEW("Audio"));
 	OutKey = FDerivedDataCacheInterface::BuildCacheKey(
 		TEXT("STREAMEDAUDIO"),
 		STREAMEDAUDIO_DERIVEDDATA_VER,
diff --git a/Engine/Source/Runtime/Engine/Private/DistanceFieldAtlas.cpp b/Engine/Source/Runtime/Engine/Private/DistanceFieldAtlas.cpp
--- a/Engine/Source/Runtime/Engine/Private/DistanceFieldAtlas.cpp
+++ b/Engine/Source/Runtime/Engine/Private/DistanceFieldAtlas.cpp
@@ -25,6 +25,7 @@
 
 #if WITH_EDITOR
 #include "DerivedDataCacheInterface.h"
+#include "DerivedDataCacheKey.h"
 #include "StaticMeshCompiler.h"
 #endif
 
@@ -137,6 +138,7 @@ FString BuildDistanceFieldDerivedDataKey(const FString& InMeshKey)
 	const float VoxelDensity = CVarDensity->GetValueOnAnyThread();
 	const FString VoxelDensityString = VoxelDensity == .1f ? TEXT("") : FString::Printf(TEXT("_%.3f"), VoxelDensity);
 
+	static UE::DerivedData::FCacheBucket LegacyBucket(TEXTVIEW("LegacyDIST"), TEXTVIEW("DistanceField"));
 	return FDerivedDataCacheInterface::BuildCacheKey(
 		TEXT("DIST"),
 		*FString::Printf(TEXT("%s_%s%s%s"), *InMeshKey, DISTANCEFIELD_DERIVEDDATA_VER, *PerMeshMaxString, *VoxelDensityString),
diff --git a/Engine/Source/Runtime/Engine/Private/Materials/MaterialShader.cpp b/Engine/Source/Runtime/Engine/Private/Materials/MaterialShader.cpp
--- a/Engine/Source/Runtime/Engine/Private/Materials/MaterialShader.cpp
+++ b/Engine/Source/Runtime/Engine/Private/Materials/MaterialShader.cpp
@@ -211,7 +211,7 @@ static FString GetMaterialShaderMapKeyString(const FMaterialShaderMapId& ShaderM
 
 static UE::DerivedData::FCacheKey GetMaterialShaderMapKey(const FStringView MaterialShaderMapKey)
 {
-	static UE::DerivedData::FCacheBucket Bucket("MaterialShaderMap");
+	static UE::DerivedData::FCacheBucket Bucket(ANSITEXTVIEW("MaterialShaderMap"), TEXTVIEW("MaterialShader"));
 	return {Bucket, FIoHash::HashBuffer(MakeMemoryView(FTCHARToUTF8(MaterialShaderMapKey)))};
 }
 
diff --git a/Engine/Source/Runtime/Engine/Private/MeshCardRepresentation.cpp b/Engine/Source/Runtime/Engine/Private/MeshCardRepresentation.cpp
--- a/Engine/Source/Runtime/Engine/Private/MeshCardRepresentation.cpp
+++ b/Engine/Source/Runtime/Engine/Private/MeshCardRepresentation.cpp
@@ -19,6 +19,7 @@
 
 #if WITH_EDITOR
 #include "DerivedDataCacheInterface.h"
+#include "DerivedDataCacheKey.h"
 #include "StaticMeshCompiler.h"
 #endif
 
@@ -106,6 +107,7 @@ FString BuildCardRepresentationDerivedDataKey(const FString& InMeshKey, int32 Ma
 	const float NormalTreshold = MeshCardRepresentation::GetNormalTreshold();
 	const bool bDebugMode = MeshCardRepresentation::IsDebugMode();
 
+	static UE::DerivedData::FCacheBucket LegacyBucket(TEXTVIEW("LegacyCARD"), TEXTVIEW("CardRepresentation"));
 	return FDerivedDataCacheInterface::BuildCacheKey(
 		TEXT("CARD"),
 		*FString::Printf(TEXT("%s_%s%s%.3f_%.3f_%d"), *InMeshKey, CARDREPRESENTATION_DERIVEDDATA_VER, bDebugMode ? TEXT("_DEBUG_") : TEXT(""),
diff --git a/Engine/Source/Runtime/Engine/Private/Particles/SubUVAnimation.cpp b/Engine/Source/Runtime/Engine/Private/Particles/SubUVAnimation.cpp
--- a/Engine/Source/Runtime/Engine/Private/Particles/SubUVAnimation.cpp
+++ b/Engine/Source/Runtime/Engine/Private/Particles/SubUVAnimation.cpp
@@ -9,6 +9,7 @@
 #include "ParticleHelper.h"
 #include "Particles/ParticleSystemComponent.h"
 #include "DerivedDataCacheInterface.h"
+#include "DerivedDataCacheKey.h"
 #include "ComponentReregisterContext.h"
 #include "RHI.h"
 
@@ -34,6 +35,11 @@ FString FSubUVDerivedData::GetDDCKeyString(const FGuid& StateId, int32 SizeX, in
 	// adding v3 to the key after allowing other formats
 	// adding v4 to the key after adding G8 support
 	KeyString += TEXT("_V4");
+
+#if WITH_EDITOR
+	static UE::DerivedData::FCacheBucket LegacyBucket(TEXTVIEW("LegacySUBUV"), TEXTVIEW("SubUV"));
+#endif
+
 	return FDerivedDataCacheInterface::BuildCacheKey(TEXT("SUBUV_"), SUBUV_DERIVEDDATA_VER, *KeyString);
 }
 
diff --git a/Engine/Source/Runtime/Engine/Private/PhysicsEngine/Experimental/ChaosDerivedData.cpp b/Engine/Source/Runtime/Engine/Private/PhysicsEngine/Experimental/ChaosDerivedData.cpp
--- a/Engine/Source/Runtime/Engine/Private/PhysicsEngine/Experimental/ChaosDerivedData.cpp
+++ b/Engine/Source/Runtime/Engine/Private/PhysicsEngine/Experimental/ChaosDerivedData.cpp
@@ -7,6 +7,10 @@
 #include "Serialization/MemoryWriter.h"
 #include "PhysicsEngine/Experimental/ChaosCooking.h"
 
+#if WITH_EDITOR
+#include "DerivedDataCacheKey.h"
+#endif
+
 class FChaosDerivedDataCookerRefHolder : public FGCObject
 {
 public:
@@ -35,6 +39,9 @@ private:
 
 const TCHAR* FChaosDerivedDataCooker::GetPluginName() const
 {
+#if WITH_EDITOR
+	static UE::DerivedData::FCacheBucket LegacyBucket(TEXTVIEW("LegacyChaosGeometryData"), TEXTVIEW("BodySetup"));
+#endif
 	return TEXT("ChaosGeometryData");
 }
 
diff --git a/Engine/Source/Runtime/Engine/Private/ShaderCompiler/ShaderCompiler.cpp b/Engine/Source/Runtime/Engine/Private/ShaderCompiler/ShaderCompiler.cpp
--- a/Engine/Source/Runtime/Engine/Private/ShaderCompiler/ShaderCompiler.cpp
+++ b/Engine/Source/Runtime/Engine/Private/ShaderCompiler/ShaderCompiler.cpp
@@ -8509,7 +8509,7 @@ static FString GetGlobalShaderMapKeyString(const FGlobalShaderMapId& ShaderMapId
 static UE::DerivedData::FCacheKey GetGlobalShaderMapKey(const FGlobalShaderMapId& ShaderMapId, EShaderPlatform Platform, const ITargetPlatform* TargetPlatform, TArray<FShaderTypeDependency> const& Dependencies)
 {
 	const FString DataKey = GetGlobalShaderMapKeyString(ShaderMapId, Platform, Dependencies);
-	static const UE::DerivedData::FCacheBucket Bucket("GlobalShaderMap");
+	static const UE::DerivedData::FCacheBucket Bucket(ANSITEXTVIEW("GlobalShaderMap"), TEXTVIEW("GlobalShader"));
 	return {Bucket, FIoHash::HashBuffer(MakeMemoryView(FTCHARToUTF8(DataKey)))};
 }
 
@@ -9606,7 +9606,7 @@ TRACE_DECLARE_MEMORY_COUNTER(Shaders_JobCacheDDCBytesSent, TEXT("Shaders/JobCach
 namespace
 {
 	/** The FCacheBucket used with the DDC, cached to avoid recreating it for each request */
-	UE::DerivedData::FCacheBucket ShaderJobCacheDDCBucket = UE::DerivedData::FCacheBucket(TEXT("FShaderJobCacheShaders"));
+	UE::DerivedData::FCacheBucket ShaderJobCacheDDCBucket = UE::DerivedData::FCacheBucket(ANSITEXTVIEW("FShaderJobCacheShaders"), TEXTVIEW("Shader"));
 	UE::DerivedData::FValueId ShaderJobCacheId = UE::DerivedData::FValueId::FromName("FShaderJobCacheShaderID");
 }
 #endif
diff --git a/Engine/Source/Runtime/Engine/Private/SkeletalMesh.cpp b/Engine/Source/Runtime/Engine/Private/SkeletalMesh.cpp
--- a/Engine/Source/Runtime/Engine/Private/SkeletalMesh.cpp
+++ b/Engine/Source/Runtime/Engine/Private/SkeletalMesh.cpp
@@ -62,6 +62,7 @@
 #include "MeshUtilities.h"
 #include "Engine/SkeletalMeshEditorData.h"
 #include "DerivedDataCacheInterface.h"
+#include "DerivedDataCacheKey.h"
 #include "Engine/RendererSettings.h"
 #include "SkeletalDebugRendering.h"
 #include "Misc/DataValidation.h"
@@ -4464,6 +4465,7 @@ FString USkeletalMesh::BuildDerivedDataKey(const ITargetPlatform* TargetPlatform
 	KeySuffix.Append(TEXT("_arm64"));
 #endif
 
+	static UE::DerivedData::FCacheBucket LegacyBucket(TEXTVIEW("LegacySKELETALMESH"), TEXTVIEW("SkeletalMesh"));
 	return FDerivedDataCacheInterface::BuildCacheKey(
 		TEXT("SKELETALMESH"),
 		*GetSkeletalMeshDerivedDataVersion(),
diff --git a/Engine/Source/Runtime/Engine/Private/TextureDerivedData.cpp b/Engine/Source/Runtime/Engine/Private/TextureDerivedData.cpp
--- a/Engine/Source/Runtime/Engine/Private/TextureDerivedData.cpp
+++ b/Engine/Source/Runtime/Engine/Private/TextureDerivedData.cpp
@@ -487,6 +487,7 @@ const FGuid& GetTextureDerivedDataVersion()
  */
 void GetTextureDerivedDataKeyFromSuffix(const FString& KeySuffix, FString& OutKey)
 {
+	static UE::DerivedData::FCacheBucket LegacyBucket(TEXTVIEW("LegacyTEXTURE"), TEXTVIEW("Texture"));
 	OutKey = FDerivedDataCacheInterface::BuildCacheKey(
 		TEXT("TEXTURE"),
 		TEXTURE_DERIVEDDATA_VER,
