Change default Local Exposure settings (contrast scale set to 0.8).
- Enabled Local Exposure in new projects via project "Default Settings".
- Modified default Local Exposure values in PPV.

In order to keep PPV in existing projects unmodified, this change adds a new BranchObjectVersion so that when loading assets saved in previous engine version, the default values are set accordingly before deserialization happens.

diff --git a/Engine/Source/Editor/GameProjectGeneration/Private/GameProjectUtils.cpp b/Engine/Source/Editor/GameProjectGeneration/Private/GameProjectUtils.cpp
--- a/Engine/Source/Editor/GameProjectGeneration/Private/GameProjectUtils.cpp
+++ b/Engine/Source/Editor/GameProjectGeneration/Private/GameProjectUtils.cpp
@@ -138,6 +138,20 @@ namespace
 			TEXT("r.DefaultFeature.AutoExposure.ExtendDefaultLuminanceRange"),
 			TEXT("True"),
 			false /* ShouldReplaceExistingValue */);
+
+		// Enable Local Exposure by default for new projects
+		ConfigValues.Emplace(TEXT("DefaultEngine.ini"),
+			TEXT("/Script/Engine.RendererSettings"),
+			TEXT("r.DefaultFeature.LocalExposure.HighlightContrastScale"),
+			TEXT("0.8"),
+			true /* ShouldReplaceExistingValue */);
+
+		// Enable Local Exposure by default for new projects
+		ConfigValues.Emplace(TEXT("DefaultEngine.ini"),
+			TEXT("/Script/Engine.RendererSettings"),
+			TEXT("r.DefaultFeature.LocalExposure.ShadowContrastScale"),
+			TEXT("0.8"),
+			true /* ShouldReplaceExistingValue */);
 	}
 
 	/** Get the configuration values for raytracing if enabled. */
diff --git a/Engine/Source/Runtime/Core/Public/UObject/UE5ReleaseStreamObjectVersion.h b/Engine/Source/Runtime/Core/Public/UObject/UE5ReleaseStreamObjectVersion.h
--- a/Engine/Source/Runtime/Core/Public/UObject/UE5ReleaseStreamObjectVersion.h
+++ b/Engine/Source/Runtime/Core/Public/UObject/UE5ReleaseStreamObjectVersion.h
@@ -144,6 +144,9 @@ struct FUE5ReleaseStreamObjectVersion
 
 		// Function data stores a map from work to debug operands
 		RigVMSaveDebugMapInGraphFunctionData,
+
+		// Changed default Local Exposure Contrast Scale from 1.0 to 0.8
+		LocalExposureDefaultChangeFrom1,
 		
 		// -----<new versions can be added above this line>-------------------------------------------------
 		VersionPlusOne,
diff --git a/Engine/Source/Runtime/Engine/Classes/Engine/RendererSettings.h b/Engine/Source/Runtime/Engine/Classes/Engine/RendererSettings.h
--- a/Engine/Source/Runtime/Engine/Classes/Engine/RendererSettings.h
+++ b/Engine/Source/Runtime/Engine/Classes/Engine/RendererSettings.h
@@ -666,6 +666,16 @@ class URendererSettings : public UDeveloperSettings
 		ConfigRestartRequired=true))
 	uint32 bExtendDefaultLuminanceRangeInAutoExposureSettings : 1;
 
+	UPROPERTY(config, EditAnywhere, Category = DefaultSettings, meta = (
+		ConsoleVariable = "r.DefaultFeature.LocalExposure.HighlightContrastScale", DisplayName = "Local Exposure Highlight Contrast Scale",
+		ToolTip = "Default Value for Local Exposure Highlight Contrast Scale.", ClampMin = "0.0", ClampMax = "1.0"))
+	float DefaultFeatureLocalExposureHighlightContrastScale;
+
+	UPROPERTY(config, EditAnywhere, Category = DefaultSettings, meta = (
+		ConsoleVariable = "r.DefaultFeature.LocalExposure.ShadowContrastScale", DisplayName = "Local Exposure Shadow Contrast Scale",
+		ToolTip = "Default Value for Local Exposure Shadow Contrast Scale.", ClampMin = "0.0", ClampMax = "1.0"))
+	float DefaultFeatureLocalExposureShadowContrastScale;
+
 	UPROPERTY(config, EditAnywhere, Category = DefaultSettings, meta = (
 		ConsoleVariable = "r.DefaultFeature.MotionBlur", DisplayName = "Motion Blur",
 		ToolTip = "Whether the default for MotionBlur is enabled or not (postprocess volume/camera/game setting can still override and enable or disable it independently)"))
diff --git a/Engine/Source/Runtime/Engine/Private/Scene.cpp b/Engine/Source/Runtime/Engine/Private/Scene.cpp
--- a/Engine/Source/Runtime/Engine/Private/Scene.cpp
+++ b/Engine/Source/Runtime/Engine/Private/Scene.cpp
@@ -509,8 +509,8 @@ FPostProcessSettings::FPostProcessSettings()
 	AutoExposureApplyPhysicalCameraExposure = 1;
 
 	LocalExposureContrastScale_DEPRECATED = 1.0f;
-	LocalExposureHighlightContrastScale = 1.0f;
-	LocalExposureShadowContrastScale = 1.0f;
+	LocalExposureHighlightContrastScale = 0.8f;
+	LocalExposureShadowContrastScale = 0.8f;
 	LocalExposureDetailStrength = 1.0f;
 	LocalExposureBlurredLuminanceBlend = 0.6f;
 	LocalExposureBlurredLuminanceKernelSizePercent = 50.0f;
@@ -666,6 +666,12 @@ bool FPostProcessSettings::Serialize(FArchive& Ar)
 	Ar.UsingCustomVersion(FReleaseObjectVersion::GUID);
 	Ar.UsingCustomVersion(FUE5ReleaseStreamObjectVersion::GUID);
 
+	if (Ar.CustomVer(FUE5ReleaseStreamObjectVersion::GUID) < FUE5ReleaseStreamObjectVersion::LocalExposureDefaultChangeFrom1)
+	{
+		LocalExposureHighlightContrastScale = 1.0f;
+		LocalExposureShadowContrastScale = 1.0f;
+	}
+
 	// Don't actually serialize, just write the custom version for PostSerialize
 	return false;
 }
diff --git a/Engine/Source/Runtime/Engine/Private/SceneView.cpp b/Engine/Source/Runtime/Engine/Private/SceneView.cpp
--- a/Engine/Source/Runtime/Engine/Private/SceneView.cpp
+++ b/Engine/Source/Runtime/Engine/Private/SceneView.cpp
@@ -182,6 +182,16 @@ static TAutoConsoleVariable<int32> CVarDefaultAutoExposureExtendDefaultLuminance
 	TEXT(" 0: Legacy range (UE4 default)\n")
 	TEXT(" 1: Extended range (UE5 default)"));
 
+static TAutoConsoleVariable<float> CVarDefaultLocalExposureHighlightContrastScale(
+	TEXT("r.DefaultFeature.LocalExposure.HighlightContrastScale"),
+	1.0f,
+	TEXT("Engine default (project setting) for Local Exposure Highlight Contrast Scale (postprocess volume/camera/game setting still can override)\n"));
+
+static TAutoConsoleVariable<float> CVarDefaultLocalExposureShadowContrastScale(
+	TEXT("r.DefaultFeature.LocalExposure.ShadowContrastScale"),
+	1.0f,
+	TEXT("Engine default (project setting) for Local Exposure Shadow Contrast Scale (postprocess volume/camera/game setting still can override)\n"));
+
 static TAutoConsoleVariable<int32> CVarDefaultMotionBlur(
 	TEXT("r.DefaultFeature.MotionBlur"),
 	1,
@@ -1830,6 +1840,18 @@ void FSceneView::StartFinalPostprocessSettings(FVector InViewLocation)
 			}
 		}
 
+		{
+			const float HighlightContrastScale = FMath::Clamp(CVarDefaultLocalExposureHighlightContrastScale.GetValueOnGameThread(), 0.0f, 1.0f);
+
+			FinalPostProcessSettings.LocalExposureHighlightContrastScale = HighlightContrastScale;
+		}
+
+		{
+			const float ShadowContrastScale = FMath::Clamp(CVarDefaultLocalExposureHighlightContrastScale.GetValueOnGameThread(), 0.0f, 1.0f);
+
+			FinalPostProcessSettings.LocalExposureShadowContrastScale = ShadowContrastScale;
+		}
+
 		if (!CVarDefaultMotionBlur.GetValueOnGameThread())
 		{
 			FinalPostProcessSettings.MotionBlurAmount = 0;
