Replace LWI HISMs with ISMs

diff --git a/Engine/Plugins/Runtime/MassGameplay/Source/MassLWI/Private/MassLWIStaticMeshManager.cpp b/Engine/Plugins/Runtime/MassGameplay/Source/MassLWI/Private/MassLWIStaticMeshManager.cpp
--- a/Engine/Plugins/Runtime/MassGameplay/Source/MassLWI/Private/MassLWIStaticMeshManager.cpp
+++ b/Engine/Plugins/Runtime/MassGameplay/Source/MassLWI/Private/MassLWIStaticMeshManager.cpp
@@ -4,7 +4,7 @@
 #include "MassLWITypes.h"
 #include "MassLWISubsystem.h"
 #include "MassLWITypes.h"
-#include "Components/HierarchicalInstancedStaticMeshComponent.h"
+#include "Components/InstancedStaticMeshComponent.h"
 #include "MassCommonFragments.h"
 #include "MassEntitySubsystem.h"
 #include "MassEntityView.h"
@@ -84,7 +84,7 @@ void AMassLWIStaticMeshManager::TransferDataToMass(FMassEntityManager& EntityMan
 		CreateMassTemplate(EntityManager);
 	}
 
-	if (MassTemplateID.IsValid() && InstancedStaticMeshComponent && FinalizedTemplate)
+	if (MassTemplateID.IsValid() && ISMComponent && FinalizedTemplate)
 	{
 		const FVector& ManagerLocation = GetActorLocation();
 
@@ -128,8 +128,7 @@ void AMassLWIStaticMeshManager::TransferDataToMass(FMassEntityManager& EntityMan
 			DataIndicesToRenderingIndices[Index] = Index;
 		}
 		
-		InstancedStaticMeshComponent->ClearInstances();
-		InstancedStaticMeshComponent->BuildTreeIfOutdated(false, true);
+		ISMComponent->ClearInstances();
 	}
 }
 
@@ -143,7 +142,7 @@ void AMassLWIStaticMeshManager::StoreMassDataInActor(FMassEntityManager& EntityM
 	FMassEntityQuery LocationQuery;
 	LocationQuery.AddRequirement<FTransformFragment>(EMassFragmentAccess::ReadOnly);
 
-	const FTransform WorldTransform = InstancedStaticMeshComponent ? InstancedStaticMeshComponent->GetComponentTransform() : GetActorTransform();
+	const FTransform WorldTransform = ISMComponent ? ISMComponent->GetComponentTransform() : GetActorTransform();
 	
 	FMassExecutionContext ExecutionContext(EntityManager);
 	for (FMassArchetypeEntityCollection& Collection : EntityCollectionsToDestroy)
@@ -264,10 +263,10 @@ void AMassLWIStaticMeshManager::CreateMassTemplate(FMassEntityManager& EntityMan
 		}
 	}
 
-	if (InstancedStaticMeshComponent)
+	if (ISMComponent)
 	{
-		MeshDesc.MaterialOverrides = InstancedStaticMeshComponent->OverrideMaterials;
-		MeshDesc.bCastShadows = (InstancedStaticMeshComponent->CastShadow != 0);
+		MeshDesc.MaterialOverrides = ISMComponent->OverrideMaterials;
+		MeshDesc.bCastShadows = (ISMComponent->CastShadow != 0);
 	}
 
 	check(MeshDesc.Mesh);
@@ -278,7 +277,7 @@ void AMassLWIStaticMeshManager::CreateMassTemplate(FMassEntityManager& EntityMan
 		FMassRepresentationFragment& RepresentationFragment = NewTemplate.AddFragment_GetRef<FMassRepresentationFragment>();
 		if (UE::Mass::Tweakables::bReuseLWIISMComponents)
 		{
-			RepresentationFragment.StaticMeshDescIndex = RepresentationSubsystem->AddVisualDescWithISMComponent(StaticMeshInstanceDesc, *InstancedStaticMeshComponent);
+			RepresentationFragment.StaticMeshDescIndex = RepresentationSubsystem->AddVisualDescWithISMComponent(StaticMeshInstanceDesc, *ISMComponent);
 		}
 		else
 		{
diff --git a/Engine/Source/Runtime/Engine/Classes/Components/HierarchicalInstancedStaticMeshComponent.h b/Engine/Source/Runtime/Engine/Classes/Components/HierarchicalInstancedStaticMeshComponent.h
--- a/Engine/Source/Runtime/Engine/Classes/Components/HierarchicalInstancedStaticMeshComponent.h
+++ b/Engine/Source/Runtime/Engine/Classes/Components/HierarchicalInstancedStaticMeshComponent.h
@@ -191,9 +191,6 @@ class UHierarchicalInstancedStaticMeshComponent : public UInstancedStaticMeshCom
 	UPROPERTY()
 	FBoxSphereBounds CacheMeshExtendedBounds;
 
-	UPROPERTY()
-	bool bDisableCollision;
-
 	// Instances to render (including removed one until the build is complete)
 	UPROPERTY()
 	int32 InstanceCountToRender;
@@ -255,8 +252,6 @@ public:
 	/** Get the transforms of instances inside the provided box */
 	ENGINE_API void GetOverlappingBoxTransforms(const FBox& Box, TArray<FTransform>& OutTransforms) const;
 
-	ENGINE_API virtual bool ShouldCreatePhysicsState() const override;
-
 	ENGINE_API bool BuildTreeIfOutdated(bool Async, bool ForceUpdate);
 	UE_DEPRECATED(5.1, "The BuildTreeAnyThread method is moving to UGrassInstancedStaticMeshComponent. Please update your project to use the new component and method or your project may not compile in the next udpate.")
 	static ENGINE_API void BuildTreeAnyThread(TArray<FMatrix>& InstanceTransforms, TArray<float>& InstanceCustomDataFloats, int32 NumCustomDataFloats, const FBox& MeshBox, TArray<FClusterNode>& OutClusterTree, TArray<int32>& OutSortedInstances, TArray<int32>& OutInstanceReorderTable, int32& OutOcclusionLayerNum, int32 MaxInstancesPerLeaf, bool InGenerateInstanceScalingRange);
diff --git a/Engine/Source/Runtime/Engine/Classes/Components/InstancedStaticMeshComponent.h b/Engine/Source/Runtime/Engine/Classes/Components/InstancedStaticMeshComponent.h
--- a/Engine/Source/Runtime/Engine/Classes/Components/InstancedStaticMeshComponent.h
+++ b/Engine/Source/Runtime/Engine/Classes/Components/InstancedStaticMeshComponent.h
@@ -88,6 +88,7 @@ class UInstancedStaticMeshComponent : public UStaticMeshComponent, public ISMIns
 {
 	GENERATED_UCLASS_BODY()
 
+	friend class ALightWeightInstanceStaticMeshManager;
 	friend class ALightWeightInstanceManager;
 	
 	/** Needs implementation in InstancedStaticMesh.cpp to compile UniquePtr for forward declared class */
@@ -146,6 +147,10 @@ class UInstancedStaticMeshComponent : public UStaticMeshComponent, public ISMIns
 	/** Mapping from PerInstanceSMData order to instance render buffer order. If empty, the PerInstanceSMData order is used. */
 	UPROPERTY()
 	TArray<int32> InstanceReorderTable;
+	
+	/** Don't create any collision when this bool is set */
+	UPROPERTY()
+	bool bDisableCollision;
 
 	/** Tracks outstanding proxysize, as this is a bit hard to do with the fire-and-forget grass. */
 	SIZE_T ProxySize;
diff --git a/Engine/Source/Runtime/Engine/Classes/GameFramework/LightWeightInstanceStaticMeshManager.h b/Engine/Source/Runtime/Engine/Classes/GameFramework/LightWeightInstanceStaticMeshManager.h
--- a/Engine/Source/Runtime/Engine/Classes/GameFramework/LightWeightInstanceStaticMeshManager.h
+++ b/Engine/Source/Runtime/Engine/Classes/GameFramework/LightWeightInstanceStaticMeshManager.h
@@ -51,6 +51,10 @@ public:
 
 	ENGINE_API virtual void GetLifetimeReplicatedProps(TArray<FLifetimeProperty>& OutLifetimeProps) const override;
 
+#if WITH_EDITORONLY_DATA
+	ENGINE_API virtual void PostLoad() override;
+#endif
+
 protected:
 
 	//~ ISMInstanceManager interface
@@ -80,8 +84,16 @@ protected:
 	UFUNCTION()
 	ENGINE_API void OnRep_StaticMesh();
 
+#if WITH_EDITORONLY_DATA
+	PRAGMA_DISABLE_DEPRECATION_WARNINGS
+	UE_DEPRECATED(5.4, "Doesn't need a UHierarchicalInstancedStaticMeshComponent anymore but is replaced by UInstancedStaticMeshComponent (see ISMComponent property).")
+	UPROPERTY()
+	TObjectPtr<class UHierarchicalInstancedStaticMeshComponent> InstancedStaticMeshComponent_DEPRECATED;
+	PRAGMA_ENABLE_DEPRECATION_WARNINGS
+#endif
+
 	UPROPERTY(EditAnywhere, BlueprintReadOnly, Category = Debug, AdvancedDisplay, meta = (BlueprintProtected = "true", AllowPrivateAccess = "true"))
-	TObjectPtr<class UHierarchicalInstancedStaticMeshComponent> InstancedStaticMeshComponent;
+	TObjectPtr<class UInstancedStaticMeshComponent> ISMComponent;
 
 	//
 	// Bookkeeping info
diff --git a/Engine/Source/Runtime/Engine/Private/HierarchicalInstancedStaticMesh.cpp b/Engine/Source/Runtime/Engine/Private/HierarchicalInstancedStaticMesh.cpp
--- a/Engine/Source/Runtime/Engine/Private/HierarchicalInstancedStaticMesh.cpp
+++ b/Engine/Source/Runtime/Engine/Private/HierarchicalInstancedStaticMesh.cpp
@@ -2726,15 +2726,6 @@ void UHierarchicalInstancedStaticMeshComponent::ClearInstances()
 	FNavigationSystem::UpdateComponentData(*this);
 }
 
-bool UHierarchicalInstancedStaticMeshComponent::ShouldCreatePhysicsState() const
-{
-	if (bDisableCollision)
-	{
-		return false;
-	}
-	return Super::ShouldCreatePhysicsState();
-}
-
 int32 UHierarchicalInstancedStaticMeshComponent::GetVertsForLOD(int32 LODIndex)
 {
 	if (GetStaticMesh() && GetStaticMesh()->HasValidRenderData(true, LODIndex))
diff --git a/Engine/Source/Runtime/Engine/Private/InstancedStaticMesh.cpp b/Engine/Source/Runtime/Engine/Private/InstancedStaticMesh.cpp
--- a/Engine/Source/Runtime/Engine/Private/InstancedStaticMesh.cpp
+++ b/Engine/Source/Runtime/Engine/Private/InstancedStaticMesh.cpp
@@ -3774,12 +3774,23 @@ bool UInstancedStaticMeshComponent::RemoveInstanceInternal(int32 InstanceIndex,
 				InstanceBody = nullptr;
 			}
 
-			InstanceBodies.RemoveAt(InstanceIndex);
+			if (bUseRemoveAtSwap && InstanceIndex != LastInstanceIndex)
+			{
+				// swap in the last instance body if we have one
+				InstanceBodies.RemoveAtSwap(InstanceIndex);
 
-			// Re-target instance indices for shifting of array.
-			for (int32 i = InstanceIndex; i < InstanceBodies.Num(); ++i)
+				// Update the Instance body index to the new swapped location
+				InstanceBodies[InstanceIndex]->InstanceBodyIndex = InstanceIndex;
+			}
+			else
 			{
-				InstanceBodies[i]->InstanceBodyIndex = i;
+				InstanceBodies.RemoveAt(InstanceIndex);
+
+				// Re-target instance indices for shifting of array.
+				for (int32 i = InstanceIndex; i < InstanceBodies.Num(); ++i)
+				{
+					InstanceBodies[i]->InstanceBodyIndex = i;
+				}
 			}
 		}
 	}
@@ -4481,7 +4492,7 @@ TArray<int32> UInstancedStaticMeshComponent::GetInstancesOverlappingBox(const FB
 
 bool UInstancedStaticMeshComponent::ShouldCreatePhysicsState() const
 {
-	return IsRegistered() && !IsBeingDestroyed() && GetStaticMesh() && !GetStaticMesh()->IsCompiling() && (bAlwaysCreatePhysicsState || IsCollisionEnabled());
+	return !bDisableCollision && IsRegistered() && !IsBeingDestroyed() && GetStaticMesh() && !GetStaticMesh()->IsCompiling() && (bAlwaysCreatePhysicsState || IsCollisionEnabled());
 }
 
 float UInstancedStaticMeshComponent::GetTextureStreamingTransformScale() const
diff --git a/Engine/Source/Runtime/Engine/Private/LightWeightInstanceStaticMeshManager.cpp b/Engine/Source/Runtime/Engine/Private/LightWeightInstanceStaticMeshManager.cpp
--- a/Engine/Source/Runtime/Engine/Private/LightWeightInstanceStaticMeshManager.cpp
+++ b/Engine/Source/Runtime/Engine/Private/LightWeightInstanceStaticMeshManager.cpp
@@ -21,9 +21,17 @@
 ALightWeightInstanceStaticMeshManager::ALightWeightInstanceStaticMeshManager(const FObjectInitializer& ObjectInitializer)
 	:Super(ObjectInitializer)
 {
-	InstancedStaticMeshComponent = CreateDefaultSubobject<UHierarchicalInstancedStaticMeshComponent>(TEXT("InstancedStaticMeshComponent0"));
-	SetRootComponent(InstancedStaticMeshComponent);
-	AddInstanceComponent(InstancedStaticMeshComponent);
+#if WITH_EDITORONLY_DATA
+PRAGMA_DISABLE_DEPRECATION_WARNINGS
+	// Still create the deprecated InstancedStaticMeshComponent so the old data can be serialized during loading
+	InstancedStaticMeshComponent_DEPRECATED = CreateDefaultSubobject<UHierarchicalInstancedStaticMeshComponent>(TEXT("InstancedStaticMeshComponent0"));
+PRAGMA_ENABLE_DEPRECATION_WARNINGS
+#endif
+
+	ISMComponent = CreateDefaultSubobject<UInstancedStaticMeshComponent>(TEXT("ISMComponent0"));
+	ISMComponent->SetRemoveSwap();
+	SetRootComponent(ISMComponent);
+	AddInstanceComponent(ISMComponent);
 
 	if (StaticMesh.IsValid())
 	{
@@ -33,6 +41,57 @@ ALightWeightInstanceStaticMeshManager::ALightWeightInstanceStaticMeshManager(con
 	SetInstancedStaticMeshParams();
 }
 
+#if WITH_EDITORONLY_DATA
+void ALightWeightInstanceStaticMeshManager::PostLoad()
+{
+	Super::PostLoad();
+
+PRAGMA_DISABLE_DEPRECATION_WARNINGS
+	if (InstancedStaticMeshComponent_DEPRECATED != nullptr && InstancedStaticMeshComponent_DEPRECATED->PerInstanceSMData.Num() > 0)
+	{
+		// assume new ISM component is still empty
+		check(ISMComponent->PerInstanceSMData.Num() == 0);
+
+		// copy over serialized instance data to ISM from HISM
+		ISMComponent->SetRelativeTransform(InstancedStaticMeshComponent_DEPRECATED->GetRelativeTransform());
+		ISMComponent->SetMobility(InstancedStaticMeshComponent_DEPRECATED->Mobility);
+		ISMComponent->SetStaticMesh(InstancedStaticMeshComponent_DEPRECATED->GetStaticMesh());
+		ISMComponent->InstancingRandomSeed = InstancedStaticMeshComponent_DEPRECATED->InstancingRandomSeed;
+
+		int32 StartCullDistance, EndCullDistance;
+		InstancedStaticMeshComponent_DEPRECATED->GetCullDistances(StartCullDistance, EndCullDistance);
+		ISMComponent->SetCullDistances(StartCullDistance, EndCullDistance);
+
+		for (int32 Idx = 0; Idx < InstancedStaticMeshComponent_DEPRECATED->GetNumOverrideMaterials(); ++Idx)
+		{
+			ISMComponent->SetMaterial(Idx, InstancedStaticMeshComponent_DEPRECATED->GetMaterial(Idx));
+		}
+		for (FTransform& InstanceTransform : InstanceTransforms)
+		{
+			ISMComponent->AddInstance(InstanceTransform, /*bWorldSpace*/false);
+		}
+
+		ISMComponent->OnPostLoadPerInstanceData();
+
+		// release all instances in the HISM
+		InstancedStaticMeshComponent_DEPRECATED->ClearInstances();
+
+		// restore the root component after destroying the old root component)
+		SetRootComponent(ISMComponent);
+
+		// all instances have been added in order of the transforms so reset the indirection lookup tables
+		DataIndicesToRenderingIndices.SetNum(InstanceTransforms.Num());
+		RenderingIndicesToDataIndices.SetNum(InstanceTransforms.Num());
+		for (int32 DataIndex = 0; DataIndex < InstanceTransforms.Num(); ++DataIndex)
+		{
+			DataIndicesToRenderingIndices[DataIndex] = DataIndex;
+			RenderingIndicesToDataIndices[DataIndex] = DataIndex;
+		}
+	}
+PRAGMA_ENABLE_DEPRECATION_WARNINGS
+}
+#endif
+
 void ALightWeightInstanceStaticMeshManager::SetRepresentedClass(UClass* ActorClass)
 {
 	Super::SetRepresentedClass(ActorClass);
@@ -49,9 +108,9 @@ void ALightWeightInstanceStaticMeshManager::SetRepresentedClass(UClass* ActorCla
 		ClearStaticMesh();
 	}
 
-	if (InstancedStaticMeshComponent)
+	if (ISMComponent)
 	{
-		InstancedStaticMeshComponent->OnPostLoadPerInstanceData();
+		ISMComponent->OnPostLoadPerInstanceData();
 	}
 }
 
@@ -121,10 +180,10 @@ void ALightWeightInstanceStaticMeshManager::AddInstanceToRendering(int32 DataInd
 		DataIndicesToRenderingIndices[DataIndex] = RenderingIdx;
 	}
 
-	// Update the HISMC
-	if (InstancedStaticMeshComponent)
+	// Update the ISMC
+	if (ISMComponent)
 	{
-		ensure(RenderingIdx == InstancedStaticMeshComponent->AddInstance(InstanceTransforms[DataIndex], /*bWorldSpace*/false));
+		ensure(RenderingIdx == ISMComponent->AddInstance(InstanceTransforms[DataIndex], /*bWorldSpace*/false));
 	}
 }
 
@@ -177,9 +236,9 @@ void ALightWeightInstanceStaticMeshManager::PostRemoveInstanceFromRendering()
 		int32 RenderingIndex = RenderingIndicesToBeDeleted[Idx];
 		if (RenderingIndex != INDEX_NONE)
 		{
-			if (InstancedStaticMeshComponent)
+			if (ISMComponent)
 			{
-				InstancedStaticMeshComponent->RemoveInstance(RenderingIndex);
+				ISMComponent->RemoveInstance(RenderingIndex);
 			}
 			RenderingIndicesToDataIndices.RemoveAtSwap(RenderingIndex);
 
@@ -216,9 +275,9 @@ void ALightWeightInstanceStaticMeshManager::OnRep_Transforms()
 {
 	Super::OnRep_Transforms();
 
-	if (InstancedStaticMeshComponent)
+	if (ISMComponent)
 	{
-		InstancedStaticMeshComponent->AddInstance(InstanceTransforms.Last(), /*bWorldSpace*/false);
+		ISMComponent->AddInstance(InstanceTransforms.Last(), /*bWorldSpace*/false);
 	}
 }
 
@@ -233,21 +292,21 @@ void ALightWeightInstanceStaticMeshManager::PostActorSpawn(const FActorInstanceH
 void ALightWeightInstanceStaticMeshManager::SetInstancedStaticMeshParams()
 {
 	FName CollisionProfileName(TEXT("LightWeightInstancedStaticMeshPhysics"));
-	InstancedStaticMeshComponent->SetCollisionProfileName(CollisionProfileName);
+	ISMComponent->SetCollisionProfileName(CollisionProfileName);
 
-	InstancedStaticMeshComponent->CanCharacterStepUpOn = ECB_Owner;
-	InstancedStaticMeshComponent->CastShadow = true;
-	InstancedStaticMeshComponent->bCastDynamicShadow = true;
-	InstancedStaticMeshComponent->bCastStaticShadow = true;
-	InstancedStaticMeshComponent->PrimaryComponentTick.bCanEverTick = false;
+	ISMComponent->CanCharacterStepUpOn = ECB_Owner;
+	ISMComponent->CastShadow = true;
+	ISMComponent->bCastDynamicShadow = true;
+	ISMComponent->bCastStaticShadow = true;
+	ISMComponent->PrimaryComponentTick.bCanEverTick = false;
 	// Allows updating in game, while optimizing rendering for the case that it is not modified
-	InstancedStaticMeshComponent->Mobility = EComponentMobility::Movable;
+	ISMComponent->Mobility = EComponentMobility::Movable;
 	// Allows per-instance selection in the editor
-	InstancedStaticMeshComponent->bHasPerInstanceHitProxies = true;
+	ISMComponent->bHasPerInstanceHitProxies = true;
 
-	while(InstancedStaticMeshComponent->InstancingRandomSeed == 0)
+	while(ISMComponent->InstancingRandomSeed == 0)
 	{
-		InstancedStaticMeshComponent->InstancingRandomSeed = FMath::Rand();
+		ISMComponent->InstancingRandomSeed = FMath::Rand();
 	}
 } 
 
@@ -263,25 +322,25 @@ void ALightWeightInstanceStaticMeshManager::OnStaticMeshSet()
 	Modify();
 #endif
 	
-	if (InstancedStaticMeshComponent)
+	if (ISMComponent)
 	{
-		EComponentMobility::Type Mobility = InstancedStaticMeshComponent->Mobility;
+		EComponentMobility::Type Mobility = ISMComponent->Mobility;
 		if (Mobility == EComponentMobility::Static)
 		{
-			InstancedStaticMeshComponent->SetMobility(EComponentMobility::Stationary);
-			InstancedStaticMeshComponent->SetStaticMesh(StaticMesh.Get());
-			InstancedStaticMeshComponent->SetMobility(Mobility);
+			ISMComponent->SetMobility(EComponentMobility::Stationary);
+			ISMComponent->SetStaticMesh(StaticMesh.Get());
+			ISMComponent->SetMobility(Mobility);
 		}
 		else
 		{
-			InstancedStaticMeshComponent->SetStaticMesh(StaticMesh.Get());
+			ISMComponent->SetStaticMesh(StaticMesh.Get());
 		}
 
 		if (UStaticMesh* Mesh = StaticMesh.Get())
 		{
 			for (int32 Idx = 0; Idx < Mesh->GetStaticMaterials().Num(); ++Idx)
 			{
-				InstancedStaticMeshComponent->SetMaterial(Idx, Mesh->GetMaterial(Idx));
+				ISMComponent->SetMaterial(Idx, Mesh->GetMaterial(Idx));
 			}
 		}
 	}
@@ -299,7 +358,7 @@ void ALightWeightInstanceStaticMeshManager::ClearStaticMesh()
 
 FText ALightWeightInstanceStaticMeshManager::GetSMInstanceDisplayName(const FSMInstanceId& InstanceId) const
 {
-	check(InstanceId.ISMComponent == InstancedStaticMeshComponent);
+	check(InstanceId.ISMComponent == ISMComponent);
 
 	const FText OwnerDisplayName = FText::FromString(RepresentedClass ? RepresentedClass->GetName() : AActor::GetName());
 	return FText::Format(NSLOCTEXT("LightWeightInstanceStaticMeshManager", "DisplayNameFmt", "{0} - Instance {1}"), OwnerDisplayName, InstanceId.InstanceIndex);
@@ -307,7 +366,7 @@ FText ALightWeightInstanceStaticMeshManager::GetSMInstanceDisplayName(const FSMI
 
 FText ALightWeightInstanceStaticMeshManager::GetSMInstanceTooltip(const FSMInstanceId& InstanceId) const
 {
-	check(InstanceId.ISMComponent == InstancedStaticMeshComponent);
+	check(InstanceId.ISMComponent == ISMComponent);
 	
 	const FText OwnerDisplayPath = FText::FromString(GetPathName(GetWorld())); // stops the path at the level of the world the object is in
 	return FText::Format(NSLOCTEXT("LightWeightInstanceStaticMeshManager", "TooltipFmt", "Instance {0} on {1}"), InstanceId.InstanceIndex, OwnerDisplayPath);
@@ -315,23 +374,23 @@ FText ALightWeightInstanceStaticMeshManager::GetSMInstanceTooltip(const FSMInsta
 
 bool ALightWeightInstanceStaticMeshManager::CanEditSMInstance(const FSMInstanceId& InstanceId) const
 {
-	check(InstanceId.ISMComponent == InstancedStaticMeshComponent);
+	check(InstanceId.ISMComponent == ISMComponent);
 
-	ISMInstanceManager* InstanceManager = InstancedStaticMeshComponent;
+	ISMInstanceManager* InstanceManager = ISMComponent;
 	return InstanceManager->CanEditSMInstance(InstanceId);
 }
 
 bool ALightWeightInstanceStaticMeshManager::CanMoveSMInstance(const FSMInstanceId& InstanceId, const ETypedElementWorldType InWorldType) const
 {
-	check(InstanceId.ISMComponent == InstancedStaticMeshComponent);
+	check(InstanceId.ISMComponent == ISMComponent);
 
-	ISMInstanceManager* InstanceManager = InstancedStaticMeshComponent;
+	ISMInstanceManager* InstanceManager = ISMComponent;
 	return InstanceManager->CanMoveSMInstance(InstanceId, InWorldType);
 }
 
 bool ALightWeightInstanceStaticMeshManager::GetSMInstanceTransform(const FSMInstanceId& InstanceId, FTransform& OutInstanceTransform, bool bWorldSpace) const
 {
-	check(InstanceId.ISMComponent == InstancedStaticMeshComponent);
+	check(InstanceId.ISMComponent == ISMComponent);
 
 	if (RenderingIndicesToDataIndices.IsValidIndex(InstanceId.InstanceIndex))
 	{
@@ -345,12 +404,12 @@ bool ALightWeightInstanceStaticMeshManager::GetSMInstanceTransform(const FSMInst
 
 bool ALightWeightInstanceStaticMeshManager::SetSMInstanceTransform(const FSMInstanceId& InstanceId, const FTransform& InstanceTransform, bool bWorldSpace, bool bMarkRenderStateDirty, bool bTeleport)
 {
-	check(InstanceId.ISMComponent == InstancedStaticMeshComponent);
+	check(InstanceId.ISMComponent == ISMComponent);
 
-	if (RenderingIndicesToDataIndices.IsValidIndex(InstanceId.InstanceIndex) && InstancedStaticMeshComponent->UpdateInstanceTransform(InstanceId.InstanceIndex, InstanceTransform, bWorldSpace, bMarkRenderStateDirty, bTeleport))
+	if (RenderingIndicesToDataIndices.IsValidIndex(InstanceId.InstanceIndex) && ISMComponent->UpdateInstanceTransform(InstanceId.InstanceIndex, InstanceTransform, bWorldSpace, bMarkRenderStateDirty, bTeleport))
 	{
 		const int32 DataIndex = RenderingIndicesToDataIndices[InstanceId.InstanceIndex];
-		InstancedStaticMeshComponent->GetInstanceTransform(InstanceId.InstanceIndex, InstanceTransforms[DataIndex], /*bWorldSpace*/false);
+		ISMComponent->GetInstanceTransform(InstanceId.InstanceIndex, InstanceTransforms[DataIndex], /*bWorldSpace*/false);
 		return true;
 	}
 
@@ -359,33 +418,33 @@ bool ALightWeightInstanceStaticMeshManager::SetSMInstanceTransform(const FSMInst
 
 void ALightWeightInstanceStaticMeshManager::NotifySMInstanceMovementStarted(const FSMInstanceId& InstanceId)
 {
-	check(InstanceId.ISMComponent == InstancedStaticMeshComponent);
+	check(InstanceId.ISMComponent == ISMComponent);
 
-	ISMInstanceManager* InstanceManager = InstancedStaticMeshComponent;
+	ISMInstanceManager* InstanceManager = ISMComponent;
 	InstanceManager->NotifySMInstanceMovementStarted(InstanceId);
 }
 
 void ALightWeightInstanceStaticMeshManager::NotifySMInstanceMovementOngoing(const FSMInstanceId& InstanceId)
 {
-	check(InstanceId.ISMComponent == InstancedStaticMeshComponent);
+	check(InstanceId.ISMComponent == ISMComponent);
 
-	ISMInstanceManager* InstanceManager = InstancedStaticMeshComponent;
+	ISMInstanceManager* InstanceManager = ISMComponent;
 	InstanceManager->NotifySMInstanceMovementOngoing(InstanceId);
 }
 
 void ALightWeightInstanceStaticMeshManager::NotifySMInstanceMovementEnded(const FSMInstanceId& InstanceId)
 {
-	check(InstanceId.ISMComponent == InstancedStaticMeshComponent);
+	check(InstanceId.ISMComponent == ISMComponent);
 
-	ISMInstanceManager* InstanceManager = InstancedStaticMeshComponent;
+	ISMInstanceManager* InstanceManager = ISMComponent;
 	InstanceManager->NotifySMInstanceMovementEnded(InstanceId);
 }
 
 void ALightWeightInstanceStaticMeshManager::NotifySMInstanceSelectionChanged(const FSMInstanceId& InstanceId, const bool bIsSelected)
 {
-	check(InstanceId.ISMComponent == InstancedStaticMeshComponent);
+	check(InstanceId.ISMComponent == ISMComponent);
 
-	ISMInstanceManager* InstanceManager = InstancedStaticMeshComponent;
+	ISMInstanceManager* InstanceManager = ISMComponent;
 	InstanceManager->NotifySMInstanceSelectionChanged(InstanceId, bIsSelected);
 }
 
@@ -398,7 +457,7 @@ bool ALightWeightInstanceStaticMeshManager::DeleteSMInstances(TArrayView<const F
 	DataIndices.Sort(TGreater<int32>());
 
 	Modify();
-	InstancedStaticMeshComponent->Modify();
+	ISMComponent->Modify();
 
 	for (int32 DataIndex : DataIndices)
 	{
@@ -414,7 +473,7 @@ bool ALightWeightInstanceStaticMeshManager::DuplicateSMInstances(TArrayView<cons
 	GetLWIDataIndices(InstanceIds, DataIndices);
 
 	Modify();
-	InstancedStaticMeshComponent->Modify();
+	ISMComponent->Modify();
 
 	TArray<int32> NewDataIndices;
 	DuplicateLWIInstances(DataIndices, NewDataIndices);
@@ -422,7 +481,7 @@ bool ALightWeightInstanceStaticMeshManager::DuplicateSMInstances(TArrayView<cons
 	OutNewInstanceIds.Reset(NewDataIndices.Num());
 	for (int32 NewDataIndex : NewDataIndices)
 	{
-		OutNewInstanceIds.Add(FSMInstanceId{ InstancedStaticMeshComponent, DataIndicesToRenderingIndices[NewDataIndex] });
+		OutNewInstanceIds.Add(FSMInstanceId{ ISMComponent, DataIndicesToRenderingIndices[NewDataIndex] });
 	}
 
 	return true;
@@ -449,7 +508,7 @@ void ALightWeightInstanceStaticMeshManager::GetLWIDataIndices(TArrayView<const F
 	OutDataIndices.Reset(InstanceIds.Num());
 	for (const FSMInstanceId& InstanceId : InstanceIds)
 	{
-		check(InstanceId.ISMComponent == InstancedStaticMeshComponent);
+		check(InstanceId.ISMComponent == ISMComponent);
 		if (RenderingIndicesToDataIndices.IsValidIndex(InstanceId.InstanceIndex))
 		{
 			OutDataIndices.Add(RenderingIndicesToDataIndices[InstanceId.InstanceIndex]);
