Reserve exact amount of array entries for AnimRetargetSources in USkeleton to prevent reallocations when adding elements to the array

diff --git a/Engine/Source/Runtime/Engine/Private/Animation/Skeleton.cpp b/Engine/Source/Runtime/Engine/Private/Animation/Skeleton.cpp
--- a/Engine/Source/Runtime/Engine/Private/Animation/Skeleton.cpp
+++ b/Engine/Source/Runtime/Engine/Private/Animation/Skeleton.cpp
@@ -433,7 +433,7 @@ void USkeleton::Serialize( FArchive& Ar )
 
 			FName RetargetSourceName;
 			FReferencePose RetargetSource;
-			AnimRetargetSources.Empty();
+			AnimRetargetSources.Empty(NumOfRetargetSources);
 			for (int32 Index=0; Index<NumOfRetargetSources; ++Index)
 			{
 				Ar << RetargetSourceName;
