Fixed AnimComposite animation notify extraction when playing in reverse order

diff --git a/Engine/Source/Runtime/Engine/Private/Animation/AnimComposite.cpp b/Engine/Source/Runtime/Engine/Private/Animation/AnimComposite.cpp
--- a/Engine/Source/Runtime/Engine/Private/Animation/AnimComposite.cpp
+++ b/Engine/Source/Runtime/Engine/Private/Animation/AnimComposite.cpp
@@ -36,34 +36,8 @@ bool UAnimComposite::IsNotifyAvailable() const
 
 void UAnimComposite::GetAnimNotifiesFromDeltaPositions(const float& PreviousPosition, const float & CurrentPosition, FAnimNotifyContext& NotifyContext) const
 {
-	const bool bMovingForward = (RateScale >= 0.f);
-
 	Super::GetAnimNotifiesFromDeltaPositions(PreviousPosition, CurrentPosition, NotifyContext);
-
-	if (bMovingForward)
-	{
-		if (PreviousPosition <= CurrentPosition)
-		{
-			AnimationTrack.GetAnimNotifiesFromTrackPositions(PreviousPosition, CurrentPosition, NotifyContext);
-		}
-		else
-		{
-			AnimationTrack.GetAnimNotifiesFromTrackPositions(PreviousPosition, GetPlayLength(), NotifyContext);
-			AnimationTrack.GetAnimNotifiesFromTrackPositions(0.f, CurrentPosition, NotifyContext);
-		}
-	}
-	else
-	{
-		if (PreviousPosition >= CurrentPosition)
-		{
-			AnimationTrack.GetAnimNotifiesFromTrackPositions(PreviousPosition, CurrentPosition, NotifyContext);
-		}
-		else
-		{
-			AnimationTrack.GetAnimNotifiesFromTrackPositions(PreviousPosition, 0.f, NotifyContext);
-			AnimationTrack.GetAnimNotifiesFromTrackPositions(GetPlayLength(), CurrentPosition, NotifyContext);
-		}
-	}
+	AnimationTrack.GetAnimNotifiesFromTrackPositions(PreviousPosition, CurrentPosition, NotifyContext);
 }
 
 void UAnimComposite::GetResourceSizeEx(FResourceSizeEx& CumulativeResourceSize)
