Fixed crash at PoseableMeshComponent CopyPoseFromSkeletalComponent function, when the component is being incrementally unregistered and the unregistering process takes longer than the time limit set

diff --git a/Engine/Source/Runtime/Engine/Private/Components/PoseableMeshComponent.cpp b/Engine/Source/Runtime/Engine/Private/Components/PoseableMeshComponent.cpp
--- a/Engine/Source/Runtime/Engine/Private/Components/PoseableMeshComponent.cpp
+++ b/Engine/Source/Runtime/Engine/Private/Components/PoseableMeshComponent.cpp
@@ -287,6 +287,11 @@ void UPoseableMeshComponent::CopyPoseFromSkeletalComponent(USkeletalMeshComponen
 	if(RequiredBones.IsValid())
 	{
 		TArray<FTransform> LocalTransforms = InComponentToCopy->GetBoneSpaceTransforms();
+		if (LocalTransforms.IsEmpty()) // during incremental unregistering this might be empty
+		{
+			return;
+		}
+
 		if(this->GetSkinnedAsset() == InComponentToCopy->GetSkinnedAsset()
 			&& LocalTransforms.Num() == BoneSpaceTransforms.Num())
 		{
