Fix for SparseVolumeTexture upload shader reading past the 2^27 maximum element limit of certain (typed) buffers and reading zero. Fixed by splitting uploads into smaller chunks.

#rnx

diff --git a/Engine/Shaders/Private/SparseVolumeTexture/UpdateSparseVolumeTexture.usf b/Engine/Shaders/Private/SparseVolumeTexture/UpdateSparseVolumeTexture.usf
--- a/Engine/Shaders/Private/SparseVolumeTexture/UpdateSparseVolumeTexture.usf
+++ b/Engine/Shaders/Private/SparseVolumeTexture/UpdateSparseVolumeTexture.usf
@@ -78,6 +78,9 @@ ByteAddressBuffer TileDataOffsetsBuffer;
 ByteAddressBuffer DstTileCoordsBuffer;
 float4 FallbackValueA;
 float4 FallbackValueB;
+uint TileIndexOffset; // Number of tiles that have already been processed in previous dispatches
+uint SrcVoxelDataOffsetA; // Start of the SrcPhysicalTileBufferA SRV. This value needs to be subtracted from any indices used for loads.
+uint SrcVoxelDataOffsetB; // Start of the SrcPhysicalTileBufferB SRV. This value needs to be subtracted from any indices used for loads.
 uint NumTilesToCopy;
 uint BufferTileStep;
 uint NumDispatchedGroups;
@@ -113,7 +116,7 @@ void SparseVolumeTextureUpdateFromSparseBufferCS(uint3 GroupId : SV_GroupID, uin
 	const uint NumOccupancyWordsPerPaddedTile = (NumVoxelsPerPaddedTile + 31u) / 32u;
 
 	// Process one tile per workgroup
-	for (uint TileIdx = GroupIdx; TileIdx < NumTilesToCopy; TileIdx += NumDispatchedGroups)
+	for (uint TileIdx = (GroupIdx + TileIndexOffset); TileIdx < (NumTilesToCopy + TileIndexOffset); TileIdx += NumDispatchedGroups)
 	{
 		const uint PackedDstCoord = DstTileCoordsBuffer.Load(TileIdx << 2u);
 		const uint3 DstBaseCoord = uint3(PackedDstCoord & 0xFFu, (PackedDstCoord >> 8u) & 0xFFu, (PackedDstCoord >> 16u) & 0xFFu) * PaddedTileSize;
@@ -140,6 +143,7 @@ void SparseVolumeTextureUpdateFromSparseBufferCS(uint3 GroupId : SV_GroupID, uin
 
 					bool bValidVoxel = false;
 					const uint LocalLoadIndex = ComputeLocalVoxelDataLoadIndex(VoxelIdx, OccupancyBits, bValidVoxel);
+					const uint VoxelLoadIndex = TileDataOffset + NumReadVoxels + LocalLoadIndex;
 					const uint3 DstCoord = DstBaseCoord + UnflattenCoord(VoxelIdx, PaddedTileSize);
 
 					BRANCH if (TextureIdx == 0)
@@ -147,7 +151,7 @@ void SparseVolumeTextureUpdateFromSparseBufferCS(uint3 GroupId : SV_GroupID, uin
 						float4 Voxel = FallbackValueA;
 						BRANCH if (bValidVoxel)
 						{
-							Voxel = SrcPhysicalTileBufferA.Load(TileDataOffset + NumReadVoxels + LocalLoadIndex);
+							Voxel = SrcPhysicalTileBufferA.Load(VoxelLoadIndex - SrcVoxelDataOffsetA);
 						}
 						DstPhysicalTileTextureA[DstCoord] = Voxel;
 					}
@@ -156,7 +160,7 @@ void SparseVolumeTextureUpdateFromSparseBufferCS(uint3 GroupId : SV_GroupID, uin
 						float4 Voxel = FallbackValueB;
 						BRANCH if (bValidVoxel)
 						{
-							Voxel = SrcPhysicalTileBufferB.Load(TileDataOffset + NumReadVoxels + LocalLoadIndex);
+							Voxel = SrcPhysicalTileBufferB.Load(VoxelLoadIndex - SrcVoxelDataOffsetB);
 						}
 						DstPhysicalTileTextureB[DstCoord] = Voxel;
 					}
diff --git a/Engine/Source/Runtime/Engine/Private/SparseVolumeTexture/SparseVolumeTextureStreamingManager.cpp b/Engine/Source/Runtime/Engine/Private/SparseVolumeTexture/SparseVolumeTextureStreamingManager.cpp
--- a/Engine/Source/Runtime/Engine/Private/SparseVolumeTexture/SparseVolumeTextureStreamingManager.cpp
+++ b/Engine/Source/Runtime/Engine/Private/SparseVolumeTexture/SparseVolumeTextureStreamingManager.cpp
@@ -212,6 +212,9 @@ class FSparseVolumeTextureUpdateFromSparseBufferCS : public FGlobalShader
 		SHADER_PARAMETER_RDG_BUFFER_SRV(ByteAddressBuffer, DstTileCoordsBuffer)
 		SHADER_PARAMETER(FVector4f, FallbackValueA)
 		SHADER_PARAMETER(FVector4f, FallbackValueB)
+		SHADER_PARAMETER(uint32, TileIndexOffset)
+		SHADER_PARAMETER(uint32, SrcVoxelDataOffsetA)
+		SHADER_PARAMETER(uint32, SrcVoxelDataOffsetB)
 		SHADER_PARAMETER(uint32, NumTilesToCopy)
 		SHADER_PARAMETER(uint32, BufferTileStep)
 		SHADER_PARAMETER(uint32, NumDispatchedGroups)
@@ -347,7 +350,11 @@ public:
 				BufferDesc.Usage |= EBufferUsageFlags::Dynamic; // Skip the unneeded copy from upload to VRAM resource on d3d12 RHI
 				AllocatePooledBuffer(BufferDesc, TileDataOffsetsUploadBuffer, TEXT("SparseVolumeTexture.TileDataOffsetsUploadBuffer"));
 
-				TileDataOffsetsAPtr = (uint8*)RHICmdList.LockBuffer(TileDataOffsetsUploadBuffer->GetRHI(), 0, BufferSize, RLM_WriteOnly);
+				// Due to a limit on the maximum number of texels in a buffer SRV, we need to upload the data in smaller chunks. In order to figure out the chunk offsets/sizes,
+				// we need to read the TileDataOffset values the caller has written to the returned pointers. We want to avoid reading from a mapped upload buffer pointer,
+				// which is why we use a temporary allocation to write the upload data to.
+				TileDataOffsets.SetNumUninitialized(NumTextures * MaxNumTiles);
+				TileDataOffsetsAPtr = (uint8*)TileDataOffsets.GetData();
 				TileDataOffsetsBPtr = TileDataOffsetsAPtr + (FormatA != PF_Unknown ? (MaxNumTiles * sizeof(uint32)) : 0);
 			}
 			// TileCoords
@@ -414,6 +421,7 @@ public:
 		DstTileCoordsUploadBuffer.SafeRelease();
 		TileDataAUploadBuffer.SafeRelease();
 		TileDataBUploadBuffer.SafeRelease();
+		TileDataOffsets.Reset();
 		ResetState();
 	}
 
@@ -426,8 +434,13 @@ public:
 			FRHICommandListBase& RHICmdList = GraphBuilder.RHICmdList;
 
 			RHICmdList.UnlockBuffer(OccupancyBitsUploadBuffer->GetRHI());
-			RHICmdList.UnlockBuffer(TileDataOffsetsUploadBuffer->GetRHI());
 			RHICmdList.UnlockBuffer(DstTileCoordsUploadBuffer->GetRHI());
+
+			// TileDataOffset values were written to a temporary allocation so that we can access them later in this function. Unlike the other buffers, we now need to copy that data over to the actual upload buffer.
+			void* TileDataOffsetsUploadPtr = RHICmdList.LockBuffer(TileDataOffsetsUploadBuffer->GetRHI(), 0, TileDataOffsets.Num() * sizeof(TileDataOffsets[0]), RLM_WriteOnly);
+			FMemory::Memcpy(TileDataOffsetsUploadPtr, TileDataOffsets.GetData(), TileDataOffsets.Num() * sizeof(TileDataOffsets[0]));
+			RHICmdList.UnlockBuffer(TileDataOffsetsUploadBuffer->GetRHI());
+
 			if (TileDataAPtr)
 			{
 				RHICmdList.UnlockBuffer(TileDataAUploadBuffer->GetRHI());
@@ -443,45 +456,112 @@ public:
 				FRDGTexture* DstTextureARDG = DstTextureA ? GraphBuilder.RegisterExternalTexture(CreateRenderTarget(DstTextureA, TEXT("SparseVolumeTexture.TileDataTextureA"))) : nullptr;
 				FRDGTexture* DstTextureBRDG = DstTextureB ? GraphBuilder.RegisterExternalTexture(CreateRenderTarget(DstTextureB, TEXT("SparseVolumeTexture.TileDataTextureB"))) : nullptr;
 
-				FRDGBufferSRV* OccupancyBitsBufferSRV = GraphBuilder.CreateSRV(GraphBuilder.RegisterExternalBuffer(OccupancyBitsUploadBuffer));
-				FRDGBufferSRV* TileDataOffsetsBufferSRV = GraphBuilder.CreateSRV(GraphBuilder.RegisterExternalBuffer(TileDataOffsetsUploadBuffer));
-				FRDGBufferSRV* DstTileCoordsBufferSRV = GraphBuilder.CreateSRV(GraphBuilder.RegisterExternalBuffer(DstTileCoordsUploadBuffer));
-				FRDGBufferSRV* TileDataABufferSRV = FormatSizeA ? GraphBuilder.CreateSRV(GraphBuilder.RegisterExternalBuffer(TileDataAUploadBuffer), FormatA) : nullptr;
-				FRDGBufferSRV* TileDataBBufferSRV = FormatSizeB ? GraphBuilder.CreateSRV(GraphBuilder.RegisterExternalBuffer(TileDataBUploadBuffer), FormatB) : nullptr;
 				FRDGTextureUAV* DummyTextureUAV = GraphBuilder.CreateUAV(DummyTexture);
 				FRDGTextureUAV* DstTextureAUAV = DstTextureARDG ? GraphBuilder.CreateUAV(DstTextureARDG) : nullptr;
 				FRDGTextureUAV* DstTextureBUAV = DstTextureBRDG ? GraphBuilder.CreateUAV(DstTextureBRDG) : nullptr;
 
+				FRDGBufferSRV* OccupancyBitsBufferSRV = GraphBuilder.CreateSRV(GraphBuilder.RegisterExternalBuffer(OccupancyBitsUploadBuffer));
+				FRDGBufferSRV* TileDataOffsetsBufferSRV = GraphBuilder.CreateSRV(GraphBuilder.RegisterExternalBuffer(TileDataOffsetsUploadBuffer));
+				FRDGBufferSRV* DstTileCoordsBufferSRV = GraphBuilder.CreateSRV(GraphBuilder.RegisterExternalBuffer(DstTileCoordsUploadBuffer));
+
 				auto ComputeShader = GetGlobalShaderMap(GMaxRHIFeatureLevel)->GetShader<FSparseVolumeTextureUpdateFromSparseBufferCS>();
 
-				FSparseVolumeTextureUpdateFromSparseBufferCS::FParameters* PassParameters = GraphBuilder.AllocParameters<FSparseVolumeTextureUpdateFromSparseBufferCS::FParameters>();
-				PassParameters->DstPhysicalTileTextureA = FormatSizeA > 0 ? DstTextureAUAV : DummyTextureUAV;
-				PassParameters->DstPhysicalTileTextureB = FormatSizeB > 0 ? DstTextureBUAV : DummyTextureUAV;
-				PassParameters->SrcPhysicalTileBufferA = FormatSizeA > 0 ? TileDataABufferSRV : TileDataBBufferSRV;
-				PassParameters->SrcPhysicalTileBufferB = FormatSizeB > 0 ? TileDataBBufferSRV : TileDataABufferSRV;
-				PassParameters->OccupancyBitsBuffer = OccupancyBitsBufferSRV;
-				PassParameters->TileDataOffsetsBuffer = TileDataOffsetsBufferSRV;
-				PassParameters->DstTileCoordsBuffer = DstTileCoordsBufferSRV;
-				PassParameters->FallbackValueA = FallbackValueA;
-				PassParameters->FallbackValueB = FallbackValueB;
-				PassParameters->NumTilesToCopy = NumWrittenTiles;
-				PassParameters->BufferTileStep = MaxNumTiles;
-				PassParameters->NumDispatchedGroups = FMath::Min(NumWrittenTiles, 1024);
-				PassParameters->PaddedTileSize = SPARSE_VOLUME_TILE_RES_PADDED;
-				PassParameters->CopyTexureMask = 0;
-				PassParameters->CopyTexureMask |= FormatSizeA > 0 ? 0x1u : 0x0u;
-				PassParameters->CopyTexureMask |= FormatSizeB > 0 ? 0x2u : 0x0u;
+				int32 NumUploadedTiles = 0;
+				int32 NumUploadedVoxelsA = 0;
+				int32 NumUploadedVoxelsB = 0;
 
-				const bool bAsyncCompute = GSupportsEfficientAsyncCompute && (GSVTStreamingAsyncCompute != 0);
+				// This is a limit on some platforms on the maximum number of texels in a texel/typed buffer. Unfortunately for R8 (1 byte) formats, this means that we can upload only 1/16 of all the texels of a 2GB texture.
+				// In order to work around this issue, the data to be uploaded is split into chunks such that this limit is not violated. We can use the TileDataOffsets values to get the number of voxels per tile and then
+				// just add as many tiles to the batch as we can.
+				const int32 MaxNumTexelsPerResource = 1 << 27;
 
-				FComputeShaderUtils::AddPass(
-					GraphBuilder,
-					RDG_EVENT_NAME("Upload SVT Tiles (TileCount: %u)", NumWrittenTiles),
-					bAsyncCompute ? ERDGPassFlags::AsyncCompute : ERDGPassFlags::Compute,
-					ComputeShader,
-					PassParameters,
-					FIntVector3(PassParameters->NumDispatchedGroups, 1, 1)
-				);
+				while (NumUploadedTiles < NumWrittenTiles)
+				{
+					// Determine the number of tiles to upload in this iteration
+					int32 NumTilesInThisBatch = 0;
+					int32 NumVoxelsAInThisBatch = 0;
+					int32 NumVoxelsBInThisBatch = 0;
+					for (int32 TileIndex = NumUploadedTiles; TileIndex < NumWrittenTiles; ++TileIndex)
+					{
+						const int32 VoxelOffsetA = FormatSizeA > 0 ? reinterpret_cast<uint32*>(TileDataOffsetsAPtr)[TileIndex] : 0;
+						const int32 VoxelOffsetB = FormatSizeB > 0 ? reinterpret_cast<uint32*>(TileDataOffsetsBPtr)[TileIndex] : 0;
+						const int32 VoxelEndIndexA = FormatSizeA > 0 && (TileIndex + 1) < NumWrittenTiles ? reinterpret_cast<uint32*>(TileDataOffsetsAPtr)[TileIndex + 1] : NumWrittenVoxelsA;
+						const int32 VoxelEndIndexB = FormatSizeB > 0 && (TileIndex + 1) < NumWrittenTiles ? reinterpret_cast<uint32*>(TileDataOffsetsBPtr)[TileIndex + 1] : NumWrittenVoxelsB;
+						const int32 TileNumVoxelsA = VoxelEndIndexA - VoxelOffsetA;
+						const int32 TileNumVoxelsB = VoxelEndIndexB - VoxelOffsetB;
+						check(TileNumVoxelsA >= 0 && TileNumVoxelsA <= SVT::NumVoxelsPerPaddedTile);
+						check(TileNumVoxelsB >= 0 && TileNumVoxelsB <= SVT::NumVoxelsPerPaddedTile);
+
+						// Adding additional voxels to the batch would exceed the limit, so exit the loop and upload the data.
+						if ((NumVoxelsAInThisBatch + TileNumVoxelsA) > MaxNumTexelsPerResource || (NumVoxelsBInThisBatch + TileNumVoxelsB) > MaxNumTexelsPerResource)
+						{
+							break;
+						}
+
+						NumTilesInThisBatch += 1;
+						NumVoxelsAInThisBatch += TileNumVoxelsA;
+						NumVoxelsBInThisBatch += TileNumVoxelsB;
+					}
+
+					FRDGBufferSRV* TileDataABufferSRV = nullptr;
+					FRDGBufferSRV* TileDataBBufferSRV = nullptr;
+
+					// This is the critical part: For every batch we create a SRV scoped to a range within the voxel data upload buffer still fitting within the typed buffer texel limit.
+					if (FormatSizeA && NumVoxelsAInThisBatch)
+					{
+						FRDGBufferSRVDesc SRVDesc(GraphBuilder.RegisterExternalBuffer(TileDataAUploadBuffer), FormatA);
+						SRVDesc.StartOffsetBytes = NumUploadedVoxelsA * FormatSizeA;
+						SRVDesc.NumElements = NumVoxelsAInThisBatch;
+						TileDataABufferSRV = GraphBuilder.CreateSRV(SRVDesc);
+					}
+					if (FormatSizeB && NumVoxelsBInThisBatch)
+					{
+						FRDGBufferSRVDesc SRVDesc(GraphBuilder.RegisterExternalBuffer(TileDataBUploadBuffer), FormatB);
+						SRVDesc.StartOffsetBytes = NumUploadedVoxelsB * FormatSizeB;
+						SRVDesc.NumElements = NumVoxelsBInThisBatch;
+						TileDataBBufferSRV = GraphBuilder.CreateSRV(SRVDesc);
+					}
+
+					FSparseVolumeTextureUpdateFromSparseBufferCS::FParameters* PassParameters = GraphBuilder.AllocParameters<FSparseVolumeTextureUpdateFromSparseBufferCS::FParameters>();
+					PassParameters->DstPhysicalTileTextureA = DstTextureAUAV ? DstTextureAUAV : DummyTextureUAV;
+					PassParameters->DstPhysicalTileTextureB = DstTextureBUAV ? DstTextureBUAV : DummyTextureUAV;
+					PassParameters->SrcPhysicalTileBufferA = TileDataABufferSRV ? TileDataABufferSRV : TileDataBBufferSRV;
+					PassParameters->SrcPhysicalTileBufferB = TileDataBBufferSRV ? TileDataBBufferSRV : TileDataABufferSRV;
+					PassParameters->OccupancyBitsBuffer = OccupancyBitsBufferSRV;
+					PassParameters->TileDataOffsetsBuffer = TileDataOffsetsBufferSRV;
+					PassParameters->DstTileCoordsBuffer = DstTileCoordsBufferSRV;
+					PassParameters->FallbackValueA = FallbackValueA;
+					PassParameters->FallbackValueB = FallbackValueB;
+					PassParameters->TileIndexOffset = NumUploadedTiles; // This lets the shader know how many tiles have already been processed in previous dispatches.
+					PassParameters->SrcVoxelDataOffsetA = NumUploadedVoxelsA; // SrcVoxelDataOffsetA and SrcVoxelDataOffsetB are subtracted from the calculated voxel data buffer read indices
+					PassParameters->SrcVoxelDataOffsetB = NumUploadedVoxelsB;
+					PassParameters->NumTilesToCopy = NumTilesInThisBatch;
+					PassParameters->BufferTileStep = MaxNumTiles;
+					PassParameters->NumDispatchedGroups = FMath::Min(NumTilesInThisBatch, GRHIMaxDispatchThreadGroupsPerDimension.X);
+					PassParameters->PaddedTileSize = SPARSE_VOLUME_TILE_RES_PADDED;
+					PassParameters->CopyTexureMask = 0;
+					PassParameters->CopyTexureMask |= FormatSizeA > 0 ? 0x1u : 0x0u;
+					PassParameters->CopyTexureMask |= FormatSizeB > 0 ? 0x2u : 0x0u;
+
+					const bool bAsyncCompute = GSupportsEfficientAsyncCompute && (GSVTStreamingAsyncCompute != 0);
+
+					FComputeShaderUtils::AddPass(
+						GraphBuilder,
+						RDG_EVENT_NAME("Upload SVT Tiles (TileCount: %u)", NumTilesInThisBatch),
+						bAsyncCompute ? ERDGPassFlags::AsyncCompute : ERDGPassFlags::Compute,
+						ComputeShader,
+						PassParameters,
+						FIntVector3(PassParameters->NumDispatchedGroups, 1, 1)
+					);
+
+					NumUploadedTiles += NumTilesInThisBatch;
+					NumUploadedVoxelsA += NumVoxelsAInThisBatch;
+					NumUploadedVoxelsB += NumVoxelsBInThisBatch;
+				}
+
+				check(NumUploadedTiles == NumWrittenTiles);
+				check(NumUploadedVoxelsA == NumWrittenVoxelsA);
+				check(NumUploadedVoxelsB == NumWrittenVoxelsB);
 			}
 		}
 		Release();
@@ -493,6 +573,7 @@ private:
 	TRefCountPtr<FRDGPooledBuffer> DstTileCoordsUploadBuffer;
 	TRefCountPtr<FRDGPooledBuffer> TileDataAUploadBuffer;
 	TRefCountPtr<FRDGPooledBuffer> TileDataBUploadBuffer;
+	TArray<uint32> TileDataOffsets; // CPU-readable per-tile offsets into tile data
 	uint8* OccupancyBitsAPtr = nullptr;
 	uint8* OccupancyBitsBPtr = nullptr;
 	uint8* TileDataOffsetsAPtr = nullptr;
@@ -1316,12 +1397,14 @@ void FStreamingManager::AddInternal(FRDGBuilder& GraphBuilder, FNewSparseVolumeT
 			if (SVTInfo.FormatA != PF_Unknown)
 			{
 				FMemory::Memzero(AddResult.OccupancyBitsPtrs[0], SVT::NumOccupancyWordsPerPaddedTile * sizeof(uint32));
+				FMemory::Memzero(AddResult.TileDataOffsetsPtrs[0], sizeof(uint32));
 			}
 			if (SVTInfo.FormatB != PF_Unknown)
 			{
 				FMemory::Memzero(AddResult.OccupancyBitsPtrs[1], SVT::NumOccupancyWordsPerPaddedTile * sizeof(uint32));
+				FMemory::Memzero(AddResult.TileDataOffsetsPtrs[1], sizeof(uint32));
 			}
-			// No need to write to TileDataPtrA. TileDataPtrB, TileDataOffsetsPtrA and TileDataOffsetsPtrB because we zeroed out all the occupancy bits.
+			// No need to write to TileDataPtrA and TileDataPtrB because we zeroed out all the occupancy bits.
 		}
 
 		// Process frames
