Hopefully build fix (temporary casting)

diff --git a/Engine/Source/Runtime/Engine/Classes/PhysicsEngine/PhysicsSettings.h b/Engine/Source/Runtime/Engine/Classes/PhysicsEngine/PhysicsSettings.h
--- a/Engine/Source/Runtime/Engine/Classes/PhysicsEngine/PhysicsSettings.h
+++ b/Engine/Source/Runtime/Engine/Classes/PhysicsEngine/PhysicsSettings.h
@@ -258,7 +258,7 @@ public:
 	UFUNCTION(BlueprintCallable, Category = "Physics")
 	int32 GetPhysicsHistoryCount() const
 	{
-		return (int32) (FMath::CeilToInt(1.f / AsyncFixedTimeStepSize) * (PhysicsPrediction.MaxSupportedLatencyPrediction / 1000.f));
+		return (int32) ((float) (FMath::CeilToInt(1.f / AsyncFixedTimeStepSize)) * (PhysicsPrediction.MaxSupportedLatencyPrediction / 1000.f));
 	}
 
 	ENGINE_API virtual void PostInitProperties() override;
