TimelineComponent: Marked several functions as blueprint callable.

diff --git a/Engine/Source/Runtime/Engine/Classes/Components/TimelineComponent.h b/Engine/Source/Runtime/Engine/Classes/Components/TimelineComponent.h
--- a/Engine/Source/Runtime/Engine/Classes/Components/TimelineComponent.h
+++ b/Engine/Source/Runtime/Engine/Classes/Components/TimelineComponent.h
@@ -395,7 +395,7 @@ private:
 	/** The actual timeline structure */
 	UPROPERTY(ReplicatedUsing=OnRep_Timeline)
 	FTimeline	TheTimeline;
-
+	
 	/** True if global time dilation should be ignored by this timeline, false otherwise. */
 	UPROPERTY()
 	uint32 bIgnoreTimeDilation : 1;
@@ -524,17 +524,21 @@ public:
 	static ETimelineSigType GetTimelineSignatureForFunction(const UFunction* InFunc);
 
 	/** Add a callback event to the timeline */
+	UFUNCTION(BlueprintCallable, Category = "Components|Timeline")
 	ENGINE_API void AddEvent(float Time, FOnTimelineEvent EventFunc);
 	
 	/** Add a vector interpolation to the timeline */
+	UFUNCTION(BlueprintCallable, Category = "Components|Timeline")
 	ENGINE_API void AddInterpVector(UCurveVector* VectorCurve, FOnTimelineVector InterpFunc, FName PropertyName = NAME_None, FName TrackName = NAME_None);
 	ENGINE_API void AddInterpVector(UCurveVector* VectorCurve, FOnTimelineVectorStatic InterpFunc);
 	
 	/** Add a float interpolation to the timeline */
+	UFUNCTION(BlueprintCallable, Category = "Components|Timeline")
 	ENGINE_API void AddInterpFloat(UCurveFloat* FloatCurve, FOnTimelineFloat InterpFunc, FName PropertyName = NAME_None, FName TrackName = NAME_None);
 	ENGINE_API void AddInterpFloat(UCurveFloat* FloatCurve, FOnTimelineFloatStatic InterpFunc);
 
 	/** Add a linear color interpolation to the timeline */
+	UFUNCTION(BlueprintCallable, Category = "Components|Timeline")
 	ENGINE_API void AddInterpLinearColor(UCurveLinearColor* LinearColorCurve, FOnTimelineLinearColor InterpFunc, FName PropertyName = NAME_None, FName TrackName = NAME_None);
 	ENGINE_API void AddInterpLinearColor(UCurveLinearColor* LinearColorCurve, FOnTimelineLinearColorStatic InterpFunc);
 
@@ -542,9 +546,11 @@ public:
 	ENGINE_API void SetPropertySetObject(UObject* NewPropertySetObject);
 
 	/** Set the delegate to call after each timeline tick */
+	UFUNCTION(BlueprintCallable, Category = "Components|Timeline")
 	ENGINE_API void SetTimelinePostUpdateFunc(FOnTimelineEvent NewTimelinePostUpdateFunc);
 
 	/** Set the delegate to call when timeline is finished */
+	UFUNCTION(BlueprintCallable, Category = "Components|Timeline")
 	ENGINE_API void SetTimelineFinishedFunc(FOnTimelineEvent NewTimelineFinishedFunc);
 	/** Set the static delegate to call when timeline is finished */
 	ENGINE_API void SetTimelineFinishedFunc(FOnTimelineEventStatic NewTimelineFinishedFunc);
