Add a new render asset streaming heuristic that prioritizes loading of desired MIPs or LODs if too many are missing. Test result shows some negative impact on level streaming speed but the perceived visual quality is overall improved. Most textures and meshes are presented at decent resolutions at the time skydiving ends. In contrast, it is very common to see blurry textures and rough meshes without the heuristic. There is no noticeable negative impacts on other perf metrics such as MVP and HPM. Map load time is also reduced when the heuristic is enabled. This change only adds the implementation and option to turn the heuristic on but doesn't enable the option.

diff --git a/Engine/Source/Runtime/Engine/Private/Streaming/AsyncTextureStreaming.cpp b/Engine/Source/Runtime/Engine/Private/Streaming/AsyncTextureStreaming.cpp
--- a/Engine/Source/Runtime/Engine/Private/Streaming/AsyncTextureStreaming.cpp
+++ b/Engine/Source/Runtime/Engine/Private/Streaming/AsyncTextureStreaming.cpp
@@ -799,18 +799,23 @@ void FRenderAssetStreamingMipCalcTask::UpdateLoadAndCancelationRequests_Async()
 	for (int32 AssetIndex = 0; AssetIndex < StreamingRenderAssets.Num() && !IsAborted(); ++AssetIndex)
 	{
 		FStreamingRenderAsset& StreamingRenderAsset = StreamingRenderAssets[AssetIndex];
+		const bool bWasMissingTooManyMips = StreamingRenderAsset.IsMissingTooManyMips();
 
 		// If we need to change the number of resident mips.
-		if (StreamingRenderAsset.UpdateLoadOrderPriority_Async(Settings.MinMipForSplitRequest))
+		if (StreamingRenderAsset.UpdateLoadOrderPriority_Async(Settings))
 		{
 			// If there is no pending update, kick one if the budget allows it.
 			if (StreamingRenderAsset.RequestedMips == StreamingRenderAsset.ResidentMips)
 			{
 				PrioritizedRenderAssets.Add(AssetIndex);
 			}
-			// Otherwise, if the update is trying to load to many mips, or if it is trying to unload required mips, (try to) cancel it.
-			else if (StreamingRenderAsset.RequestedMips > FMath::Max<int32>(StreamingRenderAsset.ResidentMips, StreamingRenderAsset.WantedMips + 1) ||
-					StreamingRenderAsset.RequestedMips < FMath::Min<int32>(StreamingRenderAsset.ResidentMips, StreamingRenderAsset.WantedMips))
+			// Otherwise, if the update is trying to load too many, too few, or unload required MIPs, (try to) cancel it.
+			else if (
+				// If marked as missing too many MIPs, a high priority request was created so be more aggressive on canceling it.
+				StreamingRenderAsset.RequestedMips > FMath::Max<int32>(StreamingRenderAsset.ResidentMips, StreamingRenderAsset.WantedMips + (bWasMissingTooManyMips ? 0 : 1)) ||
+				// If too many missing MIPs, cancel existing request if it is not loading enough so a high priority one can be created.
+				// Otherwise, only cancel if it is trying to unload resident MIPs.
+				StreamingRenderAsset.RequestedMips < (StreamingRenderAsset.IsMissingTooManyMips() ? StreamingRenderAsset.WantedMips : FMath::Min<int32>(StreamingRenderAsset.ResidentMips, StreamingRenderAsset.WantedMips)))
 			{
 				CancelationRequests.Add(AssetIndex);
 			}
diff --git a/Engine/Source/Runtime/Engine/Private/Streaming/SkeletalMeshUpdate.cpp b/Engine/Source/Runtime/Engine/Private/Streaming/SkeletalMeshUpdate.cpp
--- a/Engine/Source/Runtime/Engine/Private/Streaming/SkeletalMeshUpdate.cpp
+++ b/Engine/Source/Runtime/Engine/Private/Streaming/SkeletalMeshUpdate.cpp
@@ -15,13 +15,6 @@ SkeletalMeshUpdate.cpp: Helpers to stream in and out skeletal mesh LODs.
 #include "Streaming/RenderAssetUpdate.inl"
 #include "RHIResourceUpdates.h"
 
-int32 GStreamingSkeletalMeshIOPriority = (int32)AIOP_Low;
-static FAutoConsoleVariableRef CVarStreamingSkeletalMeshIOPriority(
-	TEXT("r.Streaming.SkeletalMeshIOPriority"),
-	GStreamingSkeletalMeshIOPriority,
-	TEXT("Base I/O priority for loading skeletal mesh LODs"),
-	ECVF_Default);
-
 extern int32 GStreamingMaxReferenceChecks;
 
 template class TRenderAssetUpdate<FSkelMeshUpdateContext>;
@@ -453,7 +446,23 @@ void FSkeletalMeshStreamIn_IO::SetIORequest(const FContext& Context)
 		// but that won't do anything because the tick would not try to acquire the lock since it is already locked.
 		TaskSynchronization.Increment();
 
-		const EAsyncIOPriorityAndFlags Priority = (EAsyncIOPriorityAndFlags)FMath::Clamp<int32>(GStreamingSkeletalMeshIOPriority + (bHighPrioIORequest ? 1 : 0), AIOP_Low, AIOP_High);
+		EAsyncIOPriorityAndFlags Priority = AIOP_Low;
+		if (bHighPrioIORequest)
+		{
+			static IConsoleVariable* CVarAsyncLoadingPrecachePriority = IConsoleManager::Get().FindConsoleVariable(TEXT("s.AsyncLoadingPrecachePriority"));
+			const bool bLoadBeforeAsyncPrecache = CVarStreamingLowResHandlingMode.GetValueOnAnyThread() == (int32)FRenderAssetStreamingSettings::LRHM_LoadBeforeAsyncPrecache;
+
+			if (CVarAsyncLoadingPrecachePriority && bLoadBeforeAsyncPrecache)
+			{
+				const int32 AsyncIOPriority = CVarAsyncLoadingPrecachePriority->GetInt();
+				// Higher priority than regular requests but don't go over max
+				Priority = (EAsyncIOPriorityAndFlags)FMath::Clamp<int32>(AsyncIOPriority + 1, AIOP_BelowNormal, AIOP_MAX);
+			}
+			else
+			{
+				Priority = AIOP_BelowNormal;
+			}
+		}
 
 		Batch.Issue(BulkData, Priority, [this](FBulkDataRequest::EStatus Status)
 		{
diff --git a/Engine/Source/Runtime/Engine/Private/Streaming/StaticMeshUpdate.cpp b/Engine/Source/Runtime/Engine/Private/Streaming/StaticMeshUpdate.cpp
--- a/Engine/Source/Runtime/Engine/Private/Streaming/StaticMeshUpdate.cpp
+++ b/Engine/Source/Runtime/Engine/Private/Streaming/StaticMeshUpdate.cpp
@@ -22,13 +22,6 @@ static FAutoConsoleVariableRef CVarStreamingMaxReferenceChecksBeforeStreamOut(
 	ECVF_Default
 );
 
-int32 GStreamingStaticMeshIOPriority = (int32)AIOP_Low;
-static FAutoConsoleVariableRef CVarStreamingStaticMeshIOPriority(
-	TEXT("r.Streaming.StaticMeshIOPriority"),
-	GStreamingStaticMeshIOPriority,
-	TEXT("Base I/O priority for loading static mesh LODs"),
-	ECVF_Default);
-
 // Instantiate TRenderAssetUpdate for FStaticMeshUpdateContext
 template class TRenderAssetUpdate<FStaticMeshUpdateContext>;
 
@@ -499,7 +492,23 @@ void FStaticMeshStreamIn_IO::SetIORequest(const FContext& Context)
 		// but that won't do anything because the tick would not try to acquire the lock since it is already locked.
 		TaskSynchronization.Increment();
 
-		const EAsyncIOPriorityAndFlags Priority = (EAsyncIOPriorityAndFlags)FMath::Clamp<int32>(GStreamingStaticMeshIOPriority + (bHighPrioIORequest ? 1 : 0), AIOP_Low, AIOP_High);
+		EAsyncIOPriorityAndFlags Priority = AIOP_Low;
+		if (bHighPrioIORequest)
+		{
+			static IConsoleVariable* CVarAsyncLoadingPrecachePriority = IConsoleManager::Get().FindConsoleVariable(TEXT("s.AsyncLoadingPrecachePriority"));
+			const bool bLoadBeforeAsyncPrecache = CVarStreamingLowResHandlingMode.GetValueOnAnyThread() == (int32)FRenderAssetStreamingSettings::LRHM_LoadBeforeAsyncPrecache;
+
+			if (CVarAsyncLoadingPrecachePriority && bLoadBeforeAsyncPrecache)
+			{
+				const int32 AsyncIOPriority = CVarAsyncLoadingPrecachePriority->GetInt();
+				// Higher priority than regular requests but don't go over max
+				Priority = (EAsyncIOPriorityAndFlags)FMath::Clamp<int32>(AsyncIOPriority + 1, AIOP_BelowNormal, AIOP_MAX);
+			}
+			else
+			{
+				Priority = AIOP_BelowNormal;
+			}
+		}
 
 		Batch.Issue(BulkData, Priority, [this](FBulkDataRequest::EStatus Status)
 		{
diff --git a/Engine/Source/Runtime/Engine/Private/Streaming/StreamingTexture.cpp b/Engine/Source/Runtime/Engine/Private/Streaming/StreamingTexture.cpp
--- a/Engine/Source/Runtime/Engine/Private/Streaming/StreamingTexture.cpp
+++ b/Engine/Source/Runtime/Engine/Private/Streaming/StreamingTexture.cpp
@@ -50,6 +50,7 @@ FStreamingRenderAsset::FStreamingRenderAsset(
 	bUseUnkownRefHeuristic = false;
 	NumMissingMips = 0;
 	bLooksLowRes = false;
+	bMissingTooManyMips = false;
 	VisibleWantedMips = MinAllowedMips;
 	HiddenWantedMips = MinAllowedMips;
 	RetentionPriority = 0;
@@ -487,12 +488,13 @@ int64 FStreamingRenderAsset::GetDropOneMipMemDelta() const
 	return GetSize(BudgetedMips + 1) - GetSize(BudgetedMips);
 }
 
-bool FStreamingRenderAsset::UpdateLoadOrderPriority_Async(int32 MinMipForSplitRequest)
+bool FStreamingRenderAsset::UpdateLoadOrderPriority_Async(const FRenderAssetStreamingSettings& Settings)
 {
 	LoadOrderPriority = 0;
+	bMissingTooManyMips = false;
 
 	// First load the visible mips, then later load the non visible part (does not apply to terrain textures as distance fields update may be waiting).
-	if (ResidentMips < VisibleWantedMips && VisibleWantedMips < BudgetedMips && BudgetedMips >= MinMipForSplitRequest && !bIsTerrainTexture)
+	if (ResidentMips < VisibleWantedMips && VisibleWantedMips < BudgetedMips && BudgetedMips >= Settings.MinMipForSplitRequest && !bIsTerrainTexture)
 	{
 		WantedMips = VisibleWantedMips;
 	}
@@ -509,7 +511,10 @@ bool FStreamingRenderAsset::UpdateLoadOrderPriority_Async(int32 MinMipForSplitRe
 
 		if (WantedMips > RequestedMips)
 		{
-			const bool bMipIsImportant = WantedMips - ResidentMips > (bLooksLowRes ? 1 : 2);
+			const bool bMipIsImportant = WantedMips - ResidentMips > (bLooksLowRes || IsMesh() ? 1 : 2);
+
+			// Only consider visible assets to reduce the number of high priority requests
+			bMissingTooManyMips = bIsVisible && bMipIsImportant && Settings.LowResHandlingMode != FRenderAssetStreamingSettings::LRHM_DoNothing;
 
 			if (bIsVisible)				LoadOrderPriority += 1024;
 			if (bMustLoadFirst)			LoadOrderPriority += 512;
@@ -584,7 +589,7 @@ void FStreamingRenderAsset::StreamWantedMips_Internal(FRenderAssetStreamingManag
 			}
 			else // WantedMips > ResidentMips
 			{
-				const bool bShouldPrioritizeAsyncIORequest = (bLocalForceFullyLoadHeuristic || bIsTerrainTexture || bLoadWithHigherPriority) && LocalWantedMips <= LocalVisibleWantedMips;
+				const bool bShouldPrioritizeAsyncIORequest = (bLocalForceFullyLoadHeuristic || bIsTerrainTexture || bLoadWithHigherPriority || IsMissingTooManyMips()) && LocalWantedMips <= LocalVisibleWantedMips;
 				RenderAsset->StreamIn(LocalWantedMips, bShouldPrioritizeAsyncIORequest);
 			}
 			UpdateStreamingStatus(false);
diff --git a/Engine/Source/Runtime/Engine/Private/Streaming/StreamingTexture.h b/Engine/Source/Runtime/Engine/Private/Streaming/StreamingTexture.h
--- a/Engine/Source/Runtime/Engine/Private/Streaming/StreamingTexture.h
+++ b/Engine/Source/Runtime/Engine/Private/Streaming/StreamingTexture.h
@@ -128,7 +128,7 @@ struct FStreamingRenderAsset
 	}
 
 	/** Init load order. Return wether this texture has any load/unload request */
-	bool UpdateLoadOrderPriority_Async(int32 MinMipForSplitRequest);
+	bool UpdateLoadOrderPriority_Async(const FRenderAssetStreamingSettings& Settings);
 
 	void UpdateOptionalMipsState_Async();
 	
@@ -153,6 +153,11 @@ struct FStreamingRenderAsset
 		return RenderAssetType != EStreamableRenderAssetType::Texture;
 	}
 
+	bool IsMissingTooManyMips() const
+	{
+		return bMissingTooManyMips;
+	}
+
 	FORCEINLINE int32 GetPerfectWantedMips() const { return FMath::Max<int32>(VisibleWantedMips,  HiddenWantedMips); }
 
 	// Whether this texture/mesh can be affected by Global Bias and Budget Bias per texture/mesh.
@@ -272,6 +277,8 @@ struct FStreamingRenderAsset
 	uint32			bUseUnkownRefHeuristic : 1;
 	/** (4) Same as force fully load, but takes into account component settings. */
 	uint32			bLooksLowRes : 1;
+	/** (4) Whether too many desired MIPs or LODs are missing. */
+	uint32			bMissingTooManyMips : 1;
 
 	/** (5) */
 	uint32			bCachedForceFullyLoadHeuristic : 1;
diff --git a/Engine/Source/Runtime/Engine/Private/Streaming/Texture2DMipDataProvider_IO.cpp b/Engine/Source/Runtime/Engine/Private/Streaming/Texture2DMipDataProvider_IO.cpp
--- a/Engine/Source/Runtime/Engine/Private/Streaming/Texture2DMipDataProvider_IO.cpp
+++ b/Engine/Source/Runtime/Engine/Private/Streaming/Texture2DMipDataProvider_IO.cpp
@@ -11,8 +11,6 @@ Texture2DMipDataProvider_IO.cpp : Implementation of FTextureMipDataProvider usin
 #include "Streaming/TextureStreamingHelpers.h"
 #include "ContentStreaming.h"
 
-extern int32 GStreamingTextureIOPriority;
-
 FTexture2DMipDataProvider_IO::FTexture2DMipDataProvider_IO(const UTexture* InTexture, bool InPrioritizedIORequest)
 	: FTextureMipDataProvider(InTexture, ETickState::Init, ETickThread::Async)
 	, bPrioritizedIORequest(InPrioritizedIORequest)
@@ -103,11 +101,29 @@ int32 FTexture2DMipDataProvider_IO::GetMips(
 		// but that won't do anything because the tick would not try to acquire the lock since it is already locked.
 		SyncOptions.Counter->Increment();
 
+		EAsyncIOPriorityAndFlags Priority = AIOP_Low;
+		if (bPrioritizedIORequest)
+		{
+			static IConsoleVariable* CVarAsyncLoadingPrecachePriority = IConsoleManager::Get().FindConsoleVariable(TEXT("s.AsyncLoadingPrecachePriority"));
+			const bool bLoadBeforeAsyncPrecache = CVarStreamingLowResHandlingMode.GetValueOnAnyThread() == (int32)FRenderAssetStreamingSettings::LRHM_LoadBeforeAsyncPrecache;
+
+			if (CVarAsyncLoadingPrecachePriority && bLoadBeforeAsyncPrecache)
+			{
+				const int32 AsyncIOPriority = CVarAsyncLoadingPrecachePriority->GetInt();
+				// Higher priority than regular requests but don't go over max
+				Priority = (EAsyncIOPriorityAndFlags)FMath::Clamp<int32>(AsyncIOPriority + 1, AIOP_BelowNormal, AIOP_MAX);
+			}
+			else
+			{
+				Priority = AIOP_BelowNormal;
+			}
+		}
+
 		IORequests[StartingMipIndex].BulkDataIORequest.Reset(
 			OwnerMip.BulkData.CreateStreamingRequest(
 				0,
 				OwnerMip.BulkData.GetBulkDataSize(),
-				(EAsyncIOPriorityAndFlags)FMath::Clamp<int32>(GStreamingTextureIOPriority + (bPrioritizedIORequest ? 1 : 0), AIOP_Low, AIOP_High) | AIOP_FLAG_DONTCACHE,
+				Priority | AIOP_FLAG_DONTCACHE,
 				&AsyncFileCallBack,
 				(uint8*)MipInfo.DestData
 			)
diff --git a/Engine/Source/Runtime/Engine/Private/Streaming/Texture2DStreamIn_IO.cpp b/Engine/Source/Runtime/Engine/Private/Streaming/Texture2DStreamIn_IO.cpp
--- a/Engine/Source/Runtime/Engine/Private/Streaming/Texture2DStreamIn_IO.cpp
+++ b/Engine/Source/Runtime/Engine/Private/Streaming/Texture2DStreamIn_IO.cpp
@@ -17,13 +17,6 @@ Texture2DStreamIn.cpp: Stream in helper for 2D textures using texture streaming
 #include "RenderUtils.h"
 #endif
 
-int32 GStreamingTextureIOPriority = (int32)AIOP_Low;
-static FAutoConsoleVariableRef CVarStreamingTextureIOPriority(
-	TEXT("r.Streaming.TextureIOPriority"),
-	GStreamingTextureIOPriority,
-	TEXT("Base I/O priority for loading texture mips"),
-	ECVF_Default);
-
 FTexture2DStreamIn_IO::FTexture2DStreamIn_IO(UTexture2D* InTexture, bool InPrioritizedIORequest)
 	: FTexture2DStreamIn(InTexture)
 	, bPrioritizedIORequest(InPrioritizedIORequest)
@@ -86,10 +79,28 @@ void FTexture2DStreamIn_IO::SetIORequests(const FContext& Context)
 			uint32 DestPitch = MipData[MipIndex].Pitch;
 			FTexture2DResource::WarnRequiresTightPackedMip(MipMap.SizeX, MipMap.SizeY, Context.Resource->GetPixelFormat(), DestPitch);
 
+			EAsyncIOPriorityAndFlags Priority = AIOP_Low;
+			if (bPrioritizedIORequest)
+			{
+				static IConsoleVariable* CVarAsyncLoadingPrecachePriority = IConsoleManager::Get().FindConsoleVariable(TEXT("s.AsyncLoadingPrecachePriority"));
+				const bool bLoadBeforeAsyncPrecache = CVarStreamingLowResHandlingMode.GetValueOnAnyThread() == (int32)FRenderAssetStreamingSettings::LRHM_LoadBeforeAsyncPrecache;
+
+				if (CVarAsyncLoadingPrecachePriority && bLoadBeforeAsyncPrecache)
+				{
+					const int32 AsyncIOPriority = CVarAsyncLoadingPrecachePriority->GetInt();
+					// Higher priority than regular requests but don't go over max
+					Priority = (EAsyncIOPriorityAndFlags)FMath::Clamp<int32>(AsyncIOPriority + 1, AIOP_BelowNormal, AIOP_MAX);
+				}
+				else
+				{
+					Priority = AIOP_BelowNormal;
+				}
+			}
+
 			IORequests[MipIndex] = MipMap.BulkData.CreateStreamingRequest(
 				0,
 				BulkDataSize,
-				(EAsyncIOPriorityAndFlags)FMath::Clamp<int32>(GStreamingTextureIOPriority + (bPrioritizedIORequest ? 1 : 0), AIOP_Low, AIOP_High) | AIOP_FLAG_DONTCACHE,
+				Priority | AIOP_FLAG_DONTCACHE,
 				&AsyncFileCallBack,
 				(uint8*)MipData[MipIndex].Data);
 		}
diff --git a/Engine/Source/Runtime/Engine/Private/Streaming/TextureStreamingHelpers.cpp b/Engine/Source/Runtime/Engine/Private/Streaming/TextureStreamingHelpers.cpp
--- a/Engine/Source/Runtime/Engine/Private/Streaming/TextureStreamingHelpers.cpp
+++ b/Engine/Source/Runtime/Engine/Private/Streaming/TextureStreamingHelpers.cpp
@@ -267,6 +267,12 @@ ENGINE_API TAutoConsoleVariable<int32> CVarFramesForFullUpdate(
 	5,
 	TEXT("Texture streaming is time sliced per frame. This values gives the number of frames to visit all textures."));
 
+TAutoConsoleVariable<int32> CVarStreamingLowResHandlingMode(
+	TEXT("r.Streaming.LowResHandlingMode"),
+	(int32)FRenderAssetStreamingSettings::LRHM_DoNothing,
+	TEXT("How to handle assets with too many missing MIPs or LODs. 0 (default): do nothing, 1: load before regular streaming requests, 2: load before async loading precache requests."),
+	ECVF_Default);
+
 static TAutoConsoleVariable<int32> CVarPrioritizeMeshLODRetention(
 	TEXT("r.Streaming.PrioritizeMeshLODRetention"),
 	1,
@@ -320,6 +326,7 @@ void FRenderAssetStreamingSettings::Update()
 	MaxTextureUVDensity = CVarStreamingMaxTextureUVDensity.GetValueOnAnyThread();
 	bUseMaterialData = bUseNewMetrics && CVarStreamingUseMaterialData.GetValueOnAnyThread() != 0;
 	HiddenPrimitiveScale = bUseNewMetrics ? CVarStreamingHiddenPrimitiveScale.GetValueOnAnyThread() : 1.f;
+	LowResHandlingMode = (ELowResHandlingMode)CVarStreamingLowResHandlingMode.GetValueOnAnyThread();
 	bMipCalculationEnablePerLevelList = CVarStreamingMipCalculationEnablePerLevelList.GetValueOnAnyThread() != 0;
 	bPrioritizeMeshLODRetention = CVarPrioritizeMeshLODRetention.GetValueOnAnyThread() != 0;
 	VRAMPercentageClamp = CVarStreamingVRAMPercentageClamp.GetValueOnAnyThread();
diff --git a/Engine/Source/Runtime/Engine/Private/Streaming/TextureStreamingHelpers.h b/Engine/Source/Runtime/Engine/Private/Streaming/TextureStreamingHelpers.h
--- a/Engine/Source/Runtime/Engine/Private/Streaming/TextureStreamingHelpers.h
+++ b/Engine/Source/Runtime/Engine/Private/Streaming/TextureStreamingHelpers.h
@@ -66,9 +66,18 @@ extern TAutoConsoleVariable<int32> CVarStreamingUseMaterialData;
 extern TAutoConsoleVariable<int32> CVarStreamingNumStaticComponentsProcessedPerFrame;
 extern TAutoConsoleVariable<int32> CVarStreamingDefragDynamicBounds;
 extern TAutoConsoleVariable<float> CVarStreamingMaxTextureUVDensity;
+extern TAutoConsoleVariable<int32> CVarStreamingLowResHandlingMode;
 
 struct FRenderAssetStreamingSettings
 {
+	// How to handle assets with too many missing MIPs or LODs
+	enum ELowResHandlingMode
+	{
+		LRHM_DoNothing,
+		LRHM_LoadBeforeRegular,			// Use higher IO priority than regular streaming requests
+		LRHM_LoadBeforeAsyncPrecache,	// Also ensure that priority is higher than async loading precache requests
+	};
+
 	FRenderAssetStreamingSettings()
 	{
 		// Make sure padding bytes don't have random values
@@ -104,6 +113,7 @@ struct FRenderAssetStreamingSettings
 	float MaxTextureUVDensity;
 	int32 MaterialQualityLevel;
 	int32 FramesForFullUpdate;
+	ELowResHandlingMode LowResHandlingMode;
 	bool bMipCalculationEnablePerLevelList;
 	bool bPrioritizeMeshLODRetention;
 	int32 VRAMPercentageClamp;
