Update FakeStereoRenderingDevice resolutions based on final ViewRects.

The GSystemResolution may be affected by cached values from previous runs, which will result in unexpected stereo-emulated aspect ratios. It seems like the safer option is to update the height and width based on the final view rect (which may also be better in cases where the window is resized, etc).

diff --git a/Engine/Source/Runtime/Engine/Private/UnrealEngine.cpp b/Engine/Source/Runtime/Engine/Private/UnrealEngine.cpp
--- a/Engine/Source/Runtime/Engine/Private/UnrealEngine.cpp
+++ b/Engine/Source/Runtime/Engine/Private/UnrealEngine.cpp
@@ -3696,6 +3696,32 @@ public:
 		Y += SizeY * (ViewIndex / 2);
 	}
 
+	virtual void SetFinalViewRect(class FRHICommandListImmediate& RHICmdList, const int32 ViewIndex, const FIntRect& FinalViewRect) override
+	{
+		static TAutoConsoleVariable<int32> CVarEmulateStereoWidth(TEXT("r.StereoEmulationWidth"), 0, TEXT("Width of the imaginable HMD for stereo emulation"), ECVF_ReadOnly);
+		static TAutoConsoleVariable<int32> CVarEmulateStereoHeight(TEXT("r.StereoEmulationHeight"), 0, TEXT("Height of the imaginable HMD for stereo emulation"), ECVF_ReadOnly);
+		int32 W = CVarEmulateStereoWidth.GetValueOnAnyThread();
+		int32 H = CVarEmulateStereoHeight.GetValueOnAnyThread();
+		
+		bool bForcedStereoRes = false;
+		if (W != 0)
+		{
+			Width = FMath::Clamp(W, 100, 10000);
+			bForcedStereoRes = true;
+		}
+		if (H != 0)
+		{
+			Height = FMath::Clamp(H, 100, 10000);
+			bForcedStereoRes = true;
+		}
+
+		if (ViewIndex == 0 && !bForcedStereoRes)
+		{
+			Width = FinalViewRect.Width();
+			Height = FinalViewRect.Height();
+		}
+	}
+
 	virtual void CalculateStereoViewOffset(const int32 ViewIndex, FRotator& ViewRotation, const float WorldToMeters, FVector& ViewLocation) override
 	{
 		// 32mm, 1/2 average interpupillary distance
