WorldPartition : Fixed lost child actors at cook when moving actor in world partition cell level.

diff --git a/Engine/Source/Runtime/Engine/Private/WorldPartition/WorldPartitionLevelHelper.cpp b/Engine/Source/Runtime/Engine/Private/WorldPartition/WorldPartitionLevelHelper.cpp
--- a/Engine/Source/Runtime/Engine/Private/WorldPartition/WorldPartitionLevelHelper.cpp
+++ b/Engine/Source/Runtime/Engine/Private/WorldPartition/WorldPartitionLevelHelper.cpp
@@ -194,22 +194,42 @@ void FWorldPartitionLevelHelper::MoveExternalActorsToLevel(const TArray<FWorldPa
 			{
 				Actor->Rename(nullptr, InLevel, REN_ForceNoResetLoaders);
 			}
-			
-			check(Actor->GetPackage() == LevelPackage);
-			if (bSameOuter && !InLevel->Actors.Contains(Actor))
+			else if (!InLevel->Actors.Contains(Actor))
 			{
 				InLevel->AddLoadedActor(Actor);
 			}
+			check(Actor->GetPackage() == LevelPackage);
 
-			// Include objects found in the source actor package in the destination level package
+			// Process objects found in the source actor package
 			TArray<UObject*> Objects;
 			const bool bIncludeNestedSubobjects = false;
-			GetObjectsWithOuter(ActorExternalPackage, Objects, bIncludeNestedSubobjects);
+			GetObjectsWithPackage(ActorExternalPackage, Objects, bIncludeNestedSubobjects);
 			for (UObject* Object : Objects)
 			{
 				if (Object->GetFName() != NAME_PackageMetaData)
 				{
-					Object->Rename(nullptr, LevelPackage, REN_ForceNoResetLoaders);
+					if (Object->GetOuter()->IsA<ULevel>())
+					{
+						// Move objects that are outered the level in the destination level
+						AActor* NestedActor = Cast<AActor>(Object);
+						if (InLevel != Object->GetOuter())
+						{
+							Object->Rename(nullptr, InLevel, REN_ForceNoResetLoaders);
+						}
+						else if (NestedActor && !InLevel->Actors.Contains(NestedActor))
+						{
+							InLevel->AddLoadedActor(NestedActor);
+						}
+						if (NestedActor)
+						{
+							LevelActors.Add(NestedActor->GetFName());
+						}
+					}
+					else
+					{
+						// Move objects in the destination level package
+						Object->Rename(nullptr, LevelPackage, REN_ForceNoResetLoaders);
+					}
 				}
 			}
 
