fix outdated editcondition type for ground clipping setting in merge actors tool's approximate actors mode

diff --git a/Engine/Source/Runtime/Engine/Classes/Engine/MeshMerging.h b/Engine/Source/Runtime/Engine/Classes/Engine/MeshMerging.h
--- a/Engine/Source/Runtime/Engine/Classes/Engine/MeshMerging.h
+++ b/Engine/Source/Runtime/Engine/Classes/Engine/MeshMerging.h
@@ -846,7 +846,7 @@ struct FMeshApproximationSettings
 	EMeshApproximationGroundPlaneClippingPolicy GroundClipping = EMeshApproximationGroundPlaneClippingPolicy::NoGroundClipping;
 
 	/** Z-Height for the ground clipping plane, if enabled */
-	UPROPERTY(BlueprintReadWrite, EditAnywhere, Category = SimplifySettings, meta = (EditCondition = "GroundClipping != EGroundPlaneClippingPolicy::NoGroundClipping"))
+	UPROPERTY(BlueprintReadWrite, EditAnywhere, Category = SimplifySettings, meta = (EditCondition = "GroundClipping != EMeshApproximationGroundPlaneClippingPolicy::NoGroundClipping"))
 	float GroundClippingZHeight = 0.0f;
 
 
