Substrate - updated opaque rough refraction classification text.

diff --git a/Engine/Source/Runtime/Engine/Private/StrataVisualizationData.cpp b/Engine/Source/Runtime/Engine/Private/StrataVisualizationData.cpp
--- a/Engine/Source/Runtime/Engine/Private/StrataVisualizationData.cpp
+++ b/Engine/Source/Runtime/Engine/Private/StrataVisualizationData.cpp
@@ -120,9 +120,9 @@ void FStrataVisualizationData::Initialize()
 
 		AddVisualizationMode(
 			AllModeMap,
-			TEXT("RoughRefractionClassification"),
-			LOCTEXT("RoughRefractionClassification", "Rough Refraction Classification"),
-			LOCTEXT("RoughRefractionClassificationDesc", "Visualizes Substrate rough refraction classification"),
+			TEXT("OpaqueRoughRefractionClassification"),
+			LOCTEXT("OpaqueRoughRefractionClassification", "Opaque Rough Refraction Classification"),
+			LOCTEXT("OpaqueRoughRefractionClassificationDesc", "Visualizes Substrate Opaque Rough Refraction Classification"),
 			FStrataViewMode::RoughRefractionClassification,
 			true,
 			Strata::IsOpaqueRoughRefractionEnabled(),
