Remove Blueprint asset sanitization from the save path to avoid invalidating harvested exports in a restricted editor mode, and alleviate BP compile/GC cost when multiple assets are involved.

Also defer the GC pass after the post-sanitized BP compile since the editor utility API is typically going to be invoked on multiple assets in sequence as part of a larger operation.

diff --git a/Engine/Source/Editor/UnrealEd/Private/Kismet2/BlueprintEditorUtils.cpp b/Engine/Source/Editor/UnrealEd/Private/Kismet2/BlueprintEditorUtils.cpp
--- a/Engine/Source/Editor/UnrealEd/Private/Kismet2/BlueprintEditorUtils.cpp
+++ b/Engine/Source/Editor/UnrealEd/Private/Kismet2/BlueprintEditorUtils.cpp
@@ -10471,7 +10471,7 @@ void FBlueprintEditorUtils::SanitizeRestrictedContent(UBlueprint* BP)
 		}
 
 		// Do a full recompile to bring everything back up-to-date and to sanitize any generated bytecode.
-		FKismetEditorUtilities::CompileBlueprint(BP);
+		FKismetEditorUtilities::CompileBlueprint(BP, EBlueprintCompileOptions::SkipGarbageCollection);
 	}
 }
 
diff --git a/Engine/Source/Runtime/Engine/Private/Blueprint.cpp b/Engine/Source/Runtime/Engine/Private/Blueprint.cpp
--- a/Engine/Source/Runtime/Engine/Private/Blueprint.cpp
+++ b/Engine/Source/Runtime/Engine/Private/Blueprint.cpp
@@ -408,11 +408,6 @@ void UBlueprint::PreSave(FObjectPreSaveContext ObjectSaveContext)
 		// Cache the BP for use (immediate, since we're about to save)
 		FFindInBlueprintSearchManager::Get().AddOrUpdateBlueprintSearchMetadata(this, Flags, OverrideVersion);
 	}
-
-#if WITH_EDITOR
-	// Sanitize the Blueprint by clearing out any objects/types that are restricted within the current editor context.
-	FBlueprintEditorUtils::SanitizeRestrictedContent(this);
-#endif	// WITH_EDITOR
 }
 
 void UBlueprint::GetPreloadDependencies(TArray<UObject*>& OutDeps)
