Add NaniteOverrideMaterial to HLOD CRC calculations.
This will track any changes to NaniteOverrideMaterials for rebuilding HLOD.
All current HLODs using NaniteOverrideMaterials will be rebuilt.

diff --git a/Engine/Source/Runtime/Engine/Private/WorldPartition/HLOD/HLODBuilder.cpp b/Engine/Source/Runtime/Engine/Private/WorldPartition/HLOD/HLODBuilder.cpp
--- a/Engine/Source/Runtime/Engine/Private/WorldPartition/HLOD/HLODBuilder.cpp
+++ b/Engine/Source/Runtime/Engine/Private/WorldPartition/HLOD/HLODBuilder.cpp
@@ -93,6 +93,20 @@ uint32 UHLODBuilder::ComputeHLODHash(const UActorComponent* InSourceComponent) c
 					UE_LOG(LogHLODBuilder, VeryVerbose, TEXT("     - Texture (%s) = %x"), *Texture->GetName(), ComponentCRC);
 				}
 			}
+			UMaterialInterface* NaniteOverride = MaterialInterface ? MaterialInterface->GetNaniteOverride() : nullptr;
+			if (NaniteOverride)
+			{
+				ComponentCRC = UHLODProxy::GetCRC(NaniteOverride, ComponentCRC);
+				UE_LOG(LogHLODBuilder, VeryVerbose, TEXT("     - Material (%s) = %x"), *NaniteOverride->GetName(), ComponentCRC);
+
+				TArray<UTexture*> Textures;
+				NaniteOverride->GetUsedTextures(Textures, EMaterialQualityLevel::High, true, ERHIFeatureLevel::SM5, true);
+				for (UTexture* Texture : Textures)
+				{
+					ComponentCRC = UHLODProxy::GetCRC(Texture, ComponentCRC);
+					UE_LOG(LogHLODBuilder, VeryVerbose, TEXT("     - Texture (%s) = %x"), *Texture->GetName(), ComponentCRC);
+				}
+			}
 		}
 	}
 	else
