Substrate - fixed material conversion issues.
Legacy ClearCoat and Cloth material were no longer converted correctly.

diff --git a/Engine/Source/Runtime/Engine/Classes/Engine/EngineTypes.h b/Engine/Source/Runtime/Engine/Classes/Engine/EngineTypes.h
--- a/Engine/Source/Runtime/Engine/Classes/Engine/EngineTypes.h
+++ b/Engine/Source/Runtime/Engine/Classes/Engine/EngineTypes.h
@@ -693,6 +693,8 @@ enum EStrataShadingModel : int
 	SSM_VolumetricFogCloud		UMETA(DisplayName = "VolumetricFogCloud"),
 	SSM_Hair					UMETA(DisplayName = "Hair"),
 	SSM_Eye						UMETA(DisplayName = "Eye"),
+	SSM_Cloth					UMETA(DisplayName = "Cloth"),
+	SSM_ClearCoat				UMETA(DisplayName = "ClearCoat"),
 	SSM_SingleLayerWater		UMETA(DisplayName = "SingleLayerWater"),
 	SSM_LightFunction			UMETA(DisplayName = "LightFunction"),
 	SSM_PostProcess				UMETA(DisplayName = "PostProcess"),
diff --git a/Engine/Source/Runtime/Engine/Private/Materials/Material.cpp b/Engine/Source/Runtime/Engine/Private/Materials/Material.cpp
--- a/Engine/Source/Runtime/Engine/Private/Materials/Material.cpp
+++ b/Engine/Source/Runtime/Engine/Private/Materials/Material.cpp
@@ -3118,11 +3118,11 @@ static void AddStrataShadingModelFromMaterialShadingModel(FStrataMaterialInfo& O
 	if (InShadingModels.HasShadingModel(MSM_DefaultLit))		{ OutInfo.AddShadingModel(EStrataShadingModel::SSM_DefaultLit); }
 	if (InShadingModels.HasShadingModel(MSM_Subsurface))		{ OutInfo.AddShadingModel(EStrataShadingModel::SSM_SubsurfaceLit); }
 	if (InShadingModels.HasShadingModel(MSM_PreintegratedSkin))	{ OutInfo.AddShadingModel(EStrataShadingModel::SSM_SubsurfaceLit); }
-	if (InShadingModels.HasShadingModel(MSM_ClearCoat))			{ OutInfo.AddShadingModel(EStrataShadingModel::SSM_DefaultLit); }
+	if (InShadingModels.HasShadingModel(MSM_ClearCoat))			{ OutInfo.AddShadingModel(EStrataShadingModel::SSM_ClearCoat); }
 	if (InShadingModels.HasShadingModel(MSM_SubsurfaceProfile))	{ OutInfo.AddShadingModel(EStrataShadingModel::SSM_SubsurfaceLit); }
 	if (InShadingModels.HasShadingModel(MSM_TwoSidedFoliage))	{ OutInfo.AddShadingModel(EStrataShadingModel::SSM_SubsurfaceLit); }
 	if (InShadingModels.HasShadingModel(MSM_Hair))				{ OutInfo.AddShadingModel(EStrataShadingModel::SSM_Hair); }
-	if (InShadingModels.HasShadingModel(MSM_Cloth))				{ OutInfo.AddShadingModel(EStrataShadingModel::SSM_DefaultLit); }
+	if (InShadingModels.HasShadingModel(MSM_Cloth))				{ OutInfo.AddShadingModel(EStrataShadingModel::SSM_Cloth); }
 	if (InShadingModels.HasShadingModel(MSM_Eye))				{ OutInfo.AddShadingModel(EStrataShadingModel::SSM_Eye); }
 	if (InShadingModels.HasShadingModel(MSM_SingleLayerWater))	{ OutInfo.AddShadingModel(EStrataShadingModel::SSM_SingleLayerWater); }
 	if (InShadingModels.HasShadingModel(MSM_ThinTranslucent))	{ OutInfo.AddShadingModel(EStrataShadingModel::SSM_DefaultLit); }
@@ -5010,6 +5010,16 @@ void UMaterial::RebuildShadingModelField()
 				MaterialDomain = EMaterialDomain::MD_Surface;
 				ShadingModel = MSM_Eye;
 			}
+			else if (StrataMaterialInfo.HasOnlyShadingModel(SSM_Cloth))
+			{
+				MaterialDomain = EMaterialDomain::MD_Surface;
+				ShadingModel = MSM_Cloth;
+			}
+			else if (StrataMaterialInfo.HasOnlyShadingModel(SSM_ClearCoat))
+			{
+				MaterialDomain = EMaterialDomain::MD_Surface;
+				ShadingModel = MSM_ClearCoat;
+			}
 			else if (StrataMaterialInfo.HasOnlyShadingModel(SSM_SingleLayerWater))
 			{
 				MaterialDomain = EMaterialDomain::MD_Surface;
diff --git a/Engine/Source/Runtime/Engine/Private/Materials/MaterialExpressions.cpp b/Engine/Source/Runtime/Engine/Private/Materials/MaterialExpressions.cpp
--- a/Engine/Source/Runtime/Engine/Private/Materials/MaterialExpressions.cpp
+++ b/Engine/Source/Runtime/Engine/Private/Materials/MaterialExpressions.cpp
@@ -23416,6 +23416,8 @@ void UMaterialExpressionStrataLegacyConversion::GatherStrataMaterialInfo(FStrata
 	if (ConvertedStrataMaterialInfo.HasShadingModel(SSM_VolumetricFogCloud))	{ StrataMaterialInfo.AddShadingModel(SSM_VolumetricFogCloud); }
 	if (ConvertedStrataMaterialInfo.HasShadingModel(SSM_Hair))					{ StrataMaterialInfo.AddShadingModel(SSM_Hair); }
 	if (ConvertedStrataMaterialInfo.HasShadingModel(SSM_Eye))					{ StrataMaterialInfo.AddShadingModel(SSM_Eye); }
+	if (ConvertedStrataMaterialInfo.HasShadingModel(SSM_Cloth))					{ StrataMaterialInfo.AddShadingModel(SSM_Cloth); }
+	if (ConvertedStrataMaterialInfo.HasShadingModel(SSM_ClearCoat))				{ StrataMaterialInfo.AddShadingModel(SSM_ClearCoat); }
 	if (ConvertedStrataMaterialInfo.HasShadingModel(SSM_SingleLayerWater))		{ StrataMaterialInfo.AddShadingModel(SSM_SingleLayerWater); }
 	if (SubsurfaceProfile)														{ StrataMaterialInfo.AddSubsurfaceProfile(SubsurfaceProfile); }
 	if (ConvertedStrataMaterialInfo.HasShadingModelFromExpression())			{ StrataMaterialInfo.SetShadingModelFromExpression(true); }
@@ -23426,13 +23428,13 @@ FStrataOperator* UMaterialExpressionStrataLegacyConversion::StrataGenerateMateri
 	// Note Thickness has no meaning/usage in the context of StrataLegacyConversionNode
 	int32 ThicknessIndex = Compiler->StrataThicknessStackGetThicknessIndex();
 
-	auto AddDefaultWorstCase = [&](bool bSSS)
+	auto AddDefaultWorstCase = [&](bool bSSS, bool bFuzz)
 	{
 		FStrataOperator& SlabOperator = Compiler->StrataCompilationRegisterOperator(STRATA_OPERATOR_BSDF_LEGACY, Compiler->StrataTreeStackGetPathUniqueId(), this, Parent, Compiler->StrataTreeStackGetParentPathUniqueId());
 		SlabOperator.BSDFType = STRATA_BSDF_TYPE_SLAB;
-		SlabOperator.bBSDFHasSSS = true;
-		SlabOperator.bBSDFHasMFPPluggedIn = true;
-		SlabOperator.bBSDFHasFuzz = true;
+		SlabOperator.bBSDFHasSSS = bSSS;
+		SlabOperator.bBSDFHasMFPPluggedIn = bSSS;
+		SlabOperator.bBSDFHasFuzz = bFuzz;
 		SlabOperator.bBSDFHasAnisotropy = Anisotropy.IsConnected();
 		SlabOperator.ThicknessIndex = ThicknessIndex;
 
@@ -23442,7 +23444,7 @@ FStrataOperator* UMaterialExpressionStrataLegacyConversion::StrataGenerateMateri
 	// Logic about shading models and complexity should match UMaterialExpressionStrataLegacyConversion::Compile.
 	if (ConvertedStrataMaterialInfo.CountShadingModels() > 1 || ConvertedStrataMaterialInfo.HasShadingModelFromExpression())
 	{
-		return AddDefaultWorstCase(true);
+		return AddDefaultWorstCase(true, true);
 	}
 	// else
 	{
@@ -23457,11 +23459,24 @@ FStrataOperator* UMaterialExpressionStrataLegacyConversion::StrataGenerateMateri
 		}
 		else if (ConvertedStrataMaterialInfo.HasShadingModel(SSM_DefaultLit))
 		{
-			return AddDefaultWorstCase(false);
+			return AddDefaultWorstCase(false, false);
 		}
 		else if (ConvertedStrataMaterialInfo.HasShadingModel(SSM_SubsurfaceLit))
 		{
-			return AddDefaultWorstCase(true);
+			return AddDefaultWorstCase(true, false);
+		}
+		else if (ConvertedStrataMaterialInfo.HasShadingModel(SSM_Cloth))
+		{
+			return AddDefaultWorstCase(false, true);
+		}
+		else if (ConvertedStrataMaterialInfo.HasShadingModel(SSM_ClearCoat))
+		{
+			FStrataOperator& Operator = Compiler->StrataCompilationRegisterOperator(STRATA_OPERATOR_BSDF_LEGACY, Compiler->StrataTreeStackGetPathUniqueId(), this, Parent, Compiler->StrataTreeStackGetParentPathUniqueId());
+			Operator.BSDFType = STRATA_BSDF_TYPE_SLAB;
+			Operator.ThicknessIndex = ThicknessIndex;
+			Operator.bBSDFHasSecondRoughnessOrSimpleClearCoat = true;
+			Operator.bBSDFHasAnisotropy = Anisotropy.IsConnected();
+			return &Operator;
 		}
 		else if (ConvertedStrataMaterialInfo.HasShadingModel(SSM_Hair))
 		{
