Allow cheat manager extensions to respond to adding/removing from cheat manager in native code, as well as BP.

diff --git a/Engine/Source/Runtime/Engine/Classes/GameFramework/CheatManager.h b/Engine/Source/Runtime/Engine/Classes/GameFramework/CheatManager.h
--- a/Engine/Source/Runtime/Engine/Classes/GameFramework/CheatManager.h
+++ b/Engine/Source/Runtime/Engine/Classes/GameFramework/CheatManager.h
@@ -79,11 +79,13 @@ public:
 	UFUNCTION(BlueprintPure, Category = "Cheat Manager")
 	ENGINE_API APlayerController* GetPlayerController() const;
 
-	UFUNCTION(BlueprintImplementableEvent, meta=(DisplayName = "Added To Cheat Manager", Keywords="Begin Play"))
+	UFUNCTION(BlueprintNativeEvent, meta=(DisplayName = "Added To Cheat Manager", Keywords="Begin Play"))
 	ENGINE_API void AddedToCheatManager();
+	ENGINE_API virtual void AddedToCheatManager_Implementation();
 
-	UFUNCTION(BlueprintImplementableEvent, meta = (DisplayName = "Removed From Cheat Manager", Keywords = "End Play"))
+	UFUNCTION(BlueprintNativeEvent, meta = (DisplayName = "Removed From Cheat Manager", Keywords = "End Play"))
 	ENGINE_API void RemovedFromCheatManager();
+	ENGINE_API virtual void RemovedFromCheatManager_Implementation();
 
 	ENGINE_API virtual void DoExtensionSpecificBugItLog(FOutputDevice& OutputFile);
 };
diff --git a/Engine/Source/Runtime/Engine/Private/CheatManager.cpp b/Engine/Source/Runtime/Engine/Private/CheatManager.cpp
--- a/Engine/Source/Runtime/Engine/Private/CheatManager.cpp
+++ b/Engine/Source/Runtime/Engine/Private/CheatManager.cpp
@@ -55,6 +55,14 @@ APlayerController* UCheatManagerExtension::GetPlayerController() const
 	return GetOuterUCheatManager()->GetPlayerController();
 }
 
+void UCheatManagerExtension::AddedToCheatManager_Implementation()
+{
+}
+
+void UCheatManagerExtension::RemovedFromCheatManager_Implementation()
+{
+}
+
 void UCheatManagerExtension::DoExtensionSpecificBugItLog(FOutputDevice& OutputFile)
 {
 }
