Properly enable GPUScene support for moble on GeometryCollection VF

diff --git a/Engine/Source/Runtime/Engine/Private/GeometryCollectionRendering.cpp b/Engine/Source/Runtime/Engine/Private/GeometryCollectionRendering.cpp
--- a/Engine/Source/Runtime/Engine/Private/GeometryCollectionRendering.cpp
+++ b/Engine/Source/Runtime/Engine/Private/GeometryCollectionRendering.cpp
@@ -55,7 +55,7 @@ void FGeometryCollectionVertexFactory::ModifyCompilationEnvironment(const FVerte
 {
 	const FStaticFeatureLevel MaxSupportedFeatureLevel = GetMaxSupportedFeatureLevel(Parameters.Platform);
 	// TODO: support GPUScene on mobile
-	const bool bUseGPUScene = UseGPUScene(Parameters.Platform, MaxSupportedFeatureLevel) && (MaxSupportedFeatureLevel > ERHIFeatureLevel::ES3_1);
+	const bool bUseGPUScene = UseGPUScene(Parameters.Platform, MaxSupportedFeatureLevel);
 	const bool bSupportsPrimitiveIdStream = Parameters.VertexFactoryType->SupportsPrimitiveIdStream();
 
 	OutEnvironment.SetDefine(TEXT("VF_SUPPORTS_PRIMITIVE_SCENE_DATA"), bSupportsPrimitiveIdStream && bUseGPUScene);
@@ -75,7 +75,6 @@ void FGeometryCollectionVertexFactory::ValidateCompiledResult(const FVertexFacto
 {
 	if (Type->SupportsPrimitiveIdStream()
 		&& UseGPUScene(Platform, GetMaxSupportedFeatureLevel(Platform))
-		&& !IsMobilePlatform(Platform) // On mobile VS may use PrimtiveUB while GPUScene is enabled
 		&& ParameterMap.ContainsParameterAllocation(FPrimitiveUniformShaderParameters::FTypeInfo::GetStructMetadata()->GetShaderVariableName()))
 	{
 		OutErrors.AddUnique(*FString::Printf(TEXT("Shader attempted to bind the Primitive uniform buffer even though Vertex Factory %s computes a PrimitiveId per-instance.  This will break auto-instancing.  Shaders should use GetPrimitiveData(Parameters).Member instead of Primitive.Member."), Type->GetName()));
