Added DataLayerInstanceWithAssetClass to allow licensee to provide a custom class.

diff --git a/Engine/Source/Runtime/Engine/Private/WorldPartition/DataLayer/DataLayerInstanceWithAsset.cpp b/Engine/Source/Runtime/Engine/Private/WorldPartition/DataLayer/DataLayerInstanceWithAsset.cpp
--- a/Engine/Source/Runtime/Engine/Private/WorldPartition/DataLayer/DataLayerInstanceWithAsset.cpp
+++ b/Engine/Source/Runtime/Engine/Private/WorldPartition/DataLayer/DataLayerInstanceWithAsset.cpp
@@ -25,6 +25,11 @@ FName UDataLayerInstanceWithAsset::MakeName(const UDataLayerAsset* DeprecatedDat
 	return FName(FString::Format(TEXT("DataLayer_{0}"), { FGuid::NewGuid().ToString() }));
 }
 
+TSubclassOf<UDataLayerInstanceWithAsset> UDataLayerInstanceWithAsset::GetDataLayerInstanceClass()
+{
+	return UDataLayerManager::GetDataLayerInstanceWithAssetClass();
+}
+
 void UDataLayerInstanceWithAsset::OnCreated(const UDataLayerAsset* Asset)
 {
 	check(!GetOuterWorldDataLayers()->HasDeprecatedDataLayers() || IsRunningCommandlet());
diff --git a/Engine/Source/Runtime/Engine/Private/WorldPartition/DataLayer/DataLayerManager.cpp b/Engine/Source/Runtime/Engine/Private/WorldPartition/DataLayer/DataLayerManager.cpp
--- a/Engine/Source/Runtime/Engine/Private/WorldPartition/DataLayer/DataLayerManager.cpp
+++ b/Engine/Source/Runtime/Engine/Private/WorldPartition/DataLayer/DataLayerManager.cpp
@@ -607,12 +607,23 @@ TSubclassOf<UDataLayerLoadingPolicy> UDataLayerManager::GetDataLayerLoadingPolic
 	}
 	if (!DataLayerLoadingPolicyClassValue)
 	{
-		// Default on UDataLayerLoadingPolicy class
+		// Default to UDataLayerLoadingPolicy class
 		DataLayerLoadingPolicyClassValue = UDataLayerLoadingPolicy::StaticClass();
 	}
 	return DataLayerLoadingPolicyClassValue;
 }
 
+TSubclassOf<UDataLayerInstanceWithAsset> UDataLayerManager::GetDataLayerInstanceWithAssetClass()
+{
+	UClass* DataLayerLoadingPolicyClassValue = GetDefault<UDataLayerManager>()->DataLayerInstanceWithAssetClass.Get();
+	if (!DataLayerLoadingPolicyClassValue)
+	{
+		// Default to UDataLayerInstanceWithAsset class
+		DataLayerLoadingPolicyClassValue = UDataLayerInstanceWithAsset::StaticClass();
+	}
+	return DataLayerLoadingPolicyClassValue;
+}
+
 void UDataLayerManager::UpdateDataLayerEditorPerProjectUserSettings() const
 {
 	TArray<FName> DataLayersNotLoadedInEditor;
diff --git a/Engine/Source/Runtime/Engine/Public/WorldPartition/DataLayer/DataLayerInstanceWithAsset.h b/Engine/Source/Runtime/Engine/Public/WorldPartition/DataLayer/DataLayerInstanceWithAsset.h
--- a/Engine/Source/Runtime/Engine/Public/WorldPartition/DataLayer/DataLayerInstanceWithAsset.h
+++ b/Engine/Source/Runtime/Engine/Public/WorldPartition/DataLayer/DataLayerInstanceWithAsset.h
@@ -16,6 +16,7 @@ class UDataLayerInstanceWithAsset : public UDataLayerInstance
 public:
 #if WITH_EDITOR
 	static ENGINE_API FName MakeName(const UDataLayerAsset* DeprecatedDataLayer);
+	static ENGINE_API TSubclassOf<UDataLayerInstanceWithAsset> GetDataLayerInstanceClass();
 	ENGINE_API void OnCreated(const UDataLayerAsset* Asset);
 
 	ENGINE_API virtual bool CanEditChange(const FProperty* InProperty) const;
diff --git a/Engine/Source/Runtime/Engine/Public/WorldPartition/DataLayer/DataLayerManager.h b/Engine/Source/Runtime/Engine/Public/WorldPartition/DataLayer/DataLayerManager.h
--- a/Engine/Source/Runtime/Engine/Public/WorldPartition/DataLayer/DataLayerManager.h
+++ b/Engine/Source/Runtime/Engine/Public/WorldPartition/DataLayer/DataLayerManager.h
@@ -171,6 +171,8 @@ private:
 	ENGINE_API bool ResolveIsLoadedInEditor(const TArray<FName>& InDataLayerInstanceNames) const;
 	//~ End Editor loading
 
+	ENGINE_API static TSubclassOf<UDataLayerInstanceWithAsset> GetDataLayerInstanceWithAssetClass();
+
 	// Helper
 	ENGINE_API TArray<const UDataLayerInstance*> GetRuntimeDataLayerInstances(const TArray<FName>& InDataLayerInstanceNames) const;
 
@@ -208,6 +210,9 @@ private:
 	UPROPERTY(Config)
 	TSoftClassPtr<UDataLayerLoadingPolicy> DataLayerLoadingPolicyClass;
 
+	UPROPERTY(Config)
+	TSoftClassPtr<UDataLayerInstanceWithAsset> DataLayerInstanceWithAssetClass;
+
 	UPROPERTY(Transient)
 	TObjectPtr<UDataLayerLoadingPolicy> DataLayerLoadingPolicy;
 
diff --git a/Engine/Source/Runtime/Engine/Public/WorldPartition/DataLayer/WorldDataLayers.h b/Engine/Source/Runtime/Engine/Public/WorldPartition/DataLayer/WorldDataLayers.h
--- a/Engine/Source/Runtime/Engine/Public/WorldPartition/DataLayer/WorldDataLayers.h
+++ b/Engine/Source/Runtime/Engine/Public/WorldPartition/DataLayer/WorldDataLayers.h
@@ -255,12 +255,22 @@ DEFINE_ACTORDESC_TYPE(AWorldDataLayers, FWorldDataLayersActorDesc);
 
 #if WITH_EDITOR
 
+template <typename T, typename = int>
+struct HasDataLayerInstanceClass : std::false_type {};
+
+template <typename T>
+struct HasDataLayerInstanceClass<T, decltype(&T::GetDataLayerInstanceClass, 0)> : std::true_type {};
+
 template<class DataLayerInstanceType, typename ...CreationsArgs>
 DataLayerInstanceType* AWorldDataLayers::CreateDataLayer(CreationsArgs... InCreationArgs)
 {
 	Modify();
 
-	DataLayerInstanceType* NewDataLayer = NewObject<DataLayerInstanceType>(this, DataLayerInstanceType::MakeName(Forward<CreationsArgs>(InCreationArgs)...), RF_Transactional | RF_NoFlags);
+	UClass* ClassToUse = DataLayerInstanceType::StaticClass();
+	if constexpr (HasDataLayerInstanceClass<DataLayerInstanceType>::value)
+		ClassToUse = DataLayerInstanceType::GetDataLayerInstanceClass();
+
+	DataLayerInstanceType* NewDataLayer = NewObject<DataLayerInstanceType>(this, ClassToUse, DataLayerInstanceType::MakeName(Forward<CreationsArgs>(InCreationArgs)...), RF_Transactional | RF_NoFlags);
 	check(NewDataLayer != nullptr);
 	NewDataLayer->OnCreated(Forward<CreationsArgs>(InCreationArgs)...);
 	DataLayerInstances.Add(NewDataLayer);
