Fixed Opus decoding and seeking for the non-streaming case

#rnx

diff --git a/Engine/Source/Runtime/Engine/Private/OpusAudioInfo.cpp b/Engine/Source/Runtime/Engine/Private/OpusAudioInfo.cpp
--- a/Engine/Source/Runtime/Engine/Private/OpusAudioInfo.cpp
+++ b/Engine/Source/Runtime/Engine/Private/OpusAudioInfo.cpp
@@ -13,6 +13,8 @@ THIRD_PARTY_INCLUDES_END
 #define USE_UE_MEM_ALLOC 1
 #define OPUS_MAX_FRAME_SIZE_MS 120
 
+DEFINE_LOG_CATEGORY_STATIC(LogOpusAudioDecoder, Log, All);
+
 ///////////////////////////////////////////////////////////////////////////////////////
 // Followed pattern used in opus_multistream_encoder.c - this will allow us to setup //
 // a multistream decoder without having to save extra information for every asset.   //
@@ -169,6 +171,7 @@ bool FOpusAudioInfo::CreateDecoder()
 	}
 
 	NumRemainingSamplesToSkip = Header.NumSilentSamplesAtBeginning;
+	PreviousDecodedUnusedSamples.Empty();
 	return true;
 }
 
@@ -192,26 +195,118 @@ FDecodeResult FOpusAudioInfo::Decode(const uint8* CompressedData, const int32 Co
 
 	if (OpusDecoderWrapper)
 	{
-		const int32 SampleSize = OutputPCMDataSize / NumChannels * sizeof(int16);
-		Result.NumCompressedBytesConsumed = CompressedDataSize;
-		Result.NumAudioFramesProduced = OpusDecoderWrapper->Decode(CompressedData, CompressedDataSize, (int16*)OutPCMData, SampleSize);
-		if (NumRemainingSamplesToSkip)
+		const int32 kFrameBytes = NumChannels * sizeof(int16);
+		int32 OutputSizeToGo = OutputPCMDataSize;
+		int32 CompressedSizeToGo = CompressedDataSize;
+		const uint8* InputDataPtr = CompressedData; 
+		uint8* OutputDataPtr = OutPCMData;
+
+		Result.NumAudioFramesProduced = 0;
+		while(OutputSizeToGo > 0)
 		{
-			if (NumRemainingSamplesToSkip >= Result.NumAudioFramesProduced)
+			// Get anything that is left over from the previous call.
+			if (PreviousDecodedUnusedSamples.Num())
+			{
+				check(NumRemainingSamplesToSkip == 0);
+				int32 MaxToCopyOut = OutputSizeToGo >= PreviousDecodedUnusedSamples.Num() ? PreviousDecodedUnusedSamples.Num() : OutputSizeToGo;
+				FMemory::Memcpy(OutputDataPtr, PreviousDecodedUnusedSamples.GetData(), MaxToCopyOut);
+				PreviousDecodedUnusedSamples.RemoveAt(0, MaxToCopyOut);
+				OutputSizeToGo -= MaxToCopyOut;
+				OutputDataPtr += MaxToCopyOut;
+				// If there is still something left over then we are done here.
+				Result.NumAudioFramesProduced += MaxToCopyOut / kFrameBytes;
+				if (OutputSizeToGo == 0 || PreviousDecodedUnusedSamples.Num())
+				{
+					Result.NumCompressedBytesConsumed = InputDataPtr - CompressedData;
+					Result.NumPcmBytesProduced = Result.NumAudioFramesProduced * kFrameBytes;
+					return Result;
+				}
+			}
+			// Something left to decompress?
+			if (CompressedSizeToGo <= 0)
+			{
+				break;
+			}
+
+			// Decode the next chunk.
+			const int32 AvailableSampleOutputSize = OutputSizeToGo / kFrameBytes;
+
+			int32 ChunkToDecodeSize = 0;
+			int32 ActualChunkSize = 0;
+			// Is this streaming?
+			const uint8* ChunkToDecodePtr = InputDataPtr;
+			if (!bIsStreaming)
+			{
+				// When not streaming each input data chunk is prepended with the size of the chunk.
+				ChunkToDecodeSize = (int32)InputDataPtr[0] + ((int32)InputDataPtr[1] << 8);
+				ActualChunkSize = ChunkToDecodeSize + 2;
+				ChunkToDecodePtr += 2;
+			}
+			else
+			{
+				// When streaming the chunk size is still prepended to the chunk, but is skipped over
+				// by the caller and provided as an argument.
+				ChunkToDecodeSize = CompressedDataSize;
+				ActualChunkSize = ChunkToDecodeSize;
+			}
+
+
+			int32 NumDecodedFrames = OpusDecoderWrapper->Decode(ChunkToDecodePtr, ChunkToDecodeSize, (int16*)OutputDataPtr, AvailableSampleOutputSize);
+			if (NumDecodedFrames >= 0)
+			{
+				CompressedSizeToGo -= ActualChunkSize;
+				InputDataPtr += ActualChunkSize;
+				if (NumRemainingSamplesToSkip)
+				{
+					if (NumRemainingSamplesToSkip >= NumDecodedFrames)
+					{
+						NumRemainingSamplesToSkip -= NumDecodedFrames;
+						NumDecodedFrames = 0;
+					}
+					else
+					{
+						uint8* FirstUsable = OutputDataPtr + NumRemainingSamplesToSkip * kFrameBytes;
+						int32 UsableSize = (Result.NumAudioFramesProduced - NumRemainingSamplesToSkip) * kFrameBytes;
+						FMemory::Memmove(OutputDataPtr, FirstUsable, UsableSize);
+						NumDecodedFrames -= NumRemainingSamplesToSkip;
+						NumRemainingSamplesToSkip = 0;
+					}
+				}
+				Result.NumAudioFramesProduced += NumDecodedFrames;
+				OutputSizeToGo -= NumDecodedFrames * kFrameBytes;
+				OutputDataPtr += NumDecodedFrames * kFrameBytes;
+			}
+			// Was the remaining output buffer too small?
+			else if (NumDecodedFrames == OPUS_BUFFER_TOO_SMALL)
 			{
-				NumRemainingSamplesToSkip -= Result.NumAudioFramesProduced;
-				Result.NumAudioFramesProduced = 0;
+				int32 NumSampleSpaceNeeded = GetMaxFrameSizeSamples();
+				int32 NumPrevBefore = PreviousDecodedUnusedSamples.Num();
+				PreviousDecodedUnusedSamples.AddUninitialized(NumSampleSpaceNeeded * kFrameBytes);
+				uint8* TempPtr = PreviousDecodedUnusedSamples.GetData() + NumPrevBefore;
+				NumDecodedFrames = OpusDecoderWrapper->Decode(ChunkToDecodePtr, ChunkToDecodeSize, (int16*)TempPtr, NumSampleSpaceNeeded);
+				if (NumDecodedFrames >= 0)
+				{
+					int32 NumPrevNow = NumDecodedFrames * kFrameBytes;
+					PreviousDecodedUnusedSamples.SetNum(NumPrevNow);
+					CompressedSizeToGo -= ActualChunkSize;
+					InputDataPtr += ActualChunkSize;
+					continue;
+				}
+				else
+				{
+					UE_LOG(LogOpusAudioDecoder, Error, TEXT("opus_multistream_decode() returned %d while decoding into temporary buffer"), NumDecodedFrames);
+					return FDecodeResult();
+				}
 			}
 			else
 			{
-				uint8* FirstUsable = OutPCMData + NumRemainingSamplesToSkip * NumChannels * sizeof(int16);
-				int32 UsableSize = (Result.NumAudioFramesProduced - NumRemainingSamplesToSkip) * NumChannels * sizeof(int16);
-				FMemory::Memmove(OutPCMData, FirstUsable, UsableSize);
-				Result.NumAudioFramesProduced -= NumRemainingSamplesToSkip;
-				NumRemainingSamplesToSkip = 0;
+				// A decode error of sorts.
+				UE_LOG(LogOpusAudioDecoder, Error, TEXT("opus_multistream_decode() returned %d"), NumDecodedFrames);
+				return FDecodeResult();
 			}
 		}
-		Result.NumPcmBytesProduced = Result.NumAudioFramesProduced * NumChannels * sizeof(int16);
+		Result.NumCompressedBytesConsumed = InputDataPtr - CompressedData;
+		Result.NumPcmBytesProduced = Result.NumAudioFramesProduced * kFrameBytes;
 	}
 	return Result;
 }
@@ -220,11 +315,51 @@ void FOpusAudioInfo::PrepareToLoop()
 {
 	IStreamedCompressedInfo::PrepareToLoop();
 	NumRemainingSamplesToSkip = Header.NumSilentSamplesAtBeginning;
+	PreviousDecodedUnusedSamples.Empty();
 }
 
 void FOpusAudioInfo::SeekToTime(const float InSeekTime)
 {
-	IStreamedCompressedInfo::SeekToTime(InSeekTime);
+	if (GetStreamingSoundWave().IsValid())
+	{
+		IStreamedCompressedInfo::SeekToTime(InSeekTime);
+	}
+	else if (SrcBufferData && SrcBufferDataSize)
+	{
+		uint32 SeekSampleNum = 0;
+		if (InSeekTime > 0.0f)
+		{
+			SeekSampleNum = (uint32)(InSeekTime * Header.SampleRate);
+		}
+
+		const uint8* ChunkPtr = SrcBufferData + AudioDataOffset;
+		const uint8* const EndPtr = SrcBufferData + SrcBufferDataSize;
+		uint32 CurrentChunkSampleNum = 0;
+		while(ChunkPtr < EndPtr)
+		{
+			uint32 ChunkSize = (uint32)ChunkPtr[0] + ((uint32)ChunkPtr[1] << 8);
+			int32 ExpectedFrames = opus_packet_get_nb_frames(ChunkPtr + 2, ChunkSize);
+			int32 ExpectedFrameSize = opus_packet_get_samples_per_frame(ChunkPtr + 2, Header.SampleRate);
+			int32 NumExpectedTotal = ExpectedFrames * ExpectedFrameSize;
+
+			if (CurrentChunkSampleNum >= SeekSampleNum && SeekSampleNum < CurrentChunkSampleNum + NumExpectedTotal)
+			{
+				CurrentSampleCount = CurrentChunkSampleNum;
+				SrcBufferOffset = ChunkPtr - SrcBufferData;
+				break;
+			}
+			CurrentChunkSampleNum += NumExpectedTotal;
+			ChunkPtr += ChunkSize + 2;
+		}
+		// Not found?
+		if (ChunkPtr >= EndPtr)
+		{
+			CurrentSampleCount = SeekSampleNum;
+			SrcBufferOffset = SrcBufferDataSize;
+		}
+	}
+
 	NumRemainingSamplesToSkip = Header.NumSilentSamplesAtBeginning;
+	PreviousDecodedUnusedSamples.Empty();
 }
 
diff --git a/Engine/Source/Runtime/Engine/Public/OpusAudioInfo.h b/Engine/Source/Runtime/Engine/Public/OpusAudioInfo.h
--- a/Engine/Source/Runtime/Engine/Public/OpusAudioInfo.h
+++ b/Engine/Source/Runtime/Engine/Public/OpusAudioInfo.h
@@ -79,4 +79,7 @@ protected:
 	
 	/** The number of decoded samples that need to be skipped over. */
 	int32 NumRemainingSamplesToSkip = 0;
+
+	/** Samples that are left over from the previous Decode() call that did not fit into the receive buffer. */
+	TArray<uint8> PreviousDecodedUnusedSamples;
 };
