Added ability to clamp DeadBlending extrapolation velocities to some maximum.

diff --git a/Engine/Source/Runtime/Engine/Classes/Animation/AnimNode_DeadBlending.h b/Engine/Source/Runtime/Engine/Classes/Animation/AnimNode_DeadBlending.h
--- a/Engine/Source/Runtime/Engine/Classes/Animation/AnimNode_DeadBlending.h
+++ b/Engine/Source/Runtime/Engine/Classes/Animation/AnimNode_DeadBlending.h
@@ -103,6 +103,34 @@ private:
 	UPROPERTY(EditAnywhere, Category = Extrapolation, meta = (Min = "0.0", UIMin = "0.0"), DisplayName = "Maximum Extrapolation Half Life")
 	float ExtrapolationHalfLifeMax = 1.0f;
 
+	/**
+	 * The maximum velocity to allow for extrapolation of bone translations in centimeters per second. Smaller values 
+	 * may help prevent the pose breaking during blending but too small values can make the blend less smooth.
+	 */
+	UPROPERTY(EditAnywhere, Category = Extrapolation, meta = (Min = "0.0", UIMin = "0.0"))
+	float MaximumTranslationVelocity = 500.0f;
+
+	/**
+	 * The maximum velocity to allow for extrapolation of bone rotations in degrees per second. Smaller values
+	 * may help prevent the pose breaking during blending but too small values can make the blend less smooth.
+	 */
+	UPROPERTY(EditAnywhere, Category = Extrapolation, meta = (Min = "0.0", UIMin = "0.0"))
+	float MaximumRotationVelocity = 360.0f;
+
+	/**
+	 * The maximum velocity to allow for extrapolation of bone scales. Smaller values may help prevent the pose 
+	 * breaking during blending but too small values can make the blend less smooth.
+	 */
+	UPROPERTY(EditAnywhere, Category = Extrapolation, meta = (Min = "0.0", UIMin = "0.0"))
+	float MaximumScaleVelocity = 4.0f;
+
+	/**
+	 * The maximum velocity to allow for extrapolation of curves. Smaller values may help prevent extreme curve values 
+	 * during blending but too small values can make the blending of curves less smooth.
+	 */
+	UPROPERTY(EditAnywhere, Category = Extrapolation, meta = (Min = "0.0", UIMin = "0.0"))
+	float MaximumCurveVelocity = 100.0f;
+
 #if WITH_EDITORONLY_DATA
 	
 	// This setting can be used to show what the extrapolation of the animation looks like.
diff --git a/Engine/Source/Runtime/Engine/Private/Animation/AnimNode_DeadBlending.cpp b/Engine/Source/Runtime/Engine/Private/Animation/AnimNode_DeadBlending.cpp
--- a/Engine/Source/Runtime/Engine/Private/Animation/AnimNode_DeadBlending.cpp
+++ b/Engine/Source/Runtime/Engine/Private/Animation/AnimNode_DeadBlending.cpp
@@ -314,6 +314,12 @@ void FAnimNode_DeadBlending::InitFrom(
 			BoneRotationVelocities[SkeletonPoseBoneIndex] = RotationDiff.ToRotationVector() / SrcPoseCurr.DeltaTime;
 			BoneScaleVelocities[SkeletonPoseBoneIndex] = UE::Anim::DeadBlending::Private::VectorLogSafe(ScaleDiff) / SrcPoseCurr.DeltaTime;
 
+			// Clamp Maximum Velocity
+
+			BoneTranslationVelocities[SkeletonPoseBoneIndex] = BoneTranslationVelocities[SkeletonPoseBoneIndex].GetClampedToMaxSize(MaximumTranslationVelocity);
+			BoneRotationVelocities[SkeletonPoseBoneIndex] = BoneRotationVelocities[SkeletonPoseBoneIndex].GetClampedToMaxSize(FMath::DegreesToRadians(MaximumRotationVelocity));
+			BoneScaleVelocities[SkeletonPoseBoneIndex] = BoneScaleVelocities[SkeletonPoseBoneIndex].GetClampedToMaxSize(MaximumScaleVelocity);
+
 			// Compute Decay HalfLives
 
 			const FTransform DstTransform = InPose[BoneIndex];
@@ -365,7 +371,7 @@ void FAnimNode_DeadBlending::InitFrom(
 	if (SrcPoseCurr.DeltaTime > UE_SMALL_NUMBER)
 	{
 		UE::Anim::FNamedValueArrayUtils::Union(CurveData, SrcPosePrev.Curves.BlendedCurve,
-			[DeltaTime = SrcPoseCurr.DeltaTime](FDeadBlendingCurveElement& OutResultElement, const UE::Anim::FCurveElement& InElement1, UE::Anim::ENamedValueUnionFlags InFlags)
+			[this, DeltaTime = SrcPoseCurr.DeltaTime](FDeadBlendingCurveElement& OutResultElement, const UE::Anim::FCurveElement& InElement1, UE::Anim::ENamedValueUnionFlags InFlags)
 			{
 				bool bSrcCurrValid = (bool)(InFlags & UE::Anim::ENamedValueUnionFlags::ValidArg0);
 				bool bSrcPrevValid = (bool)(InFlags & UE::Anim::ENamedValueUnionFlags::ValidArg1);
@@ -373,6 +379,7 @@ void FAnimNode_DeadBlending::InitFrom(
 				if (bSrcCurrValid && bSrcPrevValid)
 				{
 					OutResultElement.Velocity = (OutResultElement.Value - InElement1.Value) / DeltaTime;
+					OutResultElement.Velocity = FMath::Clamp(OutResultElement.Velocity, -MaximumCurveVelocity, MaximumCurveVelocity);
 				}
 			});
 	}
