Fix incorrect rect lighting on translucent light grid.

diff --git a/Engine/Source/Runtime/Engine/Private/Components/DirectionalLightComponent.cpp b/Engine/Source/Runtime/Engine/Private/Components/DirectionalLightComponent.cpp
--- a/Engine/Source/Runtime/Engine/Private/Components/DirectionalLightComponent.cpp
+++ b/Engine/Source/Runtime/Engine/Private/Components/DirectionalLightComponent.cpp
@@ -386,7 +386,7 @@ public:
 	}
 
 	/** Accesses parameters needed for rendering the light. */
-	virtual void GetLightShaderParameters(FLightRenderParameters& LightParameters) const override
+	virtual void GetLightShaderParameters(FLightRenderParameters& LightParameters, uint32 Flags=0) const override
 	{
 		LightParameters.WorldPosition = FVector::ZeroVector;
 		LightParameters.InvRadius = 0.0f;
diff --git a/Engine/Source/Runtime/Engine/Private/Components/PointLightComponent.cpp b/Engine/Source/Runtime/Engine/Private/Components/PointLightComponent.cpp
--- a/Engine/Source/Runtime/Engine/Private/Components/PointLightComponent.cpp
+++ b/Engine/Source/Runtime/Engine/Private/Components/PointLightComponent.cpp
@@ -33,7 +33,7 @@ void FLocalLightSceneProxy::UpdateRadius_GameThread(float ComponentRadius)
 }
 
 /** Accesses parameters needed for rendering the light. */
-void FPointLightSceneProxy::GetLightShaderParameters(FLightRenderParameters& LightParameters) const
+void FPointLightSceneProxy::GetLightShaderParameters(FLightRenderParameters& LightParameters, uint32 Flags) const
 {
 	LightParameters.WorldPosition = GetOrigin();
 	LightParameters.InvRadius = InvRadius;
diff --git a/Engine/Source/Runtime/Engine/Private/Components/RectLightComponent.cpp b/Engine/Source/Runtime/Engine/Private/Components/RectLightComponent.cpp
--- a/Engine/Source/Runtime/Engine/Private/Components/RectLightComponent.cpp
+++ b/Engine/Source/Runtime/Engine/Private/Components/RectLightComponent.cpp
@@ -226,7 +226,7 @@ bool FRectLightSceneProxy::HasSourceTexture() const
 }
 
 /** Accesses parameters needed for rendering the light. */
-void FRectLightSceneProxy::GetLightShaderParameters(FLightRenderParameters& LightParameters) const
+void FRectLightSceneProxy::GetLightShaderParameters(FLightRenderParameters& LightParameters, uint32 Flags) const
 {
 	FLinearColor LightColor = GetColor();
 	LightColor /= 0.5f * SourceWidth * SourceHeight;
@@ -260,8 +260,8 @@ void FRectLightSceneProxy::GetLightShaderParameters(FLightRenderParameters& Ligh
 		GetSceneInterface()->GetRectLightAtlasSlot(this, &LightParameters);
 	}
 	
-	// Render RectLight approximately as SpotLight on mobile
-	const bool bRenderAsSpotLight = (SceneInterface && IsMobilePlatform(SceneInterface->GetShaderPlatform()));
+	// Render RectLight approximately as SpotLight if the requester does not support rect light (e.g., translucent light grid or mobile)
+	const bool bRenderAsSpotLight = !!(Flags & ELightShaderParameterFlags::RectAsSpotLight) || (SceneInterface && IsMobilePlatform(SceneInterface->GetShaderPlatform()));
 	if (bRenderAsSpotLight)
 	{
 		float ClampedOuterConeAngle = FMath::DegreesToRadians(89.001f);
@@ -304,5 +304,4 @@ bool FRectLightSceneProxy::GetWholeSceneProjectedShadowInitializer(const FSceneV
 	}
 
 	return false;
-}
-
+}
\ No newline at end of file
diff --git a/Engine/Source/Runtime/Engine/Private/Components/SpotLightComponent.cpp b/Engine/Source/Runtime/Engine/Private/Components/SpotLightComponent.cpp
--- a/Engine/Source/Runtime/Engine/Private/Components/SpotLightComponent.cpp
+++ b/Engine/Source/Runtime/Engine/Private/Components/SpotLightComponent.cpp
@@ -56,7 +56,7 @@ public:
 	}
 
 	/** Accesses parameters needed for rendering the light. */
-	virtual void GetLightShaderParameters(FLightRenderParameters& LightParameters) const override
+	virtual void GetLightShaderParameters(FLightRenderParameters& LightParameters, uint32 Flags=0) const override
 	{
 		LightParameters.WorldPosition = GetOrigin();
 		LightParameters.InvRadius = InvRadius;
diff --git a/Engine/Source/Runtime/Engine/Private/PointLightSceneProxy.h b/Engine/Source/Runtime/Engine/Private/PointLightSceneProxy.h
--- a/Engine/Source/Runtime/Engine/Private/PointLightSceneProxy.h
+++ b/Engine/Source/Runtime/Engine/Private/PointLightSceneProxy.h
@@ -51,7 +51,7 @@ public:
 		return bInverseSquared;
 	}
 
-	virtual void GetLightShaderParameters(FLightRenderParameters& LightParameters) const override;
+	virtual void GetLightShaderParameters(FLightRenderParameters& LightParameters, uint32 Flags=0) const override;
 
 	virtual FVector GetPerObjectProjectedShadowProjectionPoint(const FBoxSphereBounds& SubjectBounds) const override
 	{
diff --git a/Engine/Source/Runtime/Engine/Public/LightSceneProxy.h b/Engine/Source/Runtime/Engine/Public/LightSceneProxy.h
--- a/Engine/Source/Runtime/Engine/Public/LightSceneProxy.h
+++ b/Engine/Source/Runtime/Engine/Public/LightSceneProxy.h
@@ -22,6 +22,11 @@ class FLightSceneInfo;
 class UTextureLightProfile;
 class FTexture;
 
+enum ELightShaderParameterFlags
+{
+	RectAsSpotLight=1,
+};
+
 /**
  * Encapsulates the data which is used to render a light by the rendering thread.
  * The constructor is called from the game thread, and after that the rendering thread owns the object.
@@ -66,7 +71,7 @@ public:
 	virtual float GetEffectiveScreenRadius(const FViewMatrices& ShadowViewMatrices) const { return 0.0f; }
 
 	/** Accesses parameters needed for rendering the light. */
-	virtual void GetLightShaderParameters(FLightRenderParameters& OutLightParameters) const {}
+	virtual void GetLightShaderParameters(FLightRenderParameters& OutLightParameters, uint32 Flags=0) const {}
 
 	virtual FVector2D GetDirectionalLightDistanceFadeParameters(ERHIFeatureLevel::Type InFeatureLevel, bool bPrecomputedLightingIsValid, int32 MaxNearCascades) const
 	{
diff --git a/Engine/Source/Runtime/Engine/Public/RectLightSceneProxy.h b/Engine/Source/Runtime/Engine/Public/RectLightSceneProxy.h
--- a/Engine/Source/Runtime/Engine/Public/RectLightSceneProxy.h
+++ b/Engine/Source/Runtime/Engine/Public/RectLightSceneProxy.h
@@ -40,7 +40,7 @@ public:
 	virtual bool HasSourceTexture() const override;
 
 	/** Accesses parameters needed for rendering the light. */
-	virtual void GetLightShaderParameters(FLightRenderParameters& LightParameters) const override;
+	virtual void GetLightShaderParameters(FLightRenderParameters& OutLightParameters, uint32 Flags=0) const override;
 
 	/**
 	* Sets up a projected shadow initializer for shadows from the entire scene.
diff --git a/Engine/Source/Runtime/Renderer/Private/LightRendering.cpp b/Engine/Source/Runtime/Renderer/Private/LightRendering.cpp
--- a/Engine/Source/Runtime/Renderer/Private/LightRendering.cpp
+++ b/Engine/Source/Runtime/Renderer/Private/LightRendering.cpp
@@ -489,12 +489,12 @@ static void RenderLight(
 	FRDGTextureRef ShadowMaskBits = nullptr,
 	int32 VirtualShadowMapId = INDEX_NONE);
 
-FDeferredLightUniformStruct GetDeferredLightParameters(const FSceneView& View, const FLightSceneInfo& LightSceneInfo)
+FDeferredLightUniformStruct GetDeferredLightParameters(const FSceneView& View, const FLightSceneInfo& LightSceneInfo, uint32 LightFlags)
 {
 	FDeferredLightUniformStruct Out;
 
 	FLightRenderParameters LightParameters;
-	LightSceneInfo.Proxy->GetLightShaderParameters(LightParameters);
+	LightSceneInfo.Proxy->GetLightShaderParameters(LightParameters, LightFlags);
 	LightParameters.MakeShaderParameters(View.ViewMatrices, View.GetLastEyeAdaptationExposure(), Out.LightParameters);
 
 	const bool bIsRayTracedLight = ShouldRenderRayTracingShadowsForLight(*LightSceneInfo.Proxy);
diff --git a/Engine/Source/Runtime/Renderer/Private/LightRendering.h b/Engine/Source/Runtime/Renderer/Private/LightRendering.h
--- a/Engine/Source/Runtime/Renderer/Private/LightRendering.h
+++ b/Engine/Source/Runtime/Renderer/Private/LightRendering.h
@@ -35,7 +35,7 @@ extern uint32 GetShadowQuality();
 
 extern float GetLightFadeFactor(const FSceneView& View, const FLightSceneProxy* Proxy);
 
-extern FDeferredLightUniformStruct GetDeferredLightParameters(const FSceneView& View, const FLightSceneInfo& LightSceneInfo);
+extern FDeferredLightUniformStruct GetDeferredLightParameters(const FSceneView& View, const FLightSceneInfo& LightSceneInfo, uint32 LightFlags=0);
 
 inline void SetDeferredLightParameters(
 	FRHIBatchedShaderParameters& BatchedParameters,
diff --git a/Engine/Source/Runtime/Renderer/Private/TranslucentLighting.cpp b/Engine/Source/Runtime/Renderer/Private/TranslucentLighting.cpp
--- a/Engine/Source/Runtime/Renderer/Private/TranslucentLighting.cpp
+++ b/Engine/Source/Runtime/Renderer/Private/TranslucentLighting.cpp
@@ -1142,7 +1142,7 @@ void InjectTranslucencyLightingVolume(
 				PassParameters->PS.ViewUniformBuffer = View.ViewUniformBuffer;
 
 				FDeferredLightUniformStruct* DeferredLightStruct = GraphBuilder.AllocParameters<FDeferredLightUniformStruct>();
-				*DeferredLightStruct = GetDeferredLightParameters(View, *LightSceneInfo);
+				*DeferredLightStruct = GetDeferredLightParameters(View, *LightSceneInfo, ELightShaderParameterFlags::RectAsSpotLight);
 				PassParameters->PS.DeferredLight = GraphBuilder.CreateUniformBuffer(DeferredLightStruct);
 
 				GetVolumeShadowingShaderParameters(GraphBuilder, View, LightSceneInfo, InjectionData.ProjectedShadowInfo, PassParameters->PS.VolumeShadowingParameters);
