WorldPartition : Fixed lost child actors at cook when moving actor in world partition cell level.

diff --git a/Engine/Source/Runtime/Engine/Private/WorldPartition/WorldPartitionLevelHelper.cpp b/Engine/Source/Runtime/Engine/Private/WorldPartition/WorldPartitionLevelHelper.cpp
--- a/Engine/Source/Runtime/Engine/Private/WorldPartition/WorldPartitionLevelHelper.cpp
+++ b/Engine/Source/Runtime/Engine/Private/WorldPartition/WorldPartitionLevelHelper.cpp
@@ -187,29 +187,51 @@ void FWorldPartitionLevelHelper::MoveExternalActorsToLevel(const TArray<FWorldPa
 
 			const bool bSameOuter = (InLevel == Actor->GetOuter());
 			Actor->SetPackageExternal(false, false);
-						
+
 			// Avoid calling Rename on the actor if it's already outered to InLevel as this will cause it's name to be changed. 
 			// (UObject::Rename doesn't check if Rename is being called with existing outer and assigns new name)
 			if (!bSameOuter)
 			{
 				Actor->Rename(nullptr, InLevel, REN_ForceNoResetLoaders);
 			}
-			
-			check(Actor->GetPackage() == LevelPackage);
-			if (bSameOuter && !InLevel->Actors.Contains(Actor))
+			else if (!InLevel->Actors.Contains(Actor))
 			{
 				InLevel->AddLoadedActor(Actor);
 			}
+			check(Actor->GetPackage() == LevelPackage);
 
-			// Include objects found in the source actor package in the destination level package
+			// Process objects found in the source actor package
 			TArray<UObject*> Objects;
 			const bool bIncludeNestedSubobjects = false;
-			GetObjectsWithOuter(ActorExternalPackage, Objects, bIncludeNestedSubobjects);
+			// Skip Garbage objects as the initial Rename on an actor with an ChildActorComponent can destroy its child actors.
+			// This happens when the component has bNeedsRecreate set to true (when it has a valid ChildActorTemplate).
+			GetObjectsWithPackage(ActorExternalPackage, Objects, bIncludeNestedSubobjects, RF_NoFlags, EInternalObjectFlags::Garbage);
 			for (UObject* Object : Objects)
 			{
 				if (Object->GetFName() != NAME_PackageMetaData)
 				{
-					Object->Rename(nullptr, LevelPackage, REN_ForceNoResetLoaders);
+					if (Object->GetOuter()->IsA<ULevel>())
+					{
+						// Move objects that are outered the level in the destination level
+						AActor* NestedActor = Cast<AActor>(Object);
+						if (InLevel != Object->GetOuter())
+						{
+							Object->Rename(nullptr, InLevel, REN_ForceNoResetLoaders);
+						}
+						else if (NestedActor && !InLevel->Actors.Contains(NestedActor))
+						{
+							InLevel->AddLoadedActor(NestedActor);
+						}
+						if (NestedActor)
+						{
+							LevelActors.Add(NestedActor->GetFName());
+						}
+					}
+					else
+					{
+						// Move objects in the destination level package
+						Object->Rename(nullptr, LevelPackage, REN_ForceNoResetLoaders);
+					}
 				}
 			}
 
