Adding ray tracing support to FTextRenderSceneProxy
- Added WPO support

diff --git a/Engine/Source/Runtime/Engine/Private/Components/TextRenderComponent.cpp b/Engine/Source/Runtime/Engine/Private/Components/TextRenderComponent.cpp
--- a/Engine/Source/Runtime/Engine/Private/Components/TextRenderComponent.cpp
+++ b/Engine/Source/Runtime/Engine/Private/Components/TextRenderComponent.cpp
@@ -22,6 +22,11 @@
 #include "Containers/Ticker.h"
 #include "PSOPrecache.h"
 
+static TAutoConsoleVariable<int32> CVarRayTracingTextMeshes(
+	TEXT("r.RayTracing.Geometry.Text"),
+	1,
+	TEXT("Include text meshes in ray tracing effects (default = 1 (text meshes enabled in ray tracing))"));
+
 #define LOCTEXT_NAMESPACE "TextRenderComponent"
 
 ATextRenderActor::ATextRenderActor(const FObjectInitializer& ObjectInitializer)
@@ -594,6 +599,15 @@ public:
 	virtual bool IsUsingDistanceCullFade() const override;
 	virtual uint32 GetMemoryFootprint() const override;
 	uint32 GetAllocatedSize() const;
+#if RHI_RAYTRACING
+	virtual void GetDynamicRayTracingInstances(FRayTracingMaterialGatheringContext& Context, TArray<FRayTracingInstance>& OutRayTracingInstances) override final;
+	virtual bool HasRayTracingRepresentation() const override { return true; }
+	virtual bool IsRayTracingRelevant() const override { return true; }
+	virtual bool IsRayTracingStaticRelevant() const override
+	{
+		return false;
+	}
+#endif
 	// End FPrimitiveSceneProxy interface
 
 private:
@@ -633,6 +647,19 @@ private:
 	TEnumAsByte<EHorizTextAligment> HorizontalAlignment;
 	TEnumAsByte<EVerticalTextAligment> VerticalAlignment;
 	bool bAlwaysRenderAsText;
+
+#if RHI_RAYTRACING
+	/** Geometry for ray tracing. */
+	FRayTracingGeometry RayTracingGeometry;
+	TArray<FMeshBatch> CachedRayTracingMaterials;
+
+	bool bSupportRayTracing;
+	bool bRayTracingWithWPO;
+	bool bNeedsToUpdateRayTracingCache;
+	FRayTracingMaskAndFlags CachedRayTracingInstanceMaskAndFlags;
+
+	void UpdateRayTracingGeometry_RenderingThread(FRHICommandListBase& RHICmdList);
+#endif
 };
 
 FTextRenderSceneProxy::FTextRenderSceneProxy( UTextRenderComponent* Component) :
@@ -680,6 +707,12 @@ FTextRenderSceneProxy::FTextRenderSceneProxy( UTextRenderComponent* Component) :
 
 	// The MID from the cache isn't known by the UTextRenderComponent
 	bVerifyUsedMaterials = false;
+
+#if RHI_RAYTRACING
+	bSupportRayTracing = IsRayTracingEnabled();
+	bNeedsToUpdateRayTracingCache = true;
+	bRayTracingWithWPO = MaterialRelevance.bUsesWorldPositionOffset;
+#endif
 }
 
 FTextRenderSceneProxy::~FTextRenderSceneProxy()
@@ -689,7 +722,7 @@ FTextRenderSceneProxy::~FTextRenderSceneProxy()
 
 void FTextRenderSceneProxy::CreateRenderThreadResources(FRHICommandListBase& RHICmdList)
 {
-	if(Font && Font->FontCacheType == EFontCacheType::Runtime)
+	if (Font && Font->FontCacheType == EFontCacheType::Runtime)
 	{
 		// Runtime fonts can't currently be used here as they use the font cache from Slate application
 		// which can only be used on the game thread
@@ -697,7 +730,7 @@ void FTextRenderSceneProxy::CreateRenderThreadResources(FRHICommandListBase& RHI
 	}
 
 	TArray<FDynamicMeshVertex> OutVertices;
-	if(BuildStringMesh(OutVertices, IndexBuffer.Indices))
+	if (BuildStringMesh(OutVertices, IndexBuffer.Indices))
 	{
 #if RHI_ENABLE_RESOURCE_INFO
 		FName Name = FName(TEXT("FTextRenderSceneProxy ") + GetOwnerName().ToString());
@@ -707,6 +740,14 @@ void FTextRenderSceneProxy::CreateRenderThreadResources(FRHICommandListBase& RHI
 		VertexBuffers.InitFromDynamicVertex(&VertexFactory, OutVertices);
 		// Enqueue initialization of render resources
 		BeginInitResource(&IndexBuffer);
+
+#if RHI_RAYTRACING
+		if (bSupportRayTracing)
+		{
+			UpdateRayTracingGeometry_RenderingThread(RHICmdList);
+			bNeedsToUpdateRayTracingCache = true;
+		}
+#endif
 	}
 }
 
@@ -717,6 +758,13 @@ void FTextRenderSceneProxy::ReleaseRenderThreadResources()
 	VertexBuffers.ColorVertexBuffer.ReleaseResource();
 	IndexBuffer.ReleaseResource();
 	VertexFactory.ReleaseResource();
+
+#if RHI_RAYTRACING
+	if (IsRayTracingEnabled())
+	{
+		RayTracingGeometry.ReleaseResource();
+	}
+#endif
 }
 
 void FTextRenderSceneProxy::GetDynamicMeshElements(const TArray<const FSceneView*>& Views, const FSceneViewFamily& ViewFamily, uint32 VisibilityMap, FMeshElementCollector& Collector) const 
@@ -841,6 +889,85 @@ uint32 FTextRenderSceneProxy::GetAllocatedSize() const
 	return( FPrimitiveSceneProxy::GetAllocatedSize() ); 
 }
 
+#if RHI_RAYTRACING
+void FTextRenderSceneProxy::GetDynamicRayTracingInstances(FRayTracingMaterialGatheringContext& Context, TArray<FRayTracingInstance>& OutRayTracingInstances)
+{
+	if (CVarRayTracingTextMeshes.GetValueOnRenderThread() == 0 || !bSupportRayTracing)
+	{
+		return;
+	}
+	// Vertex factory will not been initialized when the text string is empty or font is invalid.
+	if (!VertexFactory.IsInitialized())
+	{
+		return;
+	}
+	FRayTracingInstance& RayTracingInstance = OutRayTracingInstances.AddDefaulted_GetRef();
+
+	if (bNeedsToUpdateRayTracingCache)
+	{
+		CachedRayTracingMaterials.Reset();
+		int NumBatches = TextBatches.Num();
+		for (int32 BatchIndex = 0; BatchIndex < NumBatches; BatchIndex++)
+		{
+			const FTextBatch& TextBatch = TextBatches[BatchIndex];
+			FMeshBatch& Mesh = CachedRayTracingMaterials.AddDefaulted_GetRef();
+			FMeshBatchElement& BatchElement = Mesh.Elements[0];
+			BatchElement.IndexBuffer = &IndexBuffer;
+			Mesh.VertexFactory = &VertexFactory;
+			BatchElement.FirstIndex = TextBatch.IndexBufferOffset;
+			BatchElement.NumPrimitives = TextBatch.IndexBufferCount / 3;
+			BatchElement.MinVertexIndex = TextBatch.VertexBufferOffset;
+			BatchElement.MaxVertexIndex = TextBatch.VertexBufferOffset + TextBatch.VertexBufferCount - 1;
+			Mesh.ReverseCulling = IsLocalToWorldDeterminantNegative();
+			Mesh.bDisableBackfaceCulling = false;
+			Mesh.Type = PT_TriangleList;
+			Mesh.DepthPriorityGroup = SDPG_World;
+			Mesh.MaterialRenderProxy = TextBatch.Material->GetRenderProxy();
+			Mesh.bCanApplyViewModeOverrides = !bAlwaysRenderAsText;
+			Mesh.LODIndex = 0;
+			Mesh.bUseWireframeSelectionColoring = IsSelected() ? 1 : 0;
+			Mesh.SegmentIndex = 0;
+			Mesh.MeshIdInPrimitive = 0;
+		}
+		RayTracingInstance.MaterialsView = MakeArrayView(CachedRayTracingMaterials);
+		CachedRayTracingInstanceMaskAndFlags = Context.BuildInstanceMaskAndFlags(RayTracingInstance, *this);
+		bNeedsToUpdateRayTracingCache = false;
+	}
+	else
+	{
+		RayTracingInstance.MaterialsView = MakeArrayView(CachedRayTracingMaterials);
+	}
+
+	RayTracingInstance.Geometry = &RayTracingGeometry;
+	const FMatrix& ThisLocalToWorld = GetLocalToWorld();
+	RayTracingInstance.InstanceTransformsView = MakeArrayView(&ThisLocalToWorld, 1);
+	RayTracingInstance.bInstanceMaskAndFlagsDirty = false;
+	RayTracingInstance.MaskAndFlags = CachedRayTracingInstanceMaskAndFlags;
+
+	if (bRayTracingWithWPO && VertexFactory.GetType()->SupportsRayTracingDynamicGeometry())
+	{
+		Context.DynamicRayTracingGeometriesToUpdate.Add(
+			FRayTracingDynamicGeometryUpdateParams
+			{
+				CachedRayTracingMaterials, // TODO: this copy can be avoided if FRayTracingDynamicGeometryUpdateParams supported array views
+				false,
+				(uint32)VertexBuffers.PositionVertexBuffer.GetNumVertices(),
+				uint32((SIZE_T)VertexBuffers.PositionVertexBuffer.GetNumVertices() * sizeof(FVector3f)),
+				RayTracingGeometry.Initializer.TotalPrimitiveCount,
+				&RayTracingGeometry,
+				nullptr /* VertexBuffer */,
+				true
+			}
+		);
+	}
+	check(CachedRayTracingMaterials.Num() == RayTracingInstance.GetMaterials().Num());
+	checkf(RayTracingInstance.Geometry->Initializer.Segments.Num() == CachedRayTracingMaterials.Num(), TEXT("Segments/Materials mismatch. Number of segments: %d. Number of Materials: %d."),
+		RayTracingInstance.Geometry->Initializer.Segments.Num(),
+		CachedRayTracingMaterials.Num());
+
+}
+#endif // RHI_RAYTRACING
+
 /**
 * For the given text, constructs a mesh to be used by the vertex factory for rendering.
 */
@@ -1002,6 +1129,40 @@ bool FTextRenderSceneProxy::BuildStringMesh( TArray<FDynamicMeshVertex>& OutVert
 	return (OutVertices.Num() > 0);
 }
 
+#if RHI_RAYTRACING
+void FTextRenderSceneProxy::UpdateRayTracingGeometry_RenderingThread(FRHICommandListBase& RHICmdList)
+{
+	static const FName DebugName("FTextRenderSceneProxy");
+	static int32 DebugNumber = 0;
+
+	// Creating ray tracing geometries for text
+	FRayTracingGeometryInitializer Initializer;
+	Initializer.DebugName = FDebugName(DebugName, DebugNumber++);
+	Initializer.IndexBuffer = IndexBuffer.IndexBufferRHI;
+	Initializer.TotalPrimitiveCount = IndexBuffer.Indices.Num() / 3;
+	Initializer.GeometryType = RTGT_Triangles;
+	Initializer.bFastBuild = true;
+	Initializer.bAllowUpdate = false;
+
+	TArray<FRayTracingGeometrySegment> GeometrySections;
+	FRayTracingGeometrySegment Segment;
+	Segment.VertexBuffer = VertexBuffers.PositionVertexBuffer.VertexBufferRHI;
+	Segment.VertexBufferElementType = VET_Float3;
+	Segment.VertexBufferStride = VertexBuffers.PositionVertexBuffer.GetStride();
+	Segment.VertexBufferOffset = 0;
+	Segment.MaxVertices = VertexBuffers.PositionVertexBuffer.GetNumVertices();
+	Segment.FirstPrimitive = 0;
+	Segment.NumPrimitives = IndexBuffer.Indices.Num() / 3;
+	Segment.bEnabled = true;
+	Segment.bForceOpaque = false;
+	GeometrySections.Add(Segment);
+	Initializer.Segments = GeometrySections;
+
+	RayTracingGeometry.SetInitializer(Initializer);
+	RayTracingGeometry.InitResource(RHICmdList);
+}
+#endif
+
 // ------------------------------------------------------
 
 /** Watches for culture changes and updates all live UTextRenderComponent components */
