WorldPartition:
- Fix non spatial newly created/saved actor being unloaded after PIE
- Also properly initialize bIsForcedNonSpatiallyLoaded on new actor descs so this fix doesn't cause regression

#rnx

diff --git a/Engine/Source/Runtime/Engine/Private/WorldPartition/WorldPartition.cpp b/Engine/Source/Runtime/Engine/Private/WorldPartition/WorldPartition.cpp
--- a/Engine/Source/Runtime/Engine/Private/WorldPartition/WorldPartition.cpp
+++ b/Engine/Source/Runtime/Engine/Private/WorldPartition/WorldPartition.cpp
@@ -593,7 +593,7 @@ void UWorldPartition::Initialize(UWorld* InWorld, const FTransform& InTransform)
 
 		// Follow the world's streaming enabled value most of the times, except:
 		//	- World is instanced and from a Level Instance that supports partial loading
-		const bool bIsStreamingEnabled = bOverrideEnableStreamingInEditor.IsSet() ? *bOverrideEnableStreamingInEditor : IsStreamingEnabled();
+		const bool bIsStreamingEnabled = IsStreamingEnabledInEditor();
 
 		if (bIsInstanced)
 		{
@@ -1233,7 +1233,7 @@ bool UWorldPartition::IsSimulating(bool bIncludeTestEnableSimulationStreamingSou
 #if WITH_EDITOR
 void UWorldPartition::OnActorDescAdded(FWorldPartitionActorDesc* NewActorDesc)
 {
-	NewActorDesc->bIsForcedNonSpatiallyLoaded = !IsStreamingEnabled();
+	NewActorDesc->bIsForcedNonSpatiallyLoaded = !IsStreamingEnabledInEditor();
 
 	HashActorDesc(NewActorDesc);
 
@@ -1323,7 +1323,7 @@ void UWorldPartition::OnEnableStreamingChanged()
 	for (FActorDescContainerCollection::TIterator<> ActorDescIterator(this); ActorDescIterator; ++ActorDescIterator)
 	{
 		UnhashActorDesc(*ActorDescIterator);
-		ActorDescIterator->bIsForcedNonSpatiallyLoaded = !IsStreamingEnabled();
+		ActorDescIterator->bIsForcedNonSpatiallyLoaded = !IsStreamingEnabledInEditor();
 		HashActorDesc(*ActorDescIterator);
 	}
 
@@ -1417,6 +1417,11 @@ void UWorldPartition::OnContentBundleRemovedContent(const FContentBundleEditor*
 		}
 	}
 }
+
+bool UWorldPartition::IsStreamingEnabledInEditor() const
+{
+	return bOverrideEnableStreamingInEditor.IsSet() ? *bOverrideEnableStreamingInEditor : IsStreamingEnabled();
+}
 #endif
 
 void UWorldPartition::Serialize(FArchive& Ar)
@@ -1550,6 +1555,7 @@ void UWorldPartition::Tick(float DeltaSeconds)
 		EditorHash->Tick(DeltaSeconds);
 	}
 
+	TArray<FWorldPartitionReference> ActorReferences;
 	for (TSet<FDirtyActor>::TIterator DirtyActorIt(DirtyActors); DirtyActorIt; ++DirtyActorIt)
 	{
 		const FDirtyActor& DirtyActor = *DirtyActorIt;
@@ -1563,16 +1569,19 @@ void UWorldPartition::Tick(float DeltaSeconds)
 			}
 			else if (DirtyActor.ActorPtr.IsValid() && !DirtyActor.ActorPtr->GetPackage()->IsDirty())
 			{
-				// Transfer ownership of the last ref if actor can be pinned
-				if (Ref->GetHardRefCount() <= 1 && PinnedActors && FLoaderAdapterPinnedActors::SupportsPinning(Ref.Get()))
-				{
-					PinnedActors->AddActors({ Ref.ToHandle() });
-					DirtyActorIt.RemoveCurrent();
-				} // Clean up if a loader took a reference on the actor (ILoaderAdapter::RefreshLoadedState was called since)
-				else if (Ref->GetHardRefCount() > 1)
+				// Transfer ownership of the last ref if actor can be pinned, if not keep a temp ref so that we can refresh loaded states (which should take a add ref on unpinnable actors and prevent unloading)
+				if (Ref->GetHardRefCount() <= 1)
 				{
-					DirtyActorIt.RemoveCurrent();
-				}
+					if (PinnedActors && FLoaderAdapterPinnedActors::SupportsPinning(Ref.Get()))
+					{
+						PinnedActors->AddActors({ Ref.ToHandle() });
+					}
+					else
+					{
+						ActorReferences.Add(Ref);
+					}
+				} 
+				DirtyActorIt.RemoveCurrent();
 			}
 		}
 		else
@@ -1593,6 +1602,19 @@ void UWorldPartition::Tick(float DeltaSeconds)
 		}
 	}
 
+	if (ActorReferences.Num() > 0)
+	{
+		// Before ActorReferences get released, refresh the Loaders so that they take a reference on actors they should load.
+		IWorldPartitionActorLoaderInterface::RefreshLoadedState(false);
+
+		for (const FWorldPartitionReference& ActorReference : ActorReferences)
+		{
+			UE_CLOG(ActorReference->GetHardRefCount() <= 1, LogWorldPartition, Warning, TEXT("Releasing reference on saved actor '%s' from the DirtyActors without a loader reference."), *ActorReference->GetActor()->GetPathName());
+		}
+		// Actors should be loaded by a Loader at this point and we can release the references
+		ActorReferences.Empty();
+	}
+		
 	if (bForceGarbageCollection)
 	{
 		GEngine->ForceGarbageCollection(bForceGarbageCollectionPurge);
diff --git a/Engine/Source/Runtime/Engine/Public/WorldPartition/WorldPartition.h b/Engine/Source/Runtime/Engine/Public/WorldPartition/WorldPartition.h
--- a/Engine/Source/Runtime/Engine/Public/WorldPartition/WorldPartition.h
+++ b/Engine/Source/Runtime/Engine/Public/WorldPartition/WorldPartition.h
@@ -535,6 +535,7 @@ private:
 	ENGINE_API void HashActorDesc(FWorldPartitionActorDesc* ActorDesc);
 	ENGINE_API void UnhashActorDesc(FWorldPartitionActorDesc* ActorDesc);
 	void OnContentBundleRemovedContent(const FContentBundleEditor* ContentBundle);
+	bool IsStreamingEnabledInEditor() const;
 
 public:
 	// Editor loader adapters management
