Increase network replication priority on AGameStateBase to ensure game state transitions reliably occur on the clients

- this will indirectly prioritize these actors for inclusion in replays which is where the current issue occurs
  - increasing the net priority vs just the replay priority to safeguard against potential similar issues under high load in non-replay scenarios
  - NetPriority of 10 chosen as this is utilized by another game object to workaround similar issues given existing comments

- testing by [REDACTED_PERSON]
  - generated project specific replay without change on console platform which fails to include AGameStateBase::OnRep_ReplicatedHasBegunPlay
  - with this change in place, replay was again generated and AGameStateBase::OnRep_ReplicatedHasBegunPlay replication was included in the replay

diff --git a/Engine/Source/Runtime/Engine/Private/GameStateBase.cpp b/Engine/Source/Runtime/Engine/Private/GameStateBase.cpp
--- a/Engine/Source/Runtime/Engine/Private/GameStateBase.cpp
+++ b/Engine/Source/Runtime/Engine/Private/GameStateBase.cpp
@@ -37,6 +37,8 @@ AGameStateBase::AGameStateBase(const FObjectInitializer& ObjectInitializer)
 
 	SumServerWorldTimeSecondsDelta = 0.0;
 	NumServerWorldTimeSecondsDeltas = 0;
+
+	NetPriority = 10.f;	// We need to prioritize updates to ensure that game state transitions reliably occur on the clients.
 }
 
 const AGameModeBase* AGameStateBase::GetDefaultGameMode() const
