cluster unions should no longer manage their child components' owners' dormancy

diff --git a/Engine/Source/Runtime/Engine/Private/PhysicsEngine/ClusterUnionComponent.cpp b/Engine/Source/Runtime/Engine/Private/PhysicsEngine/ClusterUnionComponent.cpp
--- a/Engine/Source/Runtime/Engine/Private/PhysicsEngine/ClusterUnionComponent.cpp
+++ b/Engine/Source/Runtime/Engine/Private/PhysicsEngine/ClusterUnionComponent.cpp
@@ -25,9 +25,6 @@ namespace
 	bool bUseClusterUnionAccelerationStructure = true;
 	FAutoConsoleVariableRef CVarUseClusterUnionAccelerationStructure(TEXT("ClusterUnion.UseAccelerationStructure"), bUseClusterUnionAccelerationStructure, TEXT("Whether component level sweeps and overlaps against cluster unions should use an acceleration structure instead."));
 
-	bool bClusterUnionCallFlushNetDormancy = true;
-	FAutoConsoleVariableRef CVarClusterUnionCallFlushNetDormanc(TEXT("clusterunion.callflushnetdormancy"), bClusterUnionCallFlushNetDormancy, TEXT("Whether or not to call FlushNetDormancy"));
-
 	// TODO: Should this be exposed in Chaos instead?
 	using FAccelerationStructure = Chaos::TAABBTree<FExternalSpatialAccelerationPayload, Chaos::TAABBTreeLeafArray<FExternalSpatialAccelerationPayload>>;
 
@@ -788,14 +785,6 @@ void UClusterUnionComponent::HandleAddOrModifiedClusteredComponent(UPrimitiveCom
 		{
 			ReplicatedProxy->SetParticleChildToParent(Kvp.Key, Kvp.Value);
 		}
-
-		if (bClusterUnionCallFlushNetDormancy)
-		{
-			if (AActor* Owner = ChangedComponent->GetOwner())
-			{
-				Owner->FlushNetDormancy();
-			}
-		}
 	}
 
 	TArray<Chaos::FPhysicsObjectHandle> AllPhysicsObjects = ChangedComponent->GetAllPhysicsObjects();
@@ -872,11 +861,6 @@ void UClusterUnionComponent::HandleRemovedClusteredComponent(TObjectKey<UPrimiti
 					ActorToComponents.Remove(Owner);
 				}
 			}
-
-			if (bClusterUnionCallFlushNetDormancy)
-			{
-				Owner->FlushNetDormancy();
-			}
 		}
 
 		if (UPrimitiveComponent* ChangedComponent = RemovedComponent.ResolveObjectPtr())
