PhysicsReplication.cpp will no longer apply rigid body updates if the primitive component is attached to a parent.

#rnx

diff --git a/Engine/Source/Runtime/Engine/Private/PhysicsEngine/PhysicsReplication.cpp b/Engine/Source/Runtime/Engine/Private/PhysicsEngine/PhysicsReplication.cpp
--- a/Engine/Source/Runtime/Engine/Private/PhysicsEngine/PhysicsReplication.cpp
+++ b/Engine/Source/Runtime/Engine/Private/PhysicsEngine/PhysicsReplication.cpp
@@ -267,40 +267,43 @@ void FPhysicsReplication::OnTick(float DeltaSeconds, TMap<TWeakObjectPtr<UPrimit
 		bool bRemoveItr = false;
 		if (UPrimitiveComponent* PrimComp = Itr.Key().Get())
 		{
-			if (FBodyInstance* BI = PrimComp->GetBodyInstance(Itr.Value().BoneName))
+			if (PrimComp->GetAttachParent() == nullptr)
 			{
-				FReplicatedPhysicsTarget& PhysicsTarget = Itr.Value();
-				FRigidBodyState& UpdatedState = PhysicsTarget.TargetState;
-				bool bUpdated = false;
-				if (AActor* OwningActor = PrimComp->GetOwner())
+				if (FBodyInstance* BI = PrimComp->GetBodyInstance(Itr.Value().BoneName))
 				{
-					const ENetRole OwnerRole = OwningActor->GetLocalRole();
-					const bool bIsSimulated = OwnerRole == ROLE_SimulatedProxy;
-					const bool bIsReplicatedAutonomous = OwnerRole == ROLE_AutonomousProxy && PrimComp->bReplicatePhysicsToAutonomousProxy;
-					if (bIsSimulated || bIsReplicatedAutonomous)
+					FReplicatedPhysicsTarget& PhysicsTarget = Itr.Value();
+					FRigidBodyState& UpdatedState = PhysicsTarget.TargetState;
+					bool bUpdated = false;
+					if (AActor* OwningActor = PrimComp->GetOwner())
 					{
-						// Get the ping of this thing's owner. If nobody owns it,
-						// then it's server authoritative.
-						const float OwnerPing = GetOwnerPing(OwningActor, PhysicsTarget);
-
-						// Get the total ping - this approximates the time since the update was
-						// actually generated on the machine that is doing the authoritative sim.
-						// NOTE: We divide by 2 to approximate 1-way ping from 2-way ping.
-						const float PingSecondsOneWay = (LocalPing + OwnerPing) * 0.5f * 0.001f;
-
-						if (UpdatedState.Flags & ERigidBodyFlags::NeedsUpdate)
+						const ENetRole OwnerRole = OwningActor->GetLocalRole();
+						const bool bIsSimulated = OwnerRole == ROLE_SimulatedProxy;
+						const bool bIsReplicatedAutonomous = OwnerRole == ROLE_AutonomousProxy && PrimComp->bReplicatePhysicsToAutonomousProxy;
+						if (bIsSimulated || bIsReplicatedAutonomous)
 						{
-							const int32 LocalFrame = PhysicsTarget.ServerFrame - LocalFrameOffset;
-							const bool bRestoredState = ApplyRigidBodyState(DeltaSeconds, BI, PhysicsTarget, PhysicErrorCorrection, PingSecondsOneWay, LocalFrame, 0);
+							// Get the ping of this thing's owner. If nobody owns it,
+							// then it's server authoritative.
+							const float OwnerPing = GetOwnerPing(OwningActor, PhysicsTarget);
 
-							// Need to update the component to match new position.
-							if (PhysicsReplicationCVars::SkipSkeletalRepOptimization == 0 || Cast<USkeletalMeshComponent>(PrimComp) == nullptr)	//simulated skeletal mesh does its own polling of physics results so we don't need to call this as it'll happen at the end of the physics sim
-							{
-								PrimComp->SyncComponentToRBPhysics();
-							}
-							if (bRestoredState)
+							// Get the total ping - this approximates the time since the update was
+							// actually generated on the machine that is doing the authoritative sim.
+							// NOTE: We divide by 2 to approximate 1-way ping from 2-way ping.
+							const float PingSecondsOneWay = (LocalPing + OwnerPing) * 0.5f * 0.001f;
+
+							if (UpdatedState.Flags & ERigidBodyFlags::NeedsUpdate)
 							{
-								bRemoveItr = true;
+								const int32 LocalFrame = PhysicsTarget.ServerFrame - LocalFrameOffset;
+								const bool bRestoredState = ApplyRigidBodyState(DeltaSeconds, BI, PhysicsTarget, PhysicErrorCorrection, PingSecondsOneWay, LocalFrame, 0);
+
+								// Need to update the component to match new position.
+								if (PhysicsReplicationCVars::SkipSkeletalRepOptimization == 0 || Cast<USkeletalMeshComponent>(PrimComp) == nullptr)	//simulated skeletal mesh does its own polling of physics results so we don't need to call this as it'll happen at the end of the physics sim
+								{
+									PrimComp->SyncComponentToRBPhysics();
+								}
+								if (bRestoredState)
+								{
+									bRemoveItr = true;
+								}
 							}
 						}
 					}
