WorldPartition: Added IsListedInSceneOutliner in WorldPartitionActorDesc and use it to hide unloaded ActorDescItems in the World Outliner. It is also used to hide actor bounds in World Partition Editor.

diff --git a/Engine/Source/Editor/SceneOutliner/Private/ActorDescTreeItem.cpp b/Engine/Source/Editor/SceneOutliner/Private/ActorDescTreeItem.cpp
--- a/Engine/Source/Editor/SceneOutliner/Private/ActorDescTreeItem.cpp
+++ b/Engine/Source/Editor/SceneOutliner/Private/ActorDescTreeItem.cpp
@@ -256,7 +256,7 @@ FSceneOutlinerTreeItemID FActorDescTreeItem::ComputeTreeItemID(FGuid InActorGuid
 
 bool FActorDescTreeItem::ShouldDisplayInOutliner(const FWorldPartitionActorDesc* ActorDesc)
 {
-	return ActorDesc && (ActorDesc->GetActorIsRuntimeOnly() || ActorDesc->IsEditorRelevant());
+	return ActorDesc && ActorDesc->IsListedInSceneOutliner() && (ActorDesc->GetActorIsRuntimeOnly() || ActorDesc->IsEditorRelevant());
 }
 
 FSceneOutlinerTreeItemID FActorDescTreeItem::GetID() const
diff --git a/Engine/Source/Editor/WorldPartitionEditor/Private/WorldPartition/SWorldPartitionEditorGrid2D.cpp b/Engine/Source/Editor/WorldPartitionEditor/Private/WorldPartition/SWorldPartitionEditorGrid2D.cpp
--- a/Engine/Source/Editor/WorldPartitionEditor/Private/WorldPartition/SWorldPartitionEditorGrid2D.cpp
+++ b/Engine/Source/Editor/WorldPartitionEditor/Private/WorldPartition/SWorldPartitionEditorGrid2D.cpp
@@ -1067,7 +1067,7 @@ void SWorldPartitionEditorGrid2D::Tick(const FGeometry& AllottedGeometry, const
 	{
 		if (bShowActors)
 		{
-			if (!GetWorldPartition()->IsStreamingEnabled() || ActorDesc->GetIsSpatiallyLoaded())
+			if (ActorDesc->IsListedInSceneOutliner() && (!GetWorldPartition()->IsStreamingEnabled() || ActorDesc->GetIsSpatiallyLoaded()))
 			{
 				ShownActorGuids.Add(ActorDesc->GetGuid());
 			}
diff --git a/Engine/Source/Runtime/Core/Public/UObject/UE5ReleaseStreamObjectVersion.h b/Engine/Source/Runtime/Core/Public/UObject/UE5ReleaseStreamObjectVersion.h
--- a/Engine/Source/Runtime/Core/Public/UObject/UE5ReleaseStreamObjectVersion.h
+++ b/Engine/Source/Runtime/Core/Public/UObject/UE5ReleaseStreamObjectVersion.h
@@ -147,6 +147,9 @@ struct FUE5ReleaseStreamObjectVersion
 
 		// Changed default Local Exposure Contrast Scale from 1.0 to 0.8
 		LocalExposureDefaultChangeFrom1,
+
+		// Serialize bActorIsListedInSceneOutliner in WorldPartitionActorDesc
+		WorldPartitionActorDescSerializeActorIsListedInSceneOutliner,
 		
 		// -----<new versions can be added above this line>-------------------------------------------------
 		VersionPlusOne,
diff --git a/Engine/Source/Runtime/Engine/Private/WorldPartition/WorldPartitionActorDesc.cpp b/Engine/Source/Runtime/Engine/Private/WorldPartition/WorldPartitionActorDesc.cpp
--- a/Engine/Source/Runtime/Engine/Private/WorldPartition/WorldPartitionActorDesc.cpp
+++ b/Engine/Source/Runtime/Engine/Private/WorldPartition/WorldPartitionActorDesc.cpp
@@ -45,6 +45,7 @@ FWorldPartitionActorDesc::FWorldPartitionActorDesc()
 	, bActorIsRuntimeOnly(false)
 	, bActorIsMainWorldOnly(false)
 	, bActorIsHLODRelevant(false)
+	, bActorIsListedInSceneOutliner(true)
 	, bIsUsingDataLayerAsset(false)
 	, bIsBoundsValid(false)
 	, SoftRefCount(0)
@@ -97,6 +98,7 @@ void FWorldPartitionActorDesc::Init(const AActor* InActor)
 	bActorIsEditorOnly = InActor->IsEditorOnly();
 	bActorIsRuntimeOnly = InActor->IsRuntimeOnly();
 	bActorIsHLODRelevant = InActor->IsHLODRelevant();
+	bActorIsListedInSceneOutliner = InActor->IsListedInSceneOutliner();
 	bActorIsMainWorldOnly = InActor->bIsMainWorldOnly;
 	HLODLayer = InActor->GetHLODLayer() ? FSoftObjectPath(InActor->GetHLODLayer()->GetPathName()) : FSoftObjectPath();
 	
@@ -266,6 +268,7 @@ bool FWorldPartitionActorDesc::Equals(const FWorldPartitionActorDesc* Other) con
 		bActorIsRuntimeOnly == Other->bActorIsRuntimeOnly &&
 		bActorIsMainWorldOnly == Other->bActorIsMainWorldOnly &&
 		bActorIsHLODRelevant == Other->bActorIsHLODRelevant &&
+		bActorIsListedInSceneOutliner == Other->bActorIsListedInSceneOutliner &&
 		bIsUsingDataLayerAsset == Other->bIsUsingDataLayerAsset &&
 		HLODLayer == Other->HLODLayer &&
 		FolderPath == Other->FolderPath &&
@@ -404,7 +407,7 @@ FString FWorldPartitionActorDesc::ToString(EToStringMode Mode) const
 		}
 
 		Result.Appendf(
-			TEXT(" BaseClass:%s NativeClass:%s Name:%s Label:%s SpatiallyLoaded:%s Bounds:%s RuntimeGrid:%s EditorOnly:%s RuntimeOnly:%s HLODRelevant:%s IsMainWorldOnly:%s"),
+			TEXT(" BaseClass:%s NativeClass:%s Name:%s Label:%s SpatiallyLoaded:%s Bounds:%s RuntimeGrid:%s EditorOnly:%s RuntimeOnly:%s HLODRelevant:%s ListedInSceneOutliner:%s IsMainWorldOnly:%s"),
 			*BaseClass.ToString(), 
 			*NativeClass.ToString(), 
 			*GetActorName().ToString(),
@@ -415,6 +418,7 @@ FString FWorldPartitionActorDesc::ToString(EToStringMode Mode) const
 			GetBoolStr(bActorIsEditorOnly),
 			GetBoolStr(bActorIsRuntimeOnly),
 			GetBoolStr(bActorIsHLODRelevant),
+			GetBoolStr(bActorIsListedInSceneOutliner),
 			GetBoolStr(IsMainWorldOnly())
 		);
 
@@ -674,6 +678,11 @@ PRAGMA_ENABLE_DEPRECATION_WARNINGS
 		Ar << TDeltaSerialize<bool>(bActorIsMainWorldOnly);
 	}
 
+	if (Ar.CustomVer(FUE5ReleaseStreamObjectVersion::GUID) >= FUE5ReleaseStreamObjectVersion::WorldPartitionActorDescSerializeActorIsListedInSceneOutliner)
+	{
+		Ar << TDeltaSerialize<bool>(bActorIsListedInSceneOutliner);
+	}
+
 	// Fixup redirected data layer asset paths
 	if (Ar.IsLoading() && bIsUsingDataLayerAsset)
 	{
@@ -741,6 +750,11 @@ bool FWorldPartitionActorDesc::IsMainWorldOnly() const
 	return bActorIsMainWorldOnly || CastChecked<AActor>(ActorNativeClass->GetDefaultObject())->IsMainWorldOnly();
 }
 
+bool FWorldPartitionActorDesc::IsListedInSceneOutliner() const
+{
+	return bActorIsListedInSceneOutliner;
+}
+
 bool FWorldPartitionActorDesc::IsEditorRelevant() const
 {
 	if (GetActorIsRuntimeOnly())
diff --git a/Engine/Source/Runtime/Engine/Public/WorldPartition/WorldPartitionActorDesc.h b/Engine/Source/Runtime/Engine/Public/WorldPartition/WorldPartitionActorDesc.h
--- a/Engine/Source/Runtime/Engine/Public/WorldPartition/WorldPartitionActorDesc.h
+++ b/Engine/Source/Runtime/Engine/Public/WorldPartition/WorldPartitionActorDesc.h
@@ -269,6 +269,7 @@ public:
 	static ENGINE_API void RegisterActorDescDeprecator(TSubclassOf<AActor> ActorClass, const FActorDescDeprecator& Deprecator);
 
 	ENGINE_API bool IsMainWorldOnly() const;
+	ENGINE_API bool IsListedInSceneOutliner() const;
 protected:
 	ENGINE_API FWorldPartitionActorDesc();
 
@@ -300,6 +301,7 @@ protected:
 	bool							bActorIsRuntimeOnly;
 	bool							bActorIsMainWorldOnly;
 	bool							bActorIsHLODRelevant;
+	bool							bActorIsListedInSceneOutliner;
 	bool							bIsUsingDataLayerAsset; // Used to know if DataLayers array represents DataLayers Asset paths or the FNames of the deprecated version of Data Layers
 	bool							bIsBoundsValid;
 	FSoftObjectPath					HLODLayer;
