SkyAtmosphere - extending multiple scattering to go further than 2. But the UI slider remains in [0,2].

diff --git a/Engine/Source/Runtime/Engine/Classes/Components/SkyAtmosphereComponent.h b/Engine/Source/Runtime/Engine/Classes/Components/SkyAtmosphereComponent.h
--- a/Engine/Source/Runtime/Engine/Classes/Components/SkyAtmosphereComponent.h
+++ b/Engine/Source/Runtime/Engine/Classes/Components/SkyAtmosphereComponent.h
@@ -72,7 +72,7 @@ class USkyAtmosphereComponent : public USceneComponent
 	 * Multiple scattering is evaluated using a dual scattering approach. 
 	 * A value of 2 is recommended to better represent default atmosphere when r.SkyAtmosphere.MultiScatteringLUT.HighQuality=0. 
 	 */
-	UPROPERTY(EditAnywhere, BlueprintReadOnly, interp, Category = "Atmosphere", meta = (DisplayName = "MultiScattering", UIMin = 0.0, UIMax = 1.0, ClampMin = 0.0, ClampMax = 2.0))
+	UPROPERTY(EditAnywhere, BlueprintReadOnly, interp, Category = "Atmosphere", meta = (DisplayName = "MultiScattering", UIMin = 0.0, UIMax = 2.0, ClampMin = 0.0, ClampMax = 100.0))
 	float MultiScatteringFactor;
 
 	/**
diff --git a/Engine/Source/Runtime/Engine/Public/Rendering/SkyAtmosphereCommonData.cpp b/Engine/Source/Runtime/Engine/Public/Rendering/SkyAtmosphereCommonData.cpp
--- a/Engine/Source/Runtime/Engine/Public/Rendering/SkyAtmosphereCommonData.cpp
+++ b/Engine/Source/Runtime/Engine/Public/Rendering/SkyAtmosphereCommonData.cpp
@@ -43,7 +43,7 @@ FAtmosphereSetup::FAtmosphereSetup(const USkyAtmosphereComponent& SkyAtmosphereC
 	BottomRadiusKm = SkyAtmosphereComponent.BottomRadius;
 	TopRadiusKm = SkyAtmosphereComponent.BottomRadius + FMath::Max(0.1f, SkyAtmosphereComponent.AtmosphereHeight);
 	GroundAlbedo = FLinearColor(SkyAtmosphereComponent.GroundAlbedo);
-	MultiScatteringFactor = FMath::Clamp(SkyAtmosphereComponent.MultiScatteringFactor, 0.0f, 2.0f);
+	MultiScatteringFactor = FMath::Clamp(SkyAtmosphereComponent.MultiScatteringFactor, 0.0f, 100.0f);
 
 	auto ConvertCoefficientsFromSRGBToWorkingColorSpace = [](FLinearColor CoeffSRGB)
 	{
