VulkanRHI: Show toast when Virtual Shadow Maps are enabled but only SM5 Vulkan is enabled (add warning in both Windows and Linux editors).

#rnx

diff --git a/Engine/Source/Runtime/Engine/Private/RendererSettings.cpp b/Engine/Source/Runtime/Engine/Private/RendererSettings.cpp
--- a/Engine/Source/Runtime/Engine/Private/RendererSettings.cpp
+++ b/Engine/Source/Runtime/Engine/Private/RendererSettings.cpp
@@ -17,7 +17,7 @@
 #include "Misc/MessageDialog.h"
 #include "HAL/PlatformFileManager.h"
 
-#if PLATFORM_WINDOWS
+#if PLATFORM_WINDOWS || PLATFORM_LINUX
 #include "Framework/Docking/TabManager.h"
 #include "Framework/Notifications/NotificationManager.h"
 #include "Misc/ConfigCacheIni.h"
@@ -385,27 +385,49 @@ bool URendererSettings::CanEditChange(const FProperty* InProperty) const
 
 void URendererSettings::CheckForMissingShaderModels()
 {
-	// Don't show the SM6 toasts on non-Windows platforms to avoid confusion around platform requirements.
-#if PLATFORM_WINDOWS
+	// Don't show the SM6 toasts on non-Windows/Linux platforms to avoid confusion around platform requirements.
+#if PLATFORM_WINDOWS || PLATFORM_LINUX
 	if (GIsEditor && ShadowMapMethod == EShadowMapMethod::VirtualShadowMaps)
 	{
-		TArray<FString> D3D11TargetedShaderFormats;
-		GConfig->GetArray(TEXT("/Script/WindowsTargetPlatform.WindowsTargetSettings"), TEXT("D3D11TargetedShaderFormats"), D3D11TargetedShaderFormats, GEngineIni);
+		auto CopySM6Format = [](const TCHAR* ShaderFormatName, const TArray<FString>& SrcArray, TArray<FString>& DstArray)
+		{
+			if (SrcArray.Contains(ShaderFormatName))
+			{
+				DstArray.AddUnique(ShaderFormatName);
+			}
+		};
 
+		// Gather all Windows shader format settings
+		TArray<FString> D3D11TargetedShaderFormats;
 		TArray<FString> D3D12TargetedShaderFormats;
+		TArray<FString> WindowsVulkanTargetedShaderFormats;
+		TArray<FString> WindowsTargetedRHIs;
+		GConfig->GetArray(TEXT("/Script/WindowsTargetPlatform.WindowsTargetSettings"), TEXT("D3D11TargetedShaderFormats"), D3D11TargetedShaderFormats, GEngineIni);
 		GConfig->GetArray(TEXT("/Script/WindowsTargetPlatform.WindowsTargetSettings"), TEXT("D3D12TargetedShaderFormats"), D3D12TargetedShaderFormats, GEngineIni);
+		GConfig->GetArray(TEXT("/Script/WindowsTargetPlatform.WindowsTargetSettings"), TEXT("VulkanTargetedShaderFormats"), WindowsVulkanTargetedShaderFormats, GEngineIni);
+		GConfig->GetArray(TEXT("/Script/WindowsTargetPlatform.WindowsTargetSettings"), TEXT("TargetedRHIs"), WindowsTargetedRHIs, GEngineIni);
+		CopySM6Format(TEXT("PCD3D_SM6"), WindowsTargetedRHIs, D3D12TargetedShaderFormats);
+		CopySM6Format(TEXT("SF_VULKAN_SM6"), WindowsTargetedRHIs, WindowsVulkanTargetedShaderFormats);
 
-		TArray<FString> TargetedRHIs;
-		GConfig->GetArray(TEXT("/Script/WindowsTargetPlatform.WindowsTargetSettings"), TEXT("TargetedRHIs"), TargetedRHIs, GEngineIni);
+		// Gather all Linux shader format settings
+		TArray<FString> LinuxVulkanTargetedShaderFormats;
+		TArray<FString> LinuxTargetedRHIs;
+		GConfig->GetArray(TEXT("/Script/LinuxTargetPlatform.LinuxTargetSettings"), TEXT("VulkanTargetedShaderFormats"), LinuxVulkanTargetedShaderFormats, GEngineIni);
+		GConfig->GetArray(TEXT("/Script/LinuxTargetPlatform.LinuxTargetSettings"), TEXT("TargetedRHIs"), LinuxTargetedRHIs, GEngineIni);
+		CopySM6Format(TEXT("SF_VULKAN_SM6"), LinuxTargetedRHIs, LinuxVulkanTargetedShaderFormats);
 
-		if (TargetedRHIs.Contains(TEXT("PCD3D_SM6")))
-		{
-			D3D12TargetedShaderFormats.AddUnique(TEXT("PCD3D_SM6"));
-		}
+		// TODO: Gather all Mac shader format settings
 
 		const bool bProjectUsesD3D = (D3D11TargetedShaderFormats.Num() + D3D12TargetedShaderFormats.Num()) > 0;
+		const bool bProjectMissingD3DSM6 = (bProjectUsesD3D && !D3D12TargetedShaderFormats.Contains(TEXT("PCD3D_SM6")));
 
-		if (bProjectUsesD3D && !D3D12TargetedShaderFormats.Contains(TEXT("PCD3D_SM6")))
+		const bool bProjectUsesWindowsVulkan = (WindowsVulkanTargetedShaderFormats.Num() > 0);
+		const bool bProjectMissingWindowsVulkanSM6 = (bProjectUsesWindowsVulkan && !WindowsVulkanTargetedShaderFormats.Contains(TEXT("SF_VULKAN_SM6")));
+
+		const bool bProjectUsesLinuxVulkan = (LinuxTargetedRHIs.Num() > 0) || (LinuxVulkanTargetedShaderFormats.Num() > 0);
+		const bool bProjectMissingLinuxVulkanSM6 = (bProjectUsesLinuxVulkan && !LinuxVulkanTargetedShaderFormats.Contains(TEXT("SF_VULKAN_SM6")));
+
+		if (bProjectMissingD3DSM6 || bProjectMissingWindowsVulkanSM6 || bProjectMissingLinuxVulkanSM6)
 		{
 			auto DismissNotification = [this]()
 			{
@@ -435,14 +457,26 @@ void URendererSettings::CheckForMissingShaderModels()
 				SNotificationItem::CS_None));
 
 			Info.Text = LOCTEXT("NeedProjectSettings", "Missing Project Settings!");
-			Info.SubText = LOCTEXT("VirtualShadowMapsNeedsSM6Setting", "Shader Model 6 (SM6) is required to use Virtual Shadow Maps. Please enable this in:\n  Project Settings -> Platforms -> Windows -> D3D12 Targeted Shader Formats\nVirtual shadow maps will not work until this is enabled.");
 			Info.HyperlinkText = LOCTEXT("ProjectSettingsHyperlinkText", "Open Project Settings");
 			Info.Hyperlink = FSimpleDelegate::CreateLambda(OpenProjectSettings);
 
+			if (bProjectMissingD3DSM6)
+			{
+				Info.SubText = LOCTEXT("VirtualShadowMapsNeedsSM6Setting", "Shader Model 6 (SM6) is required to use Virtual Shadow Maps. Please enable this in:\n  Project Settings -> Platforms -> Windows -> D3D12 Targeted Shader Formats\nVirtual shadow maps will not work until this is enabled.");
+			}
+			else if (bProjectMissingWindowsVulkanSM6)
+			{
+				Info.SubText = LOCTEXT("VirtualShadowMapsNeedsVulkanSM6WindowsSetting", "Shader Model 6 (SM6) is required to use Virtual Shadow Maps. Please enable this in:\n  Project Settings -> Platforms -> Windows -> Vulkan Targeted Shader Formats\nVirtual shadow maps will not work in Vulkan on Windows until this is enabled.");
+			}
+			else if (bProjectMissingLinuxVulkanSM6)
+			{
+				Info.SubText = LOCTEXT("VirtualShadowMapsNeedsVulkanSM6LinuxSetting", "Shader Model 6 (SM6) is required to use Virtual Shadow Maps. Please enable this in:\n  Project Settings -> Platforms -> Linux -> Targeted RHIs\nVirtual shadow maps will not work in Vulkan on Linux until this is enabled.");
+			}
+
 			ShaderModelNotificationPtr = FSlateNotificationManager::Get().AddNotification(Info);
 		}
 	}
-#endif // PLATFORM_WINDOWS
+#endif // PLATFORM_WINDOWS || PLATFORM_LINUX
 }
 #endif // #if WITH_EDITOR
 
