AnimSequence cooking: avoid softlock when skeleton data is missing. AnimSequence::IsCachedCookedPlatformDataLoaded needs to return true when it encounters an error that will not be corrected.
#rnx

diff --git a/Engine/Source/Runtime/Engine/Private/Animation/AnimSequence.cpp b/Engine/Source/Runtime/Engine/Private/Animation/AnimSequence.cpp
--- a/Engine/Source/Runtime/Engine/Private/Animation/AnimSequence.cpp
+++ b/Engine/Source/Runtime/Engine/Private/Animation/AnimSequence.cpp
@@ -1159,7 +1159,7 @@ bool UAnimSequence::IsCachedCookedPlatformDataLoaded(const ITargetPlatform* Targ
 	if (KeyHash.IsZero())
 	{
 		UE_LOG(LogAnimation, Warning, TEXT("Zero key hash compressed animation data for %s requested platform %s"), *GetName(), *TargetPlatform->PlatformName());
-		return false;
+		return true;
 	}
 
 	if (PollCacheDerivedData(KeyHash))
