Disallow instance culling occlusion queries for HISM

HISM typically produces a huge number of instances that rarely benefit from per-instance OQ, as they're hierarchically culled by HISM system itself.

diff --git a/Engine/Source/Runtime/Engine/Public/HierarchicalStaticMeshSceneProxy.h b/Engine/Source/Runtime/Engine/Public/HierarchicalStaticMeshSceneProxy.h
--- a/Engine/Source/Runtime/Engine/Public/HierarchicalStaticMeshSceneProxy.h
+++ b/Engine/Source/Runtime/Engine/Public/HierarchicalStaticMeshSceneProxy.h
@@ -83,7 +83,9 @@ public:
 	ENGINE_API virtual const TArray<FBoxSphereBounds>* GetOcclusionQueries(const FSceneView* View) const override;
 	
 	ENGINE_API virtual void AcceptOcclusionResults(const FSceneView* View, TArray<bool>* Results, int32 ResultsStart, int32 NumResults) override;
-	
+
+	ENGINE_API virtual bool AllowInstanceCullingOcclusionQueries() const override { return false; }
+
 	virtual bool HasSubprimitiveOcclusionQueries() const override
 	{
 		return FirstOcclusionNode > 0;
