[BugFix] False positive graph node property diffs

diff --git a/Engine/Source/Runtime/Engine/Private/EdGraph/EdGraphNode.cpp b/Engine/Source/Runtime/Engine/Private/EdGraph/EdGraphNode.cpp
--- a/Engine/Source/Runtime/Engine/Private/EdGraph/EdGraphNode.cpp
+++ b/Engine/Source/Runtime/Engine/Private/EdGraph/EdGraphNode.cpp
@@ -2,6 +2,7 @@
 
 #include "EdGraph/EdGraphNode.h"
 
+#include "DiffUtils.h"
 #include "Serialization/PropertyLocalizationDataGathering.h"
 #include "UObject/BlueprintsObjectVersion.h"
 #include "UObject/FrameworkObjectVersion.h"
@@ -315,10 +316,8 @@ void UEdGraphNode::DiffProperties(UStruct* StructA, UStruct* StructB, uint8* Dat
 			continue;
 		}
 
-		const FString ValueStringA = GetPropertyNameAndValueForDiff(Prop, Prop->ContainerPtrToValuePtr<uint8>(DataA));
-		const FString ValueStringB = GetPropertyNameAndValueForDiff(PropB, PropB->ContainerPtrToValuePtr<uint8>(DataB));
-
-		if (ValueStringA != ValueStringB)
+		TArray<FPropertySoftPath> DifferingSubProperties;
+		if (!DiffUtils::Identical(FResolvedProperty(DataA, Prop), FResolvedProperty(DataB, PropB), {}, DifferingSubProperties))
 		{
 			// Only bother setting up the display data if we're storing the result
 			if (Results.CanStoreResults())
