Revert default Local Exposure values in PPV back to 1.0.
- This approach to change default values caused issues when loading existing blueprint classes that inherit from a blueprint class that overrides Local Exposure parameters.

diff --git a/Engine/Source/Runtime/Engine/Private/Scene.cpp b/Engine/Source/Runtime/Engine/Private/Scene.cpp
--- a/Engine/Source/Runtime/Engine/Private/Scene.cpp
+++ b/Engine/Source/Runtime/Engine/Private/Scene.cpp
@@ -509,8 +509,8 @@ FPostProcessSettings::FPostProcessSettings()
 	AutoExposureApplyPhysicalCameraExposure = 1;
 
 	LocalExposureContrastScale_DEPRECATED = 1.0f;
-	LocalExposureHighlightContrastScale = 0.8f;
-	LocalExposureShadowContrastScale = 0.8f;
+	LocalExposureHighlightContrastScale = 1.0f;
+	LocalExposureShadowContrastScale = 1.0f;
 	LocalExposureDetailStrength = 1.0f;
 	LocalExposureBlurredLuminanceBlend = 0.6f;
 	LocalExposureBlurredLuminanceKernelSizePercent = 50.0f;
@@ -666,12 +666,6 @@ bool FPostProcessSettings::Serialize(FArchive& Ar)
 	Ar.UsingCustomVersion(FReleaseObjectVersion::GUID);
 	Ar.UsingCustomVersion(FUE5ReleaseStreamObjectVersion::GUID);
 
-	if (Ar.CustomVer(FUE5ReleaseStreamObjectVersion::GUID) < FUE5ReleaseStreamObjectVersion::LocalExposureDefaultChangeFrom1)
-	{
-		LocalExposureHighlightContrastScale = 1.0f;
-		LocalExposureShadowContrastScale = 1.0f;
-	}
-
 	// Don't actually serialize, just write the custom version for PostSerialize
 	return false;
 }
