Add GNearClippingPlane_RenderThread to be accessed on only on the RenderThread and use GNearClippingPlane otherwise. This avoids the need for the FlushRenderingCommands call which was causing a ~50ms hitch on Switch.

diff --git a/Engine/Plugins/Runtime/OpenXR/Source/OpenXRHMD/Private/OpenXRHMD.cpp b/Engine/Plugins/Runtime/OpenXR/Source/OpenXRHMD/Private/OpenXRHMD.cpp
--- a/Engine/Plugins/Runtime/OpenXR/Source/OpenXRHMD/Private/OpenXRHMD.cpp
+++ b/Engine/Plugins/Runtime/OpenXR/Source/OpenXRHMD/Private/OpenXRHMD.cpp
@@ -1069,7 +1069,7 @@ FMatrix FOpenXRHMD::GetStereoProjectionMatrix(const int32 ViewIndex) const
 	Fov.angleLeft = tan(Fov.angleLeft);
 	Fov.angleRight = tan(Fov.angleRight);
 
-	float ZNear = GNearClippingPlane;
+	float ZNear = GNearClippingPlane_RenderThread;
 	float SumRL = (Fov.angleRight + Fov.angleLeft);
 	float SumTB = (Fov.angleUp + Fov.angleDown);
 	float InvRL = (1.0f / (Fov.angleRight - Fov.angleLeft));
@@ -1153,7 +1153,7 @@ void FOpenXRHMD::PostRenderViewFamily_RenderThread(FRDGBuilder& GraphBuilder, FS
 {
 	check(IsInRenderingThread());
 
-	const float NearZ = GNearClippingPlane / GetWorldToMetersScale();
+	const float NearZ = GNearClippingPlane_RenderThread / GetWorldToMetersScale();
 
 	for (int32 ViewIndex = 0; ViewIndex < PipelinedLayerStateRendering.ColorImages.Num(); ViewIndex++)
 	{
diff --git a/Engine/Source/Editor/VREditor/Private/VREditorModeBase.cpp b/Engine/Source/Editor/VREditor/Private/VREditorModeBase.cpp
--- a/Engine/Source/Editor/VREditor/Private/VREditorModeBase.cpp
+++ b/Engine/Source/Editor/VREditor/Private/VREditorModeBase.cpp
@@ -11,6 +11,7 @@
 #include "LevelViewportActions.h"
 #include "Modules/ModuleManager.h"
 #include "SLevelViewport.h"
+#include "RenderCore.h"
 
 
 namespace VREd
@@ -150,7 +151,8 @@ void UVREditorModeBase::StartViewport(TSharedPtr<SLevelViewport> Viewport)
 
 	// Force tiny near clip plane distance, because user can scale themselves to be very small.
 	SavedEditorState.NearClipPlane = GNearClippingPlane;
-	GNearClippingPlane = VREd::DefaultVRNearClipPlane->GetFloat();
+
+	SetNearClipPlaneGlobals(VREd::DefaultVRNearClipPlane->GetFloat());
 
 	SavedEditorState.bOnScreenMessages = GAreScreenMessagesEnabled;
 	GAreScreenMessagesEnabled = false;
@@ -225,7 +227,8 @@ void UVREditorModeBase::CloseViewport( const bool bShouldDisableStereo )
 
 		VRViewportClient.SetRealtime(SavedEditorState.bRealTime);
 
-		GNearClippingPlane = SavedEditorState.NearClipPlane;
+		SetNearClipPlaneGlobals(SavedEditorState.NearClipPlane);
+
 		GAreScreenMessagesEnabled = SavedEditorState.bOnScreenMessages;
 
 		if (bActuallyUsingVR)
diff --git a/Engine/Source/Editor/VREditor/Private/VREditorPlacement.cpp b/Engine/Source/Editor/VREditor/Private/VREditorPlacement.cpp
--- a/Engine/Source/Editor/VREditor/Private/VREditorPlacement.cpp
+++ b/Engine/Source/Editor/VREditor/Private/VREditorPlacement.cpp
@@ -30,6 +30,8 @@
 #include "SLevelViewport.h"
 #include "LevelEditorActions.h"
 
+#include "RenderCore.h"
+
 #define LOCTEXT_NAMESPACE "VREditor"
 
 
@@ -99,7 +101,8 @@ void UVREditorPlacement::UpdateNearClipPlaneOnScaleChange(const float NewWorldTo
 {
 	// Adjust the clipping plane for the user's scale, but don't let it be larger than the engine default
 	const float DefaultWorldToMetersScale = VRMode->GetSavedEditorState().WorldToMetersScale;
-	GNearClippingPlane = FMath::Min((VRMode->GetDefaultVRNearClipPlane()) * (NewWorldToMetersScale / DefaultWorldToMetersScale), VRMode->GetSavedEditorState().NearClipPlane);
+
+	SetNearClipPlaneGlobals(FMath::Min((VRMode->GetDefaultVRNearClipPlane()) * (NewWorldToMetersScale / DefaultWorldToMetersScale), VRMode->GetSavedEditorState().NearClipPlane));
 }
 
 
diff --git a/Engine/Source/Editor/VREditor/VREditor.Build.cs b/Engine/Source/Editor/VREditor/VREditor.Build.cs
--- a/Engine/Source/Editor/VREditor/VREditor.Build.cs
+++ b/Engine/Source/Editor/VREditor/VREditor.Build.cs
@@ -22,7 +22,7 @@ namespace UnrealBuildTool.Rules
                     "InputCore",
 				    "Slate",
 					"SlateCore",
-                    
+                    "RenderCore",
 					"EditorFramework",
 					"UnrealEd",
 					"UMG",
diff --git a/Engine/Source/Runtime/Core/Private/Misc/CoreGlobals.cpp b/Engine/Source/Runtime/Core/Private/Misc/CoreGlobals.cpp
--- a/Engine/Source/Runtime/Core/Private/Misc/CoreGlobals.cpp
+++ b/Engine/Source/Runtime/Core/Private/Misc/CoreGlobals.cpp
@@ -244,7 +244,8 @@ FString				GInstallBundleIni;											/* Install Bundle ini filename*/
 FString				GDeviceProfilesIni;											/* Runtime DeviceProfiles ini filename - use LoadLocalIni for other platforms' DPs */
 FString				GGameplayTagsIni;											/* Gameplay tags for the GameplayTagManager */
 
-float					GNearClippingPlane				= 10.0f;				/* Near clipping plane */
+float				GNearClippingPlane					= 10.0f;				/* Near clipping plane */
+float				GNearClippingPlane_RenderThread		= 10.0f;				/* Near clipping plane (Render Thread accessible) */
 
 bool					GExitPurge						= false;
 
diff --git a/Engine/Source/Runtime/Core/Public/CoreGlobals.h b/Engine/Source/Runtime/Core/Public/CoreGlobals.h
--- a/Engine/Source/Runtime/Core/Public/CoreGlobals.h
+++ b/Engine/Source/Runtime/Core/Public/CoreGlobals.h
@@ -419,6 +419,7 @@ extern CORE_API FString GDeviceProfilesIni;
 extern CORE_API FString GGameplayTagsIni;
 
 extern CORE_API float GNearClippingPlane;
+extern CORE_API float GNearClippingPlane_RenderThread;
 
 extern CORE_API bool GExitPurge;
 extern CORE_API TCHAR GInternalProjectName[64];
diff --git a/Engine/Source/Runtime/Engine/Private/UnrealEngine.cpp b/Engine/Source/Runtime/Engine/Private/UnrealEngine.cpp
--- a/Engine/Source/Runtime/Engine/Private/UnrealEngine.cpp
+++ b/Engine/Source/Runtime/Engine/Private/UnrealEngine.cpp
@@ -120,6 +120,7 @@ UnrealEngine.cpp: Implements the UEngine class and helpers.
 #include "GPUSkinCache.h"
 #include "Chaos/TriangleMeshImplicitObject.h"
 #include "HAL/PlatformMemoryHelpers.h"
+#include "RenderCore.h"
 #include <exception>
 
 #if UE_WITH_IRIS
@@ -2044,7 +2045,7 @@ void UEngine::Init(IEngineLoop* InEngineLoop)
 	// Update Script Maximum loop iteration count
 	FBlueprintCoreDelegates::SetScriptMaximumLoopIterations( GEngine->MaximumLoopIterationCount );
 
-	GNearClippingPlane = NearClipPlane;
+	SetNearClipPlaneGlobals(NearClipPlane);
 
 	UTextRenderComponent::InitializeMIDCache();
 
@@ -17262,8 +17263,9 @@ static void SetNearClipPlane(const TArray<FString>& Args)
 	{
 		NewClipPlane = FCString::Atof(*Args[0]);
 	}
-	FlushRenderingCommands();
-	GNearClippingPlane = FMath::Max(NewClipPlane,MinClipPlane);
+
+	SetNearClipPlaneGlobals(FMath::Max(NewClipPlane, MinClipPlane));
+
 }
 FAutoConsoleCommand GSetNearClipPlaneCmd(TEXT("r.SetNearClipPlane"),TEXT("Set the near clipping plane (in cm)"),FConsoleCommandWithArgsDelegate::CreateStatic(SetNearClipPlane));
 
diff --git a/Engine/Source/Runtime/RenderCore/Private/RenderCore.cpp b/Engine/Source/Runtime/RenderCore/Private/RenderCore.cpp
--- a/Engine/Source/Runtime/RenderCore/Private/RenderCore.cpp
+++ b/Engine/Source/Runtime/RenderCore/Private/RenderCore.cpp
@@ -13,6 +13,7 @@
 #include "RHI.h"
 #include "RenderTimer.h"
 #include "RenderCounters.h"
+#include "RenderingThread.h"
 
 void UpdateShaderDevelopmentMode();
 
@@ -324,6 +325,19 @@ RENDERCORE_API int32 GetCVarForceLODShadow_AnyThread()
 	return Ret;
 }
 
+//Setter function to keep GNearClippingPlane and GNearClippingPlane_RenderThread in sync
+RENDERCORE_API void SetNearClipPlaneGlobals(float NewNearClipPlane)
+{
+	GNearClippingPlane = NewNearClipPlane;
+
+	//Set GNearClippingPlane_RenderThread in a render command to be RT safe
+	ENQUEUE_RENDER_COMMAND(SetNearClipPlane_RenderThread)(
+		[NewNearClipPlane](FRHICommandListImmediate& RHICmdList)
+		{
+			GNearClippingPlane_RenderThread = NewNearClipPlane;
+		});
+}
+
 // Note: Enables or disables HDR support for a project. Typically this would be set on a per-project/per-platform basis in defaultengine.ini
 TAutoConsoleVariable<int32> CVarAllowHDR(
 	TEXT("r.AllowHDR"),
diff --git a/Engine/Source/Runtime/RenderCore/Public/RenderCore.h b/Engine/Source/Runtime/RenderCore/Public/RenderCore.h
--- a/Engine/Source/Runtime/RenderCore/Public/RenderCore.h
+++ b/Engine/Source/Runtime/RenderCore/Public/RenderCore.h
@@ -172,3 +172,5 @@ RENDERCORE_API int32 GetCVarForceLOD_AnyThread();
 
 RENDERCORE_API int32 GetCVarForceLODShadow();
 RENDERCORE_API int32 GetCVarForceLODShadow_AnyThread();
+
+RENDERCORE_API void SetNearClipPlaneGlobals(float NewNearClipPlane);
