addition of cvars to disable various cascade particle types

diff --git a/Engine/Source/Runtime/Engine/Private/Particles/ParticleGpuSimulation.cpp b/Engine/Source/Runtime/Engine/Private/Particles/ParticleGpuSimulation.cpp
--- a/Engine/Source/Runtime/Engine/Private/Particles/ParticleGpuSimulation.cpp
+++ b/Engine/Source/Runtime/Engine/Private/Particles/ParticleGpuSimulation.cpp
@@ -69,6 +69,14 @@ FAutoConsoleVariableRef CVarParticleSimulationSizeY(
 	ECVF_ReadOnly
 );
 
+static bool GFXCascadeGpuSpriteRenderingEnabled = true;
+static FAutoConsoleVariableRef CVarFXCascadeGpuSpriteRenderingEnabled(
+	TEXT("fx.Cascade.GpuSpriteRenderingEnabled"),
+	GFXCascadeGpuSpriteRenderingEnabled,
+	TEXT("Controls if gpu sprite rendering is enabled for Cascade"),
+	ECVF_Default
+);
+
 /** The tile size. Texture space is allocated in TileSize x TileSize units. */
 const int32 GParticleSimulationTileSize = 4;
 const int32 GParticlesPerTile = GParticleSimulationTileSize * GParticleSimulationTileSize;
@@ -2656,7 +2664,7 @@ public:
 
 		FRHICommandListBase& RHICmdList = FRHICommandListImmediate::Get();
 
-		if (RHISupportsGPUParticles())
+		if (RHISupportsGPUParticles() && GFXCascadeGpuSpriteRenderingEnabled)
 		{
 			SCOPE_CYCLE_COUNTER(STAT_GPUSpritePreRenderTime);
 
diff --git a/Engine/Source/Runtime/Engine/Private/Particles/ParticleSystemRender.cpp b/Engine/Source/Runtime/Engine/Private/Particles/ParticleSystemRender.cpp
--- a/Engine/Source/Runtime/Engine/Private/Particles/ParticleSystemRender.cpp
+++ b/Engine/Source/Runtime/Engine/Private/Particles/ParticleSystemRender.cpp
@@ -45,6 +45,38 @@ DECLARE_CYCLE_STAT(TEXT("DynamicSpriteEmitterData GetDynamicMeshElementsEmitter
 
 #include "InGamePerformanceTracker.h"
 
+static bool GFXCascadeSpriteRenderingEnabled = true;
+static FAutoConsoleVariableRef CVarFXCascadeSpriteRenderingEnabled(
+	TEXT("fx.Cascade.SpriteRenderingEnabled"),
+	GFXCascadeSpriteRenderingEnabled,
+	TEXT("Controls if sprite rendering is enabled for Cascade"),
+	ECVF_Default
+);
+
+static bool GFXCascadeMeshRenderingEnabled = true;
+static FAutoConsoleVariableRef CVarFXCascadeMeshRenderingEnabled(
+	TEXT("fx.Cascade.MeshRenderingEnabled"),
+	GFXCascadeMeshRenderingEnabled,
+	TEXT("Controls if mesh rendering is enabled for Cascade"),
+	ECVF_Default
+);
+
+static bool GFXCascadeBeamRenderingEnabled = true;
+static FAutoConsoleVariableRef CVarFXCascadeBeamRenderingEnabled(
+	TEXT("fx.Cascade.BeamRenderingEnabled"),
+	GFXCascadeBeamRenderingEnabled,
+	TEXT("Controls if beam rendering is enabled for Cascade"),
+	ECVF_Default
+);
+
+static bool GFXCascadeTrailRenderingEnabled = true;
+static FAutoConsoleVariableRef CVarFXCascadeTrailRenderingEnabled(
+	TEXT("fx.Cascade.TrailRenderingEnabled"),
+	GFXCascadeTrailRenderingEnabled,
+	TEXT("Controls if trail rendering is enabled for Cascade"),
+	ECVF_Default
+);
+
 static int32 GFXAllowParticleMeshLODs = 0;
 static FAutoConsoleVariableRef CVarFXAllowParticleMeshLODs(
 	TEXT("fx.FXAllowParticleMeshLODs"),
@@ -818,7 +850,7 @@ void FDynamicSpriteEmitterData::GetDynamicMeshElementsEmitter(const FParticleSys
 	// Sort and generate particles for this view.
 	const FDynamicSpriteEmitterReplayDataBase* SourceData = GetSourceData();
 
-	if (bValid && SourceData)
+	if (bValid && SourceData && GFXCascadeSpriteRenderingEnabled)
 	{
 		if (SourceData->EmitterRenderMode == ERM_Normal)
 		{
@@ -1311,7 +1343,7 @@ void FDynamicMeshEmitterData::GetDynamicMeshElementsEmitter(const FParticleSyste
 {
 	SCOPE_CYCLE_COUNTER(STAT_MeshRenderingTime);
 
-	if (bValid)
+	if (bValid && GFXCascadeMeshRenderingEnabled)
 	{
 		if (Source.EmitterRenderMode == ERM_Normal)
 		{
@@ -2372,6 +2404,11 @@ void FDynamicBeam2EmitterData::GetDynamicMeshElementsEmitter(const FParticleSyst
 		return;
 	}
 
+	if (!GFXCascadeBeamRenderingEnabled)
+	{
+		return;
+	}
+
 	FRHICommandListBase& RHICmdList = FRHICommandListImmediate::Get();
 
 	FIndexBuffer* IndexBuffer = nullptr;
@@ -5129,6 +5166,11 @@ void FDynamicTrailsEmitterData::GetDynamicMeshElementsEmitter(const FParticleSys
 		return;
 	}
 
+	if (!GFXCascadeTrailRenderingEnabled)
+	{
+		return;
+	}
+
 	FRHICommandListBase& RHICmdList = FRHICommandListImmediate::Get();
 
 	const bool bIsWireframe = ViewFamily.EngineShowFlags.Wireframe;
