CPU accessible textures

diff --git a/Engine/Source/Developer/TextureBuild/Private/TextureBuildFunction.cpp b/Engine/Source/Developer/TextureBuild/Private/TextureBuildFunction.cpp
--- a/Engine/Source/Developer/TextureBuild/Private/TextureBuildFunction.cpp
+++ b/Engine/Source/Developer/TextureBuild/Private/TextureBuildFunction.cpp
@@ -138,6 +138,7 @@ static FTextureBuildSettings ReadBuildSettingsFromCompactBinary(const FCbObjectV
 	BuildSettings.bApplyYCoCgBlockScale = Object["bApplyYCoCgBlockScale"].AsBool(BuildSettings.bApplyYCoCgBlockScale);
 	BuildSettings.bApplyKernelToTopMip = Object["bApplyKernelToTopMip"].AsBool(BuildSettings.bApplyKernelToTopMip);
 	BuildSettings.bRenormalizeTopMip = Object["bRenormalizeTopMip"].AsBool(BuildSettings.bRenormalizeTopMip);
+	BuildSettings.bCPUAccessible = Object["bCPUAccessible"].AsBool(BuildSettings.bCPUAccessible);
 	ReadCbField(Object["CompositeTextureMode"], BuildSettings.CompositeTextureMode);
 	ReadCbField(Object["CompositePower"], BuildSettings.CompositePower);
 	ReadCbField(Object["LODBias"], BuildSettings.LODBias);
@@ -356,6 +357,18 @@ void FTextureBuildFunction::Build(UE::DerivedData::FBuildContext& Context) const
 		return;
 	}
 
+	FSharedImageRef CPUCopy;
+	if (BuildSettings.bCPUAccessible)
+	{
+		CPUCopy = new FSharedImage();
+		SourceMips[0].CopyTo(*CPUCopy);
+	
+		// We just use a placeholder texture rather than the source.
+		SourceMips.Empty();
+		FImage& Placeholder = SourceMips.AddDefaulted_GetRef();
+		UE::TextureBuildUtilities::GetPlaceholderTextureImage(&Placeholder);
+	}
+
 	TArray<FImage> AssociatedNormalSourceMips;
 	if (FCbFieldView CompositeSource = Settings["CompositeSource"];
 		CompositeSource && !TryReadTextureSourceFromCompactBinary(CompositeSource, Context,BuildSettings, AssociatedNormalSourceMips))
@@ -419,6 +432,16 @@ void FTextureBuildFunction::Build(UE::DerivedData::FBuildContext& Context) const
 	int32 NumStreamingMips = TextureDescription.GetNumStreamingMips(&ExtendedData, EngineParameters);
 	
 	{
+		if (CPUCopy.IsValid())
+		{
+			FCbObject ImageInfoMetadata;
+			CPUCopy->ImageInfoToCompactBinary(ImageInfoMetadata);
+			Context.AddValue(UE::DerivedData::FValueId::FromName(ANSITEXTVIEW("CPUCopyImageInfo")), ImageInfoMetadata);
+
+			FSharedBuffer CPUCopyData = MakeSharedBufferFromArray(MoveTemp(CPUCopy->RawData));
+			Context.AddValue(UE::DerivedData::FValueId::FromName(ANSITEXTVIEW("CPUCopyRawData")), CPUCopyData);
+		}
+
 		Context.AddValue(UE::DerivedData::FValueId::FromName(ANSITEXTVIEW("TextureBuildMetadata")), BuildMetadata.ToCompactBinaryWithDefaults());
 		Context.AddValue(UE::DerivedData::FValueId::FromName(ANSITEXTVIEW("EncodedTextureDescription")), UE::TextureBuildUtilities::EncodedTextureDescription::ToCompactBinary(TextureDescription));
 		Context.AddValue(UE::DerivedData::FValueId::FromName(ANSITEXTVIEW("EncodedTextureExtendedData")), UE::TextureBuildUtilities::EncodedTextureExtendedData::ToCompactBinary(ExtendedData));
diff --git a/Engine/Source/Developer/TextureBuildUtilities/Private/TextureBuildUtilities.cpp b/Engine/Source/Developer/TextureBuildUtilities/Private/TextureBuildUtilities.cpp
--- a/Engine/Source/Developer/TextureBuildUtilities/Private/TextureBuildUtilities.cpp
+++ b/Engine/Source/Developer/TextureBuildUtilities/Private/TextureBuildUtilities.cpp
@@ -273,6 +273,26 @@ FTextureBuildMetadata::FTextureBuildMetadata(FCbObject InCbObject)
 	PreEncodeMipsHash = InCbObject["PreEncodeMipsHash"].AsUInt64(PreEncodeMipsHash);
 }
 
+void GetPlaceholderTextureImageInfo(FImageInfo* OutImageInfo)
+{
+	OutImageInfo->SizeX = 4;
+	OutImageInfo->SizeY = 4;
+	OutImageInfo->GammaSpace = EGammaSpace::sRGB;
+	OutImageInfo->Format = ERawImageFormat::BGRA8;
+	OutImageInfo->NumSlices = 1;
+}
+void GetPlaceholderTextureImage(FImage* OutImage)
+{
+	*OutImage = FImage();
+
+	GetPlaceholderTextureImageInfo(OutImage);
+	OutImage->RawData.AddUninitialized(sizeof(FColor) * OutImage->SizeX * OutImage->SizeY);
+	for (FColor& Color : OutImage->AsBGRA8())
+	{
+		Color = FColor::Black;
+	}
+
+}
 
 } // namespace
 }
\ No newline at end of file
diff --git a/Engine/Source/Developer/TextureBuildUtilities/Public/TextureBuildUtilities.h b/Engine/Source/Developer/TextureBuildUtilities/Public/TextureBuildUtilities.h
--- a/Engine/Source/Developer/TextureBuildUtilities/Public/TextureBuildUtilities.h
+++ b/Engine/Source/Developer/TextureBuildUtilities/Public/TextureBuildUtilities.h
@@ -22,6 +22,11 @@ namespace UE
 {
 namespace TextureBuildUtilities
 {
+	enum
+	{
+		// The width and height of the placeholder gpu texture we create when the texture is cpu accessible.
+		PLACEHOLDER_TEXTURE_SIZE = 4
+	};
 
 namespace EncodedTextureExtendedData
 {
@@ -77,5 +82,8 @@ FORCEINLINE const FName TextureFormatRemovePrefixFromName(FName const& InName)
 //	  the next step will then encode from this format to the desired output format
 TEXTUREBUILDUTILITIES_API ERawImageFormat::Type GetVirtualTextureBuildIntermediateFormat(const FTextureBuildSettings& BuildSettings);
 
+TEXTUREBUILDUTILITIES_API void GetPlaceholderTextureImageInfo(FImageInfo* OutImageInfo);
+TEXTUREBUILDUTILITIES_API void GetPlaceholderTextureImage(FImage* OutImage);
+
 } // namespace TextureBuildUtilities
 } // namespace UE
diff --git a/Engine/Source/Developer/TextureCompressor/Private/TextureCompressorModule.cpp b/Engine/Source/Developer/TextureCompressor/Private/TextureCompressorModule.cpp
--- a/Engine/Source/Developer/TextureCompressor/Private/TextureCompressorModule.cpp
+++ b/Engine/Source/Developer/TextureCompressor/Private/TextureCompressorModule.cpp
@@ -3464,6 +3464,17 @@ int32 ITextureCompressorModule::GetMipCountForBuildSettings(
 	const FTextureBuildSettings& BuildSettings,
 	int32& OutMip0SizeX, int32& OutMip0SizeY, int32& OutMip0NumSlices)
 {
+	if (BuildSettings.bCPUAccessible)
+	{
+		// CPU accessible texture generates a placeholder gpu texture with 1 mip in all cases.
+		FImageInfo PlaceholderInfo;
+		UE::TextureBuildUtilities::GetPlaceholderTextureImageInfo(&PlaceholderInfo);
+		OutMip0SizeX = PlaceholderInfo.SizeX;
+		OutMip0SizeY = PlaceholderInfo.SizeY;
+		OutMip0NumSlices = PlaceholderInfo.NumSlices;
+		return 1;
+	}
+
 	// AFAICT LatLongCubeMaps don't do any of this - pow2 is broken with them but it runs, and max texture stuff
 	// is handled internally in the extents function.
 	int32 BaseSizeX = InMip0SizeX;
diff --git a/Engine/Source/Developer/TextureCompressor/Public/TextureCompressorModule.h b/Engine/Source/Developer/TextureCompressor/Public/TextureCompressorModule.h
--- a/Engine/Source/Developer/TextureCompressor/Public/TextureCompressorModule.h
+++ b/Engine/Source/Developer/TextureCompressor/Public/TextureCompressorModule.h
@@ -159,6 +159,8 @@ struct FTextureBuildSettings
 	uint32 bApplyKernelToTopMip : 1;
 	/** 1: renormalizes the top mip (only useful for normal maps, prevents artists errors and adds quality) 0:don't */
 	uint32 bRenormalizeTopMip : 1;
+	/** If set, the texture will generate a tiny placeholder hw texture and save off a copy of the top mip for CPU access. */
+	uint32 bCPUAccessible : 1;
 	/** e.g. CTM_RoughnessFromNormalAlpha */
 	uint8 CompositeTextureMode;	// ECompositeTextureMode, opaque to avoid dependencies on engine headers.
 	/* default 1, high values result in a stronger effect */
@@ -280,6 +282,7 @@ struct FTextureBuildSettings
 		, bApplyYCoCgBlockScale(false)
 		, bApplyKernelToTopMip(false)
 		, bRenormalizeTopMip(false)
+		, bCPUAccessible(false)
 		, CompositeTextureMode(0 /*CTM_Disabled*/)
 		, CompositePower(1.0f)
 		, LODBias(0)
diff --git a/Engine/Source/Editor/TextureEditor/Private/Models/TextureEditorViewportClient.cpp b/Engine/Source/Editor/TextureEditor/Private/Models/TextureEditorViewportClient.cpp
--- a/Engine/Source/Editor/TextureEditor/Private/Models/TextureEditorViewportClient.cpp
+++ b/Engine/Source/Editor/TextureEditor/Private/Models/TextureEditorViewportClient.cpp
@@ -109,6 +109,21 @@ void FTextureEditorViewportClient::Draw(FViewport* Viewport, FCanvas* Canvas)
 
 	bool bIsVirtualTexture = false;
 
+	UTexture2D* CPUCopyTexture = nullptr;
+	if (Texture2D)
+	{
+		CPUCopyTexture = Texture2D->GetCPUCopyTexture();
+		if (CPUCopyTexture)
+		{
+			FIntPoint CenteringOffset(0, 0);
+			CenteringOffset += Viewport->GetSizeXY() / 2;
+			CenteringOffset -= FIntPoint(Width, Height) / 2;
+
+			Canvas->DrawTile(CenteringOffset.X, CenteringOffset.Y, Width, Height, 0.0f, 0.0f, 1.0f, 1.0f, FLinearColor::White, CPUCopyTexture->GetResource());
+			return;
+		}
+	}
+	
 	TRefCountPtr<FBatchedElementParameters> BatchedElementParameters;
 
 	if (GMaxRHIFeatureLevel >= ERHIFeatureLevel::SM5)
@@ -202,7 +217,8 @@ void FTextureEditorViewportClient::Draw(FViewport* Viewport, FCanvas* Canvas)
 	FTexturePlatformData** RunningPlatformDataPtr = Texture->GetRunningPlatformData();
 	float Exposure = RunningPlatformDataPtr && *RunningPlatformDataPtr && IsHDR((*RunningPlatformDataPtr)->PixelFormat) ? FMath::Pow(2.0f, (float)TextureEditorPinned->GetExposureBias()) : 1.0f;
 
-	if ( Texture->GetResource() != nullptr )
+
+	if ( Texture->GetResource() != nullptr && !CPUCopyTexture )
 	{
 		FCanvasTileItem TileItem( FVector2D( XPos, YPos ), Texture->GetResource(), FVector2D( Width, Height ), FLinearColor(Exposure, Exposure, Exposure) );
 		TileItem.BlendMode = TextureEditorPinned->GetColourChannelBlendMode();
diff --git a/Engine/Source/Editor/TextureEditor/Private/TextureEditorToolkit.cpp b/Engine/Source/Editor/TextureEditor/Private/TextureEditorToolkit.cpp
--- a/Engine/Source/Editor/TextureEditor/Private/TextureEditorToolkit.cpp
+++ b/Engine/Source/Editor/TextureEditor/Private/TextureEditorToolkit.cpp
@@ -134,7 +134,10 @@ FTextureEditorToolkit::~FTextureEditorToolkit( )
 	if (Texture2D && Texture2D->IsCurrentlyVirtualTextured())
 	{
 		FVirtualTexture2DResource* Resource = (FVirtualTexture2DResource*)Texture2D->GetResource();
-		Resource->ReleaseAllocatedVT();
+		if (Resource)
+		{
+			Resource->ReleaseAllocatedVT();
+		}
 	}
 
 	FReimportManager::Instance()->OnPreReimport().RemoveAll(this);
@@ -349,6 +352,18 @@ void FTextureEditorToolkit::CalculateTextureDimensions(int32& OutWidth, int32& O
 	OutArraySize = IsArrayTexture() ? (IsCubeTexture() ? Texture->GetSurfaceArraySize() / 6 : Texture->GetSurfaceArraySize()) : 0;
 	const int32 BorderSize = GetDefault<UTextureEditorSettings>()->TextureBorderEnabled ? 1 : 0;
 
+	
+	if (UTexture2D* Texture2D = Cast<UTexture2D>(Texture))
+	{
+		if (UTexture2D* CpuTexture = Texture2D->GetCPUCopyTexture())
+		{
+			OutWidth = CpuTexture->GetSurfaceWidth();
+			OutHeight = CpuTexture->GetSurfaceHeight();
+			OutDepth = 1;
+			OutArraySize = 0;
+		}
+	}
+
 	if (!PreviewEffectiveTextureWidth || !PreviewEffectiveTextureHeight)
 	{
 		OutWidth = 0;
@@ -809,6 +824,15 @@ void FTextureEditorToolkit::PopulateQuickInfo( )
 		MaxInGameText->SetText(FText::Format( NSLOCTEXT("TextureEditor", "QuickInfo_MaxInGame_2x", "Max In-Game: {0}x{1}{2}"), FText::AsNumber(MaxInGameWidth, &FormatOptions), FText::AsNumber(MaxInGameHeight, &FormatOptions), InGameCubemapInfo));
 	}
 
+	if (Texture2D)
+	{
+		if (UTexture2D* CpuTexture = Texture2D->GetCPUCopyTexture(); CpuTexture)
+		{
+			PreviewEffectiveTextureWidth = CpuTexture->GetSurfaceWidth();
+			PreviewEffectiveTextureHeight = CpuTexture->GetSurfaceHeight();
+		}
+	}
+
 	SizeText->SetText(FText::Format(NSLOCTEXT("TextureEditor", "QuickInfo_ResourceSize", "Resource Size: {0} KB"), FText::AsNumber(FMath::DivideAndRoundNearest(ResourceSize, (int64)1024), &FormatOptions)));
 
 	FText Method = Texture->IsCurrentlyVirtualTextured() ? NSLOCTEXT("TextureEditor", "QuickInfo_MethodVirtualStreamed", "Virtual Streamed")
diff --git a/Engine/Source/Editor/UnrealEd/Private/ThumbnailRendering/TextureThumbnailRenderer.cpp b/Engine/Source/Editor/UnrealEd/Private/ThumbnailRendering/TextureThumbnailRenderer.cpp
--- a/Engine/Source/Editor/UnrealEd/Private/ThumbnailRendering/TextureThumbnailRenderer.cpp
+++ b/Engine/Source/Editor/UnrealEd/Private/ThumbnailRendering/TextureThumbnailRenderer.cpp
@@ -41,8 +41,23 @@ void UTextureThumbnailRenderer::GetThumbnailSize(UObject* Object, float Zoom, ui
 
 	if (Texture != nullptr)
 	{
-		OutWidth = FMath::TruncToInt(Zoom * (float)Texture->GetSurfaceWidth());
-		OutHeight = FMath::TruncToInt(Zoom * (float)Texture->GetSurfaceHeight());
+		UTexture2D* Texture2D = Cast<UTexture2D>(Texture);
+		FSharedImageConstRef CpuCopy;
+		if (Texture2D)
+		{
+			CpuCopy = Texture2D->GetCPUCopy();
+		}
+
+		if (CpuCopy.IsValid())
+		{
+			OutWidth = FMath::TruncToInt(Zoom * (float)CpuCopy->GetWidth());
+			OutHeight = FMath::TruncToInt(Zoom * (float)CpuCopy->GetHeight());
+		}
+		else
+		{
+			OutWidth = FMath::TruncToInt(Zoom * (float)Texture->GetSurfaceWidth());
+			OutHeight = FMath::TruncToInt(Zoom * (float)Texture->GetSurfaceHeight());
+		}
 	}
 	else
 	{
@@ -56,100 +71,129 @@ void UTextureThumbnailRenderer::Draw(UObject* Object, int32 X, int32 Y, uint32 W
 	if (Texture != nullptr && Texture->GetResource() != nullptr)
 	{
 		UTexture2D* Texture2D = Cast<UTexture2D>(Texture);
+		UTexture* CpuCopyTexture = nullptr;
+		if (Texture2D)
+		{
+			CpuCopyTexture = Texture2D->GetCPUCopyTexture();
+		}
+		if (CpuCopyTexture == nullptr)
+		{
+			// Take the alpha channel into account for textures that have one.
+			// This provides a much better preview than just showing RGB,
+			// Because the RGB content in areas with an alpha of 0 is often garbage that will not be seen in normal conditions.
+			// Non-UI textures often have uncorrelated data in the alpha channel (like a skin mask, specular power, etc) so we only preview UI textures this way.
+			const bool bUseTranslucentBlend = Texture2D && Texture2D->HasAlphaChannel() && ((Texture2D->LODGroup == TEXTUREGROUP_UI) || (Texture2D->LODGroup == TEXTUREGROUP_Pixels2D));
+
+			UTextureCube* TextureCube = Cast<UTextureCube>(Texture);
+			UTexture2DArray* Texture2DArray = Cast<UTexture2DArray>(Texture);
+			UTextureCubeArray* TextureCubeArray = Cast<UTextureCubeArray>(Texture);
+			UTextureRenderTargetCube* RTTextureCube = Cast<UTextureRenderTargetCube>(Texture);
+			UTextureLightProfile* TextureLightProfile = Cast<UTextureLightProfile>(Texture);
+			const bool bIsVirtualTexture = Texture->IsCurrentlyVirtualTextured();
+
+			TRefCountPtr<FBatchedElementParameters> BatchedElementParameters;
+
+			if(TextureCube || TextureCubeArray || RTTextureCube)
+			{
+				// is released by the render thread when it was rendered
+				BatchedElementParameters = new FMipLevelBatchedElementParameters((float)0, (float)-1, TextureCubeArray != nullptr, FMatrix44f::Identity, true, false, false);
+			
+				// If the thumbnail is square then make it 2:1 for cubes.
+				if(Width == Height)
+				{
+					Height = Width / 2;
+					Y += Height / 2;
+				}
+			}
+			else if (Texture2DArray) 
+			{
+				bool bIsNormalMap = Texture2DArray->IsNormalMap();
+				bool bIsSingleChannel = true;
+				BatchedElementParameters = new FBatchedElementTexture2DPreviewParameters((float)0, (float)0, (float)-1, bIsNormalMap, bIsSingleChannel, false, false, true, false);
+			}
+			else if (TextureLightProfile)
+			{
+				BatchedElementParameters = new FIESLightProfileBatchedElementParameters(TextureLightProfile->Brightness);
+			}
+			else if (Texture2D && Texture2D->IsNormalMap())
+			{
+				BatchedElementParameters = new FNormalMapBatchedElementParameters();
+			}
 
-		// Take the alpha channel into account for textures that have one.
-		// This provides a much better preview than just showing RGB,
-		// Because the RGB content in areas with an alpha of 0 is often garbage that will not be seen in normal conditions.
-		// Non-UI textures often have uncorrelated data in the alpha channel (like a skin mask, specular power, etc) so we only preview UI textures this way.
-		const bool bUseTranslucentBlend = Texture2D && Texture2D->HasAlphaChannel() && ((Texture2D->LODGroup == TEXTUREGROUP_UI) || (Texture2D->LODGroup == TEXTUREGROUP_Pixels2D));
+			if (bUseTranslucentBlend)
+			{
+				// If using alpha, draw a checkerboard underneath first.
+				const int32 CheckerDensity = 8;
+				auto Checker = UThumbnailManager::Get().CheckerboardTexture;
+				Canvas->DrawTile(
+					0.0f, 0.0f, Width, Height,							// Dimensions
+					0.0f, 0.0f, CheckerDensity, CheckerDensity,			// UVs
+					FLinearColor::White, Checker->GetResource());			// Tint & Texture
+			}
 
-		UTextureCube* TextureCube = Cast<UTextureCube>(Texture);
-		UTexture2DArray* Texture2DArray = Cast<UTexture2DArray>(Texture);
-		UTextureCubeArray* TextureCubeArray = Cast<UTextureCubeArray>(Texture);
-		UTextureRenderTargetCube* RTTextureCube = Cast<UTextureRenderTargetCube>(Texture);
-		UTextureLightProfile* TextureLightProfile = Cast<UTextureLightProfile>(Texture);
-		const bool bIsVirtualTexture = Texture->IsCurrentlyVirtualTextured();
+			// Use A canvas tile item to draw
+			FCanvasTileItem CanvasTile( FVector2D( X, Y ), Texture->GetResource(), FVector2D( Width,Height ), FLinearColor::White );
+			CanvasTile.BlendMode = bUseTranslucentBlend ? SE_BLEND_Translucent : SE_BLEND_Opaque;
+			CanvasTile.BatchedElementParameters = BatchedElementParameters;
+			if (bIsVirtualTexture && Texture->Source.GetNumBlocks() > 1)
+			{
+				// Adjust UVs to display entire UDIM range, acounting for UE inverted V-axis
+				// We're not actually rendering a VT here, but the editor-only texture we're using is still using the UDIM tile layout
+				// So we use inverted Y-axis, but then normalize back to [0,1)
+				const FIntPoint BlockSize = Texture->Source.GetSizeInBlocks();
+				const float RcpBlockSizeY = 1.0f / (float)BlockSize.Y;
+				CanvasTile.UV0.Y = (1.0f - (float)BlockSize.Y) * RcpBlockSizeY;
+				CanvasTile.UV1.Y = RcpBlockSizeY;
+			}
+			CanvasTile.Draw( Canvas );
 
-		TRefCountPtr<FBatchedElementParameters> BatchedElementParameters;
+			if (TextureLightProfile)
+			{
+				float Brightness = TextureLightProfile->Brightness;
+
+				// Brightness in Lumens
+				FText BrightnessText = FText::AsNumber( Brightness );
+				FCanvasTextItem TextItem( FVector2D( 5.0f, 5.0f ), BrightnessText, GEngine->GetLargeFont(), FLinearColor::White );
+				TextItem.EnableShadow(FLinearColor::Black);
+				TextItem.Scale = FVector2D(Width / 128.0f, Height / 128.0f);
+				TextItem.Draw(Canvas);
+			}
 
-		if(TextureCube || TextureCubeArray || RTTextureCube)
-		{
-			// is released by the render thread when it was rendered
-			BatchedElementParameters = new FMipLevelBatchedElementParameters((float)0, (float)-1, TextureCubeArray != nullptr, FMatrix44f::Identity, true, false, false);
-			
-			// If the thumbnail is square then make it 2:1 for cubes.
-			if(Width == Height)
+			if (bIsVirtualTexture)
 			{
-				Height = Width / 2;
-				Y += Height / 2;
+				auto VTChars = TEXT("VT");
+				int32 VTWidth = 0;
+				int32 VTHeight = 0;
+				StringSize(GEngine->GetLargeFont(), VTWidth, VTHeight, VTChars);
+				float PaddingX = Width / 128.0f;
+				float PaddingY = Height / 128.0f;
+				float ScaleX = Width / 64.0f; //Text is 1/64'th of the size of the thumbnails
+				float ScaleY = Height / 64.0f;
+				// VT overlay
+				FCanvasTextItem TextItem(FVector2D(Width - PaddingX - VTWidth * ScaleX, Height - PaddingY - VTHeight * ScaleY), FText::FromString(VTChars), GEngine->GetLargeFont(), FLinearColor::White);
+				TextItem.EnableShadow(FLinearColor::Black);
+				TextItem.Scale = FVector2D(ScaleX, ScaleY);
+				TextItem.Draw(Canvas);
 			}
 		}
-		else if (Texture2DArray) 
-		{
-			bool bIsNormalMap = Texture2DArray->IsNormalMap();
-			bool bIsSingleChannel = true;
-			BatchedElementParameters = new FBatchedElementTexture2DPreviewParameters((float)0, (float)0, (float)-1, bIsNormalMap, bIsSingleChannel, false, false, true, false);
-		}
-		else if (TextureLightProfile)
-		{
-			BatchedElementParameters = new FIESLightProfileBatchedElementParameters(TextureLightProfile->Brightness);
-		}
-		else if (Texture2D && Texture2D->IsNormalMap())
-		{
-			BatchedElementParameters = new FNormalMapBatchedElementParameters();
-		}
-
-		if (bUseTranslucentBlend)
+		else
 		{
-			// If using alpha, draw a checkerboard underneath first.
-			const int32 CheckerDensity = 8;
-			auto Checker = UThumbnailManager::Get().CheckerboardTexture;
+			// Render a tile using a transient texture that holds the cpu copy.
 			Canvas->DrawTile(
 				0.0f, 0.0f, Width, Height,							// Dimensions
-				0.0f, 0.0f, CheckerDensity, CheckerDensity,			// UVs
-				FLinearColor::White, Checker->GetResource());			// Tint & Texture
-		}
-
-		// Use A canvas tile item to draw
-		FCanvasTileItem CanvasTile( FVector2D( X, Y ), Texture->GetResource(), FVector2D( Width,Height ), FLinearColor::White );
-		CanvasTile.BlendMode = bUseTranslucentBlend ? SE_BLEND_Translucent : SE_BLEND_Opaque;
-		CanvasTile.BatchedElementParameters = BatchedElementParameters;
-		if (bIsVirtualTexture && Texture->Source.GetNumBlocks() > 1)
-		{
-			// Adjust UVs to display entire UDIM range, acounting for UE inverted V-axis
-			// We're not actually rendering a VT here, but the editor-only texture we're using is still using the UDIM tile layout
-			// So we use inverted Y-axis, but then normalize back to [0,1)
-			const FIntPoint BlockSize = Texture->Source.GetSizeInBlocks();
-			const float RcpBlockSizeY = 1.0f / (float)BlockSize.Y;
-			CanvasTile.UV0.Y = (1.0f - (float)BlockSize.Y) * RcpBlockSizeY;
-			CanvasTile.UV1.Y = RcpBlockSizeY;
-		}
-		CanvasTile.Draw( Canvas );
-
-		if (TextureLightProfile)
-		{
-			float Brightness = TextureLightProfile->Brightness;
+				0.0f, 0.0f, 1.0f, 1.0f,								// UVs
+				FLinearColor::White, CpuCopyTexture->GetResource());			// Tint & Texture
 
-			// Brightness in Lumens
-			FText BrightnessText = FText::AsNumber( Brightness );
-			FCanvasTextItem TextItem( FVector2D( 5.0f, 5.0f ), BrightnessText, GEngine->GetLargeFont(), FLinearColor::White );
-			TextItem.EnableShadow(FLinearColor::Black);
-			TextItem.Scale = FVector2D(Width / 128.0f, Height / 128.0f);
-			TextItem.Draw(Canvas);
-		}
-
-		if (bIsVirtualTexture)
-		{
-			auto VTChars = TEXT("VT");
-			int32 VTWidth = 0;
-			int32 VTHeight = 0;
-			StringSize(GEngine->GetLargeFont(), VTWidth, VTHeight, VTChars);
+			const TCHAR* Chars = TEXT("CPU");
+			int32 TextWidth = 0;
+			int32 TextHeight = 0;
+			StringSize(GEngine->GetLargeFont(), TextWidth, TextHeight, Chars);
 			float PaddingX = Width / 128.0f;
 			float PaddingY = Height / 128.0f;
 			float ScaleX = Width / 64.0f; //Text is 1/64'th of the size of the thumbnails
 			float ScaleY = Height / 64.0f;
-			// VT overlay
-			FCanvasTextItem TextItem(FVector2D(Width - PaddingX - VTWidth * ScaleX, Height - PaddingY - VTHeight * ScaleY), FText::FromString(VTChars), GEngine->GetLargeFont(), FLinearColor::White);
+			// overlay
+			FCanvasTextItem TextItem(FVector2D(Width - PaddingX - TextWidth * ScaleX, Height - PaddingY - TextHeight * ScaleY), FText::FromString(Chars), GEngine->GetLargeFont(), FLinearColor::White);
 			TextItem.EnableShadow(FLinearColor::Black);
 			TextItem.Scale = FVector2D(ScaleX, ScaleY);
 			TextItem.Draw(Canvas);
diff --git a/Engine/Source/Runtime/Engine/Classes/Engine/Texture.h b/Engine/Source/Runtime/Engine/Classes/Engine/Texture.h
--- a/Engine/Source/Runtime/Engine/Classes/Engine/Texture.h
+++ b/Engine/Source/Runtime/Engine/Classes/Engine/Texture.h
@@ -750,7 +750,10 @@ struct FTexturePlatformData
 	FOptTexturePlatformData OptData;
 	/** Mip data or VT data. one or the other. */
 	TIndirectArray<struct FTexture2DMipMap> Mips;
-	struct FVirtualTextureBuiltData* VTData;
+	struct FVirtualTextureBuiltData* VTData=nullptr;
+
+	/** This is only valid if the texture availability is CPU only, see GetHasCpuCopy() */
+	TRefCountPtr<const struct FSharedImage> CPUCopy;
 
 #if WITH_EDITORONLY_DATA
 
@@ -855,7 +858,8 @@ struct FTexturePlatformData
 private:
 	static constexpr uint32 BitMask_CubeMap    = 1u << 31u;
 	static constexpr uint32 BitMask_HasOptData = 1u << 30u;
-	static constexpr uint32 BitMask_NumSlices  = BitMask_HasOptData - 1u;
+	static constexpr uint32 BitMask_HasCpuCopy = 1u << 29u;
+	static constexpr uint32 BitMask_NumSlices  = BitMask_HasCpuCopy - 1u;
 
 public:
 	/** Return whether TryLoadMips() would stall because async loaded mips are not yet available. */
@@ -891,9 +895,9 @@ public:
 	 */
 	void SerializeCooked(FArchive& Ar, class UTexture* Owner, bool bStreamable, const bool bSerializeMipData);
 	
-	inline void SetPackedData(int32 InNumSlices, bool bInHasOptData, bool bInCubeMap)
+	inline void SetPackedData(int32 InNumSlices, bool bInHasOptData, bool bInCubeMap, bool bInHasCpuCopy)
 	{
-		PackedData = (InNumSlices & BitMask_NumSlices) | (bInCubeMap ? BitMask_CubeMap : 0) | (bInHasOptData ? BitMask_HasOptData : 0);
+		PackedData = (InNumSlices & BitMask_NumSlices) | (bInCubeMap ? BitMask_CubeMap : 0) | (bInHasOptData ? BitMask_HasOptData : 0) | (bInHasCpuCopy ? BitMask_HasCpuCopy : 0);
 	}
 
 	inline bool GetHasOptData() const
@@ -910,6 +914,16 @@ public:
 		OptData = Data;
 	}
 
+	inline bool GetHasCpuCopy() const
+	{
+		return (PackedData & BitMask_HasCpuCopy) == BitMask_HasCpuCopy;
+	}
+
+	inline void SetHasCpuCopy(bool bInHasCpuCopy)
+	{
+		PackedData = (bInHasCpuCopy ? BitMask_HasCpuCopy : 0) | (PackedData & (~BitMask_HasCpuCopy));
+	}
+
 	inline bool IsCubemap() const
 	{
 		return (PackedData & BitMask_CubeMap) == BitMask_CubeMap; 
@@ -941,6 +955,9 @@ public:
 	}
 
 #if WITH_EDITOR
+	// Clears the data such that a new Cache() call can load new data in to the structure.
+	void Reset();
+
 	bool IsAsyncWorkComplete() const;
 
 	// Compresses the texture using the given compressor and adds the result to the DDC.
@@ -1347,6 +1364,12 @@ PRAGMA_ENABLE_DEPRECATION_WARNINGS
 	UPROPERTY(EditAnywhere, Category=LevelOfDetail, AdvancedDisplay)
 	ETextureDownscaleOptions DownscaleOptions;
 
+	/** 
+	* Whether the texture will be encoded to a gpu format and uploaded to the graphics card, or kept on the CPU for access by gamecode / blueprint. 
+	* For CPU availability, the texture will still upload a tiny black texture as a placeholder. Only applies to 2d textures.
+	*/
+	UPROPERTY(EditAnywhere, BlueprintReadWrite, Category=Texture, meta=(DisplayName="Availability"), AssetRegistrySearchable)
+	ETextureAvailability Availability;
 	
 	/** Whether Texture and its source are in SRGB Gamma color space.  Can only be used with 8-bit and compressed formats.  This should be unchecked if using alpha channels individually as masks. */
 	UPROPERTY(EditAnywhere, BlueprintReadWrite, Category=Texture, meta=(DisplayName="sRGB"), AssetRegistrySearchable)
diff --git a/Engine/Source/Runtime/Engine/Classes/Engine/Texture2D.h b/Engine/Source/Runtime/Engine/Classes/Engine/Texture2D.h
--- a/Engine/Source/Runtime/Engine/Classes/Engine/Texture2D.h
+++ b/Engine/Source/Runtime/Engine/Classes/Engine/Texture2D.h
@@ -4,6 +4,7 @@
 
 #include "CoreMinimal.h"
 #include "HAL/ThreadSafeCounter.h"
+#include "ImageCoreBP.h"
 #include "UObject/ObjectMacros.h"
 #include "UObject/ScriptMacros.h"
 #include "Engine/Texture.h"
@@ -309,11 +310,43 @@ public:
 
 #endif // WITH_EDITOR
 
+#if WITH_EDITORONLY_DATA
+	/** If we ever show the CPU accessible image in the editor we'll need a transient texture with
+	* those bits in it. If we don't have a CPU copy, then this is null. Use the GetCPUCopyTexture to
+	* access as it's created on demand.
+	*/
+	UPROPERTY();
+	TObjectPtr<UTexture2D> CPUCopyTexture;
+
+	/**
+	* Creates and returns a 2d texture that holds the CPU copy. This is just so that we have something to show
+	* in the editor for the cpu texture and should never be used at runtime or for game data. This function
+	* stalls while the texture is encoding due to GetPlatformData, and will return nullptr if the texture
+	* is not cpu accessible.
+	* 
+	* For game code access to the cpu copy image, use GetCPUCopy().
+	*/
+	ENGINE_API UTexture2D* GetCPUCopyTexture();
+#endif
+
 	friend struct FRenderAssetStreamingManager;
 	friend struct FStreamingRenderAsset;
 	
-	/** creates and initializes a new Texture2D with the requested settings */
-	ENGINE_API static class UTexture2D* CreateTransient(int32 InSizeX, int32 InSizeY, EPixelFormat InFormat = PF_B8G8R8A8, const FName InName = NAME_None);
+	/** creates and initializes a new Texture2D with the requested settings. The texture will have 1 mip level of the given size, optionally filled with the provided data. */
+	ENGINE_API static class UTexture2D* CreateTransient(int32 InSizeX, int32 InSizeY, EPixelFormat InFormat = PF_B8G8R8A8, const FName InName = NAME_None, TConstArrayView64<uint8> InImageData = TConstArrayView64<uint8>());
+
+	/** creates a new texture2d from the first slice of the given image. If the image format isn't supported a texture isn't created. */
+	ENGINE_API static class UTexture2D* CreateTransientFromImage(const FImage* InImage, const FName InName = NAME_None);
+
+	/**
+	* If the texture has a cpu accessible copy, this returns that copy. It will wait for encoding in the editor,
+	* but otherwise may return an invalid reference if the texture is not ready or doesn't have a cpu copy.
+	*/
+	ENGINE_API FSharedImageConstRef GetCPUCopy() const;
+
+	UFUNCTION(BlueprintCallable, meta=(DisplayName = "GetCPUCopy"), Category = "Rendering|Texture")
+	FSharedImageConstRefBlueprint Blueprint_GetCPUCopy() const;
+
 
 	/**
 	 * Gets the X size of the texture, in pixels
diff --git a/Engine/Source/Runtime/Engine/Classes/Engine/TextureDefines.h b/Engine/Source/Runtime/Engine/Classes/Engine/TextureDefines.h
--- a/Engine/Source/Runtime/Engine/Classes/Engine/TextureDefines.h
+++ b/Engine/Source/Runtime/Engine/Classes/Engine/TextureDefines.h
@@ -198,6 +198,13 @@ enum class ETextureMipLoadOptions : uint8
 	OnlyFirstMip,
 };
 
+UENUM()
+enum class ETextureAvailability : uint8
+{
+	GPU,
+	CPU
+};
+
 UENUM()
 enum class ETextureDownscaleOptions : uint8
 {
diff --git a/Engine/Source/Runtime/Engine/Private/Texture.cpp b/Engine/Source/Runtime/Engine/Private/Texture.cpp
--- a/Engine/Source/Runtime/Engine/Private/Texture.cpp
+++ b/Engine/Source/Runtime/Engine/Private/Texture.cpp
@@ -32,6 +32,7 @@
 #include "Engine/TextureCube.h"
 #include "Engine/RendererSettings.h"
 #include "ColorSpace.h"
+#include "ImageCoreBP.h"
 #include "ImageCoreUtils.h"
 #include "ImageUtils.h"
 #include "Algo/Unique.h"
@@ -144,6 +145,7 @@ UTexture::UTexture(const FObjectInitializer& ObjectInitializer)
 #endif
 	, TextureReference(*new FTextureReference())
 {
+	Availability = ETextureAvailability::GPU;
 	SRGB = true;
 	Filter = TF_Default;
 	MipLoadOptions = ETextureMipLoadOptions::Default;
@@ -3660,6 +3662,11 @@ FName GetDefaultTextureFormatName( const ITargetPlatform* TargetPlatform, const
 		}
 	}
 
+	// if we are using a placeholder texture, just DXT it with no alpha.
+	if (Texture->Availability == ETextureAvailability::CPU)
+	{
+		TextureFormatName = NameDXT1;
+	}
 #endif //WITH_EDITOR
 
 	//FName Result = ConditionalGetPrefixedFormat(TextureFormatName, TargetPlatform, bOodleTextureSdkVersionIsNone);
diff --git a/Engine/Source/Runtime/Engine/Private/Texture2D.cpp b/Engine/Source/Runtime/Engine/Private/Texture2D.cpp
--- a/Engine/Source/Runtime/Engine/Private/Texture2D.cpp
+++ b/Engine/Source/Runtime/Engine/Private/Texture2D.cpp
@@ -16,6 +16,7 @@
 #include "DeviceProfiles/DeviceProfile.h"
 #include "DeviceProfiles/DeviceProfileManager.h"
 #include "DerivedDataCache.h"
+#include "Math/GuardedInt.h"
 #include "Rendering/Texture2DResource.h"
 #include "Streaming/Texture2DStreamOut_AsyncCreate.h"
 #include "RenderingThread.h"
@@ -706,6 +707,10 @@ void UTexture2D::UpdateResource()
 	WaitForPendingInitOrStreaming();
 
 #if WITH_EDITOR
+	// Invalidate the CPU texture in case we changed sources - it'll recreate
+	// on access if needed.
+	CPUCopyTexture = nullptr;
+
 	// Recache platform data if the source has changed.
 	if (FTextureCompilingManager::Get().IsAsyncCompilationAllowed(this))
 	{
@@ -1089,39 +1094,158 @@ void UTexture2D::GetResourceSizeEx(FResourceSizeEx& CumulativeResourceSize)
 	}
 }
 
-UTexture2D* UTexture2D::CreateTransient(int32 InSizeX, int32 InSizeY, EPixelFormat InFormat, const FName InName)
+
+UTexture2D* UTexture2D::CreateTransientFromImage(const FImage* InImage, const FName InName)
 {
 	LLM_SCOPE(ELLMTag::Textures);
+	if (InImage == nullptr)
+	{
+		return nullptr;
+	}
+
+	TArray64<uint8> ConvertedData;
+
+	EPixelFormat Format = PF_B8G8R8A8;
+	switch (InImage->Format)
+	{
+	case ERawImageFormat::BGRE8:
+		{
+			Format = PF_A32B32G32R32F; 
+			ConvertedData.AddUninitialized(InImage->SizeX * InImage->SizeY * sizeof(FLinearColor));
+			FLinearColor* Output = (FLinearColor*)ConvertedData.GetData();
+			const FColor* Input = (const FColor*)InImage->RawData.GetData();
+
+			uint64 Pixels = (uint64)InImage->SizeX * InImage->SizeY;
+			for (uint64 Pixel = 0; Pixel < Pixels; Pixel++)
+			{
+				Output[Pixel] = Input[Pixel].FromRGBE();
+			}
+			break;
+		}
+	case ERawImageFormat::BGRA8: Format = PF_B8G8R8A8; break;
+	case ERawImageFormat::RGBA16F: Format = PF_FloatRGBA; break;
+	case ERawImageFormat::RGBA32F: Format = PF_A32B32G32R32F; break;
+	case ERawImageFormat::R16F: Format = PF_R16F; break;
+	case ERawImageFormat::R32F: Format = PF_R32_FLOAT; break;
+	case ERawImageFormat::G8:
+		{
+			Format = PF_B8G8R8A8;
+			ConvertedData.AddUninitialized(InImage->SizeX * InImage->SizeY * sizeof(FColor));
+			FColor* Output = (FColor*)ConvertedData.GetData();
+			const uint8* Input = (const uint8*)InImage->RawData.GetData();
+			
+			uint64 Pixels = (uint64)InImage->SizeX * InImage->SizeY;
+			for (uint64 Pixel = 0; Pixel < Pixels; Pixel++)
+			{
+				Output[Pixel].R = Output[Pixel].G = Output[Pixel].B = Input[Pixel];
+				Output[Pixel].A = 255;
+			}
 
-	UTexture2D* NewTexture = nullptr;
-	if (InSizeX > 0 && InSizeY > 0 &&
-		(InSizeX % GPixelFormats[InFormat].BlockSizeX) == 0 &&
-		(InSizeY % GPixelFormats[InFormat].BlockSizeY) == 0)
-	{
-		NewTexture = NewObject<UTexture2D>(
-			GetTransientPackage(),
-			InName,
-			RF_Transient
-			);
-
-		NewTexture->SetPlatformData(new FTexturePlatformData());
-		NewTexture->GetPlatformData()->SizeX = InSizeX;
-		NewTexture->GetPlatformData()->SizeY = InSizeY;
-		NewTexture->GetPlatformData()->SetNumSlices(1);
-		NewTexture->GetPlatformData()->PixelFormat = InFormat;
-
-		// Allocate first mipmap.
-		int32 NumBlocksX = InSizeX / GPixelFormats[InFormat].BlockSizeX;
-		int32 NumBlocksY = InSizeY / GPixelFormats[InFormat].BlockSizeY;
-		FTexture2DMipMap* Mip = new FTexture2DMipMap(InSizeX, InSizeY, 1);
-		NewTexture->GetPlatformData()->Mips.Add(Mip);
-		Mip->BulkData.Lock(LOCK_READ_WRITE);
-		Mip->BulkData.Realloc((int64)NumBlocksX * NumBlocksY * GPixelFormats[InFormat].BlockBytes);
-		Mip->BulkData.Unlock();
+			break;
+		}
+	case ERawImageFormat::G16:
+		{
+			Format = PF_A16B16G16R16;
+			ConvertedData.AddUninitialized(InImage->SizeX * InImage->SizeY * 4 * sizeof(uint16));
+			uint16* Output = (uint16*)ConvertedData.GetData();
+			const uint16* Input = (const uint16*)InImage->RawData.GetData();
+
+			uint64 Pixels = (uint64)InImage->SizeX * InImage->SizeY;
+			for (uint64 Pixel = 0; Pixel < Pixels; Pixel++)
+			{
+				Output[Pixel * 4 + 3] = 65535;
+				Output[Pixel * 4 + 2] = Input[Pixel];
+				Output[Pixel * 4 + 1] = Input[Pixel];
+				Output[Pixel * 4 + 0] = Input[Pixel];
+			}
+
+			break;
+		}
+	case ERawImageFormat::RGBA16:
+		{
+			Format = PF_A16B16G16R16;
+			break;
+		}
+	default:
+		{
+			UE_LOG(LogTexture, Error, TEXT("Invalid raw image format in UTexture2D::CreateTransientFromImage()"));
+			return nullptr;
+		}
 	}
-	else
+
+	// We only want to provide one slice if the source is for an array/cube
+	TConstArrayView64<uint8> ImageData = InImage->RawData;
+	if (InImage->NumSlices > 1)
+	{
+		int64 SliceBytes = ImageData.Num() / InImage->NumSlices;
+		ImageData.LeftInline(SliceBytes);
+	}
+
+	return CreateTransient(InImage->GetWidth(), InImage->GetHeight(), Format, InName, ConvertedData.Num() ? ConvertedData : ImageData);
+}
+
+
+UTexture2D* UTexture2D::CreateTransient(int32 InSizeX, int32 InSizeY, EPixelFormat InFormat, const FName InName, TConstArrayView64<uint8> InImageData)
+{
+	LLM_SCOPE(ELLMTag::Textures);
+
+	const int32 NumBlocksX = InSizeX / GPixelFormats[InFormat].BlockSizeX;
+	const int32 NumBlocksY = InSizeY / GPixelFormats[InFormat].BlockSizeY;
+
+	if (InSizeX <= 0 || InSizeY <= 0)
+	{
+		UE_LOG(LogTexture, Warning, TEXT("Negative size specified for UTexture2D::CreateTransient()"));
+		return nullptr;
+	}
+
+	if ((InSizeX % GPixelFormats[InFormat].BlockSizeX) ||
+		(InSizeY % GPixelFormats[InFormat].BlockSizeY))
+	{
+		UE_LOG(LogTexture, Warning, TEXT("Size specified isn't valid for block-based pixel format in UTexture2D::CreateTransient()"));
+		return nullptr;
+	}
+
+	FGuardedInt64 BytesForImageValidation = FGuardedInt64(NumBlocksX) * NumBlocksY * GPixelFormats[InFormat].BlockBytes;
+	if (BytesForImageValidation.IsValid() == false)
+	{
+		UE_LOG(LogTexture, Warning, TEXT("Size specified overflows in UTexture2D::CreateTransient()"));
+		return nullptr;
+	}
+
+	int64 BytesForImage = BytesForImageValidation.Get(0);
+
+	// If they provided data, it needs to be the right size.
+	if (InImageData.Num() && InImageData.Num() != BytesForImage)
+	{
+		UE_LOG(LogTexture, Warning, TEXT("Image data provided is incorrect size (%llu provided, %llu wanted) in UTexture2D::CreateTransient()"), InImageData.Num(), BytesForImage);
+		return nullptr;
+	}
+
+	UTexture2D* NewTexture = NewObject<UTexture2D>(
+		GetTransientPackage(),
+		InName,
+		RF_Transient
+		);
+
+	NewTexture->SetPlatformData(new FTexturePlatformData());
+	NewTexture->GetPlatformData()->SizeX = InSizeX;
+	NewTexture->GetPlatformData()->SizeY = InSizeY;
+	NewTexture->GetPlatformData()->SetNumSlices(1);
+	NewTexture->GetPlatformData()->PixelFormat = InFormat;
+
+	// Allocate first mipmap.
+	FTexture2DMipMap* Mip = new FTexture2DMipMap(InSizeX, InSizeY, 1);
+	NewTexture->GetPlatformData()->Mips.Add(Mip);
+	Mip->BulkData.Lock(LOCK_READ_WRITE);
+	void* DestImageData = Mip->BulkData.Realloc(BytesForImage);
+	if (InImageData.Num())
 	{
-		UE_LOG(LogTexture, Warning, TEXT("Invalid parameters specified for UTexture2D::CreateTransient()"));
+		FMemory::Memcpy(DestImageData, InImageData.GetData(), BytesForImage);
+	}
+	Mip->BulkData.Unlock();
+	if (InImageData.Num())
+	{
+		NewTexture->UpdateResource();
 	}
 	return NewTexture;
 }
@@ -1144,6 +1268,48 @@ int32 UTexture2D::Blueprint_GetSizeX() const
 	return GetSizeX();
 }
 
+#if WITH_EDITORONLY_DATA
+UTexture2D* UTexture2D::GetCPUCopyTexture()
+{
+	if (CPUCopyTexture)
+	{
+		return CPUCopyTexture;
+	}
+
+	FSharedImageConstRef CPUCopy = GetCPUCopy(); // blocks in the editor during encoding
+	if (CPUCopy)
+	{
+		CPUCopyTexture = UTexture2D::CreateTransientFromImage(CPUCopy.GetReference());
+	}
+	return CPUCopyTexture;
+}
+#endif
+
+
+FSharedImageConstRef UTexture2D::GetCPUCopy() const
+{
+	// could stall if texture isn't built!
+	const FTexturePlatformData* LocalPlatformData = GetPlatformData();
+
+	if (LocalPlatformData && LocalPlatformData->GetHasCpuCopy())
+	{
+		return LocalPlatformData->CPUCopy;
+	}
+
+	return FSharedImageConstRef();
+
+}
+
+FSharedImageConstRefBlueprint UTexture2D::Blueprint_GetCPUCopy() const
+{
+	// When cooking, blueprint construction scripts are run before textures get postloaded.
+	// We don't have valid platformdata, so we always have a null reference here which passes
+	// back to the caller.
+	FSharedImageConstRefBlueprint Result;
+	Result.Reference = GetCPUCopy();
+	return Result;
+}
+
 int32 UTexture2D::Blueprint_GetSizeY() const
 {
 #if WITH_EDITORONLY_DATA
diff --git a/Engine/Source/Runtime/Engine/Private/TextureDerivedData.cpp b/Engine/Source/Runtime/Engine/Private/TextureDerivedData.cpp
--- a/Engine/Source/Runtime/Engine/Private/TextureDerivedData.cpp
+++ b/Engine/Source/Runtime/Engine/Private/TextureDerivedData.cpp
@@ -33,6 +33,7 @@
 #include "DerivedDataRequestOwner.h"
 #include "Hash/xxhash.h"
 #include "ImageCoreUtils.h"
+#include "ImageUtils.h"
 #include "Interfaces/ITargetPlatform.h"
 #include "Interfaces/ITargetPlatformManagerModule.h"
 #include "Interfaces/ITextureFormat.h"
@@ -321,6 +322,13 @@ static void SerializeForKey(FArchive& Ar, const FTextureBuildSettings& Settings)
 		Ar << TempGuid;
 	}
 
+	if (Settings.bCPUAccessible)
+	{
+		// @todo SerializeForKey these can go away whenever we bump the overall ddc key
+		TempGuid = FGuid(0x583A3B04, 0xC41C4E2C, 0x9FB77E7D, 0xC7AEFE7E);
+		Ar << TempGuid;
+	}
+
 
 	if ( Settings.bUseNewMipFilter )
 	{
@@ -905,7 +913,14 @@ static void GetTextureBuildSettings(
 	static const auto CVarVirtualTexturesEnabled = IConsoleManager::Get().FindTConsoleVariableDataInt(TEXT("r.VirtualTextures")); check(CVarVirtualTexturesEnabled);
 	// A ULightMapVirtualTexture2D with multiple layers saved in MapBuildData could be loaded with the r.VirtualTexture disabled, it will generate DDC before we decide to invalidate the light map data, to skip the ensure failure let it generate VT DDC anyway.
 	const bool bForVirtualTextureStreamingBuild = ULightMapVirtualTexture2D::StaticClass() == Texture.GetClass();
-	const bool bVirtualTextureStreaming = bForVirtualTextureStreamingBuild || (CVarVirtualTexturesEnabled->GetValueOnAnyThread() && bPlatformSupportsVirtualTextureStreaming && Texture.VirtualTextureStreaming);
+	bool bVirtualTextureStreaming = bForVirtualTextureStreamingBuild || (CVarVirtualTexturesEnabled->GetValueOnAnyThread() && bPlatformSupportsVirtualTextureStreaming && Texture.VirtualTextureStreaming);
+	if (Texture.Availability == ETextureAvailability::CPU && TextureClass == ETextureClass::TwoD)
+	{
+		// We are swapping with a placeholder - don't VT it.
+		OutBuildSettings.bCPUAccessible = true;
+		bVirtualTextureStreaming = false;
+		MipGenSettings = TMGS_NoMipmaps;
+	}
 
 
 	// Virtual textures must have mips as VT memory management relies on a 1:1 texel/pixel mapping, which in turn
@@ -1470,6 +1485,29 @@ void FTexturePlatformData::FinishCache()
 	}
 }
 
+void FTexturePlatformData::Reset()
+{
+	Mips.Empty();
+	SizeX = 0;
+	SizeY = 0;
+	PixelFormat = PF_Unknown;
+	PackedData = 0;
+	OptData = FOptTexturePlatformData();
+	if (VTData)
+	{
+		delete VTData;
+	}
+	VTData = nullptr;
+	CPUCopy.SafeRelease();
+
+#if WITH_EDITORONLY_DATA
+	bSourceMipsAlphaDetectedValid = false;
+	PreEncodeMipsHash = 0;
+	ResultMetadata.bIsValid = false;
+#endif
+}
+
+
 typedef TArray<uint32, TInlineAllocator<MAX_TEXTURE_MIP_COUNT> > FAsyncMipHandles;
 typedef TArray<uint32> FAsyncVTChunkHandles;
 
@@ -1995,7 +2033,11 @@ FTexturePlatformData::~FTexturePlatformData()
 		AsyncTask = nullptr;
 	}
 #endif
-	if (VTData) delete VTData;
+	if (VTData) 
+	{
+		delete VTData;
+	}
+	CPUCopy = nullptr;
 }
 
 bool FTexturePlatformData::IsReadyForAsyncPostLoad() const
@@ -2782,6 +2824,23 @@ static void SerializePlatformData(
 		Ar << PlatformData->OptData;
 	}
 
+	if (PlatformData->GetHasCpuCopy())
+	{
+		if (Ar.IsLoading())
+		{
+			PlatformData->CPUCopy = FSharedImageConstRef(new FSharedImage());
+		}
+
+		// we have to cast off the const since we load in to it here as well as save.
+		FSharedImage* ImageToSerialize = (FSharedImage*)PlatformData->CPUCopy.GetReference();
+		Ar << ImageToSerialize->SizeX;
+		Ar << ImageToSerialize->SizeY;
+		Ar << ImageToSerialize->NumSlices;
+		Ar << (uint8&)ImageToSerialize->Format;
+		Ar << ImageToSerialize->GammaSpace;
+		Ar << ImageToSerialize->RawData;
+	}
+
 	if (bCooked)
 	{
 		Ar << FirstMipToSerialize;
diff --git a/Engine/Source/Runtime/Engine/Private/TextureDerivedDataBuildUtils.cpp b/Engine/Source/Runtime/Engine/Private/TextureDerivedDataBuildUtils.cpp
--- a/Engine/Source/Runtime/Engine/Private/TextureDerivedDataBuildUtils.cpp
+++ b/Engine/Source/Runtime/Engine/Private/TextureDerivedDataBuildUtils.cpp
@@ -202,6 +202,7 @@ static void WriteBuildSettings(FCbWriter& Writer, const FTextureBuildSettings& B
 	WriteCbFieldWithDefault<bool>(Writer, "bApplyYCoCgBlockScale", BuildSettings.bApplyYCoCgBlockScale, DefaultSettings.bApplyYCoCgBlockScale);
 	WriteCbFieldWithDefault<bool>(Writer, "bApplyKernelToTopMip", BuildSettings.bApplyKernelToTopMip, DefaultSettings.bApplyKernelToTopMip);
 	WriteCbFieldWithDefault<bool>(Writer, "bRenormalizeTopMip", BuildSettings.bRenormalizeTopMip, DefaultSettings.bRenormalizeTopMip);
+	WriteCbFieldWithDefault<bool>(Writer, "bCPUAccessible", BuildSettings.bCPUAccessible, DefaultSettings.bCPUAccessible);
 	WriteCbFieldWithDefault(Writer, "CompositeTextureMode", BuildSettings.CompositeTextureMode, DefaultSettings.CompositeTextureMode);
 	WriteCbFieldWithDefault(Writer, "CompositePower", BuildSettings.CompositePower, DefaultSettings.CompositePower);
 	WriteCbFieldWithDefault(Writer, "LODBias", BuildSettings.LODBias, DefaultSettings.LODBias);
diff --git a/Engine/Source/Runtime/Engine/Private/TextureDerivedDataTask.cpp b/Engine/Source/Runtime/Engine/Private/TextureDerivedDataTask.cpp
--- a/Engine/Source/Runtime/Engine/Private/TextureDerivedDataTask.cpp
+++ b/Engine/Source/Runtime/Engine/Private/TextureDerivedDataTask.cpp
@@ -98,6 +98,7 @@ public:
 	}
 };
 
+
 static FText ComposeTextureBuildText(const FString& TexturePathName, int32 SizeX, int32 SizeY, int32 NumSlices, int32 NumBlocks, int32 NumLayers, const FTextureBuildSettings& BuildSettings, ETextureEncodeSpeed InEncodeSpeed, int64 RequiredMemoryEstimate, bool bIsVT)
 {
 	FFormatNamedArguments Args;
@@ -173,6 +174,37 @@ static bool ValidateTexture2DPlatformData(const FTexturePlatformData& TextureDat
 #endif
 }
 
+void FTextureSourceData::InitAsPlaceholder()
+{
+	ReleaseMemory();
+
+	// This needs to be a tiny texture that can encode on all hardware. It's job is to
+	// take up as little memory as possible for textures where we'd rather they not create
+	// hw resources at all, but we don't want to hack in a ton of redirects/tests all over
+	// the rendering codebase.
+
+	// So we make a 4x4 black RGBA8 texture.
+	FTextureSourceBlockData& Block = Blocks.AddDefaulted_GetRef();
+	{
+		Block.NumMips = 1;
+		TArray<FImage>& MipsPerLayer = Block.MipsPerLayer.AddDefaulted_GetRef();
+		FImage& Mip = MipsPerLayer.AddDefaulted_GetRef();
+		UE::TextureBuildUtilities::GetPlaceholderTextureImage(&Mip);
+
+		Block.NumSlices = Mip.NumSlices;
+		Block.SizeX = Mip.SizeX;
+		Block.SizeY = Mip.SizeY;
+	}
+
+	FTextureSourceLayerData& Layer = Layers.AddDefaulted_GetRef();
+	{
+		Layer.ImageFormat = ERawImageFormat::BGRA8;
+		Layer.SourceGammaSpace = EGammaSpace::Linear;
+	}
+
+	bValid = true;
+}
+
 void FTextureSourceData::Init(UTexture& InTexture, TextureMipGenSettings InMipGenSettings, bool bInCubeMap, bool bInTextureArray, bool bInVolumeTexture, bool bAllowAsyncLoading)
 {
 	check( bValid == false ); // we set to true at the end, acts as our return value
@@ -659,6 +691,8 @@ static void DDC1_BuildTexture(
 		ComposeTextureBuildText(TexturePathName, TextureData, InBuildSettingsPerLayer[0], (ETextureEncodeSpeed)InBuildSettingsPerLayer[0].RepresentsEncodeSpeedNoSend, RequiredMemoryEstimate, bForVirtualTextureStreamingBuild)
 		);
 
+	DerivedData->Reset();
+
 	if (bForVirtualTextureStreamingBuild)
 	{
 		if (DerivedData->VTData == nullptr)
@@ -741,13 +775,18 @@ static void DDC1_BuildTexture(
 				*TexturePathName, TextureData.Layers.Num());
 		}
 
+		if (InBuildSettingsPerLayer[0].bCPUAccessible)
+		{
+			// Copy out the unaltered top mip for cpu access.
+			FSharedImage* CPUCopy = new FSharedImage();
+			TextureData.Blocks[0].MipsPerLayer[0][0].CopyTo(*CPUCopy);
 
-		check(DerivedData->Mips.Num() == 0);
-		DerivedData->SizeX = 0;
-		DerivedData->SizeY = 0;
-		DerivedData->PixelFormat = PF_Unknown;
-		DerivedData->SetIsCubemap(false);
-		DerivedData->VTData = nullptr;
+			DerivedData->CPUCopy = FSharedImageConstRef(CPUCopy);
+			DerivedData->SetHasCpuCopy(true);
+			
+			// Divert the texture source data to a tiny placeholder texture.
+			TextureData.InitAsPlaceholder();
+		}
 
 		uint32 NumMipsInTail;
 		uint32 ExtData;
@@ -1269,13 +1308,7 @@ FTextureCacheDerivedDataWorker::FTextureCacheDerivedDataWorker(
 	check(Texture.Source.GetId().IsValid());
 
 	// Dump any existing mips.
-	DerivedData->Mips.Empty();
-	if (DerivedData->VTData)
-	{
-		delete DerivedData->VTData;
-		DerivedData->VTData = nullptr;
-	}
-	DerivedData->bSourceMipsAlphaDetectedValid = false;
+	DerivedData->Reset();
 	UTexture::GetPixelFormatEnum();
 		
 	const bool bAllowAsyncBuild = EnumHasAnyFlags(CacheFlags, ETextureCacheFlags::AllowAsyncBuild);
@@ -1292,7 +1325,15 @@ FTextureCacheDerivedDataWorker::FTextureCacheDerivedDataWorker(
 	// All of these settings are fixed across build settings and are derived directly from the texture.
 	// So we can just use layer 0 of whatever we have.
 	TextureData.Init(Texture, (TextureMipGenSettings)BuildSettingsPerLayerFetchOrBuild[0].MipGenSettings, BuildSettingsPerLayerFetchOrBuild[0].bCubemap, BuildSettingsPerLayerFetchOrBuild[0].bTextureArray, BuildSettingsPerLayerFetchOrBuild[0].bVolume, bAllowAsyncLoading);
-	if (Texture.GetCompositeTexture() && Texture.CompositeTextureMode != CTM_Disabled && Texture.GetCompositeTexture()->Source.IsValid())
+
+	bool bNeedsCompositeData = Texture.GetCompositeTexture() && Texture.CompositeTextureMode != CTM_Disabled && Texture.GetCompositeTexture()->Source.IsValid();
+	if (BuildSettingsPerLayerFetchOrBuild[0].bCPUAccessible)
+	{
+		// CPU accessible textures don't run image processing and thus don't need the composite data.
+		bNeedsCompositeData = false;
+	}
+
+	if (bNeedsCompositeData)
 	{
 		bool bMatchingBlocks = Texture.GetCompositeTexture()->Source.GetNumBlocks() == Texture.Source.GetNumBlocks();
 		
@@ -2095,6 +2136,31 @@ static bool UnpackPlatformDataFromBuild(FTexturePlatformData& OutPlatformData, U
 	using namespace UE::DerivedData;
 	UE::DerivedData::FBuildOutput& BuildOutput = InBuildCompleteParams.Output;
 
+	FImage CPUCopy;
+	{
+		// CPUCopy might not exist if the build didn't request it
+		const FValueWithId& MetadataValue = BuildOutput.GetValue(FValueId::FromName(ANSITEXTVIEW("CPUCopyImageInfo")));
+		if (MetadataValue.IsValid())
+		{
+			if (CPUCopy.ImageInfoFromCompactBinary(FCbObject(MetadataValue.GetData().Decompress())) == false)
+			{
+				UE_LOG(LogTexture, Error, TEXT("Invalid CPUCopyImageInfo in build output '%s' by %s."), *BuildOutput.GetName(), *WriteToString<32>(BuildOutput.GetFunction()));
+				return false;
+			}
+
+			const FValueWithId& DataValue = BuildOutput.GetValue(FValueId::FromName(ANSITEXTVIEW("CPUCopyRawData")));
+			if (DataValue.IsValid() == false)
+			{
+				UE_LOG(LogTexture, Error, TEXT("Missing CPUCopyRawData in build output '%s' by %s."), *BuildOutput.GetName(), *WriteToString<32>(BuildOutput.GetFunction()));
+				return false;
+			}
+
+			FSharedBuffer Data = DataValue.GetData().Decompress();
+			CPUCopy.RawData.AddUninitialized(Data.GetSize());
+			FMemory::Memcpy(CPUCopy.RawData.GetData(), Data.GetData(), Data.GetSize());
+		}
+	}
+
 	{
 		const FValueWithId& Value = BuildOutput.GetValue(FValueId::FromName(UTF8TEXTVIEW("TextureBuildMetadata")));
 		PackTextureBuildMetadataInPlatformData(&OutPlatformData, FCbObject(Value.GetData().Decompress()));
@@ -2156,7 +2222,8 @@ static bool UnpackPlatformDataFromBuild(FTexturePlatformData& OutPlatformData, U
 	OutPlatformData.OptData.ExtData = EncodedTextureExtendedData.ExtData;
 	{
 		const bool bHasOptData = (EncodedTextureExtendedData.NumMipsInTail != 0) || (EncodedTextureExtendedData.ExtData != 0);
-		OutPlatformData.SetPackedData(EncodedTextureDescription.GetNumSlices_WithDepth(0), bHasOptData, EncodedTextureDescription.bCubeMap);
+		const bool bHasCPUCopy = CPUCopy.Format != ERawImageFormat::Invalid;
+		OutPlatformData.SetPackedData(EncodedTextureDescription.GetNumSlices_WithDepth(0), bHasOptData, EncodedTextureDescription.bCubeMap, bHasCPUCopy);
 	}
 	OutPlatformData.Mips.Empty(EncodedTextureDescription.NumMips);
 	EFileRegionType FileRegion = FFileRegion::SelectType(EncodedTextureDescription.PixelFormat);
@@ -2472,7 +2539,7 @@ public:
 		Owner.Emplace(OwnerPriority);
 
 		bool bUseCompositeTexture = false;
-		if (!IsTextureValidForBuilding(Texture, Flags, bUseCompositeTexture))
+		if (!IsTextureValidForBuilding(Texture, Flags, InSettingsFetchOrBuild.bCPUAccessible, bUseCompositeTexture))
 		{
 			return;
 		}
@@ -2675,6 +2742,17 @@ public:
 
 private:
 
+	static constexpr FAnsiStringView NonStreamingMipOutputValueNames[] = 
+	{
+		ANSITEXTVIEW("EncodedTextureDescription"),
+		ANSITEXTVIEW("EncodedTextureExtendedData"),
+		ANSITEXTVIEW("MipTail"),
+		ANSITEXTVIEW("TextureBuildMetadata"),
+		ANSITEXTVIEW("CPUCopyMetadata"),
+		ANSITEXTVIEW("CPUCopyData")
+	};
+
+
 	static UE::DerivedData::FBuildPolicy FetchFirst_CreateBuildPolicy(FBuildResultOptions InBuildResultOptions)
 	{
 		using namespace UE::DerivedData;
@@ -2686,12 +2764,12 @@ private:
 		}
 		else
 		{
-			// We only want the metadata and the mip tail.
+			// Cache everything except the streaming mips.
 			FBuildPolicyBuilder FetchFirstBuildPolicyBuilder(EBuildPolicy::CacheQuery | EBuildPolicy::SkipData);
-			FetchFirstBuildPolicyBuilder.AddValuePolicy(FValueId::FromName(ANSITEXTVIEW("EncodedTextureDescription")), EBuildPolicy::Cache);
-			FetchFirstBuildPolicyBuilder.AddValuePolicy(FValueId::FromName(ANSITEXTVIEW("EncodedTextureExtendedData")), EBuildPolicy::Cache);
-			FetchFirstBuildPolicyBuilder.AddValuePolicy(FValueId::FromName(ANSITEXTVIEW("MipTail")), EBuildPolicy::Cache);
-			FetchFirstBuildPolicyBuilder.AddValuePolicy(FValueId::FromName(ANSITEXTVIEW("TextureBuildMetadata")), EBuildPolicy::Cache);
+			for (FAnsiStringView NonStreamingValue : NonStreamingMipOutputValueNames)
+			{
+				FetchFirstBuildPolicyBuilder.AddValuePolicy(FValueId::FromName(NonStreamingValue), EBuildPolicy::Cache);
+			}
 			return FetchFirstBuildPolicyBuilder.Build();
 		}		
 	}
@@ -2711,10 +2789,10 @@ private:
 		else
 		{
 			FBuildPolicyBuilder BuildPolicyBuilder(EBuildPolicy::Build | EBuildPolicy::CacheQuery | EBuildPolicy::CacheStoreOnBuild | EBuildPolicy::SkipData);
-			BuildPolicyBuilder.AddValuePolicy(FValueId::FromName(ANSITEXTVIEW("EncodedTextureDescription")), EBuildPolicy::Cache);
-			BuildPolicyBuilder.AddValuePolicy(FValueId::FromName(ANSITEXTVIEW("EncodedTextureExtendedData")), EBuildPolicy::Cache);
-			BuildPolicyBuilder.AddValuePolicy(FValueId::FromName(ANSITEXTVIEW("MipTail")), EBuildPolicy::Default);
-			BuildPolicyBuilder.AddValuePolicy(FValueId::FromName(ANSITEXTVIEW("TextureBuildMetadata")), EBuildPolicy::Cache);
+			for (FAnsiStringView NonStreamingValue : NonStreamingMipOutputValueNames)
+			{
+				BuildPolicyBuilder.AddValuePolicy(FValueId::FromName(NonStreamingValue), EBuildPolicy::Cache);
+			}
 			return BuildPolicyBuilder.Build();
 		}
 	}
@@ -2774,7 +2852,7 @@ public:
 		return Owner->Poll();
 	}
 
-	static bool IsTextureValidForBuilding(UTexture& Texture, ETextureCacheFlags Flags, bool& bOutUseCompositeTexture)
+	static bool IsTextureValidForBuilding(UTexture& Texture, ETextureCacheFlags Flags, bool bInCPUAccessible, bool& bOutUseCompositeTexture)
 	{
 		bOutUseCompositeTexture = false;
 
@@ -2828,7 +2906,11 @@ public:
 			}
 		}
 		
-		const bool bCompositeTextureViable = Texture.GetCompositeTexture() && Texture.CompositeTextureMode != CTM_Disabled && Texture.GetCompositeTexture()->Source.IsValid();
+		bool bCompositeTextureViable = Texture.GetCompositeTexture() && Texture.CompositeTextureMode != CTM_Disabled && Texture.GetCompositeTexture()->Source.IsValid();
+		if (bInCPUAccessible)
+		{
+			bCompositeTextureViable = false;
+		}
 		bool bMatchingBlocks = bCompositeTextureViable && (Texture.GetCompositeTexture()->Source.GetNumBlocks() == Texture.Source.GetNumBlocks());
 		
 		if (bCompositeTextureViable)
@@ -3045,7 +3127,7 @@ FTexturePlatformData::FStructuredDerivedDataKey CreateTextureDerivedDataKey(
 		Texture.GetPathName(nullptr, TexturePath);
 
 		bool bUseCompositeTexture = false;
-		if (FTextureBuildTask::IsTextureValidForBuilding(Texture, CacheFlags, bUseCompositeTexture))
+		if (FTextureBuildTask::IsTextureValidForBuilding(Texture, CacheFlags, Settings.bCPUAccessible, bUseCompositeTexture))
 		{
 			// this is just to make DDC Key so I don't need RequiredMemoryEstimate
 			// but it goes in the the DDC Key, so I have to compute it
diff --git a/Engine/Source/Runtime/Engine/Private/TextureDerivedDataTask.h b/Engine/Source/Runtime/Engine/Private/TextureDerivedDataTask.h
--- a/Engine/Source/Runtime/Engine/Private/TextureDerivedDataTask.h
+++ b/Engine/Source/Runtime/Engine/Private/TextureDerivedDataTask.h
@@ -89,6 +89,9 @@ struct FTextureSourceData
 		, bValid(false)
 	{}
 
+	// Clear the current source data and make it into a placeholder texture.
+	void InitAsPlaceholder();
+
 	void Init(UTexture& InTexture, TextureMipGenSettings InMipGenSettings, bool bInCubeMap, bool bInTextureArray, bool bInVolumeTexture, bool bAllowAsyncLoading);
 	bool IsValid() const { return bValid; }
 
@@ -108,6 +111,8 @@ struct FTextureSourceData
 		// Unload BulkData loaded with LoadBulkDataWithFileReader
 		AsyncSource.RemoveBulkData();
 		Blocks.Empty();
+		Layers.Empty();
+		bValid = false;
 	}
 
 	FString TextureFullName;
diff --git a/Engine/Source/Runtime/ImageCore/Private/ImageCore.cpp b/Engine/Source/Runtime/ImageCore/Private/ImageCore.cpp
--- a/Engine/Source/Runtime/ImageCore/Private/ImageCore.cpp
+++ b/Engine/Source/Runtime/ImageCore/Private/ImageCore.cpp
@@ -3,6 +3,7 @@
 #include "ImageCore.h"
 #include "Modules/ModuleManager.h"
 #include "Async/ParallelFor.h"
+#include "Serialization/CompactBinaryWriter.h"
 #include "TransferFunctions.h"
 #include "ColorSpace.h"
 
@@ -1310,6 +1311,26 @@ void FImage::Linearize(uint8 SourceEncoding, FImage& DestImage) const
 	}
 }
 
+const FUtf8StringView ERawImageFormat::GetNameView(Type Format)
+{
+	switch (Format)
+	{
+	case ERawImageFormat::G8:  return UTF8TEXTVIEW("G8");
+	case ERawImageFormat::G16: return UTF8TEXTVIEW("G16");
+	case ERawImageFormat::R16F: return UTF8TEXTVIEW("R16F");
+	case ERawImageFormat::R32F: return UTF8TEXTVIEW("R32F");
+	case ERawImageFormat::BGRA8: return UTF8TEXTVIEW("BGRA8");
+	case ERawImageFormat::BGRE8: return UTF8TEXTVIEW("BGRE8");
+	case ERawImageFormat::RGBA16: return UTF8TEXTVIEW("RGBA16");
+	case ERawImageFormat::RGBA16F: return UTF8TEXTVIEW("RGBA16F");
+	case ERawImageFormat::RGBA32F: return UTF8TEXTVIEW("RGBA32F");
+	default:
+		check(0);
+		return UTF8TEXTVIEW("invalid");
+	}
+}
+
+
 IMAGECORE_API const TCHAR * ERawImageFormat::GetName(Type Format)
 {
 	switch (Format)
@@ -1329,6 +1350,24 @@ IMAGECORE_API const TCHAR * ERawImageFormat::GetName(Type Format)
 	}
 }
 
+bool ERawImageFormat::GetFormatFromString(FUtf8StringView InString, ERawImageFormat::Type& OutType)
+{
+	if (InString.Compare(UTF8TEXTVIEW("G8"), ESearchCase::IgnoreCase) == 0) { OutType = ERawImageFormat::G8; return true; }
+	else if (InString.Compare(UTF8TEXTVIEW("G16"), ESearchCase::IgnoreCase) == 0) { OutType = ERawImageFormat::G16; return true; }
+	else if (InString.Compare(UTF8TEXTVIEW("R16F"), ESearchCase::IgnoreCase) == 0) { OutType = ERawImageFormat::R16F; return true; }
+	else if (InString.Compare(UTF8TEXTVIEW("R32F"), ESearchCase::IgnoreCase) == 0) { OutType = ERawImageFormat::R32F; return true; }
+	else if (InString.Compare(UTF8TEXTVIEW("BGRA8"), ESearchCase::IgnoreCase) == 0) { OutType = ERawImageFormat::BGRA8; return true; }
+	else if (InString.Compare(UTF8TEXTVIEW("BGRE8"), ESearchCase::IgnoreCase) == 0) { OutType = ERawImageFormat::BGRE8; return true; }
+	else if (InString.Compare(UTF8TEXTVIEW("RGBA16"), ESearchCase::IgnoreCase) == 0) { OutType = ERawImageFormat::RGBA16; return true; }
+	else if (InString.Compare(UTF8TEXTVIEW("RGBA16F"), ESearchCase::IgnoreCase) == 0) { OutType = ERawImageFormat::RGBA16F; return true; }
+	else if (InString.Compare(UTF8TEXTVIEW("RGBA32F"), ESearchCase::IgnoreCase) == 0) { OutType = ERawImageFormat::RGBA32F; return true; }
+	else if (InString.Compare(UTF8TEXTVIEW("Invalid"), ESearchCase::IgnoreCase) == 0) { OutType = ERawImageFormat::Invalid; return true; }
+
+	// We distinguish between "they stored invalid on purpose" and "we didn't recognize it" by the return value.
+	OutType = ERawImageFormat::Invalid;
+	return false;
+}
+
 IMAGECORE_API int64 ERawImageFormat::GetBytesPerPixel(Type Format)
 {
 	int64 BytesPerPixel = 0;
@@ -1685,3 +1724,39 @@ void FImageCore::TransformToWorkingColorSpace(const FImageView& InLinearImage, c
 			ImageColors[TexelIndex] = SaturateToHalfFloat(Color);
 		});
 }
+
+static constexpr uint32 ImageInfoSerializationVersion = 1;
+
+void FImageInfo::ImageInfoToCompactBinary(FCbObject& OutObject) const
+{
+	FCbWriter Writer;
+	Writer.BeginObject();
+	Writer.AddInteger("Version", ImageInfoSerializationVersion);
+	Writer.AddInteger("SizeX", SizeX);
+	Writer.AddInteger("SizeY", SizeY);
+	Writer.AddInteger("NumSlices", NumSlices);
+	Writer.AddInteger("GammaSpace", (uint8)GammaSpace);
+	Writer.AddString("Format", ERawImageFormat::GetNameView(Format));
+	Writer.EndObject();
+	OutObject = Writer.Save().AsObject();
+}
+
+bool FImageInfo::ImageInfoFromCompactBinary(const FCbObject& InObject)
+{
+	if (InObject["Version"].AsUInt32() != ImageInfoSerializationVersion)
+	{
+		return false;
+	}
+
+	SizeX = InObject["SizeX"].AsInt32();
+	SizeY = InObject["SizeY"].AsInt32();
+	NumSlices = InObject["NumSlices"].AsInt32();
+	GammaSpace = (EGammaSpace)InObject["GammaSpace"].AsUInt8();
+	if (ERawImageFormat::GetFormatFromString(InObject["Format"].AsString(), Format) == false)
+	{
+		return false;
+	}
+
+	return true;
+}
+
diff --git a/Engine/Source/Runtime/ImageCore/Private/ImageCoreBP.cpp b/Engine/Source/Runtime/ImageCore/Private/ImageCoreBP.cpp
new file mode 100644
--- /dev/null
+++ b/Engine/Source/Runtime/ImageCore/Private/ImageCoreBP.cpp
@@ -0,0 +1,144 @@
+// Copyright Epic Games, Inc. All Rights Reserved.
+
+#include "ImageCoreBP.h"
+
+bool USharedImageConstRefBlueprintFns::IsValid(const FSharedImageConstRefBlueprint& InSharedImage)
+{
+	return InSharedImage.Reference && InSharedImage.Reference->IsImageInfoValid();
+}
+
+FVector2f USharedImageConstRefBlueprintFns::GetSize(const FSharedImageConstRefBlueprint& InSharedImage)
+{
+	if (InSharedImage.Reference &&
+		InSharedImage.Reference->IsImageInfoValid())
+	{
+		return FVector2f((float)InSharedImage.Reference->GetWidth(), (float)InSharedImage.Reference->GetHeight());
+	}
+	return FVector2f(-1, -1);
+}
+
+int USharedImageConstRefBlueprintFns::GetWidth(const FSharedImageConstRefBlueprint& InSharedImage)
+{
+	if (InSharedImage.Reference &&
+		InSharedImage.Reference->IsImageInfoValid())
+	{
+		return InSharedImage.Reference->GetWidth();
+	}
+	return -1;
+}
+
+int USharedImageConstRefBlueprintFns::GetHeight(const FSharedImageConstRefBlueprint& InSharedImage)
+{
+	if (InSharedImage.Reference &&
+		InSharedImage.Reference->IsImageInfoValid())
+	{
+		return InSharedImage.Reference->GetHeight();
+	}
+	return -1;
+}
+
+FLinearColor USharedImageConstRefBlueprintFns::GetPixelLinearColor(const FSharedImageConstRefBlueprint& InSharedImage, int32 InX, int32 InY, bool& bOutValid, FLinearColor InFailureColor)
+{
+	const FImage* Image = InSharedImage.Reference;
+
+	bOutValid = false;
+	if (Image == nullptr || Image->IsImageInfoValid() == false)
+	{
+		return InFailureColor;
+	}
+
+	if (InX < 0 || InY < 0 ||
+		InX >= Image->GetWidth() || InY >= Image->GetHeight())
+	{
+		return InFailureColor;
+	}
+
+	// GetOnePixelLinear checks() on a bad format so we validate first
+	if (Image->Format >= ERawImageFormat::MAX)
+	{
+		return InFailureColor;
+	}
+
+	bOutValid = true;
+	return ERawImageFormat::GetOnePixelLinear(Image->GetPixelPointer(InX, InY), Image->Format, Image->GammaSpace);
+}
+
+FVector4f USharedImageConstRefBlueprintFns::GetPixelValue(const FSharedImageConstRefBlueprint& InSharedImage, int32 InX, int32 InY, bool& bOutValid)
+{
+	const FImage* Image = InSharedImage.Reference;
+
+	bOutValid = false;
+	if (Image == nullptr || Image->IsImageInfoValid() == false)
+	{
+		return FVector4f(0, 0, 0, 0);
+	}
+
+	if (InX < 0 || InY < 0 ||
+		InX >= Image->GetWidth() || InY >= Image->GetHeight())
+	{
+		return FVector4f(0, 0, 0, 0);
+	}
+
+	switch (Image->Format)
+	{
+	case ERawImageFormat::G8:
+		{
+			bOutValid = true;
+			uint8 G = *(const uint8*)Image->GetPixelPointer(InX, InY);
+			return FVector4f((float)G, (float)G, (float)G, 255.0f);
+		}
+	case ERawImageFormat::BGRA8:
+		{
+			bOutValid = true;
+			FColor Result = *(const FColor*)Image->GetPixelPointer(InX, InY);
+			return FVector4f((float)Result.R, (float)Result.G, (float)Result.B, (float)Result.A);
+		}
+	case ERawImageFormat::BGRE8:
+		{
+			bOutValid = true;
+			FColor Result = *(const FColor*)Image->GetPixelPointer(InX, InY);
+			return Result.FromRGBE();
+		}
+	case ERawImageFormat::RGBA16:
+		{
+			bOutValid = true;
+			const uint16* Result = (const uint16*)Image->GetPixelPointer(InX, InY);
+			return FVector4f((float)Result[0], (float)Result[1], (float)Result[2], (float)Result[3]);
+		}
+	case ERawImageFormat::RGBA16F:
+		{
+			bOutValid = true;
+			const FFloat16* Result = (const FFloat16*)Image->GetPixelPointer(InX, InY);
+			return FVector4f(Result[0].GetFloat(), Result[1].GetFloat(), Result[2].GetFloat(), Result[3].GetFloat());
+		}
+	case ERawImageFormat::RGBA32F:
+		{
+			bOutValid = true;
+			const float* Result = (const float*)Image->GetPixelPointer(InX, InY);
+			return FVector4f(Result[0], Result[1], Result[2], Result[3]);
+		}
+	case ERawImageFormat::G16:
+		{
+			bOutValid = true;
+			uint16 Result = *(const uint16*)Image->GetPixelPointer(InX, InY);
+			return FVector4f((float)Result, (float)Result, (float)Result, 1.0f);
+		}
+	case ERawImageFormat::R16F:
+		{
+			bOutValid = true;
+			const FFloat16* Result = (const FFloat16*)Image->GetPixelPointer(InX, InY);
+			float ResultFloat = Result->GetFloat();
+			return FVector4f(ResultFloat, 0.0f, 0.0f, 1.0f);
+		}
+	case ERawImageFormat::R32F:
+		{
+			bOutValid = true;
+			float Result = *(const float*)Image->GetPixelPointer(InX, InY);
+			return FVector4f(Result, 0.0f, 0.0f, 1.0f);
+		}
+	default:
+		{
+			return FVector4f(0, 0, 0, 0);
+		}
+	}
+}
\ No newline at end of file
diff --git a/Engine/Source/Runtime/ImageCore/Public/ImageCore.h b/Engine/Source/Runtime/ImageCore/Public/ImageCore.h
--- a/Engine/Source/Runtime/ImageCore/Public/ImageCore.h
+++ b/Engine/Source/Runtime/ImageCore/Public/ImageCore.h
@@ -64,6 +64,7 @@ namespace ERawImageFormat
 		G16,     // note G8/G16 = gray = replicate to 3 channels, R16F = just in red channel
 		R16F,
 		R32F,
+		MAX, // used for validation checks < MAX = valid type.
 		Invalid = 0xFF
 	};
 	
@@ -73,6 +74,8 @@ namespace ERawImageFormat
 	IMAGECORE_API int64 GetBytesPerPixel(Type Format);
 	
 	IMAGECORE_API const TCHAR * GetName(Type Format);
+	IMAGECORE_API const FUtf8StringView GetNameView(Type Format);
+	IMAGECORE_API bool GetFormatFromString(FUtf8StringView InString, Type& OutFormat);
 	
 	IMAGECORE_API bool IsHDR(Type Format);
 	
@@ -237,6 +240,9 @@ struct FImageInfo
 		return Offset;
 	}
 
+	IMAGECORE_API void ImageInfoToCompactBinary(class FCbObject& OutObject) const;
+	// Overwrites the current info with the object's info/
+	IMAGECORE_API bool ImageInfoFromCompactBinary(const FCbObject& InObject);
 };
 
 /***
@@ -689,6 +695,16 @@ public:
 	}
 };
 
+typedef TRefCountPtr<struct FSharedImage> FSharedImageRef;
+typedef TRefCountPtr<const struct FSharedImage> FSharedImageConstRef;
+struct FSharedImage : public FImage, public FThreadSafeRefCountedObject
+{
+	FSharedImage() = default;
+	virtual ~FSharedImage() = default;
+};
+
+
+
 /* Functions
  *****************************************************************************/
 
diff --git a/Engine/Source/Runtime/ImageCore/Public/ImageCoreBP.h b/Engine/Source/Runtime/ImageCore/Public/ImageCoreBP.h
new file mode 100644
--- /dev/null
+++ b/Engine/Source/Runtime/ImageCore/Public/ImageCoreBP.h
@@ -0,0 +1,75 @@
+// Copyright Epic Games, Inc. All Rights Reserved.
+#pragma once
+
+#include "ImageCore.h"
+#include "UObject/Object.h"
+#include "ImageCoreBP.generated.h"
+
+/**
+*  Exposes a FSharedImage to blueprint in an unalterable way.
+*/
+USTRUCT(BlueprintType, meta = (DisplayName = "FSharedImageConst"))
+struct IMAGECORE_API FSharedImageConstRefBlueprint
+{
+	GENERATED_BODY()
+	FSharedImageConstRef Reference;
+
+	FSharedImageConstRefBlueprint() = default;
+	~FSharedImageConstRefBlueprint() = default;
+	FSharedImageConstRefBlueprint(const FSharedImageConstRefBlueprint&) = default;
+	FSharedImageConstRefBlueprint(FSharedImageConstRefBlueprint&&) = default;
+	FSharedImageConstRefBlueprint& operator=(const FSharedImageConstRefBlueprint&) = default;
+	FSharedImageConstRefBlueprint& operator=(FSharedImageConstRefBlueprint&&) = default;
+};
+
+UCLASS(BlueprintType)
+class IMAGECORE_API USharedImageConstRefBlueprintFns : public UObject
+{
+	GENERATED_BODY()
+
+public:
+	UFUNCTION(BlueprintCallable, Category = "SharedImage")
+    static bool IsValid(const FSharedImageConstRefBlueprint& Image);
+
+	// Returns (-1, -1) if Image is invalid
+	UFUNCTION(BlueprintCallable, Category = "SharedImage")
+    static FVector2f GetSize(const FSharedImageConstRefBlueprint& Image);
+
+	// Returns -1 if Image is invalid
+	UFUNCTION(BlueprintCallable, Category = "SharedImage")
+    static int32 GetWidth(const FSharedImageConstRefBlueprint& Image);
+
+	// Returns -1 if Image is invalid
+	UFUNCTION(BlueprintCallable, Category = "SharedImage")
+    static int32 GetHeight(const FSharedImageConstRefBlueprint& Image);
+
+	/**
+	*	Returns the color value for the given pixel. If the input position is invalid, the format is invalid,
+	*	or the reference isn't set, bValid will be false and the function will return FailureColor. The color
+	*	is converted using the image's gamma space in to linear space.
+	* 
+	*	Do not use this for full image processing as it will be extremely slow, contact support if you need such
+	*	functionality.
+	*/
+	UFUNCTION(BlueprintCallable, Category = "SharedImage")
+    static FLinearColor GetPixelLinearColor(const FSharedImageConstRefBlueprint& Image, int32 X, int32 Y, bool& bValid, FLinearColor FailureColor=FLinearColor::Black);
+
+
+	/**
+	*	Returns the value in the texture for the given pixel as a float vector. If the input position is invalid, the format is invalid,
+	*	or the reference isn't set, bValid will be false and the function will return FVector4(0,0,0,0).
+	* 
+	*	Pixel values are directly returned with no gamma transformation to allow for lookup tables. Also note that
+	*	8 bit formats that you might normally expect to be normalized to 0..1 will return their values directly as 0..256.
+	*
+	*	This supports all image formats. 
+	* 
+	*	G8 is replicated to X/Y/Z/1.
+	*	R16/R32 is returned as R/0/0/1.
+	* 
+	*	Do not use this for full image processing as it will be extremely slow, contact support if you need such
+	*	functionality.
+	*/
+	UFUNCTION(BlueprintCallable, Category = "SharedImage")
+    static FVector4f GetPixelValue(const FSharedImageConstRefBlueprint& Image, int32 X, int32 Y, bool& bValid);
+};
