When physics replication sets initial component transform, it's actually setting the root bodyinstance transform. Need to pull out the relationship between these

diff --git a/Engine/Source/Runtime/Engine/Private/PhysicsEngine/PhysicsReplication.cpp b/Engine/Source/Runtime/Engine/Private/PhysicsEngine/PhysicsReplication.cpp
--- a/Engine/Source/Runtime/Engine/Private/PhysicsEngine/PhysicsReplication.cpp
+++ b/Engine/Source/Runtime/Engine/Private/PhysicsEngine/PhysicsReplication.cpp
@@ -177,7 +177,21 @@ void FPhysicsReplication::SetReplicatedTarget(UPrimitiveComponent* Component, FN
 			// which does not get updated until the component's transform is directly set.
 			// Until that flag is set, it's root particle will be in a disabled state and not
 			// have any children, therefore replication will be dead in the water.
-			Component->SetWorldTransform(FTransform(ReplicatedTarget.Quaternion, ReplicatedTarget.Position, Component->GetRelativeScale3D()), false, nullptr, ETeleportType::TeleportPhysics);
+			//
+			// NOTE: The incoming transform is set from the bodyinstance transform, not from
+			// the primitive component transform. For staticmeshcomponent this is fine, but for
+			// skeletalmeshcomponent there can be a relative transform between the component
+			// and its root body (RootBodyData.TransformToRoot).
+			//
+			// We need to set the primitive component's transform instead of just directly
+			// setting the root bodyinstance's transform because that wouldn't call the proper
+			// primitivecomponent callbacks, which are used by clusterunion, and which is the
+			// entire reason for having this in the first place.
+			const FTransform InitialBodyTM = FTransform(ReplicatedTarget.Quaternion, ReplicatedTarget.Position);
+			USkeletalMeshComponent* SkelMeshComp = Cast<USkeletalMeshComponent>(Component);
+			const FTransform InitialCompTM = SkelMeshComp ? SkelMeshComp->RootBodyData.TransformToRoot * InitialBodyTM : InitialBodyTM;
+
+			Component->SetWorldTransform(InitialCompTM, false, nullptr, ETeleportType::TeleportPhysics);
 			Component->SetPhysicsLinearVelocity(ReplicatedTarget.LinVel);
 			Component->SetAllPhysicsAngularVelocityInDegrees(ReplicatedTarget.AngVel);
 		}
