Truncation warning fixes for FADPCMAudioInfo code in the Engine module.

#rnx

- Use UE_PTRDIFF_TO_INT32 to show that we expect the difference between the two pointers to fit in a int32 value.

diff --git a/Engine/Source/Runtime/Engine/Private/ADPCMAudioInfo.cpp b/Engine/Source/Runtime/Engine/Private/ADPCMAudioInfo.cpp
--- a/Engine/Source/Runtime/Engine/Private/ADPCMAudioInfo.cpp
+++ b/Engine/Source/Runtime/Engine/Private/ADPCMAudioInfo.cpp
@@ -501,7 +501,7 @@ bool FADPCMAudioInfo::StreamCompressedInfoInternal(const FSoundWaveProxyPtr& InW
 	}
 	
 	// if we only included the header in the zeroth chunk, skip to the next chunk.
-	const int32 SampleDataOffset = WaveInfo.SampleDataStart - ChunkData;
+	const int32 SampleDataOffset = UE_PTRDIFF_TO_INT32(WaveInfo.SampleDataStart - ChunkData);
 	check(SampleDataOffset > 0);
 	if (((uint32)SampleDataOffset) >= CurrentChunkDataSize)
 	{
@@ -512,7 +512,7 @@ bool FADPCMAudioInfo::StreamCompressedInfoInternal(const FSoundWaveProxyPtr& InW
 	}
 	else
 	{
-		FirstChunkSampleDataOffset = WaveInfo.SampleDataStart - ChunkData;
+		FirstChunkSampleDataOffset = UE_PTRDIFF_TO_INT32(WaveInfo.SampleDataStart - ChunkData);
 		FirstChunkSampleDataIndex = 0;
 	}
 
