Fix ActorComponent not returning correct asset user data array
#rnx

diff --git a/Engine/Source/Runtime/Engine/Classes/Components/ActorComponent.h b/Engine/Source/Runtime/Engine/Classes/Components/ActorComponent.h
--- a/Engine/Source/Runtime/Engine/Classes/Components/ActorComponent.h
+++ b/Engine/Source/Runtime/Engine/Classes/Components/ActorComponent.h
@@ -1043,6 +1043,7 @@ public:
 	ENGINE_API virtual void AddAssetUserData(UAssetUserData* InUserData) override;
 	ENGINE_API virtual void RemoveUserDataOfClass(TSubclassOf<UAssetUserData> InUserDataClass) override;
 	ENGINE_API virtual UAssetUserData* GetAssetUserDataOfClass(TSubclassOf<UAssetUserData> InUserDataClass) override;
+	ENGINE_API virtual const TArray<UAssetUserData*>* GetAssetUserDataArray() const override;
 	//~ End IInterface_AssetUserData Interface
 
 	//~ Begin IInterface_ActorSubobject Interface
diff --git a/Engine/Source/Runtime/Engine/Private/Components/ActorComponent.cpp b/Engine/Source/Runtime/Engine/Private/Components/ActorComponent.cpp
--- a/Engine/Source/Runtime/Engine/Private/Components/ActorComponent.cpp
+++ b/Engine/Source/Runtime/Engine/Private/Components/ActorComponent.cpp
@@ -2038,6 +2038,11 @@ UAssetUserData* UActorComponent::GetAssetUserDataOfClass(TSubclassOf<UAssetUserD
 	return NULL;
 }
 
+const TArray<UAssetUserData*>* UActorComponent::GetAssetUserDataArray() const
+{
+	return &ToRawPtrTArrayUnsafe(AssetUserData);
+}
+
 void UActorComponent::RemoveUserDataOfClass(TSubclassOf<UAssetUserData> InUserDataClass)
 {
 	for (int32 DataIdx = 0; DataIdx < AssetUserData.Num(); DataIdx++)
