Fix UInstancedStaticMeshComponent::GetNavigationPerInstanceTransforms to return only the instance transforms in the provided area box as expected.
This represents a significant performance improvement when using multiple instances ISM with AI navigation.

diff --git a/Engine/Source/Runtime/Engine/Private/InstancedStaticMesh.cpp b/Engine/Source/Runtime/Engine/Private/InstancedStaticMesh.cpp
--- a/Engine/Source/Runtime/Engine/Private/InstancedStaticMesh.cpp
+++ b/Engine/Source/Runtime/Engine/Private/InstancedStaticMesh.cpp
@@ -5297,13 +5297,22 @@ bool UInstancedStaticMeshComponent::IsNavigationRelevant() const
 
 void UInstancedStaticMeshComponent::GetNavigationPerInstanceTransforms(const FBox& AreaBox, TArray<FTransform>& InstanceData) const
 {
-	for (const auto& InstancedData : PerInstanceSMData)
+	if (GetStaticMesh())
 	{
-		//TODO: Is it worth doing per instance bounds check here ?
-		const FTransform InstanceToComponent(InstancedData.Transform);
-		if (!InstanceToComponent.GetScale3D().IsZero())
+		const FBox LocalAreaBox =  AreaBox.InverseTransformBy(GetComponentTransform());
+		const FBoxSphereBounds RenderBounds = GetStaticMesh()->GetBounds();
+
+		for (const auto& InstancedData : PerInstanceSMData)
 		{
-			InstanceData.Add(InstanceToComponent*GetComponentTransform());
+			const FTransform InstanceToComponent(InstancedData.Transform);
+			if (!InstanceToComponent.GetScale3D().IsZero())
+			{
+				const FBoxSphereBounds TransformedInstanceBounds = RenderBounds.TransformBy(InstancedData.Transform);
+				if (LocalAreaBox.Intersect(TransformedInstanceBounds.GetBox()))
+				{
+					InstanceData.Add(InstanceToComponent*GetComponentTransform());
+				}
+			}
 		}
 	}
 }
