[NetPhysics] Predictive Interpolation fixes

diff --git a/Engine/Source/Runtime/Engine/Classes/Components/SceneComponent.h b/Engine/Source/Runtime/Engine/Classes/Components/SceneComponent.h
--- a/Engine/Source/Runtime/Engine/Classes/Components/SceneComponent.h
+++ b/Engine/Source/Runtime/Engine/Classes/Components/SceneComponent.h
@@ -1354,6 +1354,13 @@ private:
 
 public:
 
+	/**
+	* Called when client receive replication data, before replication is performed.
+	* Can be overridden in derived components to make use of replication data locally.
+	* Note that replication still applies when overriding this, it's not intended to replace replication.
+	*/
+	ENGINE_API virtual void OnReceiveReplicatedState(const FVector X, const FQuat R, const FVector V, const FVector W) {};
+
 	/**
 	 * Gets the property name for RelativeLocation.
 	 *
diff --git a/Engine/Source/Runtime/Engine/Classes/PhysicsEngine/ClusterUnionComponent.h b/Engine/Source/Runtime/Engine/Classes/PhysicsEngine/ClusterUnionComponent.h
--- a/Engine/Source/Runtime/Engine/Classes/PhysicsEngine/ClusterUnionComponent.h
+++ b/Engine/Source/Runtime/Engine/Classes/PhysicsEngine/ClusterUnionComponent.h
@@ -290,6 +290,7 @@ public:
 
 	//~ Begin USceneComponent Interface
 public:
+	ENGINE_API virtual void OnReceiveReplicatedState(const FVector X, const FQuat R, const FVector V, const FVector W) override;
 	ENGINE_API virtual void OnUpdateTransform(EUpdateTransformFlags UpdateTransformFlags, ETeleportType Teleport) override;
 	ENGINE_API virtual FBoxSphereBounds CalcBounds(const FTransform& LocalToWorld) const override;
 	//~ End USceneComponent Interface
diff --git a/Engine/Source/Runtime/Engine/Classes/PhysicsEngine/PhysicsSettings.h b/Engine/Source/Runtime/Engine/Classes/PhysicsEngine/PhysicsSettings.h
--- a/Engine/Source/Runtime/Engine/Classes/PhysicsEngine/PhysicsSettings.h
+++ b/Engine/Source/Runtime/Engine/Classes/PhysicsEngine/PhysicsSettings.h
@@ -258,7 +258,7 @@ public:
 	UFUNCTION(BlueprintCallable, Category = "Physics")
 	int32 GetPhysicsHistoryCount() const
 	{
-		return (int32) ((float) (FMath::CeilToInt(1.f / AsyncFixedTimeStepSize)) * (PhysicsPrediction.MaxSupportedLatencyPrediction / 1000.f));
+		return FMath::CeilToInt(0.001f * PhysicsPrediction.MaxSupportedLatencyPrediction / AsyncFixedTimeStepSize);
 	}
 
 	ENGINE_API virtual void PostInitProperties() override;
diff --git a/Engine/Source/Runtime/Engine/Private/ActorReplication.cpp b/Engine/Source/Runtime/Engine/Private/ActorReplication.cpp
--- a/Engine/Source/Runtime/Engine/Private/ActorReplication.cpp
+++ b/Engine/Source/Runtime/Engine/Private/ActorReplication.cpp
@@ -205,6 +205,12 @@ void AActor::OnRep_ReplicatedMovement()
 			SyncReplicatedPhysicsSimulation();
 		}
 
+		// NOTE: This is only needed because ClusterUnion has a flag bHasReceivedTransform
+		// which does not get updated until the component's transform is directly set.
+		// Until that flag is set, its root particle will be in a disabled state and not
+		// have any children, therefore replication will be dead in the water.
+		RootComponent->OnReceiveReplicatedState(LocalRepMovement.Location, LocalRepMovement.Rotation.Quaternion(), LocalRepMovement.LinearVelocity, LocalRepMovement.AngularVelocity);
+
 		if (LocalRepMovement.bRepPhysics)
 		{
 			// Sync physics state
diff --git a/Engine/Source/Runtime/Engine/Private/PhysicsEngine/ClusterUnionComponent.cpp b/Engine/Source/Runtime/Engine/Private/PhysicsEngine/ClusterUnionComponent.cpp
--- a/Engine/Source/Runtime/Engine/Private/PhysicsEngine/ClusterUnionComponent.cpp
+++ b/Engine/Source/Runtime/Engine/Private/PhysicsEngine/ClusterUnionComponent.cpp
@@ -534,6 +534,22 @@ void UClusterUnionComponent::OnDestroyPhysicsState()
 	AccelerationStructure.Reset();
 }
 
+void UClusterUnionComponent::OnReceiveReplicatedState(const FVector X, const FQuat R, const FVector V, const FVector W)
+{
+	if (!bHasReceivedTransform)
+	{
+		// First time, we just directly set component state
+		//
+		// NOTE: This is only needed because ClusterUnion has a flag bHasReceivedTransform
+		// which does not get updated until the component's transform is directly set.
+		// Until that flag is set, it's root particle will be in a disabled state and not
+		// have any children, therefore replication will be dead in the water.
+		SetWorldTransform(FTransform(R, X, GetRelativeScale3D()), false, nullptr, ETeleportType::TeleportPhysics);
+		SetPhysicsLinearVelocity(V);
+		SetAllPhysicsAngularVelocityInDegrees(W);
+	}
+}
+
 void UClusterUnionComponent::OnUpdateTransform(EUpdateTransformFlags UpdateTransformFlags, ETeleportType Teleport)
 {
 	USceneComponent::OnUpdateTransform(UpdateTransformFlags, Teleport);
diff --git a/Engine/Source/Runtime/Engine/Private/PhysicsEngine/PhysicsReplication.cpp b/Engine/Source/Runtime/Engine/Private/PhysicsEngine/PhysicsReplication.cpp
--- a/Engine/Source/Runtime/Engine/Private/PhysicsEngine/PhysicsReplication.cpp
+++ b/Engine/Source/Runtime/Engine/Private/PhysicsEngine/PhysicsReplication.cpp
@@ -73,9 +73,6 @@ namespace CharacterMovementCVars
 
 	int32 ApplyAsyncSleepState = 1;
 	static FAutoConsoleVariableRef CVarApplyAsyncSleepState(TEXT("p.ApplyAsyncSleepState"), ApplyAsyncSleepState, TEXT(""));
-
-	bool bPredictiveInterpolationAlwaysHardSnap = true;
-	static FAutoConsoleVariableRef CVarPredictiveInterpolationAlwaysHardSnap(TEXT("p.PredictiveInterpolation.AlwaysHardSnap"), bPredictiveInterpolationAlwaysHardSnap, TEXT("When true, predictive interpolation replication mode will always hard snap. Used as a backup measure"));
 }
 
 namespace PhysicsReplicationCVars
@@ -92,8 +89,11 @@ namespace PhysicsReplicationCVars
 
 	namespace PredictiveInterpolationCVars
 	{
-		static float PosCorrectionTimeMultiplier = 0.5f;
-		static FAutoConsoleVariableRef CVarPosCorrectionTimeMultiplier(TEXT("np2.PredictiveInterpolation.PosCorrectionTimeMultiplier"), PosCorrectionTimeMultiplier, TEXT("Multiplier to adjust the time to correct positional offset over, which is based on the clients forward predicted time ahead of the server."));
+		static float PosCorrectionTimeBase = 0.15f;
+		static FAutoConsoleVariableRef CVarPosCorrectionTimeBase(TEXT("np2.PredictiveInterpolation.PosCorrectionTimeBase"), PosCorrectionTimeBase, TEXT("Base time to correct positional offset over. RTT * PosCorrectionTimeMultiplier are added on top of this."));
+
+		static float PosCorrectionTimeMultiplier = 1.0f;
+		static FAutoConsoleVariableRef CVarPosCorrectionTimeMultiplier(TEXT("np2.PredictiveInterpolation.PosCorrectionTimeMultiplier"), PosCorrectionTimeMultiplier, TEXT("Multiplier to adjust how much of RTT (network Round Trip Time) to add to positional offset correction."));
 
 		static float InterpolationTimeMultiplier = 1.1f;
 		static FAutoConsoleVariableRef CVarInterpolationTimeMultiplier(TEXT("np2.PredictiveInterpolation.InterpolationTimeMultiplier"), InterpolationTimeMultiplier, TEXT("Multiplier to adjust the replication interpolation time which is based on the sendrate of replication data from the server."));
@@ -104,11 +104,32 @@ namespace PhysicsReplicationCVars
 		static float MinExpectedDistanceCovered = 0.25f;
 		static FAutoConsoleVariableRef CVarMinExpectedDistanceCovered(TEXT("np2.PredictiveInterpolation.MinExpectedDistanceCovered"), MinExpectedDistanceCovered, TEXT("Value in percentage where 0.25 = 25%. How much of the expected distance based on replication velocity should the object have covered in a simulation tick to Not be considered stuck."));
 
-		static float MaxDistanceToSleepSqr = 8.f;
-		static FAutoConsoleVariableRef CVarMaxDistanceToSleepSqr(TEXT("np2.PredictiveInterpolation.MaxDistanceToSleepSqr"), MaxDistanceToSleepSqr, TEXT("Squared value. Max distance from the source target to allow the object to go to sleep."));
+		static float ErrorAccumulationDecreaseMultiplier = 0.5f;
+		static FAutoConsoleVariableRef CVarErrorAccumulationDecreaseMultiplier(TEXT("np2.PredictiveInterpolation.ErrorAccumulationDecreaseMultiplier"), ErrorAccumulationDecreaseMultiplier, TEXT("Multiplier to adjust how fast we decrease accumulated error time when we no longer accumulate error."));
+
+		static float ErrorAccumulationSeconds = 0.5f;
+		static FAutoConsoleVariableRef CVarErrorAccumulationSeconds(TEXT("np2.PredictiveInterpolation.ErrorAccumulationSeconds"), ErrorAccumulationSeconds, TEXT("Perform a reposition if replication have not been able to cover the min expected distance towards the target for this amount of time."));
+
+		static float EarlyOutDistanceSqr = 2.f;
+		static FAutoConsoleVariableRef CVarEarlyOutDistanceSqr(TEXT("np2.PredictiveInterpolation.EarlyOutDistanceSqr"), EarlyOutDistanceSqr, TEXT("Squared value. If object is within this distance from the source target, early out from replication and apply sleep if replicated."));
+		
+		static float EarlyOutAngle = 1.f;
+		static FAutoConsoleVariableRef CVarEarlyOutAngle(TEXT("np2.PredictiveInterpolation.EarlyOutAngle"), EarlyOutAngle, TEXT("If object is within this rotational angle (in degrees) from the source target, early out from replication and apply sleep if replicated."));
 		
 		static bool PostResimWaitForUpdate = true;
 		static FAutoConsoleVariableRef CVarPostResimWaitForUpdate(TEXT("np2.PredictiveInterpolation.PostResimWaitForUpdate"), PostResimWaitForUpdate, TEXT("After a resimulation, wait for replicated states that correspond to post-resim state before processing replication again."));
+		
+		static bool bVelocityBased = true;
+		static FAutoConsoleVariableRef CVarVelocityBased(TEXT("np2.PredictiveInterpolation.VelocityBased"), bVelocityBased, TEXT("When true, predictive interpolation replication mode will only apply linear velocity and angular velocity"));
+	
+		static bool bAlwaysHardSnap = false;
+		static FAutoConsoleVariableRef CVarAlwaysHardSnap(TEXT("np2.PredictiveInterpolation.AlwaysHardSnap"), bAlwaysHardSnap, TEXT("When true, predictive interpolation replication mode will always hard snap. Used as a backup measure"));
+
+		static bool bSkipReplication = false;
+		static FAutoConsoleVariableRef CVarSkipReplication(TEXT("np2.PredictiveInterpolation.SkipReplication"), bSkipReplication, TEXT("When true, predictive interpolation is not applied anymore letting the object simulate freely instead"));
+
+		static bool bDontClearTarget = false;
+		static FAutoConsoleVariableRef CVarDontClearTarget(TEXT("np2.PredictiveInterpolation.DontClearTarget"), bDontClearTarget, TEXT("When true, predictive interpolation will not lose track of the last replicated state after coming to rest."));
 	}
 
 }
@@ -143,8 +164,7 @@ void FPhysicsReplication::SetReplicatedTarget(UPrimitiveComponent* Component, FN
 {
 	// If networked physics prediction is enabled, enforce the new physics replication flow via SetReplicatedTarget() using PhysicsObject instead of BodyInstance from BoneName.
 	AActor* Owner = Component->GetOwner();
-	if (UPhysicsSettings::Get()->PhysicsPrediction.bEnablePhysicsPrediction && Owner &&
-		(PhysicsReplicationCVars::EnableDefaultReplication || Owner->GetPhysicsReplicationMode() != EPhysicsReplicationMode::Default)) // For now, only opt in to the PhysicsObject flow if not using Default replication or if default is allowed via CVar.
+	if (Owner && (PhysicsReplicationCVars::EnableDefaultReplication || Owner->GetPhysicsReplicationMode() != EPhysicsReplicationMode::Default)) // For now, only opt in to the PhysicsObject flow if not using Default replication or if default is allowed via CVar.
 	{
 		const ENetRole OwnerRole = Owner->GetLocalRole();
 		const bool bIsSimulated = OwnerRole == ROLE_SimulatedProxy;
@@ -170,30 +190,6 @@ void FPhysicsReplication::SetReplicatedTarget(UPrimitiveComponent* Component, FN
 			Target = &ComponentToTargets_DEPRECATED.Add(TargetKey);
 			Target->PrevPos = ReplicatedTarget.Position;
 			Target->PrevPosTarget = ReplicatedTarget.Position;
-
-			// First time, we just directly set component state
-			//
-			// NOTE: This is only needed because ClusterUnion has a flag bHasReceivedTransform
-			// which does not get updated until the component's transform is directly set.
-			// Until that flag is set, it's root particle will be in a disabled state and not
-			// have any children, therefore replication will be dead in the water.
-			//
-			// NOTE: The incoming transform is set from the bodyinstance transform, not from
-			// the primitive component transform. For staticmeshcomponent this is fine, but for
-			// skeletalmeshcomponent there can be a relative transform between the component
-			// and its root body (RootBodyData.TransformToRoot).
-			//
-			// We need to set the primitive component's transform instead of just directly
-			// setting the root bodyinstance's transform because that wouldn't call the proper
-			// primitivecomponent callbacks, which are used by clusterunion, and which is the
-			// entire reason for having this in the first place.
-			const FTransform InitialBodyTM = FTransform(ReplicatedTarget.Quaternion, ReplicatedTarget.Position);
-			USkeletalMeshComponent* SkelMeshComp = Cast<USkeletalMeshComponent>(Component);
-			const FTransform InitialCompTM = SkelMeshComp ? SkelMeshComp->RootBodyData.TransformToRoot * InitialBodyTM : InitialBodyTM;
-
-			Component->SetWorldTransform(InitialCompTM, false, nullptr, ETeleportType::TeleportPhysics);
-			Component->SetPhysicsLinearVelocity(ReplicatedTarget.LinVel);
-			Component->SetAllPhysicsAngularVelocityInDegrees(ReplicatedTarget.AngVel);
 		}
 
 		Target->ServerFrame = ServerFrame;
@@ -696,7 +692,7 @@ void FPhysicsReplicationAsync::OnPreSimulate_Internal()
 		for (const FPhysicsRepAsyncInputData& Input : AsyncInput->InputData)
 		{
 			UpdateRewindDataTarget(Input);
-			UpdateAsyncTarget(Input);
+			UpdateAsyncTarget(Input, RigidsSolver);
 		}
 
 		ApplyTargetStatesAsync(GetDeltaTime_Internal(), AsyncInput->ErrorCorrection, AsyncInput->InputData);
@@ -733,7 +729,7 @@ void FPhysicsReplicationAsync::UpdateRewindDataTarget(const FPhysicsRepAsyncInpu
 	}
 }
 
-void FPhysicsReplicationAsync::UpdateAsyncTarget(const FPhysicsRepAsyncInputData& Input)
+void FPhysicsReplicationAsync::UpdateAsyncTarget(const FPhysicsRepAsyncInputData& Input, Chaos::FPBDRigidsSolver* RigidsSolver)
 {
 	if (Input.PhysicsObject == nullptr)
 	{
@@ -753,10 +749,12 @@ void FPhysicsReplicationAsync::UpdateAsyncTarget(const FPhysicsRepAsyncInputData
 		Target->PrevLinVel = Input.TargetState.LinVel;
 	}
 
-	if (Input.ServerFrame > Target->ServerFrame)
+	if (Input.ServerFrame >= Target->ServerFrame)
 	{
-		Target->PrevServerFrame = Target->ServerFrame;
+		Target->PrevServerFrame = (Target->ServerFrame == INDEX_NONE) ? (Input.ServerFrame - 1) : Target->ServerFrame;
 		Target->ServerFrame = Input.ServerFrame;
+		Target->PrevReceiveFrame = (Target->ReceiveFrame == INDEX_NONE) ? (RigidsSolver->GetCurrentFrame() - 1) : Target->ReceiveFrame;
+		Target->ReceiveFrame = RigidsSolver->GetCurrentFrame();
 		Target->TargetState = Input.TargetState;
 		Target->RepMode = Input.RepMode;
 		Target->FrameOffset = Input.FrameOffset;
@@ -828,6 +826,7 @@ void FPhysicsReplicationAsync::ApplyTargetStatesAsync(const float DeltaSeconds,
 							bRemoveItr = ResimulationReplication(RigidHandle, Target, DeltaSeconds);
 							break;
 					}
+					Target.TickCount++;
 				}
 			}
 		}
@@ -1135,6 +1134,11 @@ bool FPhysicsReplicationAsync::DefaultReplication(Chaos::FPBDRigidParticleHandle
 */
 bool FPhysicsReplicationAsync::PredictiveInterpolation(Chaos::FPBDRigidParticleHandle* Handle, FReplicatedPhysicsTargetAsync& Target, const float DeltaSeconds)
 {
+	if (PhysicsReplicationCVars::PredictiveInterpolationCVars::bSkipReplication)
+	{
+		return true;
+	}
+
 	if (Target.bWaiting)
 	{
 		return false;
@@ -1146,19 +1150,40 @@ bool FPhysicsReplicationAsync::PredictiveInterpolation(Chaos::FPBDRigidParticleH
 		return true;
 	}
 
-	const float ErrorAccumulationSeconds = CharacterMovementCVars::ErrorAccumulationSeconds >= 0.0f ? CharacterMovementCVars::ErrorAccumulationSeconds : ErrorCorrectionDefault.ErrorAccumulationSeconds;
-	const float MaxRestoredStateErrorSqr = CharacterMovementCVars::MaxRestoredStateError >= 0.0f ? 
-		(CharacterMovementCVars::MaxRestoredStateError * CharacterMovementCVars::MaxRestoredStateError) :
-		(ErrorCorrectionDefault.MaxRestoredStateError * ErrorCorrectionDefault.MaxRestoredStateError);
-
+	// Get the distance from the current position to the source position of our target state
+	const float SourceDistanceSqr = (Target.PrevPosTarget - Handle->X()).SizeSquared();
 	const bool bShouldSleep = (Target.TargetState.Flags & ERigidBodyFlags::Sleeping) != 0;
-	const int32 LocalFrame = Target.ServerFrame - Target.FrameOffset;
-	const int32 NumPredictedFrames = RigidsSolver->GetCurrentFrame() - LocalFrame - Target.TickCount;
-	const float PredictedTime = DeltaSeconds * NumPredictedFrames;
-	const float SendRate = (Target.ServerFrame - Target.PrevServerFrame) * DeltaSeconds;
+	
+	// Early out if we are within range of target, also apply target sleep state
+	if (SourceDistanceSqr < PhysicsReplicationCVars::PredictiveInterpolationCVars::EarlyOutDistanceSqr)
+	{
+		// Get the rotational offset between the blended rotation target and the current rotation
+		const FQuat TargetRotDelta = Target.TargetState.Quaternion * Handle->R().Inverse();
+
+		// Convert to angle axis
+		float Angle;
+		FVector Axis;
+		TargetRotDelta.ToAxisAndAngle(Axis, Angle);
+		if (Angle < FMath::DegreesToRadians(PhysicsReplicationCVars::PredictiveInterpolationCVars::EarlyOutAngle))
+		{
+			Handle->SetV(FVector(0, 0, 0));
+			Handle->SetW(FVector(0, 0, 0));
+			Target.PrevLinVel = FVector(FVector(0, 0, 0));
 
-	const float PosCorrectionTime = PredictedTime * PhysicsReplicationCVars::PredictiveInterpolationCVars::PosCorrectionTimeMultiplier;
-	const float InterpolationTime = SendRate * PhysicsReplicationCVars::PredictiveInterpolationCVars::InterpolationTimeMultiplier;
+			if (bShouldSleep && !Handle->IsKinematic())
+			{
+				RigidsSolver->GetEvolution()->SetParticleObjectState(Handle, Chaos::EObjectStateType::Sleeping);
+			}
+
+			return bShouldSleep && !PhysicsReplicationCVars::PredictiveInterpolationCVars::bDontClearTarget;
+		}
+	}
+
+	const float RTT = LatencyOneWay * 2.f;
+	const float PosCorrectionTime = PhysicsReplicationCVars::PredictiveInterpolationCVars::PosCorrectionTimeBase + RTT * PhysicsReplicationCVars::PredictiveInterpolationCVars::PosCorrectionTimeMultiplier;
+	const float ReceiveInterval = Target.ReceiveFrame - Target.PrevReceiveFrame;
+	const float InterpolationTime = ReceiveInterval * DeltaSeconds * PhysicsReplicationCVars::PredictiveInterpolationCVars::InterpolationTimeMultiplier;
+	bool bSoftSnap = !PhysicsReplicationCVars::PredictiveInterpolationCVars::bVelocityBased;
 
 	// CurrentState
 	FRigidBodyState CurrentState;
@@ -1173,7 +1198,6 @@ bool FPhysicsReplicationAsync::PredictiveInterpolation(Chaos::FPBDRigidParticleH
 	const FVector TargetLinVel = Target.TargetState.LinVel;
 	const FVector TargetAngVel = Target.TargetState.AngVel;
 
-
 	/** --- Reconciliation ---
 	* Get the traveled direction and distance from previous frame and compare with replicated linear velocity.
 	* If the object isn't moving enough along the replicated velocity it's considered stuck and needs a hard reconciliation.
@@ -1183,17 +1207,19 @@ bool FPhysicsReplicationAsync::PredictiveInterpolation(Chaos::FPBDRigidParticleH
 	const float CoveredDistance = FVector::DotProduct(PrevDiff, Target.PrevLinVel.GetSafeNormal());
 	
 	// If the object is moving less than X% of the expected distance, accumulate error seconds
-	if (CoveredDistance / ExpectedDistance < PhysicsReplicationCVars::PredictiveInterpolationCVars::MinExpectedDistanceCovered)
+	if (ExpectedDistance > UE_SMALL_NUMBER && (CoveredDistance / ExpectedDistance) < PhysicsReplicationCVars::PredictiveInterpolationCVars::MinExpectedDistanceCovered)
 	{
 		Target.AccumulatedErrorSeconds += DeltaSeconds;
+		bSoftSnap = true;
 	}
-	else
+	else if (Target.AccumulatedErrorSeconds > 0.f)
 	{
-		Target.AccumulatedErrorSeconds = FMath::Max(Target.AccumulatedErrorSeconds - DeltaSeconds, 0.0f);
+		const float DecreaseTime = DeltaSeconds * PhysicsReplicationCVars::PredictiveInterpolationCVars::ErrorAccumulationDecreaseMultiplier;
+		Target.AccumulatedErrorSeconds = FMath::Max(Target.AccumulatedErrorSeconds - DecreaseTime, 0.0f);
+		bSoftSnap = true;
 	}
 
-	const bool bHardSnap = Target.AccumulatedErrorSeconds > ErrorAccumulationSeconds || CharacterMovementCVars::bPredictiveInterpolationAlwaysHardSnap;
-	bool bClearTarget = bHardSnap;
+	const bool bHardSnap = Target.AccumulatedErrorSeconds > PhysicsReplicationCVars::PredictiveInterpolationCVars::ErrorAccumulationSeconds || PhysicsReplicationCVars::PredictiveInterpolationCVars::bAlwaysHardSnap;
 	if (bHardSnap)
 	{
 		// Too much error so just snap state here and be done with it
@@ -1204,6 +1230,9 @@ bool FPhysicsReplicationAsync::PredictiveInterpolation(Chaos::FPBDRigidParticleH
 		Handle->SetQ(Target.PrevRotTarget);
 		Handle->SetV(Target.TargetState.LinVel);
 		Handle->SetW(Target.TargetState.AngVel);
+
+		// Cache data for next replication
+		Target.PrevLinVel = FVector(Target.TargetState.LinVel);
 	}
 	else // Velocity-based Replication
 	{
@@ -1249,51 +1278,38 @@ bool FPhysicsReplicationAsync::PredictiveInterpolation(Chaos::FPBDRigidParticleH
 		const FVector RepAngVel = WAxis * (WAngle / DeltaSeconds);
 
 
-		// Apply velocity
+		if (bSoftSnap)
+		{
+			Handle->SetX(CurrentState.Position + (RepLinVel * DeltaSeconds));
+			Handle->SetR(TargetRotBlended);
+		}
+
 		Handle->SetV(RepLinVel);
 		Handle->SetW(RepAngVel);
 
-
-		// Cache data for reconciliation
-		Target.PrevPos = FVector(CurrentState.Position);
+		// Cache data for next replication
 		Target.PrevLinVel = FVector(RepLinVel);
 	}
 
+	// Cache data for next replication
+	Target.PrevPos = FVector(CurrentState.Position);
 
-	if (bShouldSleep)
+	// --- Target Extrapolation ---
+	if (Target.TickCount < FMath::CeilToInt(ReceiveInterval * PhysicsReplicationCVars::PredictiveInterpolationCVars::ExtrapolationTimeMultiplier))
 	{
-		// --- Sleep ---
-		// Get the distance from the current position to the source position of our target state
-		const float SourceDistanceSqr = (Target.PrevPosTarget - CurrentState.Position).SizeSquared();
-		
-		// Don't allow kinematic to sleeping transition
-		if (SourceDistanceSqr < PhysicsReplicationCVars::PredictiveInterpolationCVars::MaxDistanceToSleepSqr && !Handle->IsKinematic())
-		{
-			RigidsSolver->GetEvolution()->SetParticleObjectState(Handle, Chaos::EObjectStateType::Sleeping);
-			bClearTarget = true;
-		}
-	}
-	else
-	{
-		// --- Target Extrapolation ---
-		if ((Target.TickCount * DeltaSeconds) < SendRate * PhysicsReplicationCVars::PredictiveInterpolationCVars::ExtrapolationTimeMultiplier)
-		{
-			// Extrapolate target position
-			Target.TargetState.Position = Target.TargetState.Position + Target.TargetState.LinVel * DeltaSeconds;
-
-			// Extrapolate target rotation
-			float TargetAngVelSize;
-			FVector TargetAngVelAxis;
-			Target.TargetState.AngVel.FVector::ToDirectionAndLength(TargetAngVelAxis, TargetAngVelSize);
-			TargetAngVelSize = FMath::DegreesToRadians(TargetAngVelSize);
-			const FQuat TargetRotExtrapDelta = FQuat(TargetAngVelAxis, TargetAngVelSize * DeltaSeconds);
-			Target.TargetState.Quaternion = TargetRotExtrapDelta * Target.TargetState.Quaternion;
-		}
+		// Extrapolate target position
+		Target.TargetState.Position = Target.TargetState.Position + Target.TargetState.LinVel * DeltaSeconds;
+
+		// Extrapolate target rotation
+		float TargetAngVelSize;
+		FVector TargetAngVelAxis;
+		Target.TargetState.AngVel.FVector::ToDirectionAndLength(TargetAngVelAxis, TargetAngVelSize);
+		TargetAngVelSize = FMath::DegreesToRadians(TargetAngVelSize);
+		const FQuat TargetRotExtrapDelta = FQuat(TargetAngVelAxis, TargetAngVelSize * DeltaSeconds);
+		Target.TargetState.Quaternion = TargetRotExtrapDelta * Target.TargetState.Quaternion;
 	}
-
-	Target.TickCount++;
-
-	return bClearTarget;;
+	
+	return false;
 }
 
 /** Compare states and trigger resimulation if needed */
diff --git a/Engine/Source/Runtime/Engine/Public/PhysicsReplication.h b/Engine/Source/Runtime/Engine/Public/PhysicsReplication.h
--- a/Engine/Source/Runtime/Engine/Public/PhysicsReplication.h
+++ b/Engine/Source/Runtime/Engine/Public/PhysicsReplication.h
@@ -92,7 +92,8 @@ struct FReplicatedPhysicsTargetAsync
 {
 	FReplicatedPhysicsTargetAsync(Chaos::FConstPhysicsObjectHandle POHandle)
 		: AccumulatedErrorSeconds(0.0f)
-		, ServerFrame(0)
+		, ServerFrame(INDEX_NONE)
+		, ReceiveFrame(INDEX_NONE)
 		, PhysicsObject(POHandle)
 	{ }
 
@@ -111,6 +112,9 @@ struct FReplicatedPhysicsTargetAsync
 	/** The frame offset between local client and server */
 	int32 FrameOffset;
 
+	/** The local client frame when receiving this target from the server */
+	int32 ReceiveFrame;
+
 	/** Index of physics object on component */
 	Chaos::FConstPhysicsObjectHandle PhysicsObject;
 
@@ -123,6 +127,7 @@ struct FReplicatedPhysicsTargetAsync
 	FVector PrevPos;
 	FVector PrevLinVel;
 	int32 PrevServerFrame;
+	int32 PrevReceiveFrame;
 
 	/** If this target is waiting for up-to-date data? */
 	bool bWaiting;
@@ -149,7 +154,7 @@ private:
 	TMap<Chaos::FConstPhysicsObjectHandle, FReplicatedPhysicsTargetAsync> ObjectToTarget;
 
 private:
-	void UpdateAsyncTarget(const FPhysicsRepAsyncInputData& Input);
+	void UpdateAsyncTarget(const FPhysicsRepAsyncInputData& Input, Chaos::FPBDRigidsSolver* RigidsSolver);
 	void UpdateRewindDataTarget(const FPhysicsRepAsyncInputData& Input);
 
 public:
