Fix for shader validation crash when trying to stream in empty SVT frames.

#rnx

diff --git a/Engine/Source/Runtime/Engine/Private/SparseVolumeTexture/SparseVolumeTextureStreamingManager.cpp b/Engine/Source/Runtime/Engine/Private/SparseVolumeTexture/SparseVolumeTextureStreamingManager.cpp
--- a/Engine/Source/Runtime/Engine/Private/SparseVolumeTexture/SparseVolumeTextureStreamingManager.cpp
+++ b/Engine/Source/Runtime/Engine/Private/SparseVolumeTexture/SparseVolumeTextureStreamingManager.cpp
@@ -455,6 +455,9 @@ public:
 				FRDGTexture* DummyTexture = GraphBuilder.CreateTexture(FRDGTextureDesc::Create3D(FIntVector(1), PF_R8, FClearValueBinding::None, ETextureCreateFlags::UAV), TEXT("SparseVolumeTexture.DummyTexture"));
 				FRDGTexture* DstTextureARDG = DstTextureA ? GraphBuilder.RegisterExternalTexture(CreateRenderTarget(DstTextureA, TEXT("SparseVolumeTexture.TileDataTextureA"))) : nullptr;
 				FRDGTexture* DstTextureBRDG = DstTextureB ? GraphBuilder.RegisterExternalTexture(CreateRenderTarget(DstTextureB, TEXT("SparseVolumeTexture.TileDataTextureB"))) : nullptr;
+				FRDGBuffer* SrcBufferARDG = FormatSizeA > 0 ? GraphBuilder.RegisterExternalBuffer(TileDataAUploadBuffer) : nullptr;
+				FRDGBuffer* SrcBufferBRDG = FormatSizeB > 0 ? GraphBuilder.RegisterExternalBuffer(TileDataBUploadBuffer) : nullptr;
+				check(SrcBufferARDG || SrcBufferBRDG);
 
 				FRDGTextureUAV* DummyTextureUAV = GraphBuilder.CreateUAV(DummyTexture);
 				FRDGTextureUAV* DstTextureAUAV = DstTextureARDG ? GraphBuilder.CreateUAV(DstTextureARDG) : nullptr;
@@ -463,6 +466,8 @@ public:
 				FRDGBufferSRV* OccupancyBitsBufferSRV = GraphBuilder.CreateSRV(GraphBuilder.RegisterExternalBuffer(OccupancyBitsUploadBuffer));
 				FRDGBufferSRV* TileDataOffsetsBufferSRV = GraphBuilder.CreateSRV(GraphBuilder.RegisterExternalBuffer(TileDataOffsetsUploadBuffer));
 				FRDGBufferSRV* DstTileCoordsBufferSRV = GraphBuilder.CreateSRV(GraphBuilder.RegisterExternalBuffer(DstTileCoordsUploadBuffer));
+				// Either SrcBufferARDG or SrcBufferBRDG must exist and will have at least 1 element.
+				FRDGBufferSRV* DummySrcBufferSRV = SrcBufferARDG ? GraphBuilder.CreateSRV(SrcBufferARDG, FormatA) : GraphBuilder.CreateSRV(SrcBufferBRDG, FormatB);
 
 				auto ComputeShader = GetGlobalShaderMap(GMaxRHIFeatureLevel)->GetShader<FSparseVolumeTextureUpdateFromSparseBufferCS>();
 
@@ -509,14 +514,14 @@ public:
 					// This is the critical part: For every batch we create a SRV scoped to a range within the voxel data upload buffer still fitting within the typed buffer texel limit.
 					if (FormatSizeA && NumVoxelsAInThisBatch)
 					{
-						FRDGBufferSRVDesc SRVDesc(GraphBuilder.RegisterExternalBuffer(TileDataAUploadBuffer), FormatA);
+						FRDGBufferSRVDesc SRVDesc(SrcBufferARDG, FormatA);
 						SRVDesc.StartOffsetBytes = NumUploadedVoxelsA * FormatSizeA;
 						SRVDesc.NumElements = NumVoxelsAInThisBatch;
 						TileDataABufferSRV = GraphBuilder.CreateSRV(SRVDesc);
 					}
 					if (FormatSizeB && NumVoxelsBInThisBatch)
 					{
-						FRDGBufferSRVDesc SRVDesc(GraphBuilder.RegisterExternalBuffer(TileDataBUploadBuffer), FormatB);
+						FRDGBufferSRVDesc SRVDesc(SrcBufferBRDG, FormatB);
 						SRVDesc.StartOffsetBytes = NumUploadedVoxelsB * FormatSizeB;
 						SRVDesc.NumElements = NumVoxelsBInThisBatch;
 						TileDataBBufferSRV = GraphBuilder.CreateSRV(SRVDesc);
@@ -525,8 +530,8 @@ public:
 					FSparseVolumeTextureUpdateFromSparseBufferCS::FParameters* PassParameters = GraphBuilder.AllocParameters<FSparseVolumeTextureUpdateFromSparseBufferCS::FParameters>();
 					PassParameters->DstPhysicalTileTextureA = DstTextureAUAV ? DstTextureAUAV : DummyTextureUAV;
 					PassParameters->DstPhysicalTileTextureB = DstTextureBUAV ? DstTextureBUAV : DummyTextureUAV;
-					PassParameters->SrcPhysicalTileBufferA = TileDataABufferSRV ? TileDataABufferSRV : TileDataBBufferSRV;
-					PassParameters->SrcPhysicalTileBufferB = TileDataBBufferSRV ? TileDataBBufferSRV : TileDataABufferSRV;
+					PassParameters->SrcPhysicalTileBufferA = TileDataABufferSRV ? TileDataABufferSRV : DummySrcBufferSRV;
+					PassParameters->SrcPhysicalTileBufferB = TileDataBBufferSRV ? TileDataBBufferSRV : DummySrcBufferSRV;
 					PassParameters->OccupancyBitsBuffer = OccupancyBitsBufferSRV;
 					PassParameters->TileDataOffsetsBuffer = TileDataOffsetsBufferSRV;
 					PassParameters->DstTileCoordsBuffer = DstTileCoordsBufferSRV;
