Added EditCondition to light transmission to force user to set the light to moveable prior to activate light transmission.

diff --git a/Engine/Source/Runtime/Engine/Classes/Components/LightComponentBase.h b/Engine/Source/Runtime/Engine/Classes/Components/LightComponentBase.h
--- a/Engine/Source/Runtime/Engine/Classes/Components/LightComponentBase.h
+++ b/Engine/Source/Runtime/Engine/Classes/Components/LightComponentBase.h
@@ -70,7 +70,7 @@ class ULightComponentBase : public USceneComponent
 	uint32 bAffectTranslucentLighting:1;
 
 	/** Whether light from this light transmits through surfaces with subsurface scattering profiles. Requires light to be movable. */
-	UPROPERTY(EditAnywhere, BlueprintReadOnly, Category = Light, AdvancedDisplay)
+	UPROPERTY(EditAnywhere, BlueprintReadOnly, Category = Light, AdvancedDisplay, meta = (EditCondition="Mobility==EComponentMobility::Movable"))
 	uint32 bTransmission : 1;
 
 	/** Whether the light shadows volumetric fog.  Disabling this can save GPU time. */
