Only bring up SM6 warning for the editor platform in use.

#rnx

diff --git a/Engine/Source/Runtime/Engine/Private/RendererSettings.cpp b/Engine/Source/Runtime/Engine/Private/RendererSettings.cpp
--- a/Engine/Source/Runtime/Engine/Private/RendererSettings.cpp
+++ b/Engine/Source/Runtime/Engine/Private/RendererSettings.cpp
@@ -397,26 +397,29 @@ void URendererSettings::CheckForMissingShaderModels()
 			}
 		};
 
-		// Gather all Windows shader format settings
 		TArray<FString> D3D11TargetedShaderFormats;
 		TArray<FString> D3D12TargetedShaderFormats;
 		TArray<FString> WindowsVulkanTargetedShaderFormats;
 		TArray<FString> WindowsTargetedRHIs;
+		TArray<FString> LinuxVulkanTargetedShaderFormats;
+		TArray<FString> LinuxTargetedRHIs;
+
+#if PLATFORM_WINDOWS
+		// Gather all Windows shader format settings
 		GConfig->GetArray(TEXT("/Script/WindowsTargetPlatform.WindowsTargetSettings"), TEXT("D3D11TargetedShaderFormats"), D3D11TargetedShaderFormats, GEngineIni);
 		GConfig->GetArray(TEXT("/Script/WindowsTargetPlatform.WindowsTargetSettings"), TEXT("D3D12TargetedShaderFormats"), D3D12TargetedShaderFormats, GEngineIni);
-		GConfig->GetArray(TEXT("/Script/WindowsTargetPlatform.WindowsTargetSettings"), TEXT("VulkanTargetedShaderFormats"), WindowsVulkanTargetedShaderFormats, GEngineIni);
+		//GConfig->GetArray(TEXT("/Script/WindowsTargetPlatform.WindowsTargetSettings"), TEXT("VulkanTargetedShaderFormats"), WindowsVulkanTargetedShaderFormats, GEngineIni);
 		GConfig->GetArray(TEXT("/Script/WindowsTargetPlatform.WindowsTargetSettings"), TEXT("TargetedRHIs"), WindowsTargetedRHIs, GEngineIni);
 		CopySM6Format(TEXT("PCD3D_SM6"), WindowsTargetedRHIs, D3D12TargetedShaderFormats);
-		CopySM6Format(TEXT("SF_VULKAN_SM6"), WindowsTargetedRHIs, WindowsVulkanTargetedShaderFormats);
-
-		// Gather all Linux shader format settings
-		TArray<FString> LinuxVulkanTargetedShaderFormats;
-		TArray<FString> LinuxTargetedRHIs;
-		GConfig->GetArray(TEXT("/Script/LinuxTargetPlatform.LinuxTargetSettings"), TEXT("VulkanTargetedShaderFormats"), LinuxVulkanTargetedShaderFormats, GEngineIni);
-		GConfig->GetArray(TEXT("/Script/LinuxTargetPlatform.LinuxTargetSettings"), TEXT("TargetedRHIs"), LinuxTargetedRHIs, GEngineIni);
-		CopySM6Format(TEXT("SF_VULKAN_SM6"), LinuxTargetedRHIs, LinuxVulkanTargetedShaderFormats);
-
+		//CopySM6Format(TEXT("SF_VULKAN_SM6"), WindowsTargetedRHIs, WindowsVulkanTargetedShaderFormats);
+#elif PLATFORM_LINUX
+		// TODO: Gather all Linux shader format settings
+		//GConfig->GetArray(TEXT("/Script/LinuxTargetPlatform.LinuxTargetSettings"), TEXT("VulkanTargetedShaderFormats"), LinuxVulkanTargetedShaderFormats, GEngineIni);
+		//GConfig->GetArray(TEXT("/Script/LinuxTargetPlatform.LinuxTargetSettings"), TEXT("TargetedRHIs"), LinuxTargetedRHIs, GEngineIni);
+		//CopySM6Format(TEXT("SF_VULKAN_SM6"), LinuxTargetedRHIs, LinuxVulkanTargetedShaderFormats);
+#elif PLATFORM_MAC
 		// TODO: Gather all Mac shader format settings
+#endif
 
 		const bool bProjectUsesD3D = (D3D11TargetedShaderFormats.Num() + D3D12TargetedShaderFormats.Num()) > 0;
 		const bool bProjectMissingD3DSM6 = (bProjectUsesD3D && !D3D12TargetedShaderFormats.Contains(TEXT("PCD3D_SM6")));
diff --git a/Engine/Source/Runtime/Engine/Private/StaticMesh.cpp b/Engine/Source/Runtime/Engine/Private/StaticMesh.cpp
--- a/Engine/Source/Runtime/Engine/Private/StaticMesh.cpp
+++ b/Engine/Source/Runtime/Engine/Private/StaticMesh.cpp
@@ -6550,26 +6550,29 @@ void UStaticMesh::CheckForMissingShaderModels()
 			}
 		};
 
-		// Gather all Windows shader format settings
 		TArray<FString> D3D11TargetedShaderFormats;
 		TArray<FString> D3D12TargetedShaderFormats;
 		TArray<FString> WindowsVulkanTargetedShaderFormats;
 		TArray<FString> WindowsTargetedRHIs;
+		TArray<FString> LinuxVulkanTargetedShaderFormats;
+		TArray<FString> LinuxTargetedRHIs;
+
+#if PLATFORM_WINDOWS
+		// Gather all Windows shader format settings
 		GConfig->GetArray(TEXT("/Script/WindowsTargetPlatform.WindowsTargetSettings"), TEXT("D3D11TargetedShaderFormats"), D3D11TargetedShaderFormats, GEngineIni);
 		GConfig->GetArray(TEXT("/Script/WindowsTargetPlatform.WindowsTargetSettings"), TEXT("D3D12TargetedShaderFormats"), D3D12TargetedShaderFormats, GEngineIni);
-		GConfig->GetArray(TEXT("/Script/WindowsTargetPlatform.WindowsTargetSettings"), TEXT("VulkanTargetedShaderFormats"), WindowsVulkanTargetedShaderFormats, GEngineIni);
+		//GConfig->GetArray(TEXT("/Script/WindowsTargetPlatform.WindowsTargetSettings"), TEXT("VulkanTargetedShaderFormats"), WindowsVulkanTargetedShaderFormats, GEngineIni);
 		GConfig->GetArray(TEXT("/Script/WindowsTargetPlatform.WindowsTargetSettings"), TEXT("TargetedRHIs"), WindowsTargetedRHIs, GEngineIni);
 		CopySM6Format(TEXT("PCD3D_SM6"), WindowsTargetedRHIs, D3D12TargetedShaderFormats);
-		CopySM6Format(TEXT("SF_VULKAN_SM6"), WindowsTargetedRHIs, WindowsVulkanTargetedShaderFormats);
-
-		// Gather all Linux shader format settings
-		TArray<FString> LinuxVulkanTargetedShaderFormats;
-		TArray<FString> LinuxTargetedRHIs;
-		GConfig->GetArray(TEXT("/Script/LinuxTargetPlatform.LinuxTargetSettings"), TEXT("VulkanTargetedShaderFormats"), LinuxVulkanTargetedShaderFormats, GEngineIni);
-		GConfig->GetArray(TEXT("/Script/LinuxTargetPlatform.LinuxTargetSettings"), TEXT("TargetedRHIs"), LinuxTargetedRHIs, GEngineIni);
-		CopySM6Format(TEXT("SF_VULKAN_SM6"), LinuxTargetedRHIs, LinuxVulkanTargetedShaderFormats);
-
+		//CopySM6Format(TEXT("SF_VULKAN_SM6"), WindowsTargetedRHIs, WindowsVulkanTargetedShaderFormats);
+#elif PLATFORM_LINUX
+		// TODO: Gather all Linux shader format settings
+		//GConfig->GetArray(TEXT("/Script/LinuxTargetPlatform.LinuxTargetSettings"), TEXT("VulkanTargetedShaderFormats"), LinuxVulkanTargetedShaderFormats, GEngineIni);
+		//GConfig->GetArray(TEXT("/Script/LinuxTargetPlatform.LinuxTargetSettings"), TEXT("TargetedRHIs"), LinuxTargetedRHIs, GEngineIni);
+		//CopySM6Format(TEXT("SF_VULKAN_SM6"), LinuxTargetedRHIs, LinuxVulkanTargetedShaderFormats);
+#elif PLATFORM_MAC
 		// TODO: Gather all Mac shader format settings
+#endif
 
 		const bool bProjectUsesD3D = (D3D11TargetedShaderFormats.Num() + D3D12TargetedShaderFormats.Num()) > 0;
 		const bool bProjectMissingD3DSM6 = (bProjectUsesD3D && !D3D12TargetedShaderFormats.Contains(TEXT("PCD3D_SM6")));
