[Engine]
* Removed 1.6 million allocations when starting game

diff --git a/Engine/Source/Runtime/Engine/Private/BlueprintGeneratedClass.cpp b/Engine/Source/Runtime/Engine/Private/BlueprintGeneratedClass.cpp
--- a/Engine/Source/Runtime/Engine/Private/BlueprintGeneratedClass.cpp
+++ b/Engine/Source/Runtime/Engine/Private/BlueprintGeneratedClass.cpp
@@ -764,7 +764,9 @@ bool UBlueprintGeneratedClass::BuildCustomPropertyListForPostConstruction(FCusto
 		if (FBlueprintGeneratedClassUtils::ShouldInitializePropertyDuringPostConstruction(*Property))
 		{
 			// Some properties require a full value comparison; check for those cases here.
-			const bool bAlwaysUseCompleteValue = RequiresCompleteValueForPostConstruction.Contains(Property->GetPathName());
+			TStringBuilder<256> PathName;
+			Property->GetPathName(nullptr, PathName);
+			const bool bAlwaysUseCompleteValue = RequiresCompleteValueForPostConstruction.ContainsByPredicate([&](const FString& S) { return S == PathName.ToView(); });
 
 			for (int32 Idx = 0; Idx < Property->ArrayDim; Idx++)
 			{
