PR #10608: ABrush and its descendants should respect the Transient flag when converted to an external package.

diff --git a/Engine/Source/Runtime/Engine/Private/Brush.cpp b/Engine/Source/Runtime/Engine/Private/Brush.cpp
--- a/Engine/Source/Runtime/Engine/Private/Brush.cpp
+++ b/Engine/Source/Runtime/Engine/Private/Brush.cpp
@@ -247,7 +247,7 @@ bool ABrush::SupportsLayers() const
 bool ABrush::SupportsExternalPackaging() const
 {
 	// Base class ABrush actors do not support OFPA
-	return GetClass() != ABrush::StaticClass();
+	return GetClass() != ABrush::StaticClass() && Super::SupportsExternalPackaging();
 }
 
 void ABrush::PostLoad()
