Fix for SparseVolumeTexture border voxels of mip0 sometimes not being correctly generated. The issue was caused by using an incorrect page table size to compute the list of pages to process.

#rnx

diff --git a/Engine/Source/Runtime/Engine/Private/SparseVolumeTexture/SparseVolumeTextureData.cpp b/Engine/Source/Runtime/Engine/Private/SparseVolumeTexture/SparseVolumeTextureData.cpp
--- a/Engine/Source/Runtime/Engine/Private/SparseVolumeTexture/SparseVolumeTextureData.cpp
+++ b/Engine/Source/Runtime/Engine/Private/SparseVolumeTexture/SparseVolumeTextureData.cpp
@@ -786,13 +786,13 @@ bool FTextureData::BuildDerivedData(const FTextureDataAddressingInfo& Addressing
 	{
 		// Collect the page table coordinates of all non-zero pages
 		TArray<FIntVector3> PageCoords;
-		for (int32 PageZ = 0; PageZ < Header.PageTableVolumeResolution.Z; ++PageZ)
+		for (int32 PageZ = 0; PageZ < OutDerivedData.Header.PageTableVolumeResolution.Z; ++PageZ)
 		{
-			for (int32 PageY = 0; PageY < Header.PageTableVolumeResolution.Y; ++PageY)
+			for (int32 PageY = 0; PageY < OutDerivedData.Header.PageTableVolumeResolution.Y; ++PageY)
 			{
-				for (int32 PageX = 0; PageX < Header.PageTableVolumeResolution.X; ++PageX)
+				for (int32 PageX = 0; PageX < OutDerivedData.Header.PageTableVolumeResolution.X; ++PageX)
 				{
-					const int32 PageIndex = PageZ * (Header.PageTableVolumeResolution.Y * Header.PageTableVolumeResolution.X) + PageY * Header.PageTableVolumeResolution.X + PageX;
+					const int32 PageIndex = PageZ * (OutDerivedData.Header.PageTableVolumeResolution.Y * OutDerivedData.Header.PageTableVolumeResolution.X) + PageY * OutDerivedData.Header.PageTableVolumeResolution.X + PageX;
 					const uint32 PageTableEntry = OutDerivedData.MipMaps[0].PageTable[PageIndex];
 					if (PageTableEntry != 0)
 					{
