Change to SetCursorWidget to make the widget param required

diff --git a/Engine/Source/Runtime/Engine/Classes/Engine/GameViewportClient.h b/Engine/Source/Runtime/Engine/Classes/Engine/GameViewportClient.h
--- a/Engine/Source/Runtime/Engine/Classes/Engine/GameViewportClient.h
+++ b/Engine/Source/Runtime/Engine/Classes/Engine/GameViewportClient.h
@@ -139,7 +139,8 @@ public:
 	UFUNCTION(exec)
 	ENGINE_API virtual void SetConsoleTarget(int32 PlayerIndex);
 
-	/** Sets the widget to use fore the cursor. */
+	/** Add a Widget to use for specified cursor type*/
+	UE_DEPRECATED(5.4, "AddCursorWidget is deprecated, use SetSoftwareCursorWidget instead.")
 	ENGINE_API void AddCursorWidget(EMouseCursor::Type Cursor, class UUserWidget* CursorWidget);
 
 	/** Returns a relative world context for this viewport.	 */
@@ -858,14 +859,22 @@ public:
 	 */
 	ENGINE_API bool SetDisplayConfiguration( const FIntPoint* Dimensions, EWindowMode::Type WindowMode);
 
+	UE_DEPRECATED(5.4, "SetVirtualCursorWidget is deprecated, use SetSoftwareCursorWidget instead.")
 	ENGINE_API void SetVirtualCursorWidget(EMouseCursor::Type Cursor, class UUserWidget* Widget);
 
-	/** Add a cursor to the set based on the enum and a slate widget */
+	/** Add a Widget to use for specified cursor type*/
+	UE_DEPRECATED(5.4, "AddSoftwareCursorFromSlateWidget is deprecated, use SetSoftwareCursorWidget instead.")
 	ENGINE_API void AddSoftwareCursorFromSlateWidget(EMouseCursor::Type InCursorType, TSharedPtr<SWidget> CursorWidgetPtr);
 
-	/** Adds a cursor to the set based on the enum and the class reference to it. */
+	UE_DEPRECATED(5.4, "AddSoftwareCursor is deprecated, use SetSoftwareCursorWidget instead and pass the Created Widget instead of the UserWidget Class.")
 	ENGINE_API void AddSoftwareCursor(EMouseCursor::Type Cursor, const FSoftClassPath& CursorClass);
 
+	/** Add a Widget to use for specified cursor type*/
+	ENGINE_API void SetSoftwareCursorWidget(EMouseCursor::Type InCursorType, class UUserWidget* Widget);
+
+	/** Add a Widget to use for specified cursor type*/
+	ENGINE_API void SetSoftwareCursorWidget(EMouseCursor::Type InCursorType, TSharedPtr<SWidget> CursorWidgetPtr);
+
 	/** Get the slate widget of the current software cursor */
 	ENGINE_API TSharedPtr<SWidget> GetSoftwareCursorWidget(EMouseCursor::Type Cursor) const;
 
@@ -880,6 +889,9 @@ private:
 	/** Resets the platform type shape to nullptr, to restore it to the OS default. */
 	ENGINE_API void ResetHardwareCursorStates();
 
+	/** Set the Software Cursor from the class path. This should stay private to make sure error reporting on loading is done by the caller. */
+	void SetSoftwareCursorFromClassPath(EMouseCursor::Type Cursor, const FSoftClassPath& CursorClass);
+
 	/**
 	 * Set a specific stat to either enabled or disabled (returns the number of remaining enabled stats)
 	 */
diff --git a/Engine/Source/Runtime/Engine/Classes/GameFramework/PlayerController.h b/Engine/Source/Runtime/Engine/Classes/GameFramework/PlayerController.h
--- a/Engine/Source/Runtime/Engine/Classes/GameFramework/PlayerController.h
+++ b/Engine/Source/Runtime/Engine/Classes/GameFramework/PlayerController.h
@@ -1233,7 +1233,7 @@ public:
 	 * @param CursorWidget - the widget to set the cursor to
 	 */
 	UFUNCTION(BlueprintCallable, Category="Game|Player")
-	ENGINE_API void SetMouseCursorWidget(EMouseCursor::Type Cursor, class UUserWidget* CursorWidget);
+	ENGINE_API void SetMouseCursorWidget(EMouseCursor::Type Cursor, UPARAM(Required) class UUserWidget* CursorWidget);
 
 	/** Set the view target
 	 * @param A - new actor to set as view target
diff --git a/Engine/Source/Runtime/Engine/Private/GameViewportClient.cpp b/Engine/Source/Runtime/Engine/Private/GameViewportClient.cpp
--- a/Engine/Source/Runtime/Engine/Private/GameViewportClient.cpp
+++ b/Engine/Source/Runtime/Engine/Private/GameViewportClient.cpp
@@ -551,7 +551,7 @@ void UGameViewportClient::Init(struct FWorldContext& WorldContext, UGameInstance
 	// Set all the software cursors.
 	for ( auto& Entry : UISettings->SoftwareCursors )
 	{
-		AddSoftwareCursor(Entry.Key, Entry.Value);
+		SetSoftwareCursorFromClassPath(Entry.Key, Entry.Value);
 	}
 
 	// Set all the hardware cursors.
@@ -567,7 +567,7 @@ void UGameViewportClient::RebuildCursors()
 	// Set all the software cursors.
 	for (auto& Entry : UISettings->SoftwareCursors)
 	{
-		AddSoftwareCursor(Entry.Key, Entry.Value);
+		SetSoftwareCursorFromClassPath(Entry.Key, Entry.Value);
 	}
 
 	// Set all the hardware cursors.
@@ -1030,42 +1030,72 @@ EMouseCursor::Type UGameViewportClient::GetCursor(FViewport* InViewport, int32 X
 
 void UGameViewportClient::SetVirtualCursorWidget(EMouseCursor::Type Cursor, UUserWidget* UserWidget)
 {
-	TSharedPtr<SWidget>& ExistingWidget = CursorWidgets.FindOrAdd(Cursor);
-	TSharedPtr<SWidget> NewWidget = UserWidget ? UserWidget->TakeWidget() : TSharedPtr<SWidget>();
-	if (NewWidget != ExistingWidget)
+	if (UserWidget)
 	{
-		// Pure safety
-		ExistingWidget.Reset();
-		ExistingWidget = NewWidget;
+		SetSoftwareCursorWidget(Cursor, UserWidget);
+	}
+	else
+	{
+		CursorWidgets.Remove(Cursor);
+	}
+}
+
+void UGameViewportClient::AddSoftwareCursorFromSlateWidget(EMouseCursor::Type InCursorType, TSharedPtr<SWidget> CursorWidgetPtr)
+{
+	// We set it only when it's not null to be on parity with the behavior we had before deprecation.
+	if (CursorWidgetPtr.IsValid())
+	{
+		SetSoftwareCursorWidget(InCursorType, CursorWidgetPtr);
 	}
 }
 
 void UGameViewportClient::AddSoftwareCursor(EMouseCursor::Type Cursor, const FSoftClassPath& CursorClass)
+{
+	SetSoftwareCursorFromClassPath(Cursor, CursorClass);
+}
+
+void UGameViewportClient::SetSoftwareCursorFromClassPath(EMouseCursor::Type Cursor, const FSoftClassPath & CursorClass)
 {
 	if (CursorClass.IsValid())
 	{
 		if (UClass* Class = CursorClass.TryLoadClass<UUserWidget>())
 		{
 			UUserWidget* UserWidget = CreateWidget(GetGameInstance(), Class);
-			AddCursorWidget(Cursor, UserWidget);
+			SetSoftwareCursorWidget(Cursor, UserWidget);
 		}
 		else
 		{
-			FMessageLog("PIE").Warning(FText::Format(LOCTEXT("AddCursor:LoadFailed", "UGameViewportClient::AddCursor: Could not load cursor class '{0}'."), FText::FromString(CursorClass.GetAssetName())));
+			FMessageLog("PIE").Warning(FText::Format(LOCTEXT("SetSoftwareCursorFromClassPath:LoadFailed", "UGameViewportClient::SetSoftwareCursorFromClassPath: Could not load cursor class '{0}'."), FText::FromString(CursorClass.GetAssetName())));
 		}
 	}
 	else
 	{
-		FMessageLog("PIE").Warning(LOCTEXT("AddCursor:InvalidClass", "UGameViewportClient::AddCursor: Invalid class specified."));
+		FMessageLog("PIE").Warning(LOCTEXT("SetSoftwareCursorFromClassPath:InvalidClass", "UGameViewportClient::SetSoftwareCursorFromClassPath: Invalid class specified."));
 	}
 }
 
-void UGameViewportClient::AddSoftwareCursorFromSlateWidget(EMouseCursor::Type InCursorType, TSharedPtr<SWidget> CursorWidgetPtr)
+void UGameViewportClient::SetSoftwareCursorWidget(EMouseCursor::Type InCursorType, TSharedPtr<SWidget> CursorWidgetPtr)
 {
 	if (CursorWidgetPtr.IsValid())
 	{
 		CursorWidgets.Emplace(InCursorType, CursorWidgetPtr);
 	}
+	else
+	{
+		CursorWidgets.Remove(InCursorType);
+	}
+}
+
+void UGameViewportClient::SetSoftwareCursorWidget(EMouseCursor::Type InCursorType, class UUserWidget* UserWidget)
+{
+	if (UserWidget)
+	{
+		SetSoftwareCursorWidget(InCursorType, UserWidget->TakeWidget());
+	}
+	else
+	{
+		CursorWidgets.Remove(InCursorType);
+	}
 }
 
 TSharedPtr<SWidget> UGameViewportClient::GetSoftwareCursorWidget(EMouseCursor::Type Cursor) const
@@ -1088,9 +1118,10 @@ bool UGameViewportClient::HasSoftwareCursor(EMouseCursor::Type Cursor) const
 
 void UGameViewportClient::AddCursorWidget(EMouseCursor::Type Cursor, class UUserWidget* CursorWidget)
 {
-	if (ensure(CursorWidget))
+	// We set it only when it's not null to be on parity with the behavior we had before deprecation.
+	if (CursorWidget)
 	{
-		CursorWidgets.Add(Cursor, CursorWidget->TakeWidget());
+		SetSoftwareCursorWidget(Cursor, CursorWidget->TakeWidget());
 	}
 }
 
diff --git a/Engine/Source/Runtime/Engine/Private/PlayerController.cpp b/Engine/Source/Runtime/Engine/Private/PlayerController.cpp
--- a/Engine/Source/Runtime/Engine/Private/PlayerController.cpp
+++ b/Engine/Source/Runtime/Engine/Private/PlayerController.cpp
@@ -1176,7 +1176,7 @@ void APlayerController::SetMouseCursorWidget(EMouseCursor::Type Cursor, class UU
 	ULocalPlayer* LocalPlayer = Cast<ULocalPlayer>(Player);
 	if (LocalPlayer && LocalPlayer->ViewportClient)
 	{
-		LocalPlayer->ViewportClient->AddCursorWidget(Cursor, CursorWidget);
+		LocalPlayer->ViewportClient->SetSoftwareCursorWidget(Cursor, CursorWidget);
 	}
 }
 
