Fix seamless travel not using the correct feature level

diff --git a/Engine/Source/Runtime/Engine/Private/World.cpp b/Engine/Source/Runtime/Engine/Private/World.cpp
--- a/Engine/Source/Runtime/Engine/Private/World.cpp
+++ b/Engine/Source/Runtime/Engine/Private/World.cpp
@@ -7788,6 +7788,13 @@ UWorld* FSeamlessTravelHandler::Tick()
 			// Track session change on seamless travel.
 			NETWORK_PROFILER(GNetworkProfiler.TrackSessionChange(true, LoadedWorld->URL));
 
+#if WITH_EDITOR
+			// PIE worlds should use the same feature level as the editor
+			if (CurrentContext.PIEWorldFeatureLevel != ERHIFeatureLevel::Num && LoadedWorld->GetFeatureLevel() != CurrentContext.PIEWorldFeatureLevel)
+			{
+				LoadedWorld->ChangeFeatureLevel(CurrentContext.PIEWorldFeatureLevel);
+			}
+#endif
 
 			checkSlow((LoadedWorld->GetNetMode() == NM_Client) == bIsClient);
 
