PR #10609: Nullify World->ContentBundleManager if World Partition is disabled.
* Could happens when switching from non WP world to WP world.

diff --git a/Engine/Source/Runtime/Engine/Private/WorldPartition/ContentBundle/ContentBundleEngineSubsystem.cpp b/Engine/Source/Runtime/Engine/Private/WorldPartition/ContentBundle/ContentBundleEngineSubsystem.cpp
--- a/Engine/Source/Runtime/Engine/Private/WorldPartition/ContentBundle/ContentBundleEngineSubsystem.cpp
+++ b/Engine/Source/Runtime/Engine/Private/WorldPartition/ContentBundle/ContentBundleEngineSubsystem.cpp
@@ -158,6 +158,10 @@ void UContentBundleEngineSubsystem::OnPreWorldInitialization(UWorld* World, cons
 	{
 		World->ContentBundleManager = NewObject<UContentBundleManager>(World);
 	}
+	else if (!World->GetWorldPartition() && World->ContentBundleManager != nullptr)
+	{
+		World->ContentBundleManager = nullptr;
+	}
 
 	if (World->ContentBundleManager != nullptr)
 	{
