NetDriver:
* Fix for Push Model may prevent actors from going dormant
* Now testing changelist history in CanSkipUpdate

diff --git a/Engine/Source/Runtime/Engine/Private/DataReplication.cpp b/Engine/Source/Runtime/Engine/Private/DataReplication.cpp
--- a/Engine/Source/Runtime/Engine/Private/DataReplication.cpp
+++ b/Engine/Source/Runtime/Engine/Private/DataReplication.cpp
@@ -1785,6 +1785,9 @@ bool FObjectReplicator::CanSkipUpdate(FReplicationFlags RepFlags)
 	// Any changelists to send ?
 	bCanSkip = bCanSkip && SendingRepState.LastChangelistIndex == RepChangelistState.HistoryEnd;
 
+	// Is the changelist history fully acknowledged ?
+	bCanSkip = bCanSkip && SendingRepState.HistoryStart == SendingRepState.HistoryEnd;
+
 	// Are we resending data for replay ?
 	bCanSkip = bCanSkip && Connection->ResendAllDataState == EResendAllDataState::None;
 
