Force scene component bounds to recompute when a level transform is applied.

Cached bounds will have been cached at a different transform and will never be correct unless recomputed.

diff --git a/Engine/Source/Runtime/Engine/Classes/GameFramework/Actor.h b/Engine/Source/Runtime/Engine/Classes/GameFramework/Actor.h
--- a/Engine/Source/Runtime/Engine/Classes/GameFramework/Actor.h
+++ b/Engine/Source/Runtime/Engine/Classes/GameFramework/Actor.h
@@ -1701,6 +1701,12 @@ public:
 	UFUNCTION(BlueprintCallable, Category="Transformation")
 	ENGINE_API FVector GetActorRelativeScale3D() const;
 
+	/**
+	 * Marks the bounds of all SceneComponents attached to this actor which have `bComputeBoundsOnceForGame` as needing to be recomputed the next time UpdateBounds is called. 
+	 * This might be necessary if the bounds that were cached on cook no longer reflect the actor's transform (ex. if level transform is applied).
+	 */
+	ENGINE_API void MarkNeedsRecomputeBoundsOnceForGame();
+
 	/**
 	 *	Sets the actor to be hidden in the game
 	 *	@param	bNewHidden	Whether or not to hide the actor and all its components
diff --git a/Engine/Source/Runtime/Engine/Private/Actor.cpp b/Engine/Source/Runtime/Engine/Private/Actor.cpp
--- a/Engine/Source/Runtime/Engine/Private/Actor.cpp
+++ b/Engine/Source/Runtime/Engine/Private/Actor.cpp
@@ -4688,6 +4688,17 @@ FVector AActor::GetActorRelativeScale3D() const
 	return FVector(1,1,1);
 }
 
+void AActor::MarkNeedsRecomputeBoundsOnceForGame()
+{
+	ForEachComponent<USceneComponent>(true, [](USceneComponent* SceneComponent)
+	{
+		if (SceneComponent && SceneComponent->bComputeBoundsOnceForGame)
+		{
+			SceneComponent->bComputedBoundsOnceForGame = false;
+		}
+	});
+}
+
 void AActor::SetActorHiddenInGame( bool bNewHidden )
 {
 	if (IsHidden() != bNewHidden)
diff --git a/Engine/Source/Runtime/Engine/Private/LevelUtils.cpp b/Engine/Source/Runtime/Engine/Private/LevelUtils.cpp
--- a/Engine/Source/Runtime/Engine/Private/LevelUtils.cpp
+++ b/Engine/Source/Runtime/Engine/Private/LevelUtils.cpp
@@ -444,6 +444,8 @@ void FLevelUtils::ApplyLevelTransform(const FLevelUtils::FApplyLevelTransformPar
 					RootComponent->SetRelativeLocation_Direct(TransformParams.LevelTransform.TransformPosition(RootComponent->GetRelativeLocation()));
 					RootComponent->SetRelativeRotation_Direct(TransformParams.LevelTransform.TransformRotation(RootComponent->GetRelativeRotation().Quaternion()).Rotator());
 					RootComponent->SetRelativeScale3D_Direct(TransformParams.LevelTransform.GetScale3D() * RootComponent->GetRelativeScale3D());
+
+					TransformParams.Actor->MarkNeedsRecomputeBoundsOnceForGame();
 				}
 			}
 			else
@@ -454,6 +456,9 @@ void FLevelUtils::ApplyLevelTransform(const FLevelUtils::FApplyLevelTransformPar
 				{
 					RootComponent->SetRelativeLocationAndRotation(TransformParams.LevelTransform.TransformPosition(RootComponent->GetRelativeLocation()), TransformParams.LevelTransform.TransformRotation(RootComponent->GetRelativeRotation().Quaternion()));
 					RootComponent->SetRelativeScale3D(TransformParams.LevelTransform.GetScale3D() * RootComponent->GetRelativeScale3D());
+
+					// Any components which have cached their bounds will not be accurate after a level transform is applied. Force them to recompute the bounds once more.
+					TransformParams.Actor->MarkNeedsRecomputeBoundsOnceForGame();
 				}
 			}
 #if WITH_EDITOR
@@ -500,6 +505,9 @@ void FLevelUtils::ApplyLevelTransform(const FLevelUtils::FApplyLevelTransformPar
 						RootComponent->SetRelativeLocation_Direct(TransformParams.LevelTransform.TransformPosition(RootComponent->GetRelativeLocation()));
 						RootComponent->SetRelativeRotation_Direct(TransformParams.LevelTransform.TransformRotation(RootComponent->GetRelativeRotation().Quaternion()).Rotator());
 						RootComponent->SetRelativeScale3D_Direct(TransformParams.LevelTransform.GetScale3D() * RootComponent->GetRelativeScale3D());
+
+						// Any components which have cached their bounds will not be accurate after a level transform is applied. Force them to recompute the bounds once more.
+						Actor->MarkNeedsRecomputeBoundsOnceForGame();
 					}
 				}
 			}
@@ -528,6 +536,9 @@ void FLevelUtils::ApplyLevelTransform(const FLevelUtils::FApplyLevelTransformPar
 					{
 						RootComponent->SetRelativeLocationAndRotation(TransformParams.LevelTransform.TransformPosition(RootComponent->GetRelativeLocation()), TransformParams.LevelTransform.TransformRotation(RootComponent->GetRelativeRotation().Quaternion()));
 						RootComponent->SetRelativeScale3D(TransformParams.LevelTransform.GetScale3D() * RootComponent->GetRelativeScale3D());
+
+						// Any components which have cached bounds will not be accurate after a level transform is applied. Force them to recompute the bounds once more.
+						Actor->MarkNeedsRecomputeBoundsOnceForGame();
 					}
 				}
 			}
