[RegressionFix] CL 24999829 broke forward compatibility with copy paste. This change minimizes the impact to only properties with certain special characters in their names

diff --git a/Engine/Source/Runtime/Engine/Private/UnrealExporter.cpp b/Engine/Source/Runtime/Engine/Private/UnrealExporter.cpp
--- a/Engine/Source/Runtime/Engine/Private/UnrealExporter.cpp
+++ b/Engine/Source/Runtime/Engine/Private/UnrealExporter.cpp
@@ -687,7 +687,15 @@ void ExportProperties
 		if (!Property->ShouldPort(PortFlags))
 			continue;
 
-		const FString SanitizedPropertyName = FString::Format(TEXT("\"{0}\""), {Property->GetName().ReplaceCharWithEscapedChar()});
+		FString SanitizedPropertyName = Property->GetName();
+		constexpr FAsciiSet Whitespace("\t ");
+		constexpr FAsciiSet SpecialCharacters("=()[].\"\'");
+		if (FAsciiSet::HasAny(SanitizedPropertyName, SpecialCharacters) || Whitespace.Contains(SanitizedPropertyName[0]))
+		{
+			// to increase frequency of forward compatibility, only sanitize property names that absolutely need it
+			SanitizedPropertyName = FString::Format(TEXT("\"{0}\""), {Property->GetName().ReplaceCharWithEscapedChar()});
+		}
+
 		FArrayProperty* ArrayProperty = CastField<FArrayProperty>(Property);
 		FObjectPropertyBase* ExportObjectProp = (Property->PropertyFlags & CPF_ExportObject) != 0 ? CastField<FObjectPropertyBase>(Property) : NULL;
 		const uint32 ExportFlags = PortFlags | PPF_Delimited;
