World Partition: make actors runtime grid editable for level instances and inherit parent values only if no runtime grid is set on level instance actors.

diff --git a/Engine/Source/Runtime/Engine/Private/ActorEditor.cpp b/Engine/Source/Runtime/Engine/Private/ActorEditor.cpp
--- a/Engine/Source/Runtime/Engine/Private/ActorEditor.cpp
+++ b/Engine/Source/Runtime/Engine/Private/ActorEditor.cpp
@@ -102,7 +102,7 @@ bool AActor::CanEditChange(const FProperty* PropertyThatWillChange) const
 				return false;
 			}
 
-			if (!OwningWorldPartition->IsStreamingEnabled() || GetAttachParentActor())
+			if (GetAttachParentActor())
 			{
 				if (bIsSpatiallyLoadedProperty || bIsRuntimeGridProperty)
 				{
diff --git a/Engine/Source/Runtime/Engine/Private/WorldPartition/WorldPartitionStreamingGeneration.cpp b/Engine/Source/Runtime/Engine/Private/WorldPartition/WorldPartitionStreamingGeneration.cpp
--- a/Engine/Source/Runtime/Engine/Private/WorldPartition/WorldPartitionStreamingGeneration.cpp
+++ b/Engine/Source/Runtime/Engine/Private/WorldPartition/WorldPartitionStreamingGeneration.cpp
@@ -472,8 +472,8 @@ class FWorldPartitionStreamingGenerator
 			// Apply AND logic on spatially loaded flag
 			ResultPerInstanceData.bIsSpatiallyLoaded = InActorDescView.GetIsSpatiallyLoaded() && InParentPerInstanceData.bIsSpatiallyLoaded;
 
-			// Runtime grid is only inherited from the main world, since level instance doesn't support setting this value on actors
-			ResultPerInstanceData.RuntimeGrid = InContainerCollectionInstanceDescriptor.ID.IsMainContainer() ? InActorDescView.GetRuntimeGrid() : InParentPerInstanceData.RuntimeGrid;
+			// Runtime grid is inherited from the main world if the actor has its runtime grid set to none.
+			ResultPerInstanceData.RuntimeGrid = (InContainerCollectionInstanceDescriptor.ID.IsMainContainer() || InParentPerInstanceData.RuntimeGrid.IsNone()) ? InActorDescView.GetRuntimeGrid() : InParentPerInstanceData.RuntimeGrid;
 
 			// Data layers are accumulated down the hierarchy chain, since level instances supports data layers assignation on actors
 			ResultPerInstanceData.DataLayers = InActorDescView.GetRuntimeDataLayerInstanceNames();
@@ -979,7 +979,7 @@ PRAGMA_ENABLE_DEPRECATION_WARNINGS
 					{
 						if (PassType == EPassType::ErrorReporting)
 						{
-							ErrorHandler->OnInvalidRuntimeGrid(ActorDescView, ActorDescView.GetRuntimeGrid());
+							ErrorHandler->OnInvalidRuntimeGrid(ActorDescView, PerInstanceData.RuntimeGrid);
 						}
 						else
 						{
