Moved instanced static mesh initialization work to the constructor.
 - CreateRenderThreadResources is meant for RHI resource creation and not scene proxy setup.
 - This change is necessary to move CreateRenderThreadResources to happen during UpdateAllPrimitiveSceneInfos.

diff --git a/Engine/Source/Runtime/Engine/Private/InstancedStaticMesh.cpp b/Engine/Source/Runtime/Engine/Private/InstancedStaticMesh.cpp
--- a/Engine/Source/Runtime/Engine/Private/InstancedStaticMesh.cpp
+++ b/Engine/Source/Runtime/Engine/Private/InstancedStaticMesh.cpp
@@ -1902,29 +1902,7 @@ void FInstancedStaticMeshSceneProxy::SetupProxy(UInstancedStaticMeshComponent* I
 				);
 			}
 		}
-	}
-}
-
-
-void FInstancedStaticMeshSceneProxy::CreateRenderThreadResources(FRHICommandListBase& RHICmdList)
-{
-	FStaticMeshSceneProxy::CreateRenderThreadResources(RHICmdList);
-
-	const bool bCanUseGPUScene = UseGPUScene(GetScene().GetShaderPlatform(), GetScene().GetFeatureLevel());
-	
-	// Flush upload of GPU data for ISM/HISM
-	if (ensure(InstancedRenderData.PerInstanceRenderData.IsValid()))
-	{
-		FStaticMeshInstanceBuffer& InstanceBuffer = InstancedRenderData.PerInstanceRenderData->InstanceBuffer;
-		if (!bCanUseGPUScene)
-		{
-			InstanceBuffer.FlushGPUUpload(RHICmdList);
-		}
-	}
 
-	if (bCanUseGPUScene)
-	{
-		bSupportsInstanceDataBuffer = true;
 		// TODO: can the PerInstanceRenderData ever not be valid here?
 		if (ensure(InstancedRenderData.PerInstanceRenderData.IsValid()))
 		{
@@ -1988,6 +1966,23 @@ void FInstancedStaticMeshSceneProxy::CreateRenderThreadResources(FRHICommandList
 			}
 		}
 	}
+}
+
+void FInstancedStaticMeshSceneProxy::CreateRenderThreadResources(FRHICommandListBase& RHICmdList)
+{
+	FStaticMeshSceneProxy::CreateRenderThreadResources(RHICmdList);
+
+	// Flush upload of GPU data for ISM/HISM
+	if (ensure(InstancedRenderData.PerInstanceRenderData.IsValid()))
+	{
+		const bool bCanUseGPUScene = UseGPUScene(GetScene().GetShaderPlatform(), GetScene().GetFeatureLevel());
+
+		if (!bCanUseGPUScene)
+		{
+			FStaticMeshInstanceBuffer& InstanceBuffer = InstancedRenderData.PerInstanceRenderData->InstanceBuffer;
+			InstanceBuffer.FlushGPUUpload(RHICmdList);
+		}
+	}
 
 	for (int32 i = 0; i < LODs.Num(); ++i)
 	{
