Removed render thread constraints from material render proxy deletion and uniform expression caching.

diff --git a/Engine/Plugins/2D/Paper2D/Source/Paper2D/Private/PaperRenderSceneProxy.cpp b/Engine/Plugins/2D/Paper2D/Source/Paper2D/Private/PaperRenderSceneProxy.cpp
--- a/Engine/Plugins/2D/Paper2D/Source/Paper2D/Private/PaperRenderSceneProxy.cpp
+++ b/Engine/Plugins/2D/Paper2D/Source/Paper2D/Private/PaperRenderSceneProxy.cpp
@@ -58,10 +58,6 @@ public:
 	{
 	}
 
-	virtual ~FSpriteTextureOverrideRenderProxy()
-	{
-	}
-
 	void CheckValidity(FMaterialRenderProxy* InCurrentParent)
 	{
 		if (InCurrentParent != Parent)
@@ -73,7 +69,7 @@ public:
 		if (ParentMaterialSerialNumber != Parent->GetExpressionCacheSerialNumber())
 		{
 			// Not valid, need to rebuild
-			CacheUniformExpressions(/*bRecreateUniformBuffer=*/ true);
+			CacheUniformExpressions(FRHICommandListImmediate::Get(), /*bRecreateUniformBuffer=*/ true);
 			ParentMaterialSerialNumber = Parent->GetExpressionCacheSerialNumber();
 		}
 	}
@@ -184,19 +180,6 @@ SIZE_T FPaperRenderSceneProxy::GetTypeHash() const
 	return reinterpret_cast<size_t>(&UniquePointer);
 }
 
-FPaperRenderSceneProxy::~FPaperRenderSceneProxy()
-{
-	for (FSpriteTextureOverrideRenderProxy* Proxy : MaterialTextureOverrideProxies)
-	{
-		if (Proxy != nullptr)
-		{
-			Proxy->ReleasePrimitiveResource();
-		}
-	}
-	VertexBuffer.ReleaseResource();
-	VertexFactory.ReleaseResource();
-}
-
 void FPaperRenderSceneProxy::DebugDrawBodySetup(const FSceneView* View, int32 ViewIndex, FMeshElementCollector& Collector, UBodySetup* BodySetup, const FMatrix& GeomTransformMatrix, const FLinearColor& CollisionColor, bool bDrawSolid) const
 {
 	if (FMath::Abs(GeomTransformMatrix.Determinant()) < SMALL_NUMBER)
@@ -278,6 +261,21 @@ void FPaperRenderSceneProxy::CreateRenderThreadResources(FRHICommandListBase& RH
 	}
 }
 
+void FPaperRenderSceneProxy::DestroyRenderThreadResources()
+{
+	for (FSpriteTextureOverrideRenderProxy* Proxy : MaterialTextureOverrideProxies)
+	{
+		if (Proxy != nullptr)
+		{
+			Proxy->ReleasePrimitiveResource();
+		}
+	}
+	MaterialTextureOverrideProxies.Empty();
+
+	VertexBuffer.ReleaseResource();
+	VertexFactory.ReleaseResource();
+}
+
 void FPaperRenderSceneProxy::DebugDrawCollision(const FSceneView* View, int32 ViewIndex, FMeshElementCollector& Collector, bool bDrawSolid) const
 {
 	if (MyBodySetup != nullptr)
diff --git a/Engine/Plugins/2D/Paper2D/Source/Paper2D/Private/PaperRenderSceneProxy.h b/Engine/Plugins/2D/Paper2D/Source/Paper2D/Private/PaperRenderSceneProxy.h
--- a/Engine/Plugins/2D/Paper2D/Source/Paper2D/Private/PaperRenderSceneProxy.h
+++ b/Engine/Plugins/2D/Paper2D/Source/Paper2D/Private/PaperRenderSceneProxy.h
@@ -125,7 +125,6 @@ public:
 	SIZE_T GetTypeHash() const override;
 
 	FPaperRenderSceneProxy(const UPrimitiveComponent* InComponent);
-	virtual ~FPaperRenderSceneProxy();
 
 	// FPrimitiveSceneProxy interface.
 	virtual void GetDynamicMeshElements(const TArray<const FSceneView*>& Views, const FSceneViewFamily& ViewFamily, uint32 VisibilityMap, FMeshElementCollector& Collector) const override;
@@ -134,6 +133,7 @@ public:
 	virtual bool CanBeOccluded() const override;
 	virtual bool IsUsingDistanceCullFade() const override;
 	virtual void CreateRenderThreadResources(FRHICommandListBase& RHICmdList) override;
+	virtual void DestroyRenderThreadResources() override;
 	// End of FPrimitiveSceneProxy interface.
 
 	void SetBodySetup_RenderThread(UBodySetup* NewSetup);
diff --git a/Engine/Source/Runtime/Engine/Private/ExternalTexture.cpp b/Engine/Source/Runtime/Engine/Private/ExternalTexture.cpp
--- a/Engine/Source/Runtime/Engine/Private/ExternalTexture.cpp
+++ b/Engine/Source/Runtime/Engine/Private/ExternalTexture.cpp
@@ -32,7 +32,7 @@ void FExternalTextureRegistry::RegisterExternalTexture(const FGuid& InGuid, FTex
 
 	for (const FMaterialRenderProxy* MaterialRenderProxy : ReferencingMaterialRenderProxies)
 	{
-		const_cast<FMaterialRenderProxy*>(MaterialRenderProxy)->CacheUniformExpressions(false);
+		const_cast<FMaterialRenderProxy*>(MaterialRenderProxy)->CacheUniformExpressions(FRHICommandListImmediate::Get(), false);
 	}
 }
 
@@ -45,14 +45,13 @@ void FExternalTextureRegistry::UnregisterExternalTexture(const FGuid& InGuid)
 
 	for (const FMaterialRenderProxy* MaterialRenderProxy : ReferencingMaterialRenderProxies)
 	{
-		const_cast<FMaterialRenderProxy*>(MaterialRenderProxy)->CacheUniformExpressions(false);
+		const_cast<FMaterialRenderProxy*>(MaterialRenderProxy)->CacheUniformExpressions(FRHICommandListImmediate::Get(), false);
 	}
 }
 
 
 void FExternalTextureRegistry::RemoveMaterialRenderProxyReference(const FMaterialRenderProxy* MaterialRenderProxy)
 {
-	check(IsInRenderingThread());
 	FScopeLock Lock(&CriticalSection);
 	ReferencingMaterialRenderProxies.Remove(MaterialRenderProxy);
 }
diff --git a/Engine/Source/Runtime/Engine/Private/Materials/MaterialInstance.cpp b/Engine/Source/Runtime/Engine/Private/Materials/MaterialInstance.cpp
--- a/Engine/Source/Runtime/Engine/Private/Materials/MaterialInstance.cpp
+++ b/Engine/Source/Runtime/Engine/Private/Materials/MaterialInstance.cpp
@@ -461,7 +461,7 @@ void GameThread_UpdateMIParameter(const UMaterialInstance* Instance, const Param
 			[Resource, ParameterInfo, Value](FRHICommandListImmediate& RHICmdList)
 			{
 				Resource->RenderThread_UpdateParameter(ParameterInfo, Value);
-				Resource->CacheUniformExpressions(false);
+				Resource->CacheUniformExpressions(RHICmdList, false);
 			});
 	}
 }
@@ -5132,7 +5132,7 @@ void UMaterialInstance::OverrideTextureParameterValue(const UTexture* InTextureT
 					(
 						[LocalResource](FRHICommandListImmediate& RHICmdList)
 						{
-							LocalResource->CacheUniformExpressions(false);
+							LocalResource->CacheUniformExpressions(RHICmdList, false);
 						}
 				);
 			}
diff --git a/Engine/Source/Runtime/Engine/Private/Materials/MaterialRenderProxy.cpp b/Engine/Source/Runtime/Engine/Private/Materials/MaterialRenderProxy.cpp
--- a/Engine/Source/Runtime/Engine/Private/Materials/MaterialRenderProxy.cpp
+++ b/Engine/Source/Runtime/Engine/Private/Materials/MaterialRenderProxy.cpp
@@ -353,10 +353,8 @@ IAllocatedVirtualTexture* FMaterialRenderProxy::AllocateVTStack(const FMaterialR
 }
 
 
-void FMaterialRenderProxy::EvaluateUniformExpressions(FUniformExpressionCache& OutUniformExpressionCache, const FMaterialRenderContext& Context, FUniformExpressionCacheAsyncUpdater* Updater) const
+void FMaterialRenderProxy::EvaluateUniformExpressions(FUniformExpressionCache& OutUniformExpressionCache, const FMaterialRenderContext& Context, FUniformExpressionCacheAsyncUpdater* Updater, FRHICommandListBase* RHICmdList) const
 {
-	check(IsInRenderingThread());
-
 	SCOPE_CYCLE_COUNTER(STAT_CacheUniformExpressions);
 
 	TRACE_CPUPROFILER_EVENT_SCOPE(FMaterialRenderProxy::EvaluateUniformExpressions);
@@ -430,7 +428,8 @@ void FMaterialRenderProxy::EvaluateUniformExpressions(FUniformExpressionCache& O
 
 		if (IsValidRef(OutUniformExpressionCache.UniformBuffer))
 		{
-			FRHICommandListImmediate::Get().UpdateUniformBuffer(OutUniformExpressionCache.UniformBuffer, TempBuffer);
+			check(RHICmdList);
+			RHICmdList->UpdateUniformBuffer(OutUniformExpressionCache.UniformBuffer, TempBuffer);
 		}
 		else
 		{
@@ -446,10 +445,10 @@ void FMaterialRenderProxy::EvaluateUniformExpressions(FUniformExpressionCache& O
 	++UniformExpressionCacheSerialNumber;
 }
 
-void FMaterialRenderProxy::CacheUniformExpressions(bool bRecreateUniformBuffer)
+void FMaterialRenderProxy::CacheUniformExpressions(FRHICommandListBase& RHICmdList, bool bRecreateUniformBuffer)
 {
 	// Register the render proxy's as a render resource so it can receive notifications to free the uniform buffer.
-	InitResource(FRHICommandListImmediate::Get());
+	InitResource(RHICmdList);
 
 	bool bUsingNewLoader = FPlatformProperties::RequiresCookedData();
 
@@ -462,13 +461,16 @@ void FMaterialRenderProxy::CacheUniformExpressions(bool bRecreateUniformBuffer)
 		UE_LOG(LogMaterial, Fatal, TEXT("Cannot queue the Expression Cache for Material %s when it is about to be deleted"), *MaterialName);
 	}
 	StartCacheUniformExpressions();
+
+	DeferredUniformExpressionCacheRequestsMutex.Lock();
 	DeferredUniformExpressionCacheRequests.Add(this);
+	DeferredUniformExpressionCacheRequestsMutex.Unlock();
 
 	InvalidateUniformExpressionCache(bRecreateUniformBuffer);
 
 	if (!GDeferUniformExpressionCaching)
 	{
-		FMaterialRenderProxy::UpdateDeferredCachedUniformExpressions();
+		FMaterialRenderProxy::UpdateDeferredCachedUniformExpressions(RHICmdList);
 	}
 }
 
@@ -481,15 +483,14 @@ void FMaterialRenderProxy::CacheUniformExpressions_GameThread(bool bRecreateUnif
 		FMaterialRenderProxy* RenderProxy = this;
 		ENQUEUE_RENDER_COMMAND(FCacheUniformExpressionsCommand)(
 			[RenderProxy, bRecreateUniformBuffer](FRHICommandListImmediate& RHICmdList)
-			{
-				RenderProxy->CacheUniformExpressions(bRecreateUniformBuffer);
-			});
+		{
+			RenderProxy->CacheUniformExpressions(RHICmdList, bRecreateUniformBuffer);
+		});
 	}
 }
 
 void FMaterialRenderProxy::InvalidateUniformExpressionCache(bool bRecreateUniformBuffer)
 {
-	check(IsInRenderingThread());
 	GUniformExpressionCacheAsyncUpdateTask.Wait();
 
 #if WITH_EDITOR
@@ -518,7 +519,7 @@ void FMaterialRenderProxy::InvalidateUniformExpressionCache(bool bRecreateUnifor
 	}
 }
 
-void FMaterialRenderProxy::UpdateUniformExpressionCacheIfNeeded(ERHIFeatureLevel::Type InFeatureLevel) const
+void FMaterialRenderProxy::UpdateUniformExpressionCacheIfNeeded(FRHICommandListBase& RHICmdList, ERHIFeatureLevel::Type InFeatureLevel) const
 {
 	// Don't cache uniform expressions if an entirely different FMaterialRenderProxy is going to be used for rendering
 	const FMaterial* Material = GetMaterialNoFallback(InFeatureLevel);
@@ -528,7 +529,7 @@ void FMaterialRenderProxy::UpdateUniformExpressionCacheIfNeeded(ERHIFeatureLevel
 	{
 		FMaterialRenderContext MaterialRenderContext(this, *Material, nullptr);
 		MaterialRenderContext.bShowSelection = GIsEditor;
-		EvaluateUniformExpressions(UniformExpressionCache[InFeatureLevel], MaterialRenderContext);
+		EvaluateUniformExpressions(UniformExpressionCache[InFeatureLevel], MaterialRenderContext, nullptr, &RHICmdList);
 	}
 }
 
@@ -551,13 +552,11 @@ FMaterialRenderProxy::~FMaterialRenderProxy()
 
 	if (IsInitialized())
 	{
-		check(IsInRenderingThread());
 		ReleaseResource();
 	}
 
 	if (HasVirtualTextureCallbacks)
 	{
-		check(IsInRenderingThread());
 		GetRendererModule().RemoveAllVirtualTextureProducerDestroyedCallbacks(this);
 		HasVirtualTextureCallbacks = false;
 	}
@@ -579,7 +578,9 @@ void FMaterialRenderProxy::InitRHI(FRHICommandListBase& RHICmdList)
 
 void FMaterialRenderProxy::CancelCacheUniformExpressions()
 {
+	DeferredUniformExpressionCacheRequestsMutex.Lock();
 	DeferredUniformExpressionCacheRequests.Remove(this);
+	DeferredUniformExpressionCacheRequestsMutex.Unlock();
 }
 
 void FMaterialRenderProxy::ReleaseRHI()
@@ -592,7 +593,11 @@ void FMaterialRenderProxy::ReleaseRHI()
 	}
 #endif // WITH_EDITOR
 
-	if (DeferredUniformExpressionCacheRequests.Remove(this))
+	DeferredUniformExpressionCacheRequestsMutex.Lock();
+	bool bRemoved = DeferredUniformExpressionCacheRequests.Remove(this) != 0;
+	DeferredUniformExpressionCacheRequestsMutex.Unlock();
+
+	if (bRemoved)
 	{
 		// Notify that we're finished with this inflight cache request, because the object is being released
 		FinishCacheUniformExpressions();
@@ -656,18 +661,18 @@ const FMaterial& FMaterialRenderProxy::GetIncompleteMaterialWithFallback(ERHIFea
 	return *Material;
 }
 
-void FMaterialRenderProxy::UpdateDeferredCachedUniformExpressions()
+void FMaterialRenderProxy::UpdateDeferredCachedUniformExpressions(FRHICommandListBase& RHICmdList)
 {
 	LLM_SCOPE(ELLMTag::Materials);
 
-	check(IsInRenderingThread());
-
 	TRACE_CPUPROFILER_EVENT_SCOPE(FMaterialRenderProxy::UpdateDeferredCachedUniformExpressions);
 	CSV_SCOPED_TIMING_STAT_EXCLUSIVE(Material_UpdateDeferredCachedUniformExpressions);
 	QUICK_SCOPE_CYCLE_COUNTER(STAT_UpdateDeferredCachedUniformExpressions);
 
 	FUniformExpressionCacheAsyncUpdater Updater;
-	FUniformExpressionCacheAsyncUpdater* UpdaterIfEnabled = GUniformExpressionCacheAsyncUpdateTask.IsEnabled() ? &Updater : nullptr;
+	FUniformExpressionCacheAsyncUpdater* UpdaterIfEnabled = RHICmdList.IsImmediate() && GUniformExpressionCacheAsyncUpdateTask.IsEnabled() ? &Updater : nullptr;
+
+	UE::TScopeLock Lock(DeferredUniformExpressionCacheRequestsMutex);
 
 	for (TSet<FMaterialRenderProxy*>::TConstIterator It(DeferredUniformExpressionCacheRequests); It; ++It)
 	{
@@ -678,23 +683,23 @@ void FMaterialRenderProxy::UpdateDeferredCachedUniformExpressions()
 		}
 
 		UMaterialInterface::IterateOverActiveFeatureLevels([&](ERHIFeatureLevel::Type InFeatureLevel)
+		{
+			// Don't bother caching if we'll be falling back to a different FMaterialRenderProxy for rendering anyway
+			const FMaterial* Material = MaterialProxy->GetMaterialNoFallback(InFeatureLevel);
+			if (Material && Material->GetRenderingThreadShaderMap())
 			{
-				// Don't bother caching if we'll be falling back to a different FMaterialRenderProxy for rendering anyway
-				const FMaterial* Material = MaterialProxy->GetMaterialNoFallback(InFeatureLevel);
-				if (Material && Material->GetRenderingThreadShaderMap())
-				{
-					FMaterialRenderContext MaterialRenderContext(MaterialProxy, *Material, nullptr);
-					MaterialRenderContext.bShowSelection = GIsEditor;
-					MaterialProxy->EvaluateUniformExpressions(MaterialProxy->UniformExpressionCache[(int32)InFeatureLevel], MaterialRenderContext, UpdaterIfEnabled);
-				}
-			});
+				FMaterialRenderContext MaterialRenderContext(MaterialProxy, *Material, nullptr);
+				MaterialRenderContext.bShowSelection = GIsEditor;
+				MaterialProxy->EvaluateUniformExpressions(MaterialProxy->UniformExpressionCache[(int32)InFeatureLevel], MaterialRenderContext, UpdaterIfEnabled, &RHICmdList);
+			}
+		});
 
 		MaterialProxy->FinishCacheUniformExpressions();
 	}
 
 	if (UpdaterIfEnabled)
 	{
-		Updater.Update(FRHICommandListExecutor::GetImmediateCommandList());
+		Updater.Update(FRHICommandListImmediate::Get(RHICmdList));
 	}
 
 	DeferredUniformExpressionCacheRequests.Reset();
@@ -702,6 +707,7 @@ void FMaterialRenderProxy::UpdateDeferredCachedUniformExpressions()
 
 bool FMaterialRenderProxy::HasDeferredUniformExpressionCacheRequests()
 {
+	UE::TScopeLock Lock(DeferredUniformExpressionCacheRequestsMutex);
 	return DeferredUniformExpressionCacheRequests.Num() > 0;
 }
 
@@ -710,6 +716,7 @@ TSet<FMaterialRenderProxy*> FMaterialRenderProxy::MaterialRenderProxyMap;
 FCriticalSection FMaterialRenderProxy::MaterialRenderProxyMapLock;
 #endif // WITH_EDITOR
 TSet<FMaterialRenderProxy*> FMaterialRenderProxy::DeferredUniformExpressionCacheRequests;
+UE::FMutex FMaterialRenderProxy::DeferredUniformExpressionCacheRequestsMutex;
 
 /*-----------------------------------------------------------------------------
 	FColoredMaterialRenderProxy
diff --git a/Engine/Source/Runtime/Engine/Private/Materials/MaterialShared.cpp b/Engine/Source/Runtime/Engine/Private/Materials/MaterialShared.cpp
--- a/Engine/Source/Runtime/Engine/Private/Materials/MaterialShared.cpp
+++ b/Engine/Source/Runtime/Engine/Private/Materials/MaterialShared.cpp
@@ -4967,7 +4967,7 @@ void SetShaderMapsOnMaterialResources_RenderThread(FRHICommandListImmediate& RHI
 				// Using ContainsByHash so we can pass a raw-ptr to TMap method that wants a TRefCountPtr
 				if (Material && Material->GetRenderingThreadShaderMap() && MaterialsToUpdate.ContainsByHash(GetTypeHash(Material), Material))
 				{
-					MaterialProxy->CacheUniformExpressions(true);
+					MaterialProxy->CacheUniformExpressions(RHICmdList, true);
 					bFoundAnyInitializedMaterials = true;
 
 					/*const FMaterial& MaterialForRendering = *MaterialProxy->GetMaterial(MaterialFeatureLevel);
diff --git a/Engine/Source/Runtime/Engine/Private/TextureCompiler.cpp b/Engine/Source/Runtime/Engine/Private/TextureCompiler.cpp
--- a/Engine/Source/Runtime/Engine/Private/TextureCompiler.cpp
+++ b/Engine/Source/Runtime/Engine/Private/TextureCompiler.cpp
@@ -416,7 +416,7 @@ void FTextureCompilingManager::PostCompilation(TArrayView<UTexture* const> InCom
 							ENQUEUE_RENDER_COMMAND(TextureCompiler_RecacheUniformExpressions)(
 								[RenderProxy](FRHICommandListImmediate& RHICmdList)
 								{
-									RenderProxy->CacheUniformExpressions(false);
+									RenderProxy->CacheUniformExpressions(RHICmdList, false);
 								});
 						}
 					}
diff --git a/Engine/Source/Runtime/Engine/Public/Materials/MaterialRenderProxy.h b/Engine/Source/Runtime/Engine/Public/Materials/MaterialRenderProxy.h
--- a/Engine/Source/Runtime/Engine/Public/Materials/MaterialRenderProxy.h
+++ b/Engine/Source/Runtime/Engine/Public/Materials/MaterialRenderProxy.h
@@ -7,6 +7,7 @@
 #include "Math/Vector2D.h"
 #include "RenderResource.h"
 #include "RHIImmutableSamplerState.h"
+#include "Async/Mutex.h"
 
 enum class EMaterialParameterType : uint8;
 
@@ -106,23 +107,22 @@ public:
 	/** Destructor. */
 	ENGINE_API virtual ~FMaterialRenderProxy();
 
-	UE_DEPRECATED(5.1, "EvaluateUniformExpressions with a command list is deprecated.")
-	void EvaluateUniformExpressions(FUniformExpressionCache& OutUniformExpressionCache, const FMaterialRenderContext& Context, class FRHIComputeCommandList*) const
-	{
-		EvaluateUniformExpressions(OutUniformExpressionCache, Context);
-	}
-
 	/**
 	 * Evaluates uniform expressions and stores them in OutUniformExpressionCache.
 	 * @param OutUniformExpressionCache - The uniform expression cache to build.
 	 * @param MaterialRenderContext - The context for which to cache expressions.
 	 */
-	ENGINE_API void EvaluateUniformExpressions(FUniformExpressionCache& OutUniformExpressionCache, const FMaterialRenderContext& Context, FUniformExpressionCacheAsyncUpdater* Updater = nullptr) const;
+	ENGINE_API void EvaluateUniformExpressions(FUniformExpressionCache& OutUniformExpressionCache, const FMaterialRenderContext& Context, FUniformExpressionCacheAsyncUpdater* Updater = nullptr, FRHICommandListBase* RHICmdList = nullptr) const;
 
 	/**
 	 * Caches uniform expressions for efficient runtime evaluation.
 	 */
-	ENGINE_API void CacheUniformExpressions(bool bRecreateUniformBuffer);
+	ENGINE_API void CacheUniformExpressions(FRHICommandListBase& RHICmdList, bool bRecreateUniformBuffer);
+
+	inline void CacheUniformExpressions(bool bRecreateUniformBuffer)
+	{
+		CacheUniformExpressions(FRHICommandListImmediate::Get(), bRecreateUniformBuffer);
+	}
 
 	/** Cancels an in-flight cache operation. */
 	ENGINE_API void CancelCacheUniformExpressions();
@@ -140,8 +140,12 @@ public:
 	 */
 	ENGINE_API void InvalidateUniformExpressionCache(bool bRecreateUniformBuffer);
 
-	ENGINE_API void UpdateUniformExpressionCacheIfNeeded(ERHIFeatureLevel::Type InFeatureLevel) const;
+	inline void UpdateUniformExpressionCacheIfNeeded(ERHIFeatureLevel::Type InFeatureLevel) const
+	{
+		UpdateUniformExpressionCacheIfNeeded(FRHICommandListImmediate::Get(), InFeatureLevel);
+	}
 
+	ENGINE_API void UpdateUniformExpressionCacheIfNeeded(FRHICommandListBase& RHICmdList, ERHIFeatureLevel::Type InFeatureLevel) const;
 
 	/** Returns the FMaterial, without using a fallback if the FMaterial doesn't have a valid shader map. Can return NULL. */
 	virtual const FMaterial* GetMaterialNoFallback(ERHIFeatureLevel::Type InFeatureLevel) const = 0;
@@ -213,7 +217,8 @@ public:
 	const USpecularProfile* GetSpecularProfileRT(uint32 Index) const { check(Index<uint32(SpecularProfilesRT.Num())); return SpecularProfilesRT[Index]; }
 	const uint32 NumSpecularProfileRT() const { return SpecularProfilesRT.Num(); }
 
-	static ENGINE_API void UpdateDeferredCachedUniformExpressions();
+	static void UpdateDeferredCachedUniformExpressions() { UpdateDeferredCachedUniformExpressions(FRHICommandListImmediate::Get()); }
+	static ENGINE_API void UpdateDeferredCachedUniformExpressions(FRHICommandListBase& RHICmdList);
 
 	static ENGINE_API bool HasDeferredUniformExpressionCacheRequests();
 
@@ -257,6 +262,7 @@ private:
 #endif
 
 	static ENGINE_API TSet<FMaterialRenderProxy*> DeferredUniformExpressionCacheRequests;
+	static ENGINE_API UE::FMutex DeferredUniformExpressionCacheRequestsMutex;
 };
 
 /**
diff --git a/Engine/Source/Runtime/Engine/Public/Rendering/StaticLightingSystemInterface.h b/Engine/Source/Runtime/Engine/Public/Rendering/StaticLightingSystemInterface.h
--- a/Engine/Source/Runtime/Engine/Public/Rendering/StaticLightingSystemInterface.h
+++ b/Engine/Source/Runtime/Engine/Public/Rendering/StaticLightingSystemInterface.h
@@ -32,7 +32,7 @@ public:
 	DECLARE_MULTICAST_DELEGATE_OneParam(FLightComponentBasedSignature, ULightComponentBase* /*InComponent*/);
 	DECLARE_MULTICAST_DELEGATE_TwoParams(FStationaryLightChannelReassignmentSignature, ULightComponentBase* /*InComponent*/, int32 /*NewShadowMapChannel*/);
 	DECLARE_MULTICAST_DELEGATE(FLightmassImportanceVolumeModifiedSignature);
-	DECLARE_MULTICAST_DELEGATE_OneParam(FMaterialInvalidationSignature, FMaterialRenderProxy* /*Material*/);
+	DECLARE_TS_MULTICAST_DELEGATE_OneParam(FMaterialInvalidationSignature, FMaterialRenderProxy* /*Material*/);
 
 	static ENGINE_API FPrimitiveComponentBasedSignature OnPrimitiveComponentRegistered;
 	static ENGINE_API FPrimitiveComponentBasedSignature OnPrimitiveComponentUnregistered;
