Speculative linux build fix.

diff --git a/Engine/Source/Runtime/Engine/Public/MeshDrawCommandStatsComponentData.h b/Engine/Source/Runtime/Engine/Public/MeshDrawCommandStatsComponentData.h
--- a/Engine/Source/Runtime/Engine/Public/MeshDrawCommandStatsComponentData.h
+++ b/Engine/Source/Runtime/Engine/Public/MeshDrawCommandStatsComponentData.h
@@ -36,14 +36,14 @@ using FMeshDrawCommandStatsComponentDataID = int16;
 /**
  * Manages all created FMeshDrawCommandStatsComponentData and exposed unique FMeshDrawCommandStatsComponentDataID to retrieve the data again during stat processing
  */
-class FMeshDrawCommandStatsComponentDataManager
+ENGINE_API class FMeshDrawCommandStatsComponentDataManager
 {
 public:
-	ENGINE_API static void CreateInstance();
-	ENGINE_API static FMeshDrawCommandStatsComponentDataManager* Get() { return Instance; }
+	static void CreateInstance();
+	static FMeshDrawCommandStatsComponentDataManager* Get() { return Instance; }
 
-	ENGINE_API FMeshDrawCommandStatsComponentDataID GetID(const FMeshDrawCommandStatsComponentData& ComponentData);
-	ENGINE_API FMeshDrawCommandStatsComponentData GetComponentData(FMeshDrawCommandStatsComponentDataID ID) const
+	FMeshDrawCommandStatsComponentDataID GetID(const FMeshDrawCommandStatsComponentData& ComponentData);
+	FMeshDrawCommandStatsComponentData GetComponentData(FMeshDrawCommandStatsComponentDataID ID) const
 	{
 		if (ID == INDEX_NONE)
 		{
