Speculative linux build fix.

diff --git a/Engine/Source/Runtime/Engine/Public/MeshDrawCommandStatsComponentData.h b/Engine/Source/Runtime/Engine/Public/MeshDrawCommandStatsComponentData.h
--- a/Engine/Source/Runtime/Engine/Public/MeshDrawCommandStatsComponentData.h
+++ b/Engine/Source/Runtime/Engine/Public/MeshDrawCommandStatsComponentData.h
@@ -36,7 +36,7 @@ using FMeshDrawCommandStatsComponentDataID = int16;
 /**
  * Manages all created FMeshDrawCommandStatsComponentData and exposed unique FMeshDrawCommandStatsComponentDataID to retrieve the data again during stat processing
  */
-ENGINE_API class FMeshDrawCommandStatsComponentDataManager
+class ENGINE_API FMeshDrawCommandStatsComponentDataManager
 {
 public:
 	static void CreateInstance();
