AssetManager bugfix: fix typo in ModifyCook that was preventing NeverCook and ProductionNeverCook packages from being added to the cooker's list of NeverCook packages.
#rnx

diff --git a/Engine/Source/Runtime/Engine/Private/AssetManager.cpp b/Engine/Source/Runtime/Engine/Private/AssetManager.cpp
--- a/Engine/Source/Runtime/Engine/Private/AssetManager.cpp
+++ b/Engine/Source/Runtime/Engine/Private/AssetManager.cpp
@@ -4209,7 +4209,7 @@ void UAssetManager::ModifyCook(TConstArrayView<const ITargetPlatform*> TargetPla
 				PackagesToNeverCookSet.Add(PackageName, &bAlreadyInSet);
 				if (!bAlreadyInSet)
 				{
-					PackagesToNeverCookSet.Add(PackageName);
+					PackagesToNeverCook.Add(PackageName);
 				}
 			}
 		}
