Fixed build break.

diff --git a/Engine/Source/Runtime/Engine/Private/Materials/MaterialRenderProxy.cpp b/Engine/Source/Runtime/Engine/Private/Materials/MaterialRenderProxy.cpp
--- a/Engine/Source/Runtime/Engine/Private/Materials/MaterialRenderProxy.cpp
+++ b/Engine/Source/Runtime/Engine/Private/Materials/MaterialRenderProxy.cpp
@@ -474,6 +474,11 @@ void FMaterialRenderProxy::CacheUniformExpressions(FRHICommandListBase& RHICmdLi
 	}
 }
 
+void FMaterialRenderProxy::CacheUniformExpressions(bool bRecreateUniformBuffer)
+{
+	CacheUniformExpressions(FRHICommandListImmediate::Get(), bRecreateUniformBuffer);
+}
+
 void FMaterialRenderProxy::CacheUniformExpressions_GameThread(bool bRecreateUniformBuffer)
 {
 	if (FApp::CanEverRender())
@@ -533,6 +538,11 @@ void FMaterialRenderProxy::UpdateUniformExpressionCacheIfNeeded(FRHICommandListB
 	}
 }
 
+void FMaterialRenderProxy::UpdateUniformExpressionCacheIfNeeded(ERHIFeatureLevel::Type InFeatureLevel) const
+{
+	UpdateUniformExpressionCacheIfNeeded(FRHICommandListImmediate::Get(), InFeatureLevel);
+}
+
 FMaterialRenderProxy::FMaterialRenderProxy(FString InMaterialName)
 	: SubsurfaceProfileRT(nullptr)
 	, MaterialName(MoveTemp(InMaterialName))
@@ -705,6 +715,11 @@ void FMaterialRenderProxy::UpdateDeferredCachedUniformExpressions(FRHICommandLis
 	DeferredUniformExpressionCacheRequests.Reset();
 }
 
+void FMaterialRenderProxy::UpdateDeferredCachedUniformExpressions()
+{
+	UpdateDeferredCachedUniformExpressions(FRHICommandListImmediate::Get());
+}
+
 bool FMaterialRenderProxy::HasDeferredUniformExpressionCacheRequests()
 {
 	UE::TScopeLock Lock(DeferredUniformExpressionCacheRequestsMutex);
diff --git a/Engine/Source/Runtime/Engine/Public/Materials/MaterialRenderProxy.h b/Engine/Source/Runtime/Engine/Public/Materials/MaterialRenderProxy.h
--- a/Engine/Source/Runtime/Engine/Public/Materials/MaterialRenderProxy.h
+++ b/Engine/Source/Runtime/Engine/Public/Materials/MaterialRenderProxy.h
@@ -119,10 +119,8 @@ public:
 	 */
 	ENGINE_API void CacheUniformExpressions(FRHICommandListBase& RHICmdList, bool bRecreateUniformBuffer);
 
-	inline void CacheUniformExpressions(bool bRecreateUniformBuffer)
-	{
-		CacheUniformExpressions(FRHICommandListImmediate::Get(), bRecreateUniformBuffer);
-	}
+	UE_DEPRECATED(5.4, "CacheUniformExpressions now requires a command list.")
+	ENGINE_API void CacheUniformExpressions(bool bRecreateUniformBuffer);
 
 	/** Cancels an in-flight cache operation. */
 	ENGINE_API void CancelCacheUniformExpressions();
@@ -140,11 +138,7 @@ public:
 	 */
 	ENGINE_API void InvalidateUniformExpressionCache(bool bRecreateUniformBuffer);
 
-	inline void UpdateUniformExpressionCacheIfNeeded(ERHIFeatureLevel::Type InFeatureLevel) const
-	{
-		UpdateUniformExpressionCacheIfNeeded(FRHICommandListImmediate::Get(), InFeatureLevel);
-	}
-
+	ENGINE_API void UpdateUniformExpressionCacheIfNeeded(ERHIFeatureLevel::Type InFeatureLevel) const;
 	ENGINE_API void UpdateUniformExpressionCacheIfNeeded(FRHICommandListBase& RHICmdList, ERHIFeatureLevel::Type InFeatureLevel) const;
 
 	/** Returns the FMaterial, without using a fallback if the FMaterial doesn't have a valid shader map. Can return NULL. */
@@ -217,7 +211,7 @@ public:
 	const USpecularProfile* GetSpecularProfileRT(uint32 Index) const { check(Index<uint32(SpecularProfilesRT.Num())); return SpecularProfilesRT[Index]; }
 	const uint32 NumSpecularProfileRT() const { return SpecularProfilesRT.Num(); }
 
-	static void UpdateDeferredCachedUniformExpressions() { UpdateDeferredCachedUniformExpressions(FRHICommandListImmediate::Get()); }
+	static ENGINE_API void UpdateDeferredCachedUniformExpressions();
 	static ENGINE_API void UpdateDeferredCachedUniformExpressions(FRHICommandListBase& RHICmdList);
 
 	static ENGINE_API bool HasDeferredUniformExpressionCacheRequests();
