Partial backout of 26614803 for "fix non spatial newly created/saved actor being unloaded after PIE".
When an actor descriptor is added, add it to the always loaded loader if its non spatially loaded.
Also fix duplicate actor descriptors in the dirty actor list since recent changes.

diff --git a/Engine/Source/Editor/WorldPartitionEditor/Private/WorldPartition/SWorldPartitionEditorGrid2D.cpp b/Engine/Source/Editor/WorldPartitionEditor/Private/WorldPartition/SWorldPartitionEditorGrid2D.cpp
--- a/Engine/Source/Editor/WorldPartitionEditor/Private/WorldPartition/SWorldPartitionEditorGrid2D.cpp
+++ b/Engine/Source/Editor/WorldPartitionEditor/Private/WorldPartition/SWorldPartitionEditorGrid2D.cpp
@@ -1087,19 +1087,19 @@ void SWorldPartitionEditorGrid2D::Tick(const FGeometry& AllottedGeometry, const
 
 	// Also include transient actor loader adapters that might have been spawned by blutilities, etc. Since these actors can't be saved because they are transient,
 	// they will never get an actor descriptor so they will never appear in the world partition editor. Also include unsaved, newly created actors for convenience.
-	for (auto& [Reference, Actor] : GetWorldPartition()->GetDirtyActors())
+	for (auto& [DirtyActorGuid, DirtyActor] : GetWorldPartition()->GetDirtyActors())
 	{
-		if (Actor.IsValid())
+		if (DirtyActor.ActorPtr.IsValid())
 		{
-			if (!GetWorldPartition()->GetActorDesc(Actor->GetActorGuid()) && Actor->Implements<UWorldPartitionActorLoaderInterface>())
+			if (!GetWorldPartition()->GetActorDesc(DirtyActor.ActorPtr->GetActorGuid()) && DirtyActor.ActorPtr->Implements<UWorldPartitionActorLoaderInterface>())
 			{
-				if (IWorldPartitionActorLoaderInterface::ILoaderAdapter* LoaderAdapter = Cast<IWorldPartitionActorLoaderInterface>(Actor)->GetLoaderAdapter())
+				if (IWorldPartitionActorLoaderInterface::ILoaderAdapter* LoaderAdapter = Cast<IWorldPartitionActorLoaderInterface>(DirtyActor.ActorPtr)->GetLoaderAdapter())
 				{
-					ShownLoaderInterfaces.Add(Actor.Get());
+					ShownLoaderInterfaces.Add(DirtyActor.ActorPtr.Get());
 				}
 			}
 
-			DirtyActorGuids.Add(Actor->GetActorGuid());
+			DirtyActorGuids.Add(DirtyActorGuid);
 		}
 	}
 
diff --git a/Engine/Source/Runtime/Engine/Private/WorldPartition/WorldPartition.cpp b/Engine/Source/Runtime/Engine/Private/WorldPartition/WorldPartition.cpp
--- a/Engine/Source/Runtime/Engine/Private/WorldPartition/WorldPartition.cpp
+++ b/Engine/Source/Runtime/Engine/Private/WorldPartition/WorldPartition.cpp
@@ -258,7 +258,13 @@ public:
 		bIncludeSpatiallyLoadedActors = false;
 		bIncludeNonSpatiallyLoadedActors = true;
 	}
+
+	void RefreshLoadedState()
+	{
+		FLoaderAdapterShape::RefreshLoadedState();
+	}
 };
+
 UWorldPartition::FWorldPartitionChangedEvent UWorldPartition::WorldPartitionChangedEvent;
 #endif
 
@@ -353,13 +359,13 @@ void UWorldPartition::OnPackageDirtyStateChanged(UPackage* Package)
 		{
 			if (Package->IsDirty())
 			{
-				DirtyActors.Add(FDirtyActor(ActorHandle.ToReference(), Actor));
+				DirtyActors.Add(Actor->GetActorGuid(), FDirtyActor(ActorHandle.ToReference(), Actor));
 			}
 		}
 		else
 		{
 			// This is handling a new actor (unsaved).
-			DirtyActors.Add(FDirtyActor(Actor));
+			DirtyActors.Add(Actor->GetActorGuid(), FDirtyActor(Actor));
 		}
 	}
 }
@@ -1239,7 +1245,16 @@ void UWorldPartition::OnActorDescAdded(FWorldPartitionActorDesc* NewActorDesc)
 
 	if (AActor* NewActor = NewActorDesc->GetActor())
 	{
-		DirtyActors.Add(FDirtyActor(FWorldPartitionReference(NewActorDesc->GetContainer(), NewActorDesc->GetGuid()), NewActor));
+		if (FDirtyActor* ExistingDirtyActor = DirtyActors.Find(NewActorDesc->GetGuid()))
+		{
+			check(ExistingDirtyActor->ActorPtr == NewActor);
+			check(!ExistingDirtyActor->WorldPartitionRef.IsSet());
+			ExistingDirtyActor->WorldPartitionRef = FWorldPartitionReference(NewActorDesc->GetContainer(), NewActorDesc->GetGuid());
+		}
+		else
+		{
+			DirtyActors.Add(NewActorDesc->GetGuid(), FDirtyActor(FWorldPartitionReference(NewActorDesc->GetContainer(), NewActorDesc->GetGuid()), NewActor));
+		}
 	}
 
 	if (ForceLoadedActors)
@@ -1247,6 +1262,11 @@ void UWorldPartition::OnActorDescAdded(FWorldPartitionActorDesc* NewActorDesc)
 		ForceLoadedActors->AddActors({ NewActorDesc->GetGuid() });
 	}
 
+	if (AlwaysLoadedActors && !NewActorDesc->GetIsSpatiallyLoaded())
+	{
+		AlwaysLoadedActors->RefreshLoadedState();
+	}
+
 	if (WorldPartitionEditor)
 	{
 		WorldPartitionEditor->Refresh();
@@ -1266,6 +1286,11 @@ void UWorldPartition::OnActorDescRemoved(FWorldPartitionActorDesc* ActorDesc)
 	{
 		ForceLoadedActors->RemoveActors({ ActorDesc->GetGuid() });
 	}
+
+	if (AlwaysLoadedActors && !ActorDesc->GetIsSpatiallyLoaded())
+	{
+		AlwaysLoadedActors->RefreshLoadedState();
+	}
 		
 	if (WorldPartitionEditor)
 	{
@@ -1404,9 +1429,9 @@ void UWorldPartition::OnContentBundleRemovedContent(const FContentBundleEditor*
 	if (Contains(ContentBundleActorDescContainer.Get()->GetContainerPackage()))
 	{
 		// This is handling a new actor (unsaved) from Content Bundle.
-		for (TSet<FDirtyActor>::TIterator DirtyActorIt(DirtyActors); DirtyActorIt; ++DirtyActorIt)
+		for (TMap<FGuid, FDirtyActor>::TIterator DirtyActorIt(DirtyActors); DirtyActorIt; ++DirtyActorIt)
 		{
-			const FDirtyActor& DirtyActor = *DirtyActorIt;
+			const FDirtyActor& DirtyActor = DirtyActorIt.Value();
 			if (!DirtyActor.WorldPartitionRef.IsSet())
 			{
 				if (DirtyActor.ActorPtr.IsValid() && DirtyActor.ActorPtr->GetContentBundleGuid() == ContentBundleActorDescContainer.Get()->GetContentBundleGuid())
@@ -1528,11 +1553,11 @@ void UWorldPartition::AddReferencedObjects(UObject* InThis, FReferenceCollector&
 	if (IsGarbageCollecting())
 	{
 		Collector.AllowEliminatingReferences(false);
-		for (auto& [ActorReference, Actor] : This->DirtyActors)
+		for (auto& [DirtyActorGuid, DirtyActor] : This->DirtyActors)
 		{
-			if (ActorReference.IsSet() || Actor.IsValid())
+			if (DirtyActor.WorldPartitionRef.IsSet() || DirtyActor.ActorPtr.IsValid())
 			{
-				Collector.AddReferencedObject(Actor);
+				Collector.AddReferencedObject(DirtyActor.ActorPtr);
 			}
 		}
 		Collector.AllowEliminatingReferences(true);
@@ -1555,10 +1580,9 @@ void UWorldPartition::Tick(float DeltaSeconds)
 		EditorHash->Tick(DeltaSeconds);
 	}
 
-	TArray<FWorldPartitionReference> ActorReferences;
-	for (TSet<FDirtyActor>::TIterator DirtyActorIt(DirtyActors); DirtyActorIt; ++DirtyActorIt)
+	for (TMap<FGuid, FDirtyActor>::TIterator DirtyActorIt(DirtyActors); DirtyActorIt; ++DirtyActorIt)
 	{
-		const FDirtyActor& DirtyActor = *DirtyActorIt;
+		const FDirtyActor& DirtyActor = DirtyActorIt.Value();
 		if (DirtyActor.WorldPartitionRef.IsSet())
 		{
 			const FWorldPartitionReference& Ref = DirtyActor.WorldPartitionRef.GetValue();
@@ -1569,18 +1593,12 @@ void UWorldPartition::Tick(float DeltaSeconds)
 			}
 			else if (DirtyActor.ActorPtr.IsValid() && !DirtyActor.ActorPtr->GetPackage()->IsDirty())
 			{
-				// Transfer ownership of the last ref if actor can be pinned, if not keep a temp ref so that we can refresh loaded states (which should take a add ref on unpinnable actors and prevent unloading)
-				if (Ref->GetHardRefCount() <= 1)
+				// Transfer ownership of the last ref if actor can be pinned
+				if (Ref->GetHardRefCount() <= 1 && PinnedActors && FLoaderAdapterPinnedActors::SupportsPinning(Ref.Get()))
 				{
-					if (PinnedActors && FLoaderAdapterPinnedActors::SupportsPinning(Ref.Get()))
-					{
-						PinnedActors->AddActors({ Ref.ToHandle() });
-					}
-					else
-					{
-						ActorReferences.Add(Ref);
-					}
-				} 
+					PinnedActors->AddActors({ Ref.ToHandle() });
+				}
+
 				DirtyActorIt.RemoveCurrent();
 			}
 		}
@@ -1602,19 +1620,6 @@ void UWorldPartition::Tick(float DeltaSeconds)
 		}
 	}
 
-	if (ActorReferences.Num() > 0)
-	{
-		// Before ActorReferences get released, refresh the Loaders so that they take a reference on actors they should load.
-		IWorldPartitionActorLoaderInterface::RefreshLoadedState(false);
-
-		for (const FWorldPartitionReference& ActorReference : ActorReferences)
-		{
-			UE_CLOG(ActorReference->GetHardRefCount() <= 1, LogWorldPartition, Warning, TEXT("Releasing reference on saved actor '%s' from the DirtyActors without a loader reference."), *ActorReference->GetActor()->GetPathName());
-		}
-		// Actors should be loaded by a Loader at this point and we can release the references
-		ActorReferences.Empty();
-	}
-		
 	if (bForceGarbageCollection)
 	{
 		GEngine->ForceGarbageCollection(bForceGarbageCollectionPurge);
@@ -2040,9 +2045,9 @@ bool UWorldPartition::UnregisterActorDescContainer(UActorDescContainer* InActorD
 			{
 				ActorGuids.Add(It->GetGuid());
 
-				for (TSet<FDirtyActor>::TIterator DirtyActorIt(DirtyActors); DirtyActorIt; ++DirtyActorIt)
+				for (TMap<FGuid, FDirtyActor>::TIterator DirtyActorIt(DirtyActors); DirtyActorIt; ++DirtyActorIt)
 				{
-					const FDirtyActor& DirtyActor = *DirtyActorIt;
+					const FDirtyActor& DirtyActor = DirtyActorIt.Value();
 					if (DirtyActor.WorldPartitionRef.IsSet())
 					{
 						if (DirtyActor.WorldPartitionRef.GetValue() == ActorHandle)
diff --git a/Engine/Source/Runtime/Engine/Public/WorldPartition/WorldPartition.h b/Engine/Source/Runtime/Engine/Public/WorldPartition/WorldPartition.h
--- a/Engine/Source/Runtime/Engine/Public/WorldPartition/WorldPartition.h
+++ b/Engine/Source/Runtime/Engine/Public/WorldPartition/WorldPartition.h
@@ -107,19 +107,17 @@ struct ENGINE_API FDirtyActor
 
 	FDirtyActor(AActor* InActor)
 		: ActorPtr(InActor)
-	{
-	}
+	{}
 
 	FDirtyActor(const FWorldPartitionReference& InWorldPartitionRef, AActor* InActor)
 		: WorldPartitionRef(InWorldPartitionRef)
 		, ActorPtr(InActor)
-	{
-	}
+	{}
 
 	TOptional<FWorldPartitionReference> WorldPartitionRef;
-	TWeakObjectPtr<AActor>	ActorPtr;		// TWeakObjectPtr is for undo support.
+	TWeakObjectPtr<AActor> ActorPtr; // TWeakObjectPtr is for undo support.
 
-	bool operator == (const FDirtyActor& InDirtyActor) const
+	bool operator==(const FDirtyActor& InDirtyActor) const
 	{
 		return WorldPartitionRef == InDirtyActor.WorldPartitionRef && ActorPtr == InDirtyActor.ActorPtr;
 	}
@@ -356,7 +354,7 @@ public:
 
 	bool IsEnablingStreamingJustified() const { return bEnablingStreamingJustified; }
 
-	const TSet<FDirtyActor>& GetDirtyActors() const { return DirtyActors; }
+	const TMap<FGuid, FDirtyActor>& GetDirtyActors() const { return DirtyActors; }
 #endif
 
 public:
@@ -472,7 +470,7 @@ private:
 
 	TArray<FWorldPartitionReference> LoadedSubobjects;
 
-	TSet<FDirtyActor> DirtyActors;
+	TMap<FGuid, FDirtyActor> DirtyActors;
 
 	TSet<FString> GeneratedStreamingPackageNames;
 
