Remove bIsRuntime from LocationVolume. Waiting a better use case(s) for PCG to play with loaded volume.

diff --git a/Engine/Source/Runtime/Engine/Public/LocationVolume.h b/Engine/Source/Runtime/Engine/Public/LocationVolume.h
--- a/Engine/Source/Runtime/Engine/Public/LocationVolume.h
+++ b/Engine/Source/Runtime/Engine/Public/LocationVolume.h
@@ -26,7 +26,7 @@ public:
 	//~ End UObject Interface
 
 	//~ Begin AActor Interface
-	virtual bool IsEditorOnly() const override { return !bIsRuntime; }
+	virtual bool IsEditorOnly() const override { return true; }
 #if WITH_EDITOR
 	ENGINE_API virtual void PostRegisterAllComponents();
 	ENGINE_API virtual void UnregisterAllComponents(bool bForReregister) override;
@@ -53,9 +53,6 @@ public:
 	UPROPERTY(EditAnywhere, Category=LocationVolume)
 	FColor DebugColor;
 
-	UPROPERTY(EditAnywhere, Category=LocationVolume)
-	uint8 bIsRuntime : 1;
-
 #if WITH_EDITORONLY_DATA
 	/* To support per-user last loaded location volumes */
 	uint8 bIsAutoLoad : 1;
