Fix additive curve animation based on skeleton ref pose not working

The curve refactor introduced a behavior change in TBaseBlendedCurve::InitFrom, which acted like CopyFrom rather than simply being the (old) equivalent of emptying the curve values and keeping buffer sizes in sync. In the refactored code, the function now just empties the curves and copies the filter.
Updated compressed data version to force baked additives to recompress with the fix.

diff --git a/Engine/Source/Runtime/Engine/Classes/Animation/AnimCompress.h b/Engine/Source/Runtime/Engine/Classes/Animation/AnimCompress.h
--- a/Engine/Source/Runtime/Engine/Classes/Animation/AnimCompress.h
+++ b/Engine/Source/Runtime/Engine/Classes/Animation/AnimCompress.h
@@ -272,7 +272,7 @@ namespace UE
 			// This is a version string that mimics the old versioning scheme. If you
 			// want to bump this version, generate a new guid using VS->Tools->Create GUID and
 			// return it here. Ex.
-			static FString AnimationCompressionVersionString = TEXT("75162A70E2074707ACF3DB9CC1B19597");
+			static FString AnimationCompressionVersionString = TEXT("55AC7D7DCC1F46C2A8B121A56FA37296");
 			
 			struct FAnimationCompressionMemorySummaryScope
 			{
diff --git a/Engine/Source/Runtime/Engine/Public/Animation/AnimCurveTypes.h b/Engine/Source/Runtime/Engine/Public/Animation/AnimCurveTypes.h
--- a/Engine/Source/Runtime/Engine/Public/Animation/AnimCurveTypes.h
+++ b/Engine/Source/Runtime/Engine/Public/Animation/AnimCurveTypes.h
@@ -510,39 +510,16 @@ public:
 		Super::bSorted = false;
 	}
 
-	/** Initialize from another curve */
+	/** Initialize from another curve. Just copies the filter and empties 'this'. */
 	template <typename OtherAllocator, typename OtherElementType>
 	void InitFrom(const TBaseBlendedCurve<OtherAllocator, OtherElementType>& InCurveToInitFrom)
 	{
 		CURVE_PROFILE_CYCLE_COUNTER(TBaseBlendedCurve_InitFrom);
 		
-		if constexpr(std::is_same<AllocatorType, OtherAllocator>::value && std::is_same<ElementType, OtherElementType>::value)
-		{
-			if (ensure(&InCurveToInitFrom != this))
-			{
-				Super::Elements = InCurveToInitFrom.Elements;
-				Super::bSorted = InCurveToInitFrom.bSorted;
-				Filter = InCurveToInitFrom.Filter;
-			}
-		}
-		else if constexpr(std::is_same<ElementType, OtherElementType>::value)
-		{
-			Super::Elements = InCurveToInitFrom.Elements;
-			Super::bSorted = InCurveToInitFrom.bSorted;
-			Filter = InCurveToInitFrom.Filter;
-		}
-		else
+		if (ensure(&InCurveToInitFrom != this))
 		{
 			Super::Elements.Reset();
-
-			UE::Anim::FNamedValueArrayUtils::Union(*this, InCurveToInitFrom,
-				[](ElementType& InOutResultElement, OtherElementType& InParamElement, UE::Anim::ENamedValueUnionFlags InFlags)
-				{
-					InOutResultElement.Value = InParamElement.Value;
-					InOutResultElement.Flags = InParamElement.Flags;
-				});
-
-			Super::bSorted = InCurveToInitFrom.bSorted;
+			Super::bSorted = false;
 			Filter = InCurveToInitFrom.Filter;
 		}
 	}
@@ -903,7 +880,7 @@ public:
 	{
 		CURVE_PROFILE_CYCLE_COUNTER(TBaseBlendedCurve_Override_Weighted);
 		
-		InitFrom(CurveToOverrideFrom);
+		CopyFrom(CurveToOverrideFrom);
 
 		if (!FMath::IsNearlyEqual(Weight, 1.0f))
 		{
@@ -923,7 +900,7 @@ public:
 	{
 		CURVE_PROFILE_CYCLE_COUNTER(TBaseBlendedCurve_Override);
 		
-		InitFrom(CurveToOverrideFrom);
+		CopyFrom(CurveToOverrideFrom);
 	}
 	
 	/**
