Fix for crash adding textures to texture2d arrays. Without ReadOnly, UnlockMips can change the texture source id on the texture that's being added without a PreEditChange, causing the check().

diff --git a/Engine/Source/Runtime/Engine/Classes/Engine/Texture2DArray.cpp b/Engine/Source/Runtime/Engine/Classes/Engine/Texture2DArray.cpp
--- a/Engine/Source/Runtime/Engine/Classes/Engine/Texture2DArray.cpp
+++ b/Engine/Source/Runtime/Engine/Classes/Engine/Texture2DArray.cpp
@@ -432,7 +432,7 @@ ENGINE_API bool UTexture2DArray::UpdateSourceFromSourceTextures(bool bCreatingNe
 			void* DestSliceData = DestMipData + MipSizeBytes * SourceTexIndex;
 			if (TextureSource.SizeX == SizeX && TextureSource.SizeY == SizeY && TextureSource.Format == Format)
 			{
-				void* SourceData = TextureSource.LockMip(0);
+				const uint8* SourceData = TextureSource.LockMipReadOnly(0);
 				check(TextureSource.CalcMipSize(0) == MipSizeBytes);
 				FMemory::Memcpy(DestSliceData, SourceData, MipSizeBytes);
 				TextureSource.UnlockMip(0);
