Fix a very old some copy/paste error when overriding vector axis key bindings on an input component.

diff --git a/Engine/Source/Runtime/Engine/Private/InputVectorAxisDelegateBinding.cpp b/Engine/Source/Runtime/Engine/Private/InputVectorAxisDelegateBinding.cpp
--- a/Engine/Source/Runtime/Engine/Private/InputVectorAxisDelegateBinding.cpp
+++ b/Engine/Source/Runtime/Engine/Private/InputVectorAxisDelegateBinding.cpp
@@ -31,7 +31,7 @@ void UInputVectorAxisDelegateBinding::BindToInputComponent(UInputComponent* Inpu
 				const FInputVectorAxisBinding& ExistingBind = InputComponent->VectorAxisBindings[ExistingIndex];
 				if (ExistingBind.AxisKey == VAB.AxisKey)
 				{
-					InputComponent->AxisKeyBindings.RemoveAt(ExistingIndex);
+					InputComponent->VectorAxisBindings.RemoveAt(ExistingIndex);
 				}
 			}
 		}
