"Move/Rename" Metasound/MetasoundWaveProxReader to -> Engine/SoundWaveProxyReader

(Actually, just "Add" new SoundWaveProxyReader class)

- Remove from Metasound namespace
- Rename MetasoundWaveProxyReader to SoundWaveProxyReader
- MetasoundWavePalyerNode now uses SoundWaveProxyReader
- Use `LogAudio` instead of MetasoundLogs for logging.

diff --git a/Engine/Plugins/Runtime/Metasound/Source/MetasoundEngine/Private/MetasoundWavePlayerNode.cpp b/Engine/Plugins/Runtime/Metasound/Source/MetasoundEngine/Private/MetasoundWavePlayerNode.cpp
--- a/Engine/Plugins/Runtime/Metasound/Source/MetasoundEngine/Private/MetasoundWavePlayerNode.cpp
+++ b/Engine/Plugins/Runtime/Metasound/Source/MetasoundEngine/Private/MetasoundWavePlayerNode.cpp
@@ -18,7 +18,7 @@
 #include "MetasoundTrigger.h"
 #include "MetasoundVertex.h"
 #include "MetasoundWave.h"
-#include "MetasoundWaveProxyReader.h"
+#include "Sound/SoundWaveProxyReader.h"
 
 #define LOCTEXT_NAMESPACE "MetasoundWaveNode"
 
@@ -68,9 +68,9 @@ namespace Metasound
 
 		/** FSourceBufferState tracks the current frame and loop indices held in 
 		 * a circular buffer. It describes how the content of a circular buffer
-		 * relates to the frame indices of an FWaveProxyReader 
+		 * relates to the frame indices of an FSoundWaveProxyReader 
 		 *
-		 * FSourceBufferState is tied to the implementation of the FWaveProxyReader
+		 * FSourceBufferState is tied to the implementation of the FSoundWaveProxyReader
 		 * and TCircularAudioBuffer<>, and thus does not serve much purpose outside
 		 * of this wave player node.
 		 *
@@ -114,7 +114,7 @@ namespace Metasound
 			 * @param ProxyReader - The wave proxy reader producing the audio.
 			 * @parma InSourceBuffer - The audio buffer holding a range of samples popped from the reader.
 			 */
-			FSourceBufferState(const FWaveProxyReader& ProxyReader, const Audio::FMultichannelCircularBuffer& InSourceBuffer)
+			FSourceBufferState(const FSoundWaveProxyReader& ProxyReader, const Audio::FMultichannelCircularBuffer& InSourceBuffer)
 			: FSourceBufferState(ProxyReader.GetFrameIndex(), Audio::GetMultichannelBufferNumFrames(InSourceBuffer), ProxyReader.IsLooping(), ProxyReader.GetLoopStartFrameIndex(), ProxyReader.GetLoopEndFrameIndex(), ProxyReader.GetNumFramesInWave())
 			{
 			}
@@ -158,7 +158,7 @@ namespace Metasound
 			}
 
 			/** Update loop frame indices. */
-			void SetLoopFrameIndices(const FWaveProxyReader& InProxyReader)
+			void SetLoopFrameIndices(const FSoundWaveProxyReader& InProxyReader)
 			{
 				SetLoopFrameIndices(InProxyReader.GetLoopStartFrameIndex(), InProxyReader.GetLoopEndFrameIndex());
 			}
@@ -701,14 +701,14 @@ namespace Metasound
 				Algo::SortBy(SortedCuePoints, WavePlayerNodePrivate::GetCuePointFrame);
 				
 				// Create the wave proxy reader.
-				FWaveProxyReader::FSettings WaveReaderSettings;
+				FSoundWaveProxyReader::FSettings WaveReaderSettings;
 				WaveReaderSettings.MaxDecodeSizeInFrames = MaxDecodeSizeInFrames;
 				WaveReaderSettings.StartTimeInSeconds = StartTime->GetSeconds();
 				WaveReaderSettings.LoopStartTimeInSeconds = LoopStartTime->GetSeconds();
 				WaveReaderSettings.LoopDurationInSeconds = LoopDuration->GetSeconds(); 
 				WaveReaderSettings.bIsLooping = *bLoop;
 
-				WaveProxyReader = FWaveProxyReader::Create(WaveProxy.ToSharedRef(), WaveReaderSettings);
+				WaveProxyReader = FSoundWaveProxyReader::Create(WaveProxy.ToSharedRef(), WaveReaderSettings);
 
 				if (WaveProxyReader.IsValid())
 				{
@@ -1022,7 +1022,7 @@ namespace Metasound
 		TArray<FAudioBufferWriteRef> OutputAudioBuffers;
 		TArray<FName> OutputAudioBufferVertexNames;
 
-		TUniquePtr<FWaveProxyReader> WaveProxyReader;
+		TUniquePtr<FSoundWaveProxyReader> WaveProxyReader;
 		TUniquePtr<Audio::IConvertDeinterleave> ConvertDeinterleave;
 		TUniquePtr<Audio::FMultichannelLinearResampler> Resampler;
 
diff --git a/Engine/Source/Runtime/Engine/Classes/Sound/SoundWaveProxyReader.h b/Engine/Source/Runtime/Engine/Classes/Sound/SoundWaveProxyReader.h
new file mode 100644
--- /dev/null
+++ b/Engine/Source/Runtime/Engine/Classes/Sound/SoundWaveProxyReader.h
@@ -0,0 +1,265 @@
+// Copyright Epic Games, Inc. All Rights Reserved.
+
+#pragma once
+
+#include "DSP/BufferVectorOperations.h"
+#include "HAL/Platform.h"
+#include "IAudioCodec.h"
+#include "Templates/SharedPointer.h"
+#include "Templates/UniquePtr.h"
+
+// Forward declare
+class FSoundWaveProxy;
+
+
+PRAGMA_DISABLE_DEPRECATION_WARNINGS
+
+namespace SoundWaveProxyReaderPrivate
+{
+	/** FDecoderOutput maintains a circular buffer of audio which is produced
+	 * from the decoder.
+	 */
+	class UE_DEPRECATED(5.3, "Please use IAudioFormat/IStreamedCompressedInfo instead") FDecoderOutput : public Audio::IDecoderOutput
+	{
+		static constexpr uint32 MinNumFramesPerDecode = 1;
+		static constexpr uint32 DefaultNumChannels = 1;
+		static constexpr uint32 MinNumChannels = 1;
+
+	public:
+		using FPushedAudioDetails = Audio::IDecoderOutput::FPushedAudioDetails;
+
+		/** Construct a FDecoderOutput
+		 *
+		 * @param InNumFramesPerDecode - Maximum number of frames pushed when
+		 *                               audio is decoded. The buffer will be
+		 *                               able to hold twice the number of frames.
+		 */
+		FDecoderOutput(uint32 InNumFramesPerDecode);
+
+		/** Set the number of channels in the audio being decoded. */
+		void SetNumChannels(uint32 InNumChannels);
+
+		/** Returns the number of samples in the buffer. */
+		FORCEINLINE int32 Num() const
+		{
+			return Buffer.Num();
+		}
+
+		/** Removes all samples from the buffer. */
+		void Reset();
+
+		/** Returns requirements used by the audio codec system. */
+		virtual Audio::IDecoderOutput::FRequirements GetRequirements(const Audio::FDecodedFormatInfo & InFormat) const override;
+
+		/** Adds samples to the buffer.
+		 *
+		 * This is called by the Decoder and should not be called otherwise.
+		 */
+		virtual int32 PushAudio(const Audio::IDecoderOutput::FPushedAudioDetails & InDetails, TArrayView<const int16> In16BitInterleave) override;
+
+		/** Adds samples to the buffer.
+		 *
+		 * This is called by the Decoder and should not be called otherwise.
+		 */
+		virtual int32 PushAudio(const FPushedAudioDetails & InDetails, TArrayView<const float> InFloat32Interleave) override;
+
+		/** This should not be called. It removes 16 bit PCM samples from the buffer
+		 * which is an unsupported operation of this class.
+		 */
+		virtual int32 PopAudio(TArrayView<int16> InExternalInt16Buffer, FPushedAudioDetails & OutDetails) override;
+
+		/** Copy samples to OutBuffer and remove them from this objects internal
+		 * buffer.
+		 *
+		 *
+		 * @param OutBuffer - A destination array to copy samples to.
+		 * @param OutDetails - Unused.
+		 *
+		 * @return The actual number of samples copied.
+		 */
+		virtual int32 PopAudio(TArrayView<float> OutBuffer, FPushedAudioDetails & OutDetails) override;
+
+		/** Remove samples from the internal buffer.
+		 *
+		 * @param InNumSamples - The desired number of samples to remove.
+		 *
+		 * @return The actual number of samples removed.
+		 */
+		FORCEINLINE int32 PopAudio(int32 InNumSamples)
+		{
+			return Buffer.Pop(InNumSamples);
+		}
+
+	private:
+
+		// Initialize buffer size.
+		void Init();
+		int32 PushAudioInternal(const FPushedAudioDetails & InDetails, TArrayView<const float> InBuffer);
+		int32 PushAudioInternal(const FPushedAudioDetails & InDetails, TArrayView<const int16> InBuffer);
+
+		uint32 NumFramesPerDecode = MinNumFramesPerDecode;
+		uint32 NumChannels = DefaultNumChannels;
+		Audio::TCircularAudioBuffer<float> Buffer;
+		Audio::FAlignedFloatBuffer SampleConversionBuffer;
+	};
+}
+
+/** FSoundWaveProxyReader reads a FWaveProxy and outputs 32 bit interleaved audio.
+*
+* FSoundWaveProxyReader provides controls for looping and relevant frame index
+* values.
+*/
+class FSoundWaveProxyReader
+{
+
+public:
+	using FSoundWaveProxyRef = TSharedRef<FSoundWaveProxy, ESPMode::ThreadSafe>;
+	using EDecodeResult = Audio::IDecoder::EDecodeResult;
+
+	/** Minimum number of frames to decode per a call to the decoder.  */
+	static constexpr uint32 DefaultMinDecodeSizeInFrames = 128;
+
+	/** Some codecs have strict requirements on decode size. In order to be
+	 * functional with all supported codecs, the decode size must be a multiple
+	 * of 128.
+	 */
+	static constexpr uint32 DecodeSizeQuantizationInFrames = 128;
+
+	static uint32 ConformDecodeSize(uint32 InMaxDesiredDecodeSizeInFrames);
+
+	/** Minimum duration of a loop. */
+	static constexpr float MinLoopDurationInSeconds = 0.05f;
+
+	/** Maximum duration of a loop.
+	 *
+	 * 10 years. ridiculously high. This exists to prevent floating point
+	 * undefined overflow behavior if performing calculations with
+	 * TNumericLimits<float>::Max().
+	 */
+	static constexpr float MaxLoopDurationInSeconds = 60.f * 60.f * 365.f * 10.f;
+
+	/** Settings for a FSoundWaveProxyReader. */
+	struct FSettings
+	{
+		uint32 MaxDecodeSizeInFrames = 8192;
+		float StartTimeInSeconds = 0.f;
+		bool bIsLooping = false;
+		float LoopStartTimeInSeconds = 0.f;
+		float LoopDurationInSeconds = 0.f;
+	};
+
+private:
+	/** Construct a wave proxy reader.
+	 *
+	 * @param InWaveProxy - A TSharedRef of a FSoundWaveProxy which is to be played.
+	 * @param InSettings - Reader settings.
+	 */
+	FSoundWaveProxyReader(FSoundWaveProxyRef InWaveProxy, const FSettings& InSettings);
+
+public:
+	/** Create a wave proxy reader.
+	 *
+	 * @param InWaveProxy - A TSharedRef of a FSoundWaveProxy which is to be played.
+	 * @param InSettings - Reader settings.
+	 */
+	static ENGINE_API TUniquePtr<FSoundWaveProxyReader> Create(FSoundWaveProxyRef InWaveProxy, const FSettings& InSettings);
+
+	/** Set whether the reader should loop the audio or not. */
+	ENGINE_API void SetIsLooping(bool bInIsLooping);
+
+	/** Returns true if the audio will be looped, false otherwise. */
+	ENGINE_API FORCEINLINE bool IsLooping() const
+	{
+		return Settings.bIsLooping;
+	}
+
+	/** Sets the beginning position of the loop. */
+	ENGINE_API void SetLoopStartTime(float InLoopStartTimeInSeconds);
+
+	/** Sets the duration of the loop in seconds.
+	 *
+	 * If the value is negative, the MaxLoopDurationInSeconds will be used
+	 * which will effectively loop at the end of the file.
+	 */
+	ENGINE_API void SetLoopDuration(float InLoopDurationInSeconds);
+
+	ENGINE_API FORCEINLINE float GetSampleRate() const
+	{
+		return SampleRate;
+	}
+
+	ENGINE_API FORCEINLINE int32 GetNumChannels() const
+	{
+		return NumChannels;
+	}
+
+	/** Returns the index of the playhead within the complete wave. */
+	ENGINE_API FORCEINLINE int32 GetFrameIndex() const
+	{
+		return CurrentFrameIndex;
+	}
+
+	ENGINE_API FORCEINLINE int32 GetNumFramesInWave() const
+	{
+		return NumFramesInWave;
+	}
+
+	ENGINE_API FORCEINLINE int32 GetNumFramesInLoop() const
+	{
+		return LoopEndFrameIndex - LoopStartFrameIndex;
+	}
+
+	ENGINE_API FORCEINLINE int32 GetLoopStartFrameIndex() const
+	{
+		return LoopStartFrameIndex;
+	}
+
+	ENGINE_API FORCEINLINE int32 GetLoopEndFrameIndex() const
+	{
+		return LoopEndFrameIndex;
+	}
+
+	/** Seeks to position in wave.
+	 *
+	 * @param InSeconds - The location to seek the playhead
+	 *
+	 * @return true on success, false on failuer.
+	 */
+	ENGINE_API bool SeekToTime(float InSeconds);
+
+	/** Copies audio into OutBuffer. It returns the number of samples copied.
+	 * Samples not written to will be set to zero.
+	 */
+	ENGINE_API int32 PopAudio(Audio::FAlignedFloatBuffer& OutBuffer);
+
+private:
+
+	int32 PopAudioFromDecoderOutput(TArrayView<float> OutBufferView);
+	bool InitializeDecoder(float InStartTimeInSeconds);
+	void DiscardSamples(int32 InNumSamplesToDiscard);
+	float ClampLoopStartTime(float InStartTimeInSeconds);
+	float ClampLoopDuration(float InDurationInSeconds);
+	void UpdateLoopBoundaries();
+
+	FSoundWaveProxyRef WaveProxy;
+	TUniquePtr<Audio::IDecoderInput> DecoderInput;
+	TUniquePtr<Audio::IDecoder> Decoder;
+	SoundWaveProxyReaderPrivate::FDecoderOutput DecoderOutput;
+
+	FSettings Settings;
+
+	float SampleRate = 0.f;
+	int32 NumChannels = 0;
+	EDecodeResult DecodeResult = EDecodeResult::MoreDataRemaining;
+	int32 NumFramesInWave = 0;
+	int32 CurrentFrameIndex = 0;
+	int32 LoopStartFrameIndex = 0;
+	int32 LoopEndFrameIndex = -1;
+	float DurationInSeconds = 0.f;
+	float MaxLoopStartTimeInSeconds = 0.f;
+
+	bool bIsDecoderValid = false;
+	bool bFallbackSeekMethodWarningLogged = false;
+};
+
+PRAGMA_ENABLE_DEPRECATION_WARNINGS
\ No newline at end of file
diff --git a/Engine/Source/Runtime/Engine/Engine.Build.cs b/Engine/Source/Runtime/Engine/Engine.Build.cs
--- a/Engine/Source/Runtime/Engine/Engine.Build.cs
+++ b/Engine/Source/Runtime/Engine/Engine.Build.cs
@@ -153,6 +153,7 @@ public class Engine : ModuleRules
 				"Analytics",
 				"AudioMixer",
 				"AudioMixerCore",
+				"AudioCodecEngine",
 				"IntelISPC",
 				"TraceLog",
 				"ColorManagement",
diff --git a/Engine/Source/Runtime/Engine/Private/SoundWaveProxyReader.cpp b/Engine/Source/Runtime/Engine/Private/SoundWaveProxyReader.cpp
new file mode 100644
--- /dev/null
+++ b/Engine/Source/Runtime/Engine/Private/SoundWaveProxyReader.cpp
@@ -0,0 +1,552 @@
+// Copyright Epic Games, Inc. All Rights Reserved.
+
+#include "Sound/SoundWaveProxyReader.h"
+
+#include "Audio.h"
+#include "Containers/Array.h"
+#include "Containers/ArrayView.h"
+#include "DecoderInputFactory.h"
+#include "DSP/BufferVectorOperations.h"
+#include "DSP/Dsp.h"
+#include "HAL/IConsoleManager.h"
+#include "HAL/LowLevelMemTracker.h"
+#include "HAL/Platform.h"
+#include "HAL/UnrealMemory.h"
+#include "IAudioCodec.h"
+#include "IAudioCodecRegistry.h"
+#include "Math/UnrealMathUtility.h"
+#include "Misc/AssertionMacros.h"
+#include "ProfilingDebugging/CpuProfilerTrace.h"
+#include "Sound/SoundWave.h"
+#include "Templates/SharedPointer.h"
+#include "Templates/UniquePtr.h"
+
+static int32 GSoundWaveProxyReaderSimulateSeekOnNonSeekable = 0;
+FAutoConsoleVariableRef CVarSoundWaveProxyReaderSimulateSeekOnNonSeekable(
+	TEXT("au.SoundWaveProxyReader.SimulateSeek"),
+	GSoundWaveProxyReaderSimulateSeekOnNonSeekable,
+	TEXT("If true, SoundWaves which are not of a seekable format will simulate seek calls by reading and discarding samples.\n")
+	TEXT("0: Do not simulate seek, !0: Simulate seek"),
+	ECVF_Default);
+
+
+PRAGMA_DISABLE_DEPRECATION_WARNINGS
+
+namespace SoundWaveProxyReaderPrivate
+{
+	/** Construct a FDecoderOutput
+	 *
+	 * @param InNumFramesPerDecode - Maximum number of frames pushed when
+	 *                               audio is decoded. The buffer will be
+	 *                               able to hold twice the number of frames.
+	 */
+	FDecoderOutput::FDecoderOutput(uint32 InNumFramesPerDecode)
+		: NumFramesPerDecode(FMath::Max(InNumFramesPerDecode, MinNumFramesPerDecode))
+	{
+		Init();
+	}
+
+	/** Set the number of channels in the audio being decoded. */
+	void FDecoderOutput::SetNumChannels(uint32 InNumChannels)
+	{
+		NumChannels = FMath::Max(MinNumChannels, InNumChannels);
+		Init();
+	}
+
+	/** Removes all samples from the buffer. */
+	void FDecoderOutput::Reset()
+	{
+		Buffer.SetNum(0);
+	}
+
+	/** Returns requirements used by the audio codec system. */
+	Audio::IDecoderOutput::FRequirements FDecoderOutput::GetRequirements(const Audio::FDecodedFormatInfo& InFormat) const
+	{
+		FRequirements Requirements;
+
+		Requirements.DownstreamFormat = Audio::EBitRepresentation::Float32_Interleaved;
+		Requirements.NumSampleFramesWanted = NumFramesPerDecode;
+		Requirements.NumSampleFramesPerSecond = InFormat.NumFramesPerSec;
+		Requirements.NumChannels = InFormat.NumChannels;
+
+		return Requirements;
+	}
+
+	/** Adds samples to the buffer.
+		*
+		* This is called by the Decoder and should not be called otherwise.
+		*/
+	int32 FDecoderOutput::PushAudio(const Audio::IDecoderOutput::FPushedAudioDetails& InDetails, TArrayView<const int16> In16BitInterleave)
+	{
+		if (InDetails.NumChannels != NumChannels)
+		{
+			SetNumChannels(InDetails.NumChannels);
+		}
+		return PushAudioInternal(InDetails, In16BitInterleave);
+	}
+
+	/** Adds samples to the buffer.
+		*
+		* This is called by the Decoder and should not be called otherwise.
+		*/
+	int32 FDecoderOutput::PushAudio(const FPushedAudioDetails& InDetails, TArrayView<const float> InFloat32Interleave)
+	{
+		if (InDetails.NumChannels != NumChannels)
+		{
+			SetNumChannels(InDetails.NumChannels);
+		}
+		return PushAudioInternal(InDetails, InFloat32Interleave);
+	}
+
+	/** This should not be called. It removes 16 bit PCM samples from the buffer
+		* which is an unsupported operation of this class.
+		*/
+	int32 FDecoderOutput::PopAudio(TArrayView<int16> InExternalInt16Buffer, FPushedAudioDetails& OutDetails)
+	{
+		// This buffer cannot produce 16 bit PCM audio.
+		checkNoEntry();
+		return 0;
+	}
+
+	/** Copy samples to OutBuffer and remove them from this objects internal
+		* buffer.
+		*
+		*
+		* @param OutBuffer - A destination array to copy samples to.
+		* @param OutDetails - Unused.
+		*
+		* @return The actual number of samples copied.
+		*/
+	int32 FDecoderOutput::PopAudio(TArrayView<float> OutBuffer, FPushedAudioDetails& OutDetails)
+	{
+		return Buffer.Pop(OutBuffer.GetData(), OutBuffer.Num());
+	}
+
+	// Initialize buffer size.
+	void FDecoderOutput::Init()
+	{
+		// Allow it to hold two decode buffers max.
+		uint32 MinBufferCapacity = NumChannels * NumFramesPerDecode * 2;
+		constexpr bool bRetainExistingSamples = false;
+		Buffer.Reserve(MinBufferCapacity, bRetainExistingSamples);
+	}
+
+	int32 FDecoderOutput::PushAudioInternal(const FPushedAudioDetails& InDetails, TArrayView<const float> InBuffer)
+	{
+		return Buffer.Push(InBuffer.GetData(), InBuffer.Num());
+	}
+
+	int32 FDecoderOutput::PushAudioInternal(const FPushedAudioDetails& InDetails, TArrayView<const int16> InBuffer)
+	{
+		SampleConversionBuffer.SetNumUninitialized(InBuffer.Num(), false /* bAllowShrinking */);
+
+		// Convert 16 bit pcm to 32 bit float.
+		constexpr float Scalar = 1.f / 32768.f;
+		const int16* Src = InBuffer.GetData();
+		float* Dst = SampleConversionBuffer.GetData();
+		int32 Num = InBuffer.Num();
+
+		// Convert 1 sample at a time, slow.
+		for (int32 i = 0; i < Num; ++i)
+		{
+			*Dst++ = static_cast<float>(*Src++) * Scalar;
+		}
+
+		return PushAudioInternal(InDetails, SampleConversionBuffer);
+	}
+
+}
+
+
+
+
+uint32 FSoundWaveProxyReader::ConformDecodeSize(uint32 InMaxDesiredDecodeSizeInFrames)
+{
+	static_assert(DefaultMinDecodeSizeInFrames >= DecodeSizeQuantizationInFrames, "Min decode size less than decode size quantization");
+	static_assert((DefaultMinDecodeSizeInFrames % DecodeSizeQuantizationInFrames) == 0, "Min decode size must be equally divisible by decode size quantization");
+
+	const uint32 QuantizedDecodeSizeInFrames = (InMaxDesiredDecodeSizeInFrames / DecodeSizeQuantizationInFrames) * DecodeSizeQuantizationInFrames;
+	const uint32 ConformedDecodeSizeInFrames = FMath::Max(QuantizedDecodeSizeInFrames, DefaultMinDecodeSizeInFrames);
+
+	check((ConformedDecodeSizeInFrames % QuantizedDecodeSizeInFrames) == 0);
+	check(ConformedDecodeSizeInFrames > 0);
+
+	return ConformedDecodeSizeInFrames;
+}
+
+/** Construct a wave proxy reader.
+	*
+	* @param InWaveProxy - A TSharedRef of a FSoundWaveProxy which is to be played.
+	* @param InSettings - Reader settings.
+	*/
+FSoundWaveProxyReader::FSoundWaveProxyReader(FSoundWaveProxyRef InWaveProxy, const FSettings& InSettings)
+	: WaveProxy(InWaveProxy)
+	, DecoderOutput(ConformDecodeSize(InSettings.MaxDecodeSizeInFrames))
+	, Settings(InSettings)
+{
+	// Get local copies of some values from the proxy. 
+	SampleRate = WaveProxy->GetSampleRate();
+	NumChannels = WaveProxy->GetNumChannels();
+	NumFramesInWave = WaveProxy->GetNumFrames();
+	DurationInSeconds = WaveProxy->GetDuration();
+	MaxLoopStartTimeInSeconds = FMath::Max(0, DurationInSeconds - MinLoopDurationInSeconds);
+
+	// Clamp times
+	Settings.StartTimeInSeconds = FMath::Clamp(Settings.StartTimeInSeconds, 0.f, DurationInSeconds);
+	Settings.LoopStartTimeInSeconds = ClampLoopStartTime(Settings.LoopStartTimeInSeconds);
+	Settings.LoopDurationInSeconds = ClampLoopDuration(Settings.LoopDurationInSeconds);
+
+	// Setup frame indices
+	CurrentFrameIndex = 0;
+	UpdateLoopBoundaries();
+
+	// Determine max size of decode buffer
+	Settings.MaxDecodeSizeInFrames = FMath::Max(DefaultMinDecodeSizeInFrames, Settings.MaxDecodeSizeInFrames);
+	DecoderOutput.SetNumChannels(NumChannels);
+
+	// Prepare to read audio
+	bIsDecoderValid = InitializeDecoder(Settings.StartTimeInSeconds);
+}
+
+/** Create a wave proxy reader.
+	*
+	* @param InWaveProxy - A TSharedRef of a FSoundWaveProxy which is to be played.
+	* @param InSettings - Reader settings.
+	*/
+TUniquePtr<FSoundWaveProxyReader> FSoundWaveProxyReader::Create(FSoundWaveProxyRef InWaveProxy, const FSettings& InSettings)
+{
+	if (InWaveProxy->GetSampleRate() <= 0.f)
+	{
+		UE_LOG(LogAudio, Warning, TEXT("Cannot create FSoundWaveProxyReader due to invalid sample rate (%f). Package: %s"), InWaveProxy->GetSampleRate(), *InWaveProxy->GetPackageName().ToString());
+		return TUniquePtr<FSoundWaveProxyReader>();
+	}
+
+	if (InWaveProxy->GetNumChannels() <= 0)
+	{
+		UE_LOG(LogAudio, Warning, TEXT("Cannot create FSoundWaveProxyReader due to invalid num channels (%d). Package: %s"), InWaveProxy->GetNumChannels(), *InWaveProxy->GetPackageName().ToString());
+		return TUniquePtr<FSoundWaveProxyReader>();
+	}
+
+	if (InWaveProxy->GetNumFrames() <= 0)
+	{
+		UE_LOG(LogAudio, Warning, TEXT("Cannot create FSoundWaveProxyReader due to invalid num frames (%d). Package: %s"), InWaveProxy->GetNumFrames(), *InWaveProxy->GetPackageName().ToString());
+		return TUniquePtr<FSoundWaveProxyReader>();
+	}
+
+	return TUniquePtr<FSoundWaveProxyReader>(new FSoundWaveProxyReader(InWaveProxy, InSettings));
+}
+
+
+/** Set whether the reader should loop the audio or not. */
+void FSoundWaveProxyReader::SetIsLooping(bool bInIsLooping)
+{
+	if (Settings.bIsLooping != bInIsLooping)
+	{
+		Settings.bIsLooping = bInIsLooping;
+		UpdateLoopBoundaries();
+	}
+}
+
+/** Sets the beginning position of the loop. */
+void FSoundWaveProxyReader::SetLoopStartTime(float InLoopStartTimeInSeconds)
+{
+	InLoopStartTimeInSeconds = ClampLoopStartTime(InLoopStartTimeInSeconds);
+	if (!FMath::IsNearlyEqual(Settings.LoopStartTimeInSeconds, InLoopStartTimeInSeconds))
+	{
+		Settings.LoopStartTimeInSeconds = InLoopStartTimeInSeconds;
+		UpdateLoopBoundaries();
+	}
+}
+
+/** Sets the duration of the loop in seconds. If the value is negative, the
+	* loop duration consists of the entire file. */
+void FSoundWaveProxyReader::SetLoopDuration(float InLoopDurationInSeconds)
+{
+	InLoopDurationInSeconds = ClampLoopDuration(InLoopDurationInSeconds);
+	if (!FMath::IsNearlyEqual(Settings.LoopDurationInSeconds, InLoopDurationInSeconds))
+	{
+		Settings.LoopDurationInSeconds = InLoopDurationInSeconds;
+		UpdateLoopBoundaries();
+	}
+}
+
+bool FSoundWaveProxyReader::SeekToTime(float InSeconds)
+{
+	// Direct seeking is not supported. A new decoder must be created. 
+	bIsDecoderValid = InitializeDecoder(InSeconds);
+	return bIsDecoderValid;
+}
+
+/** Copies audio into OutBuffer. It returns the number of samples copied.
+	* Samples not written to will be set to zero.
+	*/
+int32 FSoundWaveProxyReader::PopAudio(Audio::FAlignedFloatBuffer& OutBuffer)
+{
+	using namespace Audio;
+
+	TRACE_CPUPROFILER_EVENT_SCOPE(FSoundWaveProxyReader::PopAudio);
+
+	if ((0 == NumChannels) || !bIsDecoderValid)
+	{
+		if (OutBuffer.Num() > 0)
+		{
+			FMemory::Memset(OutBuffer.GetData(), 0, sizeof(float) * OutBuffer.Num());
+		}
+		return 0;
+	}
+
+	checkf(0 == (OutBuffer.Num() % NumChannels), TEXT("Output buffer size must be evenly divisible by the number of channels"));
+
+	TArrayView<float> OutBufferView{OutBuffer};
+	int32 NumSamplesUnset = OutBuffer.Num();
+	int32 NumSamplesCopied = 0;
+
+	bool bDecoderOutputHasMoreData = DecoderOutput.Num() > 0;
+	bool bDecoderCanDecodeMoreData = bIsDecoderValid && (EDecodeResult::MoreDataRemaining == DecodeResult);
+	bool bCanProduceMoreAudio = bDecoderOutputHasMoreData || bDecoderCanDecodeMoreData;
+
+	while ((NumSamplesUnset > 0) && bCanProduceMoreAudio)
+	{
+		int32 NumSamplesCopiedThisLoop = PopAudioFromDecoderOutput(OutBufferView);
+
+		// Update sample counters and output buffer view
+		NumSamplesCopied += NumSamplesCopiedThisLoop;
+		NumSamplesUnset -= NumSamplesCopiedThisLoop;
+		OutBufferView = OutBufferView.Slice(NumSamplesCopiedThisLoop, NumSamplesUnset);
+
+		if (Settings.bIsLooping)
+		{
+			// Seek to loop start if we have used up all our decodable samples or 
+			// are at the loop boundary.
+			if (0 == DecoderOutput.Num())
+			{
+				if ((!bDecoderCanDecodeMoreData) || (CurrentFrameIndex >= LoopEndFrameIndex))
+				{
+					SeekToTime(Settings.LoopStartTimeInSeconds);
+					bDecoderCanDecodeMoreData = bIsDecoderValid && (EDecodeResult::MoreDataRemaining == DecodeResult);
+				}
+			}
+		}
+
+		// Determine if we can / should decode more data. 
+		if ((NumSamplesUnset > 0) && (DecoderOutput.Num() == 0) && bDecoderCanDecodeMoreData)
+		{
+			// Looping is handle within the FSoundWaveProxyReader instead of 
+			// within the decoder.
+			DecodeResult = Decoder->Decode(false /* bIsLooping */);
+			bDecoderCanDecodeMoreData = EDecodeResult::MoreDataRemaining == DecodeResult;
+		}
+
+		bDecoderOutputHasMoreData = DecoderOutput.Num() > 0;
+		bCanProduceMoreAudio = bDecoderOutputHasMoreData || bDecoderCanDecodeMoreData;
+	}
+
+	// Zero pad any unset samples. 
+	if (OutBufferView.Num() > 0)
+	{
+		check(!bCanProduceMoreAudio);
+
+		// Zero out audio that was not set.
+		FMemory::Memset(OutBufferView.GetData(), 0, sizeof(float) * OutBufferView.Num());
+		CurrentFrameIndex += OutBufferView.Num() / NumChannels;
+	}
+
+	return NumSamplesCopied;
+}
+
+int32 FSoundWaveProxyReader::PopAudioFromDecoderOutput(TArrayView<float> OutBufferView)
+{
+	using namespace SoundWaveProxyReaderPrivate;
+
+	check(NumChannels > 0);
+
+	int32 NumSamplesCopied = 0;
+	if (DecoderOutput.Num() > 0)
+	{
+		// Get samples from the decoder buffer.
+		FDecoderOutput::FPushedAudioDetails UnusedDetails;
+		NumSamplesCopied = DecoderOutput.PopAudio(OutBufferView, UnusedDetails);
+
+		int32 NumFramesCopied = NumSamplesCopied / NumChannels;
+		CurrentFrameIndex += NumFramesCopied;
+
+		// Check whether the samples copied from the decoder extend 
+		// past the end of the loop.
+		if (Settings.bIsLooping)
+		{
+			const bool bDidOvershoot = CurrentFrameIndex >= LoopEndFrameIndex;
+			if (bDidOvershoot)
+			{
+				// Rewind sample counters if the loop boundary was overshot.
+				NumFramesCopied -= (CurrentFrameIndex - LoopEndFrameIndex);
+				CurrentFrameIndex = LoopEndFrameIndex;
+				// If Settings.bIsLooping info was altered, NumFramesCopied can end up 
+				// negative if the current frame index is past the loop end frame. 
+				NumFramesCopied = FMath::Max(0, NumFramesCopied);
+				NumSamplesCopied = NumFramesCopied * NumChannels;
+
+				// Remove any remaining samples in the decoder because they 
+				// are past the end of the loop.
+				DecoderOutput.Reset();
+			}
+
+		}
+	}
+
+	return NumSamplesCopied;
+}
+
+bool FSoundWaveProxyReader::InitializeDecoder(float InStartTimeInSeconds)
+{
+	using namespace Audio;
+
+	TRACE_CPUPROFILER_EVENT_SCOPE(FSoundWaveProxyReader::InitializeDecoder);
+
+	// Create decoder input
+	FName Format = WaveProxy->GetRuntimeFormat();
+	DecoderInput = Audio::CreateBackCompatDecoderInput(Format, WaveProxy);
+	if (!DecoderInput.IsValid())
+	{
+		UE_LOG(LogAudio, Error, TEXT("Failed to create decoder input (format:%s) for wave (package:%s)"), *Format.ToString(), *WaveProxy->GetPackageName().ToString());
+		return false;
+	}
+
+	// Seek input to start time.
+	if (!FMath::IsNearlyEqual(0.0f, InStartTimeInSeconds))
+	{
+		if (WaveProxy->IsSeekable())
+		{
+			const bool bSeekSucceeded = DecoderInput->SeekToTime(InStartTimeInSeconds);
+			if (!bSeekSucceeded)
+			{
+				UE_LOG(LogAudio, Warning, TEXT("Failed to seek decoder input during initialization: (format:%s) for wave (package:%s) to time '%.6f'"),
+					*Format.ToString(),
+					*WaveProxy->GetPackageName().ToString(),
+					InStartTimeInSeconds);
+			}
+		}
+		else
+		{
+			UE_LOG(LogAudio, Warning, TEXT("Attempt to seek on non-seekable wave: (format:%s) for wave (package:%s) to time '%.6f'"),
+				*Format.ToString(),
+				*WaveProxy->GetPackageName().ToString(),
+				InStartTimeInSeconds);
+		}
+	}
+
+	// Get codec ptr by reading the header info from the decoder input.
+	ICodecRegistry::FCodecPtr Codec = ICodecRegistry::Get().FindCodecByParsingInput(DecoderInput.Get());
+	if (nullptr == Codec)
+	{
+		UE_LOG(LogAudio, Error, TEXT("Failed to find codec (format:%s) for wave (package:%s)"), *Format.ToString(), *WaveProxy->GetPackageName().ToString());
+		return false;
+	}
+
+	// Read the sample rate and number of frames from the header 
+	// Similar to refreshing the wave data in FMixerBuffer::CreateStreamingBuffer
+	// This is a runtime hack to address incorrect sample rate on soundwaves 
+	// on platforms with Resample for Device enabled (UE-183237)
+	Audio::FFormatDescriptorSection FormatDesc;
+	DecoderInput->FindSection(FormatDesc);
+	SampleRate = FormatDesc.NumFramesPerSec;
+	if (FormatDesc.NumFrames > 0)
+	{
+		NumFramesInWave = FormatDesc.NumFrames;
+	}
+	// end hack
+
+	CurrentFrameIndex = FMath::Clamp(static_cast<int32>(InStartTimeInSeconds * GetSampleRate()), 0, GetNumFramesInWave());
+
+	// Create the decoder
+	Decoder = Codec->CreateDecoder(DecoderInput.Get(), &DecoderOutput);
+	if (!Decoder.IsValid())
+	{
+		UE_LOG(LogAudio, Error, TEXT("Failed to create decoder (format:%s) for wave (package:%s)"), *Format.ToString(), *WaveProxy->GetPackageName().ToString());
+		DecodeResult = EDecodeResult::Fail;
+		return false;
+	}
+	else
+	{
+		// The DecodeResult needs to be set to a valid state incase the prior
+		// decoder finished or failed. 
+		DecodeResult = EDecodeResult::MoreDataRemaining;
+	}
+
+	// For non-seekable streaming waves, use a fallback method to seek
+	// to the start time. 
+	const bool bUseFallbackSeekMethod = (GSoundWaveProxyReaderSimulateSeekOnNonSeekable != 0) && (InStartTimeInSeconds != 0.f) && (!WaveProxy->IsSeekable());
+	if (bUseFallbackSeekMethod)
+	{
+		if (!bFallbackSeekMethodWarningLogged)
+		{
+			UE_LOG(LogAudio, Warning, TEXT("Simulating seeking in wave which is not seekable (package:%s). For better performance, set wave to a seekable format"), *WaveProxy->GetPackageName().ToString());
+			bFallbackSeekMethodWarningLogged = true;
+		}
+
+		int32 NumSamplesToDiscard = CurrentFrameIndex * NumChannels;
+		DiscardSamples(NumSamplesToDiscard);
+	}
+
+	// return true if all the components were successfully create
+	return true;
+}
+
+void FSoundWaveProxyReader::DiscardSamples(int32 InNumSamplesToDiscard)
+{
+	while (InNumSamplesToDiscard > 0)
+	{
+		DecodeResult = Decoder->Decode(false /* bIsLooping */);
+
+		int32 NumSamplesToDiscardThisLoop = FMath::Min(DecoderOutput.Num(), InNumSamplesToDiscard);
+
+		if (NumSamplesToDiscardThisLoop < 1)
+		{
+			UE_LOG(LogAudio, Error, TEXT("Failed to decode samples (package: %s)"), *WaveProxy->GetPackageName().ToString());
+			break;
+		}
+
+		int32 ActualNumSamplesDiscarded = DecoderOutput.PopAudio(NumSamplesToDiscardThisLoop);
+		InNumSamplesToDiscard -= ActualNumSamplesDiscarded;
+
+		if ((InNumSamplesToDiscard > 0) && (DecodeResult != EDecodeResult::MoreDataRemaining))
+		{
+			UE_LOG(LogAudio, Error, TEXT("Failed to simulate seek (seek to frame: %d, package: %s)"), CurrentFrameIndex, *WaveProxy->GetPackageName().ToString());
+			break;
+		}
+	}
+}
+
+float FSoundWaveProxyReader::ClampLoopStartTime(float InStartTimeInSeconds)
+{
+	return FMath::Clamp(InStartTimeInSeconds, 0.f, MaxLoopStartTimeInSeconds);
+}
+
+float FSoundWaveProxyReader::ClampLoopDuration(float InDurationInSeconds)
+{
+	if (InDurationInSeconds <= 0.f)
+	{
+		return MaxLoopDurationInSeconds;
+	}
+	return FMath::Max(InDurationInSeconds, MinLoopDurationInSeconds);
+}
+
+void FSoundWaveProxyReader::UpdateLoopBoundaries()
+{
+	if (Settings.bIsLooping)
+	{
+		const int32 MinLoopDurationInFrames = FMath::CeilToInt(MinLoopDurationInSeconds * SampleRate);
+		const float LoopEndTime = Settings.LoopStartTimeInSeconds + Settings.LoopDurationInSeconds;
+		const int32 MinLoopEndFrameIndex = LoopStartFrameIndex + MinLoopDurationInFrames;
+		const int32 MaxLoopEndFrameIndex = NumFramesInWave;
+
+		LoopStartFrameIndex = FMath::FloorToInt(Settings.LoopStartTimeInSeconds * SampleRate);
+		LoopEndFrameIndex = FMath::Clamp(LoopEndTime * SampleRate, MinLoopEndFrameIndex, MaxLoopEndFrameIndex);
+	}
+	else
+	{
+		LoopStartFrameIndex = 0;
+		LoopEndFrameIndex = NumFramesInWave;
+	}
+}
+
+PRAGMA_ENABLE_DEPRECATION_WARNINGS
