[Backout] - CL26745178, CL26748184
Original CL Desc
-----------------------------------------------------------------
Add buoyancy subsystem plugin & initial prototype for midphase-buoyancy

CIS: resolve NumExtraFilterBits conflict in PhysicsFiltering.h and ParticleHandle.h

diff --git a/Engine/Plugins/Experimental/Buoyancy/Buoyancy.uplugin b/Engine/Plugins/Experimental/Buoyancy/Buoyancy.uplugin
deleted file mode 100644
diff --git a/Engine/Plugins/Experimental/Buoyancy/Source/Runtime/Buoyancy.Build.cs b/Engine/Plugins/Experimental/Buoyancy/Source/Runtime/Buoyancy.Build.cs
deleted file mode 100644
diff --git a/Engine/Plugins/Experimental/Buoyancy/Source/Runtime/Private/BuoyancyAlgorithms.cpp b/Engine/Plugins/Experimental/Buoyancy/Source/Runtime/Private/BuoyancyAlgorithms.cpp
deleted file mode 100644
diff --git a/Engine/Plugins/Experimental/Buoyancy/Source/Runtime/Private/BuoyancyModule.cpp b/Engine/Plugins/Experimental/Buoyancy/Source/Runtime/Private/BuoyancyModule.cpp
deleted file mode 100644
diff --git a/Engine/Plugins/Experimental/Buoyancy/Source/Runtime/Private/BuoyancySubsystem.cpp b/Engine/Plugins/Experimental/Buoyancy/Source/Runtime/Private/BuoyancySubsystem.cpp
deleted file mode 100644
diff --git a/Engine/Plugins/Experimental/Buoyancy/Source/Runtime/Public/BuoyancyAlgorithms.h b/Engine/Plugins/Experimental/Buoyancy/Source/Runtime/Public/BuoyancyAlgorithms.h
deleted file mode 100644
diff --git a/Engine/Plugins/Experimental/Buoyancy/Source/Runtime/Public/BuoyancyModule.h b/Engine/Plugins/Experimental/Buoyancy/Source/Runtime/Public/BuoyancyModule.h
deleted file mode 100644
diff --git a/Engine/Plugins/Experimental/Buoyancy/Source/Runtime/Public/BuoyancyStats.h b/Engine/Plugins/Experimental/Buoyancy/Source/Runtime/Public/BuoyancyStats.h
deleted file mode 100644
diff --git a/Engine/Plugins/Experimental/Buoyancy/Source/Runtime/Public/BuoyancySubsystem.h b/Engine/Plugins/Experimental/Buoyancy/Source/Runtime/Public/BuoyancySubsystem.h
deleted file mode 100644
diff --git a/Engine/Plugins/Experimental/Water/Source/Runtime/Public/WaterRuntimeSettings.h b/Engine/Plugins/Experimental/Water/Source/Runtime/Public/WaterRuntimeSettings.h
--- a/Engine/Plugins/Experimental/Water/Source/Runtime/Public/WaterRuntimeSettings.h
+++ b/Engine/Plugins/Experimental/Water/Source/Runtime/Public/WaterRuntimeSettings.h
@@ -66,11 +66,6 @@ public:
 	UPROPERTY(EditAnywhere, config, Category = Rendering)
 	float WaterBodyIconWorldZOffset = 75.0f;
 
-	/** Density of water used for buoyancy calculations. Given in gram per cubic centimeter */
-	// TODO: Add units to UnitConversions.h so that g/cm^3 can be explicit
-	UPROPERTY(EditAnywhere, config, Category = Buoyancy, meta=(ClampMin = 0, ForceUnits = "g/cm3"))
-	float WaterDensity = 1.f;
-
 #if WITH_EDITORONLY_DATA
 	// Delegate called whenever the curve data is updated
 	DECLARE_MULTICAST_DELEGATE_TwoParams(FOnUpdateSettings, const UWaterRuntimeSettings* /*Settings*/, EPropertyChangeType::Type /*ChangeType*/);
diff --git a/Engine/Source/Runtime/Engine/Public/Physics/PhysicsFiltering.h b/Engine/Source/Runtime/Engine/Public/Physics/PhysicsFiltering.h
--- a/Engine/Source/Runtime/Engine/Public/Physics/PhysicsFiltering.h
+++ b/Engine/Source/Runtime/Engine/Public/Physics/PhysicsFiltering.h
@@ -33,7 +33,7 @@ enum EPhysXFilterDataFlags
 // (ExtraFilter (top NumExtraFilterBits) + MyChannel (next NumCollisionChannelBits) as ECollisionChannel + Flags (remaining NumFilterDataFlagBits)
 // [NumExtraFilterBits] [NumCollisionChannelBits] [NumFilterDataFlagBits] = 32 bits
 enum { NumCollisionChannelBits = 5 };
-enum { NumFilterDataFlagBits = 32 - ::NumExtraFilterBits - ::NumCollisionChannelBits };
+enum { NumFilterDataFlagBits = 32 - NumExtraFilterBits - NumCollisionChannelBits };
 
 
 struct FPhysicsFilterBuilder
@@ -120,26 +120,26 @@ inline void CreateShapeFilterData(
 
 inline ECollisionChannel GetCollisionChannel(uint32 Word3)
 {
-	uint32 ChannelMask = (Word3 << ::NumExtraFilterBits) >> (32 - ::NumCollisionChannelBits);
+	uint32 ChannelMask = (Word3 << NumExtraFilterBits) >> (32 - NumCollisionChannelBits);
 	return (ECollisionChannel)ChannelMask;
 }
 
 inline ECollisionChannel GetCollisionChannelAndExtraFilter(uint32 Word3, FMaskFilter& OutMaskFilter)
 {
 	uint32 ChannelMask = GetCollisionChannel(Word3);
-	OutMaskFilter = Word3 >> (32 - ::NumExtraFilterBits);
+	OutMaskFilter = Word3 >> (32 - NumExtraFilterBits);
 	return (ECollisionChannel)ChannelMask;
 }
 
 inline uint32 CreateChannelAndFilter(ECollisionChannel CollisionChannel, FMaskFilter MaskFilter)
 {
-	uint32 ResultMask = (uint32(MaskFilter) << ::NumCollisionChannelBits) | (uint32)CollisionChannel;
+	uint32 ResultMask = (uint32(MaskFilter) << NumCollisionChannelBits) | (uint32)CollisionChannel;
 	return ResultMask << NumFilterDataFlagBits;
 }
 
 inline void UpdateMaskFilter(uint32& Word3, FMaskFilter NewMaskFilter)
 {
-	static_assert(::NumExtraFilterBits <= 8, "Only up to 8 extra filter bits are supported.");
-	Word3 &= (0xFFFFFFFFu >> ::NumExtraFilterBits);	//we drop the top NumExtraFilterBits bits because that's where the new mask filter is going
-	Word3 |= uint32(NewMaskFilter) << (32 - ::NumExtraFilterBits);
+	static_assert(NumExtraFilterBits <= 8, "Only up to 8 extra filter bits are supported.");
+	Word3 &= (0xFFFFFFFFu >> NumExtraFilterBits);	//we drop the top NumExtraFilterBits bits because that's where the new mask filter is going
+	Word3 |= uint32(NewMaskFilter) << (32 - NumExtraFilterBits);
 }
diff --git a/Engine/Source/Runtime/Experimental/Chaos/Public/Chaos/Collision/CollisionFilter.h b/Engine/Source/Runtime/Experimental/Chaos/Public/Chaos/Collision/CollisionFilter.h
--- a/Engine/Source/Runtime/Experimental/Chaos/Public/Chaos/Collision/CollisionFilter.h
+++ b/Engine/Source/Runtime/Experimental/Chaos/Public/Chaos/Collision/CollisionFilter.h
@@ -5,7 +5,6 @@
 
 #include "Chaos/Core.h"
 #include "Chaos/Collision/CollisionConstraintFlags.h"
-#include "Chaos/Collision/CollisionFilterBits.h"
 #include "Chaos/CollisionFilterData.h"
 #include "Chaos/GeometryParticles.h"
 #include "Chaos/ImplicitObjectType.h"
@@ -14,6 +13,8 @@
 namespace Chaos
 {
 	typedef uint8 FMaskFilter;
+	enum { NumExtraFilterBits = 6 };
+	enum { NumCollisionChannelBits = 5 };
 
 	/**
 	 * Check whether the two particles need to be considered in the broadphase.
diff --git a/Engine/Source/Runtime/Experimental/Chaos/Public/Chaos/Collision/CollisionFilterBits.h b/Engine/Source/Runtime/Experimental/Chaos/Public/Chaos/Collision/CollisionFilterBits.h
deleted file mode 100644
diff --git a/Engine/Source/Runtime/Experimental/Chaos/Public/Chaos/ParticleHandle.h b/Engine/Source/Runtime/Experimental/Chaos/Public/Chaos/ParticleHandle.h
--- a/Engine/Source/Runtime/Experimental/Chaos/Public/Chaos/ParticleHandle.h
+++ b/Engine/Source/Runtime/Experimental/Chaos/Public/Chaos/ParticleHandle.h
@@ -9,7 +9,6 @@
 #include "Chaos/ParticleHandleFwd.h"
 #include "Chaos/ParticleIterator.h"
 #include "Chaos/Properties.h"
-#include "Chaos/Collision/CollisionFilterBits.h"
 #include "ChaosCheck.h"
 #include "Chaos/ChaosDebugDrawDeclares.h"
 #if CHAOS_DEBUG_DRAW
@@ -161,8 +160,11 @@ FORCEINLINE_DEBUGGABLE bool PrePreQueryFilterImp(const FCollisionFilterData& Que
 
 FORCEINLINE_DEBUGGABLE uint32 GetChaosCollisionChannelAndExtraFilter(uint32 Word3, uint8& OutMaskFilter)
 {
-	uint32 ChannelMask = (Word3 << ChaosNumExtraFilterBits) >> (32 - ChaosNumCollisionChannelBits);
-	OutMaskFilter = Word3 >> (32 - ChaosNumExtraFilterBits);
+	enum { NumExtraFilterBits = 6 };
+	enum { NumCollisionChannelBits = 5 };
+
+	uint32 ChannelMask = (Word3 << NumExtraFilterBits) >> (32 - NumCollisionChannelBits);
+	OutMaskFilter = Word3 >> (32 - NumExtraFilterBits);
 	return (uint32)ChannelMask;
 }
 
