Missing SparseVolumeTexture inputs on SVT material nodes now cause a compiler error, which is the same behavior as for the corresponding nodes for regular textures.
Fixed class name of SVTObjectParameter node. This might potentially break materials using this node, but at this point that is fairly unlikely due to low adoption and most people using the Sample node directly instead.

#rnx

diff --git a/Engine/Source/Runtime/Engine/Classes/Materials/MaterialExpressionSparseVolumeTextureObject.h b/Engine/Source/Runtime/Engine/Classes/Materials/MaterialExpressionSparseVolumeTextureObject.h
--- a/Engine/Source/Runtime/Engine/Classes/Materials/MaterialExpressionSparseVolumeTextureObject.h
+++ b/Engine/Source/Runtime/Engine/Classes/Materials/MaterialExpressionSparseVolumeTextureObject.h
@@ -29,7 +29,7 @@ class UMaterialExpressionSparseVolumeTextureObject : public UMaterialExpressionS
 };
 
 UCLASS(collapsecategories, hidecategories = (Object, MaterialExpressionTextureSample), MinimalAPI)
-class UMaterialExpressionSparseVolumeTextureTextureObjectParameter : public UMaterialExpressionSparseVolumeTextureSampleParameter
+class UMaterialExpressionSparseVolumeTextureObjectParameter : public UMaterialExpressionSparseVolumeTextureSampleParameter
 {
 	GENERATED_UCLASS_BODY()
 
diff --git a/Engine/Source/Runtime/Engine/Private/Materials/MaterialExpressions.cpp b/Engine/Source/Runtime/Engine/Private/Materials/MaterialExpressions.cpp
--- a/Engine/Source/Runtime/Engine/Private/Materials/MaterialExpressions.cpp
+++ b/Engine/Source/Runtime/Engine/Private/Materials/MaterialExpressions.cpp
@@ -26678,10 +26678,10 @@ uint32 UMaterialExpressionSparseVolumeTextureObject::GetOutputType(int32 OutputI
 #endif // WITH_EDITOR
 
 ///////////////////////////////////////////////////////////////////////////////
-// UMaterialExpressionSparseVolumeTextureTextureObjectParameter
+// UMaterialExpressionSparseVolumeTextureObjectParameter
 ///////////////////////////////////////////////////////////////////////////////
 
-UMaterialExpressionSparseVolumeTextureTextureObjectParameter::UMaterialExpressionSparseVolumeTextureTextureObjectParameter(const FObjectInitializer& ObjectInitializer)
+UMaterialExpressionSparseVolumeTextureObjectParameter::UMaterialExpressionSparseVolumeTextureObjectParameter(const FObjectInitializer& ObjectInitializer)
 	: Super(ObjectInitializer)
 {
 	// Structure to hold one-time initialization
@@ -26714,24 +26714,24 @@ UMaterialExpressionSparseVolumeTextureTextureObjectParameter::UMaterialExpressio
 
 #if WITH_EDITOR
 
-void UMaterialExpressionSparseVolumeTextureTextureObjectParameter::GetCaption(TArray<FString>& OutCaptions) const
+void UMaterialExpressionSparseVolumeTextureObjectParameter::GetCaption(TArray<FString>& OutCaptions) const
 {
 	OutCaptions.Add(TEXT("Param Tex Object"));
 	OutCaptions.Add(FString::Printf(TEXT("'%s'"), *ParameterName.ToString()));
 }
 
-int32 UMaterialExpressionSparseVolumeTextureTextureObjectParameter::Compile(class FMaterialCompiler* Compiler, int32 OutputIndex)
+int32 UMaterialExpressionSparseVolumeTextureObjectParameter::Compile(class FMaterialCompiler* Compiler, int32 OutputIndex)
 {
 	if (!SparseVolumeTexture)
 	{
-		return INDEX_NONE;
+		return CompilerError(Compiler, TEXT("Requires valid SparseVolumeTexture"));
 	}
 
 	int32 TextureReferenceIndex = INDEX_NONE;
 	return Compiler->SparseVolumeTextureParameter(ParameterName, SparseVolumeTexture, TextureReferenceIndex, SAMPLERTYPE_LinearColor);
 }
 
-uint32 UMaterialExpressionSparseVolumeTextureTextureObjectParameter::GetOutputType(int32 OutputIndex)
+uint32 UMaterialExpressionSparseVolumeTextureObjectParameter::GetOutputType(int32 OutputIndex)
 {
 	return MCT_SparseVolumeTexture;
 }
@@ -26885,8 +26885,10 @@ int32 UMaterialExpressionSparseVolumeTextureSample::Compile(class FMaterialCompi
 			return Mip0SampleIndex;
 		}
 	}
-
-	return INDEX_NONE;
+	else
+	{
+		return CompilerError(Compiler, TEXT("Missing input texture"));
+	}
 }
 
 void UMaterialExpressionSparseVolumeTextureSample::GetCaption(TArray<FString>& OutCaptions) const
