Fix some added string cost for Blueprint class assets (regression/perfmem).

diff --git a/Engine/Source/Runtime/Engine/Classes/Engine/BlueprintGeneratedClass.h b/Engine/Source/Runtime/Engine/Classes/Engine/BlueprintGeneratedClass.h
--- a/Engine/Source/Runtime/Engine/Classes/Engine/BlueprintGeneratedClass.h
+++ b/Engine/Source/Runtime/Engine/Classes/Engine/BlueprintGeneratedClass.h
@@ -627,7 +627,7 @@ struct FBPComponentClassOverride
 	}
 };
 
-UCLASS(NeedsDeferredDependencyLoading, MinimalAPI, config=Engine)
+UCLASS(NeedsDeferredDependencyLoading, MinimalAPI)
 class UBlueprintGeneratedClass : public UClass, public IBlueprintPropertyGuidProvider
 {
 	GENERATED_UCLASS_BODY()
@@ -923,9 +923,6 @@ protected:
 	ENGINE_API bool BuildCustomArrayPropertyListForPostConstruction(FArrayProperty* ArrayProperty, FCustomPropertyListNode*& InPropertyList, const uint8* DataPtr, const uint8* DefaultDataPtr, int32 StartIndex = 0);
 
 private:
-	/** List of native class-owned properties that require complete values for comparison when generating the post-construction property list due to how they are used. These properties cannot be reduced into a list of subfields (e.g. structs). */
-	UPROPERTY(config)
-	TArray<FString> RequiresCompleteValueForPostConstruction;
 	/** List of native class-owned properties that differ from defaults. This is used to optimize property initialization during post-construction by minimizing the number of native class-owned property values that get copied to the new instance. */
 	TIndirectArray<FCustomPropertyListNode> CustomPropertyListForPostConstruction;
 	/** In some cases UObject::ConditionalPostLoad() code calls PostLoadDefaultObject() on a class that's still being serialized. */
diff --git a/Engine/Source/Runtime/Engine/Private/BlueprintGeneratedClass.cpp b/Engine/Source/Runtime/Engine/Private/BlueprintGeneratedClass.cpp
--- a/Engine/Source/Runtime/Engine/Private/BlueprintGeneratedClass.cpp
+++ b/Engine/Source/Runtime/Engine/Private/BlueprintGeneratedClass.cpp
@@ -36,6 +36,7 @@
 #include "BlueprintCompilationManager.h"
 #include "Engine/LevelScriptBlueprint.h"
 #include "WorldPartition/WorldPartitionActorDescUtils.h"
+#include "UObject/UObjectGlobals.h"
 extern UNREALED_API class UEditorEngine* GEditor;
 #else
 #include "UObject/LinkerLoad.h"
@@ -116,6 +117,61 @@ namespace UE::Runtime::Engine::Private
 			return CurrentClass;
 		}
 
+		static bool RequiresCompleteValueForPostConstruction(FProperty* Property)
+		{
+			static TSet<FProperty*> PropertiesRequiringCompleteValueInitialization;
+
+			static bool bIsInitialized = false;
+#if WITH_EDITOR
+			FCoreUObjectDelegates::ReloadCompleteDelegate.AddLambda([](EReloadCompleteReason Reason)
+			{
+				switch(Reason)
+				{
+				case EReloadCompleteReason::HotReloadManual:
+				case EReloadCompleteReason::HotReloadAutomatic:
+				{
+					// Re-initialize after hot reload in editor context.
+					bIsInitialized = false;
+					PropertiesRequiringCompleteValueInitialization.Reset();
+				}
+					break;
+
+				default:
+					break;
+				}
+			});
+#endif	// WITH_EDITOR
+
+			if (!bIsInitialized)
+			{
+				static constexpr TCHAR ConfigSection[] = TEXT("/Script/Engine.BlueprintGeneratedClass");
+				static constexpr TCHAR ConfigKeyName[] = TEXT("RequiresCompleteValueForPostConstruction");
+
+				// List of native class-owned properties that require complete values for comparison when generating the post-construction
+				// property list due to how they are used. These properties can't be further reduced into a list of subfields (e.g. structs).
+				TArray<FString> RequiresCompleteValueForPostConstruction_Value;
+				if (GConfig->GetArray(ConfigSection, ConfigKeyName, RequiresCompleteValueForPostConstruction_Value, GEngineIni))
+				{
+					PropertiesRequiringCompleteValueInitialization.Reserve(RequiresCompleteValueForPostConstruction_Value.Num());
+
+					for (const FString& PropertyPathString : RequiresCompleteValueForPostConstruction_Value)
+					{
+						TFieldPath<FProperty> PropertyPath;
+						PropertyPath.Generate(*PropertyPathString);
+						if (FProperty* ResolvedProperty = PropertyPath.Get())
+						{
+							check(ResolvedProperty->IsNative());
+							PropertiesRequiringCompleteValueInitialization.Add(ResolvedProperty);
+						}
+					}
+				}
+
+				bIsInitialized = true;
+			}
+
+			return PropertiesRequiringCompleteValueInitialization.Contains(Property);
+		}
+
 		static bool ShouldInitializePropertyDuringPostConstruction(const FProperty& Property)
 		{
 			const UClass* OwnerClass = Property.GetOwnerClass();
@@ -764,9 +820,7 @@ bool UBlueprintGeneratedClass::BuildCustomPropertyListForPostConstruction(FCusto
 		if (FBlueprintGeneratedClassUtils::ShouldInitializePropertyDuringPostConstruction(*Property))
 		{
 			// Some properties require a full value comparison; check for those cases here.
-			TStringBuilder<256> PathName;
-			Property->GetPathName(nullptr, PathName);
-			const bool bAlwaysUseCompleteValue = RequiresCompleteValueForPostConstruction.ContainsByPredicate([&](const FString& S) { return S == PathName.ToView(); });
+			const bool bAlwaysUseCompleteValue = FBlueprintGeneratedClassUtils::RequiresCompleteValueForPostConstruction(Property);
 
 			for (int32 Idx = 0; Idx < Property->ArrayDim; Idx++)
 			{
