PushModel:
* Calling ForceCompare() will now force push model properties to be evaluated even if they are not flagged dirty.

diff --git a/Engine/Source/Runtime/Engine/Private/DataChannel.cpp b/Engine/Source/Runtime/Engine/Private/DataChannel.cpp
--- a/Engine/Source/Runtime/Engine/Private/DataChannel.cpp
+++ b/Engine/Source/Runtime/Engine/Private/DataChannel.cpp
@@ -3251,7 +3251,9 @@ public:
 		{
 			if ( !RepFlags.bNetOwner )
 			{
-				Actor->SetAutonomousProxy( false, false );
+				constexpr bool bIsAutonomousProxy = false; // become simulated
+				constexpr bool bAllowForcePropertyCompare = false;
+				Actor->SetAutonomousProxy(bIsAutonomousProxy, bAllowForcePropertyCompare);
 			}
 		}
 	}
@@ -3265,7 +3267,9 @@ public:
 
 			if (ActualRemoteRole == ROLE_AutonomousProxy)
 			{
-				Actor->SetAutonomousProxy(true, false);
+				constexpr bool bIsAutonomousProxy = true; // return to autonomous
+				constexpr bool bAllowForcePropertyCompare = false;
+				Actor->SetAutonomousProxy(bIsAutonomousProxy, bAllowForcePropertyCompare);
 			}
 		}
 	}
diff --git a/Engine/Source/Runtime/Engine/Private/RepLayout.cpp b/Engine/Source/Runtime/Engine/Private/RepLayout.cpp
--- a/Engine/Source/Runtime/Engine/Private/RepLayout.cpp
+++ b/Engine/Source/Runtime/Engine/Private/RepLayout.cpp
@@ -1268,7 +1268,7 @@ ERepLayoutResult FRepLayout::UpdateChangelistMgr(
 			{
 				FReplicationFlags TempFlags = RepFlags;
 				TempFlags.bRolesOnly = true;
-				Result = CompareProperties(RepState, &InChangelistMgr.RepChangelistState, (const uint8*)InObject, TempFlags);
+				Result = CompareProperties(RepState, &InChangelistMgr.RepChangelistState, (const uint8*)InObject, TempFlags, bForceCompare);
 			}
 
 			INC_DWORD_STAT_BY(STAT_NetSkippedDynamicProps, 1);
@@ -1290,7 +1290,7 @@ ERepLayoutResult FRepLayout::UpdateChangelistMgr(
 		}
 	}
 
-	Result = CompareProperties(RepState, &InChangelistMgr.RepChangelistState, (const uint8*)InObject, RepFlags);
+	Result = CompareProperties(RepState, &InChangelistMgr.RepChangelistState, (const uint8*)InObject, RepFlags, bForceCompare);
 
 	// Currently, comparing properties should only result in Success, Empty, or FatalError.
 	// So, don't bother checking for normal errors.
@@ -1324,6 +1324,7 @@ struct FComparePropertiesSharedParams
 	const bool bIsNetworkProfilerActive = false;
 	const bool bChangedNetOwner = false;
 	const bool bForceCustomPropsActive = false;
+	const bool bForceCompareProperties = false;
 #if (WITH_PUSH_VALIDATION_SUPPORT || USE_NETWORK_PROFILER)
 	TBitArray<> PropertiesCompared;
 	TBitArray<> PropertiesChanged;
@@ -1448,7 +1449,7 @@ static bool CompareParentProperty(
 	return !!(StackParams.Changed.Num() - NumChanges);
 }
 
-namespace UE_RepLayout_Private
+namespace UE::Net::Private
 {
 	static bool CompareParentPropertyHelper(
 		const int32 ParentIndex,
@@ -1474,10 +1475,11 @@ namespace UE_RepLayout_Private
 		const FComparePropertiesSharedParams& SharedParams,
 		FComparePropertiesStackParams& StackParams)
 	{
-		return !(*SharedParams.PushModelProperties)[ParentIndex] ||
+		return SharedParams.bForceCompareProperties ||
+			!(*SharedParams.PushModelProperties)[ParentIndex] || // non-push model properties are always considered dirty			
 			SharedParams.PushModelState->IsPropertyDirty(ParentIndex) ||
 			(bRecentlyCollectedGarbage &&
-			EnumHasAnyFlags(SharedParams.Parents[ParentIndex].Flags, ERepParentFlags::HasObjectProperties | ERepParentFlags::IsNetSerialize));
+				EnumHasAnyFlags(SharedParams.Parents[ParentIndex].Flags, ERepParentFlags::HasObjectProperties | ERepParentFlags::IsNetSerialize));
 	}
 #endif // WITH_PUSH_MODEL	
 }	
@@ -1486,6 +1488,7 @@ static void CompareParentProperties(
 	const FComparePropertiesSharedParams& SharedParams,
 	FComparePropertiesStackParams& StackParams)
 {
+	using namespace UE::Net::Private;
 	check(StackParams.ShadowData);
 
 #if WITH_PUSH_MODEL
@@ -1515,7 +1518,7 @@ static void CompareParentProperties(
 
 			for (int32 ParentIndex = 0; ParentIndex < SharedParams.Parents.Num(); ++ParentIndex)
 			{
-				UE_RepLayout_Private::CompareParentPropertyHelper(ParentIndex, SharedParams, StackParams);
+				CompareParentPropertyHelper(ParentIndex, SharedParams, StackParams);
 			}
 		}
 
@@ -1530,10 +1533,9 @@ static void CompareParentProperties(
 				const ELifetimeCondition Condition = SharedParams.Parents[ParentIndex].Condition;
 				const bool bRecompareInitialProperties = SharedParams.bIsInitial && (Condition == COND_InitialOnly || (Condition == COND_Dynamic && SharedParams.RepChangedPropertyTracker && SharedParams.RepChangedPropertyTracker->GetDynamicCondition(ParentIndex) == COND_InitialOnly));
 
-				const bool bIsPropertyDirty = bRecompareInitialProperties ||
-												UE_RepLayout_Private::IsPropertyDirty(ParentIndex, bRecentlyCollectedGarbage, SharedParams, StackParams);
+				const bool bIsPropertyDirty = bRecompareInitialProperties || IsPropertyDirty(ParentIndex, bRecentlyCollectedGarbage, SharedParams, StackParams);
 				
-				const bool bDidPropertyChange = UE_RepLayout_Private::CompareParentPropertyHelper(ParentIndex, SharedParams, StackParams);
+				const bool bDidPropertyChange = CompareParentPropertyHelper(ParentIndex, SharedParams, StackParams);
 
 				ensureAlwaysMsgf(!bDidPropertyChange || bIsPropertyDirty, TEXT("Push Model Property changed value, but was not marked dirty! Property=%s"), *SharedParams.Parents[ParentIndex].Property->GetPathName());
 			}	
@@ -1573,10 +1575,10 @@ static void CompareParentProperties(
 			for (int32 ParentIndex = 0; ParentIndex < SharedParams.Parents.Num(); ++ParentIndex)
 			{
 				const ELifetimeCondition Condition = SharedParams.Parents[ParentIndex].Condition;
-				if (Condition == COND_InitialOnly || UE_RepLayout_Private::IsPropertyDirty(ParentIndex, bRecentlyCollectedGarbage, SharedParams, StackParams)
-					|| (Condition == COND_Dynamic && SharedParams.RepChangedPropertyTracker && SharedParams.RepChangedPropertyTracker->GetDynamicCondition(ParentIndex) == COND_InitialOnly))
+				if (Condition == COND_InitialOnly || IsPropertyDirty(ParentIndex, bRecentlyCollectedGarbage, SharedParams, StackParams) ||
+					(Condition == COND_Dynamic && SharedParams.RepChangedPropertyTracker && SharedParams.RepChangedPropertyTracker->GetDynamicCondition(ParentIndex) == COND_InitialOnly))
 				{
-					UE_RepLayout_Private::CompareParentPropertyHelper(ParentIndex, SharedParams, StackParams);
+					CompareParentPropertyHelper(ParentIndex, SharedParams, StackParams);
 				}
 			}
 		}
@@ -1588,7 +1590,7 @@ static void CompareParentProperties(
 
 			for (TConstSetBitIterator<> It = SharedParams.PushModelState->GetDirtyProperties(); It; ++It)
 			{
-				UE_RepLayout_Private::CompareParentPropertyHelper(It.GetIndex(), SharedParams, StackParams);
+				CompareParentPropertyHelper(It.GetIndex(), SharedParams, StackParams);
 			}
 		}
 		else
@@ -1597,9 +1599,9 @@ static void CompareParentProperties(
 
 			for (int32 ParentIndex = 0; ParentIndex < SharedParams.Parents.Num(); ++ParentIndex)
 			{
-				if (UE_RepLayout_Private::IsPropertyDirty(ParentIndex, bRecentlyCollectedGarbage, SharedParams, StackParams))
+				if (IsPropertyDirty(ParentIndex, bRecentlyCollectedGarbage, SharedParams, StackParams))
 				{
-					UE_RepLayout_Private::CompareParentPropertyHelper(ParentIndex, SharedParams, StackParams);
+					CompareParentPropertyHelper(ParentIndex, SharedParams, StackParams);
 				}
 			}
 		}
@@ -1611,7 +1613,7 @@ static void CompareParentProperties(
 
 	for (int32 ParentIndex = 0; ParentIndex < SharedParams.Parents.Num(); ++ParentIndex)
 	{
-		UE_RepLayout_Private::CompareParentPropertyHelper(ParentIndex, SharedParams, StackParams);
+		CompareParentPropertyHelper(ParentIndex, SharedParams, StackParams);
 	}
 }
 
@@ -1749,7 +1751,8 @@ ERepLayoutResult FRepLayout::CompareProperties(
 	FSendingRepState* RESTRICT RepState,
 	FRepChangelistState* RESTRICT RepChangelistState,
 	const FConstRepObjectDataBuffer Data,
-	const FReplicationFlags& RepFlags) const
+	const FReplicationFlags& RepFlags,
+	const bool bForceCompare) const
 {
 	CONDITIONAL_SCOPE_CYCLE_COUNTER(STAT_NetReplicateDynamicPropCompareTime, GUseDetailedScopeCounters);
 
@@ -1778,29 +1781,32 @@ ERepLayoutResult FRepLayout::CompareProperties(
 
 	ERepLayoutResult Result = ERepLayoutResult::Success;
 
-	FComparePropertiesSharedParams SharedParams{
-		/*bIsInitial=*/ !!RepFlags.bNetInitial,
-		/*bForceFail=*/ !!RepFlags.bNetInitial && !!RepFlags.bForceInitialDirty,
-		Flags,
-		Parents,
-		Cmds,
-		RepState,
-		RepChangelistState,
-		(RepState ? RepState->RepChangedPropertyTracker.Get() : nullptr),
-		NetSerializeLayouts,
-		/*PushModelState=*/UE_RepLayout_Private::GetPerNetDriverState(RepChangelistState),
-		/*PushModelProperties=*/ LocalPushModelProperties,
-		/*bValidateProperties=*/GbPushModelValidateProperties,
-		/*bIsNetworkProfilerActive=*/UE_RepLayout_Private::IsNetworkProfilerComparisonTrackingEnabled(),
-		/*bChangedNetOwner=*/ RepState && RepState->RepFlags.bNetOwner != RepFlags.bNetOwner,
-		/*bForceCustomPropsActive*/ !!RepFlags.bClientReplay
+	FComparePropertiesSharedParams SharedParams
+	{
+		.bIsInitial = !!RepFlags.bNetInitial,
+		.bForceFail = !!RepFlags.bNetInitial && !!RepFlags.bForceInitialDirty,
+		.Flags = Flags,
+		.Parents = Parents,
+		.Cmds = Cmds,
+		.RepState = RepState,
+		.RepChangelistState = RepChangelistState,
+		.RepChangedPropertyTracker = (RepState ? RepState->RepChangedPropertyTracker.Get() : nullptr),
+		.NetSerializeLayouts = NetSerializeLayouts,
+		.PushModelState = UE_RepLayout_Private::GetPerNetDriverState(RepChangelistState),
+		.PushModelProperties = LocalPushModelProperties,
+		.bValidateProperties = GbPushModelValidateProperties,
+		.bIsNetworkProfilerActive = UE_RepLayout_Private::IsNetworkProfilerComparisonTrackingEnabled(),
+		.bChangedNetOwner = RepState && RepState->RepFlags.bNetOwner != RepFlags.bNetOwner,
+		.bForceCustomPropsActive = !!RepFlags.bClientReplay,
+		.bForceCompareProperties = bForceCompare
 	};
 
-	FComparePropertiesStackParams StackParams{
-		Data,
-		RepChangelistState->StaticBuffer.GetData(),
-		Changed,
-		Result
+	FComparePropertiesStackParams StackParams
+	{
+		.Data = Data,
+		.ShadowData = RepChangelistState->StaticBuffer.GetData(),
+		.Changed = Changed,
+		.Result = Result
 	};
 
 	if (RepFlags.bRolesOnly)
diff --git a/Engine/Source/Runtime/Engine/Public/Net/RepLayout.h b/Engine/Source/Runtime/Engine/Public/Net/RepLayout.h
--- a/Engine/Source/Runtime/Engine/Public/Net/RepLayout.h
+++ b/Engine/Source/Runtime/Engine/Public/Net/RepLayout.h
@@ -1492,12 +1492,14 @@ private:
 	 * @param RepChangelistState	The FRepChangelistState that contains the last cached values and changelists.
 	 * @param Data					The newest Property Data available.
 	 * @param RepFlags				Flags that will be used if the object is replicated.
+	 * @param bForceCompare			Compare the property even if the dirty flag is not set.
 	 */
 	ERepLayoutResult CompareProperties(
 		FSendingRepState* RESTRICT RepState,
 		FRepChangelistState* RESTRICT RepChangelistState,
 		const FConstRepObjectDataBuffer Data,
-		const FReplicationFlags& RepFlags) const;
+		const FReplicationFlags& RepFlags,
+		const bool bForceCompare) const;
 
 	/**
 	 * Writes all changed property values from the input owner data to the given buffer.
