World Partition: fix crash when saving a newly created actor that was replaced with another class before initial save.

diff --git a/Engine/Source/Runtime/Engine/Private/WorldPartition/WorldPartition.cpp b/Engine/Source/Runtime/Engine/Private/WorldPartition/WorldPartition.cpp
--- a/Engine/Source/Runtime/Engine/Private/WorldPartition/WorldPartition.cpp
+++ b/Engine/Source/Runtime/Engine/Private/WorldPartition/WorldPartition.cpp
@@ -1247,8 +1247,8 @@ void UWorldPartition::OnActorDescAdded(FWorldPartitionActorDesc* NewActorDesc)
 	{
 		if (FDirtyActor* ExistingDirtyActor = DirtyActors.Find(NewActorDesc->GetGuid()))
 		{
-			check(ExistingDirtyActor->ActorPtr == NewActor);
 			check(!ExistingDirtyActor->WorldPartitionRef.IsSet());
+			check(!ExistingDirtyActor->ActorPtr.IsValid() || ExistingDirtyActor->ActorPtr == NewActor);			
 			ExistingDirtyActor->WorldPartitionRef = FWorldPartitionReference(NewActorDesc->GetContainer(), NewActorDesc->GetGuid());
 		}
 		else
