CIS compile fix.

diff --git a/Engine/Source/Runtime/Engine/Private/BlueprintGeneratedClass.cpp b/Engine/Source/Runtime/Engine/Private/BlueprintGeneratedClass.cpp
--- a/Engine/Source/Runtime/Engine/Private/BlueprintGeneratedClass.cpp
+++ b/Engine/Source/Runtime/Engine/Private/BlueprintGeneratedClass.cpp
@@ -21,6 +21,7 @@
 #include "Engine/InheritableComponentHandler.h"
 #include "IFieldNotificationClassDescriptor.h"
 #include "INotifyFieldValueChanged.h"
+#include "Misc/ConfigCacheIni.h"
 #include "Net/Core/PushModel/PushModel.h"
 #include "UObject/FortniteMainBranchObjectVersion.h"
 #include "UObject/Package.h" // IWYU pragma: keep
@@ -142,7 +143,7 @@ namespace UE::Runtime::Engine::Private
 			});
 #endif	// WITH_EDITOR
 
-			if (!bIsInitialized)
+			if (!bIsInitialized && GConfig)
 			{
 				static constexpr TCHAR ConfigSection[] = TEXT("/Script/Engine.BlueprintGeneratedClass");
 				static constexpr TCHAR ConfigKeyName[] = TEXT("RequiresCompleteValueForPostConstruction");
