[BugFix] Splitting the struct pin of a Post Process Settings node in an ABP removes all variable pins

diff --git a/Engine/Source/Runtime/Engine/Classes/Engine/Scene.h b/Engine/Source/Runtime/Engine/Classes/Engine/Scene.h
--- a/Engine/Source/Runtime/Engine/Classes/Engine/Scene.h
+++ b/Engine/Source/Runtime/Engine/Classes/Engine/Scene.h
@@ -668,7 +668,7 @@ struct FPostProcessSettingsDebugInfo
 // Each property consists of a bool to enable it (by default off),
 // the variable declaration and further down the default value for it.
 // The comment should include the meaning and usable range.
-USTRUCT(BlueprintType, meta=(HiddenByDefault))
+USTRUCT(BlueprintType, meta=(HiddenByDefault, DisableSplitPin))
 struct FPostProcessSettings
 {
 	GENERATED_USTRUCT_BODY()
