Shader Compiler:  FShaderJobCache::SubmitJobs now runs in multiple tasks, with fixes necessary for thread safety and minimizing lock contention applied.  In a Lyra cook, the time in SubmitJobs on the main thread is now around a fifth of the time spent in GlobalBeginCompileShader, and is thus no longer the main bottleneck on the game thread.  Lock contention is only 0.04% of the time in the SubmitJobs tasks, so further lock contention reduction wouldn't help performance.

* FShaderJobCache::ProcessFinishedJob needs to lock CompileQueueSection from FShaderCompilingManager to work from any thread.  Originally, CompileQueueSection was locked higher up in the call stack in FShaderCompilingManager::SubmitJobs.  CompileQueueSection is propagated to the FShaderJobCache constructor to make it available.
* Locking CompileQueueSection there creates deadlock issues if you also lock JobLock from FShaderJobCache, because certain code paths take the locks in the opposite order.  To solve this, we must not take both locks at the same time.  The two places ProcessFinishedJob are called from job cache logic are both moved outside of JobLock scopes.
* To reduce lock contention, the hash table of active jobs no longer uses JobLock.  Game thread related tasks allocate jobs via PrepareJob, adding to the job hash table.  An array of allocated jobs is then passed to SubmitJobs, which doesn't itself access the job hash table at all, so sharing the same lock creates a needless dependency.  Striping by the high bits of the hash key is used to further reduce contention in cases where multiple threads are generating jobs (not sure if that happens now, but maybe it will in the future).
* A wait free linked list algorithm is used to insert items to the queue of pending jobs that haven't been assigned to a worker.  Atomic linked list operations are supported for both head and tail insertion (the latter required for the default FIFO job execution mode).  Tail insertion requires maintaining a tail pointer, and thus can't use the original Core singly linked list class, and the doubly linked list class in Core fundamentally can't support atomic operations, so a version of the singly linked list implementation is copied locally to the ShaderCompiler.cpp, and adapted to our purposes.
* Wait free queue insertion is safe across multiple producer threads, which means we only need an FReadScopeLock for insertion.  Write locks are required for other list operations, but queue insertion is the massively parallel operation we are most concerned about.  Queue removal happens in a much smaller number of manager threads that call GetPendingJobs.
* Because a tail pointer is always maintained for the FIFO, insertion doesn't need to traverse the linked list to find the tail, which is potentially O(N^2), reducing the time a lock is held.
* SerializeOutput no longer occurs inside JobLock scopes for processing existing output and duplicate jobs, again reducing the time a lock is held.

#rnx

diff --git a/Engine/Source/Runtime/Engine/Private/ShaderCompiler/ShaderCompiler.cpp b/Engine/Source/Runtime/Engine/Private/ShaderCompiler/ShaderCompiler.cpp
--- a/Engine/Source/Runtime/Engine/Private/ShaderCompiler/ShaderCompiler.cpp
+++ b/Engine/Source/Runtime/Engine/Private/ShaderCompiler/ShaderCompiler.cpp
@@ -241,6 +241,69 @@ static FAutoConsoleVariableRef CVarShaderCompilerDebugDiscardCacheOutputs(
 	ECVF_Default
 );
 
+static TAutoConsoleVariable<bool> CVarShaderCompilerParallelSubmitJobs(
+	TEXT("r.ShaderCompiler.ParallelSubmitJobs"),
+	true,
+	TEXT("if != 0, FShaderJobCache::SubmitJobs will run in multiple parallel tasks, instead of the game thread."),
+	ECVF_Default
+);
+
+
+/** Copy of TIntrusiveLinkedListIterator, specific to FShaderCommonCompileJob */
+class FShaderCommonCompileJobIterator
+{
+public:
+	explicit FShaderCommonCompileJobIterator(FShaderCommonCompileJob* FirstLink)
+		: CurrentLink(FirstLink)
+	{ }
+
+	/**
+	 * Advances the iterator to the next element.
+	 */
+	FORCEINLINE void Next()
+	{
+		checkSlow(CurrentLink);
+		CurrentLink = CurrentLink->NextLink;
+	}
+
+	FORCEINLINE FShaderCommonCompileJobIterator& operator++()
+	{
+		Next();
+		return *this;
+	}
+
+	FORCEINLINE FShaderCommonCompileJobIterator operator++(int)
+	{
+		auto Tmp = *this;
+		Next();
+		return Tmp;
+	}
+
+	/** conversion to "bool" returning true if the iterator is valid. */
+	FORCEINLINE explicit operator bool() const
+	{
+		return CurrentLink != nullptr;
+	}
+
+	FORCEINLINE bool operator==(const FShaderCommonCompileJobIterator& Rhs) const { return CurrentLink == Rhs.CurrentLink; }
+	FORCEINLINE bool operator!=(const FShaderCommonCompileJobIterator& Rhs) const { return CurrentLink != Rhs.CurrentLink; }
+
+	// Accessors.
+	FORCEINLINE FShaderCommonCompileJob& operator->() const
+	{
+		checkSlow(this->CurrentLink);
+		return *(this->CurrentLink);
+	}
+
+	FORCEINLINE FShaderCommonCompileJob& operator*() const
+	{
+		checkSlow(this->CurrentLink);
+		return *(this->CurrentLink);
+	}
+
+private:
+	FShaderCommonCompileJob* CurrentLink;
+};
 
 /** Map element type for job cache */
 struct FShaderJobData
@@ -404,61 +467,77 @@ public:
 };
 
 
-/** Private implementation class for FShaderCompileJobCollection */
-class FShaderJobCache
+/**
+ * Class that provides a lock striped hash table of jobs, to reduce lock contention when adding or removing jobs
+ */
+class FShaderCompilerJobTable
 {
 public:
-	FShaderJobCache();
-	~FShaderJobCache();
+	static const int32 NUM_STRIPE_BITS = 6;
+	static const int32 NUM_STRIPES = 1 << NUM_STRIPE_BITS;
 
-	FShaderCompileJob* PrepareJob(uint32 InId, const FShaderCompileJobKey& InKey, EShaderCompileJobPriority InPriority);
-	FShaderPipelineCompileJob* PrepareJob(uint32 InId, const FShaderPipelineCompileJobKey& InKey, EShaderCompileJobPriority InPriority);
-	void RemoveJob(FShaderCommonCompileJob* InJob);
+	/** We want to use the high bits of the hash for the stripe index, as it won't have influence on the hash table index within the stripe */
+	static const int32 STRIPE_SHIFT = 32 - NUM_STRIPE_BITS;
 
-	int32 RemoveAllPendingJobsWithId(uint32 InId);
+	template<typename JobType, typename KeyType>
+	JobType* PrepareJob(uint32 InId, const KeyType& InKey, EShaderCompileJobPriority InPriority, bool& bOutNewJob)
+	{
+		const uint32 Hash = InKey.MakeHash(InId);
+		FLockStripeData& Stripe = GetStripe(JobType::Type, Hash);
 
-	void SubmitJobs(const TArray<FShaderCommonCompileJobPtr>& InJobs);
+		FWriteScopeLock Locker(Stripe.StripeLock);
 
-	/** This is an entry point for all jobs that have finished the compilation (whether real or cached). Can be called from multiple threads.*/
-	void ProcessFinishedJob(FShaderCommonCompileJob* FinishedJob, bool bWasCached);
+		JobType* ResultJob = InternalFindJob<JobType>(Hash, InId, InKey);
+		bOutNewJob = false;
 
-	/** Adds the job to cache. */
-	void AddToCacheAndProcessPending(FShaderCommonCompileJob* FinishedJob);
+		if (ResultJob == nullptr)
+		{
+			ResultJob = new JobType(Hash, InId, InPriority, InKey);
+			InternalAddJob(ResultJob);
+			bOutNewJob = true;
+		}
 
-	/** Populates caching stats in the given compiler stats struct. */
-	void GetStats(FShaderCompilerStats& OutStats) const;
+		return ResultJob;
+	}
 
-	int32 GetNumPendingJobs(EShaderCompileJobPriority InPriority) const;
+	// PrepareJob creates a job with the given key if it's unique, while this adds an existing job, typically one that is cloned from another job
+	void AddExistingJob(FShaderCommonCompileJob* InJob)
+	{
+		FLockStripeData& Stripe = GetStripe(InJob->Type, InJob->Hash);
 
-	int32 GetNumOutstandingJobs() const;
+		FWriteScopeLock Locker(Stripe.StripeLock);
+		InternalAddJob(InJob);
+	}
 
-	int32 GetNumPendingJobs() const;
+	void RemoveJob(FShaderCommonCompileJob* InJob)
+	{
+		FLockStripeData& Stripe = GetStripe(InJob->Type, InJob->Hash);
 
-	int32 GetPendingJobs(EShaderCompilerWorkerType InWorkerType, EShaderCompileJobPriority InPriority, int32 MinNumJobs, int32 MaxNumJobs, TArray<FShaderCommonCompileJobPtr>& OutJobs);
+		FWriteScopeLock Locker(Stripe.StripeLock);
 
-private:
-	using FJobInputHash = FShaderCommonCompileJob::FInputHash;
-	using FJobOutputHash = FBlake3Hash;
-	using FJobCachedOutput = FSharedBuffer;
-	using FStoredOutput = FShaderJobCacheStoredOutput;
+		const int32 JobIndex = InJob->JobIndex;
 
-	// cannot allow managing this from outside as the caching logic is not exposed
-	inline int32 InternalSubtractNumOutstandingJobs(int32 Value)
-	{
-		const int32 PrevNumOutstandingJobs = NumOutstandingJobs.Subtract(Value);
-		check(PrevNumOutstandingJobs >= Value);
-		return PrevNumOutstandingJobs - Value;
+		check(JobIndex != INDEX_NONE);
+		check(Stripe.Jobs[JobIndex] == InJob);
+		check(InJob->PendingPriority == EShaderCompileJobPriority::None);
+		InJob->JobIndex = INDEX_NONE;
+
+		Stripe.JobHash.Remove(InJob->Hash, JobIndex);
+		Stripe.FreeIndices.Add(JobIndex);
+		Stripe.Jobs[JobIndex].SafeRelease();
 	}
 
+private:
 	template<typename JobType, typename KeyType>
-	int32 InternalFindJobIndex(uint32 InJobHash, uint32 InJobId, const KeyType& InKey) const
+	JobType* InternalFindJob(uint32 InJobHash, uint32 InJobId, const KeyType& InKey) const
 	{
-		const int32 TypeIndex = (int32)JobType::Type;
+		const FLockStripeData& Stripe = GetStripe(JobType::Type, InJobHash);
+
 		uint32 CurrentPriorityIndex = 0u;
 		int32 CurrentIndex = INDEX_NONE;
-		for (int32 Index = JobHash[TypeIndex].First(InJobHash); JobHash[TypeIndex].IsValid(Index); Index = JobHash[TypeIndex].Next(Index))
+		for (int32 Index = Stripe.JobHash.First(InJobHash); Stripe.JobHash.IsValid(Index); Index = Stripe.JobHash.Next(Index))
 		{
-			const FShaderCommonCompileJob* Job = Jobs[TypeIndex][Index].GetReference();
+			const FShaderCommonCompileJob* Job = Stripe.Jobs[Index].GetReference();
 			check(Job->Type == JobType::Type);
 
 			// We find the job that matches the key with the highest priority
@@ -470,51 +549,124 @@ private:
 				CurrentIndex = Index;
 			}
 		}
-		return CurrentIndex;
+
+		return CurrentIndex != INDEX_NONE ? static_cast<JobType*>(Stripe.Jobs[CurrentIndex].GetReference()) : nullptr;
 	}
 
-	template<typename JobType, typename KeyType>
-	JobType* InternalFindJob(uint32 InJobHash, uint32 InJobId, const KeyType& InKey) const
+	void InternalAddJob(FShaderCommonCompileJob* InJob)
 	{
-		const int32 TypeIndex = (int32)JobType::Type;
-		const int32 JobIndex = InternalFindJobIndex<JobType>(InJobHash, InJobId, InKey);
-		return JobIndex != INDEX_NONE ? static_cast<JobType*>(Jobs[TypeIndex][JobIndex].GetReference()) : nullptr;
+		FLockStripeData& Stripe = GetStripe(InJob->Type, InJob->Hash);
+
+		int32 JobIndex = INDEX_NONE;
+		if (Stripe.FreeIndices.Num() > 0)
+		{
+			JobIndex = Stripe.FreeIndices.Pop(false);
+			check(!Stripe.Jobs[JobIndex].IsValid());
+			Stripe.Jobs[JobIndex] = InJob;
+		}
+		else
+		{
+			JobIndex = Stripe.Jobs.Add(InJob);
+		}
+
+		check(Stripe.Jobs[JobIndex].IsValid());
+		Stripe.JobHash.Add(InJob->Hash, JobIndex);
+
+		check(InJob->Priority != EShaderCompileJobPriority::None);
+		check(InJob->PendingPriority == EShaderCompileJobPriority::None);
+		check(InJob->JobIndex == INDEX_NONE);
+		InJob->JobIndex = JobIndex;
 	}
 
-	template<typename JobType, typename KeyType>
-	JobType* InternalPrepareJob(uint32 InId, const KeyType& InKey, EShaderCompileJobPriority InPriority)
+	struct FLockStripeData
 	{
-		const uint32 Hash = InKey.MakeHash(InId);
-		JobType* PrevJob = nullptr;
+		TArray<FShaderCommonCompileJobPtr> Jobs;
+		TArray<int32> FreeIndices;
+		FHashTable JobHash;
+		FRWLock StripeLock;
+	};
+
+	FLockStripeData Stripes[NumShaderCompileJobTypes][NUM_STRIPES];
+
+	FORCEINLINE FLockStripeData& GetStripe(EShaderCompileJobType JobType, uint32 Hash)
+	{
+		return Stripes[(int32)JobType][Hash >> STRIPE_SHIFT];
+	}
+	FORCEINLINE const FLockStripeData& GetStripe(EShaderCompileJobType JobType, uint32 Hash) const
+	{
+		return Stripes[(int32)JobType][Hash >> STRIPE_SHIFT];
+	}
+};
+
+/** Private implementation class for FShaderCompileJobCollection */
+class FShaderJobCache
+{
+public:
+	FShaderJobCache(FCriticalSection& InCompileQueueSection);
+	~FShaderJobCache();
+
+	// Returns job pointer for new job, otherwise returns NULL
+	template <typename JobType, typename KeyType>
+	JobType* PrepareJob(uint32 InId, const KeyType& InKey, EShaderCompileJobPriority InPriority)
+	{
+		bool bNewJob;
+		JobType* Result = JobTable.PrepareJob<JobType>(InId, InKey, InPriority, bNewJob);
+
+		if (bNewJob)
 		{
-			FReadScopeLock Locker(JobLock);
-			PrevJob = InternalFindJob<JobType>(Hash, InId, InKey);
+			// If it's a new job, return it
+			return Result;
 		}
-
-		JobType* NewJob = nullptr;
-		if (PrevJob == nullptr || (uint32)InPriority > (uint32)PrevJob->Priority)
+		else if (InPriority > Result->Priority)
 		{
-			FWriteScopeLock Locker(JobLock);
-			if (PrevJob == nullptr)
-			{
-				PrevJob = InternalFindJob<JobType>(Hash, InId, InKey);
-			}
-			if (PrevJob == nullptr)
-			{
-				NewJob = new JobType(Hash, InId, InPriority, InKey);
-				InternalAddJob(NewJob);
-			}
-			else if ((uint32)InPriority > (uint32)PrevJob->Priority)
-			{
-				InternalSetPriority(PrevJob, InPriority);
-			}
+			// Or if the priority changed, update that
+			InternalSetPriority(Result, InPriority);
 		}
 
-		return NewJob;
+		return nullptr;
+	}
+
+	void RemoveJob(FShaderCommonCompileJob* InJob)
+	{
+		JobTable.RemoveJob(InJob);
+	}
+
+	int32 RemoveAllPendingJobsWithId(uint32 InId);
+
+	void SubmitJob(FShaderCommonCompileJobPtr Job);
+	void SubmitJobs(const TArray<FShaderCommonCompileJobPtr>& InJobs);
+
+	/** This is an entry point for all jobs that have finished the compilation (whether real or cached). Can be called from multiple threads. Returns mutex stall time. */
+	double ProcessFinishedJob(FShaderCommonCompileJob* FinishedJob, bool bWasCached);
+
+	/** Adds the job to cache. */
+	void AddToCacheAndProcessPending(FShaderCommonCompileJob* FinishedJob);
+
+	/** Populates caching stats in the given compiler stats struct. */
+	void GetStats(FShaderCompilerStats& OutStats) const;
+
+	int32 GetNumPendingJobs(EShaderCompileJobPriority InPriority) const;
+
+	int32 GetNumOutstandingJobs() const;
+
+	int32 GetNumPendingJobs() const;
+
+	int32 GetPendingJobs(EShaderCompilerWorkerType InWorkerType, EShaderCompileJobPriority InPriority, int32 MinNumJobs, int32 MaxNumJobs, TArray<FShaderCommonCompileJobPtr>& OutJobs);
+
+private:
+	using FJobInputHash = FShaderCommonCompileJob::FInputHash;
+	using FJobOutputHash = FBlake3Hash;
+	using FJobCachedOutput = FSharedBuffer;
+	using FStoredOutput = FShaderJobCacheStoredOutput;
+
+	// cannot allow managing this from outside as the caching logic is not exposed
+	inline int32 InternalSubtractNumOutstandingJobs(int32 Value)
+	{
+		const int32 PrevNumOutstandingJobs = NumOutstandingJobs.Subtract(Value);
+		check(PrevNumOutstandingJobs >= Value);
+		return PrevNumOutstandingJobs - Value;
 	}
 
-	void InternalAddJob(FShaderCommonCompileJob* Job);
-	void InternalRemoveJob(FShaderCommonCompileJob* InJob);
 	void InternalSetPriority(FShaderCommonCompileJob* Job, EShaderCompileJobPriority InPriority);
 
 	/** Looks for or adds an entry for the given hash in the cache.  Returns cached output if it exists. */
@@ -544,19 +696,129 @@ private:
 	/** Cleans up oldest outputs to fit in the given memory budget */
 	void CullOutputsToMemoryBudget(uint64 TargetBudgetBytes);
 
+	/** Copied from TLinkedListBase::Unlink */
+	FORCEINLINE static void Unlink(FShaderCommonCompileJob& Job)
+	{
+		if (Job.NextLink)
+		{
+			Job.NextLink->PrevLink = Job.PrevLink;
+		}
+		if (Job.PrevLink)
+		{
+			*Job.PrevLink = Job.NextLink;
+		}
+		// Make it safe to call Unlink again.
+		Job.NextLink = nullptr;
+		Job.PrevLink = nullptr;
+	}
+
+	/**
+	 * Similar to TLinkedListBase::Unlink, but updates a Tail pointer if the Tail is unlinked.  The tail must
+	 * originally be initialized as Tail = &Head.
+	 */
+	FORCEINLINE void UnlinkWithTail(FShaderCommonCompileJob& Job, FShaderCommonCompileJob**& Tail)
+	{
+		// Update tail if we are removing that element
+		if (Tail == &Job.NextLink)
+		{
+			Tail = Job.PrevLink;
+		}
+		Unlink(Job);
+	}
+
+	/** Copied from TLinkedListBase::LinkAfter */
+	FORCEINLINE void LinkAfter(FShaderCommonCompileJob& Job, FShaderCommonCompileJob* After)
+	{
+		checkSlow(After != NULL);
+
+		Job.PrevLink = &After->NextLink;
+		Job.NextLink = *Job.PrevLink;
+		*Job.PrevLink = (FShaderCommonCompileJob*)&Job;
+
+		if (Job.NextLink != NULL)
+		{
+			Job.NextLink->PrevLink = &Job.NextLink;
+		}
+	}
+
+	/**
+	 * Similar to TLinkedListBase::LinkHead, but uses atomic operations to allow multiple producer threads to add to the linked list
+	 * without needing synchronization ("wait free").  Note that synchronization is required for other operations on the list, such
+	 * as traversal or removal, as the list isn't in a fully valid state mid operation (Head always points to the latest item
+	 * inserted, but it may not yet be linked with the rest of the items).  Synchronization is accomplished by using a read lock
+	 * (which multiple threads can hold) for atomic insertion operations, and a write lock for all other operations.
+	 */
+	FORCEINLINE void LinkHeadAtomic(FShaderCommonCompileJob& Job, FShaderCommonCompileJob*& Head)
+	{
+		// It's important that PrevLink is set before the InterlockedExchange, as a subsequent Head pointer exchange could write
+		// another item and need to update PrevLink for this item before this function completes.
+		Job.PrevLink = &Head;
+
+		FShaderCommonCompileJob* OldHead = (FShaderCommonCompileJob*)FPlatformAtomics::InterlockedExchange((PTRINT*)&Head, (PTRINT)&Job);
+		if (OldHead != nullptr)
+		{
+			OldHead->PrevLink = &Job.NextLink;
+		}
+		Job.NextLink = OldHead;
+	}
+
+	/**
+	 * Variation that links a job at the tail of the list.  The tail must originally be initialized as Tail = &Head.  Similar to
+	 * LinkHeadAtomic above (see more detailed comments there), a read lock is required for this operation.
+	 */
+	FORCEINLINE static void LinkTailAtomic(FShaderCommonCompileJob& Job, FShaderCommonCompileJob**& Tail)
+	{
+		FShaderCommonCompileJob** OldTail = (FShaderCommonCompileJob**)FPlatformAtomics::InterlockedExchange((PTRINT*)&Tail, (PTRINT)&Job.NextLink);
+		Job.PrevLink = OldTail;
+		Job.NextLink = nullptr;
+
+		// Update previous tail's next pointer (or OldTail may be pointing at Head if list was empty)
+		*OldTail = (FShaderCommonCompileJob*)&Job;
+	}
+
+	/** Links job into linked list of given priority */
+	FORCEINLINE void LinkJobWithPriority(FShaderCommonCompileJob& Job, int32 PriorityIndex)
+	{
+		check((uint32)PriorityIndex < (uint32)NumShaderCompileJobPriorities);
+		NumPendingJobs[PriorityIndex]++;
+#if UE_SHADERCOMPILER_FIFO_JOB_EXECUTION
+		LinkTailAtomic(Job, PendingJobsTail[PriorityIndex]);
+#else
+		LinkHeadAtomic(Job, PendingJobsHead[PriorityIndex]);
+#endif
+	}
+
+	/** Unlinks job from linked list of given priority */
+	FORCEINLINE void UnlinkJobWithPriority(FShaderCommonCompileJob& Job, int32 PriorityIndex)
+	{
+		check((uint32)PriorityIndex < (uint32)NumShaderCompileJobPriorities);
+		check(NumPendingJobs[PriorityIndex] > 0);
+		NumPendingJobs[PriorityIndex]--;
+#if UE_SHADERCOMPILER_FIFO_JOB_EXECUTION
+		UnlinkWithTail(Job, PendingJobsTail[PriorityIndex]);
+#else
+		Unlink(Job);
+#endif
+	}
+
+	/** From FShaderCompilingManager, guards access to FShaderMapCompileResults written in ProcessFinishedJob */
+	FCriticalSection& CompileQueueSection;
+
 	/** Guards access to the structure */
 	mutable FRWLock JobLock;
 
 	/** Queue of tasks that haven't been assigned to a worker yet. */
-	FShaderCommonCompileJob* PendingJobs[NumShaderCompileJobPriorities];
-	int32 NumPendingJobs[NumShaderCompileJobPriorities];
+	FShaderCommonCompileJob* PendingJobsHead[NumShaderCompileJobPriorities];
+	std::atomic_int32_t NumPendingJobs[NumShaderCompileJobPriorities];
+#if UE_SHADERCOMPILER_FIFO_JOB_EXECUTION
+	FShaderCommonCompileJob** PendingJobsTail[NumShaderCompileJobPriorities];
+#endif
 
 	/** Number of jobs currently being compiled.  This includes PendingJobs and any jobs that have been assigned to workers but aren't complete yet. */
 	FThreadSafeCounter NumOutstandingJobs;
 
-	TArray<FShaderCommonCompileJobPtr> Jobs[NumShaderCompileJobTypes];
-	TArray<int32> FreeIndices[NumShaderCompileJobTypes];
-	FHashTable JobHash[NumShaderCompileJobTypes];
+	/** Active jobs */
+	FShaderCompilerJobTable JobTable;
 
 	/* A lot of outputs can be duplicated, so they are deduplicated before storing */
 	TMap<FJobOutputHash, FStoredOutput*> Outputs;
@@ -872,7 +1134,7 @@ namespace ShaderCompiler
 }
 
 
-FShaderCompileJobCollection::FShaderCompileJobCollection()
+FShaderCompileJobCollection::FShaderCompileJobCollection(FCriticalSection& InCompileQueueSection)
 {
 	PrintStatsCmd = IConsoleManager::Get().RegisterConsoleCommand(
 		TEXT("r.ShaderCompiler.PrintStats"),
@@ -881,18 +1143,18 @@ FShaderCompileJobCollection::FShaderCompileJobCollection()
 		ECVF_Default
 	);
 
-	JobsCache = MakePimpl<FShaderJobCache>();
+	JobsCache = MakePimpl<FShaderJobCache>(InCompileQueueSection);
 }
 
 
 // Pass through functions to inner FShaderJobCache implementation class
 FShaderCompileJob* FShaderCompileJobCollection::PrepareJob(uint32 InId, const FShaderCompileJobKey& InKey, EShaderCompileJobPriority InPriority)
 {
-	return JobsCache->PrepareJob(InId, InKey, InPriority);
+	return JobsCache->PrepareJob<FShaderCompileJob>(InId, InKey, InPriority);
 }
 FShaderPipelineCompileJob* FShaderCompileJobCollection::PrepareJob(uint32 InId, const FShaderPipelineCompileJobKey& InKey, EShaderCompileJobPriority InPriority)
 {
-	return JobsCache->PrepareJob(InId, InKey, InPriority);
+	return JobsCache->PrepareJob<FShaderPipelineCompileJob>(InId, InKey, InPriority);
 }
 void FShaderCompileJobCollection::RemoveJob(FShaderCommonCompileJob* InJob)
 {
@@ -936,31 +1198,6 @@ int32 FShaderCompileJobCollection::GetPendingJobs(EShaderCompilerWorkerType InWo
 }
 
 
-void FShaderJobCache::InternalAddJob(FShaderCommonCompileJob* InJob)
-{
-	const int32 TypeIndex = (int32)InJob->Type;
-
-	int32 JobIndex = INDEX_NONE;
-	if (FreeIndices[TypeIndex].Num() > 0)
-	{
-		JobIndex = FreeIndices[TypeIndex].Pop(false);
-		check(!Jobs[TypeIndex][JobIndex].IsValid());
-		Jobs[TypeIndex][JobIndex] = InJob;
-	}
-	else
-	{
-		JobIndex = Jobs[TypeIndex].Add(InJob);
-	}
-
-	check(Jobs[TypeIndex][JobIndex].IsValid());
-	JobHash[TypeIndex].Add(InJob->Hash, JobIndex);
-	
-	check(InJob->Priority != EShaderCompileJobPriority::None);
-	check(InJob->PendingPriority == EShaderCompileJobPriority::None);
-	check(InJob->JobIndex == INDEX_NONE);
-	InJob->JobIndex = JobIndex;
-}
-
 static FShaderCommonCompileJob* CloneJob_Single(const FShaderCompileJob* SrcJob)
 {
 	FShaderCompileJob* Job = new FShaderCompileJob(SrcJob->Hash, SrcJob->Id, SrcJob->Priority, SrcJob->Key);
@@ -1012,16 +1249,16 @@ void FShaderJobCache::InternalSetPriority(FShaderCommonCompileJob* Job, EShaderC
 
 	if (Job->PendingPriority != EShaderCompileJobPriority::None)
 	{
+		// Need write lock to call UnlinkJobWithPriority
+		FWriteScopeLock Locker(JobLock);
+
 		// Job hasn't started yet, move it to the pending list for the new priority
 		const int32 PrevPriorityIndex = (int32)Job->PendingPriority;
 		check(Job->PendingPriority == Job->Priority);
-		check(NumPendingJobs[PrevPriorityIndex] > 0);
-		NumPendingJobs[PrevPriorityIndex]--;
-		Job->Unlink();
+		UnlinkJobWithPriority(*Job, PrevPriorityIndex);
 
-		NumPendingJobs[PriorityIndex]++;
 		ensure(!ShaderCompiler::IsJobCacheEnabled() || Job->bInputHashSet);
-		Job->LinkHead(PendingJobs[PriorityIndex]);
+		LinkJobWithPriority(*Job, PriorityIndex);
 		Job->Priority = InPriority;
 		Job->PendingPriority = InPriority;
 	}
@@ -1033,10 +1270,9 @@ void FShaderJobCache::InternalSetPriority(FShaderCommonCompileJob* Job, EShaderC
 		NewJob->Priority = InPriority;
 		const int32 NewNumPendingJobs = NewJob->PendingShaderMap->NumPendingJobs.Increment();
 		checkf(NewNumPendingJobs > 1, TEXT("Invalid number of pending jobs %d, should have had at least 1 job previously"), NewNumPendingJobs);
-		InternalAddJob(NewJob);
+		JobTable.AddExistingJob(NewJob);
 
 		GShaderCompilerStats->RegisterNewPendingJob(*NewJob);
-		NumPendingJobs[PriorityIndex]++;
 		ensureMsgf(NewJob->bInputHashSet == Job->bInputHashSet, TEXT("Cloned and original jobs should either both have input hash, or both not have it. Job->bInputHashSet=%d, NewJob->bInputHashSet=%d"),
 			Job->bInputHashSet,
 			NewJob->bInputHashSet
@@ -1044,7 +1280,10 @@ void FShaderJobCache::InternalSetPriority(FShaderCommonCompileJob* Job, EShaderC
 		ensureMsgf(!ShaderCompiler::IsJobCacheEnabled() || NewJob->GetInputHash() == Job->GetInputHash(),
 			TEXT("If shader jobs cache is enabled, cloned job should have the same input hash as the original, and it doesn't.")
 			);
-		NewJob->LinkHead(PendingJobs[PriorityIndex]);
+
+		// Need read lock to call LinkJobWithPriority (uses wait free LinkHeadAtomic / LinkTailAtomic -- works fine for multiple producer threads)
+		FReadScopeLock Locker(JobLock);
+		LinkJobWithPriority(*NewJob, PriorityIndex);
 		NewJob->PendingPriority = InPriority;
 		NumOutstandingJobs.Increment();
 
@@ -1052,28 +1291,6 @@ void FShaderJobCache::InternalSetPriority(FShaderCommonCompileJob* Job, EShaderC
 	}
 }
 
-void FShaderJobCache::InternalRemoveJob(FShaderCommonCompileJob* InJob)
-{
-	const int32 TypeIndex = (int32)InJob->Type;
-	const int32 JobIndex = InJob->JobIndex;
-
-	check(JobIndex != INDEX_NONE);
-	check(Jobs[TypeIndex][JobIndex] == InJob);
-	check(InJob->PendingPriority == EShaderCompileJobPriority::None);
-	InJob->JobIndex = INDEX_NONE;
-
-	JobHash[TypeIndex].Remove(InJob->Hash, JobIndex);
-	FreeIndices[TypeIndex].Add(JobIndex);
-	Jobs[TypeIndex][JobIndex].SafeRelease();
-}
-
-
-void FShaderJobCache::RemoveJob(FShaderCommonCompileJob* InJob)
-{
-	FWriteScopeLock Locker(JobLock);
-	InternalRemoveJob(InJob);
-}
-
 int32 FShaderJobCache::RemoveAllPendingJobsWithId(uint32 InId)
 {
 	int32 NumRemoved = 0;
@@ -1081,7 +1298,7 @@ int32 FShaderJobCache::RemoveAllPendingJobsWithId(uint32 InId)
 		FWriteScopeLock Locker(JobLock);
 		for (int32 PriorityIndex = 0; PriorityIndex < NumShaderCompileJobPriorities; ++PriorityIndex)
 		{
-			for (FShaderCommonCompileJob::TIterator It(PendingJobs[PriorityIndex]); It;)
+			for (FShaderCommonCompileJobIterator It(PendingJobsHead[PriorityIndex]); It;)
 			{
 				FShaderCommonCompileJob& Job = *It;
 				It.Next();
@@ -1099,11 +1316,9 @@ int32 FShaderJobCache::RemoveAllPendingJobsWithId(uint32 InId)
 						}
 					}
 
-					check(NumPendingJobs[PriorityIndex] > 0);
-					NumPendingJobs[PriorityIndex]--;
-					Job.Unlink();
+					UnlinkJobWithPriority(Job, PriorityIndex);
 					Job.PendingPriority = EShaderCompileJobPriority::None;
-					InternalRemoveJob(&Job);
+					RemoveJob(&Job);
 					++NumRemoved;
 				}
 			}
@@ -1124,12 +1339,13 @@ int32 FShaderJobCache::RemoveAllPendingJobsWithId(uint32 InId)
 					// if we're removing the list head, we need to update it to the next
 					if (JobData.DuplicateJobsWaitList == DuplicateJob)
 					{
-						JobData.DuplicateJobsWaitList = JobData.DuplicateJobsWaitList->Next();
+						JobData.DuplicateJobsWaitList = JobData.DuplicateJobsWaitList->NextLink;
 					}
 
-					DuplicateJob->Unlink();
+					// Duplicate jobs are in their own list, not one of the priority lists, so don't use UnlinkJobWithPriority
+					Unlink(*DuplicateJob);
 					DuplicateJob->PendingPriority = EShaderCompileJobPriority::None;
-					InternalRemoveJob(DuplicateJob);
+					RemoveJob(DuplicateJob);
 					++NumRemoved;
 
 					// This removes the current job (at DuplicateIndex), so we don't increment in this case
@@ -1149,147 +1365,140 @@ int32 FShaderJobCache::RemoveAllPendingJobsWithId(uint32 InId)
 	return NumRemoved;
 }
 
-void FShaderJobCache::SubmitJobs(const TArray<FShaderCommonCompileJobPtr>& InJobs)
+void FShaderJobCache::SubmitJob(FShaderCommonCompileJobPtr Job)
 {
-	if (InJobs.Num() > 0)
+	check(Job->JobIndex != INDEX_NONE);
+	check(Job->Priority != EShaderCompileJobPriority::None);
+	check(Job->PendingPriority == EShaderCompileJobPriority::None);
+
+	const int32 PriorityIndex = (int32)Job->Priority;
+	bool bNewJob = true;
+	// check caches unless we're running in validation mode (which runs _all_ jobs and compares hashes of outputs)
+	if (ShaderCompiler::IsJobCacheEnabled() && !ShaderCompiler::IsJobCacheDebugValidateEnabled())
 	{
-		// all jobs (not just actually submitted ones) count as outstanding. This needs to be done early because
-		// we may fulfill some of the jobs from the cache (and we will be subtracting them)
-		NumOutstandingJobs.Add(InJobs.Num());
+		const FShaderCommonCompileJob::FInputHash& InputHash = Job->GetInputHash();
+
+		const bool bCheckDDC = !(Job->bIsDefaultMaterial || Job->bIsGlobalShader);
 
-		int32 SubmittedJobsCount = 0;
-		int32 NumSubmittedJobs[NumShaderCompileJobPriorities] = { 0 };
+		// We don't use a scope here, because we need to release this lock before calling ProcessFinishedJob, which needs to acquire
+		// CompileQueueSection.  It's not safe to acquire CompileQueueSection where JobLock is locked first, as it will cause
+		// deadlocks due to FShaderCompileThreadRunnable::CompilingLoop calling GetPendingJobs, which acquires those two locks in
+		// the opposite order.
+		double StallStart = FPlatformTime::Seconds();
+		JobLock.WriteLock();
+		Job->TimeTaskSubmitJobsStall += FPlatformTime::Seconds() - StallStart;
+
+		FSharedBuffer* ExistingOutput;
+		FShaderJobCacheRef JobCacheRef = FindOrAdd(InputHash, bCheckDDC, ExistingOutput);
+
+		// see if there are already cached results for this job
+		if (ExistingOutput)
 		{
-			// Just precompute the InputHash for each job in multiple-thread.
-			if (ShaderCompiler::IsJobCacheEnabled())
-			{
-				TRACE_CPUPROFILER_EVENT_SCOPE(ShaderCompiler.GetInputHash);
-				ParallelFor(TEXT("ShaderCompiler.GetInputHash.PF"), InJobs.Num(),1, [&InJobs](int32 Index) 
-				{
-					ConditionalPreprocessShader(InJobs[Index]);
-					InJobs[Index]->GetInputHash(); 
-				}, EParallelForFlags::Unbalanced);
-			}
+			// Need to release the lock before calling ProcessFinishedJob, as mentioned above (and it's also good for performance to 
+			// release the lock before the relatively costly "SerializeOutput" call).
+			JobLock.WriteUnlock();
 
-			FWriteScopeLock Locker(JobLock);
+			UE_LOG(LogShaderCompilers, UE_SHADERCACHE_LOG_LEVEL, TEXT("There is already a cached job with the ihash %s, processing the new one immediately."), *LexToString(InputHash));
+			FMemoryReaderView MemReader(*ExistingOutput);
+			Job->SerializeOutput(MemReader);
 
-			// Optimization: Only search the linked list once to prevent O(n^2) behavior
-			FShaderCommonCompileJob* ExistingJobs[NumShaderCompileJobPriorities] = { nullptr };
+			// finish the job instantly
+			Job->TimeTaskSubmitJobsStall += ProcessFinishedJob(Job, true);
+			bNewJob = false;
+		}
+		else
+		{
+			FShaderJobData& JobData = GetShaderJobData(JobCacheRef);
+			Job->JobCacheRef = JobCacheRef;
 
-			for (FShaderCommonCompileJob* Job : InJobs)
+			// see if another job with the same input hash is being worked on
+			if (JobData.JobInFlight)
 			{
-				check(Job->JobIndex != INDEX_NONE);
-				check(Job->Priority != EShaderCompileJobPriority::None);
-				check(Job->PendingPriority == EShaderCompileJobPriority::None);
+				UE_LOG(LogShaderCompilers, UE_SHADERCACHE_LOG_LEVEL, TEXT("There is an outstanding job with the ihash %s, not submitting another one (adding to wait list)."), *LexToString(InputHash));
 
-				const int32 PriorityIndex = (int32)Job->Priority;
-				bool bNewJob = true;
-				// check caches unless we're running in validation mode (which runs _all_ jobs and compares hashes of outputs)
-				if (ShaderCompiler::IsJobCacheEnabled() && !ShaderCompiler::IsJobCacheDebugValidateEnabled())
+				// because of the cloned jobs, we need to maintain a separate mapping
+				FShaderCommonCompileJob** WaitListHead = &JobData.DuplicateJobsWaitList;
+				if (*WaitListHead)
 				{
-					const FShaderCommonCompileJob::FInputHash& InputHash = Job->GetInputHash();
-
-					const bool bCheckDDC = !(Job->bIsDefaultMaterial || Job->bIsGlobalShader);
+					LinkAfter(*Job, *WaitListHead);
+				}
+				else
+				{
+					*WaitListHead = Job;
+				}
+				AddDuplicateJob(Job);
+				bNewJob = false;
+			}
+			else
+			{
+				// track new jobs so we can dedupe them
+				JobData.JobInFlight = Job;
+			}
 
-					FSharedBuffer* ExistingOutput;
-					FShaderJobCacheRef JobCacheRef = FindOrAdd(InputHash, bCheckDDC, ExistingOutput);
+			JobLock.WriteUnlock();
+		}
+	}
 
-					// see if there are already cached results for this job
-					if (ExistingOutput)
-					{
-						UE_LOG(LogShaderCompilers, UE_SHADERCACHE_LOG_LEVEL, TEXT("There is already a cached job with the ihash %s, processing the new one immediately."), *LexToString(InputHash));
-						FMemoryReaderView MemReader(*ExistingOutput);
-						Job->SerializeOutput(MemReader);
+	// new job
+	if (bNewJob)
+	{
+		GShaderCompilerStats->RegisterNewPendingJob(*Job);
+		ensure(!ShaderCompiler::IsJobCacheEnabled() || Job->bInputHashSet);
 
-						// finish the job instantly
-						ProcessFinishedJob(Job, true);
+		Job->PendingPriority = Job->Priority;
 
-						continue;
-					}
-					else
-					{
-						FShaderJobData& JobData = GetShaderJobData(JobCacheRef);
-						Job->JobCacheRef = JobCacheRef;
+		// Only need a read lock here, as LinkJobWithPriority does an atomic link operation safe across multiple producer threads
+		double StallStart = FPlatformTime::Seconds();
+		FReadScopeLock Locker(JobLock);
+		Job->TimeTaskSubmitJobsStall += FPlatformTime::Seconds() - StallStart;
 
-						// see if another job with the same input hash is being worked on
-						if (JobData.JobInFlight)
-						{
-							UE_LOG(LogShaderCompilers, UE_SHADERCACHE_LOG_LEVEL, TEXT("There is an outstanding job with the ihash %s, not submitting another one (adding to wait list)."), *LexToString(InputHash));
+		LinkJobWithPriority(*Job, PriorityIndex);
+	}
+}
 
-							// because of the cloned jobs, we need to maintain a separate mapping
-							FShaderCommonCompileJob** WaitListHead = &JobData.DuplicateJobsWaitList;
-							if (*WaitListHead)
-							{
-								Job->LinkAfter(*WaitListHead);
-							}
-							else
-							{
-								*WaitListHead = Job;
-							}
-							AddDuplicateJob(Job);
-							bNewJob = false;
-						}
-						else
-						{
-							// track new jobs so we can dedupe them
-							JobData.JobInFlight = Job;
-						}
-					}
-				}
+void FShaderJobCache::SubmitJobs(const TArray<FShaderCommonCompileJobPtr>& InJobs)
+{
+	if (InJobs.Num() > 0)
+	{
+		// all jobs (not just actually submitted ones) count as outstanding. This needs to be done early because
+		// we may fulfill some of the jobs from the cache (and we will be subtracting them)
+		NumOutstandingJobs.Add(InJobs.Num());
 
-				// new job
-				if (bNewJob)
+		if (CVarShaderCompilerParallelSubmitJobs.GetValueOnGameThread())
+		{
+			for (FShaderCommonCompileJob* Job : InJobs)
+			{
+				UE::Tasks::Launch(UE_SOURCE_LOCATION, [Job, this]()
 				{
-					GShaderCompilerStats->RegisterNewPendingJob(*Job);
-					ensure(!ShaderCompiler::IsJobCacheEnabled() || Job->bInputHashSet);
-#if UE_SHADERCOMPILER_FIFO_JOB_EXECUTION
-					// link the job at the end of pending list, so we're executing them in a FIFO and not LIFO order
-					if (PendingJobs[PriorityIndex])
-					{
-						// Optimization: Only search the linked list if we have to
-						if (ExistingJobs[PriorityIndex])
-						{
-							Job->LinkAfter(ExistingJobs[PriorityIndex]);
-							ExistingJobs[PriorityIndex] = Job;
-						}
-						else
-						{
-							for (FShaderCommonCompileJob::TIterator It(PendingJobs[PriorityIndex]); It; ++It)
-							{
-								FShaderCommonCompileJob& ExistingJob = *It;
-								if (ExistingJob.GetNextLink() == nullptr)
-								{
-									Job->LinkAfter(&ExistingJob);
-									ExistingJobs[PriorityIndex] = Job;
-									break;
-								}
-							}
-						}
-					}
-					else
-#endif // UE_SHADERCOMPILER_FIFO_JOB_EXECUTION
+					TRACE_CPUPROFILER_EVENT_SCOPE(ShaderJobTask);
+					double TimeStart = FPlatformTime::Seconds();
+
+					if (ShaderCompiler::IsJobCacheEnabled())
 					{
-						Job->LinkHead(PendingJobs[PriorityIndex]);
+						ConditionalPreprocessShader(Job);
 					}
-					
+					SubmitJob(Job);
 
-					NumPendingJobs[PriorityIndex]++;
-					NumSubmittedJobs[PriorityIndex]++;
-					Job->PendingPriority = Job->Priority;
-					++SubmittedJobsCount;
-				}
+					Job->TimeTaskSubmitJobs = FPlatformTime::Seconds() - TimeStart;
+				});
 			}
 		}
-
-		UE_LOG(LogShaderCompilers, UE_SHADERCACHE_LOG_LEVEL, TEXT("Actual jobs submitted %d (of %d new), total outstanding jobs: %d."), SubmittedJobsCount, InJobs.Num(), NumOutstandingJobs.GetValue());
-
-		for (int32 PriorityIndex = 0; PriorityIndex < NumShaderCompileJobPriorities; ++PriorityIndex)
+		else
 		{
-			if (NumSubmittedJobs[PriorityIndex] > 0)
+			// Precompute the InputHash for each job in multiple-thread.
+			if (ShaderCompiler::IsJobCacheEnabled())
+			{
+				TRACE_CPUPROFILER_EVENT_SCOPE(ShaderCompiler.GetInputHash);
+				ParallelFor(TEXT("ShaderCompiler.GetInputHash.PF"), InJobs.Num(), 1, [&InJobs](int32 Index)
+				{
+					ConditionalPreprocessShader(InJobs[Index]);
+					InJobs[Index]->GetInputHash();
+				}, EParallelForFlags::Unbalanced);
+			}
+
+			for (FShaderCommonCompileJob* Job : InJobs)
 			{
-				UE_LOG(LogShaderCompilers, Verbose, TEXT("Submitted %d shader compile jobs with '%s' priority"),
-					NumSubmittedJobs[PriorityIndex],
-					ShaderCompileJobPriorityToString((EShaderCompileJobPriority)PriorityIndex));
+				SubmitJob(Job);
 			}
 		}
 	}
@@ -1301,8 +1510,10 @@ void FShaderCompileJobCollection::HandlePrintStats()
 }
 
 TRACE_DECLARE_INT_COUNTER(Shaders_Compiled, TEXT("Shaders/Compiled"));
-void FShaderJobCache::ProcessFinishedJob(FShaderCommonCompileJob* FinishedJob, bool bWasCached)
+double FShaderJobCache::ProcessFinishedJob(FShaderCommonCompileJob* FinishedJob, bool bWasCached)
 {
+	double StallTime;
+
 	FinishedJob->OnComplete();
 
 	if (!bWasCached)
@@ -1311,19 +1522,27 @@ void FShaderJobCache::ProcessFinishedJob(FShaderCommonCompileJob* FinishedJob, b
 		TRACE_COUNTER_ADD(Shaders_Compiled, 1);
 	}
 
-	// TODO: have a pending shader map critical section? not clear at this point if we can be accessing the results on another thread at the same time
-	FShaderMapCompileResults& ShaderMapResults = *(FinishedJob->PendingShaderMap);
-	ShaderMapResults.FinishedJobs.Add(FinishedJob);
-	ShaderMapResults.bAllJobsSucceeded = ShaderMapResults.bAllJobsSucceeded && FinishedJob->bSucceeded;
+	{
+		// Need to protect writes to FShaderMapCompileResults
+		double StallStart = FPlatformTime::Seconds();
+		FScopeLock Lock(&CompileQueueSection);
+		StallTime = FPlatformTime::Seconds() - StallStart;
 
-	const int32 NumPendingJobsForSM = ShaderMapResults.NumPendingJobs.Decrement();
-	checkf(NumPendingJobsForSM >= 0, TEXT("Problem tracking pending jobs for a SM (%d), number of pending jobs (%d) is negative!"), FinishedJob->Id, NumPendingJobsForSM);
+		FShaderMapCompileResults& ShaderMapResults = *(FinishedJob->PendingShaderMap);
+		ShaderMapResults.FinishedJobs.Add(FinishedJob);
+		ShaderMapResults.bAllJobsSucceeded = ShaderMapResults.bAllJobsSucceeded && FinishedJob->bSucceeded;
+
+		const int32 NumPendingJobsForSM = ShaderMapResults.NumPendingJobs.Decrement();
+		checkf(NumPendingJobsForSM >= 0, TEXT("Problem tracking pending jobs for a SM (%d), number of pending jobs (%d) is negative!"), FinishedJob->Id, NumPendingJobsForSM);
+	}
 
 	InternalSubtractNumOutstandingJobs(1);
 	if (!bWasCached && ShaderCompiler::IsJobCacheEnabled())
 	{
 		AddToCacheAndProcessPending(FinishedJob);
 	}
+
+	return StallTime;
 }
 
 void FShaderJobCache::AddToCacheAndProcessPending(FShaderCommonCompileJob* FinishedJob)
@@ -1344,49 +1563,59 @@ void FShaderJobCache::AddToCacheAndProcessPending(FShaderCommonCompileJob* Finis
 
 	FSharedBuffer Buffer = MakeSharedBufferFromArray(MoveTemp(Output));
 
-	// TODO: reduce the scope - e.g. SerializeOutput and processing finished jobs can be moved out of it
-	FWriteScopeLock JobLocker(JobLock);
-
 	FShaderJobData& JobData = GetShaderJobData(FinishedJob->JobCacheRef);
 
 	// see if there are outstanding jobs that also need to be resolved
-	int32 NumOutstandingJobsWithSameHash = 0;
-	FShaderCommonCompileJob* CurHead = JobData.DuplicateJobsWaitList;
-	while (CurHead)
+	TArray<FShaderCommonCompileJob*> FinishedDuplicateJobs;
+
 	{
-		checkf(CurHead != FinishedJob, TEXT("Job that is being added to cache was also on a waiting list! Error in bookkeeping."));
+		FWriteScopeLock JobLocker(JobLock);
 
-		FMemoryReaderView MemReader(Buffer);
-		CurHead->SerializeOutput(MemReader);
-		checkf(CurHead->bSucceeded == FinishedJob->bSucceeded, TEXT("Different success status for the job with the same ihash"));
+		FShaderCommonCompileJob* CurHead = JobData.DuplicateJobsWaitList;
+		while (CurHead)
+		{
+			checkf(CurHead != FinishedJob, TEXT("Job that is being added to cache was also on a waiting list! Error in bookkeeping."));
 
-		// finish the job instantly
-		ProcessFinishedJob(CurHead, true);
-		++NumOutstandingJobsWithSameHash;
+			// Need to add these to a list, and process them outside the JobLock scope.  ProcessFinishedJob locks CompileQueueSection,
+			// and we don't want to lock that inside a block that also locks JobLock, as it can cause a deadlock given that other
+			// code paths obtain the locks in the opposite order.  This is also good for perf, as it avoids holding the lock during
+			// the relatively costly SerializeOutput.
+			FinishedDuplicateJobs.Add(CurHead);
 
-		RemoveDuplicateJob(CurHead);
+			// This needs to happen inside the JobLocker scope
+			RemoveDuplicateJob(CurHead);
 
-		CurHead = CurHead->Next();
-	}
+			CurHead = CurHead->NextLink;
+		}
 
-	if (NumOutstandingJobsWithSameHash > 0)
-	{
-		// Clear the duplicate job list if we processed any jobs
 		JobData.DuplicateJobsWaitList = nullptr;
 
-		UE_LOG(LogShaderCompilers, UE_SHADERCACHE_LOG_LEVEL, TEXT("Processed %d outstanding jobs with the same ihash %s."), NumOutstandingJobsWithSameHash, *LexToString(InputHash));
+		if (FinishedJob->bSucceeded)
+		{
+			const bool bAddToDDC = !(FinishedJob->bIsDefaultMaterial || FinishedJob->bIsGlobalShader);
+			// we only cache jobs that succeded
+			AddJobOutput(JobData, FinishedJob, InputHash, Buffer, FinishedDuplicateJobs.Num(), bAddToDDC);
+		}
+
+		// remove ourselves from the jobs in flight
+		JobData.JobInFlight = nullptr;
+		FinishedJob->JobCacheRef.Clear();
 	}
 
-	if (FinishedJob->bSucceeded)
+	if (FinishedDuplicateJobs.Num())
 	{
-		const bool bAddToDDC = !(FinishedJob->bIsDefaultMaterial || FinishedJob->bIsGlobalShader);
-		// we only cache jobs that succeded
-		AddJobOutput(JobData, FinishedJob, InputHash, Buffer, NumOutstandingJobsWithSameHash, bAddToDDC);
-	}
+		UE_LOG(LogShaderCompilers, UE_SHADERCACHE_LOG_LEVEL, TEXT("Processed %d outstanding jobs with the same ihash %s."), FinishedDuplicateJobs.Num(), *LexToString(InputHash));
 
-	// remove ourselves from the jobs in flight
-	JobData.JobInFlight = nullptr;
-	FinishedJob->JobCacheRef.Clear();
+		for (FShaderCommonCompileJob* DuplicateJob : FinishedDuplicateJobs)
+		{
+			FMemoryReaderView MemReader(Buffer);
+			DuplicateJob->SerializeOutput(MemReader);
+			checkf(DuplicateJob->bSucceeded == FinishedJob->bSucceeded, TEXT("Different success status for the job with the same ihash"));
+
+			// finish the job instantly
+			ProcessFinishedJob(DuplicateJob, true);
+		}
+	}
 }
 
 int32 FShaderJobCache::GetNumPendingJobs(EShaderCompileJobPriority InPriority) const
@@ -1439,8 +1668,8 @@ int32 FShaderJobCache::GetPendingJobs(EShaderCompilerWorkerType InWorkerType, ES
 	}
 	
 	OutJobs.Reserve(OutJobs.Num() + FMath::Min(MaxNumJobs, NumPendingJobsOfPriority));
-	int32 NumJobs = FMath::Min(MaxNumJobs, NumPendingJobs[PriorityIndex]);
-	FShaderCommonCompileJob::TIterator It(PendingJobs[PriorityIndex]);
+	int32 NumJobs = FMath::Min(MaxNumJobs, NumPendingJobs[PriorityIndex].load());
+	FShaderCommonCompileJobIterator It(PendingJobsHead[PriorityIndex]);
 	// Randomize job selection by randomly skipping over jobs while traversing the list.
 	// Say, we need to pick 3 jobs out of 5 total. We can skip over 2 jobs in total, e.g. like this:
 	// pick one (4 more to go and we need to get 2 of 4), skip one (3 more to go, picking 2 out of 3), pick one (2 more to go, picking 1 of 2), skip one, pick one.
@@ -1457,7 +1686,13 @@ int32 FShaderJobCache::GetPendingJobs(EShaderCompilerWorkerType InWorkerType, ES
 		ensure(!ShaderCompiler::IsJobCacheEnabled() || Job.bInputHashSet);
 
 		It.Next();
-		Job.Unlink();
+
+		// Don't use UnlinkJobWithPriority here, as that updates NumPendingJobs, interfering with skipping logic -- NumPendingJobs is updated at the end of the loop
+#if UE_SHADERCOMPILER_FIFO_JOB_EXECUTION
+		UnlinkWithTail(Job, PendingJobsTail[PriorityIndex]);
+#else
+		Unlink(Job);
+#endif
 
 		Job.PendingPriority = EShaderCompileJobPriority::None;
 		Job.CurrentWorker = InWorkerType;
@@ -1482,16 +1717,6 @@ int32 FShaderJobCache::GetPendingJobs(EShaderCompilerWorkerType InWorkerType, ES
 	return NumJobs;
 }
 
-FShaderCompileJob* FShaderJobCache::PrepareJob(uint32 InId, const FShaderCompileJobKey& InKey, EShaderCompileJobPriority InPriority)
-{
-	return InternalPrepareJob<FShaderCompileJob>(InId, InKey, InPriority);
-}
-
-FShaderPipelineCompileJob* FShaderJobCache::PrepareJob(uint32 InId, const FShaderPipelineCompileJobKey& InKey, EShaderCompileJobPriority InPriority)
-{
-	return InternalPrepareJob<FShaderPipelineCompileJob>(InId, InKey, InPriority);
-}
-
 static float GRegularWorkerTimeToLive = 20.0f;
 static float GBuildWorkerTimeToLive = 600.0f;
 
@@ -4038,6 +4263,11 @@ void FShaderCompilerStats::WriteStatSummary()
 
 	UE_LOG(LogShaderCompilers, Display, TEXT("Time at least one job was in flight (either pending or executed): %.2f s"), TotalTimeAtLeastOneJobWasInFlight);
 
+	if (Counters.AccumulatedTaskSubmitJobs > 0.0)
+	{
+		UE_LOG(LogShaderCompilers, Display, TEXT("Mutex wait stall in FShaderJobCache::SubmitJobs:  %.2f%%"), 100.0 * Counters.AccumulatedTaskSubmitJobsStall / Counters.AccumulatedTaskSubmitJobs );
+	}
+
 	// print stats about the batches
 	if (Counters.LocalJobBatchesSeen > 0 && Counters.DistributedJobBatchesSeen > 0)
 	{
@@ -4530,6 +4760,12 @@ void FShaderCompilerStats::RegisterFinishedJob(FShaderCommonCompileJob& Job)
 	Counters.AccumulatedJobLifeTime += LifeTime;
 	Counters.MaxJobLifeTime = FMath::Max(LifeTime, Counters.MaxJobLifeTime);
 
+	if (Job.TimeTaskSubmitJobs)
+	{
+		Counters.AccumulatedTaskSubmitJobs += Job.TimeTaskSubmitJobs;
+		Counters.AccumulatedTaskSubmitJobsStall += Job.TimeTaskSubmitJobsStall;
+	}
+
 	// estimate lifetime without an overlap
 	ensure(Job.TimeAddedToPendingQueue != 0.0 && Job.TimeAddedToPendingQueue <= Job.TimeExecutionCompleted);
 	AddToInterval(JobLifeTimeIntervals, TInterval<double>(Job.TimeAddedToPendingQueue, Job.TimeExecutionCompleted));
@@ -4775,6 +5011,7 @@ IDistributedBuildController* FShaderCompilingManager::FindRemoteCompilerControll
 FShaderCompilingManager::FShaderCompilingManager() :
 	bCompilingDuringGame(false),
 	NumExternalJobs(0),
+	AllJobs(CompileQueueSection),
 	NumSingleThreadedRunsBeforeRetry(GSingleThreadedRunsIdle),
 #if PLATFORM_MAC
 	ShaderCompileWorkerName(FPaths::EngineDir() / TEXT("Binaries/Mac/ShaderCompileWorker")),
@@ -9753,10 +9990,17 @@ uint64 FShaderJobCache::GetCurrentMemoryBudget() const
 	return FMath::Min(AbsoluteLimit, RelativeLimit);
 }
 
-FShaderJobCache::FShaderJobCache()
+FShaderJobCache::FShaderJobCache(FCriticalSection& InCompileQueueSection)
+	: CompileQueueSection(InCompileQueueSection)
 {
-	FMemory::Memzero(PendingJobs);
+	FMemory::Memzero(PendingJobsHead);
 	FMemory::Memzero(NumPendingJobs);
+#if UE_SHADERCOMPILER_FIFO_JOB_EXECUTION
+	for (int32 PriorityIndex = 0; PriorityIndex < NumShaderCompileJobPriorities; PriorityIndex++)
+	{
+		PendingJobsTail[PriorityIndex] = &PendingJobsHead[PriorityIndex];
+	}
+#endif
 
 	CurrentlyAllocatedMemory = sizeof(*this) + InputHashToJobData.GetAllocatedSize() + Outputs.GetAllocatedSize();
 }
diff --git a/Engine/Source/Runtime/Engine/Public/ShaderCompiler.h b/Engine/Source/Runtime/Engine/Public/ShaderCompiler.h
--- a/Engine/Source/Runtime/Engine/Public/ShaderCompiler.h
+++ b/Engine/Source/Runtime/Engine/Public/ShaderCompiler.h
@@ -53,7 +53,7 @@ class FShaderJobCache;
 class FShaderCompileJobCollection
 {
 public:
-	FShaderCompileJobCollection();
+	FShaderCompileJobCollection(FCriticalSection& InCompileQueueSection);
 
 	FShaderCompileJob* PrepareJob(uint32 InId, const FShaderCompileJobKey& InKey, EShaderCompileJobPriority InPriority);
 	FShaderPipelineCompileJob* PrepareJob(uint32 InId, const FShaderPipelineCompileJobKey& InKey, EShaderCompileJobPriority InPriority);
@@ -440,6 +440,10 @@ private:
 		/** Max amount of time any single job spent being processed overall. */
 		double MaxJobLifeTime = 0;
 
+		/** Time spent in tasks generated in FShaderJobCache::SubmitJobs, plus stall time on mutex locks in those tasks */
+		double AccumulatedTaskSubmitJobs = 0.0;
+		double AccumulatedTaskSubmitJobsStall = 0.0;
+
 		/** Number of local job batches seen. */
 		int64 LocalJobBatchesSeen = 0;
 
@@ -502,6 +506,8 @@ private:
 			MaxJobExecutionTime = FMath::Max(Other.MaxJobExecutionTime, MaxJobExecutionTime);
 			AccumulatedJobLifeTime += Other.AccumulatedJobLifeTime;
 			MaxJobLifeTime = FMath::Max(Other.MaxJobLifeTime, MaxJobLifeTime);
+			AccumulatedTaskSubmitJobs += Other.AccumulatedTaskSubmitJobs;
+			AccumulatedTaskSubmitJobsStall += Other.AccumulatedTaskSubmitJobsStall;
 			LocalJobBatchesSeen += Other.LocalJobBatchesSeen;
 			TotalJobsReportedInLocalJobBatches += Other.TotalJobsReportedInLocalJobBatches;
 			DistributedJobBatchesSeen += Other.DistributedJobBatchesSeen;
diff --git a/Engine/Source/Runtime/RenderCore/Public/ShaderCompilerJobTypes.h b/Engine/Source/Runtime/RenderCore/Public/ShaderCompilerJobTypes.h
--- a/Engine/Source/Runtime/RenderCore/Public/ShaderCompilerJobTypes.h
+++ b/Engine/Source/Runtime/RenderCore/Public/ShaderCompilerJobTypes.h
@@ -70,9 +70,13 @@ struct FShaderJobCacheRef
 };
 
 /** Stores all of the common information used to compile a shader or pipeline. */
-class FShaderCommonCompileJob : public TIntrusiveLinkedList<FShaderCommonCompileJob>
+class FShaderCommonCompileJob
 {
 public:
+	/** Linked list support -- not using TIntrusiveLinkedList because we want lock free insertion not supported by the core class */
+	FShaderCommonCompileJob* NextLink = nullptr;
+	FShaderCommonCompileJob** PrevLink = nullptr;
+
 	using FInputHash = FBlake3Hash;
 
 	FPendingShaderMapCompileResultsPtr PendingShaderMap;
@@ -113,6 +117,9 @@ public:
 	double TimeAssignedToExecution = 0.0;
 	/** In-engine timestamp of job being completed. Encompasses the compile time. Not set for jobs that are satisfied from the jobs cache */
 	double TimeExecutionCompleted = 0.0;
+	/** Time spent in tasks generated in FShaderJobCache::SubmitJobs, plus stall time on mutex locks in those tasks */
+	double TimeTaskSubmitJobs = 0.0;
+	double TimeTaskSubmitJobsStall = 0.0;
 
 	FShaderJobCacheRef JobCacheRef;
 
