SetEditorProperty/GetEditorProperty can now access class sparse data, and fixed arrays can be assigned to/read from dynamic arrays

diff --git a/Engine/Source/Runtime/CoreUObject/Private/UObject/PropertyAccessUtil.cpp b/Engine/Source/Runtime/CoreUObject/Private/UObject/PropertyAccessUtil.cpp
--- a/Engine/Source/Runtime/CoreUObject/Private/UObject/PropertyAccessUtil.cpp
+++ b/Engine/Source/Runtime/CoreUObject/Private/UObject/PropertyAccessUtil.cpp
@@ -302,6 +302,44 @@ bool CopyCompletePropertyValue(const FProperty* InSrcProp, const void* InSrcValu
 
 	if (!ArePropertiesCompatible(InSrcProp, InDestProp) || InSrcProp->ArrayDim != InDestProp->ArrayDim)
 	{
+		if (InDestProp->ArrayDim > 1)
+		{
+			// handle assignment of a dynamic array to a fixed array:
+			if (const FArrayProperty* SrcArray = CastField<FArrayProperty>(InSrcProp))
+			{
+				if (ArePropertiesCompatible(SrcArray->Inner, InDestProp))
+				{
+					FScriptArrayHelper SrcArrayHelper(SrcArray, InSrcValue);
+					if (SrcArrayHelper.Num() == InDestProp->ArrayDim)
+					{
+						for (int32 I = 0; I < InDestProp->ArrayDim; ++I)
+						{
+							void* DestValue = static_cast<uint8*>(InDestValue) + InDestProp->ElementSize * I;
+							CopySinglePropertyValue(SrcArray->Inner, SrcArrayHelper.GetElementPtr(I), InDestProp, DestValue);
+						}
+						return true;
+					}
+				}
+			}
+		}
+		else if (InSrcProp->ArrayDim > 1)
+		{
+			// handle assignment of a fixed array to a dynamic array:
+			if (const FArrayProperty* DstArray = CastField<FArrayProperty>(InDestProp))
+			{
+				if (ArePropertiesCompatible(DstArray->Inner, InSrcProp))
+				{
+					FScriptArrayHelper DstArrayHelper(DstArray, InDestValue);
+					DstArrayHelper.Resize(InSrcProp->ArrayDim);
+					for (int32 I = 0; I < InSrcProp->ArrayDim; ++I)
+					{
+						const void* SrcValue = static_cast<const uint8*>(InSrcValue) + InSrcProp->ElementSize * I;
+						CopySinglePropertyValue(InSrcProp, SrcValue, DstArray->Inner, DstArrayHelper.GetElementPtr(I));
+					}
+					return true;
+				}
+			}
+		}
 		return false;
 	}
 
@@ -577,7 +615,10 @@ void EmitPostChangeNotify(const FPropertyAccessChangeNotify* InChangeNotify, con
 
 TUniquePtr<FPropertyAccessChangeNotify> BuildBasicChangeNotify(const FProperty* InProp, const UObject* InObject, const EPropertyAccessChangeNotifyMode InNotifyMode)
 {
-	check(InObject->IsA(InProp->GetOwnerClass()));
+	const UScriptStruct* SparseStruct = InObject->GetClass()->GetSparseClassDataStruct();
+	const bool bIsValidSparseProp = SparseStruct &&
+		SparseStruct->IsChildOf(InProp->GetOwnerStruct());
+	check(InObject->IsA(InProp->GetOwnerClass()) || bIsValidSparseProp);
 #if WITH_EDITOR
 	if (InNotifyMode != EPropertyAccessChangeNotifyMode::Never)
 	{
diff --git a/Engine/Source/Runtime/Engine/Classes/Kismet/KismetSystemLibrary.h b/Engine/Source/Runtime/Engine/Classes/Kismet/KismetSystemLibrary.h
--- a/Engine/Source/Runtime/Engine/Classes/Kismet/KismetSystemLibrary.h
+++ b/Engine/Source/Runtime/Engine/Classes/Kismet/KismetSystemLibrary.h
@@ -1931,7 +1931,7 @@ class UKismetSystemLibrary : public UBlueprintFunctionLibrary
 	 */
     UFUNCTION(BlueprintCallable, CustomThunk, Category = "Utilities", meta=(CustomStructureParam="PropertyValue", BlueprintInternalUseOnly="true"))
     static ENGINE_API bool GetEditorProperty(UObject* Object, const FName PropertyName, int32& PropertyValue);
-	static ENGINE_API bool Generic_GetEditorProperty(const UObject* Object, const FProperty* ObjectProp, void* ValuePtr, const FProperty* ValueProp);
+	static ENGINE_API bool Generic_GetEditorProperty(const UObject* Object, const FName PropertyName, void* ValuePtr, const FProperty* ValueProp);
 	DECLARE_FUNCTION(execGetEditorProperty);
 
 	/**
@@ -1946,7 +1946,7 @@ class UKismetSystemLibrary : public UBlueprintFunctionLibrary
 	 */
     UFUNCTION(BlueprintCallable, CustomThunk, Category = "Utilities", meta=(CustomStructureParam="PropertyValue", AdvancedDisplay="ChangeNotifyMode", BlueprintInternalUseOnly="true"))
     static ENGINE_API bool SetEditorProperty(UObject* Object, const FName PropertyName, const int32& PropertyValue, const EPropertyAccessChangeNotifyMode ChangeNotifyMode);
-	static ENGINE_API bool Generic_SetEditorProperty(UObject* Object, const FProperty* ObjectProp, const void* ValuePtr, const FProperty* ValueProp, const EPropertyAccessChangeNotifyMode ChangeNotifyMode);
+	static ENGINE_API bool Generic_SetEditorProperty(UObject* Object, const FName PropertyName, const void* ValuePtr, const FProperty* ValueProp, const EPropertyAccessChangeNotifyMode ChangeNotifyMode);
 	DECLARE_FUNCTION(execSetEditorProperty);
 #endif
 
diff --git a/Engine/Source/Runtime/Engine/Private/KismetSystemLibrary.cpp b/Engine/Source/Runtime/Engine/Private/KismetSystemLibrary.cpp
--- a/Engine/Source/Runtime/Engine/Private/KismetSystemLibrary.cpp
+++ b/Engine/Source/Runtime/Engine/Private/KismetSystemLibrary.cpp
@@ -3082,9 +3082,40 @@ bool UKismetSystemLibrary::GetEditorProperty(UObject* Object, const FName Proper
 	return false;
 }
 
-bool UKismetSystemLibrary::Generic_GetEditorProperty(const UObject* Object, const FProperty* ObjectProp, void* ValuePtr, const FProperty* ValueProp)
+bool UKismetSystemLibrary::Generic_GetEditorProperty(const UObject* Object, const FName PropertyName, void* ValuePtr, const FProperty* ValueProp)
 {
-	const EPropertyAccessResultFlags AccessResult = PropertyAccessUtil::GetPropertyValue_Object(ObjectProp, Object, ValueProp, ValuePtr, INDEX_NONE);
+	const FProperty* ObjectProp = PropertyAccessUtil::FindPropertyByName(PropertyName, Object->GetClass());
+	TOptional<EPropertyAccessResultFlags> SparseDataAccessResult;
+	if ((!ObjectProp || ObjectProp->HasAnyPropertyFlags(CPF_Deprecated)) && Object->HasAllFlags(RF_ClassDefaultObject))
+	{
+		// look for a sparse member of the same name - the sparse data is treated as an extension
+		// of the class default object by the details panel:
+		const UStruct* SparseDataStruct = Object->GetClass()->GetSparseClassDataStruct();
+		if (SparseDataStruct)
+		{
+			const FProperty* SparseProp = PropertyAccessUtil::FindPropertyByName(PropertyName, SparseDataStruct);
+			if (SparseProp)
+			{
+				void* SparseDest = Object->GetClass()->GetOrCreateSparseClassData();
+
+				SparseDataAccessResult = PropertyAccessUtil::GetPropertyValue_InContainer(
+					SparseProp,
+					SparseDest,
+					ValueProp,
+					ValuePtr,
+					INDEX_NONE);
+			}
+		}
+	}
+
+	if (!ObjectProp && !SparseDataAccessResult.IsSet())
+	{
+		FFrame::KismetExecutionMessage(*FString::Printf(TEXT("Property '%s' on '%s' (%s) was missing"), *PropertyName.ToString(), *Object->GetPathName(), *Object->GetClass()->GetName()), ELogVerbosity::Warning, UE::Blueprint::Private::PropertyGetFailedWarning);
+	}
+
+	const EPropertyAccessResultFlags AccessResult = SparseDataAccessResult.IsSet() ?
+		SparseDataAccessResult.GetValue() : 
+		PropertyAccessUtil::GetPropertyValue_Object(ObjectProp, Object, ValueProp, ValuePtr, INDEX_NONE);
 
 	if (EnumHasAnyFlags(AccessResult, EPropertyAccessResultFlags::PermissionDenied))
 	{
@@ -3140,17 +3171,9 @@ DEFINE_FUNCTION(UKismetSystemLibrary::execGetEditorProperty)
 
 	if (Object)
 	{
-		const FProperty* ObjectProp = PropertyAccessUtil::FindPropertyByName(PropertyName, Object->GetClass());
-		if (ObjectProp)
-		{
-			P_NATIVE_BEGIN;
-			bResult = Generic_GetEditorProperty(Object, ObjectProp, ValuePtr, ValueProp);
-			P_NATIVE_END;
-		}
-		else
-		{
-			FFrame::KismetExecutionMessage(*FString::Printf(TEXT("Property '%s' on '%s' (%s) was missing"), *PropertyName.ToString(), *Object->GetPathName(), *Object->GetClass()->GetName()), ELogVerbosity::Warning, UE::Blueprint::Private::PropertyGetFailedWarning);
-		}
+		P_NATIVE_BEGIN;
+		bResult = Generic_GetEditorProperty(Object, PropertyName, ValuePtr, ValueProp);
+		P_NATIVE_END;
 	}
 
 	*(bool*)RESULT_PARAM = bResult;
@@ -3163,9 +3186,47 @@ bool UKismetSystemLibrary::SetEditorProperty(UObject* Object, const FName Proper
 	return false;
 }
 
-bool UKismetSystemLibrary::Generic_SetEditorProperty(UObject* Object, const FProperty* ObjectProp, const void* ValuePtr, const FProperty* ValueProp, const EPropertyAccessChangeNotifyMode ChangeNotifyMode)
+bool UKismetSystemLibrary::Generic_SetEditorProperty(UObject* Object, const FName PropertyName, const void* ValuePtr, const FProperty* ValueProp, const EPropertyAccessChangeNotifyMode ChangeNotifyMode)
 {
-	const EPropertyAccessResultFlags AccessResult = PropertyAccessUtil::SetPropertyValue_Object(ObjectProp, Object, ValueProp, ValuePtr, INDEX_NONE, PropertyAccessUtil::EditorReadOnlyFlags, ChangeNotifyMode);
+	TOptional<EPropertyAccessResultFlags> SparseDataAccessResult;
+	const FProperty* ObjectProp = PropertyAccessUtil::FindPropertyByName(PropertyName, Object->GetClass());
+	if ((!ObjectProp || ObjectProp->HasAnyPropertyFlags(CPF_Deprecated)) && Object->HasAllFlags(RF_ClassDefaultObject))
+	{
+		// look for a sparse member of the same name - the sparse data is treated as an extension
+		// of the class default object by the details panel:
+		const UStruct* SparseDataStruct = Object->GetClass()->GetSparseClassDataStruct();
+		if (SparseDataStruct)
+		{
+			const FProperty* SparseProp = PropertyAccessUtil::FindPropertyByName(PropertyName, SparseDataStruct);
+			if (SparseProp)
+			{
+				void* SparseDest = Object->GetClass()->GetOrCreateSparseClassData();
+
+				SparseDataAccessResult = PropertyAccessUtil::SetPropertyValue_InContainer(
+					SparseProp,
+					SparseDest,
+					ValueProp,
+					ValuePtr,
+					INDEX_NONE,
+					PropertyAccessUtil::EditorReadOnlyFlags,
+					PropertyAccessUtil::IsObjectTemplate(Object),
+					[SparseProp, Object, ChangeNotifyMode]()
+					{
+						return PropertyAccessUtil::BuildBasicChangeNotify(SparseProp, Object, ChangeNotifyMode);
+					});
+			}
+		}
+	}
+
+	if(!ObjectProp && !SparseDataAccessResult.IsSet())
+	{
+		FFrame::KismetExecutionMessage(*FString::Printf(TEXT("Property '%s' on '%s' (%s) was missing"), *PropertyName.ToString(), *Object->GetPathName(), *Object->GetClass()->GetName()), ELogVerbosity::Warning, UE::Blueprint::Private::PropertySetFailedWarning);
+		return false;
+	}
+
+	const EPropertyAccessResultFlags AccessResult = SparseDataAccessResult.IsSet() ?
+		SparseDataAccessResult.GetValue() :
+		PropertyAccessUtil::SetPropertyValue_Object(ObjectProp, Object, ValueProp, ValuePtr, INDEX_NONE, PropertyAccessUtil::EditorReadOnlyFlags, ChangeNotifyMode);
 
 	if (EnumHasAnyFlags(AccessResult, EPropertyAccessResultFlags::PermissionDenied))
 	{
@@ -3241,17 +3302,9 @@ DEFINE_FUNCTION(UKismetSystemLibrary::execSetEditorProperty)
 
 	if (Object)
 	{
-		const FProperty* ObjectProp = PropertyAccessUtil::FindPropertyByName(PropertyName, Object->GetClass());
-		if (ObjectProp)
-		{
-			P_NATIVE_BEGIN;
-			bResult = Generic_SetEditorProperty(Object, ObjectProp, ValuePtr, ValueProp, ChangeNotifyMode);
-			P_NATIVE_END;
-		}
-		else
-		{
-			FFrame::KismetExecutionMessage(*FString::Printf(TEXT("Property '%s' on '%s' (%s) was missing"), *PropertyName.ToString(), *Object->GetPathName(), *Object->GetClass()->GetName()), ELogVerbosity::Warning, UE::Blueprint::Private::PropertySetFailedWarning);
-		}
+		P_NATIVE_BEGIN;
+		bResult = Generic_SetEditorProperty(Object, PropertyName, ValuePtr, ValueProp, ChangeNotifyMode);
+		P_NATIVE_END;
 	}
 
 	*(bool*)RESULT_PARAM = bResult;
