Allow users to specify bForceRootLock when sampling poses via AnimPose API

diff --git a/Engine/Source/Editor/AnimationBlueprintLibrary/Private/AnimPose.cpp b/Engine/Source/Editor/AnimationBlueprintLibrary/Private/AnimPose.cpp
--- a/Engine/Source/Editor/AnimationBlueprintLibrary/Private/AnimPose.cpp
+++ b/Engine/Source/Editor/AnimationBlueprintLibrary/Private/AnimPose.cpp
@@ -632,6 +632,9 @@ void UAnimPoseExtensions::GetAnimPoseAtTimeIntervals(const UAnimSequenceBase* An
 
         FAnimationPoseData PoseData(CompactPose, Curve, Attributes);
         FAnimExtractContext Context(0.0, EvaluationOptions.bExtractRootMotion);
+#if WITH_EDITOR
+		Context.bIgnoreRootLock = EvaluationOptions.bIncorporateRootMotionIntoPose;
+#endif // WITH_EDITOR
     
         FCompactPose BasePose;
         BasePose.SetBoneContainer(&RequiredBones);
diff --git a/Engine/Source/Editor/AnimationBlueprintLibrary/Public/AnimPose.h b/Engine/Source/Editor/AnimationBlueprintLibrary/Public/AnimPose.h
--- a/Engine/Source/Editor/AnimationBlueprintLibrary/Public/AnimPose.h
+++ b/Engine/Source/Editor/AnimationBlueprintLibrary/Public/AnimPose.h
@@ -63,6 +63,10 @@ struct ANIMATIONBLUEPRINTLIBRARY_API FAnimPoseEvaluationOptions
 	UPROPERTY(EditAnywhere, AdvancedDisplay, BlueprintReadWrite, Category="Animation|Pose")
 	bool bExtractRootMotion = false;
 
+	// Whether or not to force root motion being incorporated into retrieved pose
+	UPROPERTY(EditAnywhere, AdvancedDisplay, BlueprintReadWrite, Category="Animation|Pose")
+	bool bIncorporateRootMotionIntoPose = true;
+
 	// Optional skeletal mesh with proportions to use when evaluating a pose
 	UPROPERTY(EditAnywhere, AdvancedDisplay, BlueprintReadWrite, Category="Animation|Pose")
 	TObjectPtr<USkeletalMesh> OptionalSkeletalMesh = nullptr;
diff --git a/Engine/Source/Runtime/Engine/Classes/Animation/AnimationAsset.h b/Engine/Source/Runtime/Engine/Classes/Animation/AnimationAsset.h
--- a/Engine/Source/Runtime/Engine/Classes/Animation/AnimationAsset.h
+++ b/Engine/Source/Runtime/Engine/Classes/Animation/AnimationAsset.h
@@ -281,6 +281,10 @@ struct FAnimExtractContext
 	 * by several animation nodes to optimize evaluation time.
 	 */
 	TArray<bool> BonesRequired;
+
+#if WITH_EDITOR
+	bool bIgnoreRootLock;
+#endif 
 	
 	UE_DEPRECATED(5.1, "FAnimExtractContext construct with float-based time value is deprecated, use other signature")
 	FAnimExtractContext(float InCurrentTime, bool InbExtractRootMotion = false, FDeltaTimeRecord InDeltaTimeRecord = {}, bool InbLooping = false)
@@ -290,6 +294,9 @@ struct FAnimExtractContext
 		, bLooping(InbLooping)
 		, PoseCurves()
 		, BonesRequired()
+#if WITH_EDITOR
+		, bIgnoreRootLock(false)
+#endif 
 	{
 	}
 
@@ -300,6 +307,9 @@ struct FAnimExtractContext
 		, bLooping(InbLooping)
 		, PoseCurves()
 		, BonesRequired()
+#if WITH_EDITOR
+		, bIgnoreRootLock(false)
+#endif 
 	{
 	}
 
diff --git a/Engine/Source/Runtime/Engine/Private/Animation/AnimSequence.cpp b/Engine/Source/Runtime/Engine/Private/Animation/AnimSequence.cpp
--- a/Engine/Source/Runtime/Engine/Private/Animation/AnimSequence.cpp
+++ b/Engine/Source/Runtime/Engine/Private/Animation/AnimSequence.cpp
@@ -1654,7 +1654,12 @@ void UAnimSequence::GetBonePose(FAnimationPoseData& OutAnimationPoseData, const
 	const int32 NumTracks = CompressedData.CompressedTrackToSkeletonMapTable.Num();
 #endif 
 	const bool bTreatAnimAsAdditive = (IsValidAdditive() && !bUseRawDataForPoseExtraction); // Raw data is never additive
-	const FRootMotionReset RootMotionReset(bEnableRootMotion, RootMotionRootLock, bForceRootLock, ExtractRootTrackTransform(0.0f, &RequiredBones), bTreatAnimAsAdditive);
+	const FRootMotionReset RootMotionReset(bEnableRootMotion, RootMotionRootLock,
+#if WITH_EDITOR
+		!ExtractionContext.bIgnoreRootLock &&
+#endif // WITH_EDITOR
+		bForceRootLock,
+		ExtractRootTrackTransform(0.0f, &RequiredBones), bTreatAnimAsAdditive);
 
 #if WITH_EDITOR
 	// Evaluate raw (source) curve and bone data
diff --git a/Engine/Source/Runtime/Engine/Private/Animation/AnimStreamable.cpp b/Engine/Source/Runtime/Engine/Private/Animation/AnimStreamable.cpp
--- a/Engine/Source/Runtime/Engine/Private/Animation/AnimStreamable.cpp
+++ b/Engine/Source/Runtime/Engine/Private/Animation/AnimStreamable.cpp
@@ -316,7 +316,12 @@ void UAnimStreamable::GetAnimationPose(FAnimationPoseData& OutAnimationPoseData,
 	}
 
 	//FRootMotionReset RootMotionReset(bEnableRootMotion, RootMotionRootLock, bForceRootLock, ExtractRootTrackTransform(0.f, &RequiredBones), IsValidAdditive());
-	FRootMotionReset RootMotionReset(bEnableRootMotion, RootMotionRootLock, bForceRootLock, FTransform(), false); // MDW Does not support root motion yet
+	FRootMotionReset RootMotionReset(bEnableRootMotion, RootMotionRootLock,
+#if WITH_EDITOR
+	!ExtractionContext.bIgnoreRootLock &&
+#endif // WITH_EDITOR
+	bForceRootLock,
+	FTransform(), false); // MDW Does not support root motion yet
 
 #if WITH_EDITOR
 	if (IsDataModelValid() && (!HasRunningPlatformData() || RequiredBones.ShouldUseRawData()))
