World Partition: fix dirty actors garbage collection prevention code using stale pointers.
This can happen when the actor was reinstanced. In this case we fallback on the actor pointer resolution in the actor descriptor.

diff --git a/Engine/Source/Runtime/Engine/Private/WorldPartition/WorldPartition.cpp b/Engine/Source/Runtime/Engine/Private/WorldPartition/WorldPartition.cpp
--- a/Engine/Source/Runtime/Engine/Private/WorldPartition/WorldPartition.cpp
+++ b/Engine/Source/Runtime/Engine/Private/WorldPartition/WorldPartition.cpp
@@ -1555,9 +1555,10 @@ void UWorldPartition::AddReferencedObjects(UObject* InThis, FReferenceCollector&
 		Collector.AllowEliminatingReferences(false);
 		for (auto& [DirtyActorGuid, DirtyActor] : This->DirtyActors)
 		{
-			if (DirtyActor.WorldPartitionRef.IsSet() || DirtyActor.ActorPtr.IsValid())
+			// Actor pointer might be invalid if the actor was reinstanced, etc, in this case we fallback on the actor pointer resolution in the actor descriptor.
+			if (AActor* ActorPtr = DirtyActor.ActorPtr.IsValid() ? DirtyActor.ActorPtr.Get() : (DirtyActor.WorldPartitionRef.IsSet() ? DirtyActor.WorldPartitionRef.GetValue()->GetActor() : nullptr))
 			{
-				Collector.AddReferencedObject(DirtyActor.ActorPtr);
+				Collector.AddReferencedObject(ActorPtr);
 			}
 		}
 		Collector.AllowEliminatingReferences(true);
