Fix Blueprint watches not working for pin values in the anim graph

As a side effect, this also allows breakpoints to be used in anim node functions

diff --git a/Engine/Source/Runtime/Engine/Private/Animation/AnimInstance.cpp b/Engine/Source/Runtime/Engine/Private/Animation/AnimInstance.cpp
--- a/Engine/Source/Runtime/Engine/Private/Animation/AnimInstance.cpp
+++ b/Engine/Source/Runtime/Engine/Private/Animation/AnimInstance.cpp
@@ -27,6 +27,9 @@
 #include "Animation/ActiveMontageInstanceScope.h"
 #include "Animation/AnimNode_SaveCachedPose.h"
 #include "Animation/AnimSubsystem_SharedLinkedAnimLayers.h"
+#if WITH_EDITOR
+#include "Engine/Blueprint.h"
+#endif
 
 #include UE_INLINE_GENERATED_CPP_BY_NAME(AnimInstance)
 
@@ -771,11 +774,21 @@ void UAnimInstance::ParallelUpdateAnimation()
 bool UAnimInstance::NeedsImmediateUpdate(float DeltaSeconds, bool bNeedsValidRootMotion) const
 {
 	const bool bUseParallelUpdateAnimation = (GetDefault<UEngine>()->bAllowMultiThreadedAnimationUpdate && bUseMultiThreadedAnimationUpdate) || (CVarForceUseParallelAnimUpdate.GetValueOnGameThread() != 0);
+#if WITH_EDITOR
+	UAnimBlueprintGeneratedClass* GeneratedClass = Cast<UAnimBlueprintGeneratedClass>(GetClass());
+	UBlueprint* Blueprint = GeneratedClass ? Cast<UBlueprint>(GeneratedClass->ClassGeneratedBy) : nullptr;
+#endif
 
 	return
 		(bNeedsValidRootMotion && RootMotionMode == ERootMotionMode::RootMotionFromEverything) ||
 		!CanRunParallelWork() ||
 		GIntraFrameDebuggingGameThread ||
+#if WITH_EDITOR
+		// Force the debugged object to run its anim graph on the game thread if it is being debugged
+		// This ensures that it uses the persistent ubergraph frame and debugging facilities are available like
+		// watches, breakpoints etc.
+		(Blueprint && Blueprint->GetObjectBeingDebugged() == this) ||
+#endif
 		CVarUseParallelAnimUpdate.GetValueOnGameThread() == 0 ||
 		CVarUseParallelAnimationEvaluation.GetValueOnGameThread() == 0 ||
 		!bUseParallelUpdateAnimation ||
