Splines: Add a SetStartRollDegrees and SetEndRollDegrees function to a spline mesh component to help alert users to the units when they do the common action of getting the roll from a spline (which is in degrees) and try to set it on a spline mesh (where the original setters are in radians).

diff --git a/Engine/Source/Runtime/Engine/Classes/Components/SplineMeshComponent.h b/Engine/Source/Runtime/Engine/Classes/Components/SplineMeshComponent.h
--- a/Engine/Source/Runtime/Engine/Classes/Components/SplineMeshComponent.h
+++ b/Engine/Source/Runtime/Engine/Classes/Components/SplineMeshComponent.h
@@ -51,7 +51,7 @@ struct FSplineMeshParams
 	UPROPERTY(EditAnywhere, Category=SplineMesh, AdvancedDisplay)
 	FVector2D StartScale;
 
-	/** Roll around spline applied at start */
+	/** Roll around spline applied at start, in radians. */
 	UPROPERTY(EditAnywhere, Category=SplineMesh, AdvancedDisplay)
 	float StartRoll;
 
@@ -71,7 +71,7 @@ struct FSplineMeshParams
 	UPROPERTY(EditAnywhere, Category = SplineMesh)
 	FVector EndTangent;
 
-	/** Roll around spline applied at end. */
+	/** Roll around spline applied at end, in radians. */
 	UPROPERTY(EditAnywhere, Category=SplineMesh, AdvancedDisplay)
 	float EndRoll;
 
@@ -275,14 +275,18 @@ public:
 	UFUNCTION(BlueprintCallable, Category = SplineMesh)
 	ENGINE_API void SetStartScale(FVector2D StartScale = FVector2D(1,1), bool bUpdateMesh = true);
 
-	/** Get the start roll */
+	/** Get the start roll, in radians */
 	UFUNCTION(BlueprintCallable, Category = SplineMesh)
 	ENGINE_API float GetStartRoll() const;
 
-	/** Set the start roll */
-	UFUNCTION(BlueprintCallable, Category = SplineMesh)
+	/** Set the start roll, in radians */
+	UFUNCTION(BlueprintCallable, Category = SplineMesh, meta = (DisplayName = "Set Start Roll Radians"))
 	ENGINE_API void SetStartRoll(float StartRoll, bool bUpdateMesh = true);
 
+	/** Set the start roll in degrees */
+	UFUNCTION(BlueprintCallable, Category = SplineMesh)
+	ENGINE_API void SetStartRollDegrees(float StartRollDegrees, bool bUpdateMesh = true);
+
 	/** Get the start offset */
 	UFUNCTION(BlueprintCallable, Category = SplineMesh)
 	ENGINE_API FVector2D GetStartOffset() const;
@@ -299,14 +303,18 @@ public:
 	UFUNCTION(BlueprintCallable, Category = SplineMesh)
 	ENGINE_API void SetEndScale(FVector2D EndScale = FVector2D(1,1), bool bUpdateMesh = true);
 
-	/** Get the end roll */
+	/** Get the end roll, in radians */
 	UFUNCTION(BlueprintCallable, Category = SplineMesh)
 	ENGINE_API float GetEndRoll() const;
 
-	/** Set the end roll */
-	UFUNCTION(BlueprintCallable, Category = SplineMesh)
+	/** Set the end roll, in radians */
+	UFUNCTION(BlueprintCallable, Category = SplineMesh, meta = (DisplayName = "Set End Roll Radians"))
 	ENGINE_API void SetEndRoll(float EndRoll, bool bUpdateMesh = true);
 
+	/** Set the end roll in degrees */
+	UFUNCTION(BlueprintCallable, Category = SplineMesh)
+	ENGINE_API void SetEndRollDegrees(float EndRollDegrees, bool bUpdateMesh = true);
+
 	/** Get the end offset */
 	UFUNCTION(BlueprintCallable, Category = SplineMesh)
 	ENGINE_API FVector2D GetEndOffset() const;
diff --git a/Engine/Source/Runtime/Engine/Private/Components/SplineMeshComponent.cpp b/Engine/Source/Runtime/Engine/Private/Components/SplineMeshComponent.cpp
--- a/Engine/Source/Runtime/Engine/Private/Components/SplineMeshComponent.cpp
+++ b/Engine/Source/Runtime/Engine/Private/Components/SplineMeshComponent.cpp
@@ -450,6 +450,11 @@ void USplineMeshComponent::SetStartRoll(float StartRoll, bool bUpdateMesh)
 	}
 }
 
+void USplineMeshComponent::SetStartRollDegrees(float StartRollDegrees, bool bUpdateMesh)
+{
+	SetStartRoll(FMath::DegreesToRadians(StartRollDegrees), bUpdateMesh);
+}
+
 FVector2D USplineMeshComponent::GetStartOffset() const
 {
 	return SplineParams.StartOffset;
@@ -500,6 +505,11 @@ void USplineMeshComponent::SetEndRoll(float EndRoll, bool bUpdateMesh)
 	}
 }
 
+void USplineMeshComponent::SetEndRollDegrees(float EndRollDegrees, bool bUpdateMesh)
+{
+	SetEndRoll(FMath::DegreesToRadians(EndRollDegrees), bUpdateMesh);
+}
+
 FVector2D USplineMeshComponent::GetEndOffset() const
 {
 	return SplineParams.EndOffset;
