Disable taskification of shader job submission and bump shader version to hopefully address DDC poisoning issues.

diff --git a/Engine/Shaders/Public/ShaderVersion.ush b/Engine/Shaders/Public/ShaderVersion.ush
--- a/Engine/Shaders/Public/ShaderVersion.ush
+++ b/Engine/Shaders/Public/ShaderVersion.ush
@@ -2,4 +2,4 @@
 // Each time the console command is executed it generates a new GUID. As this file is included
 // in Platform.ush (which should be included in any shader) it allows to invalidate the shader DDC.
 // 
-// GUID = F8F483DB4AC710E5DAF54981165B2A4F
\ No newline at end of file
+// GUID = 587833ec-6a8c-4f23-9191-cc27a4b73a71
\ No newline at end of file
diff --git a/Engine/Source/Runtime/Engine/Private/ShaderCompiler/ShaderCompiler.cpp b/Engine/Source/Runtime/Engine/Private/ShaderCompiler/ShaderCompiler.cpp
--- a/Engine/Source/Runtime/Engine/Private/ShaderCompiler/ShaderCompiler.cpp
+++ b/Engine/Source/Runtime/Engine/Private/ShaderCompiler/ShaderCompiler.cpp
@@ -243,7 +243,7 @@ static FAutoConsoleVariableRef CVarShaderCompilerDebugDiscardCacheOutputs(
 
 static TAutoConsoleVariable<bool> CVarShaderCompilerParallelSubmitJobs(
 	TEXT("r.ShaderCompiler.ParallelSubmitJobs"),
-	true,
+	false,
 	TEXT("if != 0, FShaderJobCache::SubmitJobs will run in multiple parallel tasks, instead of the game thread."),
 	ECVF_Default
 );
