Make ENetRole available to Blueprints
PR #10411: Update EngineTypes.h - ENetRole to BlueprintType

diff --git a/Engine/Source/Runtime/Engine/Classes/Engine/EngineTypes.h b/Engine/Source/Runtime/Engine/Classes/Engine/EngineTypes.h
--- a/Engine/Source/Runtime/Engine/Classes/Engine/EngineTypes.h
+++ b/Engine/Source/Runtime/Engine/Classes/Engine/EngineTypes.h
@@ -2937,18 +2937,18 @@ public:
 };
 
 /** The network role of an actor on a local/remote network context */
-UENUM()
+UENUM(BlueprintType)
 enum ENetRole : int
 {
 	/** No role at all. */
-	ROLE_None,
+	ROLE_None UMETA(DisplayName = "None"),
 	/** Locally simulated proxy of this actor. */
-	ROLE_SimulatedProxy,
+	ROLE_SimulatedProxy UMETA(DisplayName = "Simulated Proxy"),
 	/** Locally autonomous proxy of this actor. */
-	ROLE_AutonomousProxy,
+	ROLE_AutonomousProxy UMETA(DisplayName = "Autonomous Proxy"),
 	/** Authoritative control over the actor. */
-	ROLE_Authority,
-	ROLE_MAX,
+	ROLE_Authority UMETA(DisplayName = "Authority"),
+	ROLE_MAX UMETA(Hidden),
 };
 
 /** Describes if an actor can enter a low network bandwidth dormant mode */
