Fix GetEditorProperty not gracefully failing for missing properties post 26760895, missed an early return
#rnx

diff --git a/Engine/Source/Runtime/Engine/Private/KismetSystemLibrary.cpp b/Engine/Source/Runtime/Engine/Private/KismetSystemLibrary.cpp
--- a/Engine/Source/Runtime/Engine/Private/KismetSystemLibrary.cpp
+++ b/Engine/Source/Runtime/Engine/Private/KismetSystemLibrary.cpp
@@ -3111,6 +3111,7 @@ bool UKismetSystemLibrary::Generic_GetEditorProperty(const UObject* Object, cons
 	if (!ObjectProp && !SparseDataAccessResult.IsSet())
 	{
 		FFrame::KismetExecutionMessage(*FString::Printf(TEXT("Property '%s' on '%s' (%s) was missing"), *PropertyName.ToString(), *Object->GetPathName(), *Object->GetClass()->GetName()), ELogVerbosity::Warning, UE::Blueprint::Private::PropertyGetFailedWarning);
+		return false;
 	}
 
 	const EPropertyAccessResultFlags AccessResult = SparseDataAccessResult.IsSet() ?
