World Partition: always include editor-only references in the instancing context when cooking instanced packages.
This fixes the case where a cooked actor is referencing another cooked actors through an editor-only reference.

diff --git a/Engine/Source/Runtime/Engine/Private/WorldPartition/WorldPartitionRuntimeLevelStreamingCell.cpp b/Engine/Source/Runtime/Engine/Private/WorldPartition/WorldPartitionRuntimeLevelStreamingCell.cpp
--- a/Engine/Source/Runtime/Engine/Private/WorldPartition/WorldPartitionRuntimeLevelStreamingCell.cpp
+++ b/Engine/Source/Runtime/Engine/Private/WorldPartition/WorldPartitionRuntimeLevelStreamingCell.cpp
@@ -228,23 +228,20 @@ void UWorldPartitionRuntimeLevelStreamingCell::AddActorToCell(const FWorldPartit
 {
 	check(!ActorDescView.GetActorIsEditorOnly());
 
-	if (!IsRunningCookCommandlet())
+	for (const FGuid& EditorReferenceGuid : ActorDescView.GetEditorReferences())
 	{
-		for (const FGuid& EditorReferenceGuid : ActorDescView.GetEditorReferences())
-		{
-			const FWorldPartitionActorDesc& ReferenceActorDesc = InContainer->GetActorDescChecked(EditorReferenceGuid);
-
-			Packages.Emplace(
-				ReferenceActorDesc.GetActorPackage(), 
-				*ReferenceActorDesc.GetActorSoftPath().ToString(), 
-				InContainerID, 
-				InContainerTransform, 
-				InContainer->GetContainerPackage(), 
-				GetWorld()->GetPackage()->GetFName(), 
-				InContainerID.GetActorGuid(ReferenceActorDesc.GetGuid()),
-				true
-			);
-		}
+		const FWorldPartitionActorDesc& ReferenceActorDesc = InContainer->GetActorDescChecked(EditorReferenceGuid);
+
+		Packages.Emplace(
+			ReferenceActorDesc.GetActorPackage(), 
+			*ReferenceActorDesc.GetActorSoftPath().ToString(), 
+			InContainerID, 
+			InContainerTransform, 
+			InContainer->GetContainerPackage(), 
+			GetWorld()->GetPackage()->GetFName(), 
+			InContainerID.GetActorGuid(ReferenceActorDesc.GetGuid()),
+			true
+		);
 	}
 
 	Packages.Emplace(
