World Partition: always output streaming generation logs during automation tests, even on build machines.

diff --git a/Engine/Source/Runtime/Engine/Private/WorldPartition/WorldPartitionStreamingGeneration.cpp b/Engine/Source/Runtime/Engine/Private/WorldPartition/WorldPartitionStreamingGeneration.cpp
--- a/Engine/Source/Runtime/Engine/Private/WorldPartition/WorldPartitionStreamingGeneration.cpp
+++ b/Engine/Source/Runtime/Engine/Private/WorldPartition/WorldPartitionStreamingGeneration.cpp
@@ -1287,7 +1287,7 @@ bool UWorldPartition::GenerateContainerStreaming(const FGenerateStreamingParams&
 	TUniquePtr<FArchive> LogFileAr;
 	TUniquePtr<FHierarchicalLogArchive> HierarchicalLogAr;
 
-	if (IsMainWorldPartition() && (!GIsBuildMachine || IsRunningCookCommandlet()))
+	if (IsMainWorldPartition() && (!GIsBuildMachine || GIsAutomationTesting || IsRunningCookCommandlet()))
 	{
 		TStringBuilder<256> StateLogSuffix;
 		StateLogSuffix += bIsPIE ? TEXT("PIE") : (IsRunningGame() ? TEXT("Game") : (IsRunningCookCommandlet() ? TEXT("Cook") : (GIsAutomationTesting ? TEXT("UnitTest") : TEXT("Manual"))));
