Comment trail so the next time it doesn't take a day to debug this.

diff --git a/Engine/Source/Runtime/Engine/Private/TextureCompiler.cpp b/Engine/Source/Runtime/Engine/Private/TextureCompiler.cpp
--- a/Engine/Source/Runtime/Engine/Private/TextureCompiler.cpp
+++ b/Engine/Source/Runtime/Engine/Private/TextureCompiler.cpp
@@ -241,6 +241,18 @@ int32 FTextureCompilingManager::GetNumRemainingAssets() const
 void FTextureCompilingManager::AddTextures(TArrayView<UTexture* const> InTextures)
 {
 	check(IsInGameThread());
+
+	// If you hit this, it's because above this in the stack you'll see PostCompilation(). In that function you'll see:
+	// 	Texture->FinishCachePlatformData();
+	//	Texture->UpdateResource();
+	// UpdateResource ends up doing another CachePlatformData() - so what's happened is you finished pulling in the derived data
+	// and then immediately tried again - and then tried to launch another build because the ddc keys changed. This means that
+	// during the async build, a property or otherwise that is an input to the ddc key changed. This shouldn't happen because
+	// PreEditChange completes the async build before allowing the change.
+	// Debugging this can be a huge pain. If you have a repro, IMO the best way is to hack GetTextureDerivedDataKeySuffix
+	// to strcmp on the name of the repro texture and just log the full key suffix. Then you should immediately see the changed
+	// keys right before the crash and you can backsolve what value changed. Once you have that, you can set a data breakpoint on
+	// the property and see who is poking it.
 	checkf(bIsRoutingPostCompilation == false,
 		TEXT("Registering a texture to the compile manager from inside a texture postcompilation is not supported and usually indicate that the previous async operation wasn't completed (i.e. missing call to PreEditChange) before modifying a texture property."));
 
