DeadBlending Node memory optimizations.

diff --git a/Engine/Source/Runtime/Engine/Classes/Animation/AnimNode_DeadBlending.h b/Engine/Source/Runtime/Engine/Classes/Animation/AnimNode_DeadBlending.h
--- a/Engine/Source/Runtime/Engine/Classes/Animation/AnimNode_DeadBlending.h
+++ b/Engine/Source/Runtime/Engine/Classes/Animation/AnimNode_DeadBlending.h
@@ -204,19 +204,19 @@ private:
 
 	// Recorded pose state at point of transition.
 
-	TBitArray<TInlineAllocator<512>> BoneValid;
+	TBitArray<TInlineAllocator<16>> BoneValid;
 	TArray<FVector> BoneTranslations;
 	TArray<FQuat> BoneRotations;
-	TArray<FQuat> BoneRotationDirections;
+	TArray<FQuat4f> BoneRotationDirections;
 	TArray<FVector> BoneScales;
 
-	TArray<FVector> BoneTranslationVelocities;
-	TArray<FVector> BoneRotationVelocities;
-	TArray<FVector> BoneScaleVelocities;
+	TArray<FVector3f> BoneTranslationVelocities;
+	TArray<FVector3f> BoneRotationVelocities;
+	TArray<FVector3f> BoneScaleVelocities;
 
-	TArray<FVector> BoneTranslationDecayHalfLives;
-	TArray<FVector> BoneRotationDecayHalfLives;
-	TArray<FVector> BoneScaleDecayHalfLives;
+	TArray<FVector3f> BoneTranslationDecayHalfLives;
+	TArray<FVector3f> BoneRotationDecayHalfLives;
+	TArray<FVector3f> BoneScaleDecayHalfLives;
 
 	// Recorded curve state at the point of transition.
 
@@ -228,7 +228,7 @@ private:
 		FDeadBlendingCurveElement() = default;
 	};
 
-	TBaseBlendedCurve<TInlineAllocator<32>, FDeadBlendingCurveElement> CurveData;
+	TBaseBlendedCurve<TInlineAllocator<8>, FDeadBlendingCurveElement> CurveData;
 
 private:
 
diff --git a/Engine/Source/Runtime/Engine/Private/Animation/AnimNode_DeadBlending.cpp b/Engine/Source/Runtime/Engine/Private/Animation/AnimNode_DeadBlending.cpp
--- a/Engine/Source/Runtime/Engine/Private/Animation/AnimNode_DeadBlending.cpp
+++ b/Engine/Source/Runtime/Engine/Private/Animation/AnimNode_DeadBlending.cpp
@@ -77,14 +77,22 @@ namespace UE::Anim::DeadBlending::Private
 		return RefSkeleton.GetNum();
 	}
 
-	static inline FVector VectorDivMax(const float V, const FVector W, const float Epsilon = UE_SMALL_NUMBER)
+	static inline FVector3f VectorDivMax(const float V, const FVector3f W, const float Epsilon = UE_SMALL_NUMBER)
 	{
-		return FVector(
+		return FVector3f(
 			V / FMath::Max(W.X, Epsilon),
 			V / FMath::Max(W.Y, Epsilon),
 			V / FMath::Max(W.Z, Epsilon));
 	}
 
+	static inline FVector3f VectorDivMax(const FVector3f V, const FVector3f W, const float Epsilon = UE_SMALL_NUMBER)
+	{
+		return FVector3f(
+			V.X / FMath::Max(W.X, Epsilon),
+			V.Y / FMath::Max(W.Y, Epsilon),
+			V.Z / FMath::Max(W.Z, Epsilon));
+	}
+
 	static inline FVector VectorDivMax(const FVector V, const FVector W, const float Epsilon = UE_SMALL_NUMBER)
 	{
 		return FVector(
@@ -93,9 +101,9 @@ namespace UE::Anim::DeadBlending::Private
 			V.Z / FMath::Max(W.Z, Epsilon));
 	}
 
-	static inline FVector VectorInvExpApprox(const FVector V)
+	static inline FVector3f VectorInvExpApprox(const FVector3f V)
 	{
-		return FVector(
+		return FVector3f(
 			FMath::InvExpApprox(V.X),
 			FMath::InvExpApprox(V.Y),
 			FMath::InvExpApprox(V.Z));
@@ -134,15 +142,15 @@ namespace UE::Anim::DeadBlending::Private
 
 	static inline FVector ExtrapolateTranslation(
 		const FVector Translation,
-		const FVector Velocity,
+		const FVector3f Velocity,
 		const float Time,
-		const FVector DecayHalflife,
+		const FVector3f DecayHalflife,
 		const float Epsilon = UE_SMALL_NUMBER)
 	{
 		if (Velocity.SquaredLength() > Epsilon)
 		{
-			const FVector C = VectorDivMax(Ln2, DecayHalflife, Epsilon);
-			return Translation + VectorDivMax(Velocity, C, Epsilon) * (FVector::OneVector - VectorInvExpApprox(C * Time));
+			const FVector3f C = VectorDivMax(Ln2, DecayHalflife, Epsilon);
+			return Translation + (FVector)(VectorDivMax(Velocity, C, Epsilon) * (FVector3f::OneVector - VectorInvExpApprox(C * Time)));
 		}
 		else
 		{
@@ -152,15 +160,15 @@ namespace UE::Anim::DeadBlending::Private
 
 	static inline FQuat ExtrapolateRotation(
 		const FQuat Rotation,
-		const FVector Velocity,
+		const FVector3f Velocity,
 		const float Time,
-		const FVector DecayHalflife,
+		const FVector3f DecayHalflife,
 		const float Epsilon = UE_SMALL_NUMBER)
 	{
 		if (Velocity.SquaredLength() > Epsilon)
 		{
-			const FVector C = VectorDivMax(Ln2, DecayHalflife, Epsilon);
-			return FQuat::MakeFromRotationVector(VectorDivMax(Velocity, C, Epsilon) * (FVector::OneVector - VectorInvExpApprox(C * Time))) * Rotation;
+			const FVector3f C = VectorDivMax(Ln2, DecayHalflife, Epsilon);
+			return FQuat::MakeFromRotationVector((FVector)(VectorDivMax(Velocity, C, Epsilon) * (FVector3f::OneVector - VectorInvExpApprox(C * Time)))) * Rotation;
 		}
 		else
 		{
@@ -170,15 +178,15 @@ namespace UE::Anim::DeadBlending::Private
 
 	static inline FVector ExtrapolateScale(
 		const FVector Scale,
-		const FVector Velocity,
+		const FVector3f Velocity,
 		const float Time,
-		const FVector DecayHalflife,
+		const FVector3f DecayHalflife,
 		const float Epsilon = UE_SMALL_NUMBER)
 	{
 		if (Velocity.SquaredLength() > Epsilon)
 		{
-			const FVector C = VectorDivMax(Ln2, DecayHalflife, Epsilon);
-			return VectorExp(VectorDivMax(Velocity, C, Epsilon) * (FVector::OneVector - VectorInvExpApprox(C * Time))) * Scale;
+			const FVector3f C = VectorDivMax(Ln2, DecayHalflife, Epsilon);
+			return VectorExp((FVector)(VectorDivMax(Velocity, C, Epsilon) * (FVector3f::OneVector - VectorInvExpApprox(C * Time)))) * Scale;
 		}
 		else
 		{
@@ -231,15 +239,15 @@ namespace UE::Anim::DeadBlending::Private
 		return FMath::Clamp(HalfLife * (SrcDstDiff / ClipMagnitudeToGreaterThanEpsilon(SrcVelocity, Epsilon)), HalfLifeMin, HalfLifeMax);
 	}
 
-	static inline FVector ComputeDecayHalfLifeFromDiffAndVelocity(
+	static inline FVector3f ComputeDecayHalfLifeFromDiffAndVelocity(
 		const FVector SrcDstDiff,
-		const FVector SrcVelocity,
+		const FVector3f SrcVelocity,
 		const float HalfLife,
 		const float HalfLifeMin,
 		const float HalfLifeMax,
 		const float Epsilon = UE_KINDA_SMALL_NUMBER)
 	{
-		return FVector(
+		return FVector3f(
 			ComputeDecayHalfLifeFromDiffAndVelocity(SrcDstDiff.X, SrcVelocity.X, HalfLife, HalfLifeMin, HalfLifeMax, Epsilon),
 			ComputeDecayHalfLifeFromDiffAndVelocity(SrcDstDiff.Y, SrcVelocity.Y, HalfLife, HalfLifeMin, HalfLifeMax, Epsilon),
 			ComputeDecayHalfLifeFromDiffAndVelocity(SrcDstDiff.Z, SrcVelocity.Z, HalfLife, HalfLifeMin, HalfLifeMax, Epsilon));
@@ -261,16 +269,16 @@ void FAnimNode_DeadBlending::InitFrom(
 	BoneValid.Init(false, NumSkeletonBones);
 	BoneTranslations.Init(FVector::ZeroVector, NumSkeletonBones);
 	BoneRotations.Init(FQuat::Identity, NumSkeletonBones);
-	BoneRotationDirections.Init(FQuat::Identity, NumSkeletonBones);
+	BoneRotationDirections.Init(FQuat4f::Identity, NumSkeletonBones);
 	BoneScales.Init(FVector::OneVector, NumSkeletonBones);
 
-	BoneTranslationVelocities.Init(FVector::ZeroVector, NumSkeletonBones);
-	BoneRotationVelocities.Init(FVector::ZeroVector, NumSkeletonBones);
-	BoneScaleVelocities.Init(FVector::ZeroVector, NumSkeletonBones);
+	BoneTranslationVelocities.Init(FVector3f::ZeroVector, NumSkeletonBones);
+	BoneRotationVelocities.Init(FVector3f::ZeroVector, NumSkeletonBones);
+	BoneScaleVelocities.Init(FVector3f::ZeroVector, NumSkeletonBones);
 
-	BoneTranslationDecayHalfLives.Init(ExtrapolationHalfLifeMin * FVector::OneVector, NumSkeletonBones);
-	BoneRotationDecayHalfLives.Init(ExtrapolationHalfLifeMin * FVector::OneVector, NumSkeletonBones);
-	BoneScaleDecayHalfLives.Init(ExtrapolationHalfLifeMin * FVector::OneVector, NumSkeletonBones);
+	BoneTranslationDecayHalfLives.Init(ExtrapolationHalfLifeMin * FVector3f::OneVector, NumSkeletonBones);
+	BoneRotationDecayHalfLives.Init(ExtrapolationHalfLifeMin * FVector3f::OneVector, NumSkeletonBones);
+	BoneScaleDecayHalfLives.Init(ExtrapolationHalfLifeMin * FVector3f::OneVector, NumSkeletonBones);
 
 	// Record bone state
 
@@ -310,9 +318,9 @@ void FAnimNode_DeadBlending::InitFrom(
 
 			const FVector ScaleDiff = UE::Anim::DeadBlending::Private::VectorDivMax(SrcTransformCurr.GetScale3D(), SrcTransformPrev.GetScale3D());
 
-			BoneTranslationVelocities[SkeletonPoseBoneIndex] = TranslationDiff / SrcPoseCurr.DeltaTime;
-			BoneRotationVelocities[SkeletonPoseBoneIndex] = RotationDiff.ToRotationVector() / SrcPoseCurr.DeltaTime;
-			BoneScaleVelocities[SkeletonPoseBoneIndex] = UE::Anim::DeadBlending::Private::VectorLogSafe(ScaleDiff) / SrcPoseCurr.DeltaTime;
+			BoneTranslationVelocities[SkeletonPoseBoneIndex] = (FVector3f)(TranslationDiff / SrcPoseCurr.DeltaTime);
+			BoneRotationVelocities[SkeletonPoseBoneIndex] = (FVector3f)(RotationDiff.ToRotationVector() / SrcPoseCurr.DeltaTime);
+			BoneScaleVelocities[SkeletonPoseBoneIndex] = (FVector3f)(UE::Anim::DeadBlending::Private::VectorLogSafe(ScaleDiff) / SrcPoseCurr.DeltaTime);
 
 			// Clamp Maximum Velocity
 
@@ -456,10 +464,10 @@ void FAnimNode_DeadBlending::ApplyTo(FCompactPose& InOutPose, FBlendedCurve& InO
 		// side of this rotation.
 
 		FQuat RotationDiff = ExtrapolatedRotation * InOutPose[BoneIndex].GetRotation().Inverse();
-		RotationDiff.EnforceShortestArcWith(BoneRotationDirections[SkeletonPoseBoneIndex]);
+		RotationDiff.EnforceShortestArcWith((FQuat)BoneRotationDirections[SkeletonPoseBoneIndex]);
 
 		// Update BoneRotationDirections to match our current path
-		BoneRotationDirections[SkeletonPoseBoneIndex] = RotationDiff;
+		BoneRotationDirections[SkeletonPoseBoneIndex] = (FQuat4f)RotationDiff;
 
 		// Compute Blend Alpha
 
