Shader Compiler:  Fix for asserts accessing CVar in wrong thread -- use reference variable instead.  Fix for priority crash -- need to check priority again after locking in case it was changed by another thread.

#rnx

diff --git a/Engine/Source/Runtime/Engine/Private/ShaderCompiler/ShaderCompiler.cpp b/Engine/Source/Runtime/Engine/Private/ShaderCompiler/ShaderCompiler.cpp
--- a/Engine/Source/Runtime/Engine/Private/ShaderCompiler/ShaderCompiler.cpp
+++ b/Engine/Source/Runtime/Engine/Private/ShaderCompiler/ShaderCompiler.cpp
@@ -241,9 +241,10 @@ static FAutoConsoleVariableRef CVarShaderCompilerDebugDiscardCacheOutputs(
 	ECVF_Default
 );
 
-static TAutoConsoleVariable<bool> CVarShaderCompilerParallelSubmitJobs(
+int32 GShaderCompilerParallelSubmitJobs = 0;
+static FAutoConsoleVariableRef CVarShaderCompilerParallelSubmitJobs(
 	TEXT("r.ShaderCompiler.ParallelSubmitJobs"),
-	false,
+	GShaderCompilerParallelSubmitJobs,
 	TEXT("if != 0, FShaderJobCache::SubmitJobs will run in multiple parallel tasks, instead of the game thread."),
 	ECVF_Default
 );
@@ -1252,17 +1253,24 @@ void FShaderJobCache::InternalSetPriority(FShaderCommonCompileJob* Job, EShaderC
 		// Need write lock to call UnlinkJobWithPriority
 		FWriteScopeLock Locker(JobLock);
 
-		// Job hasn't started yet, move it to the pending list for the new priority
-		const int32 PrevPriorityIndex = (int32)Job->PendingPriority;
-		check(Job->PendingPriority == Job->Priority);
-		UnlinkJobWithPriority(*Job, PrevPriorityIndex);
+		// Check priority again, as the job may have been kicked off by another thread while waiting on the lock
+		if (Job->PendingPriority != EShaderCompileJobPriority::None)
+		{
+			// Job hasn't started yet, move it to the pending list for the new priority
+			const int32 PrevPriorityIndex = (int32)Job->PendingPriority;
+			check(Job->PendingPriority == Job->Priority);
+			UnlinkJobWithPriority(*Job, PrevPriorityIndex);
 
-		ensure(!ShaderCompiler::IsJobCacheEnabled() || Job->bInputHashSet);
-		LinkJobWithPriority(*Job, PriorityIndex);
-		Job->Priority = InPriority;
-		Job->PendingPriority = InPriority;
+			ensure(!ShaderCompiler::IsJobCacheEnabled() || Job->bInputHashSet);
+			LinkJobWithPriority(*Job, PriorityIndex);
+			Job->Priority = InPriority;
+			Job->PendingPriority = InPriority;
+			
+			return;
+		}
 	}
-	else if (!Job->bFinalized &&
+
+	if (!Job->bFinalized &&
 		Job->CurrentWorker == EShaderCompilerWorkerType::Distributed &&
 		InPriority == EShaderCompileJobPriority::ForceLocal)
 	{
@@ -1464,7 +1472,7 @@ void FShaderJobCache::SubmitJobs(const TArray<FShaderCommonCompileJobPtr>& InJob
 		// we may fulfill some of the jobs from the cache (and we will be subtracting them)
 		NumOutstandingJobs.Add(InJobs.Num());
 
-		if (CVarShaderCompilerParallelSubmitJobs.GetValueOnGameThread())
+		if (GShaderCompilerParallelSubmitJobs)
 		{
 			for (FShaderCommonCompileJob* Job : InJobs)
 			{
