Added option to DeadBlending node to allow it to dynamically allocate and deallocate memory when blends become active and inactive, rather than preallocate all memory up-front.

diff --git a/Engine/Source/Runtime/Engine/Classes/Animation/AnimNode_DeadBlending.h b/Engine/Source/Runtime/Engine/Classes/Animation/AnimNode_DeadBlending.h
--- a/Engine/Source/Runtime/Engine/Classes/Animation/AnimNode_DeadBlending.h
+++ b/Engine/Source/Runtime/Engine/Classes/Animation/AnimNode_DeadBlending.h
@@ -131,6 +131,14 @@ private:
 	UPROPERTY(EditAnywhere, Category = Extrapolation, meta = (Min = "0.0", UIMin = "0.0"))
 	float MaximumCurveVelocity = 100.0f;
 
+	/**
+	 * Enable this to pre-allocate memory for the node rather than to allocate and deallocate memory when blending 
+	 * becomes active and inactive. This improves performance, but causes larger memory usage, in particular when you 
+	 * have multiple Dead Blending nodes in an animation graph that are not all used at once.
+	 */
+	UPROPERTY(EditAnywhere, Category = Memory)
+	bool bPreallocateMemory = false;
+
 #if WITH_EDITORONLY_DATA
 	
 	// This setting can be used to show what the extrapolation of the animation looks like.
@@ -160,6 +168,11 @@ public: // FAnimNode_Base
 
 private:
 	
+	/**
+	 * Deactivates the inertialization and frees any temporary memory (unless bPreallocateMemory is set).
+	 */
+	void Deactivate();
+
 	/**
 	 * Records the pose and velocity of the animation being transitioned from, and computes the extrapolation half-lives.
 	 * 
diff --git a/Engine/Source/Runtime/Engine/Private/Animation/AnimNode_DeadBlending.cpp b/Engine/Source/Runtime/Engine/Private/Animation/AnimNode_DeadBlending.cpp
--- a/Engine/Source/Runtime/Engine/Private/Animation/AnimNode_DeadBlending.cpp
+++ b/Engine/Source/Runtime/Engine/Private/Animation/AnimNode_DeadBlending.cpp
@@ -254,6 +254,30 @@ namespace UE::Anim::DeadBlending::Private
 	}
 }
 
+void FAnimNode_DeadBlending::Deactivate()
+{
+	InertializationState = EInertializationState::Inactive;
+
+	if (!bPreallocateMemory)
+	{
+		BoneValid.Empty();
+		BoneTranslations.Empty();
+		BoneRotations.Empty();
+		BoneRotationDirections.Empty();
+		BoneScales.Empty();
+
+		BoneTranslationVelocities.Empty();
+		BoneRotationVelocities.Empty();
+		BoneScaleVelocities.Empty();
+
+		BoneTranslationDecayHalfLives.Empty();
+		BoneRotationDecayHalfLives.Empty();
+		BoneScaleDecayHalfLives.Empty();
+
+		InertializationDurationPerBone.Empty();
+	}
+}
+
 void FAnimNode_DeadBlending::InitFrom(
 	const FCompactPose& InPose,
 	const FBlendedCurve& InCurves,
@@ -571,7 +595,7 @@ void FAnimNode_DeadBlending::Initialize_AnyThread(const FAnimationInitializeCont
 
 	RequestQueue.Reserve(8);
 
-	const int32 NumSkeletonBones = Context.AnimInstanceProxy->GetSkeleton()->GetReferenceSkeleton().GetNum();
+	const int32 NumSkeletonBones = bPreallocateMemory ? Context.AnimInstanceProxy->GetSkeleton()->GetReferenceSkeleton().GetNum() : 0;
 
 	BoneValid.Empty(NumSkeletonBones);
 	BoneTranslations.Empty(NumSkeletonBones);
@@ -700,7 +724,7 @@ void FAnimNode_DeadBlending::Evaluate_AnyThread(FPoseContext& Output)
 
 	if (bTeleported)
 	{
-		InertializationState = EInertializationState::Inactive;
+		Deactivate();
 	}
 	
 	// If we don't have any Pose Snapshots recorded it means this is the first time this node has been evaluated in 
@@ -821,7 +845,7 @@ void FAnimNode_DeadBlending::Evaluate_AnyThread(FPoseContext& Output)
 
 		if (InertializationTime >= InertializationMaxDuration)
 		{
-			InertializationState = EInertializationState::Inactive;
+			Deactivate();
 		}
 	}
 
