Speculative fix for reachable garbage

diff --git a/Engine/Source/Runtime/Engine/Private/Components/ChildActorComponent.cpp b/Engine/Source/Runtime/Engine/Private/Components/ChildActorComponent.cpp
--- a/Engine/Source/Runtime/Engine/Private/Components/ChildActorComponent.cpp
+++ b/Engine/Source/Runtime/Engine/Private/Components/ChildActorComponent.cpp
@@ -932,6 +932,10 @@ void UChildActorComponent::DestroyChildActor()
 
 		ChildActor = nullptr;
 	}
+	else if (!IsValid(ChildActor))
+	{
+		ChildActor = nullptr;
+	}
 }
 
 void UChildActorComponent::BeginPlay()
