Fix crash from calling GetParticle on destroyed FPhysicsObject.
Add a callback to SimCallbackObject for when a physics proxy owning FPhysicsObjects gets destroyed.

diff --git a/Engine/Source/Runtime/Engine/Private/PhysicsEngine/PhysicsReplication.cpp b/Engine/Source/Runtime/Engine/Private/PhysicsEngine/PhysicsReplication.cpp
--- a/Engine/Source/Runtime/Engine/Private/PhysicsEngine/PhysicsReplication.cpp
+++ b/Engine/Source/Runtime/Engine/Private/PhysicsEngine/PhysicsReplication.cpp
@@ -122,7 +122,7 @@ namespace PhysicsReplicationCVars
 		static bool bVelocityBased = true;
 		static FAutoConsoleVariableRef CVarVelocityBased(TEXT("np2.PredictiveInterpolation.VelocityBased"), bVelocityBased, TEXT("When true, predictive interpolation replication mode will only apply linear velocity and angular velocity"));
 	
-		static bool bAlwaysHardSnap = true;
+		static bool bAlwaysHardSnap = false;
 		static FAutoConsoleVariableRef CVarAlwaysHardSnap(TEXT("np2.PredictiveInterpolation.AlwaysHardSnap"), bAlwaysHardSnap, TEXT("When true, predictive interpolation replication mode will always hard snap. Used as a backup measure"));
 
 		static bool bSkipReplication = false;
@@ -207,6 +207,11 @@ void FPhysicsReplication::SetReplicatedTarget(UPrimitiveComponent* Component, FN
 
 void FPhysicsReplication::SetReplicatedTarget(Chaos::FConstPhysicsObjectHandle PhysicsObject, const FRigidBodyState& ReplicatedTarget, int32 ServerFrame, EPhysicsReplicationMode ReplicationMode)
 {
+	if (!PhysicsObject)
+	{
+		return;
+	}
+
 	UWorld* OwningWorld = GetOwningWorld();
 	if (OwningWorld == nullptr)
 	{
@@ -658,6 +663,11 @@ void FPhysicsReplication::PrepareAsyncData_External(const FRigidBodyErrorCorrect
 }
 
 #pragma region AsyncPhysicsReplication
+void FPhysicsReplicationAsync::OnPhysicsObjectUnregistered_Internal(Chaos::FConstPhysicsObjectHandle PhysicsObject)
+{
+	ObjectToTarget.Remove(PhysicsObject);
+}
+
 void FPhysicsReplicationAsync::OnPreSimulate_Internal()
 {
 	if (const FPhysicsReplicationAsyncInput* AsyncInput = GetConsumerInput_Internal())
@@ -742,7 +752,7 @@ void FPhysicsReplicationAsync::UpdateAsyncTarget(const FPhysicsRepAsyncInputData
 		// First time we add a target, set it's previous and correction
 		// positions to the target position to avoid math with uninitialized
 		// memory.
-		Target = &ObjectToTarget.Add(Input.PhysicsObject, FReplicatedPhysicsTargetAsync(Input.PhysicsObject));
+		Target = &ObjectToTarget.Add(Input.PhysicsObject, FReplicatedPhysicsTargetAsync());
 		Target->PrevPos = Input.TargetState.Position;
 		Target->PrevPosTarget = Input.TargetState.Position;
 		Target->PrevRotTarget = Input.TargetState.Quaternion;
@@ -795,39 +805,37 @@ void FPhysicsReplicationAsync::ApplyTargetStatesAsync(const float DeltaSeconds,
 	{
 		bool bRemoveItr = true; // Remove current cached replication target unless replication logic tells us to store it for next tick
 
-		FReplicatedPhysicsTargetAsync& Target = Itr.Value();
-
-		if (Target.PhysicsObject != nullptr)
+		Chaos::FConstPhysicsObjectHandle POHandle = Itr.Key();
+		Chaos::FWritePhysicsObjectInterface_Internal Interface = Chaos::FPhysicsObjectInternalInterface::GetWrite();
+		if (FGeometryParticleHandle* Handle = Interface.GetParticle(POHandle))
 		{
-			Chaos::FWritePhysicsObjectInterface_Internal Interface = Chaos::FPhysicsObjectInternalInterface::GetWrite();
-			if (FGeometryParticleHandle* Handle = Interface.GetParticle(Target.PhysicsObject))
+			FReplicatedPhysicsTargetAsync& Target = Itr.Value();
+
+			if (FPBDRigidParticleHandle* RigidHandle = Handle->CastToRigidParticle())
 			{
-				if (FPBDRigidParticleHandle* RigidHandle = Handle->CastToRigidParticle())
-				{
-					EPhysicsReplicationMode RepMode = Target.RepMode;
+				EPhysicsReplicationMode RepMode = Target.RepMode;
 
-					// TODO, Remove the resim option from project settings, we only need the physics prediction one now
-					if (!Chaos::FPBDRigidsSolver::IsPhysicsResimulationEnabled() && RepMode == EPhysicsReplicationMode::Resimulation)
-					{
-						RepMode = EPhysicsReplicationMode::Default;
-					}
+				// TODO, Remove the resim option from project settings, we only need the physics prediction one now
+				if (!Chaos::FPBDRigidsSolver::IsPhysicsResimulationEnabled() && RepMode == EPhysicsReplicationMode::Resimulation)
+				{
+					RepMode = EPhysicsReplicationMode::Default;
+				}
 
-					switch (RepMode)
-					{
-						case EPhysicsReplicationMode::Default:
-							bRemoveItr = DefaultReplication(RigidHandle, Target, DeltaSeconds);
-							break;
+				switch (RepMode)
+				{
+					case EPhysicsReplicationMode::Default:
+						bRemoveItr = DefaultReplication(RigidHandle, Target, DeltaSeconds);
+						break;
 
-						case EPhysicsReplicationMode::PredictiveInterpolation:
-							bRemoveItr = PredictiveInterpolation(RigidHandle, Target, DeltaSeconds);
-							break;
+					case EPhysicsReplicationMode::PredictiveInterpolation:
+						bRemoveItr = PredictiveInterpolation(RigidHandle, Target, DeltaSeconds);
+						break;
 
-						case EPhysicsReplicationMode::Resimulation:
-							bRemoveItr = ResimulationReplication(RigidHandle, Target, DeltaSeconds);
-							break;
-					}
-					Target.TickCount++;
+					case EPhysicsReplicationMode::Resimulation:
+						bRemoveItr = ResimulationReplication(RigidHandle, Target, DeltaSeconds);
+						break;
 				}
+				Target.TickCount++;
 			}
 		}
 
@@ -1289,24 +1297,22 @@ bool FPhysicsReplicationAsync::PredictiveInterpolation(Chaos::FPBDRigidParticleH
 
 		// Cache data for next replication
 		Target.PrevLinVel = FVector(RepLinVel);
-	}
+		Target.PrevPos = FVector(CurrentState.Position);
 
-	// Cache data for next replication
-	Target.PrevPos = FVector(CurrentState.Position);
-
-	// --- Target Extrapolation ---
-	if (Target.TickCount < FMath::CeilToInt(ReceiveInterval * PhysicsReplicationCVars::PredictiveInterpolationCVars::ExtrapolationTimeMultiplier))
-	{
-		// Extrapolate target position
-		Target.TargetState.Position = Target.TargetState.Position + Target.TargetState.LinVel * DeltaSeconds;
-
-		// Extrapolate target rotation
-		float TargetAngVelSize;
-		FVector TargetAngVelAxis;
-		Target.TargetState.AngVel.FVector::ToDirectionAndLength(TargetAngVelAxis, TargetAngVelSize);
-		TargetAngVelSize = FMath::DegreesToRadians(TargetAngVelSize);
-		const FQuat TargetRotExtrapDelta = FQuat(TargetAngVelAxis, TargetAngVelSize * DeltaSeconds);
-		Target.TargetState.Quaternion = TargetRotExtrapDelta * Target.TargetState.Quaternion;
+		// --- Target Extrapolation ---
+		if (Target.TickCount < FMath::CeilToInt(ReceiveInterval * PhysicsReplicationCVars::PredictiveInterpolationCVars::ExtrapolationTimeMultiplier))
+		{
+			// Extrapolate target position
+			Target.TargetState.Position = Target.TargetState.Position + Target.TargetState.LinVel * DeltaSeconds;
+
+			// Extrapolate target rotation
+			float TargetAngVelSize;
+			FVector TargetAngVelAxis;
+			Target.TargetState.AngVel.FVector::ToDirectionAndLength(TargetAngVelAxis, TargetAngVelSize);
+			TargetAngVelSize = FMath::DegreesToRadians(TargetAngVelSize);
+			const FQuat TargetRotExtrapDelta = FQuat(TargetAngVelAxis, TargetAngVelSize * DeltaSeconds);
+			Target.TargetState.Quaternion = TargetRotExtrapDelta * Target.TargetState.Quaternion;
+		}
 	}
 	
 	return false;
diff --git a/Engine/Source/Runtime/Engine/Public/PhysicsReplication.h b/Engine/Source/Runtime/Engine/Public/PhysicsReplication.h
--- a/Engine/Source/Runtime/Engine/Public/PhysicsReplication.h
+++ b/Engine/Source/Runtime/Engine/Public/PhysicsReplication.h
@@ -90,11 +90,10 @@ struct FPhysicsReplicationAsyncInput : public Chaos::FSimCallbackInput
 
 struct FReplicatedPhysicsTargetAsync
 {
-	FReplicatedPhysicsTargetAsync(Chaos::FConstPhysicsObjectHandle POHandle)
+	FReplicatedPhysicsTargetAsync()
 		: AccumulatedErrorSeconds(0.0f)
 		, ServerFrame(INDEX_NONE)
 		, ReceiveFrame(INDEX_NONE)
-		, PhysicsObject(POHandle)
 	{ }
 
 	/** The target state replicated by server */
@@ -115,9 +114,6 @@ struct FReplicatedPhysicsTargetAsync
 	/** The local client frame when receiving this target from the server */
 	int32 ReceiveFrame;
 
-	/** Index of physics object on component */
-	Chaos::FConstPhysicsObjectHandle PhysicsObject;
-
 	/** The replication mode this PhysicsObject should use */
 	EPhysicsReplicationMode RepMode;
 
@@ -136,10 +132,12 @@ struct FReplicatedPhysicsTargetAsync
 class FPhysicsReplicationAsync : public Chaos::TSimCallbackObject<
 	FPhysicsReplicationAsyncInput,
 	Chaos::FSimCallbackNoOutput,
-	Chaos::ESimCallbackOptions::Presimulate>
+	Chaos::ESimCallbackOptions::Presimulate | Chaos::ESimCallbackOptions::PhysicsObjectUnregister>
 {
 	virtual FName GetFNameForStatId() const override;
 	virtual void OnPreSimulate_Internal() override;
+	virtual void OnPhysicsObjectUnregistered_Internal(Chaos::FConstPhysicsObjectHandle PhysicsObject) override;
+
 	virtual void ApplyTargetStatesAsync(const float DeltaSeconds, const FPhysicsRepErrorCorrectionData& ErrorCorrection, const TArray<FPhysicsRepAsyncInputData>& TargetStates);
 
 	// Replication functions
diff --git a/Engine/Source/Runtime/Experimental/Chaos/Private/PBDRigidsSolver.cpp b/Engine/Source/Runtime/Experimental/Chaos/Private/PBDRigidsSolver.cpp
--- a/Engine/Source/Runtime/Experimental/Chaos/Private/PBDRigidsSolver.cpp
+++ b/Engine/Source/Runtime/Experimental/Chaos/Private/PBDRigidsSolver.cpp
@@ -1061,6 +1061,15 @@ namespace Chaos
 				continue;
 			}
 
+			// Callback for Unregister PhysicsObject
+			if (PhysicsObjectUnregistrationWatchers.Num() > 0)
+			{
+				for (ISimCallbackObject* Callback : PhysicsObjectUnregistrationWatchers)
+				{
+					Callback->OnPhysicsObjectUnregistered_Internal(Proxy->GetPhysicsObjectHandle());
+				}
+			}
+
 			if (FPBDRigidsEvolutionGBF* Evolution = GetEvolution())
 			{
 				// Remove the cluster union this proxy manages which will in turn also the destroy the necessary particles.
@@ -1092,6 +1101,15 @@ namespace Chaos
 			FPendingDestroyInfo& Info = PendingDestroyPhysicsProxy[Idx];
 			if(Info.DestroyOnStep <= GetCurrentFrame() || IsShuttingDown())
 			{
+				// Callback for Unregister PhysicsObject
+				if (PhysicsObjectUnregistrationWatchers.Num() > 0)
+				{
+					for (ISimCallbackObject* Callback : PhysicsObjectUnregistrationWatchers)
+					{
+						Callback->OnPhysicsObjectUnregistered_Internal(Info.Proxy->GetPhysicsObject());
+					}
+				}
+
 				// finally let's release the unique index
 				GetEvolution()->ReleaseUniqueIdx(Info.UniqueIdx);
 
@@ -1114,6 +1132,19 @@ namespace Chaos
 			GetEvolution()->GetRigidClustering().CleanupInternalClustersForProxies(TArrayView<IPhysicsProxyBase*>{reinterpret_cast<IPhysicsProxyBase**>(&PendingDestroyGeometryCollectionPhysicsProxy[0]), PendingDestroyGeometryCollectionPhysicsProxy.Num() });
 			for (auto Proxy : PendingDestroyGeometryCollectionPhysicsProxy)
 			{
+				// Callback for Unregister PhysicsObject
+				if (PhysicsObjectUnregistrationWatchers.Num() > 0)
+				{
+					TArray<FPhysicsObjectHandle> PhysicsObjects = Proxy->GetAllPhysicsObjects();
+					for (FConstPhysicsObjectHandle PhysicsObject : PhysicsObjects)
+					{
+						for (ISimCallbackObject* Callback : PhysicsObjectUnregistrationWatchers)
+						{
+							Callback->OnPhysicsObjectUnregistered_Internal(PhysicsObject);
+						}
+					}
+				}
+
 				// Removing the geometry collection from the solver a bit delayed. This lets the cluster union do its cleanup first before
 				// the geometry collection if they're all being destroyed at the same time.
 				Proxy->OnRemoveFromSolver(this);
@@ -1633,6 +1664,11 @@ namespace Chaos
 					MRewindCallback->RegisterRewindableSimCallback_Internal(SimCallbackObject);
 				}
 			}
+
+			if (SimCallbackObject->HasOption(ESimCallbackOptions::PhysicsObjectUnregister))
+			{
+				PhysicsObjectUnregistrationWatchers.Add(SimCallbackObject);
+			}
 		}
 
 		//save any pending data for this particular interval
@@ -1709,6 +1745,15 @@ namespace Chaos
 			}
 		}
 
+		for (int32 Idx = PhysicsObjectUnregistrationWatchers.Num() - 1; Idx >= 0; --Idx)
+		{
+			if (PhysicsObjectUnregistrationWatchers[Idx]->bPendingDelete)
+			{
+				//will also be in SimCallbackObjects so we'll delete it in that loop
+				PhysicsObjectUnregistrationWatchers.RemoveAtSwap(Idx, 1, false);
+			}
+		}
+
 		ProcessSinglePushedData_Internal(PushData);
 
 		//run any commands passed in. These don't generate outputs and are a one off so just do them here
diff --git a/Engine/Source/Runtime/Experimental/Chaos/Public/Chaos/Framework/PhysicsSolverBase.h b/Engine/Source/Runtime/Experimental/Chaos/Public/Chaos/Framework/PhysicsSolverBase.h
--- a/Engine/Source/Runtime/Experimental/Chaos/Public/Chaos/Framework/PhysicsSolverBase.h
+++ b/Engine/Source/Runtime/Experimental/Chaos/Public/Chaos/Framework/PhysicsSolverBase.h
@@ -728,6 +728,7 @@ namespace Chaos
 	TArray<ISimCallbackObject*> ContactModifiers;
 	TArray<ISimCallbackObject*> RegistrationWatchers;
 	TArray<ISimCallbackObject*> UnregistrationWatchers;
+	TArray<ISimCallbackObject*> PhysicsObjectUnregistrationWatchers;
 
 	TUniquePtr<FRewindData> MRewindData;
 	TUniquePtr<IRewindCallback> MRewindCallback;
diff --git a/Engine/Source/Runtime/Experimental/Chaos/Public/Chaos/SimCallbackObject.h b/Engine/Source/Runtime/Experimental/Chaos/Public/Chaos/SimCallbackObject.h
--- a/Engine/Source/Runtime/Experimental/Chaos/Public/Chaos/SimCallbackObject.h
+++ b/Engine/Source/Runtime/Experimental/Chaos/Public/Chaos/SimCallbackObject.h
@@ -7,6 +7,7 @@
 #include "Chaos/CollisionResolutionTypes.h"
 #include "Chaos/GeometryParticlesfwd.h"
 #include "ChaosStats.h"
+#include "Chaos/PhysicsObject.h"
 
 // Enable through build or just here in code to cause an untracked callback to fail to compile
 // Untracked callbacks will show up a such in profiling sessions and this can help track them down
@@ -38,7 +39,8 @@ enum class ESimCallbackOptions : uint16
 	ParticleRegister		= 1 << 6,
 	ParticleUnregister		= 1 << 7,
 	RunOnFrozenGameThread	= 1 << 8,
-	Rewind					= 1 << 9
+	Rewind					= 1 << 9,
+	PhysicsObjectUnregister	= 1 << 10
 };
 ENUM_CLASS_FLAGS(ESimCallbackOptions)
 
@@ -281,6 +283,12 @@ private:
 		check(false);
 	}
 
+	/** Called when physics proxies owning FPhysicsObjects unregister */
+	virtual void OnPhysicsObjectUnregistered_Internal(FConstPhysicsObjectHandle PhysicsObject)
+	{
+		check(false);
+	}
+
 	/** If we've already allocated an output in this simulation step, use it again. This way multiple sim callbacks only generate one output. Otherwise allocate an output for us and mark it as pending for external thread */
 	virtual void OnFinalizeOutputData_Internal(FSimCallbackOutput* CurOutput)
 	{
