Modified USkeletalMeshComponent::ComponentIsTouchingSelectionFrustum to mirror the behavior in USkeletalMeshComponent::ComponentIsTouchingSelectionBox to consider the bMustEncompassentireComponent arg when determining if a skeletal mesh is in the given frustum.

PR #10460

diff --git a/Engine/Source/Runtime/Engine/Private/Components/SkeletalMeshComponent.cpp b/Engine/Source/Runtime/Engine/Private/Components/SkeletalMeshComponent.cpp
--- a/Engine/Source/Runtime/Engine/Private/Components/SkeletalMeshComponent.cpp
+++ b/Engine/Source/Runtime/Engine/Private/Components/SkeletalMeshComponent.cpp
@@ -3820,8 +3820,8 @@ bool USkeletalMeshComponent::ComponentIsTouchingSelectionBox(const FBox& InSelBB
 					}
 				}	
 
-				// If the selection box has to encompass all of the component and none of the component's verts failed the intersection test, this component
-				// is consider touching
+				// if bMustEncompassEntireComponent == false, we require at least one intersection, so if we are here, no intersection happened
+				// if bMustEncompassEntireComponent == true, all of them intersected, else the for loop would have returned false already
 				if (bMustEncompassEntireComponent)
 				{
 					return true;
@@ -3867,9 +3867,12 @@ bool USkeletalMeshComponent::ComponentIsTouchingSelectionFrustum(const FConvexVo
 					}
 				}
 
-				// If the selection box has to encompass all of the component and none of the component's verts failed the intersection test, this component
-				// is consider touching
-				return true;
+				// if bMustEncompassEntireComponent == false, we require at least one intersection, so if we are here, no intersection happened
+				// if bMustEncompassEntireComponent == true, all of them intersected, else the for loop would have returned false already
+				if (bMustEncompassEntireComponent)
+				{
+					return true;
+				}
 			}
 		}
 	}
