World Partition: added support for editor and server only data layers activation.
This can be used for data layers adding details on high end platforms.
This will be used for server-only content bundles once EDL is in.

diff --git a/Engine/Source/Runtime/Engine/Private/WorldPartition/DataLayer/DataLayerAsset.cpp b/Engine/Source/Runtime/Engine/Private/WorldPartition/DataLayer/DataLayerAsset.cpp
--- a/Engine/Source/Runtime/Engine/Private/WorldPartition/DataLayer/DataLayerAsset.cpp
+++ b/Engine/Source/Runtime/Engine/Private/WorldPartition/DataLayer/DataLayerAsset.cpp
@@ -59,6 +59,11 @@ bool UDataLayerAsset::CanEditChange(const FProperty* InProperty) const
 			// UDataLayerAsset outered to a UDataLayerInstance do not support Runtime type
 			return !IsPrivate();
 		}
+		else if (InProperty->GetFName() == GET_MEMBER_NAME_CHECKED(UDataLayerAsset, LoadFilter))
+		{
+			// Only runtime data layer assets can be set to client only
+			return IsRuntime();
+		}
 	}
 
 	return true;
diff --git a/Engine/Source/Runtime/Engine/Private/WorldPartition/DataLayer/DataLayerInstance.cpp b/Engine/Source/Runtime/Engine/Private/WorldPartition/DataLayer/DataLayerInstance.cpp
--- a/Engine/Source/Runtime/Engine/Private/WorldPartition/DataLayer/DataLayerInstance.cpp
+++ b/Engine/Source/Runtime/Engine/Private/WorldPartition/DataLayer/DataLayerInstance.cpp
@@ -24,9 +24,7 @@ UDataLayerInstance::UDataLayerInstance(const FObjectInitializer& ObjectInitializ
 	, bIsLocked(false)
 #endif
 	, InitialRuntimeState(EDataLayerRuntimeState::Unloaded)
-{
-
-}
+{}
 
 void UDataLayerInstance::PostLoad()
 {
@@ -246,6 +244,11 @@ bool UDataLayerInstance::IsParentDataLayerTypeCompatible(const UDataLayerInstanc
 		return false;
 	}
 
+	if (IsClientOnly() || IsServerOnly())
+	{
+		return false;
+	}
+
 	EDataLayerType ParentDataLayerType = InParent->GetType();
 
 	return GetType() != EDataLayerType::Unknown
@@ -415,15 +418,4 @@ EDataLayerRuntimeState UDataLayerInstance::GetEffectiveRuntimeState() const
 	return GetOuterWorldDataLayers()->GetDataLayerEffectiveRuntimeStateByName(GetDataLayerFName());
 }
 
-bool UDataLayerInstance::SetRuntimeState(EDataLayerRuntimeState InState, bool bInIsRecursive) const
-{
-	if (GetOuterWorldDataLayers()->HasAuthority())
-	{
-		GetOuterWorldDataLayers()->SetDataLayerRuntimeState(this, InState, bInIsRecursive);
-		return true;
-	}
-	UE_LOG(LogWorldPartition, Error, TEXT("SetDataLayerRuntimeState can only execute on authority"));
-	return false;
-}
-
 #undef LOCTEXT_NAMESPACE
diff --git a/Engine/Source/Runtime/Engine/Private/WorldPartition/DataLayer/DataLayerManager.cpp b/Engine/Source/Runtime/Engine/Private/WorldPartition/DataLayer/DataLayerManager.cpp
--- a/Engine/Source/Runtime/Engine/Private/WorldPartition/DataLayer/DataLayerManager.cpp
+++ b/Engine/Source/Runtime/Engine/Private/WorldPartition/DataLayer/DataLayerManager.cpp
@@ -70,7 +70,7 @@ FAutoConsoleCommand UDataLayerManager::ToggleDataLayerActivation(
 		for (const FWorldContext& Context : GEngine->GetWorldContexts())
 		{
 			UWorld* World = Context.World();
-			if (World && World->IsGameWorld() && !World->IsNetMode(NM_Client))
+			if (World && World->IsGameWorld())
 			{
 				if (UWorldPartitionSubsystem* WorldPartitionSubsystem = World->GetSubsystem<UWorldPartitionSubsystem>())
 				{
@@ -116,7 +116,7 @@ FAutoConsoleCommand UDataLayerManager::SetDataLayerRuntimeStateCommand(
 		for (const FWorldContext& Context : GEngine->GetWorldContexts())
 		{
 			UWorld* World = Context.World();
-			if (World && World->IsGameWorld() && !World->IsNetMode(NM_Client))
+			if (World && World->IsGameWorld())
 			{
 				if (UWorldPartitionSubsystem* WorldPartitionSubsystem = World->GetSubsystem<UWorldPartitionSubsystem>())
 				{
@@ -268,8 +268,13 @@ const UDataLayerInstance* UDataLayerManager::GetDataLayerInstanceFromAssetName(c
 
 bool UDataLayerManager::SetDataLayerInstanceRuntimeState(const UDataLayerInstance* InDataLayerInstance, EDataLayerRuntimeState InState, bool bInIsRecursive)
 {
-	UE_CLOG(!InDataLayerInstance, LogWorldPartition, Error, TEXT("Invalid Data Layer Instance."));
-	return InDataLayerInstance ? InDataLayerInstance->SetRuntimeState(InState, bInIsRecursive) : false;
+	if (InDataLayerInstance)
+	{
+		InDataLayerInstance->GetOuterWorldDataLayers()->SetDataLayerRuntimeState(InDataLayerInstance, InState, bInIsRecursive);
+		return true;
+	}
+	UE_LOG(LogWorldPartition, Error, TEXT("Invalid Data Layer Instance."));
+	return false;
 }
 
 bool UDataLayerManager::SetDataLayerRuntimeState(const UDataLayerAsset* InDataLayerAsset, EDataLayerRuntimeState InState, bool bInIsRecursive)
@@ -296,14 +301,22 @@ void UDataLayerManager::BroadcastOnDataLayerInstanceRuntimeStateChanged(const UD
 
 EDataLayerRuntimeState UDataLayerManager::GetDataLayerInstanceRuntimeState(const UDataLayerInstance* InDataLayerInstance) const
 {
-	UE_CLOG(!InDataLayerInstance, LogWorldPartition, Error, TEXT("Invalid Data Layer Instance."));
-	return InDataLayerInstance ? InDataLayerInstance->GetRuntimeState() : EDataLayerRuntimeState::Unloaded;
+	if (InDataLayerInstance)
+	{
+		return InDataLayerInstance->GetRuntimeState();
+	}
+	UE_LOG(LogWorldPartition, Error, TEXT("Invalid Data Layer Instance."));
+	return EDataLayerRuntimeState::Unloaded;
 }
 
 EDataLayerRuntimeState UDataLayerManager::GetDataLayerInstanceEffectiveRuntimeState(const UDataLayerInstance* InDataLayerInstance) const
 {
-	UE_CLOG(!InDataLayerInstance, LogWorldPartition, Error, TEXT("Invalid Data Layer Instance."));
-	return InDataLayerInstance ? InDataLayerInstance->GetEffectiveRuntimeState() : EDataLayerRuntimeState::Unloaded;
+	if (InDataLayerInstance)
+	{
+		return InDataLayerInstance->GetEffectiveRuntimeState();
+	}
+	UE_LOG(LogWorldPartition, Error, TEXT("Invalid Data Layer Instance."));
+	return EDataLayerRuntimeState::Unloaded;
 }
 
 const TSet<FName>& UDataLayerManager::GetEffectiveActiveDataLayerNames() const
@@ -592,7 +605,6 @@ void UDataLayerManager::ForEachDataLayerInstance(TFunctionRef<bool(UDataLayerIns
  */
 
 #if WITH_EDITOR
-
 TSubclassOf<UDataLayerLoadingPolicy> UDataLayerManager::GetDataLayerLoadingPolicyClass() const
 {
 	// Use UDataLayerManager::DataLayerLoadingPolicyClass
@@ -837,5 +849,4 @@ void FDataLayersEditorBroadcast::StaticOnActorDataLayersEditorLoadingStateChange
 	Get().DataLayerEditorLoadingStateChanged.Broadcast(bIsFromUserChange);
 	IWorldPartitionActorLoaderInterface::RefreshLoadedState(bIsFromUserChange);
 }
-
 #endif
\ No newline at end of file
diff --git a/Engine/Source/Runtime/Engine/Private/WorldPartition/DataLayer/WorldDataLayers.cpp b/Engine/Source/Runtime/Engine/Private/WorldPartition/DataLayer/WorldDataLayers.cpp
--- a/Engine/Source/Runtime/Engine/Private/WorldPartition/DataLayer/WorldDataLayers.cpp
+++ b/Engine/Source/Runtime/Engine/Private/WorldPartition/DataLayer/WorldDataLayers.cpp
@@ -76,20 +76,16 @@ void AWorldDataLayers::BeginPlay()
 {
 	Super::BeginPlay();
 
-	// When running a Replay we want to reset our state to the initial state and rely on the Replay/Replication.
-	// Unfortunately this can't be tested in the PostLoad as the World doesn't have a demo driver yet.
-	if (GetWorld()->IsPlayingReplay())
-	{
-		ResetDataLayerRuntimeStates();
-		InitializeDataLayerRuntimeStates();
-	}
+	// This must be done after the NetDriver has been set
+	ResetDataLayerRuntimeStates();
+	InitializeDataLayerRuntimeStates();
 }
 
 void AWorldDataLayers::RewindForReplay()
 {
 	Super::RewindForReplay();
 
-	// Same as BeginPlay when rewinding we want to reset our state to the initial state and rely on the Replay/Replication.
+	// When rewinding we want to reset our state to the initial state and rely on the Replay/Replication.
 	ResetDataLayerRuntimeStates();
 	InitializeDataLayerRuntimeStates();
 }
@@ -106,37 +102,12 @@ void AWorldDataLayers::InitializeDataLayerRuntimeStates()
 			return;
 		}
 #endif
-		ForEachDataLayer([this](class UDataLayerInstance* DataLayer)
-		{
-			if (DataLayer && DataLayer->IsRuntime())
-			{
-				if (DataLayer->GetInitialRuntimeState() == EDataLayerRuntimeState::Activated)
-				{
-					ActiveDataLayerNames.Add(DataLayer->GetDataLayerFName());
-				}
-				else if (DataLayer->GetInitialRuntimeState() == EDataLayerRuntimeState::Loaded)
-				{
-					LoadedDataLayerNames.Add(DataLayer->GetDataLayerFName());
-				}
-			}
-			return true;
-		});
-
-		RepActiveDataLayerNames = ActiveDataLayerNames.Array();
-		RepLoadedDataLayerNames = LoadedDataLayerNames.Array();
-
-		ForEachDataLayer([this](class UDataLayerInstance* DataLayer)
+		ForEachDataLayer([this](class UDataLayerInstance* DataLayerInstance)
 		{
-			if (DataLayer && DataLayer->IsRuntime())
-			{
-				ResolveEffectiveRuntimeState(DataLayer, /*bNotifyChange*/false);
-			}
+			SetDataLayerRuntimeState(DataLayerInstance, DataLayerInstance->GetInitialRuntimeState());
 			return true;
 		});
 
-		RepEffectiveActiveDataLayerNames = EffectiveActiveDataLayerNames.Array();
-		RepEffectiveLoadedDataLayerNames = EffectiveLoadedDataLayerNames.Array();
-
 		UE_CLOG(RepEffectiveActiveDataLayerNames.Num() || RepEffectiveLoadedDataLayerNames.Num(), LogWorldPartition, Log, TEXT("Initial Data Layer Effective States Activated(%s) Loaded(%s)"), *JoinDataLayerShortNamesFromInstanceNames(this, RepEffectiveActiveDataLayerNames), *JoinDataLayerShortNamesFromInstanceNames(this, RepEffectiveLoadedDataLayerNames));
 	}
 }
@@ -147,81 +118,110 @@ void AWorldDataLayers::ResetDataLayerRuntimeStates()
 	LoadedDataLayerNames.Reset();
 	RepActiveDataLayerNames.Reset();
 	RepLoadedDataLayerNames.Reset();
+	LocalActiveDataLayerNames.Reset();
+	LocalLoadedDataLayerNames.Reset();
 
 	EffectiveActiveDataLayerNames.Reset();
 	EffectiveLoadedDataLayerNames.Reset();
 	RepEffectiveActiveDataLayerNames.Reset();
 	RepEffectiveLoadedDataLayerNames.Reset();
+	LocalEffectiveActiveDataLayerNames.Reset();
+	LocalEffectiveLoadedDataLayerNames.Reset();
 }
 
 void AWorldDataLayers::SetDataLayerRuntimeState(const UDataLayerInstance* InDataLayerInstance, EDataLayerRuntimeState InState, bool bInIsRecursive)
 {
-	if (ensure(GetLocalRole() == ROLE_Authority))
+	if (!InDataLayerInstance || !InDataLayerInstance->IsRuntime())
 	{
-		if (!InDataLayerInstance || !InDataLayerInstance->IsRuntime())
-		{
-			return;
-		}
+		return;
+	}
+
+	const ENetMode NetMode = GetNetMode();
+	const bool bDataLayerClientOnly = InDataLayerInstance->IsClientOnly();
+	const bool bDataLayerServerOnly = InDataLayerInstance->IsServerOnly();
 
-		EDataLayerRuntimeState CurrentState = GetDataLayerRuntimeStateByName(InDataLayerInstance->GetDataLayerFName());
-		if (CurrentState != InState)
+	// Filter out irrelevant requests for the current net mode
+	if ((bDataLayerClientOnly && NetMode != NM_Client && NetMode != NM_Standalone) || (bDataLayerServerOnly && NetMode == NM_Client))
+	{
+		return;
+	}
+
+	EDataLayerRuntimeState CurrentState = GetDataLayerRuntimeStateByName(InDataLayerInstance->GetDataLayerFName());
+	if (CurrentState != InState)
+	{
+		if (GetWorld()->IsGameWorld())
 		{
 			PRAGMA_DISABLE_DEPRECATION_WARNINGS
-				if (GetWorld()->IsGameWorld())
+			const FName DataLayerShortName(InDataLayerInstance->GetDataLayerShortName());
+			if (DataLayersFilterDelegate.IsBound())
+			{
+				if (!DataLayersFilterDelegate.Execute(DataLayerShortName, CurrentState, InState))
 				{
-					FName DataLayerShortName(InDataLayerInstance->GetDataLayerShortName());
-					if (DataLayersFilterDelegate.IsBound())
-					{
-						if (!DataLayersFilterDelegate.Execute(DataLayerShortName, CurrentState, InState))
-						{
-							UE_LOG(LogWorldPartition, Log, TEXT("Data Layer '%s' was filtered out: %s -> %s"),
-								*DataLayerShortName.ToString(),
-								*StaticEnum<EDataLayerRuntimeState>()->GetDisplayNameTextByValue((int64)CurrentState).ToString(),
-								*StaticEnum<EDataLayerRuntimeState>()->GetDisplayNameTextByValue((int64)InState).ToString());
-							return;
-						}
-					}
+					UE_LOG(LogWorldPartition, Log, TEXT("Data Layer '%s' was filtered out: %s -> %s"),
+						*DataLayerShortName.ToString(),
+						*StaticEnum<EDataLayerRuntimeState>()->GetDisplayNameTextByValue((int64)CurrentState).ToString(),
+						*StaticEnum<EDataLayerRuntimeState>()->GetDisplayNameTextByValue((int64)InState).ToString());
+					return;
 				}
+			}
 			PRAGMA_ENABLE_DEPRECATION_WARNINGS
+		}
 
-			LoadedDataLayerNames.Remove(InDataLayerInstance->GetDataLayerFName());
-			ActiveDataLayerNames.Remove(InDataLayerInstance->GetDataLayerFName());
+		LoadedDataLayerNames.Remove(InDataLayerInstance->GetDataLayerFName());
+		ActiveDataLayerNames.Remove(InDataLayerInstance->GetDataLayerFName());
+		LocalLoadedDataLayerNames.Remove(InDataLayerInstance->GetDataLayerFName());
 
-			if (InState == EDataLayerRuntimeState::Loaded)
+		if (InState == EDataLayerRuntimeState::Loaded)
+		{
+			if (bDataLayerClientOnly || bDataLayerServerOnly)
+			{
+				LocalLoadedDataLayerNames.Add(InDataLayerInstance->GetDataLayerFName());
+			}
+			else
 			{
 				LoadedDataLayerNames.Add(InDataLayerInstance->GetDataLayerFName());
 			}
-			else if (InState == EDataLayerRuntimeState::Activated)
+		}
+		else if (InState == EDataLayerRuntimeState::Activated)
+		{
+			if (bDataLayerClientOnly || bDataLayerServerOnly)
+			{
+				LocalActiveDataLayerNames.Add(InDataLayerInstance->GetDataLayerFName());
+			}
+			else
 			{
 				ActiveDataLayerNames.Add(InDataLayerInstance->GetDataLayerFName());
 			}
+		}
 
-			// Update Replicated Properties
+		// Update replicated properties
+		if (NetMode == NM_DedicatedServer || NetMode == NM_ListenServer)
+		{
 			RepActiveDataLayerNames = ActiveDataLayerNames.Array();
 			RepLoadedDataLayerNames = LoadedDataLayerNames.Array();
+		}
 
-			++DataLayersStateEpoch;
+		++DataLayersStateEpoch;
 
 #if !NO_LOGGING || CSV_PROFILER
-			const FString DataLayerShortName = InDataLayerInstance->GetDataLayerShortName();
-			UE_LOG(LogWorldPartition, Log, TEXT("Data Layer '%s' state changed: %s -> %s"),
-				*DataLayerShortName,
-				*StaticEnum<EDataLayerRuntimeState>()->GetDisplayNameTextByValue((int64)CurrentState).ToString(),
-				*StaticEnum<EDataLayerRuntimeState>()->GetDisplayNameTextByValue((int64)InState).ToString());
+		const FString DataLayerShortName = InDataLayerInstance->GetDataLayerShortName();
+		UE_LOG(LogWorldPartition, Log, TEXT("Data Layer '%s' state changed: %s -> %s"),
+			*DataLayerShortName,
+			*StaticEnum<EDataLayerRuntimeState>()->GetDisplayNameTextByValue((int64)CurrentState).ToString(),
+			*StaticEnum<EDataLayerRuntimeState>()->GetDisplayNameTextByValue((int64)InState).ToString());
 
-			CSV_EVENT_GLOBAL(TEXT("DataLayer-%s-%s"), *DataLayerShortName, *StaticEnum<EDataLayerRuntimeState>()->GetDisplayNameTextByValue((int64)InState).ToString());
+		CSV_EVENT_GLOBAL(TEXT("DataLayer-%s-%s"), *DataLayerShortName, *StaticEnum<EDataLayerRuntimeState>()->GetDisplayNameTextByValue((int64)InState).ToString());
 #endif
 
-			ResolveEffectiveRuntimeState(InDataLayerInstance);
-		}
-		if (bInIsRecursive)
+		ResolveEffectiveRuntimeState(InDataLayerInstance);
+	}
+	if (bInIsRecursive)
+	{
+		InDataLayerInstance->ForEachChild([this, InState, bInIsRecursive](const UDataLayerInstance* Child)
 		{
-			InDataLayerInstance->ForEachChild([this, InState, bInIsRecursive](const UDataLayerInstance* Child)
-				{
-					SetDataLayerRuntimeState(Child, InState, bInIsRecursive);
-					return true;
-				});
-		}
+			SetDataLayerRuntimeState(Child, InState, bInIsRecursive);
+			return true;
+		});
 	}
 }
 
@@ -259,6 +259,16 @@ EDataLayerRuntimeState AWorldDataLayers::GetDataLayerRuntimeStateByName(FName In
 		check(!ActiveDataLayerNames.Contains(InDataLayerName));
 		return EDataLayerRuntimeState::Loaded;
 	}
+	else if (LocalActiveDataLayerNames.Contains(InDataLayerName))
+	{
+		check(!LocalLoadedDataLayerNames.Contains(InDataLayerName));
+		return EDataLayerRuntimeState::Activated;
+	}
+	else if (LocalLoadedDataLayerNames.Contains(InDataLayerName))
+	{
+		check(!LocalActiveDataLayerNames.Contains(InDataLayerName));
+		return EDataLayerRuntimeState::Loaded;
+	}
 
 	return EDataLayerRuntimeState::Unloaded;
 }
@@ -289,17 +299,51 @@ EDataLayerRuntimeState AWorldDataLayers::GetDataLayerEffectiveRuntimeStateByName
 		check(!EffectiveActiveDataLayerNames.Contains(InDataLayerName));
 		return EDataLayerRuntimeState::Loaded;
 	}
+	else if (LocalEffectiveActiveDataLayerNames.Contains(InDataLayerName))
+	{
+		check(!LocalEffectiveLoadedDataLayerNames.Contains(InDataLayerName));
+		return EDataLayerRuntimeState::Activated;
+	}
+	else if (LocalEffectiveLoadedDataLayerNames.Contains(InDataLayerName))
+	{
+		check(!LocalEffectiveActiveDataLayerNames.Contains(InDataLayerName));
+		return EDataLayerRuntimeState::Loaded;
+	}
 
 	return EDataLayerRuntimeState::Unloaded;
 }
 
-void AWorldDataLayers::ResolveEffectiveRuntimeState(const UDataLayerInstance* InDataLayer, bool bInNotifyChange)
+const TSet<FName>& AWorldDataLayers::GetEffectiveActiveDataLayerNames() const
+{
+	AllEffectiveActiveDataLayerNames = EffectiveActiveDataLayerNames;
+	AllEffectiveActiveDataLayerNames.Append(LocalEffectiveActiveDataLayerNames);
+	return AllEffectiveActiveDataLayerNames;
+}
+
+const TSet<FName>& AWorldDataLayers::GetEffectiveLoadedDataLayerNames() const
+{
+	AllEffectiveLoadedDataLayerNames = EffectiveLoadedDataLayerNames;
+	AllEffectiveLoadedDataLayerNames.Append(LocalEffectiveLoadedDataLayerNames);
+	return AllEffectiveLoadedDataLayerNames;
+}
+
+void AWorldDataLayers::ResolveEffectiveRuntimeState(const UDataLayerInstance* InDataLayerInstance, bool bInNotifyChange)
 {
-	check(InDataLayer);
-	const FName DataLayerName = InDataLayer->GetDataLayerFName();
+	check(InDataLayerInstance);
+	const ENetMode NetMode = GetNetMode();
+	const bool bDataLayerClientOnly = InDataLayerInstance->IsClientOnly();
+	const bool bDataLayerServerOnly = InDataLayerInstance->IsServerOnly();
+	const FName DataLayerName = InDataLayerInstance->GetDataLayerFName();
 	EDataLayerRuntimeState CurrentEffectiveRuntimeState = GetDataLayerEffectiveRuntimeStateByName(DataLayerName);
 	EDataLayerRuntimeState NewEffectiveRuntimeState = GetDataLayerRuntimeStateByName(DataLayerName);
-	const UDataLayerInstance* Parent = InDataLayer->GetParent();
+	const UDataLayerInstance* Parent = InDataLayerInstance->GetParent();
+
+	// Filter out irrelevant requests for the current net mode
+	if ((bDataLayerClientOnly && NetMode != NM_Client && NetMode != NM_Standalone) || (bDataLayerServerOnly && NetMode == NM_Client))
+	{
+		return;
+	}
+
 	while (Parent && (NewEffectiveRuntimeState != EDataLayerRuntimeState::Unloaded))
 	{
 		if (Parent->IsRuntime())
@@ -314,32 +358,52 @@ void AWorldDataLayers::ResolveEffectiveRuntimeState(const UDataLayerInstance* In
 	{
 		EffectiveLoadedDataLayerNames.Remove(DataLayerName);
 		EffectiveActiveDataLayerNames.Remove(DataLayerName);
+		LocalEffectiveLoadedDataLayerNames.Remove(DataLayerName);
+		LocalEffectiveActiveDataLayerNames.Remove(DataLayerName);
 
 		if (NewEffectiveRuntimeState == EDataLayerRuntimeState::Loaded)
 		{
-			EffectiveLoadedDataLayerNames.Add(DataLayerName);
+			if (bDataLayerClientOnly || bDataLayerServerOnly)
+			{
+				LocalEffectiveLoadedDataLayerNames.Add(DataLayerName);
+			}
+			else
+			{
+				EffectiveLoadedDataLayerNames.Add(DataLayerName);
+			}
 		}
 		else if (NewEffectiveRuntimeState == EDataLayerRuntimeState::Activated)
 		{
-			EffectiveActiveDataLayerNames.Add(DataLayerName);
+			if (bDataLayerClientOnly || bDataLayerServerOnly)
+			{
+				LocalEffectiveActiveDataLayerNames.Add(DataLayerName);
+			}
+			else
+			{
+				EffectiveActiveDataLayerNames.Add(DataLayerName);
+			}
 		}
 
 		// Update Replicated Properties
-		RepEffectiveActiveDataLayerNames = EffectiveActiveDataLayerNames.Array();
-		RepEffectiveLoadedDataLayerNames = EffectiveLoadedDataLayerNames.Array();
+		if (NetMode == NM_DedicatedServer || NetMode == NM_ListenServer)
+		{
+			RepEffectiveActiveDataLayerNames = EffectiveActiveDataLayerNames.Array();
+			RepEffectiveLoadedDataLayerNames = EffectiveLoadedDataLayerNames.Array();
+		}
+
 		++DataLayersStateEpoch;
 
 		if (bInNotifyChange)
 		{
-			UE_LOG(LogWorldPartition, Log, TEXT("Data Layer '%s' effective state changed: %s -> %s"),
-				*InDataLayer->GetDataLayerShortName(),
+			UE_LOG(LogWorldPartition, Log, TEXT("Data Layer Instance '%s' effective state changed: %s -> %s"),
+				*InDataLayerInstance->GetDataLayerShortName(),
 				*StaticEnum<EDataLayerRuntimeState>()->GetDisplayNameTextByValue((int64)CurrentEffectiveRuntimeState).ToString(),
 				*StaticEnum<EDataLayerRuntimeState>()->GetDisplayNameTextByValue((int64)NewEffectiveRuntimeState).ToString());
 
-			OnDataLayerRuntimeStateChanged(InDataLayer, NewEffectiveRuntimeState);
+			OnDataLayerRuntimeStateChanged(InDataLayerInstance, NewEffectiveRuntimeState);
 		}
 
-		for (const UDataLayerInstance* Child : InDataLayer->GetChildren())
+		for (const UDataLayerInstance* Child : InDataLayerInstance->GetChildren())
 		{
 			ResolveEffectiveRuntimeState(Child);
 		}
@@ -418,7 +482,7 @@ void AWorldDataLayers::DumpDataLayers(FOutputDevice& OutputDevice) const
 	OutputDevice.Logf(TEXT(" Data Layers Hierarchy"));
 	ForEachDataLayer([this, &OutputDevice](UDataLayerInstance* DataLayerInstance)
 	{
-		if (DataLayerInstance && !DataLayerInstance->GetParent())
+		if (!DataLayerInstance->GetParent())
 		{
 			DumpDataLayerRecursively(DataLayerInstance, TEXT(""), OutputDevice);
 		}
@@ -721,9 +785,9 @@ const UDataLayerInstance* AWorldDataLayers::GetDataLayerInstance(const UDataLaye
 
 void AWorldDataLayers::ForEachDataLayer(TFunctionRef<bool(UDataLayerInstance*)> Func)
 {
-	for (UDataLayerInstance* DataLayer : DataLayerInstances)
+	for (UDataLayerInstance* DataLayerInstance : DataLayerInstances)
 	{
-		if (!Func(DataLayer))
+		if (DataLayerInstance && !Func(DataLayerInstance))
 		{
 			break;
 		}
@@ -732,9 +796,9 @@ void AWorldDataLayers::ForEachDataLayer(TFunctionRef<bool(UDataLayerInstance*)>
 
 void AWorldDataLayers::ForEachDataLayer(TFunctionRef<bool(UDataLayerInstance*)> Func) const
 {
-	for (UDataLayerInstance* DataLayer : DataLayerInstances)
+	for (UDataLayerInstance* DataLayerInstance : DataLayerInstances)
 	{
-		if (!Func(DataLayer))
+		if (DataLayerInstance && !Func(DataLayerInstance))
 		{
 			break;
 		}
@@ -785,11 +849,11 @@ void AWorldDataLayers::PostLoad()
 		});
 
 		TArray<UDataLayerInstance*> EditorDataLayers;
-		ForEachDataLayer([&EditorDataLayers](UDataLayerInstance* DataLayer)
+		ForEachDataLayer([&EditorDataLayers](UDataLayerInstance* DataLayerInstance)
 		{
-			if (DataLayer && !DataLayer->IsRuntime())
+			if (!DataLayerInstance->IsRuntime())
 			{
-				EditorDataLayers.Add(DataLayer);
+				EditorDataLayers.Add(DataLayerInstance);
 			}
 			return true;
 		});
@@ -845,8 +909,6 @@ void AWorldDataLayers::PostLoad()
 		}
 	}
 #endif
-
-	InitializeDataLayerRuntimeStates();
 }
 
 #if WITH_EDITOR
diff --git a/Engine/Source/Runtime/Engine/Public/WorldPartition/DataLayer/DataLayerAsset.h b/Engine/Source/Runtime/Engine/Public/WorldPartition/DataLayer/DataLayerAsset.h
--- a/Engine/Source/Runtime/Engine/Public/WorldPartition/DataLayer/DataLayerAsset.h
+++ b/Engine/Source/Runtime/Engine/Public/WorldPartition/DataLayer/DataLayerAsset.h
@@ -13,6 +13,14 @@
 
 class AActor;
 
+UENUM(BlueprintType)
+enum class EDataLayerLoadFilter : uint8
+{
+	None,
+	ClientOnly,
+	ServerOnly
+};
+
 UCLASS(BlueprintType, editinlinenew, MinimalAPI)
 class UDataLayerAsset : public UObject
 {
@@ -36,6 +44,7 @@ public:
 	ENGINE_API bool CanBeReferencedByActor(AActor* InActor) const;
 	static ENGINE_API bool CanBeReferencedByActor(const TSoftObjectPtr<UDataLayerAsset>& InDataLayerAsset, AActor* InActor);
 #endif
+
 	ENGINE_API bool IsPrivate() const;
 
 	UFUNCTION(Category = "Data Layer", BlueprintCallable)
@@ -48,6 +57,13 @@ public:
 	FColor GetDebugColor() const { return DebugColor; }
 
 	bool SupportsActorFilters() const { return bSupportsActorFilters; }
+
+	UFUNCTION(Category = "Data Layer|Runtime", BlueprintCallable)
+	bool IsClientOnly() const { return IsRuntime() && LoadFilter == EDataLayerLoadFilter::ClientOnly; }
+
+	UFUNCTION(Category = "Data Layer|Runtime", BlueprintCallable)
+	bool IsServerOnly() const { return IsRuntime() && LoadFilter == EDataLayerLoadFilter::ServerOnly; }
+
 private:
 	/** Whether the Data Layer affects actor runtime loading */
 	UPROPERTY(Category = "Data Layer", EditAnywhere)
@@ -58,4 +74,7 @@ private:
 		
 	UPROPERTY(Category = "Runtime", EditAnywhere)
 	FColor DebugColor;
+
+	UPROPERTY(Category = "Runtime", EditAnywhere)
+	EDataLayerLoadFilter LoadFilter;
 };
diff --git a/Engine/Source/Runtime/Engine/Public/WorldPartition/DataLayer/DataLayerInstance.h b/Engine/Source/Runtime/Engine/Public/WorldPartition/DataLayer/DataLayerInstance.h
--- a/Engine/Source/Runtime/Engine/Public/WorldPartition/DataLayer/DataLayerInstance.h
+++ b/Engine/Source/Runtime/Engine/Public/WorldPartition/DataLayer/DataLayerInstance.h
@@ -141,6 +141,12 @@ public:
 	UFUNCTION(Category = "Data Layer|Runtime", BlueprintCallable)
 	virtual bool IsRuntime() const { return false; }
 
+	UFUNCTION(Category = "Data Layer|Runtime", BlueprintCallable)
+	virtual bool IsClientOnly() const { return false; }
+
+	UFUNCTION(Category = "Data Layer|Runtime", BlueprintCallable)
+	virtual bool IsServerOnly() const { return false; }
+
 	UFUNCTION(Category = "Data Layer|Runtime", BlueprintCallable)
 	virtual FColor GetDebugColor() const { return FColor::Black; }
 
diff --git a/Engine/Source/Runtime/Engine/Public/WorldPartition/DataLayer/DataLayerInstanceWithAsset.h b/Engine/Source/Runtime/Engine/Public/WorldPartition/DataLayer/DataLayerInstanceWithAsset.h
--- a/Engine/Source/Runtime/Engine/Public/WorldPartition/DataLayer/DataLayerInstanceWithAsset.h
+++ b/Engine/Source/Runtime/Engine/Public/WorldPartition/DataLayer/DataLayerInstanceWithAsset.h
@@ -36,14 +36,17 @@ public:
 
 	const UDataLayerAsset* GetAsset() const override { return DataLayerAsset; }
 
-	virtual EDataLayerType GetType() const override { return DataLayerAsset != nullptr ? DataLayerAsset->GetType() : EDataLayerType::Unknown; }
+	virtual EDataLayerType GetType() const override { return DataLayerAsset ? DataLayerAsset->GetType() : EDataLayerType::Unknown; }
 
-	virtual bool IsRuntime() const override { return DataLayerAsset != nullptr ? DataLayerAsset->IsRuntime() : false; }
+	virtual bool IsRuntime() const override { return DataLayerAsset && DataLayerAsset->IsRuntime(); }
 
-	virtual FColor GetDebugColor() const override { return DataLayerAsset != nullptr ? DataLayerAsset->GetDebugColor() : FColor::Black; }
+	virtual bool IsClientOnly() const override { return DataLayerAsset && DataLayerAsset->IsClientOnly(); }
+	virtual bool IsServerOnly() const override { return DataLayerAsset && DataLayerAsset->IsServerOnly(); }
 
-	virtual FString GetDataLayerShortName() const override { return DataLayerAsset != nullptr ? DataLayerAsset->GetName() : GetDataLayerFName().ToString(); }
-	virtual FString GetDataLayerFullName() const override { return DataLayerAsset != nullptr ? DataLayerAsset->GetPathName() : GetDataLayerFName().ToString(); }
+	virtual FColor GetDebugColor() const override { return DataLayerAsset ? DataLayerAsset->GetDebugColor() : FColor::Black; }
+
+	virtual FString GetDataLayerShortName() const override { return DataLayerAsset ? DataLayerAsset->GetName() : GetDataLayerFName().ToString(); }
+	virtual FString GetDataLayerFullName() const override { return DataLayerAsset ? DataLayerAsset->GetPathName() : GetDataLayerFName().ToString(); }
 
 protected:
 #if WITH_EDITOR
diff --git a/Engine/Source/Runtime/Engine/Public/WorldPartition/DataLayer/WorldDataLayers.h b/Engine/Source/Runtime/Engine/Public/WorldPartition/DataLayer/WorldDataLayers.h
--- a/Engine/Source/Runtime/Engine/Public/WorldPartition/DataLayer/WorldDataLayers.h
+++ b/Engine/Source/Runtime/Engine/Public/WorldPartition/DataLayer/WorldDataLayers.h
@@ -110,8 +110,8 @@ public:
 	ENGINE_API void SetDataLayerRuntimeState(const UDataLayerInstance* InDataLayerInstance, EDataLayerRuntimeState InState, bool bIsRecursive = false);
 	ENGINE_API EDataLayerRuntimeState GetDataLayerRuntimeStateByName(FName InDataLayerName) const;
 	ENGINE_API EDataLayerRuntimeState GetDataLayerEffectiveRuntimeStateByName(FName InDataLAyerName) const;
-	const TSet<FName>& GetEffectiveActiveDataLayerNames() const { return EffectiveActiveDataLayerNames; }
-	const TSet<FName>& GetEffectiveLoadedDataLayerNames() const { return EffectiveLoadedDataLayerNames; }
+	const TSet<FName>& GetEffectiveActiveDataLayerNames() const;
+	const TSet<FName>& GetEffectiveLoadedDataLayerNames() const;
 	UFUNCTION(NetMulticast, Reliable)
 	ENGINE_API void OnDataLayerRuntimeStateChanged(const UDataLayerInstance* InDataLayer, EDataLayerRuntimeState InState);
 	static int32 GetDataLayersStateEpoch() { return DataLayersStateEpoch; }
@@ -229,6 +229,9 @@ private:
 	TSet<FName> ActiveDataLayerNames;
 	TSet<FName> LoadedDataLayerNames;
 
+	TSet<FName> LocalActiveDataLayerNames;
+	TSet<FName> LocalLoadedDataLayerNames;
+
 	UPROPERTY(Transient, Replicated, ReplicatedUsing=OnRep_EffectiveActiveDataLayerNames)
 	TArray<FName> RepEffectiveActiveDataLayerNames;
 		
@@ -239,6 +242,12 @@ private:
 	TSet<FName> EffectiveActiveDataLayerNames;
 	TSet<FName> EffectiveLoadedDataLayerNames;
 
+	TSet<FName> LocalEffectiveActiveDataLayerNames;
+	TSet<FName> LocalEffectiveLoadedDataLayerNames;
+
+	mutable TSet<FName> AllEffectiveActiveDataLayerNames;
+	mutable TSet<FName> AllEffectiveLoadedDataLayerNames;
+
 	static ENGINE_API int32 DataLayersStateEpoch;
 
 	static_assert(DATALAYER_TO_INSTANCE_RUNTIME_CONVERSION_ENABLED, "bHasDeprecatedDataLayers is deprecated and needs to be deleted.");
@@ -321,11 +330,11 @@ void AWorldDataLayers::OverwriteDataLayerRuntimeStates(const TArray<IdentifierTy
 			RepLoadedDataLayerNames = LoadedDataLayerNames.Array();
 		}
 
-		ForEachDataLayer([this](class UDataLayerInstance* DataLayer)
+		ForEachDataLayer([this](class UDataLayerInstance* DataLayerInstance)
 		{
-			if (DataLayer && DataLayer->IsRuntime())
+			if (DataLayerInstance->IsRuntime())
 			{
-				ResolveEffectiveRuntimeState(DataLayer, /*bNotifyChange*/false);
+				ResolveEffectiveRuntimeState(DataLayerInstance, /*bNotifyChange*/false);
 			}
 			return true;
 		});
