World Partition - HLOD: UHLODTemplatedInstancedStaticMeshComponent should only null out the properties that are identical to the source template actor class ones

diff --git a/Engine/Source/Runtime/Engine/Private/WorldPartition/HLOD/HLODTemplatedInstancedStaticMeshComponent.cpp b/Engine/Source/Runtime/Engine/Private/WorldPartition/HLOD/HLODTemplatedInstancedStaticMeshComponent.cpp
--- a/Engine/Source/Runtime/Engine/Private/WorldPartition/HLOD/HLODTemplatedInstancedStaticMeshComponent.cpp
+++ b/Engine/Source/Runtime/Engine/Private/WorldPartition/HLOD/HLODTemplatedInstancedStaticMeshComponent.cpp
@@ -31,11 +31,34 @@ void UHLODTemplatedInstancedStaticMeshComponent::PreSave(FObjectPreSaveContext O
 		};
 
 		FObjectSaveOverride ObjectSaveOverride;
-		ObjectSaveOverride.PropOverrides.Add(GetTransientPropertyOverride(GetMemberNameChecked_StaticMesh()));
-		ObjectSaveOverride.PropOverrides.Add(GetTransientPropertyOverride(GET_MEMBER_NAME_CHECKED(UMeshComponent, OverrideMaterials)));
-		ObjectSaveOverride.PropOverrides.Add(GetTransientPropertyOverride(GET_MEMBER_NAME_CHECKED(UMeshComponent, OverlayMaterial)));
 
-		ObjectSaveContext.AddSaveOverride(this, ObjectSaveOverride);
+		// Only add overrides to clear properties that are identical to the ones provided by template actor class
+		const UStaticMeshComponent* TemplateSMC = AActor::GetActorClassDefaultComponentByName<UStaticMeshComponent>(TemplateActorClass, TemplateComponentName);
+		if (TemplateSMC)
+		{
+			// StaticMesh
+			if (GetStaticMesh() == TemplateSMC->GetStaticMesh())
+			{
+				ObjectSaveOverride.PropOverrides.Add(GetTransientPropertyOverride(GetMemberNameChecked_StaticMesh()));
+			}
+
+			// OverrideMaterials
+			if (OverrideMaterials == TemplateSMC->OverrideMaterials)
+			{
+				ObjectSaveOverride.PropOverrides.Add(GetTransientPropertyOverride(GET_MEMBER_NAME_CHECKED(UMeshComponent, OverrideMaterials)));
+			}
+
+			// OverlayMaterial
+			if (OverlayMaterial == TemplateSMC->OverlayMaterial)
+			{
+				ObjectSaveOverride.PropOverrides.Add(GetTransientPropertyOverride(GET_MEMBER_NAME_CHECKED(UMeshComponent, OverlayMaterial)));
+			}
+		}		
+
+		if (!ObjectSaveOverride.PropOverrides.IsEmpty())
+		{
+			ObjectSaveContext.AddSaveOverride(this, ObjectSaveOverride);
+		}
 	}
 #endif
 
@@ -46,10 +69,6 @@ void UHLODTemplatedInstancedStaticMeshComponent::PostLoad()
 {
 	if (!HasAnyFlags(RF_ArchetypeObject | RF_ClassDefaultObject))
 	{
-		check(GetStaticMesh() == nullptr);
-		check(!OverrideMaterials.ContainsByPredicate([](const UMaterialInterface* MaterialInterface) { return MaterialInterface != nullptr; }));
-		check(OverlayMaterial == nullptr);
-
 		RestoreAssetsFromActorTemplate();
 	}
 
@@ -68,20 +87,30 @@ void UHLODTemplatedInstancedStaticMeshComponent::SetTemplateComponentName(const
 
 void UHLODTemplatedInstancedStaticMeshComponent::RestoreAssetsFromActorTemplate()
 {
+	// Restore null properties from the template actor class
 	const UStaticMeshComponent* TemplateSMC = AActor::GetActorClassDefaultComponentByName<UStaticMeshComponent>(TemplateActorClass, TemplateComponentName);
 	if (TemplateSMC)
 	{
 		// StaticMesh
-		if (UStaticMesh* TemplateStaticMesh = TemplateSMC->GetStaticMesh())
+		if (GetStaticMesh() == nullptr)
 		{
-			SetStaticMesh(TemplateStaticMesh);
-			SetForcedLodModel(TemplateStaticMesh->GetNumLODs());
+			if (UStaticMesh* TemplateStaticMesh = TemplateSMC->GetStaticMesh())
+			{
+				SetStaticMesh(TemplateStaticMesh);
+				SetForcedLodModel(TemplateStaticMesh->GetNumLODs());
+			}
 		}
 
 		// OverrideMaterials
-		OverrideMaterials = TemplateSMC->OverrideMaterials;
+		if (OverrideMaterials.IsEmpty())
+		{
+			OverrideMaterials = TemplateSMC->OverrideMaterials;
+		}
 
 		// OverlayMaterial
-		OverlayMaterial = TemplateSMC->OverlayMaterial;
+		if (OverlayMaterial == nullptr)
+		{
+			OverlayMaterial = TemplateSMC->OverlayMaterial;
+		}
 	}
 }
