World Partition
* Add an editor preference to log the generate streaming in PIE (default true).

diff --git a/Engine/Source/Editor/WorldPartitionEditor/Private/WorldPartitionEditorModule.cpp b/Engine/Source/Editor/WorldPartitionEditor/Private/WorldPartitionEditorModule.cpp
--- a/Engine/Source/Editor/WorldPartitionEditor/Private/WorldPartitionEditorModule.cpp
+++ b/Engine/Source/Editor/WorldPartitionEditor/Private/WorldPartitionEditorModule.cpp
@@ -297,6 +297,16 @@ void FWorldPartitionEditorModule::SetEnableLoadingInEditor(bool bInEnableLoading
 	GetMutableDefault<UWorldPartitionEditorSettings>()->bEnableLoadingInEditor = bInEnableLoadingInEditor;
 }
 
+bool FWorldPartitionEditorModule::GetEnableStreamingGenerationLogOnPIE() const
+{
+	return GetDefault<UWorldPartitionEditorSettings>()->bEnableStreamingGenerationLogOnPIE;
+}
+
+void FWorldPartitionEditorModule::SetEnableStreamingGenerationLogOnPIE(bool bEnableStreamingGenerationLogOnPIE)
+{
+	GetMutableDefault<UWorldPartitionEditorSettings>()->bEnableStreamingGenerationLogOnPIE = bEnableStreamingGenerationLogOnPIE;
+}
+
 bool FWorldPartitionEditorModule::GetDisablePIE() const
 {
 	return GetDefault<UWorldPartitionEditorSettings>()->bDisablePIE;
@@ -819,6 +829,7 @@ UWorldPartitionEditorSettings::UWorldPartitionEditorSettings()
 	InstancedFoliageGridSize = 25600;
 	MinimapLowQualityWorldUnitsPerPixelThreshold = 12800;
 	bEnableLoadingInEditor = true;
+	bEnableStreamingGenerationLogOnPIE = true;
 	bDisablePIE = false;
 	bDisableBugIt = false;
 	bAdvancedMode = true;
diff --git a/Engine/Source/Editor/WorldPartitionEditor/Public/WorldPartition/WorldPartitionEditorSettings.h b/Engine/Source/Editor/WorldPartitionEditor/Public/WorldPartition/WorldPartitionEditorSettings.h
--- a/Engine/Source/Editor/WorldPartitionEditor/Public/WorldPartition/WorldPartitionEditorSettings.h
+++ b/Engine/Source/Editor/WorldPartitionEditor/Public/WorldPartition/WorldPartitionEditorSettings.h
@@ -29,6 +29,9 @@ public:
 	UPROPERTY(Config, EditAnywhere, Category = WorldPartition, Meta = (ToolTip = "Wheter to enable dynamic loading in the editor through loading regions"))
 	bool bEnableLoadingInEditor;
 
+	UPROPERTY(Config, EditAnywhere, Category = WorldPartition, Meta = (ToolTip = "Wheter to enable streaming generation log on PIE"))
+	bool bEnableStreamingGenerationLogOnPIE;
+
 	bool bDisableBugIt;
 	bool bDisablePIE;
 	bool bAdvancedMode;
diff --git a/Engine/Source/Editor/WorldPartitionEditor/Public/WorldPartitionEditorModule.h b/Engine/Source/Editor/WorldPartitionEditor/Public/WorldPartitionEditorModule.h
--- a/Engine/Source/Editor/WorldPartitionEditor/Public/WorldPartitionEditorModule.h
+++ b/Engine/Source/Editor/WorldPartitionEditor/Public/WorldPartitionEditorModule.h
@@ -53,6 +53,12 @@ public:
 	virtual bool GetEnableLoadingInEditor() const override;
 	virtual void SetEnableLoadingInEditor(bool bInEnableLoadingInEditor) override;
 
+	/**
+	* Returns if streaming generation log on PIE is enabled or not.
+	*/
+	virtual bool GetEnableStreamingGenerationLogOnPIE() const override;
+	virtual void SetEnableStreamingGenerationLogOnPIE(bool bInEnableLoadingInEditor) override;
+
 	/**
 	 * Returns if pie is disabled or not.
 	 */
diff --git a/Engine/Source/Runtime/Engine/Private/WorldPartition/WorldPartitionStreamingGeneration.cpp b/Engine/Source/Runtime/Engine/Private/WorldPartition/WorldPartitionStreamingGeneration.cpp
--- a/Engine/Source/Runtime/Engine/Private/WorldPartition/WorldPartitionStreamingGeneration.cpp
+++ b/Engine/Source/Runtime/Engine/Private/WorldPartition/WorldPartitionStreamingGeneration.cpp
@@ -18,6 +18,7 @@
 #include "Misc/HashBuilder.h"
 #include "Misc/Paths.h"
 #include "ProfilingDebugging/ScopedTimers.h"
+#include "WorldPartition/IWorldPartitionEditorModule.h"
 #include "WorldPartition/WorldPartitionRuntimeHash.h"
 #include "WorldPartition/WorldPartitionStreamingPolicy.h"
 #include "WorldPartition/WorldPartitionLevelStreamingPolicy.h"
@@ -1287,7 +1288,14 @@ bool UWorldPartition::GenerateContainerStreaming(const FGenerateStreamingParams&
 	TUniquePtr<FArchive> LogFileAr;
 	TUniquePtr<FHierarchicalLogArchive> HierarchicalLogAr;
 
-	if (IsMainWorldPartition() && (!GIsBuildMachine || GIsAutomationTesting || IsRunningCookCommandlet()))
+	bool bAllowStreamingGenerationLog = true;
+	IWorldPartitionEditorModule& WorldPartitionEditorModule = FModuleManager::LoadModuleChecked<IWorldPartitionEditorModule>("WorldPartitionEditor");
+	if (bIsPIE && !WorldPartitionEditorModule.GetEnableStreamingGenerationLogOnPIE())
+	{
+		bAllowStreamingGenerationLog = false;
+	}
+
+	if (bAllowStreamingGenerationLog && IsMainWorldPartition() && (!GIsBuildMachine || GIsAutomationTesting || IsRunningCookCommandlet()))
 	{
 		TStringBuilder<256> StateLogSuffix;
 		StateLogSuffix += bIsPIE ? TEXT("PIE") : (IsRunningGame() ? TEXT("Game") : (IsRunningCookCommandlet() ? TEXT("Cook") : (GIsAutomationTesting ? TEXT("UnitTest") : TEXT("Manual"))));
diff --git a/Engine/Source/Runtime/Engine/Public/WorldPartition/IWorldPartitionEditorModule.h b/Engine/Source/Runtime/Engine/Public/WorldPartition/IWorldPartitionEditorModule.h
--- a/Engine/Source/Runtime/Engine/Public/WorldPartition/IWorldPartitionEditorModule.h
+++ b/Engine/Source/Runtime/Engine/Public/WorldPartition/IWorldPartitionEditorModule.h
@@ -39,6 +39,9 @@ public:
 	virtual bool GetEnableLoadingInEditor() const = 0;
 	virtual void SetEnableLoadingInEditor(bool bInEnableLoadingInEditor) = 0;
 
+	virtual bool GetEnableStreamingGenerationLogOnPIE() const = 0;
+	virtual void SetEnableStreamingGenerationLogOnPIE(bool bInEnableLoadingInEditor) = 0;
+
 	virtual bool GetDisablePIE() const = 0;
 	virtual void SetDisablePIE(bool bInDisablePIE) = 0;
 
