Added a flag to always make the actor spawned by a ChildActorComponent transient

diff --git a/Engine/Source/Runtime/Engine/Classes/Components/ChildActorComponent.h b/Engine/Source/Runtime/Engine/Classes/Components/ChildActorComponent.h
--- a/Engine/Source/Runtime/Engine/Classes/Components/ChildActorComponent.h
+++ b/Engine/Source/Runtime/Engine/Classes/Components/ChildActorComponent.h
@@ -139,6 +139,13 @@ private:
 	EChildActorComponentTreeViewVisualizationMode EditorTreeViewVisualizationMode;
 #endif
 
+	/**
+	 * Should the spawned actor be marked as transient?
+	 * @note The spawned actor will also be marked transient if this component or its owner actor are transient, regardless of the state of this flag.
+	 */
+	UPROPERTY(EditDefaultsOnly, Category=ChildActorComponent)
+	uint8 bChildActorIsTransient:1;
+
 	/** Flag indicating that when the component is registered that the child actor should be recreated */
 	uint8 bNeedsRecreate:1;
 
diff --git a/Engine/Source/Runtime/Engine/Private/Components/ChildActorComponent.cpp b/Engine/Source/Runtime/Engine/Private/Components/ChildActorComponent.cpp
--- a/Engine/Source/Runtime/Engine/Private/Components/ChildActorComponent.cpp
+++ b/Engine/Source/Runtime/Engine/Private/Components/ChildActorComponent.cpp
@@ -794,7 +794,7 @@ void UChildActorComponent::CreateChildActor(TFunction<void(AActor*)> CustomizerF
 				{
 					Params.ObjectFlags &= ~RF_Transactional;
 				}
-				if (HasAllFlags(RF_Transient) || (MyOwner && MyOwner->HasAllFlags(RF_Transient)))
+				if (bChildActorIsTransient || HasAllFlags(RF_Transient) || (MyOwner && MyOwner->HasAllFlags(RF_Transient)))
 				{
 					// If this component or its owner are transient, set our created actor to transient. 
 					Params.ObjectFlags |= RF_Transient;
