Fix HTML formatting for FPSChart

diff --git a/Engine/Source/Runtime/Engine/Private/ChartCreation.cpp b/Engine/Source/Runtime/Engine/Private/ChartCreation.cpp
--- a/Engine/Source/Runtime/Engine/Private/ChartCreation.cpp
+++ b/Engine/Source/Runtime/Engine/Private/ChartCreation.cpp
@@ -1637,7 +1637,7 @@ TEXT("<TD CLASS=\"columnSeparator\"><DIV>&nbsp;</DIV></TD>\n")
 TEXT("\n")
 TEXT("<TD CLASS=\"rowHeaderTimes\"><DIV>time</DIV></TD>\n")
 TEXT("<TD CLASS=\"rowHeaderTimes\"><DIV>frame count</DIV></TD>\n")
-TEXT("<TD CLASS=\"rowHeaderTimes\"<DIV>time disregarded</DIV></TD>\n")
+TEXT("<TD CLASS=\"rowHeaderTimes\"><DIV>time disregarded</DIV></TD>\n")
 TEXT("\n")
 TEXT("<TD CLASS=\"columnSeparator\"><DIV>&nbsp;</DIV></TD>\n")
 TEXT("<TD><DIV CLASS=\"rowHeaderTimes\">Total Hitches</DIV></TD>\n")
