Rewind debugger fixes for tracks not showing correctly until restarting pie or re recording

diff --git a/Engine/Plugins/Animation/GameplayInsights/Source/RewindDebugger/Private/RewindDebugger.cpp b/Engine/Plugins/Animation/GameplayInsights/Source/RewindDebugger/Private/RewindDebugger.cpp
--- a/Engine/Plugins/Animation/GameplayInsights/Source/RewindDebugger/Private/RewindDebugger.cpp
+++ b/Engine/Plugins/Animation/GameplayInsights/Source/RewindDebugger/Private/RewindDebugger.cpp
@@ -351,8 +351,8 @@ void FRewindDebugger::RefreshDebugTracks()
     			// remove any existing tracks that don't match the current list of object ids
     			for (int TrackIndex = DebugTracks.Num()-1; TrackIndex>=0; TrackIndex--)
     			{
-    				uint64* FoundId = TargetObjectIds.FindByPredicate([this, TrackIndex](const uint64& TrackId) { return DebugTracks[TrackIndex]->GetObjectId() == TrackId; });
-    				if (!FoundId)
+    				uint64 TrackId = DebugTracks[TrackIndex]->GetObjectId();
+    				if (!TargetObjectIds.Contains(TrackId))
     				{
     					DebugTracks.RemoveAt(TrackIndex);
     				}
@@ -401,6 +401,13 @@ namespace
 	}
 }
 
+void FRewindDebugger::OnConnection()
+{
+	// queue up some operations to happen on the game thread next tick
+	bTraceJustConnected = true;
+	FTraceAuxiliary::OnConnection.RemoveAll(this);
+}
+
 void FRewindDebugger::StartRecording()
 {
 	if (!CanStartRecording())
@@ -425,41 +432,14 @@ void FRewindDebugger::StartRecording()
 	// Clear all buffered data and prevent data from previous recordings from leaking into the new recording
 	FTraceAuxiliary::FOptions Options;
 	Options.bExcludeTail = true;
-	
-	FTraceAuxiliary::Start(FTraceAuxiliary::EConnectionType::Network, TEXT("127.0.0.1"), TEXT(""), &Options, LogRewindDebugger);
 
-	UE::Trace::ToggleChannel(TEXT("Object"), true);
-	UE::Trace::ToggleChannel(TEXT("ObjectProperties"), true);
-	UE::Trace::ToggleChannel(TEXT("Animation"), true);
-	UE::Trace::ToggleChannel(TEXT("Frame"), true);
+	FTraceAuxiliary::OnConnection.AddRaw(this, &FRewindDebugger::OnConnection); 
+
 	
+	FTraceAuxiliary::Start(FTraceAuxiliary::EConnectionType::Network, TEXT("127.0.0.1"), TEXT(""), &Options, LogRewindDebugger);
 	UnrealInsightsModule->StartAnalysisForLastLiveSession();
 
-	// update extensions
-	IterateExtensions([this](IRewindDebuggerExtension* Extension)
-		{
-			Extension->RecordingStarted(this);
-		}
-	);
 
-#if OBJECT_TRACE_ENABLED
-	// trace each play-in-editor world, and all the actors in it.
-	for (TObjectIterator<UWorld> World; World; ++World)
-	{
-		if (World->IsPlayInEditor())
-		{
-			FObjectTrace::ResetWorldElapsedTime(*World);
-			FObjectTrace::SetWorldRecordingIndex(*World, RecordingIndex);
-			
-			TRACE_WORLD(*World);
-				
-			for (TActorIterator<AActor> Iterator(*World); Iterator; ++Iterator)
-			{
-				TRACE_OBJECT_LIFETIME_BEGIN(*Iterator);
-			}
-		}
-	}
-#endif // OBJECT_TRACE_ENABLED
 
 	TargetObjectIds.Empty(2);
 
@@ -817,6 +797,43 @@ const TraceServices::IAnalysisSession* FRewindDebugger::GetAnalysisSession() con
 void FRewindDebugger::Tick(float DeltaTime)
 {
 	TRACE_CPUPROFILER_EVENT_SCOPE(FRewindDebugger::Tick);
+
+	if (bTraceJustConnected)
+	{
+		bTraceJustConnected = false;
+		
+		UE::Trace::ToggleChannel(TEXT("Object"), true);
+		UE::Trace::ToggleChannel(TEXT("ObjectProperties"), true);
+		UE::Trace::ToggleChannel(TEXT("Animation"), true);
+		UE::Trace::ToggleChannel(TEXT("Frame"), true);
+
+		// update extensions
+		IterateExtensions([this](IRewindDebuggerExtension* Extension)
+			{
+				Extension->RecordingStarted(this);
+			}
+		);
+
+	#if OBJECT_TRACE_ENABLED
+		// trace each play-in-editor world, and all the actors in it.
+		for (TObjectIterator<UWorld> World; World; ++World)
+		{
+			if (World->IsPlayInEditor())
+			{
+				FObjectTrace::ResetWorldElapsedTime(*World);
+				FObjectTrace::SetWorldRecordingIndex(*World, RecordingIndex);
+				
+				TRACE_WORLD(*World);
+					
+				for (TActorIterator<AActor> Iterator(*World); Iterator; ++Iterator)
+				{
+					TRACE_OBJECT_LIFETIME_BEGIN(*Iterator);
+				}
+			}
+		}
+	#endif // OBJECT_TRACE_ENABLED
+	}
+	
 	if (const TraceServices::IAnalysisSession* Session = GetAnalysisSession())
 	{
 		RefreshDebugTracks();
diff --git a/Engine/Plugins/Animation/GameplayInsights/Source/RewindDebugger/Private/RewindDebugger.h b/Engine/Plugins/Animation/GameplayInsights/Source/RewindDebugger/Private/RewindDebugger.h
--- a/Engine/Plugins/Animation/GameplayInsights/Source/RewindDebugger/Private/RewindDebugger.h
+++ b/Engine/Plugins/Animation/GameplayInsights/Source/RewindDebugger/Private/RewindDebugger.h
@@ -45,7 +45,8 @@ public:
 	virtual TSharedPtr<FDebugObjectInfo> GetSelectedComponent() const override;
 	virtual TSharedPtr<RewindDebugger::FRewindDebuggerTrack> GetSelectedTrack() const override;
 	virtual TArray<TSharedPtr<FDebugObjectInfo>>& GetDebugComponents() override;
-	
+
+	void OnConnection();
 
 	void GetTargetObjectIds(TArray<uint64>& OutActorIds) const;
 
@@ -150,6 +151,7 @@ private:
 	FOnComponentListChanged ComponentListChangedDelegate;
 	FOnTrackCursor TrackCursorDelegate;
 
+	bool bTraceJustConnected = false;
 	bool bPIEStarted;
 	bool bPIESimulating;
 	
diff --git a/Engine/Source/Runtime/Engine/Private/ObjectTrace.cpp b/Engine/Source/Runtime/Engine/Private/ObjectTrace.cpp
--- a/Engine/Source/Runtime/Engine/Private/ObjectTrace.cpp
+++ b/Engine/Source/Runtime/Engine/Private/ObjectTrace.cpp
@@ -1,6 +1,7 @@
 // Copyright Epic Games, Inc. All Rights Reserved.
 
 #include "ObjectTrace.h"
+#include "Misc/ScopeRWLock.h"
 
 #include UE_INLINE_GENERATED_CPP_BY_NAME(ObjectTrace)
 
@@ -88,22 +89,6 @@ UE_TRACE_EVENT_BEGIN(Object, View)
 UE_TRACE_EVENT_END()
 
 // Object annotations used for tracing
-struct FTracedObjectAnnotation
-{
-	FTracedObjectAnnotation()
-		: bTraced(false)
-	{}
-
-	// Whether this object has been traced this session
-	bool bTraced;
-
-	/** Determine if this annotation is default - required for annotations */
-	FORCEINLINE bool IsDefault() const
-	{
-		return bTraced == false;
-	}
-};
-
 struct FObjectIdAnnotation
 {
 	FObjectIdAnnotation()
@@ -131,7 +116,9 @@ int32 GetTypeHash(const FObjectIdAnnotation& Annotation)
 }
 
 // Object annotations used for tracing
-FUObjectAnnotationSparse<FTracedObjectAnnotation, true> GObjectTracedAnnotations;
+// FUObjectAnnotationSparse<FTracedObjectAnnotation, true> GObjectTracedAnnotations;
+FRWLock GObjectTracedSetLock;
+TSet<uint64> GObjectTracedSet;
 FUObjectAnnotationSparseSearchable<FObjectIdAnnotation, true> GObjectIdAnnotations;
 
 // Handle used to hook to world tick
@@ -172,7 +159,11 @@ void FObjectTrace::Destroy()
 void FObjectTrace::Reset()
 {
 	GObjectIdAnnotations.RemoveAllAnnotations();
-	GObjectTracedAnnotations.RemoveAllAnnotations();
+	
+	{
+		FRWScopeLock ScopeLock(GObjectTracedSetLock,SLT_Write);
+		GObjectTracedSet.Empty();
+	}
 }
 
 uint64 FObjectTrace::GetObjectId(const UObject* InObject)
@@ -229,8 +220,6 @@ void FObjectTrace::ResetWorldElapsedTime(const UWorld* World)
 	{
 		WorldSubsystem->ElapsedTime = 0;
 	}
-
-	GObjectTracedAnnotations.RemoveAllAnnotations();
 }
 
 double FObjectTrace::GetWorldElapsedTime(const UWorld* World)
@@ -314,25 +303,29 @@ void FObjectTrace::OutputClass(const UClass* InClass)
 	{
 		return;
 	}
+	uint64 ObjectId = GetObjectId(InClass);
 
-	FTracedObjectAnnotation Annotation = GObjectTracedAnnotations.GetAnnotation(InClass);
-	if(Annotation.bTraced)
 	{
-		// Already traced, so skip
-		return;
+		FRWScopeLock ScopeLock(GObjectTracedSetLock, SLT_ReadOnly);
+		if (GObjectTracedSet.Contains(ObjectId))
+		{
+			// Already traced, so skip
+			return;
+		}
+	}
+	{
+		FRWScopeLock ScopeLock(GObjectTracedSetLock, SLT_Write);
+		GObjectTracedSet.Add(ObjectId);
 	}
 
 	OutputClass(InClass->GetSuperClass());
 
-	Annotation.bTraced = true;
-	GObjectTracedAnnotations.AddAnnotation(InClass, MoveTemp(Annotation));
-
 	FString ClassPathName = InClass->GetPathName();
 	TCHAR ClassName[FName::StringBufferSize];
 	uint32 ClassNameLength = InClass->GetFName().ToString(ClassName);
 
 	UE_TRACE_LOG(Object, Class, ObjectChannel)
-		<< Class.Id(GetObjectId(InClass))
+		<< Class.Id(ObjectId)
 		<< Class.SuperId(GetObjectId(InClass->GetSuperClass()))
 		<< Class.Name(ClassName, ClassNameLength)
 		<< Class.Path(*ClassPathName, ClassPathName.Len());
@@ -387,19 +380,23 @@ void FObjectTrace::OutputObject(const UObject* InObject)
 	{
 		return;
 	}
-
-	FTracedObjectAnnotation Annotation = GObjectTracedAnnotations.GetAnnotation(InObject);
-	if(Annotation.bTraced)
+	uint64 ObjectId = GetObjectId(InObject);
+	
 	{
-		// Already traced, so skip
-		return;
+		FRWScopeLock ScopeLock(GObjectTracedSetLock, SLT_ReadOnly);
+		if (GObjectTracedSet.Contains(ObjectId))
+		{
+			// Already traced, so skip
+			return;
+		}
+	}
+	{
+		FRWScopeLock ScopeLock(GObjectTracedSetLock, SLT_Write);
+		GObjectTracedSet.Add(ObjectId);
 	}
 
 	OutputObject(InObject->GetOuter());
 
-	Annotation.bTraced = true;
-	GObjectTracedAnnotations.AddAnnotation(InObject, MoveTemp(Annotation));
-
 	// Trace the object's class first
 	TRACE_CLASS(InObject->GetClass());
 
@@ -408,7 +405,7 @@ void FObjectTrace::OutputObject(const UObject* InObject)
 	FString ObjectPathName = InObject->GetPathName();
 
 	UE_TRACE_LOG(Object, Object, ObjectChannel)
-		<< Object.Id(GetObjectId(InObject))
+		<< Object.Id(ObjectId)
 		<< Object.ClassId(GetObjectId(InObject->GetClass()))
 		<< Object.OuterId(GetObjectId(InObject->GetOuter()))
 		<< Object.Name(ObjectName, ObjectNameLength)
