Added preallocate memory option to the Inertialization node.

diff --git a/Engine/Source/Runtime/Engine/Classes/Animation/AnimNode_Inertialization.h b/Engine/Source/Runtime/Engine/Classes/Animation/AnimNode_Inertialization.h
--- a/Engine/Source/Runtime/Engine/Classes/Animation/AnimNode_Inertialization.h
+++ b/Engine/Source/Runtime/Engine/Classes/Animation/AnimNode_Inertialization.h
@@ -304,9 +304,9 @@ struct FInertializationPoseDiff
 	{
 	}
 
-	void Reset()
+	void Reset(uint32 NumBonesSlack = 0)
 	{
-		BoneDiffs.Empty();
+		BoneDiffs.Empty(NumBonesSlack);
 		CurveDiffs.Empty();
 		InertializationSpace = EInertializationSpace::Default;
 	}
@@ -372,6 +372,14 @@ private:
 	UPROPERTY(EditAnywhere, Category = Filter)
 	TArray<FName> FilteredCurves;
 
+	/**
+	 * Enable this to pre-allocate memory for the node rather than to allocate and deallocate memory when blending
+	 * becomes active and inactive. This improves performance, but causes larger memory usage, in particular when you
+	 * have multiple Inertialization nodes in an animation graph that are not all used at once.
+	 */
+	UPROPERTY(EditAnywhere, Category = Memory)
+	bool bPreallocateMemory = false;
+
 public: // FAnimNode_Inertialization
 
 	ENGINE_API FAnimNode_Inertialization();
diff --git a/Engine/Source/Runtime/Engine/Private/Animation/AnimNode_Inertialization.cpp b/Engine/Source/Runtime/Engine/Private/Animation/AnimNode_Inertialization.cpp
--- a/Engine/Source/Runtime/Engine/Private/Animation/AnimNode_Inertialization.cpp
+++ b/Engine/Source/Runtime/Engine/Private/Animation/AnimNode_Inertialization.cpp
@@ -149,7 +149,7 @@ void FAnimNode_Inertialization::Initialize_AnyThread(const FAnimationInitializeC
 	Source.Initialize(Context);
 
 
-	const int32 NumSkeletonBones = Context.AnimInstanceProxy->GetSkeleton()->GetReferenceSkeleton().GetNum();
+	const int32 NumSkeletonBones = bPreallocateMemory ? Context.AnimInstanceProxy->GetSkeleton()->GetReferenceSkeleton().GetNum() : 0;
 
 	PoseSnapshots.Empty(INERTIALIZATION_MAX_POSE_SNAPSHOTS);
 	RequestQueue.Reserve(8);
@@ -161,7 +161,7 @@ void FAnimNode_Inertialization::Initialize_AnyThread(const FAnimationInitializeC
 
 	Deactivate();
 
-	InertializationPoseDiff.Reset();
+	InertializationPoseDiff.Reset(NumSkeletonBones);
 }
 
 
@@ -542,10 +542,16 @@ void FAnimNode_Inertialization::ApplyInertialization(FPoseContext& Context, cons
 
 void FAnimNode_Inertialization::Deactivate()
 {
+	if (!bPreallocateMemory)
+	{
+		// InertializationPoseDiff::Reset actually empties and frees memory here
+		InertializationPoseDiff.Reset();
+		InertializationDurationPerBone.Empty();
+	}
+
 	InertializationState = EInertializationState::Inactive;
 	InertializationElapsedTime = 0.0f;
 	InertializationDuration = 0.0f;
-	InertializationDurationPerBone.Reset();
 	InertializationMaxDuration = 0.0f;
 	InertializationDeficit = 0.0f;
 }
