Reduce log spamming by making the 'scene viewport resized log' verbose.

diff --git a/Engine/Source/Runtime/Engine/Private/Slate/SceneViewport.cpp b/Engine/Source/Runtime/Engine/Private/Slate/SceneViewport.cpp
--- a/Engine/Source/Runtime/Engine/Private/Slate/SceneViewport.cpp
+++ b/Engine/Source/Runtime/Engine/Private/Slate/SceneViewport.cpp
@@ -1598,7 +1598,7 @@ void FSceneViewport::ResizeViewport(uint32 NewSizeX, uint32 NewSizeY, EWindowMod
 		default: WindowModeName = TEXT("INVALID"); break;
 		}
 
-		UE_LOG(LogViewport, Log, TEXT("Scene viewport resized to %dx%d, mode %s."), NewSizeX, NewSizeY, WindowModeName);
+		UE_LOG(LogViewport, Verbose, TEXT("Scene viewport resized to %dx%d, mode %s."), NewSizeX, NewSizeY, WindowModeName);
 
 		bIsResizing = true;
 
