Making it so a sound cue with nothing hooked to it's output isn't played and parsed.

diff --git a/Engine/Source/Runtime/Engine/Private/AudioDevice.cpp b/Engine/Source/Runtime/Engine/Private/AudioDevice.cpp
--- a/Engine/Source/Runtime/Engine/Private/AudioDevice.cpp
+++ b/Engine/Source/Runtime/Engine/Private/AudioDevice.cpp
@@ -5089,7 +5089,8 @@ void FAudioDevice::AddNewActiveSoundInternal(const FActiveSound& InNewActiveSoun
 	}
 
 	USoundBase* Sound = InNewActiveSound.GetSound();
-	if (Sound == nullptr)
+	// A non-playable sound could be a sound cue with nothing hooked to output.
+	if (Sound == nullptr || !Sound->IsPlayable())
 	{
 		ReportSoundFailedToStart(InNewActiveSound.AudioComponentID, InVirtualLoopToRetrigger);
 		return;
