Implement `recompileshaders listtypes` command.

- This dumps a list of all shader types and vertex factory types in the engine.  It is dumped in csv format and includes class name and shader file name.
- This list is useful in combination with other `recompileshaders` commands to compile a single shader (where you need the file name).
- Minor formatting and whitespace fixing.

Example:

Cmd: recompileshaders listtypes
ShaderTypeName, Filename
FHardwareVideoDecodingYCoCgAlphaPS, /Plugin/WmfMedia/Private/MediaHardwareVideoDecoding.usf
FHardwareVideoDecodingYCoCgPS, /Plugin/WmfMedia/Private/MediaHardwareVideoDecoding.usf
FHardwareVideoDecodingY416PS, /Plugin/WmfMedia/Private/MediaHardwareVideoDecoding.usf
...
VertexFactoryTypeName, Filename
TWaterVertexFactory< true>, /Plugin/Water/Private/WaterMeshVertexFactory.ush
TWaterVertexFactory< false>, /Plugin/Water/Private/WaterMeshVertexFactory.ush
FNiagaraSpriteVertexFactory, /Plugin/FX/Niagara/Private/NiagaraSpriteVertexFactory.ush
FNiagaraRibbonVertexFactory, /Plugin/FX/Niagara/Private/NiagaraRibbonVertexFactory.ush
FNiagaraMeshVertexFactory, /Plugin/FX/Niagara/Private/NiagaraMeshVertexFactory.ush
FVectorFieldVisualizationVertexFactory, /Engine/Private/VectorFieldVisualizationVertexFactory.ush

diff --git a/Engine/Source/Runtime/Engine/Private/ShaderCompiler/ShaderCompiler.cpp b/Engine/Source/Runtime/Engine/Private/ShaderCompiler/ShaderCompiler.cpp
--- a/Engine/Source/Runtime/Engine/Private/ShaderCompiler/ShaderCompiler.cpp
+++ b/Engine/Source/Runtime/Engine/Private/ShaderCompiler/ShaderCompiler.cpp
@@ -7966,6 +7966,21 @@ protected:
 
 namespace
 {
+	void ListAllShaderTypes()
+	{
+		UE_LOG(LogShaderCompilers, Display, TEXT("ShaderTypeName, Filename"));
+		for (TLinkedList<FShaderType*>::TIterator It(FShaderType::GetTypeList()); It; It.Next())
+		{
+			UE_LOG(LogShaderCompilers, Display, TEXT("%s, %s "), (*It)->GetName(), (*It)->GetShaderFilename());
+		}
+
+		UE_LOG(LogShaderCompilers, Display, TEXT("VertexFactoryTypeName, Filename"));
+		for (TLinkedList<FVertexFactoryType*>::TIterator It(FVertexFactoryType::GetTypeList()); It; It.Next())
+		{
+			UE_LOG(LogShaderCompilers, Display, TEXT("%s, %s"), (*It)->GetName(), (*It)->GetShaderFilename());
+		}
+	}
+
 	ODSCRecompileCommand ParseRecompileCommandString(const TCHAR* CmdString, TArray<FString>& OutMaterialsToLoad, FString& OutShaderTypesToLoad)
 	{
 		FString CmdName = FParse::Token(CmdString, 0);
@@ -8029,6 +8044,10 @@ namespace
 				OutMaterialsToLoad.Add(It->GetPathName());
 			}
 		}
+		else if (FCString::Stricmp(*CmdName, TEXT("listtypes")))
+		{
+			ListAllShaderTypes();
+		}
 		else
 		{
 			CommandType = ODSCRecompileCommand::SingleShader;
@@ -8183,7 +8202,7 @@ bool RecompileShaders(const TCHAR* Cmd, FOutputDevice& Ar)
 
 #if WITH_EDITOR
 	FString FlagStr(FParse::Token(Cmd, 0));
-	if( FlagStr.Len() > 0 )
+	if (FlagStr.Len() > 0)
 	{
 		TRACE_CPUPROFILER_EVENT_SCOPE(RecompileShaders);
 		GWarn->BeginSlowTask( NSLOCTEXT("ShaderCompilingManager", "BeginRecompilingShadersTask", "Recompiling shaders"), true );
@@ -8192,7 +8211,7 @@ bool RecompileShaders(const TCHAR* Cmd, FOutputDevice& Ar)
 		FlushShaderFileCache();
 		FlushRenderingCommands();
 
-		if( FCString::Stricmp(*FlagStr,TEXT("Changed"))==0)
+		if (FCString::Stricmp(*FlagStr,TEXT("Changed")) == 0)
 		{
 			TArray<const FShaderType*> OutdatedShaderTypes;
 			TArray<const FVertexFactoryType*> OutdatedFactoryTypes;
@@ -8229,12 +8248,12 @@ bool RecompileShaders(const TCHAR* Cmd, FOutputDevice& Ar)
 				UE_LOG(LogShaderCompilers, Warning, TEXT("No Shader changes found."));
 			}
 		}
-		else if( FCString::Stricmp(*FlagStr,TEXT("Global"))==0)
+		else if (FCString::Stricmp(*FlagStr, TEXT("Global")) == 0)
 		{
 			FRecompileShadersTimer TestTimer(TEXT("RecompileShaders Global"));
 			RecompileGlobalShaders();
 		}
-		else if( FCString::Stricmp(*FlagStr,TEXT("Material"))==0)
+		else if (FCString::Stricmp(*FlagStr, TEXT("Material")) == 0)
 		{
 			FString RequestedMaterialName(FParse::Token(Cmd, 0));
 			FRecompileShadersTimer TestTimer(FString::Printf(TEXT("Recompile Material %s"), *RequestedMaterialName));
@@ -8283,7 +8302,7 @@ bool RecompileShaders(const TCHAR* Cmd, FOutputDevice& Ar)
 				UE_LOG(LogShaderCompilers, Warning, TEXT("Couldn't find Material %s!"), *RequestedMaterialName);
 			}
 		}
-		else if( FCString::Stricmp(*FlagStr,TEXT("All"))==0)
+		else if (FCString::Stricmp(*FlagStr, TEXT("All")) == 0)
 		{
 			FRecompileShadersTimer TestTimer(TEXT("RecompileShaders"));
 			RecompileGlobalShaders();
@@ -8304,6 +8323,10 @@ bool RecompileShaders(const TCHAR* Cmd, FOutputDevice& Ar)
 				}
 			}
 		}
+		else if (FCString::Stricmp(*FlagStr, TEXT("listtypes")) == 0)
+		{
+			ListAllShaderTypes();
+		}
 		else
 		{
 			TArray<const FShaderType*> ShaderTypes = FShaderType::GetShaderTypesByFilename(*FlagStr);
@@ -8327,7 +8350,15 @@ bool RecompileShaders(const TCHAR* Cmd, FOutputDevice& Ar)
 		return true;
 	}
 
-	UE_LOG(LogShaderCompilers, Warning, TEXT("Invalid parameter. Options are: \n'Changed', 'Global', 'Material [name]', 'All'."));
+	UE_LOG(LogShaderCompilers, Warning, TEXT("Invalid parameter. \n"
+											 "Options are: \n"
+											 "    'Changed'             Recompile just the shaders that have source file changes.\n"
+											 "    'Global'              Recompile just the global shaders.\n"
+											 "    'Material [name]'     Recompile all the shaders for a single material.\n"
+											 "    'Listtypes'           List all the shader type and vertex factory type class names and their source file path.  Can be used to find shader file names to be used with `recompileshaders [shaderfilename]`.\n"
+											 "    'All'                 Recompile all materials.\n"
+											 "    [shaderfilename]      Compile the single shader associated with a specific filename.\n"
+											 ));
 #endif // WITH_EDITOR
 
 	return true;
