Added static_assert to verify the size of Chaos::TImplicitObjectTransformed

diff --git a/Engine/Source/Runtime/Engine/Private/PhysicsEngine/Experimental/ChaosInterfaceUtils.cpp b/Engine/Source/Runtime/Engine/Private/PhysicsEngine/Experimental/ChaosInterfaceUtils.cpp
--- a/Engine/Source/Runtime/Engine/Private/PhysicsEngine/Experimental/ChaosInterfaceUtils.cpp
+++ b/Engine/Source/Runtime/Engine/Private/PhysicsEngine/Experimental/ChaosInterfaceUtils.cpp
@@ -21,6 +21,9 @@
 #define FORCE_ANALYTICS 0
 #define CREATE_STRAIGHT_CAPSULE_GEOMETRY_FOR_TAPERED_CAPSULES
 
+static_assert(sizeof(Chaos::TImplicitObjectTransformed<Chaos::FReal, 3>) <= 192, 
+	"TImplicitObjectTransformed was optimized to fit into MB3 192 bin to reduce excess memory waste");
+
 namespace ChaosInterface
 {
 	float Chaos_Collision_MarginFraction = -1.0f;
