Avoid potential double regristration of static mesh proxies after moving an actor, this fixes one cause of "Assertion failed: !Primitive->SceneProxy"

diff --git a/Engine/Source/Runtime/Engine/Private/ActorEditor.cpp b/Engine/Source/Runtime/Engine/Private/ActorEditor.cpp
--- a/Engine/Source/Runtime/Engine/Private/ActorEditor.cpp
+++ b/Engine/Source/Runtime/Engine/Private/ActorEditor.cpp
@@ -254,6 +254,9 @@ void AActor::PostEditMove(bool bFinished)
 
 			FNavigationLockContext NavLock(GetWorld(), ENavigationLockReason::AllowUnregister);
 			RerunConstructionScripts();
+			// Construction scripts can have all manner of side effects, including creation of proxies..
+			// remove any static mesh components that have had their proxies created by RerunConstructionScripts:
+			ReregisterContext.SanitizeMeshComponents();
 		}
 	}
 
diff --git a/Engine/Source/Runtime/Engine/Private/StaticMeshResources.cpp b/Engine/Source/Runtime/Engine/Private/StaticMeshResources.cpp
--- a/Engine/Source/Runtime/Engine/Private/StaticMeshResources.cpp
+++ b/Engine/Source/Runtime/Engine/Private/StaticMeshResources.cpp
@@ -189,6 +189,12 @@ void FStaticMeshComponentBulkReregisterContext::AddSimpleConstructionScript(USim
 	}
 }
 
+void FStaticMeshComponentBulkReregisterContext::SanitizeMeshComponents()
+{
+	// for contexts in which side effects are less predictable, e.g. in editor:
+	StaticMeshComponents.RemoveAllSwap([](const UPrimitiveComponent* Component) { return !IsValidChecked(Component) || Component->SceneProxy != nullptr; }, false);
+}
+
 void FStaticMeshComponentBulkReregisterContext::AddConstructedComponent(USceneComponent* SceneComp)
 {
 	UStaticMeshComponent* StaticMeshComponent = Cast<UStaticMeshComponent>(SceneComp);
diff --git a/Engine/Source/Runtime/Engine/Public/StaticMeshResources.h b/Engine/Source/Runtime/Engine/Public/StaticMeshResources.h
--- a/Engine/Source/Runtime/Engine/Public/StaticMeshResources.h
+++ b/Engine/Source/Runtime/Engine/Public/StaticMeshResources.h
@@ -827,6 +827,8 @@ public:
 
 	ENGINE_API void AddSimpleConstructionScript(USimpleConstructionScript* SCS);
 
+	/** Removes any static mesh components that have had their proxy created after the context is created*/
+	void SanitizeMeshComponents();
 private:
 	/** Called by USCS_Node (child of USimpleConstructionScript) to track re-created components */
 	ENGINE_API void AddConstructedComponent(USceneComponent* SceneComp);
