Heterogeneous Volume Component updates

Align component bounds with the unscaled voxel dimensions.

Add animation parameters for Sequencer-driven workflows:
* Frame: Represents the active frame being rendered. It can be edited by the user and keyed in Sequencer or automatically driven when playing an animation.
* Start Frame: Represents the starting frame in an animated range.
* End Frame: Represents the ending frame in an animated range.
* Playing: Enables active playing of the animation data.
* Looping: Enables looping of animation data.

Add sparse volume texture streaming parameters:
* Blocking Requests: Issues blocking requests to the streaming service.
* MIP: Sets the active MIP for streaming requests. This value should correspond to the MIP level accessing the SVT in the shader.

diff --git a/Engine/Source/Runtime/Engine/Classes/Components/HeterogeneousVolumeComponent.h b/Engine/Source/Runtime/Engine/Classes/Components/HeterogeneousVolumeComponent.h
--- a/Engine/Source/Runtime/Engine/Classes/Components/HeterogeneousVolumeComponent.h
+++ b/Engine/Source/Runtime/Engine/Classes/Components/HeterogeneousVolumeComponent.h
@@ -22,24 +22,63 @@ class UHeterogeneousVolumeComponent : public UMeshComponent
 	UPROPERTY(EditAnywhere, BlueprintReadOnly, Category = Volume)
 	FIntVector VolumeResolution;
 
-	UPROPERTY(EditAnywhere, BlueprintReadOnly, Category = Volume)
-	float Framerate;
+	UPROPERTY(EditAnywhere, BlueprintReadWrite, Category = Animation, meta = (EditCondition = "bPlaying == false"))
+	float Frame;
 
-	UPROPERTY(EditAnywhere, BlueprintReadOnly, Category = Volume)
-	uint32 bAnimate : 1;
+	UPROPERTY(EditAnywhere, Category = Animation)
+	float FrameRate;
+
+	UPROPERTY(EditAnywhere, Category = Animation)
+	float StartFrame;
+
+	UPROPERTY(EditAnywhere, Category = Animation)
+	float EndFrame;
+
+	UPROPERTY(EditAnywhere, Category = Animation)
+	uint32 bPlaying : 1;
 
-	UPROPERTY(EditAnywhere, BlueprintReadOnly, Category = Lighting)
+	UPROPERTY(EditAnywhere, Category = Animation)
+	uint32 bLooping : 1;
+
+	UPROPERTY(EditAnywhere, Category = SparseVolumeTextureStreamin, AdvancedDisplay)
+	uint32 MipLevel = 0;
+
+	UPROPERTY(EditAnywhere, Category = SparseVolumeTextureStreaming)
+	uint32 bIssueBlockingRequests : 1;
+
+	UPROPERTY(EditAnywhere, Category = Lighting)
 	float StepFactor;
 
-	UPROPERTY(EditAnywhere, BlueprintReadOnly, Category = Lighting)
+	UPROPERTY(EditAnywhere, Category = Lighting)
 	float ShadowStepFactor;
 
-	UPROPERTY(EditAnywhere, BlueprintReadOnly, Category = Lighting)
+	UPROPERTY(EditAnywhere, Category = Lighting)
 	float LightingDownsampleFactor;
-	
+
 	UPROPERTY(Transient)
 	TObjectPtr<UMaterialInstanceDynamic> MaterialInstanceDynamic;
 
+	UFUNCTION(BlueprintCallable, Category = "Animation")
+	ENGINE_API void SetFrame(float NewValue);
+
+	UFUNCTION(BlueprintCallable, Category = "Animation")
+	ENGINE_API void SetFrameRate(float NewValue);
+
+	UFUNCTION(BlueprintCallable, Category = "Animation")
+	ENGINE_API void SetStartFrame(float NewValue);
+
+	UFUNCTION(BlueprintCallable, Category = "Animation")
+	ENGINE_API void SetEndFrame(float NewValue);
+
+	UFUNCTION(BlueprintCallable, Category = "Animation")
+	ENGINE_API void SetPlaying(bool NewValue);
+
+	UFUNCTION(BlueprintCallable, Category = "Animation")
+	ENGINE_API void SetLooping(bool NewValue);
+
+	UFUNCTION(BlueprintCallable, Category = "Animation")
+	ENGINE_API void Play();
+
 	~UHeterogeneousVolumeComponent() {}
 
 public:
@@ -62,7 +101,10 @@ public:
 	//~ End UPrimitiveComponent Interface.
 
 private:
-	float Time;
+	const USparseVolumeTexture* PreviousSVT;
+
+	const USparseVolumeTexture* GetSparseVolumeTexture() const;
+	void OnSparseVolumeTextureChanged(const USparseVolumeTexture* SparseVolumeTexture);
 };
 
 /**
diff --git a/Engine/Source/Runtime/Engine/Private/Components/HeterogeneousVolumeComponent.cpp b/Engine/Source/Runtime/Engine/Private/Components/HeterogeneousVolumeComponent.cpp
--- a/Engine/Source/Runtime/Engine/Private/Components/HeterogeneousVolumeComponent.cpp
+++ b/Engine/Source/Runtime/Engine/Private/Components/HeterogeneousVolumeComponent.cpp
@@ -76,7 +76,7 @@ FHeterogeneousVolumeSceneProxy::FHeterogeneousVolumeSceneProxy(UHeterogeneousVol
 	VoxelSize.Z /= InComponent->VolumeResolution.Z;
 	HeterogeneousVolumeData.MinimumVoxelSize = FMath::Max(VoxelSize.GetMin(), 0.001);
 
-	if (InComponent->bAnimate && InComponent->MaterialInstanceDynamic)
+	if (InComponent->MaterialInstanceDynamic)
 	{
 		MaterialInterface = InComponent->MaterialInstanceDynamic;
 	}
@@ -226,12 +226,88 @@ UHeterogeneousVolumeComponent::UHeterogeneousVolumeComponent(const FObjectInitia
 
 	MaterialInstanceDynamic = nullptr;
 	VolumeResolution = FIntVector(128);
-	Framerate = 24.0f;
-	bAnimate = false;
+	FrameRate = 24.0f;
+	bPlaying = false;
+	bLooping = false;
+	Frame = 0;
+	StartFrame = 0;
+	EndFrame = 0;
 	StepFactor = 1.0f;
 	ShadowStepFactor = 8.0f;
 	LightingDownsampleFactor = 2.0f;
-	Time = 0.0f;
+	MipLevel = 0;
+	bIssueBlockingRequests = false;
+	PreviousSVT = nullptr;
+}
+
+void UHeterogeneousVolumeComponent::SetFrame(float NewValue)
+{
+	if (AreDynamicDataChangesAllowed()
+		&& Frame != NewValue)
+	{
+		Frame = NewValue;
+		MarkRenderStateDirty();
+	}
+}
+
+void UHeterogeneousVolumeComponent::SetFrameRate(float NewValue)
+{
+	if (AreDynamicDataChangesAllowed()
+		&& FrameRate != NewValue)
+	{
+		FrameRate = NewValue;
+		MarkRenderStateDirty();
+	}
+}
+
+void UHeterogeneousVolumeComponent::SetStartFrame(float NewValue)
+{
+	if (AreDynamicDataChangesAllowed()
+		&& StartFrame != NewValue)
+	{
+		StartFrame = NewValue;
+		MarkRenderStateDirty();
+	}
+}
+
+void UHeterogeneousVolumeComponent::SetEndFrame(float NewValue)
+{
+	if (AreDynamicDataChangesAllowed()
+		&& EndFrame != NewValue)
+	{
+		EndFrame = NewValue;
+		MarkRenderStateDirty();
+	}
+}
+
+void UHeterogeneousVolumeComponent::SetPlaying(bool NewValue)
+{
+	if (AreDynamicDataChangesAllowed()
+		&& bPlaying != NewValue)
+	{
+		bPlaying = NewValue;
+		MarkRenderStateDirty();
+	}
+}
+
+void UHeterogeneousVolumeComponent::SetLooping(bool NewValue)
+{
+	if (AreDynamicDataChangesAllowed()
+		&& bLooping != NewValue)
+	{
+		bLooping = NewValue;
+		MarkRenderStateDirty();
+	}
+}
+
+void UHeterogeneousVolumeComponent::Play()
+{
+	if (AreDynamicDataChangesAllowed())
+	{
+		bPlaying = 1;
+		Frame = 0;
+		MarkRenderStateDirty();
+	}
 }
 
 FPrimitiveSceneProxy* UHeterogeneousVolumeComponent::CreateSceneProxy()
@@ -243,7 +319,7 @@ FBoxSphereBounds UHeterogeneousVolumeComponent::CalcBounds(const FTransform& Loc
 {
 	FBoxSphereBounds NewBounds;
 	NewBounds.Origin = FVector::ZeroVector;
-	NewBounds.BoxExtent = FVector(50, 50, 50);
+	NewBounds.BoxExtent = FVector(VolumeResolution) * 0.5;
 	NewBounds.SphereRadius = NewBounds.BoxExtent.Length();
 
 	return NewBounds.TransformBy(LocalToWorld);
@@ -267,6 +343,48 @@ void UHeterogeneousVolumeComponent::PostLoad()
 	}
 }
 
+const USparseVolumeTexture* UHeterogeneousVolumeComponent::GetSparseVolumeTexture() const
+{
+	USparseVolumeTexture* SparseVolumeTexture = nullptr;
+
+	uint32 MaterialIndex = 0;
+	UMaterialInterface* MaterialInterface = GetMaterial(MaterialIndex);
+	if (MaterialInterface)
+	{
+		TArray<FMaterialParameterInfo> ParameterInfo;
+		TArray<FGuid> ParameterIds;
+		MaterialInterface->GetAllSparseVolumeTextureParameterInfo(ParameterInfo, ParameterIds);
+
+		if (!ParameterInfo.IsEmpty())
+		{
+			MaterialInterface->GetSparseVolumeTextureParameterValue(ParameterInfo[MaterialIndex], SparseVolumeTexture);
+		}
+	}
+
+	return SparseVolumeTexture;
+}
+
+void UHeterogeneousVolumeComponent::OnSparseVolumeTextureChanged(const USparseVolumeTexture* SparseVolumeTexture)
+{
+	if (SparseVolumeTexture)
+	{
+		VolumeResolution = SparseVolumeTexture ? SparseVolumeTexture->GetVolumeResolution() : FIntVector(1);
+		StartFrame = 0;
+		EndFrame = SparseVolumeTexture->GetNumFrames() - 1;
+		Frame = FMath::Clamp(Frame, StartFrame, EndFrame);
+	}
+	else
+	{
+		VolumeResolution = FIntVector(128);
+		Frame = 0.0f;
+		StartFrame = 0.0f;
+		EndFrame = 0.0f;
+	}
+
+	PreviousSVT = SparseVolumeTexture;
+	MarkRenderStateDirty();
+}
+
 #if WITH_EDITOR
 void UHeterogeneousVolumeComponent::PostEditChangeProperty(FPropertyChangedEvent& PropertyChangedEvent)
 {
@@ -280,7 +398,8 @@ void UHeterogeneousVolumeComponent::PostEditChangeProperty(FPropertyChangedEvent
 
 	if (PropertyName == GET_MEMBER_NAME_CHECKED(UHeterogeneousVolumeComponent, OverrideMaterials))
 	{
-		UMaterialInterface* MaterialInterface = GetMaterial(0);
+		uint32 MaterialIndex = 0;
+		UMaterialInterface* MaterialInterface = GetMaterial(MaterialIndex);
 		if (MaterialInterface)
 		{
 			const UMaterial* Material = MaterialInterface->GetMaterial();
@@ -291,6 +410,47 @@ void UHeterogeneousVolumeComponent::PostEditChangeProperty(FPropertyChangedEvent
 			}
 
 			MaterialInstanceDynamic = UMaterialInstanceDynamic::Create(MaterialInterface, nullptr);
+			OnSparseVolumeTextureChanged(GetSparseVolumeTexture());
+		}
+	}
+
+	if (PropertyName == GET_MEMBER_NAME_CHECKED(UHeterogeneousVolumeComponent, VolumeResolution))
+	{
+		// Prevent resolution changes when using SVT
+		const USparseVolumeTexture* SparseVolumeTexture = GetSparseVolumeTexture();
+		if (SparseVolumeTexture)
+		{
+			VolumeResolution = SparseVolumeTexture->GetVolumeResolution();
+		}
+	}
+
+	if (PropertyName == GET_MEMBER_NAME_CHECKED(UHeterogeneousVolumeComponent, Frame))
+	{
+		const USparseVolumeTexture* SparseVolumeTexture = GetSparseVolumeTexture();
+		if (SparseVolumeTexture)
+		{
+			Frame = FMath::Clamp(Frame, StartFrame, EndFrame);
+		}
+	}
+
+	if (PropertyName == GET_MEMBER_NAME_CHECKED(UHeterogeneousVolumeComponent, StartFrame))
+	{
+		const USparseVolumeTexture* SparseVolumeTexture = GetSparseVolumeTexture();
+		if (SparseVolumeTexture)
+		{
+			StartFrame = FMath::Clamp(StartFrame, 0, EndFrame);
+			Frame = FMath::Clamp(Frame, StartFrame, EndFrame);
+		}
+	}
+
+	if (PropertyName == GET_MEMBER_NAME_CHECKED(UHeterogeneousVolumeComponent, EndFrame))
+	{
+		const USparseVolumeTexture* SparseVolumeTexture = GetSparseVolumeTexture();
+		if (SparseVolumeTexture)
+		{
+			const int32 FrameCount = SparseVolumeTexture->GetNumFrames();
+			EndFrame = FMath::Clamp(EndFrame, StartFrame, FrameCount - 1);
+			Frame = FMath::Clamp(Frame, StartFrame, EndFrame);
 		}
 	}
 }
@@ -311,36 +471,64 @@ void UHeterogeneousVolumeComponent::TickComponent(float DeltaTime, ELevelTick Ti
 	Super::TickComponent(DeltaTime, TickType, ThisTickFunction);
 
 	uint32 MaterialIndex = 0;
-	UMaterialInterface* Material = UMeshComponent::GetMaterial(MaterialIndex);
-	if (Material)
+	UMaterialInterface* MaterialInterface = UMeshComponent::GetMaterial(MaterialIndex);
+	if (MaterialInterface)
 	{
 		// Get all SVT params in the material
-		TArray<FMaterialParameterInfo> SVTParameterInfo;
-		TArray<FGuid> SVTParameterIds;
-		Material->GetAllSparseVolumeTextureParameterInfo(SVTParameterInfo, SVTParameterIds);
+		TArray<FMaterialParameterInfo> ParameterInfo;
+		TArray<FGuid> ParameterIds;
+		MaterialInterface->GetAllSparseVolumeTextureParameterInfo(ParameterInfo, ParameterIds);
 
-		if (!SVTParameterInfo.IsEmpty())
+		if (!ParameterInfo.IsEmpty())
 		{
-			USparseVolumeTexture* DefaultSVT = nullptr;
-			bool bValid = Material->GetSparseVolumeTextureParameterDefaultValue(SVTParameterInfo[MaterialIndex], DefaultSVT);
-			if (DefaultSVT)
+			USparseVolumeTexture* SparseVolumeTexture = nullptr;
+			bool bValid = MaterialInterface->GetSparseVolumeTextureParameterValue(ParameterInfo[MaterialIndex], SparseVolumeTexture);
+#if WITH_EDITOR
+			// Detect an update to the material
+			if (SparseVolumeTexture != PreviousSVT)
 			{
-				const int32 FrameCount = DefaultSVT->GetNumFrames();
-				const float FrameIndexF = bAnimate ? FMath::Fmod(Time * Framerate, (float)FrameCount) : 0;
-				const int32 MipLevel = 0;
-				bool bIsBlocking = false;
-				USparseVolumeTextureFrame* SVTFrame = USparseVolumeTextureFrame::GetFrameAndIssueStreamingRequest(DefaultSVT, FrameIndexF, MipLevel, bIsBlocking);
-				VolumeResolution = SVTFrame ? SVTFrame->GetVolumeResolution() : FIntVector(1);
-
-				if (bAnimate && MaterialInstanceDynamic)
+				OnSparseVolumeTextureChanged(SparseVolumeTexture);
+			}
+#endif // WITH_EDITOR
+
+			if (SparseVolumeTexture)
+			{
+				const int32 FrameCount = SparseVolumeTexture->GetNumFrames();
+
+				// Determine active frame based on animation controls if playing
+				if (bPlaying)
 				{
-					MaterialInstanceDynamic->SetSparseVolumeTextureParameterValue(SVTParameterInfo[MaterialIndex].Name, SVTFrame);
+					Frame += DeltaTime * FrameRate;
+				}
+
+				if (bLooping)
+				{
+					float FrameRange = EndFrame - StartFrame + 1;
+					Frame = FMath::Fmod(Frame - StartFrame, (float)FrameRange) + StartFrame;
+				}
+				else
+				{
+					Frame = FMath::Clamp(Frame, StartFrame, EndFrame);
+				}
+
+				bool bIsBlocking = bIssueBlockingRequests != 0;
+				USparseVolumeTextureFrame* SparseVolumeTextureFrame = USparseVolumeTextureFrame::GetFrameAndIssueStreamingRequest(SparseVolumeTexture, Frame, MipLevel, bIsBlocking);
+				if (SparseVolumeTextureFrame)
+				{
+					FIntVector PerFrameVolumeResolution = SparseVolumeTextureFrame->GetVolumeResolution();
+					if (VolumeResolution != PerFrameVolumeResolution)
+					{
+						MarkRenderTransformDirty();
+					}
+				}
+
+				if (MaterialInstanceDynamic)
+				{
+					MaterialInstanceDynamic->SetSparseVolumeTextureParameterValue(ParameterInfo[MaterialIndex].Name, SparseVolumeTextureFrame);
 				}
 			}
 		}
 	}
-
-	Time += DeltaTime;
 }
 
 AHeterogeneousVolume::AHeterogeneousVolume(const FObjectInitializer& ObjectInitializer)
