Adding forward declaration.

diff --git a/Engine/Source/Runtime/Engine/Classes/Components/HeterogeneousVolumeComponent.h b/Engine/Source/Runtime/Engine/Classes/Components/HeterogeneousVolumeComponent.h
--- a/Engine/Source/Runtime/Engine/Classes/Components/HeterogeneousVolumeComponent.h
+++ b/Engine/Source/Runtime/Engine/Classes/Components/HeterogeneousVolumeComponent.h
@@ -11,6 +11,8 @@
 
 #include "HeterogeneousVolumeComponent.generated.h"
 
+class USparseVolumeTexture;
+
 /**
  * A component that represents a heterogeneous volume.
  */
@@ -40,7 +42,7 @@ class UHeterogeneousVolumeComponent : public UMeshComponent
 	UPROPERTY(EditAnywhere, Category = Animation)
 	uint32 bLooping : 1;
 
-	UPROPERTY(EditAnywhere, Category = SparseVolumeTextureStreamin, AdvancedDisplay)
+	UPROPERTY(EditAnywhere, Category = SparseVolumeTextureStreaming, AdvancedDisplay)
 	uint32 MipLevel = 0;
 
 	UPROPERTY(EditAnywhere, Category = SparseVolumeTextureStreaming)
