World Partition Streaming:
- Don't move editor-only actor references to streaming cells when generating streaming.
- Always include editor-only actor references in instancing context, even if actors aren't instanced.

This is for the case where always-loaded objects (moved into the persistent level) have editor-only references to non-editor only actors placed in streaming cells.

diff --git a/Engine/Source/Runtime/Engine/Private/WorldPartition/WorldPartitionLevelHelper.cpp b/Engine/Source/Runtime/Engine/Private/WorldPartition/WorldPartitionLevelHelper.cpp
--- a/Engine/Source/Runtime/Engine/Private/WorldPartition/WorldPartitionLevelHelper.cpp
+++ b/Engine/Source/Runtime/Engine/Private/WorldPartition/WorldPartitionLevelHelper.cpp
@@ -179,6 +179,11 @@ void FWorldPartitionLevelHelper::MoveExternalActorsToLevel(const TArray<FWorldPa
 			continue;
 		}
 
+		if (PackageObjectMapping.bIsEditorOnly)
+		{
+			continue;
+		}
+
 		AActor* Actor = FindObject<AActor>(nullptr, *PackageObjectMapping.LoadedPath.ToString());
 		if (Actor)
 		{
@@ -429,13 +434,14 @@ bool FWorldPartitionLevelHelper::LoadActors(UWorld* InOuterWorld, ULevel* InDest
 		}
 		
 		const FName ContainerPackageInstanceName = Context->RemapPackage(PackageObjectMapping.ContainerPackage);
-		if (PackageObjectMapping.ContainerPackage != ContainerPackageInstanceName)
+
+		if (PackageObjectMapping.bIsEditorOnly || PackageObjectMapping.ContainerPackage != ContainerPackageInstanceName)
 		{
 			const FName ActorPackageName = *FPackageName::ObjectPathToPackageName(PackageObjectMapping.Package.ToString());
 			const FName ActorPackageInstanceName = PackageObjectMapping.bIsEditorOnly ? NAME_None : FName(*ULevel::GetExternalActorPackageInstanceName(ContainerPackageInstanceName.ToString(), ActorPackageName.ToString()));
 
 			Context->AddPackageMapping(ActorPackageName, ActorPackageInstanceName);
-		}
+	}
 
 		if (!PackageObjectMapping.bIsEditorOnly)
 		{
@@ -628,4 +634,4 @@ bool FWorldPartitionLevelHelper::LoadActors(UWorld* InOuterWorld, ULevel* InDest
 	return (LoadProgress->NumPendingLoadRequests == 0);
 }
 
-#endif
\ No newline at end of file
+#endif
