PR #10237: Convenience IsInterfaceValid() blueprint node

diff --git a/Engine/Source/Runtime/Engine/Classes/Kismet/KismetSystemLibrary.h b/Engine/Source/Runtime/Engine/Classes/Kismet/KismetSystemLibrary.h
--- a/Engine/Source/Runtime/Engine/Classes/Kismet/KismetSystemLibrary.h
+++ b/Engine/Source/Runtime/Engine/Classes/Kismet/KismetSystemLibrary.h
@@ -270,6 +270,10 @@ class UKismetSystemLibrary : public UBlueprintFunctionLibrary
 	UFUNCTION(BlueprintPure, meta=(DisplayName = "To Object (Interface)", CompactNodeTitle = "->", BlueprintAutocast), Category="Utilities")
 	static ENGINE_API UObject* Conv_InterfaceToObject(const FScriptInterface& Interface); 
 
+	/** Checks if the backing object is valid. */
+	UFUNCTION(BlueprintPure, Category="Utilities")
+	static bool IsValidInterface(const FScriptInterface& Interface); 
+
 	/** Builds a Soft Object Path struct from a string that contains a full /folder/packagename.object path */
 	UFUNCTION(BlueprintPure, Category = "Utilities", meta = (Keywords = "construct build", NativeMakeFunc, BlueprintThreadSafe, BlueprintAutocast))
 	static ENGINE_API FSoftObjectPath MakeSoftObjectPath(const FString& PathString);
diff --git a/Engine/Source/Runtime/Engine/Private/KismetSystemLibrary.cpp b/Engine/Source/Runtime/Engine/Private/KismetSystemLibrary.cpp
--- a/Engine/Source/Runtime/Engine/Private/KismetSystemLibrary.cpp
+++ b/Engine/Source/Runtime/Engine/Private/KismetSystemLibrary.cpp
@@ -340,6 +340,11 @@ UObject* UKismetSystemLibrary::Conv_InterfaceToObject(const FScriptInterface& In
 	return Interface.GetObject();
 }
 
+bool UKismetSystemLibrary::IsValidInterface(const FScriptInterface& Interface)
+{
+	return IsValid(Interface.GetObject());
+}
+
 void UKismetSystemLibrary::LogString(const FString& InString, bool bPrintToLog)
 {
 #if !(UE_BUILD_SHIPPING || UE_BUILD_TEST) || USE_LOGGING_IN_SHIPPING // Do not Print in Shipping or Test unless explictly enabled.
