World Partition: replace dirty actor list when reinstancing actgors to fix dirty actors garbage collection prevention code using stale pointers.

diff --git a/Engine/Source/Runtime/Engine/Private/WorldPartition/WorldPartition.cpp b/Engine/Source/Runtime/Engine/Private/WorldPartition/WorldPartition.cpp
--- a/Engine/Source/Runtime/Engine/Private/WorldPartition/WorldPartition.cpp
+++ b/Engine/Source/Runtime/Engine/Private/WorldPartition/WorldPartition.cpp
@@ -339,6 +339,21 @@ void UWorldPartition::OnGCPostReachabilityAnalysis()
 	}
 }
 
+void UWorldPartition::OnObjectsReplaced(const TMap<UObject*, UObject*>& OldToNewObjectMap)
+{
+	// Patch up actor pointers in dirty actor list
+	for (auto [OldObject, NewObject] : OldToNewObjectMap)
+	{
+		if (AActor* OldActor = Cast<AActor>(OldObject))
+		{
+			if (FDirtyActor* DirtyActor = DirtyActors.Find(OldActor->GetActorGuid()))
+			{
+				DirtyActor->ActorPtr = CastChecked<AActor>(NewObject);
+			}
+		}
+	}
+}
+
 void UWorldPartition::OnPackageDirtyStateChanged(UPackage* Package)
 {	
 	auto ShouldHandleActor = [this](AActor* Actor)
@@ -917,6 +932,7 @@ void UWorldPartition::RegisterDelegates()
 			FEditorDelegates::CancelPIE.AddUObject(this, &UWorldPartition::OnCancelPIE);
 			FGameDelegates::Get().GetEndPlayMapDelegate().AddUObject(this, &UWorldPartition::OnEndPlay);
 			FCoreUObjectDelegates::PostReachabilityAnalysis.AddUObject(this, &UWorldPartition::OnGCPostReachabilityAnalysis);
+			FCoreUObjectDelegates::OnObjectsReplaced.AddUObject(this, &UWorldPartition::OnObjectsReplaced);
 			GEditor->OnPostBugItGoCalled().AddUObject(this, &UWorldPartition::OnPostBugItGoCalled);
 			GEditor->OnEditorClose().AddUObject(this, &UWorldPartition::SavePerUserSettings);
 			FWorldDelegates::OnPostWorldRename.AddUObject(this, &UWorldPartition::OnWorldRenamed);
@@ -966,6 +982,7 @@ void UWorldPartition::UnregisterDelegates()
 			if (!IsEngineExitRequested())
 			{
 				FCoreUObjectDelegates::PostReachabilityAnalysis.RemoveAll(this);
+				FCoreUObjectDelegates::OnObjectsReplaced.RemoveAll(this);
 			}
 
 			GEditor->OnPostBugItGoCalled().RemoveAll(this);
@@ -1555,10 +1572,7 @@ void UWorldPartition::AddReferencedObjects(UObject* InThis, FReferenceCollector&
 		Collector.AllowEliminatingReferences(false);
 		for (auto& [DirtyActorGuid, DirtyActor] : This->DirtyActors)
 		{
-			if (DirtyActor.WorldPartitionRef.IsSet() || DirtyActor.ActorPtr.IsValid())
-			{
-				Collector.AddReferencedObject(DirtyActor.ActorPtr);
-			}
+			Collector.AddReferencedObject(DirtyActor.ActorPtr);
 		}
 		Collector.AllowEliminatingReferences(true);
 	}
diff --git a/Engine/Source/Runtime/Engine/Public/WorldPartition/WorldPartition.h b/Engine/Source/Runtime/Engine/Public/WorldPartition/WorldPartition.h
--- a/Engine/Source/Runtime/Engine/Public/WorldPartition/WorldPartition.h
+++ b/Engine/Source/Runtime/Engine/Public/WorldPartition/WorldPartition.h
@@ -171,6 +171,7 @@ private:
 	ENGINE_API void SavePerUserSettings();
 		
 	ENGINE_API void OnGCPostReachabilityAnalysis();
+	ENGINE_API void OnObjectsReplaced(const TMap<UObject*, UObject*>& OldToNewObjectMap);
 	ENGINE_API void OnPackageDirtyStateChanged(UPackage* Package);
 
 	// PIE/Game
