Attempting to clarify comments around non-spatialized radius feature.

diff --git a/Engine/Source/Runtime/Engine/Classes/Sound/SoundAttenuation.h b/Engine/Source/Runtime/Engine/Classes/Sound/SoundAttenuation.h
--- a/Engine/Source/Runtime/Engine/Classes/Sound/SoundAttenuation.h
+++ b/Engine/Source/Runtime/Engine/Classes/Sound/SoundAttenuation.h
@@ -273,15 +273,15 @@ struct FSoundAttenuationSettings : public FBaseAttenuationSettings
  	float OmniRadius_DEPRECATED;
 #endif
 
-	/** The distance below which a sound begins to be non-spatialized (2D). This prevents near-field audio from flipping as audio crosses the listener's position. This does not apply when using a 3rd party binaural plugin (audio will remain spatialized). */
+	/** The distance below which a sound begins to linearly interpolate towards being non-spatialized (2D). See "Non Spatialized Radius End" to define the end of the interpolation and the "Non Spatialized Radius Mode" for the mode of the interpolation. Note: this does not apply when using a 3rd party binaural plugin (audio will remain spatialized). */
 	UPROPERTY(EditAnywhere, BlueprintReadWrite, Category = AttenuationSpatialization, meta = (ClampMin = "0", EditCondition = "bSpatialize"))
 	float NonSpatializedRadiusStart;
 
-	/** The distance below which a sound is fully non-spatialized (2D). */
+	/** The distance below which a sound is fully non-spatialized (2D). See "Non Spatialized Radius Start" to define the start of the interpolation and the "Non Spatialized Radius Mode" for the mode of the interpolation. */
 	UPROPERTY(EditAnywhere, BlueprintReadWrite, Category = AttenuationSpatialization, meta = (ClampMin = "0", EditCondition = "bSpatialize"))
 	float NonSpatializedRadiusEnd;
 
-	/** Defines how to blend to the 2D sound when using the non-spatialized radius. */
+	/** Defines how to interpolate a 3D sound towards a 2D sound when using the non-spatialized radius start and end properties. */
 	UPROPERTY(EditAnywhere, BlueprintReadWrite, Category = AttenuationSpatialization, meta = (ClampMin = "0", EditCondition = "bSpatialize"))
 	ENonSpatializedRadiusSpeakerMapMode NonSpatializedRadiusMode;
 
