Add particle component PSO precache events as prerequisite to clearing the reference that's used to check for the end of precaching.

diff --git a/Engine/Source/Runtime/Engine/Private/Particles/ParticleComponents.cpp b/Engine/Source/Runtime/Engine/Private/Particles/ParticleComponents.cpp
--- a/Engine/Source/Runtime/Engine/Private/Particles/ParticleComponents.cpp
+++ b/Engine/Source/Runtime/Engine/Private/Particles/ParticleComponents.cpp
@@ -172,7 +172,7 @@ void UFXSystemAsset::LaunchPSOPrecaching(TArrayView<VFsPerMaterialData> VFsPerMa
 		}
 	}
 
-	// Create task the clear cached events when done
+	// Create task to signal that the PSO precache events are done by adding them as prerequisite to the task.
 	if (PrecachePSOsEvents.Num() > 0)
 	{
 		struct FReleasePrecachePSOsEventTask
@@ -195,7 +195,7 @@ void UFXSystemAsset::LaunchPSOPrecaching(TArrayView<VFsPerMaterialData> VFsPerMa
 		};
 
 		// need to set `PrecachePSOsEvent` before the task is launched to not race with its execution
-		TGraphTask<FReleasePrecachePSOsEventTask>* ReleasePrecachePSOsEventTask = TGraphTask<FReleasePrecachePSOsEventTask>::CreateTask().ConstructAndHold(PrecachePSOsEvent);
+		TGraphTask<FReleasePrecachePSOsEventTask>* ReleasePrecachePSOsEventTask = TGraphTask<FReleasePrecachePSOsEventTask>::CreateTask(&PrecachePSOsEvents).ConstructAndHold(PrecachePSOsEvent);
 		PrecachePSOsEvent = ReleasePrecachePSOsEventTask->GetCompletionEvent();
 		ReleasePrecachePSOsEventTask->Unlock();
 	}
