NetDriver:
* Add cpu scope with net driver definition.

#rnx

diff --git a/Engine/Source/Runtime/Engine/Private/NetDriver.cpp b/Engine/Source/Runtime/Engine/Private/NetDriver.cpp
--- a/Engine/Source/Runtime/Engine/Private/NetDriver.cpp
+++ b/Engine/Source/Runtime/Engine/Private/NetDriver.cpp
@@ -927,6 +927,8 @@ void UNetDriver::TickFlush(float DeltaSeconds)
 
 	CSV_SCOPED_TIMING_STAT_EXCLUSIVE(NetworkOutgoing);
 	SCOPE_CYCLE_COUNTER(STAT_NetTickFlush);
+	TRACE_CPUPROFILER_EVENT_SCOPE_TEXT(*NetDriverDefinition.ToString());
+
 	bool bEnableTimer = (NetDriverName == NAME_GameNetDriver) && ShouldEnableScopeSecondsTimers();
 	if (bEnableTimer)
 	{
