ISM support for partial navigation update

diff --git a/Engine/Source/Runtime/Engine/Classes/Components/InstancedStaticMeshComponent.h b/Engine/Source/Runtime/Engine/Classes/Components/InstancedStaticMeshComponent.h
--- a/Engine/Source/Runtime/Engine/Classes/Components/InstancedStaticMeshComponent.h
+++ b/Engine/Source/Runtime/Engine/Classes/Components/InstancedStaticMeshComponent.h
@@ -508,11 +508,14 @@ protected:
 	/** Terminate all body instances owned by this component. */
 	ENGINE_API void ClearAllInstanceBodies();
 
+	/** Request to navigation system to update for the bounds of the ISM. */
+	ENGINE_API virtual void FullNavigationUpdate();
+
 	/** Request to navigation system to update only part of navmesh occupied by specified instance. */
 	ENGINE_API virtual void PartialNavigationUpdate(int32 InstanceIdx);
 
 	/** Does this component support partial navigation updates */
-	virtual bool SupportsPartialNavigationUpdate() const { return false; }
+	virtual bool SupportsPartialNavigationUpdate() const { return true; }
 
 	/** Internal version of AddInstance */
 	ENGINE_API int32 AddInstanceInternal(int32 InstanceIndex, FInstancedStaticMeshInstanceData* InNewInstanceData, const FTransform& InstanceTransform, bool bWorldSpace);
diff --git a/Engine/Source/Runtime/Engine/Private/InstancedStaticMesh.cpp b/Engine/Source/Runtime/Engine/Private/InstancedStaticMesh.cpp
--- a/Engine/Source/Runtime/Engine/Private/InstancedStaticMesh.cpp
+++ b/Engine/Source/Runtime/Engine/Private/InstancedStaticMesh.cpp
@@ -3628,8 +3628,7 @@ TArray<int32> UInstancedStaticMeshComponent::AddInstancesInternal(TConstArrayVie
 
 	if (!SupportsPartialNavigationUpdate())
 	{
-		// Index parameter is ignored if partial navigation updates are not supported
-		PartialNavigationUpdate(0);
+		FullNavigationUpdate();
 	}
 
 	// Batch update the render state after all instances are finished building
@@ -3946,6 +3945,9 @@ void UInstancedStaticMeshComponent::UpdateInstanceBodyTransform(int32 InstanceIn
 	{
 		if (InstanceBodyInstance)
 		{
+			// Update navigation at current position (before removal)
+			PartialNavigationUpdate(InstanceIndex);
+			
 			// delete BodyInstance
 			InstanceBodyInstance->TermBody();
 			delete InstanceBodyInstance;
@@ -3956,6 +3958,9 @@ void UInstancedStaticMeshComponent::UpdateInstanceBodyTransform(int32 InstanceIn
 	{
 		if (InstanceBodyInstance)
 		{
+			// Update navigation at current position (before applying the transform)
+			PartialNavigationUpdate(InstanceIndex);
+			
 			// Update existing BodyInstance
 			InstanceBodyInstance->SetBodyTransform(WorldSpaceInstanceTransform, TeleportFlagToEnum(bTeleport));
 			InstanceBodyInstance->UpdateBodyScale(WorldSpaceInstanceTransform.GetScale3D());
@@ -3966,6 +3971,9 @@ void UInstancedStaticMeshComponent::UpdateInstanceBodyTransform(int32 InstanceIn
 			InstanceBodyInstance = new FBodyInstance();
 			InitInstanceBody(InstanceIndex, InstanceBodyInstance);
 		}
+
+		// Update navigation at new instance location
+		PartialNavigationUpdate(InstanceIndex);
 	}
 }
 
@@ -3994,9 +4002,6 @@ bool UInstancedStaticMeshComponent::UpdateInstanceTransform(int32 InstanceIndex,
 		UpdateInstanceBodyTransform(InstanceIndex, WorldTransform, bTeleport);
 	}
 
-	// Request navigation update
-	PartialNavigationUpdate(InstanceIndex);
-
 	// Force recreation of the render data when proxy is created
 	InstanceUpdateCmdBuffer.Edit();
 
@@ -4050,13 +4055,8 @@ bool UInstancedStaticMeshComponent::BatchUpdateInstancesTransforms(int32 StartIn
 			FTransform WorldTransform = bWorldSpace ? NewInstanceTransform : (LocalTransform * GetComponentTransform());
 			UpdateInstanceBodyTransform(InstanceIndex, WorldTransform, bTeleport);
 		}
-
-		
 	}
 
-	// Request navigation update - Execute on a single index as it updates everything anyway
-	PartialNavigationUpdate(StartInstanceIndex);
-
 	// Force recreation of the render data when proxy is created
 	InstanceUpdateCmdBuffer.Edit();
 
@@ -4261,8 +4261,7 @@ bool UInstancedStaticMeshComponent::UpdateInstances(
 	check(PerInstanceSMData.Num() == PerInstanceIds.Num());
 	check(PerInstanceSMCustomData.Num() == (NumCustomDataFloats * PerInstanceSMData.Num()));
 
-	// #todo (jnadro) InstanceIndex is ignored by this function so just pass zero.
-	PartialNavigationUpdate(0);
+	FullNavigationUpdate();
 
 	// #todo (jnadro) Updating the PerInstanceRenderData this way currently does not work.
 	// Updating the PerInstanceRenderData from the InstanceUpdateCmdBuffer causes
@@ -4341,9 +4340,6 @@ bool UInstancedStaticMeshComponent::BatchUpdateInstancesTransformsInternal(int32
 		InstanceIndex++;
 	}
 
-	// Request navigation update - Execute on a single index as it updates everything anyway
-	PartialNavigationUpdate(StartInstanceIndex);
-
 	// Force recreation of the render data when proxy is created
 	InstanceUpdateCmdBuffer.Edit();
 
@@ -4384,9 +4380,6 @@ bool UInstancedStaticMeshComponent::BatchUpdateInstancesTransform(int32 StartIns
 		}
 	}
 
-	// Request navigation update - Execute on a single index as it updates everything anyway
-	PartialNavigationUpdate(StartInstanceIndex);
-
 	// Force recreation of the render data when proxy is created
 	InstanceUpdateCmdBuffer.Edit();
 
@@ -4407,12 +4400,12 @@ bool UInstancedStaticMeshComponent::BatchUpdateInstancesData(int32 StartInstance
 
 	Modify();
 
-	for (int32 i = 0; i < NumInstances; ++i)
+	for (int32 Index = 0; Index < NumInstances; ++Index)
 	{
-		int32 InstanceIndex = StartInstanceIndex + i;
+		int32 InstanceIndex = StartInstanceIndex + Index;
 		FInstancedStaticMeshInstanceData& InstanceData = PerInstanceSMData[InstanceIndex];
 
-		InstanceData = StartInstanceData[i];
+		InstanceData = StartInstanceData[Index];
 
 		if (bPhysicsStateCreated)
 		{
@@ -4422,9 +4415,6 @@ bool UInstancedStaticMeshComponent::BatchUpdateInstancesData(int32 StartInstance
 		}
 	}
 
-	// Request navigation update - Execute on a single index as it updates everything anyway
-	PartialNavigationUpdate(StartInstanceIndex);
-
 	// Force recreation of the render data when proxy is created
 	InstanceUpdateCmdBuffer.Edit();
 
@@ -5085,10 +5075,34 @@ void UInstancedStaticMeshComponent::OnPostLoadPerInstanceData()
 	}
 }
 
+void UInstancedStaticMeshComponent::FullNavigationUpdate()
+{
+	FNavigationSystem::UpdateComponentData(*this); // just update everything
+}
+
 void UInstancedStaticMeshComponent::PartialNavigationUpdate(int32 InstanceIdx)
 {
-	// Just update everything
-	FNavigationSystem::UpdateComponentData(*this);
+	if (!IsNavigationRelevant())
+	{
+		return;
+	}
+	
+	if (!InstanceBodies.IsValidIndex(InstanceIdx))
+	{
+		// The physics state might not be created and InstanceBodies might not be populated yet.
+		// This flow can occur from PostEditChangeChainProperty then AddInstanceInternal. 
+		return;
+	}
+
+	const FBodyInstance* const InstanceBodyInstance = InstanceBodies[InstanceIdx];
+	// Not having a body is a valid case when our physics state cannot be created (see CreateAllInstanceBodies)
+	if (!InstanceBodyInstance)
+	{
+		return;
+	}
+	
+	const FBox InstanceBounds = InstanceBodyInstance->GetBodyBounds();
+	FNavigationSystem::OnComponentBoundsChanged(*this, GetNavigationBounds(),InstanceBounds);
 }
 
 bool UInstancedStaticMeshComponent::DoCustomNavigableGeometryExport(FNavigableGeometryExport& GeomExport) const
@@ -5402,7 +5416,8 @@ void UInstancedStaticMeshComponent::PostEditChangeChainProperty(FPropertyChanged
 		}
 		else if (PropertyChangedEvent.Property->GetFName() == GET_MEMBER_NAME_CHECKED(FInstancedStaticMeshInstanceData, Transform))
 		{
-			PartialNavigationUpdate(-1);
+			FullNavigationUpdate();
+
 			// Force recreation of the render data
 			InstanceUpdateCmdBuffer.Edit();
 			MarkRenderStateDirty();
