Addressing volume clamping of modulation params in SoundSubmix

Note: the min/max values of the properties are defined by the audio modulation plugin content settings. This change uses the current min value of volume as the default value of new submix dry param values.

diff --git a/Engine/Source/Runtime/Engine/Classes/Sound/SoundSubmix.h b/Engine/Source/Runtime/Engine/Classes/Sound/SoundSubmix.h
--- a/Engine/Source/Runtime/Engine/Classes/Sound/SoundSubmix.h
+++ b/Engine/Source/Runtime/Engine/Classes/Sound/SoundSubmix.h
@@ -296,15 +296,15 @@ public:
 	float DryLevel;
 
 	/** The output volume of the submix in Decibels. Applied after submix effects and analysis are performed.*/
-	UPROPERTY(EditAnywhere, BlueprintReadWrite, Category = SubmixLevel, meta = (DisplayName = "Output Volume (dB)", AudioParam = "Volume", AudioParamClass = "SoundModulationParameterVolume", ClampMin = "-96.0", ClampMax = "0.0", UIMin = "-96.0", UIMax = "0.0"))
+	UPROPERTY(EditAnywhere, BlueprintReadWrite, Category = SubmixLevel, meta = (DisplayName = "Output Volume (dB)", AudioParam = "Volume", AudioParamClass = "SoundModulationParameterVolume"))
 	FSoundModulationDestinationSettings OutputVolumeModulation;
 
 	/** The wet level of the submixin Decibels. Applied after submix effects and analysis are performed. */
-	UPROPERTY(EditAnywhere, BlueprintReadWrite, Category = SubmixLevel, meta = (DisplayName = "Wet Level (dB)", AudioParam = "Volume", AudioParamClass = "SoundModulationParameterVolume", ClampMin = "-96.0", ClampMax = "0.0", UIMin = "-96.0", UIMax = "0.0"))
+	UPROPERTY(EditAnywhere, BlueprintReadWrite, Category = SubmixLevel, meta = (DisplayName = "Wet Level (dB)", AudioParam = "Volume", AudioParamClass = "SoundModulationParameterVolume"))
 	FSoundModulationDestinationSettings WetLevelModulation;
 
 	/** The dry level of the submix in Decibels. Applied before submix effects and analysis are performed. */
-	UPROPERTY(EditAnywhere, BlueprintReadWrite, Category = SubmixLevel, meta = (DisplayName = "Dry Level (dB)", AudioParam = "Volume", AudioParamClass = "SoundModulationParameterVolume", ClampMin = "-96.0", ClampMax = "0.0", UIMin = "-96.0", UIMax = "0.0"))
+	UPROPERTY(EditAnywhere, BlueprintReadWrite, Category = SubmixLevel, meta = (DisplayName = "Dry Level (dB)", AudioParam = "Volume", AudioParamClass = "SoundModulationParameterVolume"))
 	FSoundModulationDestinationSettings DryLevelModulation;
 	
 	/** Whether to send this Submix to AudioLink (when AudioLink is Enabled)*/
diff --git a/Engine/Source/Runtime/Engine/Private/SoundSubmix.cpp b/Engine/Source/Runtime/Engine/Private/SoundSubmix.cpp
--- a/Engine/Source/Runtime/Engine/Private/SoundSubmix.cpp
+++ b/Engine/Source/Runtime/Engine/Private/SoundSubmix.cpp
@@ -37,6 +37,19 @@ USoundSubmixBase::USoundSubmixBase(const FObjectInitializer& ObjectInitializer)
 #endif // WITH_EDITORONLY_DATA
 {}
 
+static float GetMinVolumeParameter()
+{
+	if (Audio::IsModulationParameterRegistered("Volume"))
+	{
+		const Audio::FModulationParameter& VolumeParam = Audio::GetModulationParameter(TEXT("Volume"));
+		return VolumeParam.MinValue;
+	}
+	else
+	{
+		return -60.0f;
+	}
+}
+
 USoundSubmix::USoundSubmix(const FObjectInitializer& ObjectInitializer)
 	: Super(ObjectInitializer)
 	, bMuteWhenBackgrounded(0)
@@ -49,7 +62,9 @@ USoundSubmix::USoundSubmix(const FObjectInitializer& ObjectInitializer)
 {
 	OutputVolumeModulation.Value = 0.f;
 	WetLevelModulation.Value = 0.f;
-	DryLevelModulation.Value = -96.f;
+
+	DryLevelModulation.Value = GetMinVolumeParameter();
+
 }
 
 void USoundSubmix::Serialize(FArchive& Ar)
@@ -67,7 +82,7 @@ void USoundSubmix::Serialize(FArchive& Ar)
 		{
 			if (OutputVolume <= LinearNeg96dB)
 			{
-				OutputVolumeModulation.Value = -96.f;
+				OutputVolumeModulation.Value = GetMinVolumeParameter();
 			}
 			else
 			{
@@ -80,7 +95,7 @@ void USoundSubmix::Serialize(FArchive& Ar)
 		{
 			if (WetLevel <= LinearNeg96dB)
 			{
-				WetLevelModulation.Value = -96.f;
+				WetLevelModulation.Value = GetMinVolumeParameter();
 			}
 			else
 			{
@@ -93,7 +108,7 @@ void USoundSubmix::Serialize(FArchive& Ar)
 		{
 			if (DryLevel <= LinearNeg96dB)
 			{
-				DryLevelModulation.Value = -96.f;
+				DryLevelModulation.Value = GetMinVolumeParameter();
 			}
 			else
 			{
