Fixing dynamic bus sends

diff --git a/Engine/Source/Runtime/AudioMixer/Private/AudioMixerSource.cpp b/Engine/Source/Runtime/AudioMixer/Private/AudioMixerSource.cpp
--- a/Engine/Source/Runtime/AudioMixer/Private/AudioMixerSource.cpp
+++ b/Engine/Source/Runtime/AudioMixer/Private/AudioMixerSource.cpp
@@ -1854,16 +1854,18 @@ namespace Audio
 
 		SetupBusData();
 
-		if (!bSendingAudioToBuses)
+		FActiveSound* ActiveSound = WaveInstance->ActiveSound;
+		check(ActiveSound);
+
+		// Check if the user actively called a function that alters bus sends since the last update
+		bool bHasNewBusSends = ActiveSound->HasNewBusSends();
+
+		if (!bSendingAudioToBuses && !bHasNewBusSends && !DynamicBusSendInfos.Num())
 		{
 			return;
 		}
 
-		//If the user actively called a function that alters bus sends since the last update
-		FActiveSound* ActiveSound = WaveInstance->ActiveSound;
-		check(ActiveSound);
-
-		if (ActiveSound->HasNewBusSends())
+		if (bHasNewBusSends)
 		{
 			TArray<TTuple<EBusSendType, FSoundSourceBusSendInfo>> NewBusSends = ActiveSound->GetNewBusSends();
 			for (TTuple<EBusSendType, FSoundSourceBusSendInfo>& NewSend : NewBusSends)
@@ -1871,11 +1873,13 @@ namespace Audio
 				if (NewSend.Value.SoundSourceBus)
 				{
 					MixerSourceVoice->SetAudioBusSendInfo(NewSend.Key, NewSend.Value.SoundSourceBus->GetUniqueID(), NewSend.Value.SendLevel);
+					bSendingAudioToBuses = true;
 				}
 
 				if (NewSend.Value.AudioBus)
 				{
 					MixerSourceVoice->SetAudioBusSendInfo(NewSend.Key, NewSend.Value.AudioBus->GetUniqueID(), NewSend.Value.SendLevel);
+					bSendingAudioToBuses = true;
 				}
 			}
 
diff --git a/Engine/Source/Runtime/Engine/Classes/Components/AudioComponent.h b/Engine/Source/Runtime/Engine/Classes/Components/AudioComponent.h
--- a/Engine/Source/Runtime/Engine/Classes/Components/AudioComponent.h
+++ b/Engine/Source/Runtime/Engine/Classes/Components/AudioComponent.h
@@ -917,6 +917,17 @@ private:
 	FSharedISourceBufferListenerPtr SourceBufferListener;
 	bool bShouldSourceBufferListenerZeroBuffer = false;
 
+	/** Pending submix and bus sends. */
+	TArray<FSoundSubmixSendInfo> PendingSubmixSends;
+
+	struct FPendingSourceBusSendInfo
+	{
+		EBusSendType BusSendType = EBusSendType::PreEffect;
+		FSoundSourceBusSendInfo BusSendInfo;
+	};
+
+	TArray<FPendingSourceBusSendInfo> PendingBusSends;
+
 protected:
 
 	/** Utility function called by Play and FadeIn to start a sound playing. */
diff --git a/Engine/Source/Runtime/Engine/Private/Components/AudioComponent.cpp b/Engine/Source/Runtime/Engine/Private/Components/AudioComponent.cpp
--- a/Engine/Source/Runtime/Engine/Private/Components/AudioComponent.cpp
+++ b/Engine/Source/Runtime/Engine/Private/Components/AudioComponent.cpp
@@ -816,6 +816,19 @@ void UAudioComponent::PlayInternal(const PlayInternalRequestData& InPlayRequestD
 
 	NewActiveSound.MaxDistance = MaxDistance;
 
+	// Setup the submix and bus sends that may have been set before playing
+	for (FSoundSubmixSendInfo& SubmixSendInfo : PendingSubmixSends)
+	{
+		NewActiveSound.SetSubmixSend(SubmixSendInfo);
+	}
+	PendingSubmixSends.Reset();
+
+	for (FPendingSourceBusSendInfo& PendingBusSend : PendingBusSends)
+	{
+		NewActiveSound.SetSourceBusSend(PendingBusSend.BusSendType, PendingBusSend.BusSendInfo);
+	}
+	PendingBusSends.Reset();
+
 	Audio::FVolumeFader& Fader = NewActiveSound.ComponentVolumeFader;
 	Fader.SetVolume(0.0f); // Init to 0.0f to fade as default is 1.0f
 	Fader.StartFade(InPlayRequestData.FadeVolumeLevel, InPlayRequestData.FadeInDuration, static_cast<Audio::EFaderCurve>(InPlayRequestData.FadeCurve));
@@ -1441,36 +1454,52 @@ void UAudioComponent::AdjustAttenuation(const FSoundAttenuationSettings& InAtten
 
 void UAudioComponent::SetSubmixSend(USoundSubmixBase* Submix, float SendLevel)
 {
-	if (FAudioDevice* AudioDevice = GetAudioDevice())
-	{
-		const uint64 MyAudioComponentID = AudioComponentID;
-
-		FSoundSubmixSendInfo SendInfo;
-		SendInfo.SoundSubmix = Submix;
-		SendInfo.SendLevel = SendLevel;
+	FSoundSubmixSendInfo SendInfo;
+	SendInfo.SoundSubmix = Submix;
+	SendInfo.SendLevel = SendLevel;
 
-		DECLARE_CYCLE_STAT(TEXT("FAudioThreadTask.AudioSetSubmixSend"), STAT_SetSubmixSend, STATGROUP_AudioThreadCommands);
-		AudioDevice->SendCommandToActiveSounds(AudioComponentID, [SendInfo](FActiveSound& ActiveSound)
+	if (IsPlaying())
+	{
+		if (FAudioDevice* AudioDevice = GetAudioDevice())
 		{
-			ActiveSound.SetSubmixSend(SendInfo);
-		}, GET_STATID(STAT_SetSubmixSend));
+			DECLARE_CYCLE_STAT(TEXT("FAudioThreadTask.AudioSetSubmixSend"), STAT_SetSubmixSend, STATGROUP_AudioThreadCommands);
+			AudioDevice->SendCommandToActiveSounds(AudioComponentID, [SendInfo](FActiveSound& ActiveSound)
+			{
+				ActiveSound.SetSubmixSend(SendInfo);
+			}, GET_STATID(STAT_SetSubmixSend));
+		}
+	}
+	else
+	{
+		PendingSubmixSends.Add(SendInfo);
 	}
 }
 
 void UAudioComponent::SetBusSendEffectInternal(USoundSourceBus* InSourceBus, UAudioBus* InAudioBus, float SendLevel, EBusSendType InBusSendType)
 {
-	if (FAudioDevice* AudioDevice = GetAudioDevice())
-	{
+	FSoundSourceBusSendInfo BusSendInfo;
+	BusSendInfo.SoundSourceBus = InSourceBus;
+	BusSendInfo.AudioBus = InAudioBus;
+	BusSendInfo.SendLevel = SendLevel;
 
-		FSoundSourceBusSendInfo SourceBusSendInfo;
-		SourceBusSendInfo.SoundSourceBus = InSourceBus;
-		SourceBusSendInfo.AudioBus = InAudioBus;
-		SourceBusSendInfo.SendLevel = SendLevel;
-
-		AudioDevice->SendCommandToActiveSounds(AudioComponentID, [InBusSendType, SourceBusSendInfo](FActiveSound& ActiveSound)
+	if (IsPlaying())
+	{
+		if (FAudioDevice* AudioDevice = GetAudioDevice())
 		{
-			ActiveSound.SetSourceBusSend(InBusSendType, SourceBusSendInfo);
-		});
+			DECLARE_CYCLE_STAT(TEXT("FAudioThreadTask.AudioSetBusSend"), STAT_SetBusSend, STATGROUP_AudioThreadCommands);
+			AudioDevice->SendCommandToActiveSounds(AudioComponentID, [InBusSendType, BusSendInfo](FActiveSound& ActiveSound)
+			{
+				ActiveSound.SetSourceBusSend(InBusSendType, BusSendInfo);
+			}, GET_STATID(STAT_SetBusSend));
+		}
+	}
+	else
+	{
+		FPendingSourceBusSendInfo NewPendingSourceBusInfo;
+		NewPendingSourceBusInfo.BusSendInfo = BusSendInfo;
+		NewPendingSourceBusInfo.BusSendType = InBusSendType;
+
+		PendingBusSends.Add(NewPendingSourceBusInfo);
 	}
 }
 
