Moved CommandletHelpers to Engine, to be alongside UCommandlet itself

#rnx

diff --git a/Engine/Source/Editor/UnrealEd/Private/Commandlets/CommandletHelpers.cpp b/Engine/Source/Editor/UnrealEd/Private/Commandlets/CommandletHelpers.cpp
deleted file mode 100644
diff --git a/Engine/Source/Editor/UnrealEd/Public/Commandlets/CommandletHelpers.h b/Engine/Source/Editor/UnrealEd/Public/Commandlets/CommandletHelpers.h
deleted file mode 100644
diff --git a/Engine/Source/Runtime/Engine/Private/Commandlets/CommandletHelpers.cpp b/Engine/Source/Runtime/Engine/Private/Commandlets/CommandletHelpers.cpp
new file mode 100644
--- /dev/null
+++ b/Engine/Source/Runtime/Engine/Private/Commandlets/CommandletHelpers.cpp
@@ -0,0 +1,43 @@
+// Copyright Epic Games, Inc. All Rights Reserved.
+
+#include "Commandlets/CommandletHelpers.h"
+
+#include "Misc/AssertionMacros.h"
+#include "Misc/CString.h"
+
+namespace CommandletHelpers
+{
+	FString BuildCommandletProcessArguments(const TCHAR* const CommandletName, const TCHAR* const ProjectPath, const TCHAR* const AdditionalArguments)
+	{
+		check(CommandletName && FCString::Strlen(CommandletName) > 0);
+
+		FString Arguments;
+
+		if (ProjectPath && FCString::Strlen(ProjectPath) > 0)
+		{
+			if (!Arguments.IsEmpty())
+			{
+				Arguments += TEXT(" ");
+			}
+			Arguments += ProjectPath;
+		}
+
+		if (!Arguments.IsEmpty())
+		{
+			Arguments += TEXT(" ");
+		}
+		Arguments += TEXT("-run=");
+		Arguments += CommandletName;
+
+		if (AdditionalArguments && FCString::Strlen(AdditionalArguments) > 0)
+		{
+			if (!Arguments.IsEmpty())
+			{
+				Arguments += TEXT(" ");
+			}
+			Arguments += AdditionalArguments;
+		}
+
+		return Arguments;
+	}
+}
diff --git a/Engine/Source/Runtime/Engine/Public/Commandlets/CommandletHelpers.h b/Engine/Source/Runtime/Engine/Public/Commandlets/CommandletHelpers.h
new file mode 100644
--- /dev/null
+++ b/Engine/Source/Runtime/Engine/Public/Commandlets/CommandletHelpers.h
@@ -0,0 +1,10 @@
+// Copyright Epic Games, Inc. All Rights Reserved.
+
+#pragma once
+
+#include "Containers/UnrealString.h"
+
+namespace CommandletHelpers
+{
+	ENGINE_API FString BuildCommandletProcessArguments(const TCHAR* const CommandletName, const TCHAR* const ProjectPath, const TCHAR* const AdditionalArguments);
+}
