AudioComponent: Add missing #includes

#rnx

diff --git a/Engine/Source/Runtime/Engine/Classes/Components/AudioComponent.h b/Engine/Source/Runtime/Engine/Classes/Components/AudioComponent.h
--- a/Engine/Source/Runtime/Engine/Classes/Components/AudioComponent.h
+++ b/Engine/Source/Runtime/Engine/Classes/Components/AudioComponent.h
@@ -1,6 +1,7 @@
 // Copyright Epic Games, Inc. All Rights Reserved.
 #pragma once
 
+#include "Audio.h"
 #include "Audio/SoundParameterControllerInterface.h"
 #include "Components/SceneComponent.h"
 #include "CoreMinimal.h"
@@ -12,6 +13,8 @@
 #include "Sound/QuartzSubscription.h"
 #include "Sound/SoundAttenuation.h"
 #include "Sound/SoundModulationDestination.h"
+#include "Sound/SoundSubmixSend.h"
+#include "Sound/SoundSourceBusSend.h"
 #if UE_ENABLE_INCLUDE_ORDER_DEPRECATED_IN_5_2
 #include "Sound/SoundWave.h"
 #include "Quartz/AudioMixerClockHandle.h"
@@ -30,7 +33,6 @@ class USoundClass;
 class USoundConcurrency;
 class USoundEffectSourcePresetChain;
 class USoundWave;
-enum class EBusSendType : uint8;
 struct FAudioComponentParam;
 using FSharedISourceBufferListenerPtr = TSharedPtr<ISourceBufferListener, ESPMode::ThreadSafe>;
 
