Delete FComponentRecreateRenderStateContext copy constructor and copy assignment operator. Define move constructor and move assignment operator that clears the input object instead of leaving the data in it.

This prevents a crash when executing ALandscapeProxy::UpdateAllComponentMaterialInstances, which was resulting in copies of FComponentRecreateRenderStateContext getting destroyed at inappropriate times.

diff --git a/Engine/Source/Runtime/Engine/Public/ComponentRecreateRenderStateContext.h b/Engine/Source/Runtime/Engine/Public/ComponentRecreateRenderStateContext.h
--- a/Engine/Source/Runtime/Engine/Public/ComponentRecreateRenderStateContext.h
+++ b/Engine/Source/Runtime/Engine/Public/ComponentRecreateRenderStateContext.h
@@ -12,9 +12,9 @@ class FComponentRecreateRenderStateContext
 {
 private:
 	/** Pointer to component we are recreating render state for */
-	UActorComponent* Component;
+	UActorComponent* Component = nullptr;
 
-	TSet<FSceneInterface*>* ScenesToUpdateAllPrimitiveSceneInfos;
+	TSet<FSceneInterface*>* ScenesToUpdateAllPrimitiveSceneInfos = nullptr;
 
 public:
 	FComponentRecreateRenderStateContext(UActorComponent* InComponent, TSet<FSceneInterface*>* InScenesToUpdateAllPrimitiveSceneInfos = nullptr)
@@ -36,6 +36,26 @@ public:
 		}
 	}
 
+	FComponentRecreateRenderStateContext(const FComponentRecreateRenderStateContext&) = delete;
+	FComponentRecreateRenderStateContext& operator=(const FComponentRecreateRenderStateContext&) = delete;
+	
+	FComponentRecreateRenderStateContext(FComponentRecreateRenderStateContext&& Other)
+		: Component(Other.Component)
+		, ScenesToUpdateAllPrimitiveSceneInfos(Other.ScenesToUpdateAllPrimitiveSceneInfos)
+	{
+		Other.Component = nullptr;
+		Other.ScenesToUpdateAllPrimitiveSceneInfos = nullptr;
+	}
+
+	FComponentRecreateRenderStateContext& operator=(FComponentRecreateRenderStateContext&& Other)
+	{
+		Component = Other.Component;
+		ScenesToUpdateAllPrimitiveSceneInfos = Other.ScenesToUpdateAllPrimitiveSceneInfos;
+		Other.Component = nullptr;
+		Other.ScenesToUpdateAllPrimitiveSceneInfos = nullptr;
+		return *this;
+	}
+
 	~FComponentRecreateRenderStateContext()
 	{
 		if (Component && !Component->IsRenderStateCreated() && Component->IsRegistered())
