Enable High Priority Anim Completion Tasks

diff --git a/Engine/Source/Runtime/Engine/Private/Components/SkeletalMeshComponent.cpp b/Engine/Source/Runtime/Engine/Private/Components/SkeletalMeshComponent.cpp
--- a/Engine/Source/Runtime/Engine/Private/Components/SkeletalMeshComponent.cpp
+++ b/Engine/Source/Runtime/Engine/Private/Components/SkeletalMeshComponent.cpp
@@ -91,7 +91,7 @@ DEFINE_STAT(STAT_PostAnimEvaluation);
 CSV_DECLARE_CATEGORY_MODULE_EXTERN(ENGINE_API, Animation);
 CSV_DECLARE_CATEGORY_MODULE_EXTERN(CORE_API, Basic);
 
-static bool GParallelAnimCompletionTaskHighPriority = false;
+static bool GParallelAnimCompletionTaskHighPriority = true;
 static FAutoConsoleVariableRef CVarParallelAnimCompletionTaskHighPriority(
 	TEXT("TaskGraph.TaskPriorities.ParallelAnimCompletionTaskHighPriority"),
 	GParallelAnimCompletionTaskHighPriority,
