Fixed failing check on a UAnimMontage::TickAssetPlayer() causing crashes when the montage position was being set and caused a loop.

The ensures that fails checks that the sync positions (prev and next) stay the same if no sync markers have been passed during that tick, however, this assumption cannot be made if the montage instance's FMakerTickRecord is invalid, which happens when the montage instance's position is modified via the setter.

diff --git a/Engine/Source/Runtime/Engine/Private/Animation/AnimMontage.cpp b/Engine/Source/Runtime/Engine/Private/Animation/AnimMontage.cpp
--- a/Engine/Source/Runtime/Engine/Private/Animation/AnimMontage.cpp
+++ b/Engine/Source/Runtime/Engine/Private/Animation/AnimMontage.cpp
@@ -1208,9 +1208,10 @@ void UAnimMontage::TickAssetPlayer(FAnimTickRecord& Instance, struct FAnimNotify
 			{
 				FMarkerTickRecord* MarkerTickRecord = Instance.MarkerTickRecord;
 				FMarkerTickContext& MarkerTickContext = Context.MarkerTickContext;
-
+				const bool bIsMarkerTickRecordValid = MarkerTickRecord->IsValid(Instance.bLooping);
+				
 				// Store the sync anim position BEFORE the asset has being ticked.
-				if (MarkerTickRecord->IsValid(Instance.bLooping))
+				if (bIsMarkerTickRecordValid)
 				{
 					MarkerTickContext.SetMarkerSyncStartPosition(GetMarkerSyncPositionFromMarkerIndicies(MarkerTickRecord->PreviousMarker.MarkerIndex, MarkerTickRecord->NextMarker.MarkerIndex, PreviousTime, nullptr));
 				}
@@ -1235,7 +1236,9 @@ void UAnimMontage::TickAssetPlayer(FAnimTickRecord& Instance, struct FAnimNotify
 				MarkerTickContext.MarkersPassedThisTick = *Instance.Montage.MarkersPassedThisTick;
 
 #if DO_CHECK
-				if(MarkerTickContext.MarkersPassedThisTick.Num() == 0)
+				// The marker tick record gets invalidated when the montage position is set externally and due to this change we cannot assume
+				// its sync positions will be the same as the previous tick. 
+				if (MarkerTickContext.MarkersPassedThisTick.Num() == 0 && bIsMarkerTickRecordValid)
 				{
 					const FMarkerSyncAnimPosition& StartPosition = MarkerTickContext.GetMarkerSyncStartPosition();
 					const FMarkerSyncAnimPosition& EndPosition = MarkerTickContext.GetMarkerSyncEndPosition();
