[Texture Streaming] Fix desynchronization between CachedSRRState and the renderer resource when a petition is canceled during the FinalizeNewMips step.

#rnx

diff --git a/Engine/Source/Runtime/Engine/Private/Streaming/TextureStreamIn.cpp b/Engine/Source/Runtime/Engine/Private/Streaming/TextureStreamIn.cpp
--- a/Engine/Source/Runtime/Engine/Private/Streaming/TextureStreamIn.cpp
+++ b/Engine/Source/Runtime/Engine/Private/Streaming/TextureStreamIn.cpp
@@ -386,11 +386,7 @@ void FTextureStreamIn::FinalizeNewMips(const FContext& Context)
 	DECLARE_SCOPE_CYCLE_COUNTER(TEXT("FTextureStreamIn::FinalizeNewMips"), STAT_TextureStreamIn_FinalizeNewMips, STATGROUP_StreamingDetails);
 
 	// Execute
-	if (DoFinalizeNewMips(Context))
-	{
-		MarkAsSuccessfullyFinished();
-	}
-	else
+	if (!DoFinalizeNewMips(Context))
 	{
 		MarkAsCancelled();
 	}
@@ -404,6 +400,9 @@ void FTextureStreamIn::FinalizeNewMips(const FContext& Context)
 	}
 	else if (!IsCancelled())
 	{
+		// Finalize new mips completed successfully.
+		MarkAsSuccessfullyFinished();
+		
 		// Release the mip allocator.
 		MipAllocator.Reset();
 		MipInfos.Empty();
@@ -421,6 +420,12 @@ void FTextureStreamIn::FinalizeNewMips(const FContext& Context)
 	}
 	else
 	{
+		// Mark as successfully finished even if cancelled to keep the CachedSRRState in sync with the RHI texture.
+		if (MipAllocator->GetNextTickState() == FTextureMipAllocator::ETickState::Done)
+		{
+			MarkAsSuccessfullyFinished();
+		}
+
 		PushTask(Context, TT_None, nullptr, GetCancelThread(), SRA_UPDATE_CALLBACK(Cancel));
 	}
 }
