Added the 'concept' of common frame-rate to Animation Montages/Composites
- On load/change it tries to work out the highest compatible frame-rate across all the contained animation assets, and uses it as the sampling rate for the montage/composite itself
- If one or more assets are not compatible the sampling rate reverts back to 100000hz, and the user is informed about it
- Project wide option to enforce this, to be extended to other systems later, which will prevent the user from adding incompatible assets to a montage/composite
- CVar to mute warning on-load, useful for upgrading content without breaking validation

diff --git a/Engine/Source/Editor/Persona/Private/AnimTimeline/AnimModel.cpp b/Engine/Source/Editor/Persona/Private/AnimTimeline/AnimModel.cpp
--- a/Engine/Source/Editor/Persona/Private/AnimTimeline/AnimModel.cpp
+++ b/Engine/Source/Editor/Persona/Private/AnimTimeline/AnimModel.cpp
@@ -14,10 +14,10 @@
 const FAnimModel::FSnapType FAnimModel::FSnapType::Frames("Frames", LOCTEXT("FramesSnapName", "Frames"), [](const FAnimModel& InModel, double InTime)
 {
 	// Round to nearest frame
-	double FrameRate = InModel.GetFrameRate();
-	if(FrameRate > 0)
+	FFrameRate FrameRate = InModel.GetFrameRate();
+	if(FrameRate.IsValid())
 	{
-		return FMath::RoundToDouble(InTime * FrameRate) / FrameRate;
+		return (double)FrameRate.AsFrameNumber(InTime).Value;
 	}
 	
 	return InTime;
@@ -52,21 +52,14 @@ FAnimatedRange FAnimModel::GetWorkingRange() const
 	return WorkingRange;
 }
 
-double FAnimModel::GetFrameRate() const
+FFrameRate FAnimModel::GetFrameRate() const
 {
-	if(UAnimSequence* AnimSequence = Cast<UAnimSequence>(GetAnimSequenceBase()))
-	{
-		return AnimSequence->GetSamplingFrameRate().AsDecimal();
-	}
-	else
-	{
-		return 30.0;
-	}
+	return GetAnimSequenceBase()->GetSamplingFrameRate();
 }
 
 int32 FAnimModel::GetTickResolution() const
 {
-	return FMath::RoundToInt32((double)GetDefault<UPersonaOptions>()->TimelineScrubSnapValue * GetFrameRate());
+	return FMath::RoundToInt32((double)GetDefault<UPersonaOptions>()->TimelineScrubSnapValue * GetFrameRate().AsDecimal());
 }
 
 TRange<FFrameNumber> FAnimModel::GetPlaybackRange() const
diff --git a/Engine/Source/Editor/Persona/Private/AnimTimeline/AnimModel.h b/Engine/Source/Editor/Persona/Private/AnimTimeline/AnimModel.h
--- a/Engine/Source/Editor/Persona/Private/AnimTimeline/AnimModel.h
+++ b/Engine/Source/Editor/Persona/Private/AnimTimeline/AnimModel.h
@@ -164,7 +164,7 @@ public:
 	virtual void OnSetEditableTime(int32 TimeIndex, double Time, bool bIsDragging) {}
 
 	/** Get the framerate specified by the anim sequence */
-	double GetFrameRate() const;
+	FFrameRate GetFrameRate() const;
 
 	/** Get the tick resolution we are displaying at */
 	int32 GetTickResolution() const;
diff --git a/Engine/Source/Editor/Persona/Private/AnimTimeline/AnimTimelineTrack.cpp b/Engine/Source/Editor/Persona/Private/AnimTimeline/AnimTimelineTrack.cpp
--- a/Engine/Source/Editor/Persona/Private/AnimTimeline/AnimTimelineTrack.cpp
+++ b/Engine/Source/Editor/Persona/Private/AnimTimeline/AnimTimelineTrack.cpp
@@ -184,7 +184,8 @@ float FAnimTimelineTrack::GetViewMaxInput() const
 
 float FAnimTimelineTrack::GetScrubValue() const
 {
-	const int64 Resolution = FMath::RoundToInt(static_cast<double>(GetDefault<UPersonaOptions>()->TimelineScrubSnapValue) * GetModel()->GetFrameRate());
+	const FFrameRate FrameRate = GetModel()->GetFrameRate();
+	const int64 Resolution = FMath::RoundToInt(static_cast<double>(GetDefault<UPersonaOptions>()->TimelineScrubSnapValue) * FrameRate.AsDecimal());
 	return static_cast<float>(static_cast<double>(GetModel()->GetScrubPosition().Value) / static_cast<double>(Resolution));
 }
 
diff --git a/Engine/Source/Editor/Persona/Private/AnimTimeline/SAnimTimeline.cpp b/Engine/Source/Editor/Persona/Private/AnimTimeline/SAnimTimeline.cpp
--- a/Engine/Source/Editor/Persona/Private/AnimTimeline/SAnimTimeline.cpp
+++ b/Engine/Source/Editor/Persona/Private/AnimTimeline/SAnimTimeline.cpp
@@ -52,8 +52,6 @@ void SAnimTimeline::Construct(const FArguments& InArgs, const TSharedRef<FAnimMo
 	Model = InModel;
 	OnReceivedFocus = InArgs._OnReceivedFocus;
 
-	int32 TickResolutionValue = InModel->GetTickResolution();
-	int32 SequenceFrameRate = static_cast<int32>(FMath::RoundToInt(InModel->GetFrameRate()));
 
 	if (InModel->GetPreviewScene()->GetPreviewMeshComponent()->PreviewInstance)
 	{
@@ -72,14 +70,14 @@ void SAnimTimeline::Construct(const FArguments& InArgs, const TSharedRef<FAnimMo
 		return GetDefault<UPersonaOptions>()->TimelineDisplayFormat == EFrameNumberDisplayFormats::Frames ? EFrameNumberDisplayFormats::Seconds : EFrameNumberDisplayFormats::Frames;
 	});
 
-	const TAttribute<FFrameRate> TickResolution = MakeAttributeLambda([TickResolutionValue]()
+	const TAttribute<FFrameRate> TickResolution = MakeAttributeLambda([this]()
 	{
-		return FFrameRate(TickResolutionValue, 1);
+		return FFrameRate(Model.Pin()->GetTickResolution(), 1);
 	});
 
-	const TAttribute<FFrameRate> DisplayRate = MakeAttributeLambda([SequenceFrameRate]()
+	const TAttribute<FFrameRate> DisplayRate = MakeAttributeLambda([this]()
 	{
-		return FFrameRate(SequenceFrameRate, 1);
+		return Model.Pin()->GetFrameRate();
 	});
 
 	// Create our numeric type interface so we can pass it to the time slider below.
@@ -767,7 +765,7 @@ bool SAnimTimeline::GetGridMetrics(float PhysicalWidth, double& OutMajorInterval
 	const FSlateFontInfo SmallLayoutFont = FCoreStyle::GetDefaultFontStyle("Regular", 8);
 	const TSharedRef<FSlateFontMeasure> FontMeasureService = FSlateApplication::Get().GetRenderer()->GetFontMeasureService();
 
-	const FFrameRate DisplayRate(static_cast<int32>(FMath::Max(FMath::RoundToInt(Model.Pin()->GetFrameRate()), 1)), 1);
+	const FFrameRate DisplayRate = Model.Pin()->GetFrameRate();
 	const double BiggestTime = ViewRange.Get().GetUpperBoundValue();
 	const FString TickString = NumericTypeInterface->ToString((BiggestTime * DisplayRate).FrameNumber.Value);
 	const FVector2D MaxTextSize = FontMeasureService->Measure(TickString, SmallLayoutFont);
@@ -834,7 +832,7 @@ void SAnimTimeline::HandleScrubPositionChanged(FFrameTime NewScrubPosition, bool
 
 double SAnimTimeline::GetSpinboxDelta() const
 {
-	return FFrameRate(Model.Pin()->GetTickResolution(), 1).AsDecimal() * FFrameRate(static_cast<int32>(FMath::RoundToInt(Model.Pin()->GetFrameRate())), 1).AsInterval();
+	return FFrameRate(Model.Pin()->GetTickResolution(), 1).AsDecimal() * Model.Pin()->GetFrameRate().AsInterval();
 }
 
 void SAnimTimeline::SetPlayTime(double InFrameTime)
diff --git a/Engine/Source/Editor/Persona/Private/AnimationEditorViewportClient.cpp b/Engine/Source/Editor/Persona/Private/AnimationEditorViewportClient.cpp
--- a/Engine/Source/Editor/Persona/Private/AnimationEditorViewportClient.cpp
+++ b/Engine/Source/Editor/Persona/Private/AnimationEditorViewportClient.cpp
@@ -1041,6 +1041,25 @@ FText FAnimationViewportClient::GetDisplayInfo(bool bDisplayAllInfo) const
 				TextValue = ConcatenateLine(TextValue, LOCTEXT("ApplyToCompressedDataWarning", "<AnimViewport.WarningText>Animation is being edited. To apply to compressed data (and recalculate baked additives), click \"Apply\"</>"));
 			}
 		}
+
+		UAnimCompositeBase* CompositeBase = Cast<UAnimCompositeBase>(PreviewInstance->GetCurrentAsset());
+		if (CompositeBase && !CompositeBase->GetCommonTargetFrameRate().IsValid())
+		{
+			FString AssetString;
+			TArray<UAnimationAsset*> Assets;
+			if(CompositeBase->GetAllAnimationSequencesReferred(Assets, false))
+			{
+				for (const UAnimationAsset* AnimAsset : Assets)
+				{
+					if (const UAnimSequenceBase* AnimSequenceBase = Cast<UAnimSequenceBase>(AnimAsset))
+					{
+						AssetString.Append(FString::Printf(TEXT("\n\t<AnimViewport.WarningText>%s - %s</>"), *AnimSequenceBase->GetName(), *AnimSequenceBase->GetSamplingFrameRate().ToPrettyText().ToString()));
+					}
+				}
+			}
+			
+			TextValue = ConcatenateLine(TextValue, FText::Format(LOCTEXT("IncompatibleFrameRatesCompositeWarning", "<AnimViewport.WarningText>{0} is composed of assets with incompatible framerates:</>{1}"), CompositeBase->GetClass()->GetDisplayNameText(), FText::FromString(AssetString)));
+		}
 	}
 
 	if (PreviewMeshComponent->IsUsingInGameBounds())
diff --git a/Engine/Source/Editor/Persona/Private/SAnimCompositePanel.cpp b/Engine/Source/Editor/Persona/Private/SAnimCompositePanel.cpp
--- a/Engine/Source/Editor/Persona/Private/SAnimCompositePanel.cpp
+++ b/Engine/Source/Editor/Persona/Private/SAnimCompositePanel.cpp
@@ -11,6 +11,7 @@
 #include "AnimTimeline/AnimModel.h"
 #include "AssetToolsModule.h"
 #include "IAssetTypeActions.h"
+#include "Animation/AnimationSettings.h"
 #include "Animation/AnimSequence.h"
 #include "Animation/AnimMontage.h"
 
@@ -90,6 +91,7 @@ void SAnimCompositePanel::Update()
 				.OnAnimSegmentNodeClicked( this, &SAnimCompositePanel::ShowSegmentInDetailsView )
 				.OnPreAnimUpdate( this, &SAnimCompositePanel::PreAnimUpdate )
 				.OnPostAnimUpdate( this, &SAnimCompositePanel::PostAnimUpdate )
+				.OnIsAnimAssetValid(this, &SAnimCompositePanel::OnIsAnimAssetValid)
 			];
 	}
 }
@@ -170,6 +172,8 @@ void SAnimCompositePanel::SortAndUpdateComposite()
 
 	Composite->AnimationTrack.SortAnimSegments();
 
+	Composite->UpdateCommonTargetFrameRate();
+
 	WeakModel.Pin()->RecalculateSequenceLength();
 
 	// Update view (this will recreate everything)
@@ -241,4 +245,30 @@ void SAnimCompositePanel::HandleObjectsSelected(const TArray<UObject*>& InObject
 	}
 }
 
+bool SAnimCompositePanel::OnIsAnimAssetValid(const UAnimSequenceBase* AnimSequenceBase, FText* OutReason)
+{
+	if (AnimSequenceBase)
+	{
+		if (UAnimationSettings::Get()->bEnforceSupportedFrameRates)
+		{
+			const FFrameRate AssetFrameRate = AnimSequenceBase->GetSamplingFrameRate();
+			const UAnimComposite* AnimComposite = WeakModel.Pin()->GetAsset<UAnimComposite>();
+			const FFrameRate CompositeFrameRate = AnimComposite->GetCommonTargetFrameRate();
+			const bool bContainsSegments = AnimComposite->AnimationTrack.AnimSegments.Num() != 0;
+			if (CompositeFrameRate.IsValid() && bContainsSegments && !AssetFrameRate.IsMultipleOf(CompositeFrameRate) && !AssetFrameRate.IsFactorOf(CompositeFrameRate))
+			{
+				if (OutReason)
+				{
+					*OutReason = FText::Format(LOCTEXT("InvalidFrameRate", "Animation Asset {0} its framerate {1} is incompatible with the Anim Composite's {2}"), FText::FromString(AnimSequenceBase->GetName()), AssetFrameRate.ToPrettyText(), CompositeFrameRate.ToPrettyText());
+				}
+			
+				return false;
+			}
+		}
+		
+		return true;	
+	}
+	return false;
+}
+
 #undef LOCTEXT_NAMESPACE
diff --git a/Engine/Source/Editor/Persona/Private/SAnimCompositePanel.h b/Engine/Source/Editor/Persona/Private/SAnimCompositePanel.h
--- a/Engine/Source/Editor/Persona/Private/SAnimCompositePanel.h
+++ b/Engine/Source/Editor/Persona/Private/SAnimCompositePanel.h
@@ -81,6 +81,8 @@ public:
 
 	void HandleObjectsSelected(const TArray<UObject*>& InObjects);
 
+	bool OnIsAnimAssetValid(const UAnimSequenceBase* AnimSequenceBase, FText* OutReason);
+
 private:
 	/** Reference to our editor model */
 	TWeakPtr<FAnimModel> WeakModel;
diff --git a/Engine/Source/Editor/Persona/Private/SAnimMontagePanel.cpp b/Engine/Source/Editor/Persona/Private/SAnimMontagePanel.cpp
--- a/Engine/Source/Editor/Persona/Private/SAnimMontagePanel.cpp
+++ b/Engine/Source/Editor/Persona/Private/SAnimMontagePanel.cpp
@@ -23,6 +23,7 @@
 #include "AnimTimeline/AnimModel.h"
 #include "AnimPreviewInstance.h"
 #include "ScopedTransaction.h"
+#include "Animation/AnimationSettings.h"
 #include "Misc/MessageDialog.h"
 #include "Animation/EditorAnimCompositeSegment.h"
 #include "Factories/AnimMontageFactory.h"
@@ -99,6 +100,33 @@ UAnimPreviewInstance* SAnimMontagePanel::GetPreviewInstance() const
 	return PreviewMeshComponent && PreviewMeshComponent->IsPreviewOn()? PreviewMeshComponent->PreviewInstance : nullptr;
 }
 
+bool SAnimMontagePanel::OnIsAnimAssetValid(const UAnimSequenceBase* AnimSequenceBase, FText* OutReason)
+{
+	if (AnimSequenceBase)
+	{
+		if (UAnimationSettings::Get()->bEnforceSupportedFrameRates)
+		{
+			const FFrameRate AssetFrameRate = AnimSequenceBase->GetSamplingFrameRate();
+			
+			const UAnimMontage* AnimMontage = WeakModel.Pin()->GetAsset<UAnimMontage>();
+			const FFrameRate MontageFrameRate = WeakModel.Pin()->GetAsset<UAnimMontage>()->GetCommonTargetFrameRate();
+			const bool bContainsSegments = AnimMontage->SlotAnimTracks.Num() != 0 && AnimMontage->SlotAnimTracks[0].AnimTrack.AnimSegments.Num() != 0;
+			if (MontageFrameRate.IsValid() && bContainsSegments && !AssetFrameRate.IsMultipleOf(MontageFrameRate) && !AssetFrameRate.IsFactorOf(MontageFrameRate))
+			{
+				if (OutReason)
+				{
+					*OutReason = FText::Format(LOCTEXT("InvalidFrameRate", "Animation Asset {0} its framerate {1} is incompatible with the Anim Montage's {2}"), FText::FromString(AnimSequenceBase->GetName()), AssetFrameRate.ToPrettyText(), MontageFrameRate.ToPrettyText());
+				}				
+				
+				return false;
+			}
+		}
+		
+		return true;	
+	}
+	return false;
+}
+
 FReply SAnimMontagePanel::OnFindParentClassInContentBrowserClicked()
 {
 	if (Montage != nullptr)
@@ -658,6 +686,7 @@ void SAnimMontagePanel::Update()
 						.OnDiffFromParentAsset(this, &SAnimMontagePanel::IsDiffererentFromParent)
 						.TrackMaxValue(this, &SAnimMontagePanel::GetSequenceLength)
 						.TrackNumDiscreteValues(Montage->GetNumberOfSampledKeys())
+						.OnIsAnimAssetValid(this, &SAnimMontagePanel::OnIsAnimAssetValid)
 					];
 
 				}
@@ -690,6 +719,7 @@ void SAnimMontagePanel::Update()
 						.OnPostAnimUpdate(this, &SAnimMontagePanel::PostAnimUpdate)
 						.OnAnimSegmentRemoved(this, &SAnimMontagePanel::OnAnimSegmentRemoved, SlotAnimIdx)
 						.OnTrackRightClickContextMenu(this, &SAnimMontagePanel::SummonTrackContextMenu, static_cast<int>(SlotAnimIdx))
+						.OnIsAnimAssetValid(this, &SAnimMontagePanel::OnIsAnimAssetValid)
 					];
 				}
 			}
diff --git a/Engine/Source/Editor/Persona/Private/SAnimMontagePanel.h b/Engine/Source/Editor/Persona/Private/SAnimMontagePanel.h
--- a/Engine/Source/Editor/Persona/Private/SAnimMontagePanel.h
+++ b/Engine/Source/Editor/Persona/Private/SAnimMontagePanel.h
@@ -172,6 +172,7 @@ public:
 
 	UAnimPreviewInstance* GetPreviewInstance() const;
 
+	bool OnIsAnimAssetValid(const UAnimSequenceBase* AnimSequenceBase, FText* OutReason);
 public:
 
 	/** These are meant to be callbacks used by the montage editing UI widgets */
diff --git a/Engine/Source/Editor/Persona/Private/SAnimSegmentsPanel.cpp b/Engine/Source/Editor/Persona/Private/SAnimSegmentsPanel.cpp
--- a/Engine/Source/Editor/Persona/Private/SAnimSegmentsPanel.cpp
+++ b/Engine/Source/Editor/Persona/Private/SAnimSegmentsPanel.cpp
@@ -44,6 +44,7 @@ void SAnimSegmentsPanel::Construct(const FArguments& InArgs)
 	OnAnimReplaceMapping					= InArgs._OnAnimReplaceMapping;
 	OnDiffFromParentAsset					= InArgs._OnDiffFromParentAsset;
 	OnGetNodeColor							= InArgs._OnGetNodeColor;
+	OnIsAnimAssetValid						= InArgs._OnIsAnimAssetValid;
 
 	bChildAnimMontage = InArgs._bChildAnimMontage;
 
@@ -84,6 +85,7 @@ void SAnimSegmentsPanel::Construct(const FArguments& InArgs)
 					.TrackNumDiscreteValues(InArgs._TrackNumDiscreteValues)
 					.OnTrackRightClickContextMenu(InArgs._OnTrackRightClickContextMenu)
 					.OnTrackDragDrop(this, &SAnimSegmentsPanel::OnTrackDragDrop)
+					.OnAssetDragDrop(this, &SAnimSegmentsPanel::OnAssetDragDrop)
 				];
 		}
 		else
@@ -106,6 +108,7 @@ void SAnimSegmentsPanel::Construct(const FArguments& InArgs)
 					.TrackNumDiscreteValues(InArgs._TrackNumDiscreteValues)
 					.OnTrackRightClickContextMenu(InArgs._OnTrackRightClickContextMenu)
 					.OnTrackDragDrop(this, &SAnimSegmentsPanel::OnTrackDragDrop)
+					.OnAssetDragDrop(this, &SAnimSegmentsPanel::OnAssetDragDrop)
 				];
 		}
 
@@ -361,14 +364,19 @@ void SAnimSegmentsPanel::ReplaceAnimSegment(UAnimSequenceBase* NewSequenceBase,
 	ReplaceAnimSegment(SegmentIdx, NewSequenceBase);
 }
 
-bool SAnimSegmentsPanel::IsValidToAdd(UAnimSequenceBase* NewSequenceBase) const
+bool SAnimSegmentsPanel::IsValidToAdd(UAnimSequenceBase* NewSequenceBase, FText* OutReason /*= nullptr*/) const
 {
-	if (AnimTrack == NULL || NewSequenceBase == NULL)
+	if (AnimTrack == nullptr || NewSequenceBase == nullptr)
 	{
 		return false;
 	}
 
-	return (AnimTrack->IsValidToAdd(NewSequenceBase));
+	if (!AnimTrack->IsValidToAdd(NewSequenceBase, OutReason) || (OnIsAnimAssetValid.IsBound() && !OnIsAnimAssetValid.Execute(NewSequenceBase, OutReason)))
+    {
+        return false;
+    } 
+
+	return true;
 }
 
 void SAnimSegmentsPanel::RemoveAnimSegment(int32 AnimSegmentIndex)
@@ -501,13 +509,35 @@ void SAnimSegmentsPanel::OnTrackDragDrop( TSharedPtr<FDragDropOperation> DragDro
 					}
 				}
 			}
+		}
+	}
+}
 
-			if(bFailedToAdd)
+bool SAnimSegmentsPanel::OnAssetDragDrop(TSharedPtr<FAssetDragDropOp> AssetDragDropOp)
+{
+	if (AssetDragDropOp.IsValid())
+	{
+		if (AssetDragDropOp->HasAssets())
+		{			
+			for(const FAssetData& DroppedAssetData : AssetDragDropOp->GetAssets())
 			{
-				FMessageDialog::Open(EAppMsgType::Ok, LOCTEXT("FailedToAdd", "Make sure the target animations are valid. Check to make sure if they are the same additive type if additive."));
+				FText FailureReason;
+				UAnimSequenceBase* DroppedSequence = Cast<UAnimSequenceBase>(DroppedAssetData.GetAsset());
+
+				bool bInvalidAsset = false;
+				if (!IsValidToAdd(DroppedSequence, &FailureReason))
+				{
+					AssetDragDropOp->SetToolTip(FailureReason, FAppStyle::GetBrush(TEXT("Graph.ConnectorFeedback.Error")));
+					return false;
+				}
 			}
+			
+			AssetDragDropOp->SetToolTip(FText(), FAppStyle::GetBrush(TEXT("Graph.ConnectorFeedback.OK")));
+			return true;
 		}
 	}
+
+	return false;
 }
 
 void SAnimSegmentsPanel::OnAnimSegmentNodeClicked(int32 SegmentIdx)
diff --git a/Engine/Source/Editor/Persona/Private/SAnimSegmentsPanel.h b/Engine/Source/Editor/Persona/Private/SAnimSegmentsPanel.h
--- a/Engine/Source/Editor/Persona/Private/SAnimSegmentsPanel.h
+++ b/Engine/Source/Editor/Persona/Private/SAnimSegmentsPanel.h
@@ -80,6 +80,7 @@ public:
 	SLATE_EVENT(FOnAnimSegmentRemoved, OnAnimSegmentRemoved)
 	SLATE_EVENT(FOnAnimReplaceMapping, OnAnimReplaceMapping)
 	SLATE_EVENT(FOnDiffFromParentAsset, OnDiffFromParentAsset)
+	SLATE_EVENT(FIsAnimAssetValid, OnIsAnimAssetValid)
 
 	SLATE_EVENT( FOnBarDrag,				OnBarDrag)
 	SLATE_EVENT( FOnBarDrop,				OnBarDrop)
@@ -131,8 +132,9 @@ private:
 	void				SummonSegmentNodeContextMenu( FMenuBuilder& MenuBuilder, int32 AnimSegmentIndex );
 
 	void				AddAnimSegment(UAnimSequenceBase *NewSequenceBase, float NewStartPos );
-	bool				IsValidToAdd(UAnimSequenceBase* NewSequenceBase) const;
+	bool				IsValidToAdd(UAnimSequenceBase* NewSequenceBase, FText* OutReason = nullptr) const;
 	void				OnTrackDragDrop( TSharedPtr<FDragDropOperation> DragDropOp, float DataPos );
+	bool				OnAssetDragDrop(TSharedPtr<FAssetDragDropOp> AssetDragDropOp);
 	void				OnAnimSegmentNodeClicked(int32 SegmentIdx);
 	void				OnAnimSegmentNodeDoubleClicked(int32 SegmentIdx);
 
@@ -174,6 +176,7 @@ private:
 
 	TAttribute<FLinearColor> DefaultNodeColor;
 
+	FIsAnimAssetValid OnIsAnimAssetValid;
  	/* 
 	 * Child Anim Montage: Child Anim Montage only can replace name of animations, and no other meaningful edits 
 	 * as it will derive every data from Parent. There might be some other data that will allow to be replaced, but for now, it is
diff --git a/Engine/Source/Editor/Persona/Private/STrack.cpp b/Engine/Source/Editor/Persona/Private/STrack.cpp
--- a/Engine/Source/Editor/Persona/Private/STrack.cpp
+++ b/Engine/Source/Editor/Persona/Private/STrack.cpp
@@ -18,6 +18,7 @@
 #include "Rendering/DrawElements.h"
 #include "Rendering/RenderingCommon.h"
 #include "SCurveEditor.h"
+#include "DragAndDrop/AssetDragDropOp.h"
 #include "Styling/CoreStyle.h"
 #include "Templates/TypeHash.h"
 #include "Widgets/Layout/SBorder.h"
@@ -414,6 +415,7 @@ void STrack::Construct( const FArguments& InArgs )
 	OnBarClicked = InArgs._OnBarClicked;
 	OnBarDrop = InArgs._OnBarDrop;
 	OnTrackDragDrop = InArgs._OnTrackDragDrop;
+	OnAssetDragDrop = InArgs._OnAssetDragDrop;
 	OnSummonContextMenu = InArgs._OnSummonContextMenu;
 	OnTrackRightClickContextMenu = InArgs._OnTrackRightClickContextMenu;
 
@@ -571,6 +573,28 @@ FReply STrack::OnDragDetected(const FGeometry& MyGeometry, const FPointerEvent&
 	return FReply::Unhandled();
 }
 
+void STrack::OnDragEnter(const FGeometry& MyGeometry, const FDragDropEvent& DragDropEvent)
+{
+	TSharedPtr<FAssetDragDropOp> DragDropOp = DragDropEvent.GetOperationAs<FAssetDragDropOp>();
+	if (DragDropOp.IsValid())
+	{
+		if(OnAssetDragDrop.IsBound())
+		{
+			OnAssetDragDrop.Execute(DragDropOp);
+		}
+	}
+}
+
+void STrack::OnDragLeave(const FDragDropEvent& DragDropEvent)
+{
+	SPanel::OnDragLeave(DragDropEvent);
+	TSharedPtr<FAssetDragDropOp> DragDropOp = DragDropEvent.GetOperationAs<FAssetDragDropOp>();
+	if (DragDropOp.IsValid())
+	{
+		DragDropOp->ResetToDefaultToolTip();
+	}
+}
+
 float STrack::GetNodeDragDropDataPos( const FGeometry& MyGeometry, const FDragDropEvent& DragDropEvent ) const
 {
 	float DataPos = 0.f;
diff --git a/Engine/Source/Editor/Persona/Private/STrack.h b/Engine/Source/Editor/Persona/Private/STrack.h
--- a/Engine/Source/Editor/Persona/Private/STrack.h
+++ b/Engine/Source/Editor/Persona/Private/STrack.h
@@ -23,6 +23,7 @@
 #include "Widgets/DeclarativeSyntaxSupport.h"
 #include "Widgets/SCompoundWidget.h"
 #include "Widgets/SPanel.h"
+#include "DragAndDrop/AssetDragDropOp.h"
 
 class FArrangedChildren;
 class FChildren;
@@ -47,6 +48,8 @@ DECLARE_DELEGATE_OneParam( FOnBarClicked, int32)
 DECLARE_DELEGATE_TwoParams( FOnBarDrag, int32, float)
 DECLARE_DELEGATE_OneParam( FOnBarDrop, int32 )
 DECLARE_DELEGATE_TwoParams( FOnTrackDragDop, TSharedPtr<FDragDropOperation>, float )
+DECLARE_DELEGATE_RetVal_OneParam( bool, FOnAssetDragDrop, TSharedPtr<FAssetDragDropOp>)
+DECLARE_DELEGATE_RetVal_TwoParams(bool, FIsAnimAssetValid, const UAnimSequenceBase*, FText* OutReason)
 
 DECLARE_DELEGATE_RetVal( FString, FOnGetNodeName )
 DECLARE_DELEGATE_OneParam( FOnTrackNodeDragged, float )
@@ -270,6 +273,7 @@ public:
 		SLATE_EVENT( FOnBarDrag, OnBarDrag)
 		SLATE_EVENT( FOnBarClicked, OnBarClicked)
 		SLATE_EVENT( FOnTrackDragDop, OnTrackDragDrop )
+		SLATE_EVENT( FOnAssetDragDrop, OnAssetDragDrop )
 		SLATE_EVENT( FOnBarDrop, OnBarDrop )
 		SLATE_EVENT( FOnSummonContextMenu, OnSummonContextMenu )
 		SLATE_EVENT( FOnTrackRightClickContextMenu, OnTrackRightClickContextMenu )
@@ -286,6 +290,8 @@ public:
 	virtual FReply	OnDrop( const FGeometry& MyGeometry, const FDragDropEvent& DragDropEvent ) override;	
 	virtual FReply	OnDragOver( const FGeometry& MyGeometry, const FDragDropEvent& DragDropEvent ) override;
 	virtual FReply	OnDragDetected( const FGeometry& MyGeometry, const FPointerEvent& MouseEvent ) override;
+	virtual void OnDragEnter(const FGeometry& MyGeometry, const FDragDropEvent& DragDropEvent) override;
+	virtual void OnDragLeave(const FDragDropEvent& DragDropEvent) override;
 
 	int32			GetHitNode(const FGeometry& MyGeometry, const FVector2D& CursorPosition);
 
@@ -337,6 +343,7 @@ protected:
 
 	FOnGetBarPos							OnGetDraggableBarPos;
 	FOnTrackDragDop							OnTrackDragDrop;
+	FOnAssetDragDrop						OnAssetDragDrop;
 	
 	TAttribute<const FSlateBrush*>			StyleInfo;
 	FOnSummonContextMenu					OnSummonContextMenu;
diff --git a/Engine/Source/Runtime/Engine/Classes/Animation/AnimComposite.h b/Engine/Source/Runtime/Engine/Classes/Animation/AnimComposite.h
--- a/Engine/Source/Runtime/Engine/Classes/Animation/AnimComposite.h
+++ b/Engine/Source/Runtime/Engine/Classes/Animation/AnimComposite.h
@@ -39,6 +39,10 @@ public:
 	virtual void GetResourceSizeEx(FResourceSizeEx& CumulativeResourceSize) override;
 	//~ End UObject Interface
 
+#if WITH_EDITOR
+	virtual void PostEditChangeProperty(FPropertyChangedEvent& PropertyChangedEvent) override;
+#endif // WITH_EDITOR
+
 	//~ Begin UAnimSequenceBase Interface
 	ENGINE_API virtual void HandleAssetPlayerTickedInternal(FAnimAssetTickContext &Context, const float PreviousTime, const float MoveDelta, const FAnimTickRecord &Instance, struct FAnimNotifyQueue& NotifyQueue) const override;
 
@@ -58,6 +62,7 @@ public:
 	virtual class UAnimSequence* GetAdditiveBasePose() const override;
 	virtual bool GetAllAnimationSequencesReferred(TArray<UAnimationAsset*>& AnimationAssets, bool bRecursive = true) override;
 	virtual void ReplaceReferredAnimations(const TMap<UAnimationAsset*, UAnimationAsset*>& ReplacementMap) override;
+	virtual void UpdateCommonTargetFrameRate() override;	
 #endif
 	//~ End UAnimSequence Interface
 
@@ -66,6 +71,7 @@ public:
 	virtual bool ContainRecursive(TArray<UAnimCompositeBase*>& CurrentAccumulatedList) override;
 	virtual void SetCompositeLength(float InLength) override;
 	virtual void PostLoad() override;
+	virtual FFrameRate GetSamplingFrameRate() const override;
 	//~End UAnimCompositeBase Interface
 };
 
diff --git a/Engine/Source/Runtime/Engine/Classes/Animation/AnimCompositeBase.h b/Engine/Source/Runtime/Engine/Classes/Animation/AnimCompositeBase.h
--- a/Engine/Source/Runtime/Engine/Classes/Animation/AnimCompositeBase.h
+++ b/Engine/Source/Runtime/Engine/Classes/Animation/AnimCompositeBase.h
@@ -16,6 +16,11 @@ class UAnimCompositeBase;
 class UAnimSequence;
 struct FCompactPose;
 
+namespace UE { namespace Anim
+{
+	extern TAutoConsoleVariable<bool> CVarOutputMontageFrameRateWarning;
+}}
+
 /** Struct defining a RootMotionExtractionStep.
  * When extracting RootMotion we can encounter looping animations (wrap around), or different animations.
  * We break those up into different steps, to help with RootMotion extraction, 
@@ -267,7 +272,7 @@ struct FAnimTrack
 	void ValidateSegmentTimes();
 
 	/** return true if valid to add */
-	ENGINE_API bool IsValidToAdd(const UAnimSequenceBase* SequenceBase) const;
+	ENGINE_API bool IsValidToAdd(const UAnimSequenceBase* SequenceBase, FText* OutReason = nullptr) const;
 
 	/** Gets the index of the segment at the given absolute montage time. */	
 	ENGINE_API int32 GetSegmentIndexAtTime(float InTime) const;
@@ -371,6 +376,13 @@ class UAnimCompositeBase : public UAnimSequenceBase
 
 #if WITH_EDITOR
 	virtual void PopulateWithExistingModel(TScriptInterface<IAnimationDataModel> ExistingDataModel) override;
+
+	virtual void UpdateCommonTargetFrameRate() PURE_VIRTUAL(UAnimCompositeBase::UpdateCommonTargetFrameRate, );	
 #endif // WITH_EDITOR
+	FFrameRate GetCommonTargetFrameRate() const { return CommonTargetFrameRate; }
+protected:
+	/** Frame-rate used to represent this Animation Montage (best fitting for placed Animation Sequences)*/
+	UPROPERTY(VisibleAnywhere, Category = AnimationComposite)
+	FFrameRate CommonTargetFrameRate;
 };
 
diff --git a/Engine/Source/Runtime/Engine/Classes/Animation/AnimMontage.h b/Engine/Source/Runtime/Engine/Classes/Animation/AnimMontage.h
--- a/Engine/Source/Runtime/Engine/Classes/Animation/AnimMontage.h
+++ b/Engine/Source/Runtime/Engine/Classes/Animation/AnimMontage.h
@@ -725,6 +725,8 @@ public:
 	PRAGMA_ENABLE_DEPRECATION_WARNINGS
 	virtual void PreSave(FObjectPreSaveContext ObjectSaveContext) override;
 
+	virtual FFrameRate GetSamplingFrameRate() const override;
+
 	// Gets the sequence length of the montage by calculating it from the lengths of the segments in the montage
 	ENGINE_API float CalculateSequenceLength();
 
@@ -913,7 +915,6 @@ private:
 	/** Cached list of Branching Point markers */
 	UPROPERTY()
 	TArray<FBranchingPointMarker> BranchingPointMarkers;
-
 public:
 
 	/** Keep track of which AnimNotify_State are marked as BranchingPoints, so we can update their state when the Montage is ticked */
@@ -966,5 +967,6 @@ private:
 #if WITH_EDITOR
 	void BakeTimeStretchCurve();
 	//~End Time Stretch Curve
+	virtual void UpdateCommonTargetFrameRate() override;
 #endif // WITH_EDITOR
 };
diff --git a/Engine/Source/Runtime/Engine/Classes/Animation/AnimationAsset.h b/Engine/Source/Runtime/Engine/Classes/Animation/AnimationAsset.h
--- a/Engine/Source/Runtime/Engine/Classes/Animation/AnimationAsset.h
+++ b/Engine/Source/Runtime/Engine/Classes/Animation/AnimationAsset.h
@@ -1114,7 +1114,6 @@ public:
 	// Helper function for GetAllAnimationSequencesReferred, it adds itself first and call GetAllAnimationSEquencesReferred
 	ENGINE_API void HandleAnimReferenceCollection(TArray<UAnimationAsset*>& AnimationAssets, bool bRecursive);
 
-protected:
 	/** Retrieve all animations that are used by this asset 
 	 * 
 	 * @param (out)		AnimationSequences 
diff --git a/Engine/Source/Runtime/Engine/Classes/Animation/AnimationSettings.h b/Engine/Source/Runtime/Engine/Classes/Animation/AnimationSettings.h
--- a/Engine/Source/Runtime/Engine/Classes/Animation/AnimationSettings.h
+++ b/Engine/Source/Runtime/Engine/Classes/Animation/AnimationSettings.h
@@ -103,7 +103,11 @@ class UAnimationSettings : public UDeveloperSettings
 
 	/** Project specific default frame-rate used when (re)initializing any animation based data */
 	UPROPERTY(config, EditAnywhere, Category = AnimationData)
-	FFrameRate DefaultFrameRate;	
+	FFrameRate DefaultFrameRate;
+
+	/** Whether to enforce the project to only use entries from SupportedFrameRates for the animation assets, if disable will warn instead */
+	UPROPERTY(config, EditAnywhere, Category = AnimationData)
+	bool bEnforceSupportedFrameRates;
 public:
 	static UAnimationSettings * Get() { return CastChecked<UAnimationSettings>(UAnimationSettings::StaticClass()->GetDefaultObject()); }
 
diff --git a/Engine/Source/Runtime/Engine/Private/Animation/AnimComposite.cpp b/Engine/Source/Runtime/Engine/Private/Animation/AnimComposite.cpp
--- a/Engine/Source/Runtime/Engine/Private/Animation/AnimComposite.cpp
+++ b/Engine/Source/Runtime/Engine/Private/Animation/AnimComposite.cpp
@@ -8,6 +8,7 @@
 #include "Animation/AnimData/IAnimationDataController.h"
 #include "Animation/AnimationPoseData.h"
 #include "EngineLogs.h"
+#include "Animation/AnimationSettings.h"
 
 #include UE_INLINE_GENERATED_CPP_BY_NAME(AnimComposite)
 
@@ -27,6 +28,12 @@ void UAnimComposite::ReplaceReferredAnimations(const TMap<UAnimationAsset*, UAni
 {
 	AnimationTrack.ReplaceReferredAnimations(ReplacementMap);
 }
+
+void UAnimComposite::PostEditChangeProperty(FPropertyChangedEvent& PropertyChangedEvent)
+{
+	Super::PostEditChangeProperty(PropertyChangedEvent);
+	UpdateCommonTargetFrameRate();
+}
 #endif
 
 bool UAnimComposite::IsNotifyAvailable() const
@@ -47,6 +54,7 @@ void UAnimComposite::GetResourceSizeEx(FResourceSizeEx& CumulativeResourceSize)
 	CumulativeResourceSize.AddDedicatedSystemMemoryBytes(AnimationTrack.GetTotalBytesUsed());
 }
 
+
 void UAnimComposite::HandleAssetPlayerTickedInternal(FAnimAssetTickContext &Context, const float PreviousTime, const float MoveDelta, const FAnimTickRecord &Instance, struct FAnimNotifyQueue& NotifyQueue) const
 {
 	Super::HandleAssetPlayerTickedInternal(Context, PreviousTime, MoveDelta, Instance, NotifyQueue);
@@ -188,3 +196,79 @@ void UAnimComposite::PostLoad()
 #endif
 }
 
+#if WITH_EDITOR
+void UAnimComposite::UpdateCommonTargetFrameRate()
+{
+	CommonTargetFrameRate = FFrameRate(0,0);
+	FFrameRate TargetRate = UAnimationSettings::Get()->GetDefaultFrameRate();
+
+	bool bFirst = true;
+	bool bValidFrameRate = AnimationTrack.AnimSegments.Num() != 0;
+	for (const FAnimSegment& Segment : AnimationTrack.AnimSegments)
+	{
+		const UAnimSequenceBase* Base = Segment.GetAnimReference();
+		if (Base && Base != this)
+		{
+			const FFrameRate BaseFrameRate = Base->GetSamplingFrameRate();
+			if (bFirst)
+			{
+				TargetRate = BaseFrameRate;
+				bFirst = false;
+			}
+			else
+			{
+				if (BaseFrameRate.IsValid())
+				{					
+					if (TargetRate.IsMultipleOf(BaseFrameRate))
+					{
+						TargetRate = BaseFrameRate;
+					}
+					else if (TargetRate != BaseFrameRate && !BaseFrameRate.IsMultipleOf(TargetRate))
+					{						
+						FString AssetString;
+						TArray<UAnimationAsset*> Assets;
+						if(GetAllAnimationSequencesReferred(Assets, false))
+						{
+							for (const UAnimationAsset* AnimAsset : Assets)
+							{
+								if (const UAnimSequenceBase* AnimSequenceBase = Cast<UAnimSequenceBase>(AnimAsset))
+								{
+									AssetString.Append(FString::Printf(TEXT("\n\t%s - %s"), *AnimSequenceBase->GetName(), *AnimSequenceBase->GetSamplingFrameRate().ToPrettyText().ToString()));
+								}
+							}
+						}						
+
+						if (UE::Anim::CVarOutputMontageFrameRateWarning.GetValueOnAnyThread() == true)
+						{
+							UE_LOG(LogAnimation, Warning, TEXT("Frame rate of animation %s (%s) is incompatible with other animations in Animation Composite %s - underlying frame-rate will be set to %s:%s"), *Base->GetName(), *BaseFrameRate.ToPrettyText().ToString(), *GetName(), *Super::GetSamplingFrameRate().ToPrettyText().ToString(), *AssetString);
+						}
+						
+						bValidFrameRate = false;
+						break;
+					}
+				}
+				else
+				{
+					UE_LOG(LogAnimation, Warning, TEXT("Invalid frame rate %s for %s in %s"), *BaseFrameRate.ToPrettyText().ToString(), *Base->GetName(), *GetName());
+				}
+			}			
+		}	
+	}
+
+	if (bValidFrameRate)
+	{
+		CommonTargetFrameRate = TargetRate;
+	}
+}
+#endif // WITH_EDITOR
+
+FFrameRate UAnimComposite::GetSamplingFrameRate() const
+{
+	if (CommonTargetFrameRate.IsValid())
+	{
+		return CommonTargetFrameRate;
+	}
+
+	return Super::GetSamplingFrameRate();
+}
+
diff --git a/Engine/Source/Runtime/Engine/Private/Animation/AnimCompositeBase.cpp b/Engine/Source/Runtime/Engine/Private/Animation/AnimCompositeBase.cpp
--- a/Engine/Source/Runtime/Engine/Private/Animation/AnimCompositeBase.cpp
+++ b/Engine/Source/Runtime/Engine/Private/Animation/AnimCompositeBase.cpp
@@ -17,6 +17,17 @@
 
 #include UE_INLINE_GENERATED_CPP_BY_NAME(AnimCompositeBase)
 
+namespace UE
+{
+	namespace Anim
+	{		
+		TAutoConsoleVariable<bool> CVarOutputMontageFrameRateWarning(
+			TEXT("a.OutputMontageFrameRateWarning"),
+			true,
+			TEXT("If true will warn the user about Animation Montages/Composites composed of incompatible animation assets (incompatible frame-rates)."));
+	}
+}
+
 ///////////////////////////////////////////////////////
 // FAnimSegment
 ///////////////////////////////////////////////////////
@@ -714,35 +725,51 @@ int32 FAnimTrack::GetTotalBytesUsed() const
 	return AnimSegments.GetAllocatedSize();
 }
 
-bool FAnimTrack::IsValidToAdd(const UAnimSequenceBase* SequenceBase) const
+bool FAnimTrack::IsValidToAdd(const UAnimSequenceBase* SequenceBase, FText* OutReason /*= nullptr*/) const
 {
 	bool bValid = false;
 	// remove asset if invalid
 	if (SequenceBase)
 	{
-		if (SequenceBase->GetPlayLength() <= 0.f)
+		const float PlayLength = SequenceBase->GetPlayLength();
+		if (PlayLength <= 0.f)
 		{
 			UE_LOG(LogAnimation, Warning, TEXT("Remove Empty Sequence (%s)"), *SequenceBase->GetFullName());
+
+			if (OutReason)
+			{
+				*OutReason = FText::FromString(FString::Printf(TEXT("Animation Asset %s has invalid playable length of %f"), *SequenceBase->GetName(), PlayLength));
+			}			
+			
+			return false;
 		}
-		else if (!SequenceBase->CanBeUsedInComposition())
+
+		if (!SequenceBase->CanBeUsedInComposition())
 		{
 			UE_LOG(LogAnimation, Warning, TEXT("Remove Invalid Sequence (%s)"), *SequenceBase->GetFullName());
-		}
-		else
-		{
-			int32 TrackType = GetTrackAdditiveType();
-			if ((TrackType == -1) || (TrackType == SequenceBase->GetAdditiveAnimType()))
+			if (OutReason)
 			{
-				bValid = true;
+				*OutReason = FText::FromString(FString::Printf(TEXT("Animation Asset %s cannot be used in an Animation Composite/Montage"), *SequenceBase->GetName()));
 			}
-			else
+			return false;
+		}
+		
+		const int32 TrackType = GetTrackAdditiveType();
+		const EAdditiveAnimationType AnimAdditiveType = SequenceBase->GetAdditiveAnimType();
+		if (TrackType != AnimAdditiveType && TrackType != INDEX_NONE)
+		{
+			const UEnum* TypeEnum = FindObject<UEnum>(nullptr, TEXT("/Script/Engine.EAdditiveAnimationType"));	
+			if (OutReason)
 			{
-				UE_LOG(LogAnimation, Warning, TEXT("Additivie type (%s) does not match. Make sure you add same type of additive animation."), *SequenceBase->GetFullName());
+				*OutReason = FText::FromString(FString::Printf(TEXT("Animation Asset %s has an additive type %s that does not match the target's %s"), *SequenceBase->GetName(), *TypeEnum->GetNameStringByValue(AnimAdditiveType), *TypeEnum->GetNameStringByValue(TrackType)));
 			}
+			return false;
 		}
+		
+		return true;
 	}
 
-	return bValid;
+	return true;
 }
 ///////////////////////////////////////////////////////
 // UAnimCompositeBase
@@ -789,6 +816,10 @@ void UAnimCompositeBase::PostLoad()
 {
 	Super::PostLoad();
 
+#if WITH_EDITOR
+	UpdateCommonTargetFrameRate();
+#endif // WITH_EDITOR
+
 	InvalidateRecursiveAsset();
 }
 
diff --git a/Engine/Source/Runtime/Engine/Private/Animation/AnimMontage.cpp b/Engine/Source/Runtime/Engine/Private/Animation/AnimMontage.cpp
--- a/Engine/Source/Runtime/Engine/Private/Animation/AnimMontage.cpp
+++ b/Engine/Source/Runtime/Engine/Private/Animation/AnimMontage.cpp
@@ -20,6 +20,7 @@
 #include "Components/SkeletalMeshComponent.h"
 #include "Animation/AnimTrace.h"
 #include "Animation/ActiveMontageInstanceScope.h"
+#include "Animation/AnimationSettings.h"
 
 #include UE_INLINE_GENERATED_CPP_BY_NAME(AnimMontage)
 
@@ -372,6 +373,16 @@ void UAnimMontage::PreSave(FObjectPreSaveContext ObjectSaveContext)
 	Super::PreSave(ObjectSaveContext);
 }
 
+FFrameRate UAnimMontage::GetSamplingFrameRate() const
+{	
+	if (CommonTargetFrameRate.IsValid())
+	{
+		return CommonTargetFrameRate;
+	}
+
+	return Super::GetSamplingFrameRate();
+}
+
 void UAnimMontage::PostLoad()
 {
 	Super::PostLoad();
@@ -739,6 +750,8 @@ void UAnimMontage::PostEditChangeProperty(FPropertyChangedEvent& PropertyChanged
 		CollectMarkers();
 	}
 
+	UpdateCommonTargetFrameRate();
+
 	PropagateChanges();
 }
 
@@ -2994,6 +3007,74 @@ void UAnimMontage::PopulateWithExistingModel(TScriptInterface<IAnimationDataMode
 	const float CurrentCalculatedLength = CalculateSequenceLength();
 	SetCompositeLength(CurrentCalculatedLength);
 }
+
+void UAnimMontage::UpdateCommonTargetFrameRate()
+{
+	CommonTargetFrameRate = FFrameRate(0,0);
+	FFrameRate TargetRate = UAnimationSettings::Get()->GetDefaultFrameRate();
+
+	bool bValidFrameRate = true;
+	bool bFirst = true;
+	for (const FSlotAnimationTrack& Track : SlotAnimTracks)
+	{
+		for (const FAnimSegment& Segment : Track.AnimTrack.AnimSegments)
+		{
+			const UAnimSequenceBase* Base = Segment.GetAnimReference();
+			if (Base && Base != this)
+			{
+				const FFrameRate BaseFrameRate = Base->GetSamplingFrameRate();
+				if (bFirst)
+				{
+					TargetRate = BaseFrameRate;
+					bFirst = false;
+				}
+				else
+				{
+					if (BaseFrameRate.IsValid())
+					{
+						if (TargetRate.IsMultipleOf(BaseFrameRate))
+						{
+							TargetRate = BaseFrameRate;
+						}
+						else if (TargetRate != BaseFrameRate && !BaseFrameRate.IsMultipleOf(TargetRate))
+						{
+							FString AssetString;
+							TArray<UAnimationAsset*> Assets;
+							if(GetAllAnimationSequencesReferred(Assets, false))
+							{
+								for (const UAnimationAsset* AnimAsset : Assets)
+								{
+									if (const UAnimSequenceBase* AnimSequenceBase = Cast<UAnimSequenceBase>(AnimAsset))
+									{
+										AssetString.Append(FString::Printf(TEXT("\n\t%s - %s"), *AnimSequenceBase->GetName(), *AnimSequenceBase->GetSamplingFrameRate().ToPrettyText().ToString()));
+									}
+								}
+							}						
+
+							if (UE::Anim::CVarOutputMontageFrameRateWarning.GetValueOnAnyThread() == true)
+							{
+								UE_LOG(LogAnimation, Warning, TEXT("Frame rate of animation %s (%s) is incompatible with other animations in Animation Montage %s - underlying frame-rate will be set to %s:%s"), *Base->GetName(), *BaseFrameRate.ToPrettyText().ToString(), *GetName(), *Super::GetSamplingFrameRate().ToPrettyText().ToString(), *AssetString);
+							}
+						
+							bValidFrameRate = false;
+							break;
+						}
+					}
+					else
+					{
+						UE_LOG(LogAnimMontage, Warning, TEXT("Invalid frame rate %s for %s in %s"), *BaseFrameRate.ToPrettyText().ToString(), *Base->GetName(), *GetName());
+					}
+				}			
+			}	
+		}			
+	}
+
+	if (bValidFrameRate)
+	{
+		CommonTargetFrameRate = TargetRate;
+	}
+}
+
 #endif // WITH_EDITOR
 
 FMontageBlendSettings::FMontageBlendSettings()
diff --git a/Engine/Source/Runtime/Engine/Private/Animation/AnimationSettings.cpp b/Engine/Source/Runtime/Engine/Private/Animation/AnimationSettings.cpp
--- a/Engine/Source/Runtime/Engine/Private/Animation/AnimationSettings.cpp
+++ b/Engine/Source/Runtime/Engine/Private/Animation/AnimationSettings.cpp
@@ -43,6 +43,7 @@ UAnimationSettings::UAnimationSettings(const FObjectInitializer& ObjectInitializ
 	};
 	
 	DefaultFrameRate = FFrameRate(30,1);
+	bEnforceSupportedFrameRates = false;
 }
 
 TArray<FString> UAnimationSettings::GetBoneCustomAttributeNamesToImport() const
