CIS fixes following up on CL 26886518

diff --git a/Engine/Source/Editor/Persona/Private/AnimationEditorViewportClient.cpp b/Engine/Source/Editor/Persona/Private/AnimationEditorViewportClient.cpp
--- a/Engine/Source/Editor/Persona/Private/AnimationEditorViewportClient.cpp
+++ b/Engine/Source/Editor/Persona/Private/AnimationEditorViewportClient.cpp
@@ -45,6 +45,7 @@
 #include "UnrealWidget.h"
 #include "Engine/PoseWatchRenderData.h"
 #include "UObject/UE5MainStreamObjectVersion.h"
+#include "Animation/AnimCompositeBase.h"
 
 namespace {
 	static const float AnimationEditorViewport_RotateSpeed = 0.02f;
diff --git a/Engine/Source/Editor/Persona/Private/STrack.h b/Engine/Source/Editor/Persona/Private/STrack.h
--- a/Engine/Source/Editor/Persona/Private/STrack.h
+++ b/Engine/Source/Editor/Persona/Private/STrack.h
@@ -35,6 +35,7 @@ class FUICommandList;
 class FWidgetStyle;
 class SWidget;
 class UObject;
+class UAnimSequenceBase;
 struct FPointerEvent;
 struct FSlateBrush;
 
diff --git a/Engine/Source/Runtime/Engine/Classes/Animation/AnimCompositeBase.h b/Engine/Source/Runtime/Engine/Classes/Animation/AnimCompositeBase.h
--- a/Engine/Source/Runtime/Engine/Classes/Animation/AnimCompositeBase.h
+++ b/Engine/Source/Runtime/Engine/Classes/Animation/AnimCompositeBase.h
@@ -10,17 +10,21 @@
 #include "CoreMinimal.h"
 #include "UObject/ObjectMacros.h"
 #include "Animation/AnimSequenceBase.h"
+#include "HAL/IConsoleManager.h"
 #include "AnimCompositeBase.generated.h"
 
 class UAnimCompositeBase;
 class UAnimSequence;
 struct FCompactPose;
 
+#if WITH_EDITOR
 namespace UE { namespace Anim
 {
 	extern TAutoConsoleVariable<bool> CVarOutputMontageFrameRateWarning;
 }}
 
+#endif // WITH_EDITOR
+
 /** Struct defining a RootMotionExtractionStep.
  * When extracting RootMotion we can encounter looping animations (wrap around), or different animations.
  * We break those up into different steps, to help with RootMotion extraction, 
diff --git a/Engine/Source/Runtime/Engine/Private/Animation/AnimCompositeBase.cpp b/Engine/Source/Runtime/Engine/Private/Animation/AnimCompositeBase.cpp
--- a/Engine/Source/Runtime/Engine/Private/Animation/AnimCompositeBase.cpp
+++ b/Engine/Source/Runtime/Engine/Private/Animation/AnimCompositeBase.cpp
@@ -17,16 +17,18 @@
 
 #include UE_INLINE_GENERATED_CPP_BY_NAME(AnimCompositeBase)
 
+#if WITH_EDITOR
 namespace UE
 {
 	namespace Anim
 	{		
 		TAutoConsoleVariable<bool> CVarOutputMontageFrameRateWarning(
 			TEXT("a.OutputMontageFrameRateWarning"),
-			true,
+			false,
 			TEXT("If true will warn the user about Animation Montages/Composites composed of incompatible animation assets (incompatible frame-rates)."));
 	}
 }
+#endif if WITH_EDITOR
 
 ///////////////////////////////////////////////////////
 // FAnimSegment
