And another CIS fix

diff --git a/Engine/Source/Runtime/Engine/Private/Animation/AnimCompositeBase.cpp b/Engine/Source/Runtime/Engine/Private/Animation/AnimCompositeBase.cpp
--- a/Engine/Source/Runtime/Engine/Private/Animation/AnimCompositeBase.cpp
+++ b/Engine/Source/Runtime/Engine/Private/Animation/AnimCompositeBase.cpp
@@ -28,7 +28,7 @@ namespace UE
 			TEXT("If true will warn the user about Animation Montages/Composites composed of incompatible animation assets (incompatible frame-rates)."));
 	}
 }
-#endif if WITH_EDITOR
+#endif // WITH_EDITOR
 
 ///////////////////////////////////////////////////////
 // FAnimSegment
