Fix Submix Volume Unit Ranges

diff --git a/Engine/Source/Runtime/Engine/Classes/Sound/SoundSubmix.h b/Engine/Source/Runtime/Engine/Classes/Sound/SoundSubmix.h
--- a/Engine/Source/Runtime/Engine/Classes/Sound/SoundSubmix.h
+++ b/Engine/Source/Runtime/Engine/Classes/Sound/SoundSubmix.h
@@ -296,15 +296,15 @@ public:
 	float DryLevel;
 
 	/** The output volume of the submix in Decibels. Applied after submix effects and analysis are performed.*/
-	UPROPERTY(EditAnywhere, BlueprintReadWrite, Category = SubmixLevel, meta = (DisplayName = "Output Volume (dB)", AudioParam = "Volume", AudioParamClass = "SoundModulationParameterVolume"))
+	UPROPERTY(EditAnywhere, BlueprintReadWrite, Category = SubmixLevel, meta = (DisplayName = "Output Volume (dB)", AudioParam = "Volume", AudioParamClass = "SoundModulationParameterVolume", ClampMin = "-96.0", ClampMax = "0.0", UIMin = "-96.0", UIMax = "0.0"))
 	FSoundModulationDestinationSettings OutputVolumeModulation;
 
 	/** The wet level of the submixin Decibels. Applied after submix effects and analysis are performed. */
-	UPROPERTY(EditAnywhere, BlueprintReadWrite, Category = SubmixLevel, meta = (DisplayName = "Wet Level (dB)", AudioParam = "Volume", AudioParamClass = "SoundModulationParameterVolume"))
+	UPROPERTY(EditAnywhere, BlueprintReadWrite, Category = SubmixLevel, meta = (DisplayName = "Wet Level (dB)", AudioParam = "Volume", AudioParamClass = "SoundModulationParameterVolume", ClampMin = "-96.0", ClampMax = "0.0", UIMin = "-96.0", UIMax = "0.0"))
 	FSoundModulationDestinationSettings WetLevelModulation;
 
 	/** The dry level of the submix in Decibels. Applied before submix effects and analysis are performed. */
-	UPROPERTY(EditAnywhere, BlueprintReadWrite, Category = SubmixLevel, meta = (DisplayName = "Dry Level (dB)", AudioParam = "Volume", AudioParamClass = "SoundModulationParameterVolume"))
+	UPROPERTY(EditAnywhere, BlueprintReadWrite, Category = SubmixLevel, meta = (DisplayName = "Dry Level (dB)", AudioParam = "Volume", AudioParamClass = "SoundModulationParameterVolume", ClampMin = "-96.0", ClampMax = "0.0", UIMin = "-96.0", UIMax = "0.0"))
 	FSoundModulationDestinationSettings DryLevelModulation;
 	
 	/** Whether to send this Submix to AudioLink (when AudioLink is Enabled)*/
