Demote the VT NoMips warning to Display.

diff --git a/Engine/Source/Runtime/Engine/Private/Texture.cpp b/Engine/Source/Runtime/Engine/Private/Texture.cpp
--- a/Engine/Source/Runtime/Engine/Private/Texture.cpp
+++ b/Engine/Source/Runtime/Engine/Private/Texture.cpp
@@ -616,7 +616,7 @@ void UTexture::ValidateSettingsAfterImportOrEdit(bool * pRequiresNotifyMaterials
 			}
 			if (MipGenSettings == TMGS_NoMipmaps)
 			{
-				UE_LOG(LogTexture, Warning, TEXT("Virtual textures require mips and MipGenSettings is NoMipmaps: Forcing to SimpleAverage (%s)"), *GetName());
+				UE_LOG(LogTexture, Display, TEXT("Virtual textures require mips and MipGenSettings is NoMipmaps: Forcing to SimpleAverage (%s)"), *GetName());
 				MipGenSettings = TMGS_SimpleAverage;
 			}
 		}
diff --git a/Engine/Source/Runtime/Engine/Private/TextureDerivedData.cpp b/Engine/Source/Runtime/Engine/Private/TextureDerivedData.cpp
--- a/Engine/Source/Runtime/Engine/Private/TextureDerivedData.cpp
+++ b/Engine/Source/Runtime/Engine/Private/TextureDerivedData.cpp
@@ -928,7 +928,7 @@ static void GetTextureBuildSettings(
 	if (bVirtualTextureStreaming && MipGenSettings == TMGS_NoMipmaps)
 	{
 		MipGenSettings = TMGS_SimpleAverage;
-		UE_LOG(LogTexture, Warning, TEXT("Texture %s is virtual and has NoMips - forcing to SimpleAverage."), *Texture.GetPathName());
+		UE_LOG(LogTexture, Display, TEXT("Texture %s is virtual and has NoMips - forcing to SimpleAverage."), *Texture.GetPathName());
 	}
 	if (bVirtualTextureStreaming && MipGenSettings == TMGS_LeaveExistingMips)
 	{
