Moved GPU skin cache and compute framework to run asynchronously during scene update.
 - Moves GPU skin cache dispatches off the critical path.
 - Merges compute graph RDG builders into the main RDG builder to take advantage of parallelism.

diff --git a/Engine/Source/Runtime/Engine/Private/GPUSkinCache.cpp b/Engine/Source/Runtime/Engine/Private/GPUSkinCache.cpp
--- a/Engine/Source/Runtime/Engine/Private/GPUSkinCache.cpp
+++ b/Engine/Source/Runtime/Engine/Private/GPUSkinCache.cpp
@@ -604,6 +604,7 @@ protected:
 	int BoneInfluenceType;
 	bool bUse16BitBoneIndex;
 	bool bUse16BitBoneWeight;
+	bool bQueuedForDispatch = false;
 	uint32 InputWeightIndexSize;
 	uint32 InputWeightStride;
 	FShaderResourceViewRHIRef InputWeightStreamSRV;
@@ -1307,6 +1308,13 @@ void FGPUSkinCache::DoDispatch(FRHICommandList& RHICmdList)
 	int32 BatchCount = BatchDispatches.Num();
 	INC_DWORD_STAT_BY(STAT_GPUSkinCache_TotalNumChunks, BatchCount);
 
+	if (!BatchCount)
+	{
+		return;
+	}
+
+	TRACE_CPUPROFILER_EVENT_SCOPE(FGPUSkinCache::DoDispatch);
+
 	bool bCapture = BatchCount > 0 && GNumDispatchesToCapture > 0;
 	RenderCaptureInterface::FScopedCapture RenderCapture(bCapture, &RHICmdList, TEXT("GPUSkinCache"));
 	GNumDispatchesToCapture -= bCapture ? 1 : 0;
@@ -1322,6 +1330,9 @@ void FGPUSkinCache::DoDispatch(FRHICommandList& RHICmdList)
 		{
 			FDispatchEntry& DispatchItem = BatchDispatches[i];
 			PrepareUpdateSkinning(DispatchItem.SkinCacheEntry, DispatchItem.Section, DispatchItem.RevisionNumber, &BuffersToTransitionForSkinning);
+
+			// Clear the flag that this is queued for dispatch.
+			DispatchItem.SkinCacheEntry->bQueuedForDispatch = false;
 		}
 
 		{
@@ -1570,6 +1581,20 @@ void FGPUSkinCache::DoDispatch(FRHICommandList& RHICmdList)
 		TRACE_CPUPROFILER_EVENT_SCOPE(TransitionAllToReadable);
 		TransitionAllToReadable(RHICmdList, BuffersToTransitionToRead);
 	}
+
+#if RHI_RAYTRACING
+	if (IsGPUSkinCacheRayTracingSupported())
+	{
+		for (FGPUSkinCacheEntry* SkinCacheEntry : PendingProcessRTGeometryEntries)
+		{
+			ProcessRayTracingGeometryToUpdate(RHICmdList, SkinCacheEntry);
+		}
+
+		PendingProcessRTGeometryEntries.Reset();
+	}
+#endif
+
+	BatchDispatches.Reset();
 }
 
 void FGPUSkinCache::DoDispatch(FRHICommandList& RHICmdList, FGPUSkinCacheEntry* SkinCacheEntry, int32 Section, int32 RevisionNumber)
@@ -1829,6 +1854,7 @@ bool FGPUSkinCache::ProcessEntry(
 
 	if (bShouldBatchDispatches)
 	{
+		InOutEntry->bQueuedForDispatch = true;
 		BatchDispatches.Add({ InOutEntry, RevisionNumber, uint32(Section) });
 	}
 	else
@@ -1878,32 +1904,14 @@ void FGPUSkinCache::ProcessRayTracingGeometryToUpdate(FRHICommandList& RHICmdLis
 
 #endif
 
-void FGPUSkinCache::BeginBatchDispatch(FRHICommandList& RHICmdList)
+void FGPUSkinCache::BeginBatchDispatch()
 {
-	check(BatchDispatches.Num() == 0);
 	bShouldBatchDispatches = true;
 	DispatchCounter = 0;
 }
 
-void FGPUSkinCache::EndBatchDispatch(FRHICommandList& RHICmdList)
+void FGPUSkinCache::EndBatchDispatch()
 {
-	TRACE_CPUPROFILER_EVENT_SCOPE(FGPUSkinCache::EndBatchDispatch);
-
-	DoDispatch(RHICmdList);
-
-#if RHI_RAYTRACING
-	if (IsGPUSkinCacheRayTracingSupported())
-	{
-		for (FGPUSkinCacheEntry* SkinCacheEntry : PendingProcessRTGeometryEntries)
-		{
-			ProcessRayTracingGeometryToUpdate(RHICmdList, SkinCacheEntry);
-		}
-
-		PendingProcessRTGeometryEntries.Reset();
-	}
-#endif
-
-	BatchDispatches.Reset();
 	bShouldBatchDispatches = false;
 }
 
@@ -1915,6 +1923,24 @@ void FGPUSkinCache::Release(FGPUSkinCacheEntry*& SkinCacheEntry)
 		check(SkinCache);
 		SkinCache->PendingProcessRTGeometryEntries.Remove(SkinCacheEntry);
 
+		if (SkinCacheEntry->bQueuedForDispatch)
+		{
+			for (int32 Index = 0; Index < SkinCache->BatchDispatches.Num(); )
+			{
+				if (SkinCache->BatchDispatches[Index].SkinCacheEntry == SkinCacheEntry)
+				{
+					SkinCache->BatchDispatches.RemoveAtSwap(Index);
+
+					// Continue to search for other sections associated with this skin cache entry.
+				}
+				else
+				{
+					++Index;
+				}
+			}
+			SkinCacheEntry->bQueuedForDispatch = false;
+		}
+
 		ReleaseSkinCacheEntry(SkinCacheEntry);
 		SkinCacheEntry = nullptr;
 	}
@@ -2463,4 +2489,4 @@ void FGPUSkinCache::DrawVisualizationInfoText(const FName& GPUSkinCacheVisualiza
 
 #undef IMPLEMENT_SKIN_CACHE_SHADER_CLOTH
 #undef IMPLEMENT_SKIN_CACHE_SHADER_ALL_SKIN_TYPES
-#undef IMPLEMENT_SKIN_CACHE_SHADER
+#undef IMPLEMENT_SKIN_CACHE_SHADER
\ No newline at end of file
diff --git a/Engine/Source/Runtime/Engine/Private/LevelTick.cpp b/Engine/Source/Runtime/Engine/Private/LevelTick.cpp
--- a/Engine/Source/Runtime/Engine/Private/LevelTick.cpp
+++ b/Engine/Source/Runtime/Engine/Private/LevelTick.cpp
@@ -958,13 +958,11 @@ struct FSendAllEndOfFrameUpdates
 		if (InScene != nullptr)
 		{
 			GPUSkinCache = InScene->GetGPUSkinCache();
-			InScene->GetComputeTaskWorkers(ComputeTaskWorkers);
 			FeatureLevel = InScene->GetFeatureLevel();
 		}
 	}
 	
 	FGPUSkinCache* GPUSkinCache = nullptr;
-	TArray<IComputeTaskWorker*> ComputeTaskWorkers;
 	ERHIFeatureLevel::Type FeatureLevel = ERHIFeatureLevel::Num;
 
 #if WANTS_DRAW_MESH_EVENTS
@@ -977,11 +975,11 @@ void BeginSendEndOfFrameUpdatesDrawEvent(FSendAllEndOfFrameUpdates& SendAllEndOf
 	BEGIN_DRAW_EVENTF_GAMETHREAD(SendAllEndOfFrameUpdates, SendAllEndOfFrameUpdates.DrawEvent, TEXT("SendAllEndOfFrameUpdates"));
 	
 	ENQUEUE_RENDER_COMMAND(BeginDrawEventCommand)(
-		[GPUSkinCache = SendAllEndOfFrameUpdates.GPUSkinCache](FRHICommandListImmediate& RHICmdList)
+		[GPUSkinCache = SendAllEndOfFrameUpdates.GPUSkinCache](FRHICommandListImmediate&)
 		{
 			if (GPUSkinCache != nullptr)
 			{
-				GPUSkinCache->BeginBatchDispatch(RHICmdList);
+				GPUSkinCache->BeginBatchDispatch();
 			}
 		});
 }
@@ -991,28 +989,13 @@ DECLARE_GPU_STAT(GPUSkinCacheRayTracingGeometry);
 void EndSendEndOfFrameUpdatesDrawEvent(FSendAllEndOfFrameUpdates& SendAllEndOfFrameUpdates)
 {
 	ENQUEUE_RENDER_COMMAND(EndDrawEventCommand)(
-		[GPUSkinCache = SendAllEndOfFrameUpdates.GPUSkinCache, ComputeTaskWorkers = SendAllEndOfFrameUpdates.ComputeTaskWorkers, FeatureLevel = SendAllEndOfFrameUpdates.FeatureLevel](FRHICommandListImmediate& RHICmdList)
+		[GPUSkinCache = SendAllEndOfFrameUpdates.GPUSkinCache] (FRHICommandListImmediate&)
+	{
+		if (GPUSkinCache != nullptr)
 		{
-			SCOPED_GPU_STAT(RHICmdList, EndOfFrameUpdates);
-			TRACE_CPUPROFILER_EVENT_SCOPE(EndSendEndOfFrameUpdatesDrawEvent_RT);
-
-			if (GPUSkinCache != nullptr)
-			{
-				// Once all the individual components have received their DoDeferredRenderUpdates_Concurrent()
-				// allow the GPU Skin Cache system to update.
-				GPUSkinCache->EndBatchDispatch(RHICmdList);
-			}
-
-			for (IComputeTaskWorker* ComputeTaskWorker : ComputeTaskWorkers)
-			{
-				if (ComputeTaskWorker->HasWork(ComputeTaskExecutionGroup::EndOfFrameUpdate))
-				{
-					FRDGBuilder GraphBuilder(RHICmdList, RDG_EVENT_NAME("ComputeTaskWorker"));
-					ComputeTaskWorker->SubmitWork(GraphBuilder, ComputeTaskExecutionGroup::EndOfFrameUpdate, FeatureLevel);
-					GraphBuilder.Execute();
-				}
-			}
-		});
+			GPUSkinCache->EndBatchDispatch();
+		}
+	});
 
 	STOP_DRAW_EVENT_GAMETHREAD(SendAllEndOfFrameUpdates.DrawEvent);
 }
diff --git a/Engine/Source/Runtime/Engine/Public/GPUSkinCache.h b/Engine/Source/Runtime/Engine/Public/GPUSkinCache.h
--- a/Engine/Source/Runtime/Engine/Public/GPUSkinCache.h
+++ b/Engine/Source/Runtime/Engine/Public/GPUSkinCache.h
@@ -420,11 +420,14 @@ public:
 	void ProcessRayTracingGeometryToUpdate(FRHICommandList& RHICmdList, FGPUSkinCacheEntry* SkinCacheEntry);
 #endif // RHI_RAYTRACING
 
-	void BeginBatchDispatch(FRHICommandList& RHICmdList);
-	void EndBatchDispatch(FRHICommandList& RHICmdList);
+	void BeginBatchDispatch();
+	void EndBatchDispatch();
+	void ENGINE_API DoDispatch(FRHICommandList& RHICmdList);
 
 	inline ERHIFeatureLevel::Type GetFeatureLevel() const { return FeatureLevel; }
 
+	inline bool HasWork() const { return !BatchDispatches.IsEmpty(); }
+
 protected:
 	void MakeBufferTransitions(FRHICommandList& RHICmdList, TArray<FSkinCacheRWBuffer*>& Buffers, ERHIAccess ToState);
 	void GetBufferUAVs(const TArray<FSkinCacheRWBuffer*>& InBuffers, TArray<FRHIUnorderedAccessView*>& OutUAVs);
@@ -435,7 +438,6 @@ protected:
 	TArray<FDispatchEntry> BatchDispatches;
 
 	FRWBuffersAllocation* TryAllocBuffer(uint32 NumVertices, bool WithTangnents, bool UseIntermediateTangents, uint32 NumTriangles, FRHICommandList& RHICmdList, const FName& OwnerName);
-	void DoDispatch(FRHICommandList& RHICmdList);
 	void DoDispatch(FRHICommandList& RHICmdList, FGPUSkinCacheEntry* SkinCacheEntry, int32 Section, int32 RevisionNumber);
 	void DispatchUpdateSkinTangents(FRHICommandList& RHICmdList, FGPUSkinCacheEntry* Entry, int32 SectionIndex, FSkinCacheRWBuffer*& StagingBuffer, bool bTrianglePass);
 
diff --git a/Engine/Source/Runtime/Renderer/Private/RendererScene.cpp b/Engine/Source/Runtime/Renderer/Private/RendererScene.cpp
--- a/Engine/Source/Runtime/Renderer/Private/RendererScene.cpp
+++ b/Engine/Source/Runtime/Renderer/Private/RendererScene.cpp
@@ -76,6 +76,7 @@
 #include "ProfilingDebugging/CountersTrace.h"
 #include "SceneCulling/SceneCulling.h"
 #include "InstanceCulling/InstanceCullingOcclusionQuery.h"
+#include "ComputeWorkerInterface.h"
 
 #if RHI_RAYTRACING
 #include "Nanite/NaniteRayTracing.h"
@@ -406,28 +407,6 @@ FSceneViewState::FSceneViewState(ERHIFeatureLevel::Type FeatureLevel, FSceneView
 	}
 }
 
-void DestroyRenderResource(FRenderResource* RenderResource)
-{
-	if (RenderResource) 
-	{
-		ENQUEUE_RENDER_COMMAND(DestroySceneViewStateRenderResource)(
-			[RenderResource](FRHICommandList&)
-			{
-				RenderResource->ReleaseResource();
-				delete RenderResource;
-			});
-	}
-}
-
-void DestroyRWBuffer(FRWBuffer* RWBuffer)
-{
-	ENQUEUE_RENDER_COMMAND(DestroyRWBuffer)(
-		[RWBuffer](FRHICommandList&)
-		{
-			delete RWBuffer;
-		});
-}
-
 FSceneViewState::~FSceneViewState()
 {
 	CachedVisibilityChunk = NULL;
@@ -5233,6 +5212,25 @@ void FScene::UpdateAllPrimitiveSceneInfos(FRDGBuilder& GraphBuilder, EUpdateAllP
 	SCOPE_CYCLE_COUNTER(STAT_UpdateScenePrimitiveRenderThreadTime);
 
 	check(IsInRenderingThread());
+
+	UE::Tasks::FTask GPUSkinCacheTask;
+
+	if (GPUSkinCache && GPUSkinCache->HasWork())
+	{
+		GPUSkinCacheTask = GraphBuilder.AddCommandListSetupTask([this] (FRHICommandList& RHICmdList)
+		{
+			GPUSkinCache->DoDispatch(RHICmdList);
+		});
+	}
+
+	for (IComputeTaskWorker* ComputeTaskWorker : ComputeTaskWorkers)
+	{
+		if (ComputeTaskWorker->HasWork(ComputeTaskExecutionGroup::EndOfFrameUpdate))
+		{
+			ComputeTaskWorker->SubmitWork(GraphBuilder, ComputeTaskExecutionGroup::EndOfFrameUpdate, FeatureLevel);
+		}
+	}
+
 	FSceneRenderer::WaitForCleanUpTasks(GraphBuilder.RHICmdList);
 	FMaterialRenderProxy::UpdateDeferredCachedUniformExpressions();
 
@@ -6145,6 +6143,8 @@ void FScene::UpdateAllPrimitiveSceneInfos(FRDGBuilder& GraphBuilder, EUpdateAllP
 	UpdateCachedShadowState(SceneUpdateChangeSetStorage.GetPreUpdateSet(), SceneUpdateChangeSetStorage.GetPostUpdateSet());
 	ShadowScene->PostSceneUpdate(SceneUpdateChangeSetStorage.GetPreUpdateSet(), SceneUpdateChangeSetStorage.GetPostUpdateSet());
 
+	GPUSkinCacheTask.Wait();
+
 	// Must run prior to static mesh gathering as callbacks can evaluate uniform expression caches.
 	FVirtualTextureSystem::Get().CallPendingCallbacks();
 
