PR #10229: Allow BlueprintExtensions to add their graphs to their owning blueprint's `GetAllGraphs` result.

diff --git a/Engine/Source/Runtime/Engine/Private/Blueprint.cpp b/Engine/Source/Runtime/Engine/Private/Blueprint.cpp
--- a/Engine/Source/Runtime/Engine/Private/Blueprint.cpp
+++ b/Engine/Source/Runtime/Engine/Private/Blueprint.cpp
@@ -1886,6 +1886,24 @@ void UBlueprint::GetAllGraphs(TArray<UEdGraph*>& Graphs) const
 			}
 		}
 	}
+
+	for (const UBlueprintExtension* Extension : GetExtensions())
+	{
+		if (Extension != nullptr)
+		{
+			TArray<UEdGraph*> ExtensionGraphs;
+			Extension->GetAllGraphs(ExtensionGraphs);
+			for (int32 i = 0; i < ExtensionGraphs.Num(); ++i)
+			{
+				UEdGraph* Graph = ExtensionGraphs[i];
+				if(Graph)
+				{
+					Graphs.Add(Graph);
+					Graph->GetAllChildrenGraphs(Graphs);
+				}
+			}
+		}
+	}
 #endif // WITH_EDITORONLY_DATA
 }
 
diff --git a/Engine/Source/Runtime/Engine/Public/Blueprint/BlueprintExtension.h b/Engine/Source/Runtime/Engine/Public/Blueprint/BlueprintExtension.h
--- a/Engine/Source/Runtime/Engine/Public/Blueprint/BlueprintExtension.h
+++ b/Engine/Source/Runtime/Engine/Public/Blueprint/BlueprintExtension.h
@@ -70,6 +70,11 @@ public:
 		return HandleGatherSearchData(OwningBlueprint);
 	}
 
+	/**
+	 * Override this function to inform editor tools of any graphs in this extension (eg. blueprint diff, asset search)
+	 */
+	virtual void GetAllGraphs(TArray<UEdGraph*>& Graphs) const {}
+
 private:
 
 	/**
