[Quartz] Fixing up mispelled enum type.

diff --git a/Engine/Source/Runtime/Engine/Classes/Sound/QuartzQuantizationUtilities.h b/Engine/Source/Runtime/Engine/Classes/Sound/QuartzQuantizationUtilities.h
--- a/Engine/Source/Runtime/Engine/Classes/Sound/QuartzQuantizationUtilities.h
+++ b/Engine/Source/Runtime/Engine/Classes/Sound/QuartzQuantizationUtilities.h
@@ -178,7 +178,7 @@ enum class EQuartzDelegateType : uint8
 
 
 // An enumeration for specifying quantization boundary reference frame
-UENUM(BlueprintType)
+UENUM(BlueprintType, meta = (DisplayName="EQuartz Quantization Reference")) // display name to hide c++ typo for now
 enum class EQuarztQuantizationReference : uint8
 {
 	BarRelative				UMETA(DisplayName = "Bar Relative", ToolTip = "Will occur on the next occurence of this duration from the start of a bar (i.e. On beat 3)"),
