Added method to query the FOnMontageEndedDelegate of an UAnimInstance.

diff --git a/Engine/Source/Runtime/Engine/Classes/Animation/AnimInstance.h b/Engine/Source/Runtime/Engine/Classes/Animation/AnimInstance.h
--- a/Engine/Source/Runtime/Engine/Classes/Animation/AnimInstance.h
+++ b/Engine/Source/Runtime/Engine/Classes/Animation/AnimInstance.h
@@ -708,6 +708,10 @@ public:
 	********************************************************************************************* */
 public:	
 	ENGINE_API void Montage_SetEndDelegate(FOnMontageEnded & InOnMontageEnded, UAnimMontage* Montage = NULL);
+
+	/** Get pointer to Ended delegate for Montage.
+	    If Montage reference is NULL, it will pick the first active montage found.*/
+	ENGINE_API FOnMontageEnded* Montage_GetEndedDelegate(UAnimMontage* Montage = nullptr);
 	
 	ENGINE_API void Montage_SetBlendingOutDelegate(FOnMontageBlendingOutStarted & InOnMontageBlendingOut, UAnimMontage* Montage = NULL);
 	
diff --git a/Engine/Source/Runtime/Engine/Private/Animation/AnimInstance.cpp b/Engine/Source/Runtime/Engine/Private/Animation/AnimInstance.cpp
--- a/Engine/Source/Runtime/Engine/Private/Animation/AnimInstance.cpp
+++ b/Engine/Source/Runtime/Engine/Private/Animation/AnimInstance.cpp
@@ -2583,6 +2583,32 @@ void UAnimInstance::Montage_SetBlendingOutDelegate(FOnMontageBlendingOutStarted&
 	}
 }
 
+FOnMontageEnded* UAnimInstance::Montage_GetEndedDelegate(UAnimMontage* Montage)
+{
+	if (Montage)
+	{
+		FAnimMontageInstance* MontageInstance = GetActiveInstanceForMontage(Montage);
+		if (MontageInstance)
+		{
+			return &MontageInstance->OnMontageEnded;
+		}
+	}
+	else
+	{
+		// If no Montage reference, use first active one found.
+		for (int32 InstanceIndex = 0; InstanceIndex < MontageInstances.Num(); InstanceIndex++)
+		{
+			FAnimMontageInstance* MontageInstance = MontageInstances[InstanceIndex];
+			if (MontageInstance && MontageInstance->IsActive())
+			{
+				return &MontageInstance->OnMontageEnded;
+			}
+		}
+	}
+
+	return nullptr;
+}
+
 FOnMontageBlendingOutStarted* UAnimInstance::Montage_GetBlendingOutDelegate(UAnimMontage* Montage)
 {
 	if (Montage)
