[CisFix] fixed missing forward declaration for UEdGraph

diff --git a/Engine/Source/Runtime/Engine/Public/Blueprint/BlueprintExtension.h b/Engine/Source/Runtime/Engine/Public/Blueprint/BlueprintExtension.h
--- a/Engine/Source/Runtime/Engine/Public/Blueprint/BlueprintExtension.h
+++ b/Engine/Source/Runtime/Engine/Public/Blueprint/BlueprintExtension.h
@@ -9,6 +9,7 @@
 
 class UBlueprint;
 class FKismetCompilerContext;
+class UEdGraph;
 
 /**
  * Per-instance extension object that can be added to UBlueprint::Extensions in order to augment built-in blueprint functionality
