PR #10699: Fix for world partition server streaming is editor servers

diff --git a/Engine/Source/Runtime/Engine/Private/WorldPartition/WorldPartition.cpp b/Engine/Source/Runtime/Engine/Private/WorldPartition/WorldPartition.cpp
--- a/Engine/Source/Runtime/Engine/Private/WorldPartition/WorldPartition.cpp
+++ b/Engine/Source/Runtime/Engine/Private/WorldPartition/WorldPartition.cpp
@@ -1164,9 +1164,13 @@ bool UWorldPartition::IsServerStreamingEnabled() const
 				{
 					switch (UWorldPartition::GlobalEnableServerStreaming)
 					{
-					case 1:	bIsEnabled = true;
+					case 1:
+						bIsEnabled = true;
+						break;
 #if WITH_EDITOR
-					case 2: bIsEnabled = bIsPIE;
+					case 2:
+						bIsEnabled = bIsPIE;
+						break;
 #endif
 					}
 				}
