Level Instance: Fix crash with Disaster Recovery plugin being enabled

- Avoid trying to resolve data layers on uninitialized WP Actors when creating actor desc

#rnx

diff --git a/Engine/Source/Runtime/Engine/Private/WorldPartition/WorldPartitionActorDesc.cpp b/Engine/Source/Runtime/Engine/Private/WorldPartition/WorldPartitionActorDesc.cpp
--- a/Engine/Source/Runtime/Engine/Private/WorldPartition/WorldPartitionActorDesc.cpp
+++ b/Engine/Source/Runtime/Engine/Private/WorldPartition/WorldPartitionActorDesc.cpp
@@ -144,7 +144,8 @@ void FWorldPartitionActorDesc::Init(const AActor* InActor)
 			const UDataLayerManager* ResolvingDataLayerManager = !ActorOwningWorld->IsGameWorld() ? ActorOwningWorld->GetDataLayerManager() : ActorWorldPartition->GetDataLayerManager();
 			// The only case where it's normal not to have a valid ResolvingDataLayerManager is if OwningWorld is not game and not partitioned
 			// Also, consider that we don't need the DataLayerManager for any running commandlet except for cook
-			const bool bCanSkipDataLayerManagerValidation = (IsRunningCommandlet() && !IsRunningCookCommandlet()) || (!ActorOwningWorld->IsGameWorld() && !ActorOwningWorld->IsPartitionedWorld());
+			// If Actor's World Partition isn't initialized we don't need to resolve (ex: Disaster Recovery Plugin will save Actor Packages in temp location which will Create a temp ActorDesc)
+			const bool bCanSkipDataLayerManagerValidation = (IsRunningCommandlet() && !IsRunningCookCommandlet()) || (!ActorOwningWorld->IsGameWorld() && (!ActorOwningWorld->IsPartitionedWorld() || !ActorWorldPartition->IsInitialized()));
 			check(ResolvingDataLayerManager || bCanSkipDataLayerManagerValidation);
 			if (ResolvingDataLayerManager)
 			{
