World Partition Class Descriptor Registry: don't react to all BP changes, just recompiles as this is when the generated class is updated.

diff --git a/Engine/Source/Runtime/Engine/Private/WorldPartition/WorldPartitionClassDescRegistry.cpp b/Engine/Source/Runtime/Engine/Private/WorldPartition/WorldPartitionClassDescRegistry.cpp
--- a/Engine/Source/Runtime/Engine/Private/WorldPartition/WorldPartitionClassDescRegistry.cpp
+++ b/Engine/Source/Runtime/Engine/Private/WorldPartition/WorldPartitionClassDescRegistry.cpp
@@ -501,13 +501,17 @@ void FWorldPartitionClassDescRegistry::OnObjectPreSave(UObject* InObject, FObjec
 
 void FWorldPartitionClassDescRegistry::OnObjectPropertyChanged(UObject* InObject, FPropertyChangedEvent& InPropertyChangedEvent)
 {
-	if (UBlueprint* Blueprint = Cast<UBlueprint>(InObject))
+	// We only want to handle BP class compiles, not property changes.
+	if (!InPropertyChangedEvent.Property)
 	{
-		// The generated class is invalid in some situations, like renaming a blueprint, etc.
-		if (Blueprint->GeneratedClass && Blueprint->GeneratedClass->IsChildOf<AActor>())
+		if (UBlueprint* Blueprint = Cast<UBlueprint>(InObject))
 		{
-			PrefetchClassDesc(Blueprint->GeneratedClass);
-			UpdateClassDescriptor(InObject, true);
+			// The generated class is invalid in some situations, like renaming a blueprint, etc.
+			if (Blueprint->GeneratedClass && Blueprint->GeneratedClass->IsChildOf<AActor>())
+			{
+				PrefetchClassDesc(Blueprint->GeneratedClass);
+				UpdateClassDescriptor(InObject, true);
+			}
 		}
 	}
 }
