World Partition - Streaming Generation: Handle special case where world bounds is a point

diff --git a/Engine/Source/Runtime/Engine/Private/WorldPartition/RuntimeSpatialHash/RuntimeSpatialHashGridHelper.cpp b/Engine/Source/Runtime/Engine/Private/WorldPartition/RuntimeSpatialHash/RuntimeSpatialHashGridHelper.cpp
--- a/Engine/Source/Runtime/Engine/Private/WorldPartition/RuntimeSpatialHash/RuntimeSpatialHashGridHelper.cpp
+++ b/Engine/Source/Runtime/Engine/Private/WorldPartition/RuntimeSpatialHash/RuntimeSpatialHashGridHelper.cpp
@@ -152,7 +152,8 @@ FSquare2DGridHelper GetPartitionedActors(const FBox& WorldBounds, const FSpatial
 	// Create the hierarchical grids for the game
 	//	
 	FSquare2DGridHelper PartitionedActors = GetGridHelper(WorldBounds, FVector(Grid.Origin, 0), Grid.CellSize);
-	if (ensure(PartitionedActors.Levels.Num()) && WorldBounds.IsValid)
+	const bool bWorldBoundsValid = WorldBounds.IsValid && WorldBounds.GetExtent().Size2D() > 0;
+	if (ensure(PartitionedActors.Levels.Num()) && bWorldBoundsValid)
 	{
 		int32 IntersectingCellCount = 0;
 		FSquare2DGridHelper::FGridLevel& LastGridLevel = PartitionedActors.Levels.Last();
