[NetPhysics] Predictive Interpolation improvements
Clear replication target if:
- Changing replication mode
- Physics object no longer dynamic or sleeping
- When physics replication is switched off at runtime.

Improve softsnapping by only snapping the amount not covered by velocity of the expected replicated distance.

CVar to skip replication all together implemented into the new async flow.
CVar to disable softsnap in predictive interpolation.

diff --git a/Engine/Source/Runtime/Engine/Private/ActorReplication.cpp b/Engine/Source/Runtime/Engine/Private/ActorReplication.cpp
--- a/Engine/Source/Runtime/Engine/Private/ActorReplication.cpp
+++ b/Engine/Source/Runtime/Engine/Private/ActorReplication.cpp
@@ -1099,6 +1099,21 @@ bool AActor::CanTriggerResimulation() const
 void AActor::SetPhysicsReplicationMode(const EPhysicsReplicationMode ReplicationMode)
 {
 	PhysicsReplicationMode = ReplicationMode;
+
+	// When switching replication mode, also clear the current replication target.
+	if (UPrimitiveComponent* RootPrimComp = Cast<UPrimitiveComponent>(RootComponent))
+	{
+		if (UWorld* World = GetWorld())
+		{
+			if (FPhysScene* PhysScene = World->GetPhysicsScene())
+			{
+				if (IPhysicsReplication* PhysicsReplication = PhysScene->GetPhysicsReplication())
+				{
+					PhysicsReplication->RemoveReplicatedTarget(RootPrimComp);
+				}
+			}
+		}
+	}
 }
 
 EPhysicsReplicationMode AActor::GetPhysicsReplicationMode()
diff --git a/Engine/Source/Runtime/Engine/Private/PhysicsEngine/PhysicsReplication.cpp b/Engine/Source/Runtime/Engine/Private/PhysicsEngine/PhysicsReplication.cpp
--- a/Engine/Source/Runtime/Engine/Private/PhysicsEngine/PhysicsReplication.cpp
+++ b/Engine/Source/Runtime/Engine/Private/PhysicsEngine/PhysicsReplication.cpp
@@ -102,7 +102,7 @@ namespace PhysicsReplicationCVars
 		static FAutoConsoleVariableRef CVarExtrapolationTimeMultiplier(TEXT("np2.PredictiveInterpolation.ExtrapolationTimeMultiplier"), ExtrapolationTimeMultiplier, TEXT("Multiplier to adjust the time to extrapolate the target forward over, the time is based on current send-rate."));
 
 		static float MinExpectedDistanceCovered = 0.25f;
-		static FAutoConsoleVariableRef CVarMinExpectedDistanceCovered(TEXT("np2.PredictiveInterpolation.MinExpectedDistanceCovered"), MinExpectedDistanceCovered, TEXT("Value in percentage where 0.25 = 25%. How much of the expected distance based on replication velocity should the object have covered in a simulation tick to Not be considered stuck."));
+		static FAutoConsoleVariableRef CVarMinExpectedDistanceCovered(TEXT("np2.PredictiveInterpolation.MinExpectedDistanceCovered"), MinExpectedDistanceCovered, TEXT("Value between 0-1, in percentage where 0.25 = 25%. How much of the expected distance based on replication velocity should the object have covered in a simulation tick to Not be considered stuck."));
 
 		static float ErrorAccumulationDecreaseMultiplier = 0.5f;
 		static FAutoConsoleVariableRef CVarErrorAccumulationDecreaseMultiplier(TEXT("np2.PredictiveInterpolation.ErrorAccumulationDecreaseMultiplier"), ErrorAccumulationDecreaseMultiplier, TEXT("Multiplier to adjust how fast we decrease accumulated error time when we no longer accumulate error."));
@@ -121,6 +121,9 @@ namespace PhysicsReplicationCVars
 		
 		static bool bVelocityBased = true;
 		static FAutoConsoleVariableRef CVarVelocityBased(TEXT("np2.PredictiveInterpolation.VelocityBased"), bVelocityBased, TEXT("When true, predictive interpolation replication mode will only apply linear velocity and angular velocity"));
+		
+		static bool bDisableSoftSnap = false;
+		static FAutoConsoleVariableRef CVarDisableSoftSnap(TEXT("np2.PredictiveInterpolation.DisableSoftSnap"), bDisableSoftSnap, TEXT("When true, predictive interpolation will not use softsnap to correct the replication with when velocity fails. Hardsnap will still eventually kick in if replication can't reach the target."));
 	
 		static bool bAlwaysHardSnap = false;
 		static FAutoConsoleVariableRef CVarAlwaysHardSnap(TEXT("np2.PredictiveInterpolation.AlwaysHardSnap"), bAlwaysHardSnap, TEXT("When true, predictive interpolation replication mode will always hard snap. Used as a backup measure"));
@@ -233,7 +236,18 @@ void FPhysicsReplication::SetReplicatedTarget(Chaos::FConstPhysicsObjectHandle P
 
 void FPhysicsReplication::RemoveReplicatedTarget(UPrimitiveComponent* Component)
 {
+	if (Component == nullptr)
+	{
+		return;
+	}
+
+	// Remove from legacy flow
 	ComponentToTargets_DEPRECATED.Remove(Component);
+	
+	// Remove from FPhysicsObject flow
+	Chaos::FConstPhysicsObjectHandle PhysicsObject = Component->GetPhysicsObjectByName(NAME_None);
+	FReplicatedPhysicsTarget Target(PhysicsObject); // This creates a new but empty target and when it tries to update the current target in the async flow it will remove it from replication since it's empty.
+	ReplicatedTargetsQueue.Add(Target);
 }
 
 
@@ -336,21 +350,18 @@ void FPhysicsReplication::OnTick(float DeltaSeconds, TMap<TWeakObjectPtr<UPrimit
 	// PhysicsObject replication flow
 	for (FReplicatedPhysicsTarget& PhysicsTarget : ReplicatedTargetsQueue)
 	{
-		if (PhysicsTarget.TargetState.Flags & ERigidBodyFlags::NeedsUpdate)
-		{
-			const float PingSecondsOneWay = LocalPing * 0.5f * 0.001f;
-
-			// Queue up the target state for async replication
-			FPhysicsRepAsyncInputData AsyncInputData(PhysicsTarget.PhysicsObject);
-			AsyncInputData.TargetState = PhysicsTarget.TargetState;
-			AsyncInputData.Proxy = nullptr;
-			AsyncInputData.RepMode = PhysicsTarget.ReplicationMode;
-			AsyncInputData.ServerFrame = PhysicsTarget.ServerFrame;
-			AsyncInputData.FrameOffset = LocalFrameOffset;
-			AsyncInputData.LatencyOneWay = PingSecondsOneWay;
-
-			AsyncInput->InputData.Add(AsyncInputData);
-		}
+		const float PingSecondsOneWay = LocalPing * 0.5f * 0.001f;
+
+		// Queue up the target state for async replication
+		FPhysicsRepAsyncInputData AsyncInputData(PhysicsTarget.PhysicsObject);
+		AsyncInputData.TargetState = PhysicsTarget.TargetState;
+		AsyncInputData.Proxy = nullptr;
+		AsyncInputData.RepMode = PhysicsTarget.ReplicationMode;
+		AsyncInputData.ServerFrame = PhysicsTarget.ServerFrame;
+		AsyncInputData.FrameOffset = LocalFrameOffset;
+		AsyncInputData.LatencyOneWay = PingSecondsOneWay;
+
+		AsyncInput->InputData.Add(AsyncInputData);
 	}
 	ReplicatedTargetsQueue.Reset();
 
@@ -670,6 +681,11 @@ void FPhysicsReplicationAsync::OnPhysicsObjectUnregistered_Internal(Chaos::FCons
 
 void FPhysicsReplicationAsync::OnPreSimulate_Internal()
 {
+	if (FPhysicsReplication::ShouldSkipPhysicsReplication())
+	{
+		return;
+	}
+
 	if (const FPhysicsReplicationAsyncInput* AsyncInput = GetConsumerInput_Internal())
 	{
 		Chaos::FPBDRigidsSolver* RigidsSolver = static_cast<Chaos::FPBDRigidsSolver*>(GetSolver());
@@ -701,6 +717,13 @@ void FPhysicsReplicationAsync::OnPreSimulate_Internal()
 		// Update async targets with target input
 		for (const FPhysicsRepAsyncInputData& Input : AsyncInput->InputData)
 		{
+			if (Input.TargetState.Flags == ERigidBodyFlags::None)
+			{
+				// Remove replication target 
+				ObjectToTarget.Remove(Input.PhysicsObject);
+				continue;
+			}
+
 			UpdateRewindDataTarget(Input);
 			UpdateAsyncTarget(Input, RigidsSolver);
 		}
@@ -787,7 +810,7 @@ void FPhysicsReplicationAsync::ApplyTargetStatesAsync(const float DeltaSeconds,
 {
 	using namespace Chaos;
 
-	// Deprecated, BodyInstance flow
+	// Deprecated, legacy BodyInstance flow
 	for (const FPhysicsRepAsyncInputData& Input : InputData)
 	{
 		if (Input.Proxy != nullptr)
@@ -1161,7 +1184,8 @@ bool FPhysicsReplicationAsync::PredictiveInterpolation(Chaos::FPBDRigidParticleH
 	// Get the distance from the current position to the source position of our target state
 	const float SourceDistanceSqr = (Target.PrevPosTarget - Handle->X()).SizeSquared();
 	const bool bShouldSleep = (Target.TargetState.Flags & ERigidBodyFlags::Sleeping) != 0;
-	
+	const bool bCanSimulate = Handle->IsDynamic() || Handle->IsSleeping();
+
 	// Early out if we are within range of target, also apply target sleep state
 	if (SourceDistanceSqr < PhysicsReplicationCVars::PredictiveInterpolationCVars::EarlyOutDistanceSqr)
 	{
@@ -1178,12 +1202,13 @@ bool FPhysicsReplicationAsync::PredictiveInterpolation(Chaos::FPBDRigidParticleH
 			Handle->SetW(FVector(0, 0, 0));
 			Target.PrevLinVel = FVector(FVector(0, 0, 0));
 
-			if (bShouldSleep && !Handle->IsKinematic())
+			if (bShouldSleep && bCanSimulate)
 			{
 				RigidsSolver->GetEvolution()->SetParticleObjectState(Handle, Chaos::EObjectStateType::Sleeping);
 			}
 
-			return bShouldSleep && !PhysicsReplicationCVars::PredictiveInterpolationCVars::bDontClearTarget;
+			const bool bClearTarget = (!bCanSimulate || bShouldSleep) && !PhysicsReplicationCVars::PredictiveInterpolationCVars::bDontClearTarget;
+			return bClearTarget;
 		}
 	}
 
@@ -1213,9 +1238,10 @@ bool FPhysicsReplicationAsync::PredictiveInterpolation(Chaos::FPBDRigidParticleH
 	const FVector PrevDiff = CurrentState.Position - Target.PrevPos;
 	const float	ExpectedDistance = (Target.PrevLinVel * DeltaSeconds).Size();
 	const float CoveredDistance = FVector::DotProduct(PrevDiff, Target.PrevLinVel.GetSafeNormal());
-	
+	const float CoveredAplha = FMath::Clamp(CoveredDistance / ExpectedDistance, 0.f, 1.f);
+
 	// If the object is moving less than X% of the expected distance, accumulate error seconds
-	if (ExpectedDistance > UE_SMALL_NUMBER && (CoveredDistance / ExpectedDistance) < PhysicsReplicationCVars::PredictiveInterpolationCVars::MinExpectedDistanceCovered)
+	if (ExpectedDistance > UE_SMALL_NUMBER && CoveredAplha < PhysicsReplicationCVars::PredictiveInterpolationCVars::MinExpectedDistanceCovered)
 	{
 		Target.AccumulatedErrorSeconds += DeltaSeconds;
 		bSoftSnap = true;
@@ -1227,7 +1253,12 @@ bool FPhysicsReplicationAsync::PredictiveInterpolation(Chaos::FPBDRigidParticleH
 		bSoftSnap = true;
 	}
 
-	const bool bHardSnap = Target.AccumulatedErrorSeconds > PhysicsReplicationCVars::PredictiveInterpolationCVars::ErrorAccumulationSeconds || PhysicsReplicationCVars::PredictiveInterpolationCVars::bAlwaysHardSnap;
+	if (PhysicsReplicationCVars::PredictiveInterpolationCVars::bDisableSoftSnap && PhysicsReplicationCVars::PredictiveInterpolationCVars::bVelocityBased)
+	{
+		bSoftSnap = false;
+	}
+
+	const bool bHardSnap = !bCanSimulate || Target.AccumulatedErrorSeconds > PhysicsReplicationCVars::PredictiveInterpolationCVars::ErrorAccumulationSeconds || PhysicsReplicationCVars::PredictiveInterpolationCVars::bAlwaysHardSnap;
 	if (bHardSnap)
 	{
 		// Too much error so just snap state here and be done with it
@@ -1288,7 +1319,8 @@ bool FPhysicsReplicationAsync::PredictiveInterpolation(Chaos::FPBDRigidParticleH
 
 		if (bSoftSnap)
 		{
-			Handle->SetX(CurrentState.Position + (RepLinVel * DeltaSeconds));
+			const float AlphaToSoftSnap = 1.f - CoveredAplha; // If we covered 20% of the distance by velocity (1.f - 0.2f = 0.8f) cover 80% of the distance via softsnap.
+			Handle->SetX(CurrentState.Position + (RepLinVel * (AlphaToSoftSnap * DeltaSeconds)));
 			Handle->SetR(TargetRotBlended);
 		}
 
@@ -1315,7 +1347,7 @@ bool FPhysicsReplicationAsync::PredictiveInterpolation(Chaos::FPBDRigidParticleH
 		}
 	}
 	
-	return false;
+	return !bCanSimulate && !PhysicsReplicationCVars::PredictiveInterpolationCVars::bDontClearTarget; // If the object can't simulate, clear target
 }
 
 /** Compare states and trigger resimulation if needed */
