Fix FDisplayDebugManager not respecting DPI scaling

diff --git a/Engine/Source/Runtime/Engine/Private/UserInterface/Canvas.cpp b/Engine/Source/Runtime/Engine/Private/UserInterface/Canvas.cpp
--- a/Engine/Source/Runtime/Engine/Private/UserInterface/Canvas.cpp
+++ b/Engine/Source/Runtime/Engine/Private/UserInterface/Canvas.cpp
@@ -2205,15 +2205,15 @@ void FDisplayDebugManager::DrawString(const FString& InDebugString, const float&
 		DebugTextItem.Text = FText::FromString(InDebugString);
 		Canvas->DrawItem(DebugTextItem, FVector2D(CurrentPos.X + OptionalXOffset, CurrentPos.Y));
 
-		NextColumXPos = FMath::Max(NextColumXPos, CurrentPos.X + OptionalXOffset + DebugTextItem.DrawnSize.X);
-		CurrentPos.Y += FMath::Max(GetYStep(), DebugTextItem.DrawnSize.Y);
+		NextColumXPos = FMath::Max(NextColumXPos, CurrentPos.X + OptionalXOffset + (DebugTextItem.DrawnSize.X / Canvas->GetDPIScale()));
+		CurrentPos.Y += FMath::Max(GetYStep(), (DebugTextItem.DrawnSize.Y / Canvas->GetDPIScale()));
 		AddColumnIfNeeded();
 	}
 }
 
 float FDisplayDebugManager::GetTextScale() const
 {
-	return Canvas ? FMath::Max(Canvas->SizeX / 1920.f, 1.f) : 1.f;
+	return Canvas ? FMath::Max(Canvas->SizeX / (1920.f * Canvas->GetDPIScale()), 1.f) : 1.f;
 }
 
 void FDisplayDebugManager::AddColumnIfNeeded()
@@ -2221,7 +2221,7 @@ void FDisplayDebugManager::AddColumnIfNeeded()
 	if (Canvas)
 	{
 		const float YStep = GetYStep();
-		if ((CurrentPos.Y + YStep) > Canvas->SizeY)
+		if ((CurrentPos.Y + YStep) > Canvas->ClipY)
 		{
 			CurrentPos.Y = InitialPos.Y;
 			CurrentPos.X = NextColumXPos + YStep * 2.f;
