[Backout] - CL26858924 (Invalid fix for UE-181010)

Original CL Desc
-----------------------------------------------------------------
Addressing volume clamping of modulation params in SoundSubmix

Note: the min/max values of the properties are defined by the audio modulation plugin content settings. This change uses the current min value of volume as the default value of new submix dry param values.

diff --git a/Engine/Source/Runtime/Engine/Private/SoundSubmix.cpp b/Engine/Source/Runtime/Engine/Private/SoundSubmix.cpp
--- a/Engine/Source/Runtime/Engine/Private/SoundSubmix.cpp
+++ b/Engine/Source/Runtime/Engine/Private/SoundSubmix.cpp
@@ -37,19 +37,6 @@ USoundSubmixBase::USoundSubmixBase(const FObjectInitializer& ObjectInitializer)
 #endif // WITH_EDITORONLY_DATA
 {}
 
-static float GetMinVolumeParameter()
-{
-	if (Audio::IsModulationParameterRegistered("Volume"))
-	{
-		const Audio::FModulationParameter& VolumeParam = Audio::GetModulationParameter(TEXT("Volume"));
-		return VolumeParam.MinValue;
-	}
-	else
-	{
-		return -60.0f;
-	}
-}
-
 USoundSubmix::USoundSubmix(const FObjectInitializer& ObjectInitializer)
 	: Super(ObjectInitializer)
 	, bMuteWhenBackgrounded(0)
@@ -62,9 +49,7 @@ USoundSubmix::USoundSubmix(const FObjectInitializer& ObjectInitializer)
 {
 	OutputVolumeModulation.Value = 0.f;
 	WetLevelModulation.Value = 0.f;
-
-	DryLevelModulation.Value = GetMinVolumeParameter();
-
+	DryLevelModulation.Value = -96.f;
 }
 
 void USoundSubmix::Serialize(FArchive& Ar)
@@ -82,7 +67,7 @@ void USoundSubmix::Serialize(FArchive& Ar)
 		{
 			if (OutputVolume <= LinearNeg96dB)
 			{
-				OutputVolumeModulation.Value = GetMinVolumeParameter();
+				OutputVolumeModulation.Value = -96.f;
 			}
 			else
 			{
@@ -95,7 +80,7 @@ void USoundSubmix::Serialize(FArchive& Ar)
 		{
 			if (WetLevel <= LinearNeg96dB)
 			{
-				WetLevelModulation.Value = GetMinVolumeParameter();
+				WetLevelModulation.Value = -96.f;
 			}
 			else
 			{
@@ -108,7 +93,7 @@ void USoundSubmix::Serialize(FArchive& Ar)
 		{
 			if (DryLevel <= LinearNeg96dB)
 			{
-				DryLevelModulation.Value = GetMinVolumeParameter();
+				DryLevelModulation.Value = -96.f;
 			}
 			else
 			{
