Exposed version of BlendCurves with BlendOption param

diff --git a/Engine/Source/Runtime/Engine/Public/AnimationRuntime.h b/Engine/Source/Runtime/Engine/Public/AnimationRuntime.h
--- a/Engine/Source/Runtime/Engine/Public/AnimationRuntime.h
+++ b/Engine/Source/Runtime/Engine/Public/AnimationRuntime.h
@@ -61,6 +61,7 @@ FORCEINLINE void BlendTransform<ETransformBlendMode::Accumulate>(const FTransfor
 
 ENGINE_API void BlendCurves(const TArrayView<const FBlendedCurve> SourceCurves, const TArrayView<const float> SourceWeights, FBlendedCurve& OutCurve);
 void ENGINE_API BlendCurves(const TArrayView<const FBlendedCurve* const> SourceCurves, const TArrayView<const float> SourceWeights, FBlendedCurve& OutCurve);
+ENGINE_API void BlendCurves(const TArrayView<const FBlendedCurve* const> SourceCurves, const TArrayView<const float> SourceWeights, FBlendedCurve& OutCurve, ECurveBlendOption::Type BlendOption);
 
 
 /////////////////////////////////////////////////////////
