Avoid compiling Nanite input in "Nanite Pass" material experssion for platforms that can not support Nanite

diff --git a/Engine/Source/Runtime/Engine/Private/Materials/MaterialExpressions.cpp b/Engine/Source/Runtime/Engine/Private/Materials/MaterialExpressions.cpp
--- a/Engine/Source/Runtime/Engine/Private/Materials/MaterialExpressions.cpp
+++ b/Engine/Source/Runtime/Engine/Private/Materials/MaterialExpressions.cpp
@@ -19291,7 +19291,7 @@ int32 UMaterialExpressionNaniteReplace::Compile(class FMaterialCompiler* Compile
 	else
 	{
 		const int32 Arg1 = Default.Compile(Compiler);
-		const int32 Arg2 = Nanite.Compile(Compiler);
+		const int32 Arg2 = FDataDrivenShaderPlatformInfo::GetSupportsNanite(Compiler->GetShaderPlatform()) ? Nanite.Compile(Compiler) : Arg1;
 		return Compiler->NaniteReplace(Arg1, Arg2);
 	}
 }
