WorldPartition: Fix Cook not calling RerunConstructionScripts on sub container actors part of the PersistentLevel cell

#rnx

diff --git a/Engine/Source/Runtime/Engine/Private/WorldPartition/WorldPartitionLevelHelper.cpp b/Engine/Source/Runtime/Engine/Private/WorldPartition/WorldPartitionLevelHelper.cpp
--- a/Engine/Source/Runtime/Engine/Private/WorldPartition/WorldPartitionLevelHelper.cpp
+++ b/Engine/Source/Runtime/Engine/Private/WorldPartition/WorldPartitionLevelHelper.cpp
@@ -198,6 +198,14 @@ void FWorldPartitionLevelHelper::MoveExternalActorsToLevel(const TArray<FWorldPa
 			if (!bSameOuter)
 			{
 				Actor->Rename(nullptr, InLevel, REN_ForceNoResetLoaders);
+
+				// AActor::Rename will register components but doesn't call RerunConstructionScripts like AddLoadedActors does.
+				// If bIsWorldInitialized is false. RerunConstructionScripts will get called as part of UEditorEngine::InitializePhysicsSceneForSaveIfNecessary during Cell package save
+				// Current behavior is that the PersistentLevel Cell is initialized here (PopulateGeneratorPackageForCook) and other cells aren't yet (PopulateGeneratedPackageForCook)
+				if (InLevel->GetWorld()->bIsWorldInitialized)
+				{
+					Actor->RerunConstructionScripts();
+				}
 			}
 			else if (!InLevel->Actors.Contains(Actor))
 			{
