Properly account for probe collision in ChaosInterface::CreateGeometry using the same CollisionEnabledIntersection algorithm that is used in UpdatePhysicsFilterData.

In order to ensure that the underlying operation is identical to UpdatePhysicsFilterData, added methods CollisionEnabledFromFlags and CollisionEnabledToFlags.

diff --git a/Engine/Source/Runtime/Engine/Classes/Engine/EngineTypes.h b/Engine/Source/Runtime/Engine/Classes/Engine/EngineTypes.h
--- a/Engine/Source/Runtime/Engine/Classes/Engine/EngineTypes.h
+++ b/Engine/Source/Runtime/Engine/Classes/Engine/EngineTypes.h
@@ -1548,6 +1548,42 @@ FORCEINLINE ECollisionEnabled::Type CollisionEnabledIntersection(ECollisionEnabl
 	return ECollisionEnabled::NoCollision;
 }
 
+/** Convert a set of three bools into an ECollisionEnabled */
+FORCEINLINE ECollisionEnabled::Type CollisionEnabledFromFlags(const bool bQuery, const bool bPhysics, const bool bProbe)
+{
+	// Convert to ints for bit manipulation
+	const int32 QueryBit = static_cast<uint32>(bQuery);
+	const int32 PhysicsBit = static_cast<uint32>(bPhysics);
+	const int32 ProbeBit = static_cast<uint32>(bProbe);
+
+	// NOTE:
+	// We use &~ between physics and probe because we cannot have a case with
+	// both enabled, and in every case probe "beats" physics. See the collision
+	// type rules outlined in the comments in CollisionEnabledIntersection.
+	//
+	// index    type              bits: probe/physics/query
+	// ----------------------------------------------------
+	// 0        NoCollision       000
+	// 1        QueryOnly         001
+	// 2        PhysicsOnly       010
+	// 3        QueryAndPhysics   011
+	// 4        ProbeOnly         100 (& 110)
+	// 5        QueryAndProbe     101 (& 111)
+	return static_cast<ECollisionEnabled::Type>(
+		(QueryBit << 0) |
+		(PhysicsBit & (~ProbeBit)) << 1 |
+		(ProbeBit << 2));
+}
+
+/** Convert an ECollisionEnabled enum into a set of three bools */
+FORCEINLINE void CollisionEnabledToFlags(const ECollisionEnabled::Type CollisionEnabled, bool& bQuery, bool& bPhysics, bool& bProbe)
+{
+	const int32 Bits = static_cast<int32>(CollisionEnabled);
+	bQuery   = ((Bits & 0x1) != 0);
+	bPhysics = ((Bits & 0x2) != 0);
+	bProbe   = ((Bits & 0x4) != 0);
+}
+
 /** Describes type of wake/sleep event sent to the physics system */
 enum class ESleepEvent : uint8
 {
diff --git a/Engine/Source/Runtime/Engine/Private/PhysicsEngine/Experimental/ChaosInterfaceUtils.cpp b/Engine/Source/Runtime/Engine/Private/PhysicsEngine/Experimental/ChaosInterfaceUtils.cpp
--- a/Engine/Source/Runtime/Engine/Private/PhysicsEngine/Experimental/ChaosInterfaceUtils.cpp
+++ b/Engine/Source/Runtime/Engine/Private/PhysicsEngine/Experimental/ChaosInterfaceUtils.cpp
@@ -109,14 +109,19 @@ namespace ChaosInterface
 			NewShape->UpdateShapeBounds(InParams.WorldTransform);
 			NewShape->SetUserData(UserData);
 
-			// The following does nearly the same thing that happens in UpdatePhysicsFilterData.
-			// TODO: Refactor so that this code is not duplicated
-			const bool bBodyEnableSim = InParams.CollisionData.CollisionFlags.bEnableSimCollisionSimple || InParams.CollisionData.CollisionFlags.bEnableSimCollisionComplex;
-			const bool bBodyEnableQuery = InParams.CollisionData.CollisionFlags.bEnableQueryCollision;
-			const bool bShapeEnableSim = ShapeCollisionEnabled == ECollisionEnabled::QueryAndPhysics || ShapeCollisionEnabled == ECollisionEnabled::PhysicsOnly;
-			const bool bShapeEnableQuery = ShapeCollisionEnabled == ECollisionEnabled::QueryAndPhysics || ShapeCollisionEnabled == ECollisionEnabled::QueryOnly;
-			NewShape->SetSimEnabled(bBodyEnableSim && bShapeEnableSim);
-			NewShape->SetQueryEnabled(bBodyEnableQuery && bShapeEnableQuery);
+			// Combine shape and body collision-enabled using the same method as UpdatePhysicsFilterData,
+			// but with the added nastiness of needing collision mode represented as bool flags.
+			bool bEnableSim = InParams.CollisionData.CollisionFlags.bEnableSimCollisionSimple || InParams.CollisionData.CollisionFlags.bEnableSimCollisionComplex;
+			bool bEnableQuery = InParams.CollisionData.CollisionFlags.bEnableQueryCollision;
+			bool bEnableProbe = InParams.CollisionData.CollisionFlags.bEnableProbeCollision;
+			const ECollisionEnabled::Type BodyCollisionEnabled = CollisionEnabledFromFlags(bEnableQuery, bEnableSim, bEnableProbe);
+			const ECollisionEnabled::Type CombinedCollisionEnabled = CollisionEnabledIntersection(BodyCollisionEnabled, ShapeCollisionEnabled);
+			CollisionEnabledToFlags(CombinedCollisionEnabled, bEnableQuery, bEnableSim, bEnableProbe);
+
+			// Set the parameters on the shape
+			NewShape->SetSimEnabled(bEnableSim);
+			NewShape->SetQueryEnabled(bEnableQuery);
+			NewShape->SetIsProbe(bEnableProbe);
 
 			return NewShape;
 		};
