Only bring up SM6 warning (Nanite and VSM) for the RHI in use.

#rnx

diff --git a/Engine/Source/Runtime/Engine/Private/RendererSettings.cpp b/Engine/Source/Runtime/Engine/Private/RendererSettings.cpp
--- a/Engine/Source/Runtime/Engine/Private/RendererSettings.cpp
+++ b/Engine/Source/Runtime/Engine/Private/RendererSettings.cpp
@@ -22,6 +22,7 @@
 #include "Framework/Notifications/NotificationManager.h"
 #include "Misc/ConfigCacheIni.h"
 #include "Widgets/Notifications/SNotificationList.h"
+#include "DataDrivenShaderPlatformInfo.h"
 #endif
 
 /** The editor object. */
@@ -406,17 +407,27 @@ void URendererSettings::CheckForMissingShaderModels()
 
 #if PLATFORM_WINDOWS
 		// Gather all Windows shader format settings
-		GConfig->GetArray(TEXT("/Script/WindowsTargetPlatform.WindowsTargetSettings"), TEXT("D3D11TargetedShaderFormats"), D3D11TargetedShaderFormats, GEngineIni);
-		GConfig->GetArray(TEXT("/Script/WindowsTargetPlatform.WindowsTargetSettings"), TEXT("D3D12TargetedShaderFormats"), D3D12TargetedShaderFormats, GEngineIni);
-		//GConfig->GetArray(TEXT("/Script/WindowsTargetPlatform.WindowsTargetSettings"), TEXT("VulkanTargetedShaderFormats"), WindowsVulkanTargetedShaderFormats, GEngineIni);
-		GConfig->GetArray(TEXT("/Script/WindowsTargetPlatform.WindowsTargetSettings"), TEXT("TargetedRHIs"), WindowsTargetedRHIs, GEngineIni);
-		CopySM6Format(TEXT("PCD3D_SM6"), WindowsTargetedRHIs, D3D12TargetedShaderFormats);
-		//CopySM6Format(TEXT("SF_VULKAN_SM6"), WindowsTargetedRHIs, WindowsVulkanTargetedShaderFormats);
+		{
+			GConfig->GetArray(TEXT("/Script/WindowsTargetPlatform.WindowsTargetSettings"), TEXT("TargetedRHIs"), WindowsTargetedRHIs, GEngineIni);
+
+			// If using Vulkan in Windows, warn about Vulkan settings
+			if (IsVulkanPlatform(GMaxRHIShaderPlatform))
+			{
+				GConfig->GetArray(TEXT("/Script/WindowsTargetPlatform.WindowsTargetSettings"), TEXT("VulkanTargetedShaderFormats"), WindowsVulkanTargetedShaderFormats, GEngineIni);
+				CopySM6Format(TEXT("SF_VULKAN_SM6"), WindowsTargetedRHIs, WindowsVulkanTargetedShaderFormats);
+			}
+			else
+			{
+				GConfig->GetArray(TEXT("/Script/WindowsTargetPlatform.WindowsTargetSettings"), TEXT("D3D11TargetedShaderFormats"), D3D11TargetedShaderFormats, GEngineIni);
+				GConfig->GetArray(TEXT("/Script/WindowsTargetPlatform.WindowsTargetSettings"), TEXT("D3D12TargetedShaderFormats"), D3D12TargetedShaderFormats, GEngineIni);
+				CopySM6Format(TEXT("PCD3D_SM6"), WindowsTargetedRHIs, D3D12TargetedShaderFormats);
+			}
+		}
 #elif PLATFORM_LINUX
-		// TODO: Gather all Linux shader format settings
-		//GConfig->GetArray(TEXT("/Script/LinuxTargetPlatform.LinuxTargetSettings"), TEXT("VulkanTargetedShaderFormats"), LinuxVulkanTargetedShaderFormats, GEngineIni);
-		//GConfig->GetArray(TEXT("/Script/LinuxTargetPlatform.LinuxTargetSettings"), TEXT("TargetedRHIs"), LinuxTargetedRHIs, GEngineIni);
-		//CopySM6Format(TEXT("SF_VULKAN_SM6"), LinuxTargetedRHIs, LinuxVulkanTargetedShaderFormats);
+		// Gather all Linux shader format settings
+		GConfig->GetArray(TEXT("/Script/LinuxTargetPlatform.LinuxTargetSettings"), TEXT("VulkanTargetedShaderFormats"), LinuxVulkanTargetedShaderFormats, GEngineIni);
+		GConfig->GetArray(TEXT("/Script/LinuxTargetPlatform.LinuxTargetSettings"), TEXT("TargetedRHIs"), LinuxTargetedRHIs, GEngineIni);
+		CopySM6Format(TEXT("SF_VULKAN_SM6"), LinuxTargetedRHIs, LinuxVulkanTargetedShaderFormats);
 #elif PLATFORM_MAC
 		// TODO: Gather all Mac shader format settings
 #endif
diff --git a/Engine/Source/Runtime/Engine/Private/StaticMesh.cpp b/Engine/Source/Runtime/Engine/Private/StaticMesh.cpp
--- a/Engine/Source/Runtime/Engine/Private/StaticMesh.cpp
+++ b/Engine/Source/Runtime/Engine/Private/StaticMesh.cpp
@@ -61,6 +61,7 @@
 
 #include "DerivedDataCache.h"
 #include "DerivedDataRequestOwner.h"
+#include "DataDrivenShaderPlatformInfo.h"
 #else
 #include "Interfaces/ITargetPlatform.h"
 #endif // #if WITH_EDITOR
@@ -6559,17 +6560,27 @@ void UStaticMesh::CheckForMissingShaderModels()
 
 #if PLATFORM_WINDOWS
 		// Gather all Windows shader format settings
-		GConfig->GetArray(TEXT("/Script/WindowsTargetPlatform.WindowsTargetSettings"), TEXT("D3D11TargetedShaderFormats"), D3D11TargetedShaderFormats, GEngineIni);
-		GConfig->GetArray(TEXT("/Script/WindowsTargetPlatform.WindowsTargetSettings"), TEXT("D3D12TargetedShaderFormats"), D3D12TargetedShaderFormats, GEngineIni);
-		//GConfig->GetArray(TEXT("/Script/WindowsTargetPlatform.WindowsTargetSettings"), TEXT("VulkanTargetedShaderFormats"), WindowsVulkanTargetedShaderFormats, GEngineIni);
-		GConfig->GetArray(TEXT("/Script/WindowsTargetPlatform.WindowsTargetSettings"), TEXT("TargetedRHIs"), WindowsTargetedRHIs, GEngineIni);
-		CopySM6Format(TEXT("PCD3D_SM6"), WindowsTargetedRHIs, D3D12TargetedShaderFormats);
-		//CopySM6Format(TEXT("SF_VULKAN_SM6"), WindowsTargetedRHIs, WindowsVulkanTargetedShaderFormats);
+		{
+			GConfig->GetArray(TEXT("/Script/WindowsTargetPlatform.WindowsTargetSettings"), TEXT("TargetedRHIs"), WindowsTargetedRHIs, GEngineIni);
+
+			// If using Vulkan in Windows, warn about Vulkan settings
+			if (IsVulkanPlatform(GMaxRHIShaderPlatform))
+			{
+				GConfig->GetArray(TEXT("/Script/WindowsTargetPlatform.WindowsTargetSettings"), TEXT("VulkanTargetedShaderFormats"), WindowsVulkanTargetedShaderFormats, GEngineIni);
+				CopySM6Format(TEXT("SF_VULKAN_SM6"), WindowsTargetedRHIs, WindowsVulkanTargetedShaderFormats);
+			}
+			else
+			{
+				GConfig->GetArray(TEXT("/Script/WindowsTargetPlatform.WindowsTargetSettings"), TEXT("D3D11TargetedShaderFormats"), D3D11TargetedShaderFormats, GEngineIni);
+				GConfig->GetArray(TEXT("/Script/WindowsTargetPlatform.WindowsTargetSettings"), TEXT("D3D12TargetedShaderFormats"), D3D12TargetedShaderFormats, GEngineIni);
+				CopySM6Format(TEXT("PCD3D_SM6"), WindowsTargetedRHIs, D3D12TargetedShaderFormats);
+			}
+		}
 #elif PLATFORM_LINUX
-		// TODO: Gather all Linux shader format settings
-		//GConfig->GetArray(TEXT("/Script/LinuxTargetPlatform.LinuxTargetSettings"), TEXT("VulkanTargetedShaderFormats"), LinuxVulkanTargetedShaderFormats, GEngineIni);
-		//GConfig->GetArray(TEXT("/Script/LinuxTargetPlatform.LinuxTargetSettings"), TEXT("TargetedRHIs"), LinuxTargetedRHIs, GEngineIni);
-		//CopySM6Format(TEXT("SF_VULKAN_SM6"), LinuxTargetedRHIs, LinuxVulkanTargetedShaderFormats);
+		// Gather all Linux shader format settings
+		GConfig->GetArray(TEXT("/Script/LinuxTargetPlatform.LinuxTargetSettings"), TEXT("VulkanTargetedShaderFormats"), LinuxVulkanTargetedShaderFormats, GEngineIni);
+		GConfig->GetArray(TEXT("/Script/LinuxTargetPlatform.LinuxTargetSettings"), TEXT("TargetedRHIs"), LinuxTargetedRHIs, GEngineIni);
+		CopySM6Format(TEXT("SF_VULKAN_SM6"), LinuxTargetedRHIs, LinuxVulkanTargetedShaderFormats);
 #elif PLATFORM_MAC
 		// TODO: Gather all Mac shader format settings
 #endif
