Sparse class data requires bIsSparseClassDataSerializable be set to true in order to be written to an asset

diff --git a/Engine/Source/Runtime/Engine/Private/KismetSystemLibrary.cpp b/Engine/Source/Runtime/Engine/Private/KismetSystemLibrary.cpp
--- a/Engine/Source/Runtime/Engine/Private/KismetSystemLibrary.cpp
+++ b/Engine/Source/Runtime/Engine/Private/KismetSystemLibrary.cpp
@@ -4,6 +4,7 @@
 #include "AssetRegistry/ARFilter.h"
 #include "Blueprint/BlueprintSupport.h"
 #include "Engine/AssetManagerTypes.h"
+#include "Engine/BlueprintGeneratedClass.h"
 #include "HAL/FileManager.h"
 #include "Engine/World.h"
 #include "Misc/EngineVersion.h"
@@ -3220,6 +3221,17 @@ bool UKismetSystemLibrary::Generic_SetEditorProperty(UObject* Object, const FNam
 					{
 						return PropertyAccessUtil::BuildBasicChangeNotify(SparseProp, Object, ChangeNotifyMode);
 					});
+				if (*SparseDataAccessResult == EPropertyAccessResultFlags::Success)
+				{
+					if(UBlueprintGeneratedClass* AsBPGC = Cast<UBlueprintGeneratedClass>(Object->GetClass()))
+					{
+						AsBPGC->bIsSparseClassDataSerializable = true;
+					}
+					else
+					{
+						FFrame::KismetExecutionMessage(*FString::Printf(TEXT("Property '%s' on '%s' (%s) is in sparse class data but not owned by a BlueprintGeneratedClass and may not be saved"), *ObjectProp->GetName(), *Object->GetPathName(), *Object->GetClass()->GetName()), ELogVerbosity::Warning, UE::Blueprint::Private::PropertySetFailedWarning);
+					}
+				}
 			}
 		}
 	}
