Shader Compiler:  Promote a duplicate (if present) to become the in-flight job, when an in-flight job is removed.  Solves shader compiler hang that otherwise occurs where the duplicate job will never complete.

#rnx

diff --git a/Engine/Source/Runtime/Engine/Private/ShaderCompiler/ShaderCompiler.cpp b/Engine/Source/Runtime/Engine/Private/ShaderCompiler/ShaderCompiler.cpp
--- a/Engine/Source/Runtime/Engine/Private/ShaderCompiler/ShaderCompiler.cpp
+++ b/Engine/Source/Runtime/Engine/Private/ShaderCompiler/ShaderCompiler.cpp
@@ -1320,7 +1320,38 @@ int32 FShaderJobCache::RemoveAllPendingJobsWithId(uint32 InId)
 							FShaderJobData& JobData = GetShaderJobData(Job.JobCacheRef);
 
 							check(JobData.JobInFlight == &Job);
-							JobData.JobInFlight = nullptr;
+
+							// If we are removing an in-flight job, we need to promote a duplicate to be the new in-flight job, if present.
+							// Make sure the duplicate we choose doesn't have the same ID as what we're removing.
+							FShaderCommonCompileJob* DuplicateJob;
+							for (DuplicateJob = JobData.DuplicateJobsWaitList; DuplicateJob; DuplicateJob = DuplicateJob->NextLink)
+							{
+								if (DuplicateJob->Id != InId)
+								{
+									break;
+								}
+							}
+
+							if (DuplicateJob)
+							{
+								// Advance head if we are unlinking the head, then remove
+								if (JobData.DuplicateJobsWaitList == DuplicateJob)
+								{
+									JobData.DuplicateJobsWaitList = DuplicateJob->NextLink;
+								}
+								Unlink(*DuplicateJob);
+								RemoveDuplicateJob(DuplicateJob);
+
+								// Add it as pending at the appropriate priority
+								GShaderCompilerStats->RegisterNewPendingJob(*DuplicateJob);
+
+								DuplicateJob->PendingPriority = DuplicateJob->Priority;
+
+								LinkJobWithPriority(*DuplicateJob, (int32)DuplicateJob->Priority);
+							}
+
+							// DuplicateJob will be nullptr if there was no duplicate to promote
+							JobData.JobInFlight = DuplicateJob;
 						}
 					}
 
