Cleanup deprecated PlatformData and Resource from UTexture and friends

#rnx

diff --git a/Engine/Source/Runtime/Engine/Classes/Engine/Texture.h b/Engine/Source/Runtime/Engine/Classes/Engine/Texture.h
--- a/Engine/Source/Runtime/Engine/Classes/Engine/Texture.h
+++ b/Engine/Source/Runtime/Engine/Classes/Engine/Texture.h
@@ -21,7 +21,6 @@
 #endif
 #include "Engine/StreamableRenderAsset.h"
 #include "PerPlatformProperties.h"
-#include "Misc/FieldAccessor.h"
 #include "ImageCore.h"
 #if WITH_EDITORONLY_DATA
 #include "Misc/TVariant.h"
@@ -29,14 +28,6 @@
 #include "DerivedDataCacheKeyProxy.h"
 #endif
 
-#ifndef WITH_TEXTURE_RESOURCE_DEPRECATIONS
-#define WITH_TEXTURE_RESOURCE_DEPRECATIONS 1
-#endif
-
-#ifndef WITH_TEXTURE_PLATFORMDATA_DEPRECATIONS
-#define WITH_TEXTURE_PLATFORMDATA_DEPRECATIONS 1
-#endif
-
 #if WITH_EDITOR
 #include "Templates/DontCopy.h"
 #endif
@@ -1476,11 +1467,6 @@ private:
 	class FTextureResource* PrivateResourceRenderThread;
 
 public:
-#if WITH_TEXTURE_RESOURCE_DEPRECATIONS
-	UE_DEPRECATED(5.0, "Use GetResource() / SetResource() accessors instead.")
-	TFieldPtrAccessor<FTextureResource> Resource;
-#endif
-
 	ENGINE_API UTexture(FVTableHelper& Helper);
 	ENGINE_API virtual ~UTexture();
 	
diff --git a/Engine/Source/Runtime/Engine/Classes/Engine/Texture2D.h b/Engine/Source/Runtime/Engine/Classes/Engine/Texture2D.h
--- a/Engine/Source/Runtime/Engine/Classes/Engine/Texture2D.h
+++ b/Engine/Source/Runtime/Engine/Classes/Engine/Texture2D.h
@@ -12,7 +12,6 @@
 #include "TextureResource.h"
 #endif
 #include "Engine/TextureAllMipDataProviderFactory.h"
-#include "Misc/FieldAccessor.h"
 #include "Serialization/BulkData.h"
 #include "Texture2D.generated.h"
 
@@ -78,10 +77,6 @@ private:
 	FTexturePlatformData* PrivatePlatformData;
 
 public:
-#if WITH_TEXTURE_PLATFORMDATA_DEPRECATIONS
-	UE_DEPRECATED(5.0, "Use GetPlatformData() / SetPlatformData() accessors instead.")
-	TFieldPtrAccessor<FTexturePlatformData> PlatformData;
-#endif
 
 	/** Set the derived data for this texture on this platform. */
 	ENGINE_API void SetPlatformData(FTexturePlatformData* PlatformData);
diff --git a/Engine/Source/Runtime/Engine/Classes/Engine/Texture2DArray.cpp b/Engine/Source/Runtime/Engine/Classes/Engine/Texture2DArray.cpp
--- a/Engine/Source/Runtime/Engine/Classes/Engine/Texture2DArray.cpp
+++ b/Engine/Source/Runtime/Engine/Classes/Engine/Texture2DArray.cpp
@@ -95,9 +95,6 @@ static UTexture2DArray* GetDefaultTexture2DArray(const UTexture2DArray* Texture)
 UTexture2DArray::UTexture2DArray(const FObjectInitializer& ObjectInitializer)
 	: Super(ObjectInitializer)
 	, PrivatePlatformData(nullptr)
-	, PlatformData(
-		[this]()-> FTexturePlatformData* { return GetPlatformData(); },
-		[this](FTexturePlatformData* InPlatformData) { SetPlatformData(InPlatformData); })
 {
 #if WITH_EDITORONLY_DATA
 	SRGB = true;
diff --git a/Engine/Source/Runtime/Engine/Classes/Engine/Texture2DArray.h b/Engine/Source/Runtime/Engine/Classes/Engine/Texture2DArray.h
--- a/Engine/Source/Runtime/Engine/Classes/Engine/Texture2DArray.h
+++ b/Engine/Source/Runtime/Engine/Classes/Engine/Texture2DArray.h
@@ -20,9 +20,6 @@ class UTexture2DArray : public UTexture
 	FTexturePlatformData* PrivatePlatformData;
 
 public:
-	UE_DEPRECATED(5.1, "Use GetPlatformData() / SetPlatformData() accessors instead.")
-	TFieldPtrAccessor<FTexturePlatformData> PlatformData;
-
 	/** Set the derived data for this texture on this platform. */
 	ENGINE_API void SetPlatformData(FTexturePlatformData* PlatformData);
 	/** Get the derived data for this texture on this platform. */
diff --git a/Engine/Source/Runtime/Engine/Classes/Engine/TextureCube.h b/Engine/Source/Runtime/Engine/Classes/Engine/TextureCube.h
--- a/Engine/Source/Runtime/Engine/Classes/Engine/TextureCube.h
+++ b/Engine/Source/Runtime/Engine/Classes/Engine/TextureCube.h
@@ -6,7 +6,6 @@
 #include "CoreMinimal.h"
 #include "UObject/ObjectMacros.h"
 #include "Engine/Texture.h"
-#include "Misc/FieldAccessor.h"
 #include "TextureCube.generated.h"
 
 class FTextureResource;
@@ -20,11 +19,6 @@ class UTextureCube : public UTexture
 	FTexturePlatformData* PrivatePlatformData;
 
 public:
-#if WITH_TEXTURE_PLATFORMDATA_DEPRECATIONS
-	UE_DEPRECATED(5.0, "Use GetPlatformData() / SetPlatformData() accessors instead.")
-	TFieldPtrAccessor<FTexturePlatformData> PlatformData;
-#endif
-
 	/** Set the derived data for this texture on this platform. */
 	ENGINE_API void SetPlatformData(FTexturePlatformData* PlatformData);
 	/** Get the derived data for this texture on this platform. */
diff --git a/Engine/Source/Runtime/Engine/Classes/Engine/TextureCubeArray.h b/Engine/Source/Runtime/Engine/Classes/Engine/TextureCubeArray.h
--- a/Engine/Source/Runtime/Engine/Classes/Engine/TextureCubeArray.h
+++ b/Engine/Source/Runtime/Engine/Classes/Engine/TextureCubeArray.h
@@ -16,9 +16,6 @@ class UTextureCubeArray : public UTexture
 	FTexturePlatformData* PrivatePlatformData;
 
 public:
-	UE_DEPRECATED(5.1, "Use GetPlatformData() / SetPlatformData() accessors instead.")
-	TFieldPtrAccessor<FTexturePlatformData> PlatformData;
-
 	/** Set the derived data for this texture on this platform. */
 	ENGINE_API void SetPlatformData(FTexturePlatformData* PlatformData);
 	/** Get the derived data for this texture on this platform. */
diff --git a/Engine/Source/Runtime/Engine/Classes/Engine/VolumeTexture.h b/Engine/Source/Runtime/Engine/Classes/Engine/VolumeTexture.h
--- a/Engine/Source/Runtime/Engine/Classes/Engine/VolumeTexture.h
+++ b/Engine/Source/Runtime/Engine/Classes/Engine/VolumeTexture.h
@@ -21,8 +21,6 @@ class UVolumeTexture : public UTexture
 	FTexturePlatformData* PrivatePlatformData;
 
 public:
-	UE_DEPRECATED(5.1, "Use GetPlatformData() / SetPlatformData() accessors instead.")
-	TFieldPtrAccessor<FTexturePlatformData> PlatformData;
 
 	/** Set the derived data for this texture on this platform. */
 	ENGINE_API void SetPlatformData(FTexturePlatformData* PlatformData);
diff --git a/Engine/Source/Runtime/Engine/Private/Texture.cpp b/Engine/Source/Runtime/Engine/Private/Texture.cpp
--- a/Engine/Source/Runtime/Engine/Private/Texture.cpp
+++ b/Engine/Source/Runtime/Engine/Private/Texture.cpp
@@ -138,11 +138,6 @@ UTexture::UTexture(const FObjectInitializer& ObjectInitializer)
 	: Super(ObjectInitializer)
 	, PrivateResource(nullptr)
 	, PrivateResourceRenderThread(nullptr)
-#if WITH_TEXTURE_RESOURCE_DEPRECATIONS
-	, Resource(
-		[this]()-> FTextureResource* { return GetResource(); },
-		[this](FTextureResource* InTextureResource) { SetResource(InTextureResource); })
-#endif
 	, TextureReference(*new FTextureReference())
 {
 	Availability = ETextureAvailability::GPU;
diff --git a/Engine/Source/Runtime/Engine/Private/Texture2D.cpp b/Engine/Source/Runtime/Engine/Private/Texture2D.cpp
--- a/Engine/Source/Runtime/Engine/Private/Texture2D.cpp
+++ b/Engine/Source/Runtime/Engine/Private/Texture2D.cpp
@@ -55,11 +55,6 @@
 UTexture2D::UTexture2D(const FObjectInitializer& ObjectInitializer)
 	: Super(ObjectInitializer)
 	, PrivatePlatformData(nullptr)
-#if WITH_TEXTURE_PLATFORMDATA_DEPRECATIONS
-	, PlatformData(
-		[this]()-> FTexturePlatformData* { return GetPlatformData(); },
-		[this](FTexturePlatformData* InPlatformData) { SetPlatformData(InPlatformData); })
-#endif
 	, ResourceMem(nullptr)
 {
 	PendingUpdate = nullptr;
diff --git a/Engine/Source/Runtime/Engine/Private/TextureCube.cpp b/Engine/Source/Runtime/Engine/Private/TextureCube.cpp
--- a/Engine/Source/Runtime/Engine/Private/TextureCube.cpp
+++ b/Engine/Source/Runtime/Engine/Private/TextureCube.cpp
@@ -79,11 +79,6 @@ static UTextureCube* GetDefaultTextureCube(const UTextureCube* Texture)
 UTextureCube::UTextureCube(const FObjectInitializer& ObjectInitializer)
 	: Super(ObjectInitializer)
 	, PrivatePlatformData(nullptr)
-#if WITH_TEXTURE_PLATFORMDATA_DEPRECATIONS
-	, PlatformData(
-		[this]()-> FTexturePlatformData* { return GetPlatformData(); },
-		[this](FTexturePlatformData* InPlatformData) { SetPlatformData(InPlatformData); })
-#endif
 {
 	SRGB = true;
 }
diff --git a/Engine/Source/Runtime/Engine/Private/TextureCubeArray.cpp b/Engine/Source/Runtime/Engine/Private/TextureCubeArray.cpp
--- a/Engine/Source/Runtime/Engine/Private/TextureCubeArray.cpp
+++ b/Engine/Source/Runtime/Engine/Private/TextureCubeArray.cpp
@@ -358,9 +358,6 @@ static UTextureCubeArray* GetDefaultTextureCubeArray(const UTextureCubeArray* Te
 UTextureCubeArray::UTextureCubeArray(const FObjectInitializer& ObjectInitializer)
 	: Super(ObjectInitializer)
 	, PrivatePlatformData(nullptr)
-	, PlatformData(
-		[this]()-> FTexturePlatformData* { return GetPlatformData(); },
-		[this](FTexturePlatformData* InPlatformData) { SetPlatformData(InPlatformData); })
 {
 #if WITH_EDITORONLY_DATA
 	SRGB = true;
diff --git a/Engine/Source/Runtime/Engine/Private/VolumeTexture.cpp b/Engine/Source/Runtime/Engine/Private/VolumeTexture.cpp
--- a/Engine/Source/Runtime/Engine/Private/VolumeTexture.cpp
+++ b/Engine/Source/Runtime/Engine/Private/VolumeTexture.cpp
@@ -103,9 +103,6 @@ static UVolumeTexture* GetDefaultVolumeTexture(const UVolumeTexture* Texture)
 UVolumeTexture::UVolumeTexture(const FObjectInitializer& ObjectInitializer)
 	: Super(ObjectInitializer)
 	, PrivatePlatformData(nullptr)
-	, PlatformData(
-		[this]()-> FTexturePlatformData* { return GetPlatformData(); },
-		[this](FTexturePlatformData* InPlatformData) { SetPlatformData(InPlatformData); })
 {
 	SRGB = true;
 }
