Bindless
- Adding a shader compiler flag to force a shader to stay "bindful" (aka not bindless)
- Adding a DDSPI entry to allow a platform to keep certain utility shaders bindful
- Using new DDSPI entry to keep clear replacement and update texture shaders bindful for some platforms

diff --git a/Engine/Source/Runtime/Engine/Private/ShaderCompiler/ShaderCompiler.cpp b/Engine/Source/Runtime/Engine/Private/ShaderCompiler/ShaderCompiler.cpp
--- a/Engine/Source/Runtime/Engine/Private/ShaderCompiler/ShaderCompiler.cpp
+++ b/Engine/Source/Runtime/Engine/Private/ShaderCompiler/ShaderCompiler.cpp
@@ -7838,7 +7838,7 @@ void GlobalBeginCompileShader(
 	SET_SHADER_DEFINE(Input.Environment, PLATFORM_SUPPORTS_UNIFORM_BUFFER_OBJECTS, FDataDrivenShaderPlatformInfo::GetSupportsUniformBufferObjects(EShaderPlatform(Target.Platform)) ? 1 : 0);
 	
 	bool bEnableBindlessMacro = false;
-	if (RHIGetBindlessSupport(ShaderPlatform) != ERHIBindlessSupport::Unsupported)
+	if (RHIGetBindlessSupport(ShaderPlatform) != ERHIBindlessSupport::Unsupported && !Input.Environment.CompilerFlags.Contains(CFLAG_ForceBindful))
 	{
 		const bool bIsRaytracingShader = IsRayTracingShaderFrequency(Input.Target.GetFrequency());
 
diff --git a/Engine/Source/Runtime/Engine/Public/UpdateTextureShaders.h b/Engine/Source/Runtime/Engine/Public/UpdateTextureShaders.h
--- a/Engine/Source/Runtime/Engine/Public/UpdateTextureShaders.h
+++ b/Engine/Source/Runtime/Engine/Public/UpdateTextureShaders.h
@@ -4,10 +4,12 @@
 #pragma once
 
 #include "CoreMinimal.h"
+#include "DataDrivenShaderPlatformInfo.h"
 #include "ShaderParameters.h"
 #include "Shader.h"
 #include "GlobalShader.h"
 #include "DataDrivenShaderPlatformInfo.h"
+#include "ShaderCompilerCore.h"
 
 enum class EUpdateTextureValueType
 {
@@ -16,13 +18,38 @@ enum class EUpdateTextureValueType
 	Uint32
 };
 
-class FUpdateTexture2DSubresourceCS : public FGlobalShader
+class FUpdateTextureShaderBase : public FGlobalShader
+{
+public:
+	FUpdateTextureShaderBase() {}
+	FUpdateTextureShaderBase(const ShaderMetaType::CompiledShaderInitializerType& Initializer)
+		: FGlobalShader(Initializer)
+	{
+	}
+
+	static bool ShouldCompilePermutation(const FGlobalShaderPermutationParameters& Parameters)
+	{
+		return IsFeatureLevelSupported(Parameters.Platform, ERHIFeatureLevel::SM5);
+	}
+
+	static inline void ModifyCompilationEnvironment(const FGlobalShaderPermutationParameters& Parameters, FShaderCompilerEnvironment& OutEnvironment)
+	{
+		FGlobalShader::ModifyCompilationEnvironment(Parameters, OutEnvironment);
+
+		if (FDataDrivenShaderPlatformInfo::GetRequiresBindfulUtilityShaders(Parameters.Platform))
+		{
+			OutEnvironment.CompilerFlags.Add(CFLAG_ForceBindful);
+		}
+	}
+};
+
+class FUpdateTexture2DSubresourceCS : public FUpdateTextureShaderBase
 {
 	DECLARE_EXPORTED_SHADER_TYPE(FUpdateTexture2DSubresourceCS, Global, ENGINE_API);
 public:
 	FUpdateTexture2DSubresourceCS() {}
 	FUpdateTexture2DSubresourceCS(const ShaderMetaType::CompiledShaderInitializerType& Initializer)
-		: FGlobalShader( Initializer )
+		: FUpdateTextureShaderBase( Initializer )
 	{
 		SrcPosPitchParameter.Bind(Initializer.ParameterMap, TEXT("TSrcPosPitch"), SPF_Mandatory);
 		SrcBuffer.Bind(Initializer.ParameterMap, TEXT("TSrcBuffer"), SPF_Mandatory);
@@ -33,10 +60,7 @@ public:
 	static const TCHAR* GetSourceFilename() { return TEXT("/Engine/Private/UpdateTextureShaders.usf"); }
 	static const TCHAR* GetFunctionName() { return TEXT("TUpdateTexture2DSubresourceCS"); }
 
-	static bool ShouldCompilePermutation(const FGlobalShaderPermutationParameters& Parameters)
-	{
-		return IsFeatureLevelSupported(Parameters.Platform, ERHIFeatureLevel::SM5);
-	}
+	using FUpdateTextureShaderBase::ModifyCompilationEnvironment;
 
 	void SetParameters(FRHIBatchedShaderParameters& BatchedParameters, const FUpdateTextureRegion2D& UpdateRegionInBlocks, uint32 SrcElementPitch)
 	{
@@ -77,6 +101,8 @@ public:
 
 	static void ModifyCompilationEnvironment(const FGlobalShaderPermutationParameters& Parameters, FShaderCompilerEnvironment& OutEnvironment)
 	{
+		FUpdateTexture2DSubresourceCS::ModifyCompilationEnvironment(Parameters, OutEnvironment);
+
 		const TCHAR* Type = nullptr;
 		switch (ValueType)
 		{
@@ -106,13 +132,13 @@ inline TShaderRef<FUpdateTexture2DSubresourceCS> FUpdateTexture2DSubresourceCS::
 	}
 }
 
-class FUpdateTexture3DSubresourceCS : public FGlobalShader
+class FUpdateTexture3DSubresourceCS : public FUpdateTextureShaderBase
 {
 	DECLARE_EXPORTED_SHADER_TYPE(FUpdateTexture3DSubresourceCS, Global, ENGINE_API);
 public:
 	FUpdateTexture3DSubresourceCS() {}
 	FUpdateTexture3DSubresourceCS(const ShaderMetaType::CompiledShaderInitializerType& Initializer)
-		: FGlobalShader(Initializer)
+		: FUpdateTextureShaderBase(Initializer)
 	{
 		SrcPitchParameter.Bind(Initializer.ParameterMap, TEXT("SrcPitch"), SPF_Mandatory);
 		SrcDepthPitchParameter.Bind(Initializer.ParameterMap, TEXT("SrcDepthPitch"), SPF_Mandatory);
@@ -129,11 +155,6 @@ public:
 	
 	static const TCHAR* GetFunctionName() { return TEXT("UpdateTexture3DSubresourceCS"); };
 
-	static bool ShouldCompilePermutation(const FGlobalShaderPermutationParameters& Parameters)
-	{
-		return IsFeatureLevelSupported(Parameters.Platform, ERHIFeatureLevel::SM5);
-	}
-
 	LAYOUT_FIELD(FShaderParameter, SrcPitchParameter);
 	LAYOUT_FIELD(FShaderParameter, SrcDepthPitchParameter);
 	LAYOUT_FIELD(FShaderResourceParameter, SrcBuffer);
@@ -145,13 +166,13 @@ public:
 };
 
 template<uint32 ElementsPerThread>
-class TCopyDataCS : public FGlobalShader
+class TCopyDataCS : public FUpdateTextureShaderBase
 {
 	DECLARE_EXPORTED_SHADER_TYPE(TCopyDataCS, Global, ENGINE_API);
 public:
 	TCopyDataCS() {}
 	TCopyDataCS(const ShaderMetaType::CompiledShaderInitializerType& Initializer)
-		: FGlobalShader(Initializer)
+		: FUpdateTextureShaderBase(Initializer)
 	{
 		SrcBuffer.Bind(Initializer.ParameterMap, TEXT("SrcCopyBuffer"), SPF_Mandatory);
 		DestBuffer.Bind(Initializer.ParameterMap, TEXT("DestBuffer"), SPF_Mandatory);		
@@ -169,11 +190,6 @@ public:
 		return nullptr;
 	}
 
-	static bool ShouldCompilePermutation(const FGlobalShaderPermutationParameters& Parameters)
-	{
-		return IsFeatureLevelSupported(Parameters.Platform, ERHIFeatureLevel::SM5);
-	}
-
 	LAYOUT_FIELD(FShaderResourceParameter, SrcBuffer);
 	LAYOUT_FIELD(FShaderResourceParameter, DestBuffer);
 };
diff --git a/Engine/Source/Runtime/RHI/Private/DataDrivenShaderPlatformInfo.cpp b/Engine/Source/Runtime/RHI/Private/DataDrivenShaderPlatformInfo.cpp
--- a/Engine/Source/Runtime/RHI/Private/DataDrivenShaderPlatformInfo.cpp
+++ b/Engine/Source/Runtime/RHI/Private/DataDrivenShaderPlatformInfo.cpp
@@ -288,6 +288,7 @@ void FGenericDataDrivenShaderPlatformInfo::ParseDataDrivenShaderInfo(const FConf
 	GET_SECTION_BOOL_HELPER(bSupportsNNEShaders);
 	GET_SECTION_BOOL_HELPER(bSupportsShaderPipelines);
 	GET_SECTION_BOOL_HELPER(bSupportsUniformBufferObjects);
+	GET_SECTION_BOOL_HELPER(bRequiresBindfulUtilityShaders);
 #undef GET_SECTION_BOOL_HELPER
 #undef GET_SECTION_INT_HELPER
 #undef GET_SECTION_SUPPORT_HELPER
diff --git a/Engine/Source/Runtime/RHI/Public/DataDrivenShaderPlatformInfo.h b/Engine/Source/Runtime/RHI/Public/DataDrivenShaderPlatformInfo.h
--- a/Engine/Source/Runtime/RHI/Public/DataDrivenShaderPlatformInfo.h
+++ b/Engine/Source/Runtime/RHI/Public/DataDrivenShaderPlatformInfo.h
@@ -116,6 +116,7 @@ class FGenericDataDrivenShaderPlatformInfo
 	uint32 bSupportsNNEShaders: 1;
 	uint32 bSupportsShaderPipelines : 1;
 	uint32 bSupportsUniformBufferObjects : 1;
+	uint32 bRequiresBindfulUtilityShaders : 1;
 	// NOTE: When adding fields, you must also add to ParseDataDrivenShaderInfo!
 	uint32 bContainsValidPlatformInfo : 1;
 
@@ -745,6 +746,12 @@ public:
 		return Infos[Platform].bSupportsUniformBufferObjects;
 	}
 	
+	static FORCEINLINE_DEBUGGABLE const bool GetRequiresBindfulUtilityShaders(const FStaticShaderPlatform Platform)
+	{
+		check(IsValid(Platform));
+		return Infos[Platform].bRequiresBindfulUtilityShaders;
+	}
+	
 	static FORCEINLINE_DEBUGGABLE const bool IsValid(const FStaticShaderPlatform Platform)
 	{
 		return Infos[Platform].bContainsValidPlatformInfo;
diff --git a/Engine/Source/Runtime/RenderCore/Public/ClearReplacementShaders.h b/Engine/Source/Runtime/RenderCore/Public/ClearReplacementShaders.h
--- a/Engine/Source/Runtime/RenderCore/Public/ClearReplacementShaders.h
+++ b/Engine/Source/Runtime/RenderCore/Public/ClearReplacementShaders.h
@@ -4,6 +4,7 @@
 #pragma once
 
 #include "CoreMinimal.h"
+#include "DataDrivenShaderPlatformInfo.h"
 #include "GlobalShader.h"
 #include "Math/IntVector.h"
 #include "Math/UnrealMathSSE.h"
@@ -14,6 +15,7 @@
 #include "RHIDefinitions.h"
 #include "Serialization/MemoryLayout.h"
 #include "Shader.h"
+#include "ShaderCompilerCore.h"
 #include "ShaderCore.h"
 #include "ShaderParameterUtils.h"
 #include "ShaderParameters.h"
@@ -73,6 +75,11 @@ public:
 	{
 		FGlobalShader::ModifyCompilationEnvironment(Parameters, OutEnvironment);
 
+		if (FDataDrivenShaderPlatformInfo::GetRequiresBindfulUtilityShaders(Parameters.Platform))
+		{
+			OutEnvironment.CompilerFlags.Add(CFLAG_ForceBindful);
+		}
+
 		OutEnvironment.SetDefine(TEXT("ENABLE_CLEAR_VALUE"), !bZeroOutput);
 		OutEnvironment.SetDefine(TEXT("ENABLE_BOUNDS"), bEnableBounds);
 
diff --git a/Engine/Source/Runtime/RenderCore/Public/ShaderCompilerCore.h b/Engine/Source/Runtime/RenderCore/Public/ShaderCompilerCore.h
--- a/Engine/Source/Runtime/RenderCore/Public/ShaderCompilerCore.h
+++ b/Engine/Source/Runtime/RenderCore/Public/ShaderCompilerCore.h
@@ -136,6 +136,8 @@ enum ECompilerFlags
 	CFLAG_BindlessResources,
 	// Enabled if bindless samplers are enabled for the platform
 	CFLAG_BindlessSamplers,
+	// Force disable bindless resources and samplers on a shader
+	CFLAG_ForceBindful,
 	// EXPERIMENTAL: Run the shader re-writer that removes any unused functions/resources/types from source code before compilation.
 	CFLAG_RemoveDeadCode,
 	CFLAG_UseLegacyPreprocessor UE_DEPRECATED(5.3, "Legacy preprocessor has been removed as of UE 5.3; please report any issues with the new preprocessor to the UE rendering team."),
