Make sure to flush all level streaming before looking for player starts when entering PIE, so we can gather them from sublevels.

diff --git a/Engine/Source/Runtime/Engine/Private/GameInstance.cpp b/Engine/Source/Runtime/Engine/Private/GameInstance.cpp
--- a/Engine/Source/Runtime/Engine/Private/GameInstance.cpp
+++ b/Engine/Source/Runtime/Engine/Private/GameInstance.cpp
@@ -504,9 +504,18 @@ FGameInstancePIEResult UGameInstance::StartPlayInEditorGameInstance(ULocalPlayer
 		}
 
 		SlowTask.EnterProgressFrame(10, NSLOCTEXT("UnrealEd", "PIEFlushingLevelStreaming", "Starting PIE (Loading always loaded objects)..."));
+
 		// Make sure "always loaded" sub-levels are fully loaded
 		PlayWorld->FlushLevelStreaming(EFlushLevelStreamingType::Visibility);
 
+		UGameViewportClient* const GameViewport = GetGameViewportClient();
+		if (GameViewport != NULL && GameViewport->Viewport != NULL)
+		{
+			SlowTask.EnterProgressFrame(25, NSLOCTEXT("UnrealEd", "PIEWaitingForInitialLevelStreaming", "Starting PIE (Waiting for initial level streaming)..."));
+			// Stream any always loaded levels now that need to be loaded before the game starts
+			GEngine->BlockTillLevelStreamingCompleted(PlayWorld);
+		}
+
 		SlowTask.EnterProgressFrame(10, NSLOCTEXT("UnrealEd", "PIECreatingAISystem", "Starting PIE (Creating AI System)..."));
 		PlayWorld->CreateAISystem();
 
@@ -523,14 +532,13 @@ FGameInstancePIEResult UGameInstance::StartPlayInEditorGameInstance(ULocalPlayer
 			{
 				return FGameInstancePIEResult::Failure(FText::Format(NSLOCTEXT("UnrealEd", "Error_CouldntSpawnPlayer", "Couldn't spawn player: {0}"), FText::FromString(Error)));
 			}
-		}
 
-		UGameViewportClient* const GameViewport = GetGameViewportClient();
-		if (GameViewport != NULL && GameViewport->Viewport != NULL)
-		{
-			SlowTask.EnterProgressFrame(50, NSLOCTEXT("UnrealEd", "PIEWaitingForLevelStreaming", "Starting PIE (Waiting for level streaming)..."));
-			// Stream any levels now that need to be loaded before the game starts
-			GEngine->BlockTillLevelStreamingCompleted(PlayWorld);
+			if (GameViewport != NULL && GameViewport->Viewport != NULL)
+			{
+				SlowTask.EnterProgressFrame(25, NSLOCTEXT("UnrealEd", "PIEWaitingForLevelStreaming", "Starting PIE (Waiting for level streaming)..."));
+				// Stream any levels now that need to be loaded before the game starts as a result of spawning the local player
+				GEngine->BlockTillLevelStreamingCompleted(PlayWorld);
+			}
 		}
 
 		if (Params.NetMode == PIE_ListenServer)
@@ -1561,5 +1569,4 @@ void UGameInstance::RegisterReferencedObject(UObject* ObjectToReference)
 void UGameInstance::UnregisterReferencedObject(UObject* ObjectToReference)
 {
 	ReferencedObjects.RemoveSingleSwap(ObjectToReference);
-}
-
+}
\ No newline at end of file
