World Partition - Streaming Generation: If a non spatially loaded actor is attached to a spatially loaded one, ignore the parent setting and log a warning

diff --git a/Engine/Source/Runtime/Engine/Private/WorldPartition/WorldPartitionActorDescView.cpp b/Engine/Source/Runtime/Engine/Private/WorldPartition/WorldPartitionActorDescView.cpp
--- a/Engine/Source/Runtime/Engine/Private/WorldPartition/WorldPartitionActorDescView.cpp
+++ b/Engine/Source/Runtime/Engine/Private/WorldPartition/WorldPartitionActorDescView.cpp
@@ -77,12 +77,13 @@ bool FWorldPartitionActorDescView::GetIsSpatiallyLoaded() const
 		return false;
 	}
 
-	if (ParentView)
+	bool bIsSpatiallyLoaded = ActorDesc->GetIsSpatiallyLoaded();
+	if (bIsSpatiallyLoaded && ParentView)
 	{
-		return ParentView->GetIsSpatiallyLoaded();
+		bIsSpatiallyLoaded = ParentView->GetIsSpatiallyLoaded();
 	}
 	
-	return ActorDesc->GetIsSpatiallyLoaded();
+	return bIsSpatiallyLoaded;
 }
 
 bool FWorldPartitionActorDescView::GetActorIsHLODRelevant() const
