Fix to UMaterialExpressionNamedRerouteUsage::CompilePreview crashing when opening a material.

diff --git a/Engine/Source/Runtime/Engine/Private/Materials/MaterialExpressions.cpp b/Engine/Source/Runtime/Engine/Private/Materials/MaterialExpressions.cpp
--- a/Engine/Source/Runtime/Engine/Private/Materials/MaterialExpressions.cpp
+++ b/Engine/Source/Runtime/Engine/Private/Materials/MaterialExpressions.cpp
@@ -18931,11 +18931,15 @@ int32 UMaterialExpressionNamedRerouteUsage::CompilePreview(FMaterialCompiler* Co
 {
 	int32 ResultCodeChunk = Compile(Compiler, OutputIndex);
 
-	FExpressionInput Input = Declaration->TraceInputsToRealInput();
-	if (Input.Expression && Input.Expression->IsResultStrataMaterial(Input.OutputIndex))
+	if (IsDeclarationValid())
 	{
-		ResultCodeChunk = Compiler->StrataCompilePreview(ResultCodeChunk);
+		FExpressionInput Input = Declaration->TraceInputsToRealInput();
+		if (Input.Expression && Input.Expression->IsResultStrataMaterial(Input.OutputIndex))
+		{
+			ResultCodeChunk = Compiler->StrataCompilePreview(ResultCodeChunk);
+		}
 	}
+
 	return ResultCodeChunk;
 }
 
