Probable crash fix with PGO enabled by removing undefined behaviour: Server Crash in FGenericPhysicsInterface::RaycastSingle

diff --git a/Engine/Source/Runtime/Engine/Private/Collision/CollisionConversions.cpp b/Engine/Source/Runtime/Engine/Private/Collision/CollisionConversions.cpp
--- a/Engine/Source/Runtime/Engine/Private/Collision/CollisionConversions.cpp
+++ b/Engine/Source/Runtime/Engine/Private/Collision/CollisionConversions.cpp
@@ -414,7 +414,7 @@ EConvertQueryResult ConvertQueryImpactHit(const UWorld* World, const ChaosInterf
 }
 
 template <typename HitType>
-EConvertQueryResult ConvertTraceResults(bool& OutHasValidBlockingHit, const UWorld* World, int32 NumHits, HitType* Hits, float CheckLength, const FCollisionFilterData& QueryFilter, TArray<FHitResult>& OutHits, const FVector& StartLoc, const FVector& EndLoc, const FPhysicsGeometry& Geom, const FTransform& QueryTM, float MaxDistance, bool bReturnFaceIndex, bool bReturnPhysMat)
+EConvertQueryResult ConvertTraceResults(bool& OutHasValidBlockingHit, const UWorld* World, int32 NumHits, HitType* Hits, float CheckLength, const FCollisionFilterData& QueryFilter, TArray<FHitResult>& OutHits, const FVector& StartLoc, const FVector& EndLoc, const FPhysicsGeometry* Geom, const FTransform& QueryTM, float MaxDistance, bool bReturnFaceIndex, bool bReturnPhysMat)
 {
 	OutHits.Reserve(OutHits.Num() + NumHits);
 	EConvertQueryResult ConvertResult = EConvertQueryResult::Valid;
@@ -435,7 +435,7 @@ EConvertQueryResult ConvertTraceResults(bool& OutHasValidBlockingHit, const UWor
 			}
 
 			FHitResult& NewResult = OutHits[OutHits.AddDefaulted()];
-			if (ConvertQueryImpactHit(World, Hit, NewResult, CheckLength, QueryFilter, StartLoc, EndLoc, &Geom, QueryTM, bReturnFaceIndex, bReturnPhysMat) == EConvertQueryResult::Valid)
+			if (ConvertQueryImpactHit(World, Hit, NewResult, CheckLength, QueryFilter, StartLoc, EndLoc, Geom, QueryTM, bReturnFaceIndex, bReturnPhysMat) == EConvertQueryResult::Valid)
 			{
 				bHadBlockingHit |= NewResult.bBlockingHit;
 			}
@@ -456,7 +456,7 @@ EConvertQueryResult ConvertTraceResults(bool& OutHasValidBlockingHit, const UWor
 }
 
 template <typename Hit>
-EConvertQueryResult ConvertTraceResults(bool& OutHasValidBlockingHit, const UWorld* World, int32 NumHits, Hit* Hits, float CheckLength, const FCollisionFilterData& QueryFilter, FHitResult& OutHit, const FVector& StartLoc, const FVector& EndLoc, const FPhysicsGeometry& Geom, const FTransform& QueryTM, float MaxDistance, bool bReturnFaceIndex, bool bReturnPhysMat)
+EConvertQueryResult ConvertTraceResults(bool& OutHasValidBlockingHit, const UWorld* World, int32 NumHits, Hit* Hits, float CheckLength, const FCollisionFilterData& QueryFilter, FHitResult& OutHit, const FVector& StartLoc, const FVector& EndLoc, const FPhysicsGeometry* Geom, const FTransform& QueryTM, float MaxDistance, bool bReturnFaceIndex, bool bReturnPhysMat)
 {
 	const FVector Dir = (EndLoc - StartLoc).GetSafeNormal();
 	if constexpr (std::is_same_v<Hit, FHitSweep>)
@@ -466,19 +466,19 @@ EConvertQueryResult ConvertTraceResults(bool& OutHasValidBlockingHit, const UWor
 			SetInternalFaceIndex(Hits[0], FindFaceIndex(Hits[0], Dir));
 		}
 	}
-	EConvertQueryResult Result = ConvertQueryImpactHit(World, Hits[0], OutHit, CheckLength, QueryFilter, StartLoc, EndLoc, &Geom, QueryTM, bReturnFaceIndex, bReturnPhysMat);
+	EConvertQueryResult Result = ConvertQueryImpactHit(World, Hits[0], OutHit, CheckLength, QueryFilter, StartLoc, EndLoc, Geom, QueryTM, bReturnFaceIndex, bReturnPhysMat);
 	OutHasValidBlockingHit = Result == EConvertQueryResult::Valid;
 	return Result;
 }
 
-template EConvertQueryResult ConvertTraceResults<FHitSweep>(bool& OutHasValidBlockingHit, const UWorld* World, int32 NumHits, FHitSweep* Hits, float CheckLength, const FCollisionFilterData& QueryFilter, TArray<FHitResult>& OutHits, const FVector& StartLoc, const FVector& EndLoc, const FPhysicsGeometry& Geom, const FTransform& QueryTM, float MaxDistance, bool bReturnFaceIndex, bool bReturnPhysMat);
-template EConvertQueryResult ConvertTraceResults<FHitSweep>(bool& OutHasValidBlockingHit, const UWorld* World, int32 NumHits, FHitSweep* Hits, float CheckLength, const FCollisionFilterData& QueryFilter, FHitResult& OutHit, const FVector& StartLoc, const FVector& EndLoc, const FPhysicsGeometry& Geom, const FTransform& QueryTM, float MaxDistance, bool bReturnFaceIndex, bool bReturnPhysMat);
-template EConvertQueryResult ConvertTraceResults<ChaosInterface::FPTSweepHit>(bool& OutHasValidBlockingHit, const UWorld* World, int32 NumHits, ChaosInterface::FPTSweepHit* Hits, float CheckLength, const FCollisionFilterData& QueryFilter, TArray<FHitResult>& OutHits, const FVector& StartLoc, const FVector& EndLoc, const FPhysicsGeometry& Geom, const FTransform& QueryTM, float MaxDistance, bool bReturnFaceIndex, bool bReturnPhysMat);
-template EConvertQueryResult ConvertTraceResults<ChaosInterface::FPTSweepHit>(bool& OutHasValidBlockingHit, const UWorld* World, int32 NumHits, ChaosInterface::FPTSweepHit* Hits, float CheckLength, const FCollisionFilterData& QueryFilter, FHitResult& OutHit, const FVector& StartLoc, const FVector& EndLoc, const FPhysicsGeometry& Geom, const FTransform& QueryTM, float MaxDistance, bool bReturnFaceIndex, bool bReturnPhysMat);
-template EConvertQueryResult ConvertTraceResults<FHitRaycast>(bool& OutHasValidBlockingHit, const UWorld* World, int32 NumHits, FHitRaycast* Hits, float CheckLength, const FCollisionFilterData& QueryFilter, TArray<FHitResult>& OutHits, const FVector& StartLoc, const FVector& EndLoc, const FPhysicsGeometry& Geom, const FTransform& QueryTM, float MaxDistance, bool bReturnFaceIndex, bool bReturnPhysMat);
-template EConvertQueryResult ConvertTraceResults<FHitRaycast>(bool& OutHasValidBlockingHit, const UWorld* World, int32 NumHits, FHitRaycast* Hits, float CheckLength, const FCollisionFilterData& QueryFilter, FHitResult& OutHit, const FVector& StartLoc, const FVector& EndLoc, const FPhysicsGeometry& Geom, const FTransform& QueryTM, float MaxDistance, bool bReturnFaceIndex, bool bReturnPhysMat);
-template EConvertQueryResult ConvertTraceResults<ChaosInterface::FPTRaycastHit>(bool& OutHasValidBlockingHit, const UWorld* World, int32 NumHits, ChaosInterface::FPTRaycastHit* Hits, float CheckLength, const FCollisionFilterData& QueryFilter, TArray<FHitResult>& OutHits, const FVector& StartLoc, const FVector& EndLoc, const FPhysicsGeometry& Geom, const FTransform& QueryTM, float MaxDistance, bool bReturnFaceIndex, bool bReturnPhysMat);
-template EConvertQueryResult ConvertTraceResults<ChaosInterface::FPTRaycastHit>(bool& OutHasValidBlockingHit, const UWorld* World, int32 NumHits, ChaosInterface::FPTRaycastHit* Hits, float CheckLength, const FCollisionFilterData& QueryFilter, FHitResult& OutHit, const FVector& StartLoc, const FVector& EndLoc, const FPhysicsGeometry& Geom, const FTransform& QueryTM, float MaxDistance, bool bReturnFaceIndex, bool bReturnPhysMat);
+template EConvertQueryResult ConvertTraceResults<FHitSweep>(bool& OutHasValidBlockingHit, const UWorld* World, int32 NumHits, FHitSweep* Hits, float CheckLength, const FCollisionFilterData& QueryFilter, TArray<FHitResult>& OutHits, const FVector& StartLoc, const FVector& EndLoc, const FPhysicsGeometry* Geom, const FTransform& QueryTM, float MaxDistance, bool bReturnFaceIndex, bool bReturnPhysMat);
+template EConvertQueryResult ConvertTraceResults<FHitSweep>(bool& OutHasValidBlockingHit, const UWorld* World, int32 NumHits, FHitSweep* Hits, float CheckLength, const FCollisionFilterData& QueryFilter, FHitResult& OutHit, const FVector& StartLoc, const FVector& EndLoc, const FPhysicsGeometry* Geom, const FTransform& QueryTM, float MaxDistance, bool bReturnFaceIndex, bool bReturnPhysMat);
+template EConvertQueryResult ConvertTraceResults<ChaosInterface::FPTSweepHit>(bool& OutHasValidBlockingHit, const UWorld* World, int32 NumHits, ChaosInterface::FPTSweepHit* Hits, float CheckLength, const FCollisionFilterData& QueryFilter, TArray<FHitResult>& OutHits, const FVector& StartLoc, const FVector& EndLoc, const FPhysicsGeometry* Geom, const FTransform& QueryTM, float MaxDistance, bool bReturnFaceIndex, bool bReturnPhysMat);
+template EConvertQueryResult ConvertTraceResults<ChaosInterface::FPTSweepHit>(bool& OutHasValidBlockingHit, const UWorld* World, int32 NumHits, ChaosInterface::FPTSweepHit* Hits, float CheckLength, const FCollisionFilterData& QueryFilter, FHitResult& OutHit, const FVector& StartLoc, const FVector& EndLoc, const FPhysicsGeometry* Geom, const FTransform& QueryTM, float MaxDistance, bool bReturnFaceIndex, bool bReturnPhysMat);
+template EConvertQueryResult ConvertTraceResults<FHitRaycast>(bool& OutHasValidBlockingHit, const UWorld* World, int32 NumHits, FHitRaycast* Hits, float CheckLength, const FCollisionFilterData& QueryFilter, TArray<FHitResult>& OutHits, const FVector& StartLoc, const FVector& EndLoc, const FPhysicsGeometry* Geom, const FTransform& QueryTM, float MaxDistance, bool bReturnFaceIndex, bool bReturnPhysMat);
+template EConvertQueryResult ConvertTraceResults<FHitRaycast>(bool& OutHasValidBlockingHit, const UWorld* World, int32 NumHits, FHitRaycast* Hits, float CheckLength, const FCollisionFilterData& QueryFilter, FHitResult& OutHit, const FVector& StartLoc, const FVector& EndLoc, const FPhysicsGeometry* Geom, const FTransform& QueryTM, float MaxDistance, bool bReturnFaceIndex, bool bReturnPhysMat);
+template EConvertQueryResult ConvertTraceResults<ChaosInterface::FPTRaycastHit>(bool& OutHasValidBlockingHit, const UWorld* World, int32 NumHits, ChaosInterface::FPTRaycastHit* Hits, float CheckLength, const FCollisionFilterData& QueryFilter, TArray<FHitResult>& OutHits, const FVector& StartLoc, const FVector& EndLoc, const FPhysicsGeometry* Geom, const FTransform& QueryTM, float MaxDistance, bool bReturnFaceIndex, bool bReturnPhysMat);
+template EConvertQueryResult ConvertTraceResults<ChaosInterface::FPTRaycastHit>(bool& OutHasValidBlockingHit, const UWorld* World, int32 NumHits, ChaosInterface::FPTRaycastHit* Hits, float CheckLength, const FCollisionFilterData& QueryFilter, FHitResult& OutHit, const FVector& StartLoc, const FVector& EndLoc, const FPhysicsGeometry* Geom, const FTransform& QueryTM, float MaxDistance, bool bReturnFaceIndex, bool bReturnPhysMat);
 
 /** Util to convert an overlapped shape into a sweep hit result, returns whether it was a blocking hit. */
 template <typename THitLocation>
diff --git a/Engine/Source/Runtime/Engine/Private/Collision/CollisionConversions.h b/Engine/Source/Runtime/Engine/Private/Collision/CollisionConversions.h
--- a/Engine/Source/Runtime/Engine/Private/Collision/CollisionConversions.h
+++ b/Engine/Source/Runtime/Engine/Private/Collision/CollisionConversions.h
@@ -49,10 +49,10 @@ EConvertQueryResult ConvertQueryImpactHit(const UWorld* World, const ChaosInterf
  * @return	Whether all results passed NaN/Inf checks.
  */
 template <typename Hit>
-EConvertQueryResult ConvertTraceResults(bool& OutHasValidBlockingHit, const UWorld* World, int32 NumHits, Hit* Hits, float CheckLength, const FCollisionFilterData& QueryFilter, TArray<FHitResult>& OutHits, const FVector& StartLoc, const FVector& EndLoc, const FPhysicsGeometry& Geom, const FTransform& QueryTM, float MaxDistance, bool bReturnFaceIndex, bool bReturnPhysMat);
+EConvertQueryResult ConvertTraceResults(bool& OutHasValidBlockingHit, const UWorld* World, int32 NumHits, Hit* Hits, float CheckLength, const FCollisionFilterData& QueryFilter, TArray<FHitResult>& OutHits, const FVector& StartLoc, const FVector& EndLoc, const FPhysicsGeometry* Geom, const FTransform& QueryTM, float MaxDistance, bool bReturnFaceIndex, bool bReturnPhysMat);
 
 template <typename Hit>
-EConvertQueryResult ConvertTraceResults(bool& OutHasValidBlockingHit, const UWorld* World, int32 NumHits, Hit* Hits, float CheckLength, const FCollisionFilterData& QueryFilter, FHitResult& OutHits, const FVector& StartLoc, const FVector& EndLoc, const FPhysicsGeometry& Geom, const FTransform& QueryTM, float MaxDistance, bool bReturnFaceIndex, bool bReturnPhysMat);
+EConvertQueryResult ConvertTraceResults(bool& OutHasValidBlockingHit, const UWorld* World, int32 NumHits, Hit* Hits, float CheckLength, const FCollisionFilterData& QueryFilter, FHitResult& OutHits, const FVector& StartLoc, const FVector& EndLoc, const FPhysicsGeometry* Geom, const FTransform& QueryTM, float MaxDistance, bool bReturnFaceIndex, bool bReturnPhysMat);
 
 /** 
  * Util to convert physX or Chaos overlap query to our overlap result
diff --git a/Engine/Source/Runtime/Engine/Private/Collision/SceneQuery.cpp b/Engine/Source/Runtime/Engine/Private/Collision/SceneQuery.cpp
--- a/Engine/Source/Runtime/Engine/Private/Collision/SceneQuery.cpp
+++ b/Engine/Source/Runtime/Engine/Private/Collision/SceneQuery.cpp
@@ -486,7 +486,7 @@ bool TSceneCastCommonImp(const UWorld* World, typename Traits::TOutHits& OutHits
 		}
 		if (NumHits > 0 && !Traits::IsTest())
 		{
-			bool bSuccess = ConvertTraceResults(bBlockingHit, World, NumHits, Traits::GetHits(HitBufferSync), DeltaMag, Filter, OutHits, Start, End, *GeomInputs.GetGeometry(), StartTM, MinBlockingDistance, Params.bReturnFaceIndex, Params.bReturnPhysicalMaterial) == EConvertQueryResult::Valid;
+			bool bSuccess = ConvertTraceResults(bBlockingHit, World, NumHits, Traits::GetHits(HitBufferSync), DeltaMag, Filter, OutHits, Start, End, GeomInputs.GetGeometry(), StartTM, MinBlockingDistance, Params.bReturnFaceIndex, Params.bReturnPhysicalMaterial) == EConvertQueryResult::Valid;
 
 			if (!bSuccess)
 			{
diff --git a/Engine/Source/Runtime/Engine/Private/Components/PrimitiveComponent.cpp b/Engine/Source/Runtime/Engine/Private/Components/PrimitiveComponent.cpp
--- a/Engine/Source/Runtime/Engine/Private/Components/PrimitiveComponent.cpp
+++ b/Engine/Source/Runtime/Engine/Private/Components/PrimitiveComponent.cpp
@@ -3096,7 +3096,7 @@ bool UPrimitiveComponent::SweepComponent(FHitResult& OutHit, const FVector Start
 		ChaosInterface::SetFlags(BestHit, EHitFlags::Distance | EHitFlags::Normal | EHitFlags::Position | EHitFlags::FaceIndex);
 
 		bool bHasHit = false;
-		ConvertTraceResults<ChaosInterface::FSweepHit>(bHasHit, GetWorld(), 1, &BestHit, (End - Start).Size(), QueryFilter, OutHit, Start, End, Geometry, FTransform{ ShapeWorldRotation, Start }, 0.f, Params.bReturnFaceIndex, Params.bReturnPhysicalMaterial);
+		ConvertTraceResults<ChaosInterface::FSweepHit>(bHasHit, GetWorld(), 1, &BestHit, (End - Start).Size(), QueryFilter, OutHit, Start, End, &Geometry, FTransform{ ShapeWorldRotation, Start }, 0.f, Params.bReturnFaceIndex, Params.bReturnPhysicalMaterial);
 		return bHasHit;
 	}
 
