Fixed an issue where the physics state scale was out of sync w/ the component scale

#rnx

diff --git a/Engine/Source/Runtime/Engine/Private/PhysicsEngine/BodyInstance.cpp b/Engine/Source/Runtime/Engine/Private/PhysicsEngine/BodyInstance.cpp
--- a/Engine/Source/Runtime/Engine/Private/PhysicsEngine/BodyInstance.cpp
+++ b/Engine/Source/Runtime/Engine/Private/PhysicsEngine/BodyInstance.cpp
@@ -1948,15 +1948,11 @@ bool FBodyInstance::UpdateBodyScale(const FVector& InScale3D, bool bForceUpdate)
 
 			FKShapeElem* ShapeElem = FChaosUserData::Get<FKShapeElem>(FPhysicsInterface::GetUserData(ShapeHandle));
 
-			if(!ShapeElem)
-			{
-				continue;
-			}
 			switch(ConcreteType)
 			{
 				case ImplicitObjectType::Sphere:
 				{
-					if (!CHAOS_ENSURE(!bIsInstanced && !bIsScaled))
+					if (!ShapeElem || !CHAOS_ENSURE(!bIsInstanced && !bIsScaled))
 					{
 						// No support for Instanced, Scaled not supported as we bake the scale below
 						break;
@@ -1976,7 +1972,7 @@ bool FBodyInstance::UpdateBodyScale(const FVector& InScale3D, bool bForceUpdate)
 				}
 				case ImplicitObjectType::Box:
 				{
-					if (!CHAOS_ENSURE(!bIsScaled && !bIsInstanced))
+					if (!ShapeElem || !CHAOS_ENSURE(!bIsScaled && !bIsInstanced))
 					{
 						// No support for ScaledImplicit Box yet or Instanced, scale is baked below
 						break;
@@ -2021,7 +2017,7 @@ bool FBodyInstance::UpdateBodyScale(const FVector& InScale3D, bool bForceUpdate)
 				}
 				case ImplicitObjectType::Capsule:
 				{
-					if (!CHAOS_ENSURE(!bIsInstanced && !bIsScaled))
+					if (!ShapeElem || !CHAOS_ENSURE(!bIsInstanced && !bIsScaled))
 					{
 						// No support for Instanced
 						break;
@@ -2082,7 +2078,7 @@ bool FBodyInstance::UpdateBodyScale(const FVector& InScale3D, bool bForceUpdate)
 				}
 				case ImplicitObjectType::Convex:
 				{
-					if(!CHAOS_ENSURE(bIsInstanced || bIsScaled))
+					if (!ShapeElem || !CHAOS_ENSURE(bIsInstanced || bIsScaled))
 					{
 						// Expecting instanced or scaled.
 						break;
