TargetPlatform/DeviceManager/ProjectInfo: Logging: Downgrade log message to verbose to make cookcommandlet logs easier to read.
#rnx

diff --git a/Engine/Source/Runtime/Core/Private/Misc/UProjectInfo.cpp b/Engine/Source/Runtime/Core/Private/Misc/UProjectInfo.cpp
--- a/Engine/Source/Runtime/Core/Private/Misc/UProjectInfo.cpp
+++ b/Engine/Source/Runtime/Core/Private/Misc/UProjectInfo.cpp
@@ -6,7 +6,7 @@
 #include "Misc/FileHelper.h"
 #include "Misc/Paths.h"
 
-DECLARE_LOG_CATEGORY_EXTERN(LogUProjectInfo, Verbose, All);
+DECLARE_LOG_CATEGORY_EXTERN(LogUProjectInfo, Log, All);
 DEFINE_LOG_CATEGORY(LogUProjectInfo);
 
 FUProjectDictionary::FUProjectDictionary(const FString& InRootDir)
@@ -140,10 +140,10 @@ FUProjectDictionary& FUProjectDictionary::GetDefault()
 	static bool bHaveLoggedProjects = false;
 	if(!bHaveLoggedProjects)
 	{
-		UE_LOG(LogUProjectInfo, Log, TEXT("Found projects:"));
+		UE_LOG(LogUProjectInfo, Verbose, TEXT("Found projects:"));
 		for(TMap<FString, FString>::TConstIterator Iter(DefaultDictionary.ShortProjectNameDictionary); Iter; ++Iter)
 		{
-			UE_LOG(LogUProjectInfo, Log, TEXT("    %s: \"%s\""), *Iter.Key(), *Iter.Value());
+			UE_LOG(LogUProjectInfo, Verbose, TEXT("    %s: \"%s\""), *Iter.Key(), *Iter.Value());
 		}
 		bHaveLoggedProjects = true;
 	}
diff --git a/Engine/Source/Runtime/Engine/Private/DeviceProfiles/DeviceProfileManager.cpp b/Engine/Source/Runtime/Engine/Private/DeviceProfiles/DeviceProfileManager.cpp
--- a/Engine/Source/Runtime/Engine/Private/DeviceProfiles/DeviceProfileManager.cpp
+++ b/Engine/Source/Runtime/Engine/Private/DeviceProfiles/DeviceProfileManager.cpp
@@ -1314,13 +1314,13 @@ void UDeviceProfileManager::SetActiveDeviceProfile( UDeviceProfile* DeviceProfil
 {
 	ActiveDeviceProfile = DeviceProfile;
 
-	UE_LOG(LogDeviceProfileManager, Log, TEXT("Available device profiles:"));
+	UE_LOG(LogDeviceProfileManager, Verbose, TEXT("Available device profiles:"));
 	for (int32 Idx = 0; Idx < Profiles.Num(); ++Idx)
 	{
 		UDeviceProfile* Profile = Cast<UDeviceProfile>(Profiles[Idx]);
 		const void* TextureLODGroupsAddr = Profile ? Profile->TextureLODGroups.GetData() : nullptr;
 		const int32 NumTextureLODGroups = Profile ? Profile->TextureLODGroups.Num() : 0;
-		UE_LOG(LogDeviceProfileManager, Log, TEXT("\t[%p][%p %d] %s, "), Profile, TextureLODGroupsAddr, NumTextureLODGroups, Profile ? *Profile->GetName() : TEXT("None"));
+		UE_LOG(LogDeviceProfileManager, Verbose, TEXT("\t[%p][%p %d] %s, "), Profile, TextureLODGroupsAddr, NumTextureLODGroups, Profile ? *Profile->GetName() : TEXT("None"));
 	}
 
 	const void* TextureLODGroupsAddr = ActiveDeviceProfile ? ActiveDeviceProfile->TextureLODGroups.GetData() : nullptr;
