cluster union initialization now initializes the cluster particle to the component's initial transform and all added particles now have its child to parent set based on their location relative to the cluster union component transform at the time of the addition on the GT

diff --git a/Engine/Source/Runtime/Engine/Classes/PhysicsEngine/ClusterUnionComponent.h b/Engine/Source/Runtime/Engine/Classes/PhysicsEngine/ClusterUnionComponent.h
--- a/Engine/Source/Runtime/Engine/Classes/PhysicsEngine/ClusterUnionComponent.h
+++ b/Engine/Source/Runtime/Engine/Classes/PhysicsEngine/ClusterUnionComponent.h
@@ -150,8 +150,6 @@ public:
 	UFUNCTION()
 	bool IsComponentAdded(UPrimitiveComponent* Component) { return PerComponentData.Contains(Component) || PendingComponentSync.Contains(Component); }
 
-	bool HasReceivedTransform() const { return bHasReceivedTransform; }
-
 	// Multi-trace/sweep functions that only make sense in the context of a cluster union.
 	ENGINE_API bool LineTraceComponent(TArray<FHitResult>& OutHit, const FVector Start, const FVector End, ECollisionChannel TraceChannel, const struct FCollisionQueryParams& Params, const struct FCollisionResponseParams& ResponseParams, const struct FCollisionObjectQueryParams& ObjectParams);
 	ENGINE_API bool SweepComponent(TArray<FHitResult>& OutHit, const FVector Start, const FVector End, const FQuat& ShapeWorldRotation, const FPhysicsGeometry& Geometry, ECollisionChannel TraceChannel, const struct FCollisionQueryParams& Params, const struct FCollisionResponseParams& ResponseParams, const struct FCollisionObjectQueryParams& ObjectParams);
@@ -231,7 +229,6 @@ private:
 	ENGINE_API FPhysScene_Chaos* GetChaosScene() const;
 
 	Chaos::FClusterUnionPhysicsProxy* PhysicsProxy;
-	bool bHasReceivedTransform;
 
 	// User data to be able to tie the cluster particle back to this component.
 	FChaosUserData PhysicsUserData;
diff --git a/Engine/Source/Runtime/Engine/Classes/PhysicsEngine/ClusterUnionReplicatedProxyComponent.h b/Engine/Source/Runtime/Engine/Classes/PhysicsEngine/ClusterUnionReplicatedProxyComponent.h
--- a/Engine/Source/Runtime/Engine/Classes/PhysicsEngine/ClusterUnionReplicatedProxyComponent.h
+++ b/Engine/Source/Runtime/Engine/Classes/PhysicsEngine/ClusterUnionReplicatedProxyComponent.h
@@ -95,8 +95,7 @@ private:
 	FTimerHandle DeferSetChildToParentHandle;
 	ENGINE_API void DeferSetChildToParentChildUntilClusteredComponentInParentUnion();
 
-	FTimerHandle DeferAddComponentToClusterHandle;
-	ENGINE_API void DeferAddComponentToClusterHandleUntilInitialTransformUpdate();
+	ENGINE_API void AddComponentToCluster();
 
 	//~ Begin UActorComponent Interface
 public:
diff --git a/Engine/Source/Runtime/Engine/Private/PhysicsEngine/ClusterUnionComponent.cpp b/Engine/Source/Runtime/Engine/Private/PhysicsEngine/ClusterUnionComponent.cpp
--- a/Engine/Source/Runtime/Engine/Private/PhysicsEngine/ClusterUnionComponent.cpp
+++ b/Engine/Source/Runtime/Engine/Private/PhysicsEngine/ClusterUnionComponent.cpp
@@ -54,7 +54,6 @@ UClusterUnionComponent::UClusterUnionComponent(const FObjectInitializer& ObjectI
 	PhysicsProxy = nullptr;
 	SetIsReplicatedByDefault(true);
 	bComputeBoundsOnceForGame = false;
-	bHasReceivedTransform = false;
 	bHasCachedLocalBounds = false;
 #if WITH_EDITORONLY_DATA
 	bVisualizeComponent = true;
@@ -131,18 +130,32 @@ void UClusterUnionComponent::AddComponentToCluster(UPrimitiveComponent* InCompon
 	PendingData.BoneIds.Empty(Objects.Num());
 	PendingData.AccelerationPayloads.Empty(Objects.Num());
 
+	// We want to fix the ChildToParent relationship between the particles we're adding and the cluster union
+	// based on the current game-thread state, not the future physics-thread state as the two may diverge
+	// which could cause the alignment between the particle and the cluster union to be different from what
+	// would be expected.
+	TArray<FTransform> ChildToParents;
+	ChildToParents.Reserve(Objects.Num());
+
+	const FTransform CurrentComponentTransform = GetComponentTransform();
 	for (Chaos::FPhysicsObjectHandle Object : Objects)
 	{
-		const int32 BoneId = Chaos::FPhysicsObjectInterface::GetId(Object);
-		PendingData.BoneIds.Add(BoneId);
-
-		if (AccelerationStructure && Object)
+		if (Object)
 		{
-			const FBox HandleBounds = Interface->GetWorldBounds({ &Object, 1 });
-			FExternalSpatialAccelerationPayload Handle;
-			Handle.Initialize(InComponent, BoneId);
-			AccelerationStructure->UpdateElement(Handle, Chaos::TAABB<Chaos::FReal, 3>{HandleBounds.Min, HandleBounds.Max}, HandleBounds.IsValid != 0);
-			PendingData.AccelerationPayloads.Add(Handle);
+			const int32 BoneId = Chaos::FPhysicsObjectInterface::GetId(Object);
+			PendingData.BoneIds.Add(BoneId);
+
+			if (AccelerationStructure)
+			{
+				const FBox HandleBounds = Interface->GetWorldBounds({ &Object, 1 });
+				FExternalSpatialAccelerationPayload Handle;
+				Handle.Initialize(InComponent, BoneId);
+				AccelerationStructure->UpdateElement(Handle, Chaos::TAABB<Chaos::FReal, 3>{HandleBounds.Min, HandleBounds.Max}, HandleBounds.IsValid != 0);
+				PendingData.AccelerationPayloads.Add(Handle);
+			}
+
+			const FTransform CurrentParticleTransform = Interface->GetTransform(Object);
+			ChildToParents.Add(CurrentParticleTransform.GetRelativeTransform(CurrentComponentTransform));
 		}
 	}
 
@@ -153,6 +166,8 @@ void UClusterUnionComponent::AddComponentToCluster(UPrimitiveComponent* InCompon
 	PendingComponentSync.Add(InComponent, PendingData);
 
 	PhysicsProxy->AddPhysicsObjects_External(Objects);
+	ForceSetChildToParent(InComponent, PendingData.BoneIds, ChildToParents);
+
 	if (bRebuildGeometry)
 	{
 		if(!bIncrementalUnionBuild)
@@ -551,7 +566,7 @@ void UClusterUnionComponent::OnCreatePhysicsState()
 	InitData.UserData = static_cast<void*>(&PhysicsUserData);
 	InitData.ActorId = GetOwner()->GetUniqueID();
 	InitData.ComponentId = GetUniqueID();
-	InitData.bNeedsClusterXRInitialization = bHasAuthority;
+	InitData.InitialTransform = GetComponentTransform();
 
 	// Client needs to be set to unbreakable so the server is authoritative.
 	InitData.bUnbreakable = !bHasAuthority;
@@ -559,7 +574,6 @@ void UClusterUnionComponent::OnCreatePhysicsState()
 	// Only need to check connectivity on the server and have the client rely on replication to get the memo on when to release from cluster union.
 	InitData.bCheckConnectivity = bHasAuthority;
 
-	bHasReceivedTransform = false;
 	PhysicsProxy = new Chaos::FClusterUnionPhysicsProxy{ this, Parameters, InitData };
 	PhysicsProxy->Initialize_External();
 	if (FPhysScene_Chaos* Scene = GetChaosScene())
@@ -625,24 +639,11 @@ void UClusterUnionComponent::OnDestroyPhysicsState()
 
 void UClusterUnionComponent::OnReceiveReplicatedState(const FVector X, const FQuat R, const FVector V, const FVector W)
 {
-	if (!bHasReceivedTransform)
-	{
-		// First time, we just directly set component state
-		//
-		// NOTE: This is only needed because ClusterUnion has a flag bHasReceivedTransform
-		// which does not get updated until the component's transform is directly set.
-		// Until that flag is set, it's root particle will be in a disabled state and not
-		// have any children, therefore replication will be dead in the water.
-		SetWorldTransform(FTransform(R, X, GetRelativeScale3D()), false, nullptr, ETeleportType::TeleportPhysics);
-		SetPhysicsLinearVelocity(V);
-		SetAllPhysicsAngularVelocityInDegrees(W);
-	}
 }
 
 void UClusterUnionComponent::OnUpdateTransform(EUpdateTransformFlags UpdateTransformFlags, ETeleportType Teleport)
 {
 	USceneComponent::OnUpdateTransform(UpdateTransformFlags, Teleport);
-	bHasReceivedTransform = true;
 
 	if (PhysicsProxy && !(UpdateTransformFlags & EUpdateTransformFlags::SkipPhysicsUpdate))
 	{
diff --git a/Engine/Source/Runtime/Engine/Private/PhysicsEngine/ClusterUnionReplicatedProxyComponent.cpp b/Engine/Source/Runtime/Engine/Private/PhysicsEngine/ClusterUnionReplicatedProxyComponent.cpp
--- a/Engine/Source/Runtime/Engine/Private/PhysicsEngine/ClusterUnionReplicatedProxyComponent.cpp
+++ b/Engine/Source/Runtime/Engine/Private/PhysicsEngine/ClusterUnionReplicatedProxyComponent.cpp
@@ -123,10 +123,8 @@ void UClusterUnionReplicatedProxyComponent::PostRepNotifies()
 	const bool bIsInitialReplication = bNetUpdateParentClusterUnion || bNetUpdateChildClusteredComponent || bNetUpdateParticleBoneIds;
 	if (bIsInitialReplication)
 	{
-		if (!DeferAddComponentToClusterHandle.IsValid())
-		{
-			DeferAddComponentToClusterHandleUntilInitialTransformUpdate();
-		}
+		AddComponentToCluster();
+
 		bNetUpdateParentClusterUnion = false;
 		bNetUpdateChildClusteredComponent = false;
 		bNetUpdateParticleBoneIds = false;
@@ -144,58 +142,49 @@ void UClusterUnionReplicatedProxyComponent::PostRepNotifies()
 	}
 }
 
-void UClusterUnionReplicatedProxyComponent::DeferAddComponentToClusterHandleUntilInitialTransformUpdate()
+void UClusterUnionReplicatedProxyComponent::AddComponentToCluster()
 {
 	if (!ParentClusterUnion.IsValid() || !ChildClusteredComponent.IsValid() || IsPendingDeletion())
 	{
 		return;
 	}
 
-	DeferAddComponentToClusterHandle.Invalidate();
-
-	if (ParentClusterUnion->HasReceivedTransform())
-	{
-		// Need to check if we're *losing* bones instead and handle that situation as well as adding new bones into the cluster union.
-		// This extra check also does a bit of prevention on the client side from adding duplicate particles into the cluster union.
-		TSet<int32> NewBoneIdSet{ ParticleBoneIds };
-
-		TArray<int32> ToAdd;
-		ToAdd.Reserve(NewBoneIdSet.Num());
+	// Need to check if we're *losing* bones instead and handle that situation as well as adding new bones into the cluster union.
+	// This extra check also does a bit of prevention on the client side from adding duplicate particles into the cluster union.
+	TSet<int32> NewBoneIdSet{ ParticleBoneIds };
 
-		for (int32 BoneId : NewBoneIdSet)
-		{
-			if (!LastSyncedBoneIds.Contains(BoneId))
-			{
-				ToAdd.Add(BoneId);
-			}
-		}
+	TArray<int32> ToAdd;
+	ToAdd.Reserve(NewBoneIdSet.Num());
 
-		if (!ToAdd.IsEmpty())
+	for (int32 BoneId : NewBoneIdSet)
+	{
+		if (!LastSyncedBoneIds.Contains(BoneId))
 		{
-			ParentClusterUnion->AddComponentToCluster(ChildClusteredComponent.Get(), ToAdd);
+			ToAdd.Add(BoneId);
 		}
+	}
 
-		TArray<int32> ToRemove;
-		ToRemove.Reserve(LastSyncedBoneIds.Num());
-		for (int32 BoneId : LastSyncedBoneIds)
-		{
-			if (!NewBoneIdSet.Contains(BoneId))
-			{
-				ToRemove.Add(BoneId);
-			}
-		}
+	if (!ToAdd.IsEmpty())
+	{
+		ParentClusterUnion->AddComponentToCluster(ChildClusteredComponent.Get(), ToAdd);
+	}
 
-		if (!ToRemove.IsEmpty())
+	TArray<int32> ToRemove;
+	ToRemove.Reserve(LastSyncedBoneIds.Num());
+	for (int32 BoneId : LastSyncedBoneIds)
+	{
+		if (!NewBoneIdSet.Contains(BoneId))
 		{
-			ParentClusterUnion->RemoveComponentBonesFromCluster(ChildClusteredComponent.Get(), ToRemove);
+			ToRemove.Add(BoneId);
 		}
-
-		LastSyncedBoneIds = NewBoneIdSet;
 	}
-	else if (AActor* Owner = GetOwner())
+
+	if (!ToRemove.IsEmpty())
 	{
-		DeferAddComponentToClusterHandle = Owner->GetWorldTimerManager().SetTimerForNextTick(this, &UClusterUnionReplicatedProxyComponent::DeferAddComponentToClusterHandleUntilInitialTransformUpdate);
+		ParentClusterUnion->RemoveComponentBonesFromCluster(ChildClusteredComponent.Get(), ToRemove);
 	}
+
+	LastSyncedBoneIds = NewBoneIdSet;
 }
 
 void UClusterUnionReplicatedProxyComponent::DeferSetChildToParentChildUntilClusteredComponentInParentUnion()
diff --git a/Engine/Source/Runtime/Experimental/Chaos/Private/PhysicsProxy/ClusterUnionPhysicsProxy.cpp b/Engine/Source/Runtime/Experimental/Chaos/Private/PhysicsProxy/ClusterUnionPhysicsProxy.cpp
--- a/Engine/Source/Runtime/Experimental/Chaos/Private/PhysicsProxy/ClusterUnionPhysicsProxy.cpp
+++ b/Engine/Source/Runtime/Experimental/Chaos/Private/PhysicsProxy/ClusterUnionPhysicsProxy.cpp
@@ -123,13 +123,17 @@ namespace Chaos
 			Particle_Internal->SetPhysicsProxy(this);
 			Particle_Internal->GTGeometryParticle() = Particle_External.Get();
 			Particle_Internal->SetUnbreakable(InitData.bUnbreakable);
+			Particle_Internal->SetX(InitData.InitialTransform.GetTranslation());
+			Particle_Internal->SetR(InitData.InitialTransform.GetRotation());
 
 #if !(UE_BUILD_SHIPPING || UE_BUILD_TEST)
 			Particle_Internal->SetDebugName(MakeShared<FString, ESPMode::ThreadSafe>(FString::Printf(TEXT("%s"), *GetOwner()->GetName())));
 #endif
 
-			// On the client, we'd rather wait for the server to initialize the particle properly.
-			ClusterUnion->bNeedsXRInitialization = InitData.bNeedsClusterXRInitialization;
+			// For explicit cluster unions (i.e. cluster unions created via a cluster union component rather than the cluster group index), the cluster union itself is responsible for initializing
+			// the particle's XR from the component's transform. Therefore, we never want the PT to try and initialize the XR of the cluster union - the only changes to XR from the PT should
+			// be a result of physical simulation.
+			ClusterUnion->bNeedsXRInitialization = false;
 			ClusterUnion->bCheckConnectivity = InitData.bCheckConnectivity;
 		}
 	}
diff --git a/Engine/Source/Runtime/Experimental/Chaos/Public/PhysicsProxy/ClusterUnionPhysicsProxy.h b/Engine/Source/Runtime/Experimental/Chaos/Public/PhysicsProxy/ClusterUnionPhysicsProxy.h
--- a/Engine/Source/Runtime/Experimental/Chaos/Public/PhysicsProxy/ClusterUnionPhysicsProxy.h
+++ b/Engine/Source/Runtime/Experimental/Chaos/Public/PhysicsProxy/ClusterUnionPhysicsProxy.h
@@ -28,9 +28,9 @@ namespace Chaos
 	struct FClusterUnionInitData
 	{
 		void* UserData;
+		FTransform InitialTransform;
 		uint32 ActorId = INDEX_NONE;
 		uint32 ComponentId = INDEX_NONE;
-		bool bNeedsClusterXRInitialization = true;
 		bool bCheckConnectivity = true;
 		bool bUnbreakable = false;
 	};
