Fix bug where Modulation Setters on Submix didn't do anything

diff --git a/Engine/Source/Runtime/Engine/Classes/Sound/SoundSubmix.h b/Engine/Source/Runtime/Engine/Classes/Sound/SoundSubmix.h
--- a/Engine/Source/Runtime/Engine/Classes/Sound/SoundSubmix.h
+++ b/Engine/Source/Runtime/Engine/Classes/Sound/SoundSubmix.h
@@ -296,15 +296,15 @@ public:
 	float DryLevel;
 
 	/** The output volume of the submix in Decibels. Applied after submix effects and analysis are performed.*/
-	UPROPERTY(EditAnywhere, BlueprintReadWrite, Category = SubmixLevel, meta = (DisplayName = "Output Volume (dB)", AudioParam = "Volume", AudioParamClass = "SoundModulationParameterVolume", ClampMin = "-96.0", ClampMax = "0.0", UIMin = "-96.0", UIMax = "0.0"))
+	UPROPERTY(EditAnywhere, BlueprintSetter=SetOutputVolumeModulation, Category = SubmixLevel, meta = (DisplayName = "Output Volume (dB)", AudioParam = "Volume", AudioParamClass = "SoundModulationParameterVolume", ClampMin = "-96.0", ClampMax = "0.0", UIMin = "-96.0", UIMax = "0.0"))
 	FSoundModulationDestinationSettings OutputVolumeModulation;
 
 	/** The wet level of the submixin Decibels. Applied after submix effects and analysis are performed. */
-	UPROPERTY(EditAnywhere, BlueprintReadWrite, Category = SubmixLevel, meta = (DisplayName = "Wet Level (dB)", AudioParam = "Volume", AudioParamClass = "SoundModulationParameterVolume", ClampMin = "-96.0", ClampMax = "0.0", UIMin = "-96.0", UIMax = "0.0"))
+	UPROPERTY(EditAnywhere, BlueprintSetter=SetWetVolumeModulation, Category = SubmixLevel, meta = (DisplayName = "Wet Level (dB)", AudioParam = "Volume", AudioParamClass = "SoundModulationParameterVolume", ClampMin = "-96.0", ClampMax = "0.0", UIMin = "-96.0", UIMax = "0.0"))
 	FSoundModulationDestinationSettings WetLevelModulation;
 
 	/** The dry level of the submix in Decibels. Applied before submix effects and analysis are performed. */
-	UPROPERTY(EditAnywhere, BlueprintReadWrite, Category = SubmixLevel, meta = (DisplayName = "Dry Level (dB)", AudioParam = "Volume", AudioParamClass = "SoundModulationParameterVolume", ClampMin = "-96.0", ClampMax = "0.0", UIMin = "-96.0", UIMax = "0.0"))
+	UPROPERTY(EditAnywhere, BlueprintSetter=SetWetVolumeModulation, Category = SubmixLevel, meta = (DisplayName = "Dry Level (dB)", AudioParam = "Volume", AudioParamClass = "SoundModulationParameterVolume", ClampMin = "-96.0", ClampMax = "0.0", UIMin = "-96.0", UIMax = "0.0"))
 	FSoundModulationDestinationSettings DryLevelModulation;
 	
 	/** Whether to send this Submix to AudioLink (when AudioLink is Enabled)*/
@@ -408,6 +408,19 @@ protected:
 	ENGINE_API virtual void PostEditChangeProperty(struct FPropertyChangedEvent& PropertyChangedEvent) override;
 #endif
 
+	// Custom settors for Modulation Destinations
+	UFUNCTION(BlueprintCallable, BlueprintInternalUseOnly)
+	void SetOutputVolumeModulation(const FSoundModulationDestinationSettings& InVolMod);
+
+	UFUNCTION(BlueprintCallable, BlueprintInternalUseOnly)
+	void SetWetVolumeModulation(const FSoundModulationDestinationSettings& InVolMod);
+
+	UFUNCTION(BlueprintCallable, BlueprintInternalUseOnly)
+	void SetDryVolumeModulation(const FSoundModulationDestinationSettings& InVolMod);
+
+	// Send submix modulation changes to audio devices
+	void PushModulationChanges();
+
 	// State handling for bouncing output.
 	TUniquePtr<Audio::FAudioRecordingData> RecordingData;
 };
diff --git a/Engine/Source/Runtime/Engine/Private/SoundSubmix.cpp b/Engine/Source/Runtime/Engine/Private/SoundSubmix.cpp
--- a/Engine/Source/Runtime/Engine/Private/SoundSubmix.cpp
+++ b/Engine/Source/Runtime/Engine/Private/SoundSubmix.cpp
@@ -495,14 +495,7 @@ void USoundSubmix::PostEditChangeProperty(struct FPropertyChangedEvent& Property
 			{
 				USoundSubmix* SoundSubmix = this;
 
-				TSet<TObjectPtr<USoundModulatorBase>> NewVolumeMod = OutputVolumeModulation.Modulators;
-				TSet<TObjectPtr<USoundModulatorBase>> NewWetLevelMod = WetLevelModulation.Modulators;
-				TSet<TObjectPtr<USoundModulatorBase>> NewDryLevelMod = DryLevelModulation.Modulators;
-
-				AudioDeviceManager->IterateOverAllDevices([SoundSubmix, VolMod = MoveTemp(NewVolumeMod), WetMod = MoveTemp(NewWetLevelMod), DryMod = MoveTemp(NewDryLevelMod)](Audio::FDeviceId Id, FAudioDevice* Device) mutable
-				{
-					Device->UpdateSubmixModulationSettings(SoundSubmix, VolMod, WetMod, DryMod);
-				});
+				PushModulationChanges();
 
 				float NewVolumeModBase = OutputVolumeModulation.Value;
 				float NewWetModBase = WetLevelModulation.Value;
@@ -525,6 +518,39 @@ void USoundSubmix::PostEditChangeProperty(struct FPropertyChangedEvent& Property
 }
 #endif
 
+void USoundSubmix::SetOutputVolumeModulation(const FSoundModulationDestinationSettings& InVolMod)
+{
+	OutputVolumeModulation = InVolMod;
+	PushModulationChanges();
+}
+
+void USoundSubmix::SetWetVolumeModulation(const FSoundModulationDestinationSettings& InVolMod)
+{
+	WetLevelModulation = InVolMod;
+	PushModulationChanges();
+}
+
+void USoundSubmix::SetDryVolumeModulation(const FSoundModulationDestinationSettings& InVolMod)
+{
+	DryLevelModulation = InVolMod;
+	PushModulationChanges();
+}
+
+void USoundSubmix::PushModulationChanges()
+{
+	if (FAudioDeviceManager* AudioDeviceManager = FAudioDeviceManager::Get())
+	{
+		// Send the changes to the Modulation System
+		TSet<TObjectPtr<USoundModulatorBase>> NewVolumeMod = OutputVolumeModulation.Modulators;
+		TSet<TObjectPtr<USoundModulatorBase>> NewWetLevelMod = WetLevelModulation.Modulators;
+		TSet<TObjectPtr<USoundModulatorBase>> NewDryLevelMod = DryLevelModulation.Modulators;
+		AudioDeviceManager->IterateOverAllDevices([SoundSubmix = this, VolMod = MoveTemp(NewVolumeMod), WetMod = MoveTemp(NewWetLevelMod), DryMod = MoveTemp(NewDryLevelMod)](Audio::FDeviceId Id, FAudioDevice* Device) mutable
+		{
+			Device->UpdateSubmixModulationSettings(SoundSubmix, VolMod, WetMod, DryMod);
+		});
+	}
+}
+
 FString USoundSubmixBase::GetDesc()
 {
 	return FString(TEXT("Sound Submix"));
