InstancedStaticMesh Serialization Optim+Fix when bEditableWhenInherited == false (PackedLevelActors)
- Avoid persisting PerInstanceSMData, PerInstanceSMCustomData in this case and copy values from Archetype
- Also assign NumCustomDataFloats from Archetype in this case to fix warning (issue was with PerInstanceSMCustomData being always serialized and NumCustomDataFloats being delta serialized causing warnings before ReRunConstructionScript got to run)
- Also in this CL: Remove useless PerInstanceSMCustomData which was initialized and never used. This field gets properly cached by the base FInstancedStaticMeshComponentInstanceData class which serializes PerInstanceSMCustomData and PerInstanceSMData even if they have SkipSerialization flag because of the PPF_ForceTaggedSerialization flag
- Also remove useless copy of PerInstanceSMData as bMatch would be false if value differed. This might have been necessary at some point when this struct contained more than a Transform.

Resaving AncientGame with this change reduces the ExternalActors folder size by ~150mb (from ~500 to ~350)
Tested loading AncientGame without resave, pie, cook also tested with resave, pie and cook

diff --git a/Engine/Source/Runtime/Core/Public/UObject/FortniteMainBranchObjectVersion.h b/Engine/Source/Runtime/Core/Public/UObject/FortniteMainBranchObjectVersion.h
--- a/Engine/Source/Runtime/Core/Public/UObject/FortniteMainBranchObjectVersion.h
+++ b/Engine/Source/Runtime/Core/Public/UObject/FortniteMainBranchObjectVersion.h
@@ -357,6 +357,9 @@ struct FFortniteMainBranchObjectVersion
 		// Fix missing binding extensions for some anim graph nodes
 		FixMissingAnimGraphNodeBindingExtensions,
 
+		// EditableWhenInherited: Skip custom serialization on non Archetypes
+		ISMComponentEditableWhenInheritedSkipSerialization,
+
 		// -----<new versions can be added above this line>-------------------------------------------------
 		VersionPlusOne,
 		LatestVersion = VersionPlusOne - 1
diff --git a/Engine/Source/Runtime/Engine/Classes/Components/InstancedStaticMeshComponent.h b/Engine/Source/Runtime/Engine/Classes/Components/InstancedStaticMeshComponent.h
--- a/Engine/Source/Runtime/Engine/Classes/Components/InstancedStaticMeshComponent.h
+++ b/Engine/Source/Runtime/Engine/Classes/Components/InstancedStaticMeshComponent.h
@@ -645,9 +645,6 @@ public:
 	UPROPERTY()
 	TArray<FInstancedStaticMeshInstanceData> PerInstanceSMData;
 
-	UPROPERTY()
-	TArray<float> PerInstanceSMCustomData;
-
 	/** The cached selected instances */
 	TBitArray<> SelectedInstances;
 
diff --git a/Engine/Source/Runtime/Engine/Private/InstancedStaticMesh.cpp b/Engine/Source/Runtime/Engine/Private/InstancedStaticMesh.cpp
--- a/Engine/Source/Runtime/Engine/Private/InstancedStaticMesh.cpp
+++ b/Engine/Source/Runtime/Engine/Private/InstancedStaticMesh.cpp
@@ -2596,9 +2596,9 @@ TStructOnScope<FActorComponentInstanceData> UInstancedStaticMeshComponent::GetCo
 		StaticMeshInstanceData->CachedStaticLighting.MapBuildDataIds.Add(LODDataEntry.MapBuildDataId);
 	}
 
-	// Back up per-instance info
+	// Back up per-instance info (this is strictly for Comparison in UInstancedStaticMeshComponent::ApplyComponentInstanceData 
+	// as this Property will get serialized by base class FActorComponentInstanceData through FComponentPropertyWriter which uses the PPF_ForceTaggedSerialization to backup all properties even the custom serialized ones
 	StaticMeshInstanceData->PerInstanceSMData = PerInstanceSMData;
-	StaticMeshInstanceData->PerInstanceSMCustomData = PerInstanceSMCustomData;
 
 	// Back up instance selection
 	StaticMeshInstanceData->SelectedInstances = SelectedInstances;
@@ -2664,8 +2664,6 @@ void UInstancedStaticMeshComponent::ApplyComponentInstanceData(FInstancedStaticM
 		{
 			LODData[i].MapBuildDataId = InstancedMeshData->CachedStaticLighting.MapBuildDataIds[i];
 		}
-
-		PerInstanceSMData = InstancedMeshData->PerInstanceSMData;
 	}
 
 	SelectedInstances = InstancedMeshData->SelectedInstances;
@@ -3476,6 +3474,31 @@ void UInstancedStaticMeshComponent::Serialize(FArchive& Ar)
 		Ar << bCooked;
 	}
 
+	// Inherit properties when bEditableWhenInherited == true (when the component isn't a template and we are persisting data)
+	const UInstancedStaticMeshComponent* Archetype = Cast<UInstancedStaticMeshComponent>(GetArchetype());
+	const bool bInheritSkipSerializationProperties = !bEditableWhenInherited && Archetype && Ar.IsPersistent() && !IsTemplate();
+	
+
+	// Check if we need have SkipSerialization property data to load/save
+	bool bHasSkipSerializationPropertiesData = !bInheritSkipSerializationProperties;
+	if (Ar.IsLoading())
+	{
+#if WITH_EDITOR
+		if (Ar.CustomVer(FFortniteMainBranchObjectVersion::GUID) < FFortniteMainBranchObjectVersion::ISMComponentEditableWhenInheritedSkipSerialization)
+		{
+			bHasSkipSerializationPropertiesData = true;
+		}
+		else
+#endif
+		{
+			Ar << bHasSkipSerializationPropertiesData;
+		}
+	}
+	else
+	{
+		Ar << bHasSkipSerializationPropertiesData;
+	}
+			
 #if WITH_EDITOR
 	if (Ar.IsLoading() && Ar.CustomVer(FMobileObjectVersion::GUID) < FMobileObjectVersion::InstancedStaticMeshLightmapSerialization)
 	{
@@ -3488,12 +3511,39 @@ void UInstancedStaticMeshComponent::Serialize(FArchive& Ar)
 	}
 	else
 #endif //WITH_EDITOR
+	if (Ar.IsLoading())
 	{
-		PerInstanceSMData.BulkSerialize(Ar, Ar.UEVer() < EUnrealEngineObjectUE5Version::LARGE_WORLD_COORDINATES); 
-	}
+		// Read existing data if it was serialized
+		TArray<FInstancedStaticMeshInstanceData> TempPerInstanceSMData;
+		TArray<float> TempPerInstanceSMCustomData;
+		if (bHasSkipSerializationPropertiesData)
+		{
+			TempPerInstanceSMData.BulkSerialize(Ar, Ar.UEVer() < EUnrealEngineObjectUE5Version::LARGE_WORLD_COORDINATES);
+			if(Ar.CustomVer(FRenderingObjectVersion::GUID) >= FRenderingObjectVersion::PerInstanceCustomData)
+			{
+				TempPerInstanceSMCustomData.BulkSerialize(Ar);
+			}
+		}
+						
+		// If we should inherit use Archetype Data
+		if (bInheritSkipSerializationProperties)
+		{
+			PerInstanceSMData = Archetype->PerInstanceSMData;
+			PerInstanceSMCustomData = Archetype->PerInstanceSMCustomData;
 
-	if (!Ar.IsLoading() || Ar.CustomVer(FRenderingObjectVersion::GUID) >= FRenderingObjectVersion::PerInstanceCustomData)
+			// Make sure that if we inherit the PerInstanceSMCustomData we also do inherit this value
+			NumCustomDataFloats = Archetype->NumCustomDataFloats;
+		}
+		else
+		{
+			check(bHasSkipSerializationPropertiesData);
+			PerInstanceSMData = MoveTemp(TempPerInstanceSMData);
+			PerInstanceSMCustomData = MoveTemp(TempPerInstanceSMCustomData);
+		}
+	}
+	else if(bHasSkipSerializationPropertiesData)
 	{
+		PerInstanceSMData.BulkSerialize(Ar, Ar.UEVer() < EUnrealEngineObjectUE5Version::LARGE_WORLD_COORDINATES);
 		PerInstanceSMCustomData.BulkSerialize(Ar);
 	}
 
