World Partition: when an actor is detroyed through a user action, the transaction buffer will keep the actor alive in case the user undo the operation. When that  happens without a transaction (editor tools), the actor will be destroyed on the next GC, so we must remove the actor descriptor. In rare cases this heppens in the editor, report the destroyed actor as invalid so the user can at least repair and get a valid changelist to reflect the change.

fyi aris.theophanidis

diff --git a/Engine/Source/Runtime/Engine/Private/WorldPartition/ActorDescContainer.cpp b/Engine/Source/Runtime/Engine/Private/WorldPartition/ActorDescContainer.cpp
--- a/Engine/Source/Runtime/Engine/Private/WorldPartition/ActorDescContainer.cpp
+++ b/Engine/Source/Runtime/Engine/Private/WorldPartition/ActorDescContainer.cpp
@@ -13,6 +13,7 @@
 #include "AssetRegistry/AssetRegistryModule.h"
 #include "Modules/ModuleManager.h"
 #include "WorldPartition/WorldPartition.h"
+#include "WorldPartition/WorldPartitionActorDesc.h"
 #include "WorldPartition/WorldPartitionActorDescUtils.h"
 #include "WorldPartition/WorldPartitionClassDescRegistry.h"
 #include "WorldPartition/DataLayer/DataLayerManager.h"
@@ -410,6 +411,29 @@ void UActorDescContainer::OnClassDescriptorUpdated(const FWorldPartitionActorDes
 	}
 }
 
+void UActorDescContainer::OnActorDestroyed(AActor* Actor)
+{
+	// When an actor is detroyed through a user action, the transaction buffer will keep the actor alive in case the user undo the operation. When that
+	// happens without a transaction (editor tools), the actor will be destroyed on the next GC, so we must remove the actor descriptor. In rare cases
+	// this heppens in the editor, report the destroyed actor as invalid so the user can at least repair and get a valid changelist to reflect the change.
+	if (ShouldHandleActorEvent(Actor) && !GUndo && !GIsTransacting)
+	{
+		if (TUniquePtr<FWorldPartitionActorDesc>* ExistingActorDesc = GetActorDescriptor(Actor->GetActorGuid()))
+		{
+			OnActorDescRemoved(ExistingActorDesc->Get());
+			UnregisterActorDescriptor(ExistingActorDesc->Get());
+			ExistingActorDesc->Get()->Unload();
+			ExistingActorDesc->Reset();
+
+			FARFilter Filter;
+			Filter.bIncludeOnlyOnDiskAssets = true;
+			Filter.PackageNames.Add(Actor->GetPackage()->GetFName());
+			IAssetRegistry& AssetRegistry = FModuleManager::LoadModuleChecked<FAssetRegistryModule>(TEXT("AssetRegistry")).Get();
+			AssetRegistry.GetAssets(Filter, InvalidActors);
+		}
+	}
+}
+
 bool UActorDescContainer::RemoveActor(const FGuid& ActorGuid)
 {
 	if (TUniquePtr<FWorldPartitionActorDesc>* ExistingActorDesc = GetActorDescriptor(ActorGuid))
@@ -455,6 +479,12 @@ void UActorDescContainer::RegisterEditorDelegates()
 		
 		FWorldPartitionClassDescRegistry& ClassDescRegistry = FWorldPartitionClassDescRegistry::Get();
 		ClassDescRegistry.OnClassDescriptorUpdated().AddUObject(this, &UActorDescContainer::OnClassDescriptorUpdated);
+
+		if (UWorld* OwningWorld = GetWorldPartition() ? GetWorldPartition()->GetWorld() : nullptr)
+		{
+			FOnActorDestroyed::FDelegate ActorDestroyedDelegate = FOnActorDestroyed::FDelegate::CreateUObject(this, &UActorDescContainer::OnActorDestroyed);
+			OnActorDestroyedHandle = OwningWorld->AddOnActorDestroyedHandler(ActorDestroyedDelegate);
+		}
 	}
 }
 
@@ -468,6 +498,12 @@ void UActorDescContainer::UnregisterEditorDelegates()
 
 		FWorldPartitionClassDescRegistry& ClassDescRegistry = FWorldPartitionClassDescRegistry::Get();
 		ClassDescRegistry.OnClassDescriptorUpdated().RemoveAll(this);
+
+		if (UWorld* OwningWorld = GetWorldPartition() ? GetWorldPartition()->GetWorld() : nullptr)
+		{
+			OwningWorld->RemoveOnActorDestroyededHandler(OnActorDestroyedHandle);
+			OnActorDestroyedHandle.Reset();
+		}
 	}
 }
 
diff --git a/Engine/Source/Runtime/Engine/Public/WorldPartition/ActorDescContainer.h b/Engine/Source/Runtime/Engine/Public/WorldPartition/ActorDescContainer.h
--- a/Engine/Source/Runtime/Engine/Public/WorldPartition/ActorDescContainer.h
+++ b/Engine/Source/Runtime/Engine/Public/WorldPartition/ActorDescContainer.h
@@ -55,6 +55,7 @@ public:
 	ENGINE_API void OnPackageDeleted(UPackage* Package);
 	ENGINE_API void OnObjectsReplaced(const TMap<UObject*, UObject*>& OldToNewObjectMap);
 	ENGINE_API void OnClassDescriptorUpdated(const FWorldPartitionActorDesc* InClassDesc);
+	ENGINE_API void OnActorDestroyed(AActor* Actor);
 
 	FName GetContainerPackage() const { return ContainerPackageName; }
 	void SetContainerPackage(const FName& InContainerPackageName) { ContainerPackageName = InContainerPackageName; }
@@ -111,6 +112,8 @@ public:
 
 	TArray<FAssetData> InvalidActors;
 
+	FDelegateHandle OnActorDestroyedHandle;
+
 protected:
 	FNameActorDescMap ActorsByName;
 
