World Partition: added support for client and server only data layers activation.
- This can be used for data layers adding details on high end platforms.
- This can be used in content bundles with server-only data layers.

This is CL26804958 brought back with fixes.

diff --git a/Engine/Source/Editor/DataLayerEditor/Private/DataLayer/DataLayerHierarchy.cpp b/Engine/Source/Editor/DataLayerEditor/Private/DataLayer/DataLayerHierarchy.cpp
--- a/Engine/Source/Editor/DataLayerEditor/Private/DataLayer/DataLayerHierarchy.cpp
+++ b/Engine/Source/Editor/DataLayerEditor/Private/DataLayer/DataLayerHierarchy.cpp
@@ -250,7 +250,7 @@ void FDataLayerHierarchy::CreateItems(TArray<FSceneOutlinerTreeItemPtr>& OutItem
 					if (OwningWorldWorldDataLayers && CurrentLevelWorldDataLayers)
 					{
 						check(OwningWorldWorldDataLayers != CurrentLevelWorldDataLayers);
-						OwningWorldWorldDataLayers->ForEachDataLayer([&WorldToLevelDataLayerMap, CurrentLevelWorldDataLayers](UDataLayerInstance* DataLayerInstance)
+						OwningWorldWorldDataLayers->ForEachDataLayerInstance([&WorldToLevelDataLayerMap, CurrentLevelWorldDataLayers](UDataLayerInstance* DataLayerInstance)
 						{
 							if (UDataLayerInstanceWithAsset* DataLayerInstanceWithAsset = Cast<UDataLayerInstanceWithAsset>(DataLayerInstance))
 							{
diff --git a/Engine/Source/Editor/DataLayerEditor/Private/DataLayer/DataLayerInstanceCustomization.cpp b/Engine/Source/Editor/DataLayerEditor/Private/DataLayer/DataLayerInstanceCustomization.cpp
--- a/Engine/Source/Editor/DataLayerEditor/Private/DataLayer/DataLayerInstanceCustomization.cpp
+++ b/Engine/Source/Editor/DataLayerEditor/Private/DataLayer/DataLayerInstanceCustomization.cpp
@@ -27,7 +27,7 @@ void FDataLayerInstanceDetails::CustomizeDetails(IDetailLayoutBuilder& DetailBui
 	for (const TWeakObjectPtr<UObject>& SelectedObject : ObjectsBeingCustomized)
 	{
 		UDataLayerInstance* DataLayerInstance = Cast<UDataLayerInstance>(SelectedObject.Get());
-		if (DataLayerInstance && DataLayerInstance->IsRuntime())
+		if (DataLayerInstance && DataLayerInstance->IsRuntime() && !DataLayerInstance->IsClientOnly() && !DataLayerInstance->IsServerOnly())
 		{
 			bHasRuntime = true;
 			break;
diff --git a/Engine/Source/Runtime/Engine/Private/WorldPartition/DataLayer/DataLayerAsset.cpp b/Engine/Source/Runtime/Engine/Private/WorldPartition/DataLayer/DataLayerAsset.cpp
--- a/Engine/Source/Runtime/Engine/Private/WorldPartition/DataLayer/DataLayerAsset.cpp
+++ b/Engine/Source/Runtime/Engine/Private/WorldPartition/DataLayer/DataLayerAsset.cpp
@@ -59,6 +59,11 @@ bool UDataLayerAsset::CanEditChange(const FProperty* InProperty) const
 			// UDataLayerAsset outered to a UDataLayerInstance do not support Runtime type
 			return !IsPrivate();
 		}
+		else if (InProperty->GetFName() == GET_MEMBER_NAME_CHECKED(UDataLayerAsset, LoadFilter))
+		{
+			// Only runtime data layer assets can be set to client only
+			return IsRuntime();
+		}
 	}
 
 	return true;
diff --git a/Engine/Source/Runtime/Engine/Private/WorldPartition/DataLayer/DataLayerInstance.cpp b/Engine/Source/Runtime/Engine/Private/WorldPartition/DataLayer/DataLayerInstance.cpp
--- a/Engine/Source/Runtime/Engine/Private/WorldPartition/DataLayer/DataLayerInstance.cpp
+++ b/Engine/Source/Runtime/Engine/Private/WorldPartition/DataLayer/DataLayerInstance.cpp
@@ -24,9 +24,7 @@ UDataLayerInstance::UDataLayerInstance(const FObjectInitializer& ObjectInitializ
 	, bIsLocked(false)
 #endif
 	, InitialRuntimeState(EDataLayerRuntimeState::Unloaded)
-{
-
-}
+{}
 
 void UDataLayerInstance::PostLoad()
 {
@@ -122,9 +120,12 @@ bool UDataLayerInstance::CanEditChange(const FProperty* InProperty) const
 		return false;
 	}
 
-	if (!IsRuntime() && (InProperty->GetFName() == GET_MEMBER_NAME_CHECKED(UDataLayerInstance, InitialRuntimeState)))
+	if (InProperty->GetFName() == GET_MEMBER_NAME_CHECKED(UDataLayerInstance, InitialRuntimeState))
 	{
-		return false;
+		if (!IsRuntime() || IsClientOnly() || IsServerOnly())
+		{
+			return false;
+		}
 	}
 
 	return true;
@@ -246,6 +247,11 @@ bool UDataLayerInstance::IsParentDataLayerTypeCompatible(const UDataLayerInstanc
 		return false;
 	}
 
+	if (IsClientOnly() || IsServerOnly())
+	{
+		return false;
+	}
+
 	EDataLayerType ParentDataLayerType = InParent->GetType();
 
 	return GetType() != EDataLayerType::Unknown
@@ -415,15 +421,4 @@ EDataLayerRuntimeState UDataLayerInstance::GetEffectiveRuntimeState() const
 	return GetOuterWorldDataLayers()->GetDataLayerEffectiveRuntimeStateByName(GetDataLayerFName());
 }
 
-bool UDataLayerInstance::SetRuntimeState(EDataLayerRuntimeState InState, bool bInIsRecursive) const
-{
-	if (GetOuterWorldDataLayers()->HasAuthority())
-	{
-		GetOuterWorldDataLayers()->SetDataLayerRuntimeState(this, InState, bInIsRecursive);
-		return true;
-	}
-	UE_LOG(LogWorldPartition, Error, TEXT("SetDataLayerRuntimeState can only execute on authority"));
-	return false;
-}
-
 #undef LOCTEXT_NAMESPACE
diff --git a/Engine/Source/Runtime/Engine/Private/WorldPartition/DataLayer/DataLayerManager.cpp b/Engine/Source/Runtime/Engine/Private/WorldPartition/DataLayer/DataLayerManager.cpp
--- a/Engine/Source/Runtime/Engine/Private/WorldPartition/DataLayer/DataLayerManager.cpp
+++ b/Engine/Source/Runtime/Engine/Private/WorldPartition/DataLayer/DataLayerManager.cpp
@@ -70,7 +70,7 @@ FAutoConsoleCommand UDataLayerManager::ToggleDataLayerActivation(
 		for (const FWorldContext& Context : GEngine->GetWorldContexts())
 		{
 			UWorld* World = Context.World();
-			if (World && World->IsGameWorld() && !World->IsNetMode(NM_Client))
+			if (World && World->IsGameWorld())
 			{
 				if (UWorldPartitionSubsystem* WorldPartitionSubsystem = World->GetSubsystem<UWorldPartitionSubsystem>())
 				{
@@ -116,7 +116,7 @@ FAutoConsoleCommand UDataLayerManager::SetDataLayerRuntimeStateCommand(
 		for (const FWorldContext& Context : GEngine->GetWorldContexts())
 		{
 			UWorld* World = Context.World();
-			if (World && World->IsGameWorld() && !World->IsNetMode(NM_Client))
+			if (World && World->IsGameWorld())
 			{
 				if (UWorldPartitionSubsystem* WorldPartitionSubsystem = World->GetSubsystem<UWorldPartitionSubsystem>())
 				{
@@ -268,8 +268,13 @@ const UDataLayerInstance* UDataLayerManager::GetDataLayerInstanceFromAssetName(c
 
 bool UDataLayerManager::SetDataLayerInstanceRuntimeState(const UDataLayerInstance* InDataLayerInstance, EDataLayerRuntimeState InState, bool bInIsRecursive)
 {
-	UE_CLOG(!InDataLayerInstance, LogWorldPartition, Error, TEXT("Invalid Data Layer Instance."));
-	return InDataLayerInstance ? InDataLayerInstance->SetRuntimeState(InState, bInIsRecursive) : false;
+	if (InDataLayerInstance)
+	{
+		InDataLayerInstance->GetOuterWorldDataLayers()->SetDataLayerRuntimeState(InDataLayerInstance, InState, bInIsRecursive);
+		return true;
+	}
+	UE_LOG(LogWorldPartition, Error, TEXT("Invalid Data Layer Instance."));
+	return false;
 }
 
 bool UDataLayerManager::SetDataLayerRuntimeState(const UDataLayerAsset* InDataLayerAsset, EDataLayerRuntimeState InState, bool bInIsRecursive)
@@ -296,14 +301,22 @@ void UDataLayerManager::BroadcastOnDataLayerInstanceRuntimeStateChanged(const UD
 
 EDataLayerRuntimeState UDataLayerManager::GetDataLayerInstanceRuntimeState(const UDataLayerInstance* InDataLayerInstance) const
 {
-	UE_CLOG(!InDataLayerInstance, LogWorldPartition, Error, TEXT("Invalid Data Layer Instance."));
-	return InDataLayerInstance ? InDataLayerInstance->GetRuntimeState() : EDataLayerRuntimeState::Unloaded;
+	if (InDataLayerInstance)
+	{
+		return InDataLayerInstance->GetRuntimeState();
+	}
+	UE_LOG(LogWorldPartition, Error, TEXT("Invalid Data Layer Instance."));
+	return EDataLayerRuntimeState::Unloaded;
 }
 
 EDataLayerRuntimeState UDataLayerManager::GetDataLayerInstanceEffectiveRuntimeState(const UDataLayerInstance* InDataLayerInstance) const
 {
-	UE_CLOG(!InDataLayerInstance, LogWorldPartition, Error, TEXT("Invalid Data Layer Instance."));
-	return InDataLayerInstance ? InDataLayerInstance->GetEffectiveRuntimeState() : EDataLayerRuntimeState::Unloaded;
+	if (InDataLayerInstance)
+	{
+		return InDataLayerInstance->GetEffectiveRuntimeState();
+	}
+	UE_LOG(LogWorldPartition, Error, TEXT("Invalid Data Layer Instance."));
+	return EDataLayerRuntimeState::Unloaded;
 }
 
 const TSet<FName>& UDataLayerManager::GetEffectiveActiveDataLayerNames() const
@@ -578,7 +591,7 @@ void UDataLayerManager::ForEachDataLayerInstance(TFunctionRef<bool(UDataLayerIns
 			return bShouldContinue;
 		};
 
-		WorldDataLayers->ForEachDataLayer(CallAndSetContinueFunc);
+		WorldDataLayers->ForEachDataLayerInstance(CallAndSetContinueFunc);
 	}
 }
 
@@ -592,7 +605,6 @@ void UDataLayerManager::ForEachDataLayerInstance(TFunctionRef<bool(UDataLayerIns
  */
 
 #if WITH_EDITOR
-
 TSubclassOf<UDataLayerLoadingPolicy> UDataLayerManager::GetDataLayerLoadingPolicyClass() const
 {
 	// Use UDataLayerManager::DataLayerLoadingPolicyClass
@@ -837,5 +849,4 @@ void FDataLayersEditorBroadcast::StaticOnActorDataLayersEditorLoadingStateChange
 	Get().DataLayerEditorLoadingStateChanged.Broadcast(bIsFromUserChange);
 	IWorldPartitionActorLoaderInterface::RefreshLoadedState(bIsFromUserChange);
 }
-
 #endif
\ No newline at end of file
diff --git a/Engine/Source/Runtime/Engine/Private/WorldPartition/DataLayer/WorldDataLayers.cpp b/Engine/Source/Runtime/Engine/Private/WorldPartition/DataLayer/WorldDataLayers.cpp
--- a/Engine/Source/Runtime/Engine/Private/WorldPartition/DataLayer/WorldDataLayers.cpp
+++ b/Engine/Source/Runtime/Engine/Private/WorldPartition/DataLayer/WorldDataLayers.cpp
@@ -106,18 +106,26 @@ void AWorldDataLayers::InitializeDataLayerRuntimeStates()
 			return;
 		}
 #endif
-		ForEachDataLayer([this](class UDataLayerInstance* DataLayer)
+		TArray<UDataLayerInstance*> RuntimeDataLayerInstances;
+		RuntimeDataLayerInstances.Reserve(DataLayerInstances.Num());
+
+		ForEachDataLayerInstance([this, &RuntimeDataLayerInstances](UDataLayerInstance* DataLayerInstance)
 		{
-			if (DataLayer && DataLayer->IsRuntime())
+			const bool bDataLayerClientOnly = DataLayerInstance->IsClientOnly();
+			const bool bDataLayerServerOnly = DataLayerInstance->IsServerOnly();
+
+			if (DataLayerInstance->IsRuntime() && !bDataLayerClientOnly && !bDataLayerServerOnly)
 			{
-				if (DataLayer->GetInitialRuntimeState() == EDataLayerRuntimeState::Activated)
+				if (DataLayerInstance->GetInitialRuntimeState() == EDataLayerRuntimeState::Activated)
 				{
-					ActiveDataLayerNames.Add(DataLayer->GetDataLayerFName());
+					ActiveDataLayerNames.Add(DataLayerInstance->GetDataLayerFName());
 				}
-				else if (DataLayer->GetInitialRuntimeState() == EDataLayerRuntimeState::Loaded)
+				else if (DataLayerInstance->GetInitialRuntimeState() == EDataLayerRuntimeState::Loaded)
 				{
-					LoadedDataLayerNames.Add(DataLayer->GetDataLayerFName());
+					LoadedDataLayerNames.Add(DataLayerInstance->GetDataLayerFName());
 				}
+
+				RuntimeDataLayerInstances.Add(DataLayerInstance);
 			}
 			return true;
 		});
@@ -125,14 +133,11 @@ void AWorldDataLayers::InitializeDataLayerRuntimeStates()
 		RepActiveDataLayerNames = ActiveDataLayerNames.Array();
 		RepLoadedDataLayerNames = LoadedDataLayerNames.Array();
 
-		ForEachDataLayer([this](class UDataLayerInstance* DataLayer)
+		const bool bNotifyChange = false;
+		for (UDataLayerInstance* DataLayerInstance : RuntimeDataLayerInstances)
 		{
-			if (DataLayer && DataLayer->IsRuntime())
-			{
-				ResolveEffectiveRuntimeState(DataLayer, /*bNotifyChange*/false);
-			}
-			return true;
-		});
+			ResolveEffectiveRuntimeState(DataLayerInstance, bNotifyChange);
+		}
 
 		RepEffectiveActiveDataLayerNames = EffectiveActiveDataLayerNames.Array();
 		RepEffectiveLoadedDataLayerNames = EffectiveLoadedDataLayerNames.Array();
@@ -147,81 +152,133 @@ void AWorldDataLayers::ResetDataLayerRuntimeStates()
 	LoadedDataLayerNames.Reset();
 	RepActiveDataLayerNames.Reset();
 	RepLoadedDataLayerNames.Reset();
+	LocalActiveDataLayerNames.Reset();
+	LocalLoadedDataLayerNames.Reset();
 
 	EffectiveActiveDataLayerNames.Reset();
 	EffectiveLoadedDataLayerNames.Reset();
 	RepEffectiveActiveDataLayerNames.Reset();
 	RepEffectiveLoadedDataLayerNames.Reset();
+	LocalEffectiveActiveDataLayerNames.Reset();
+	LocalEffectiveLoadedDataLayerNames.Reset();
 }
 
 void AWorldDataLayers::SetDataLayerRuntimeState(const UDataLayerInstance* InDataLayerInstance, EDataLayerRuntimeState InState, bool bInIsRecursive)
 {
-	if (ensure(GetLocalRole() == ROLE_Authority))
+	if (!InDataLayerInstance || !InDataLayerInstance->IsRuntime())
+	{
+		return;
+	}
+
+	const ENetMode NetMode = GetNetMode();
+	const bool bDataLayerClientOnly = InDataLayerInstance->IsClientOnly();
+	const bool bDataLayerServerOnly = InDataLayerInstance->IsServerOnly();
+	const EDataLayerRuntimeState CurrentState = GetDataLayerRuntimeStateByName(InDataLayerInstance->GetDataLayerFName());
+
+	if (bDataLayerClientOnly)
 	{
-		if (!InDataLayerInstance || !InDataLayerInstance->IsRuntime())
+		if (NetMode != NM_Standalone && NetMode != NM_Client)
 		{
+			UE_LOG(LogWorldPartition, Log, TEXT("Client Only Data Layer state change '%s' was ignored: %s -> %s"),
+				*InDataLayerInstance->GetDataLayerShortName(),
+				*StaticEnum<EDataLayerRuntimeState>()->GetDisplayNameTextByValue((int64)CurrentState).ToString(),
+				*StaticEnum<EDataLayerRuntimeState>()->GetDisplayNameTextByValue((int64)InState).ToString());
+			return;
+		}
+	}
+	else if (bDataLayerServerOnly)
+	{
+		if (NetMode != NM_Standalone && NetMode == NM_Client)
+		{
+			UE_LOG(LogWorldPartition, Log, TEXT("Server Only Data Layer state change '%s' was ignored: %s -> %s"),
+				*InDataLayerInstance->GetDataLayerShortName(),
+				*StaticEnum<EDataLayerRuntimeState>()->GetDisplayNameTextByValue((int64)CurrentState).ToString(),
+				*StaticEnum<EDataLayerRuntimeState>()->GetDisplayNameTextByValue((int64)InState).ToString());
 			return;
 		}
+	}
+	else if (NetMode == NM_Client)
+	{
+		UE_LOG(LogWorldPartition, Log, TEXT("Data Layer '%s' state change was ignored on client: %s -> %s"),
+			*InDataLayerInstance->GetDataLayerShortName(),
+			*StaticEnum<EDataLayerRuntimeState>()->GetDisplayNameTextByValue((int64)CurrentState).ToString(),
+			*StaticEnum<EDataLayerRuntimeState>()->GetDisplayNameTextByValue((int64)InState).ToString());
+		return;
+	}
 
-		EDataLayerRuntimeState CurrentState = GetDataLayerRuntimeStateByName(InDataLayerInstance->GetDataLayerFName());
-		if (CurrentState != InState)
+	if (CurrentState != InState)
+	{
+		if (GetWorld()->IsGameWorld())
 		{
 			PRAGMA_DISABLE_DEPRECATION_WARNINGS
-				if (GetWorld()->IsGameWorld())
+			if (DataLayersFilterDelegate.IsBound())
+			{
+				const FString DataLayerShortName(InDataLayerInstance->GetDataLayerShortName());
+				if (!DataLayersFilterDelegate.Execute(FName(*DataLayerShortName), CurrentState, InState))
 				{
-					FName DataLayerShortName(InDataLayerInstance->GetDataLayerShortName());
-					if (DataLayersFilterDelegate.IsBound())
-					{
-						if (!DataLayersFilterDelegate.Execute(DataLayerShortName, CurrentState, InState))
-						{
-							UE_LOG(LogWorldPartition, Log, TEXT("Data Layer '%s' was filtered out: %s -> %s"),
-								*DataLayerShortName.ToString(),
-								*StaticEnum<EDataLayerRuntimeState>()->GetDisplayNameTextByValue((int64)CurrentState).ToString(),
-								*StaticEnum<EDataLayerRuntimeState>()->GetDisplayNameTextByValue((int64)InState).ToString());
-							return;
-						}
-					}
+					UE_LOG(LogWorldPartition, Log, TEXT("Data Layer '%s' state change was filtered out: %s -> %s"),
+						*DataLayerShortName,
+						*StaticEnum<EDataLayerRuntimeState>()->GetDisplayNameTextByValue((int64)CurrentState).ToString(),
+						*StaticEnum<EDataLayerRuntimeState>()->GetDisplayNameTextByValue((int64)InState).ToString());
+					return;
 				}
+			}
 			PRAGMA_ENABLE_DEPRECATION_WARNINGS
+		}
 
-			LoadedDataLayerNames.Remove(InDataLayerInstance->GetDataLayerFName());
-			ActiveDataLayerNames.Remove(InDataLayerInstance->GetDataLayerFName());
+		LoadedDataLayerNames.Remove(InDataLayerInstance->GetDataLayerFName());
+		ActiveDataLayerNames.Remove(InDataLayerInstance->GetDataLayerFName());
+		LocalLoadedDataLayerNames.Remove(InDataLayerInstance->GetDataLayerFName());
 
-			if (InState == EDataLayerRuntimeState::Loaded)
+		if (InState == EDataLayerRuntimeState::Loaded)
+		{
+			if (bDataLayerClientOnly || bDataLayerServerOnly)
+			{
+				LocalLoadedDataLayerNames.Add(InDataLayerInstance->GetDataLayerFName());
+			}
+			else
 			{
 				LoadedDataLayerNames.Add(InDataLayerInstance->GetDataLayerFName());
 			}
-			else if (InState == EDataLayerRuntimeState::Activated)
+		}
+		else if (InState == EDataLayerRuntimeState::Activated)
+		{
+			if (bDataLayerClientOnly || bDataLayerServerOnly)
+			{
+				LocalActiveDataLayerNames.Add(InDataLayerInstance->GetDataLayerFName());
+			}
+			else
 			{
 				ActiveDataLayerNames.Add(InDataLayerInstance->GetDataLayerFName());
 			}
+		}
 
-			// Update Replicated Properties
+		// Update replicated properties
+		if (NetMode == NM_DedicatedServer || NetMode == NM_ListenServer)
+		{
 			RepActiveDataLayerNames = ActiveDataLayerNames.Array();
 			RepLoadedDataLayerNames = LoadedDataLayerNames.Array();
+		}
+
+		++DataLayersStateEpoch;
 
-			++DataLayersStateEpoch;
+		UE_LOG(LogWorldPartition, Log, TEXT("Data Layer '%s' state changed: %s -> %s"),
+			*InDataLayerInstance->GetDataLayerShortName(),
+			*StaticEnum<EDataLayerRuntimeState>()->GetDisplayNameTextByValue((int64)CurrentState).ToString(),
+			*StaticEnum<EDataLayerRuntimeState>()->GetDisplayNameTextByValue((int64)InState).ToString());
 
-#if !NO_LOGGING || CSV_PROFILER
-			const FString DataLayerShortName = InDataLayerInstance->GetDataLayerShortName();
-			UE_LOG(LogWorldPartition, Log, TEXT("Data Layer '%s' state changed: %s -> %s"),
-				*DataLayerShortName,
-				*StaticEnum<EDataLayerRuntimeState>()->GetDisplayNameTextByValue((int64)CurrentState).ToString(),
-				*StaticEnum<EDataLayerRuntimeState>()->GetDisplayNameTextByValue((int64)InState).ToString());
+		CSV_EVENT_GLOBAL(TEXT("DataLayer-%s-%s"), *InDataLayerInstance->GetDataLayerShortName(), *StaticEnum<EDataLayerRuntimeState>()->GetDisplayNameTextByValue((int64)InState).ToString());
 
-			CSV_EVENT_GLOBAL(TEXT("DataLayer-%s-%s"), *DataLayerShortName, *StaticEnum<EDataLayerRuntimeState>()->GetDisplayNameTextByValue((int64)InState).ToString());
-#endif
+		ResolveEffectiveRuntimeState(InDataLayerInstance);
+	}
 
-			ResolveEffectiveRuntimeState(InDataLayerInstance);
-		}
-		if (bInIsRecursive)
+	if (bInIsRecursive)
+	{
+		InDataLayerInstance->ForEachChild([this, InState, bInIsRecursive](const UDataLayerInstance* Child)
 		{
-			InDataLayerInstance->ForEachChild([this, InState, bInIsRecursive](const UDataLayerInstance* Child)
-				{
-					SetDataLayerRuntimeState(Child, InState, bInIsRecursive);
-					return true;
-				});
-		}
+			SetDataLayerRuntimeState(Child, InState, bInIsRecursive);
+			return true;
+		});
 	}
 }
 
@@ -259,6 +316,16 @@ EDataLayerRuntimeState AWorldDataLayers::GetDataLayerRuntimeStateByName(FName In
 		check(!ActiveDataLayerNames.Contains(InDataLayerName));
 		return EDataLayerRuntimeState::Loaded;
 	}
+	else if (LocalActiveDataLayerNames.Contains(InDataLayerName))
+	{
+		check(!LocalLoadedDataLayerNames.Contains(InDataLayerName));
+		return EDataLayerRuntimeState::Activated;
+	}
+	else if (LocalLoadedDataLayerNames.Contains(InDataLayerName))
+	{
+		check(!LocalActiveDataLayerNames.Contains(InDataLayerName));
+		return EDataLayerRuntimeState::Loaded;
+	}
 
 	return EDataLayerRuntimeState::Unloaded;
 }
@@ -289,17 +356,45 @@ EDataLayerRuntimeState AWorldDataLayers::GetDataLayerEffectiveRuntimeStateByName
 		check(!EffectiveActiveDataLayerNames.Contains(InDataLayerName));
 		return EDataLayerRuntimeState::Loaded;
 	}
+	else if (LocalEffectiveActiveDataLayerNames.Contains(InDataLayerName))
+	{
+		check(!LocalEffectiveLoadedDataLayerNames.Contains(InDataLayerName));
+		return EDataLayerRuntimeState::Activated;
+	}
+	else if (LocalEffectiveLoadedDataLayerNames.Contains(InDataLayerName))
+	{
+		check(!LocalEffectiveActiveDataLayerNames.Contains(InDataLayerName));
+		return EDataLayerRuntimeState::Loaded;
+	}
 
 	return EDataLayerRuntimeState::Unloaded;
 }
 
-void AWorldDataLayers::ResolveEffectiveRuntimeState(const UDataLayerInstance* InDataLayer, bool bInNotifyChange)
+const TSet<FName>& AWorldDataLayers::GetEffectiveActiveDataLayerNames() const
+{
+	AllEffectiveActiveDataLayerNames = EffectiveActiveDataLayerNames;
+	AllEffectiveActiveDataLayerNames.Append(LocalEffectiveActiveDataLayerNames);
+	return AllEffectiveActiveDataLayerNames;
+}
+
+const TSet<FName>& AWorldDataLayers::GetEffectiveLoadedDataLayerNames() const
+{
+	AllEffectiveLoadedDataLayerNames = EffectiveLoadedDataLayerNames;
+	AllEffectiveLoadedDataLayerNames.Append(LocalEffectiveLoadedDataLayerNames);
+	return AllEffectiveLoadedDataLayerNames;
+}
+
+void AWorldDataLayers::ResolveEffectiveRuntimeState(const UDataLayerInstance* InDataLayerInstance, bool bInNotifyChange)
 {
-	check(InDataLayer);
-	const FName DataLayerName = InDataLayer->GetDataLayerFName();
+	check(InDataLayerInstance);
+	const ENetMode NetMode = GetNetMode();
+	const bool bDataLayerClientOnly = InDataLayerInstance->IsClientOnly();
+	const bool bDataLayerServerOnly = InDataLayerInstance->IsServerOnly();
+	const FName DataLayerName = InDataLayerInstance->GetDataLayerFName();
 	EDataLayerRuntimeState CurrentEffectiveRuntimeState = GetDataLayerEffectiveRuntimeStateByName(DataLayerName);
 	EDataLayerRuntimeState NewEffectiveRuntimeState = GetDataLayerRuntimeStateByName(DataLayerName);
-	const UDataLayerInstance* Parent = InDataLayer->GetParent();
+	const UDataLayerInstance* Parent = InDataLayerInstance->GetParent();
+
 	while (Parent && (NewEffectiveRuntimeState != EDataLayerRuntimeState::Unloaded))
 	{
 		if (Parent->IsRuntime())
@@ -308,38 +403,58 @@ void AWorldDataLayers::ResolveEffectiveRuntimeState(const UDataLayerInstance* In
 			NewEffectiveRuntimeState = (EDataLayerRuntimeState)FMath::Min((int32)NewEffectiveRuntimeState, (int32)GetDataLayerRuntimeStateByName(Parent->GetDataLayerFName()));
 		}
 		Parent = Parent->GetParent();
-	};
+	}
 
 	if (CurrentEffectiveRuntimeState != NewEffectiveRuntimeState)
 	{
 		EffectiveLoadedDataLayerNames.Remove(DataLayerName);
 		EffectiveActiveDataLayerNames.Remove(DataLayerName);
+		LocalEffectiveLoadedDataLayerNames.Remove(DataLayerName);
+		LocalEffectiveActiveDataLayerNames.Remove(DataLayerName);
 
 		if (NewEffectiveRuntimeState == EDataLayerRuntimeState::Loaded)
 		{
-			EffectiveLoadedDataLayerNames.Add(DataLayerName);
+			if (bDataLayerClientOnly || bDataLayerServerOnly)
+			{
+				LocalEffectiveLoadedDataLayerNames.Add(DataLayerName);
+			}
+			else
+			{
+				EffectiveLoadedDataLayerNames.Add(DataLayerName);
+			}
 		}
 		else if (NewEffectiveRuntimeState == EDataLayerRuntimeState::Activated)
 		{
-			EffectiveActiveDataLayerNames.Add(DataLayerName);
+			if (bDataLayerClientOnly || bDataLayerServerOnly)
+			{
+				LocalEffectiveActiveDataLayerNames.Add(DataLayerName);
+			}
+			else
+			{
+				EffectiveActiveDataLayerNames.Add(DataLayerName);
+			}
 		}
 
 		// Update Replicated Properties
-		RepEffectiveActiveDataLayerNames = EffectiveActiveDataLayerNames.Array();
-		RepEffectiveLoadedDataLayerNames = EffectiveLoadedDataLayerNames.Array();
+		if (NetMode == NM_DedicatedServer || NetMode == NM_ListenServer)
+		{
+			RepEffectiveActiveDataLayerNames = EffectiveActiveDataLayerNames.Array();
+			RepEffectiveLoadedDataLayerNames = EffectiveLoadedDataLayerNames.Array();
+		}
+
 		++DataLayersStateEpoch;
 
 		if (bInNotifyChange)
 		{
-			UE_LOG(LogWorldPartition, Log, TEXT("Data Layer '%s' effective state changed: %s -> %s"),
-				*InDataLayer->GetDataLayerShortName(),
+			UE_LOG(LogWorldPartition, Log, TEXT("Data Layer Instance '%s' effective state changed: %s -> %s"),
+				*InDataLayerInstance->GetDataLayerShortName(),
 				*StaticEnum<EDataLayerRuntimeState>()->GetDisplayNameTextByValue((int64)CurrentEffectiveRuntimeState).ToString(),
 				*StaticEnum<EDataLayerRuntimeState>()->GetDisplayNameTextByValue((int64)NewEffectiveRuntimeState).ToString());
 
-			OnDataLayerRuntimeStateChanged(InDataLayer, NewEffectiveRuntimeState);
+			OnDataLayerRuntimeStateChanged(InDataLayerInstance, NewEffectiveRuntimeState);
 		}
 
-		for (const UDataLayerInstance* Child : InDataLayer->GetChildren())
+		for (const UDataLayerInstance* Child : InDataLayerInstance->GetChildren())
 		{
 			ResolveEffectiveRuntimeState(Child);
 		}
@@ -377,7 +492,7 @@ void AWorldDataLayers::DumpDataLayerRecursively(const UDataLayerInstance* DataLa
 	{
 		DumpDataLayerRecursively(Child, Prefix + TEXT(" | "), OutputDevice);
 	}
-};
+}
 
 void AWorldDataLayers::DumpDataLayers(FOutputDevice& OutputDevice) const
 {
@@ -416,9 +531,9 @@ void AWorldDataLayers::DumpDataLayers(FOutputDevice& OutputDevice) const
 	
 	OutputDevice.Logf(TEXT("----------------------------------------------------"));
 	OutputDevice.Logf(TEXT(" Data Layers Hierarchy"));
-	ForEachDataLayer([this, &OutputDevice](UDataLayerInstance* DataLayerInstance)
+	ForEachDataLayerInstance([this, &OutputDevice](UDataLayerInstance* DataLayerInstance)
 	{
-		if (DataLayerInstance && !DataLayerInstance->GetParent())
+		if (!DataLayerInstance->GetParent())
 		{
 			DumpDataLayerRecursively(DataLayerInstance, TEXT(""), OutputDevice);
 		}
@@ -637,7 +752,7 @@ TArray<UDataLayerInstance*> AWorldDataLayers::GetActorEditorContextDataLayers()
 		{
 			Result.Add(const_cast<UDataLayerInstance*>(DataLayerInstance));
 		}
-	};
+	}
 	return Result;
 }
 
@@ -719,22 +834,22 @@ const UDataLayerInstance* AWorldDataLayers::GetDataLayerInstance(const UDataLaye
 	return nullptr;
 }
 
-void AWorldDataLayers::ForEachDataLayer(TFunctionRef<bool(UDataLayerInstance*)> Func)
+void AWorldDataLayers::ForEachDataLayerInstance(TFunctionRef<bool(UDataLayerInstance*)> Func)
 {
-	for (UDataLayerInstance* DataLayer : DataLayerInstances)
+	for (UDataLayerInstance* DataLayerInstance : DataLayerInstances)
 	{
-		if (!Func(DataLayer))
+		if (DataLayerInstance && !Func(DataLayerInstance))
 		{
 			break;
 		}
 	}
 }
 
-void AWorldDataLayers::ForEachDataLayer(TFunctionRef<bool(UDataLayerInstance*)> Func) const
+void AWorldDataLayers::ForEachDataLayerInstance(TFunctionRef<bool(UDataLayerInstance*)> Func) const
 {
-	for (UDataLayerInstance* DataLayer : DataLayerInstances)
+	for (UDataLayerInstance* DataLayerInstance : DataLayerInstances)
 	{
-		if (!Func(DataLayer))
+		if (DataLayerInstance && !Func(DataLayerInstance))
 		{
 			break;
 		}
@@ -778,18 +893,18 @@ void AWorldDataLayers::PostLoad()
 	// Remove all Editor Data Layers when cooking or when in a game world
 	if (IsRunningCookCommandlet() || GetWorld()->IsGameWorld())
 	{
-		ForEachDataLayer([](UDataLayerInstance* DataLayer)
+		ForEachDataLayerInstance([](UDataLayerInstance* DataLayerInstance)
 		{
-			DataLayer->ConditionalPostLoad();
+			DataLayerInstance->ConditionalPostLoad();
 			return true;
 		});
 
 		TArray<UDataLayerInstance*> EditorDataLayers;
-		ForEachDataLayer([&EditorDataLayers](UDataLayerInstance* DataLayer)
+		ForEachDataLayerInstance([&EditorDataLayers](UDataLayerInstance* DataLayerInstance)
 		{
-			if (DataLayer && !DataLayer->IsRuntime())
+			if (!DataLayerInstance->IsRuntime())
 			{
-				EditorDataLayers.Add(DataLayer);
+				EditorDataLayers.Add(DataLayerInstance);
 			}
 			return true;
 		});
diff --git a/Engine/Source/Runtime/Engine/Private/WorldPartition/DataLayer/WorldDataLayersActorDesc.cpp b/Engine/Source/Runtime/Engine/Private/WorldPartition/DataLayer/WorldDataLayersActorDesc.cpp
--- a/Engine/Source/Runtime/Engine/Private/WorldPartition/DataLayer/WorldDataLayersActorDesc.cpp
+++ b/Engine/Source/Runtime/Engine/Private/WorldPartition/DataLayer/WorldDataLayersActorDesc.cpp
@@ -206,7 +206,7 @@ void FWorldDataLayersActorDesc::Init(const AActor* InActor)
 	FWorldPartitionActorDesc::Init(InActor);
 
 	const AWorldDataLayers* WorldDataLayers = CastChecked<AWorldDataLayers>(InActor);
-	WorldDataLayers->ForEachDataLayer([this](UDataLayerInstance* DataLayerInstance)
+	WorldDataLayers->ForEachDataLayerInstance([this](UDataLayerInstance* DataLayerInstance)
 	{
 		FDataLayerInstanceDesc& DataLayerInstanceDesc = DataLayerInstances.Emplace_GetRef();
 		DataLayerInstanceDesc.Init(DataLayerInstance);
diff --git a/Engine/Source/Runtime/Engine/Public/WorldPartition/DataLayer/DataLayerAsset.h b/Engine/Source/Runtime/Engine/Public/WorldPartition/DataLayer/DataLayerAsset.h
--- a/Engine/Source/Runtime/Engine/Public/WorldPartition/DataLayer/DataLayerAsset.h
+++ b/Engine/Source/Runtime/Engine/Public/WorldPartition/DataLayer/DataLayerAsset.h
@@ -13,6 +13,14 @@
 
 class AActor;
 
+UENUM(BlueprintType)
+enum class EDataLayerLoadFilter : uint8
+{
+	None,
+	ClientOnly,
+	ServerOnly
+};
+
 UCLASS(BlueprintType, editinlinenew, MinimalAPI)
 class UDataLayerAsset : public UObject
 {
@@ -36,6 +44,7 @@ public:
 	ENGINE_API bool CanBeReferencedByActor(AActor* InActor) const;
 	static ENGINE_API bool CanBeReferencedByActor(const TSoftObjectPtr<UDataLayerAsset>& InDataLayerAsset, AActor* InActor);
 #endif
+
 	ENGINE_API bool IsPrivate() const;
 
 	UFUNCTION(Category = "Data Layer", BlueprintCallable)
@@ -48,6 +57,13 @@ public:
 	FColor GetDebugColor() const { return DebugColor; }
 
 	bool SupportsActorFilters() const { return bSupportsActorFilters; }
+
+	UFUNCTION(Category = "Data Layer|Runtime", BlueprintCallable)
+	bool IsClientOnly() const { return IsRuntime() && LoadFilter == EDataLayerLoadFilter::ClientOnly; }
+
+	UFUNCTION(Category = "Data Layer|Runtime", BlueprintCallable)
+	bool IsServerOnly() const { return IsRuntime() && LoadFilter == EDataLayerLoadFilter::ServerOnly; }
+
 private:
 	/** Whether the Data Layer affects actor runtime loading */
 	UPROPERTY(Category = "Data Layer", EditAnywhere)
@@ -58,4 +74,7 @@ private:
 		
 	UPROPERTY(Category = "Runtime", EditAnywhere)
 	FColor DebugColor;
+
+	UPROPERTY(Category = "Runtime", EditAnywhere)
+	EDataLayerLoadFilter LoadFilter;
 };
diff --git a/Engine/Source/Runtime/Engine/Public/WorldPartition/DataLayer/DataLayerInstance.h b/Engine/Source/Runtime/Engine/Public/WorldPartition/DataLayer/DataLayerInstance.h
--- a/Engine/Source/Runtime/Engine/Public/WorldPartition/DataLayer/DataLayerInstance.h
+++ b/Engine/Source/Runtime/Engine/Public/WorldPartition/DataLayer/DataLayerInstance.h
@@ -141,11 +141,17 @@ public:
 	UFUNCTION(Category = "Data Layer|Runtime", BlueprintCallable)
 	virtual bool IsRuntime() const { return false; }
 
+	UFUNCTION(Category = "Data Layer|Runtime", BlueprintCallable)
+	virtual bool IsClientOnly() const { return false; }
+
+	UFUNCTION(Category = "Data Layer|Runtime", BlueprintCallable)
+	virtual bool IsServerOnly() const { return false; }
+
 	UFUNCTION(Category = "Data Layer|Runtime", BlueprintCallable)
 	virtual FColor GetDebugColor() const { return FColor::Black; }
 
 	UFUNCTION(Category = "Data Layer|Runtime", BlueprintCallable)
-	EDataLayerRuntimeState GetInitialRuntimeState() const { return IsRuntime() ? InitialRuntimeState : EDataLayerRuntimeState::Unloaded; }
+	EDataLayerRuntimeState GetInitialRuntimeState() const { return IsRuntime() && !IsClientOnly() && !IsServerOnly() ? InitialRuntimeState : EDataLayerRuntimeState::Unloaded; }
 
 	virtual FString GetDataLayerShortName() const { return TEXT("Invalid Data Layer"); }
 	virtual FString GetDataLayerFullName() const { return TEXT("Invalid Data Layer"); }
@@ -219,6 +225,7 @@ protected:
 	uint32 bIsLocked : 1;
 #endif
 
+	/** Initial runtime state of this data layer instance. Only supported if it's runtime and not client/server only. */
 	UPROPERTY(Category = "Runtime", EditAnywhere)
 	EDataLayerRuntimeState InitialRuntimeState;
 
diff --git a/Engine/Source/Runtime/Engine/Public/WorldPartition/DataLayer/DataLayerInstanceWithAsset.h b/Engine/Source/Runtime/Engine/Public/WorldPartition/DataLayer/DataLayerInstanceWithAsset.h
--- a/Engine/Source/Runtime/Engine/Public/WorldPartition/DataLayer/DataLayerInstanceWithAsset.h
+++ b/Engine/Source/Runtime/Engine/Public/WorldPartition/DataLayer/DataLayerInstanceWithAsset.h
@@ -36,14 +36,17 @@ public:
 
 	const UDataLayerAsset* GetAsset() const override { return DataLayerAsset; }
 
-	virtual EDataLayerType GetType() const override { return DataLayerAsset != nullptr ? DataLayerAsset->GetType() : EDataLayerType::Unknown; }
+	virtual EDataLayerType GetType() const override { return DataLayerAsset ? DataLayerAsset->GetType() : EDataLayerType::Unknown; }
 
-	virtual bool IsRuntime() const override { return DataLayerAsset != nullptr ? DataLayerAsset->IsRuntime() : false; }
+	virtual bool IsRuntime() const override { return DataLayerAsset && DataLayerAsset->IsRuntime(); }
 
-	virtual FColor GetDebugColor() const override { return DataLayerAsset != nullptr ? DataLayerAsset->GetDebugColor() : FColor::Black; }
+	virtual bool IsClientOnly() const override { return DataLayerAsset && DataLayerAsset->IsClientOnly(); }
+	virtual bool IsServerOnly() const override { return DataLayerAsset && DataLayerAsset->IsServerOnly(); }
 
-	virtual FString GetDataLayerShortName() const override { return DataLayerAsset != nullptr ? DataLayerAsset->GetName() : GetDataLayerFName().ToString(); }
-	virtual FString GetDataLayerFullName() const override { return DataLayerAsset != nullptr ? DataLayerAsset->GetPathName() : GetDataLayerFName().ToString(); }
+	virtual FColor GetDebugColor() const override { return DataLayerAsset ? DataLayerAsset->GetDebugColor() : FColor::Black; }
+
+	virtual FString GetDataLayerShortName() const override { return DataLayerAsset ? DataLayerAsset->GetName() : GetDataLayerFName().ToString(); }
+	virtual FString GetDataLayerFullName() const override { return DataLayerAsset ? DataLayerAsset->GetPathName() : GetDataLayerFName().ToString(); }
 
 protected:
 #if WITH_EDITOR
diff --git a/Engine/Source/Runtime/Engine/Public/WorldPartition/DataLayer/WorldDataLayers.h b/Engine/Source/Runtime/Engine/Public/WorldPartition/DataLayer/WorldDataLayers.h
--- a/Engine/Source/Runtime/Engine/Public/WorldPartition/DataLayer/WorldDataLayers.h
+++ b/Engine/Source/Runtime/Engine/Public/WorldPartition/DataLayer/WorldDataLayers.h
@@ -101,8 +101,15 @@ public:
 	ENGINE_API const UDataLayerInstance* GetDataLayerInstance(const FName& InDataLayerInstanceName) const;
 	ENGINE_API const UDataLayerInstance* GetDataLayerInstance(const UDataLayerAsset* InDataLayerAsset) const;
 	ENGINE_API const UDataLayerInstance* GetDataLayerInstanceFromAssetName(const FName& InDataLayerAssetFullName) const;
-	ENGINE_API void ForEachDataLayer(TFunctionRef<bool(UDataLayerInstance*)> Func);
-	ENGINE_API void ForEachDataLayer(TFunctionRef<bool(UDataLayerInstance*)> Func) const;
+
+	UE_DEPRECATED(5.4, "Use ForEachDataLayerInstance() instead.")
+	ENGINE_API void ForEachDataLayer(TFunctionRef<bool(UDataLayerInstance*)> Func) { return ForEachDataLayerInstance(Func); }
+
+	UE_DEPRECATED(5.4, "Use ForEachDataLayerInstance() instead.")
+	ENGINE_API void ForEachDataLayer(TFunctionRef<bool(UDataLayerInstance*)> Func) const { return ForEachDataLayerInstance(Func); }
+
+	ENGINE_API void ForEachDataLayerInstance(TFunctionRef<bool(UDataLayerInstance*)> Func);
+	ENGINE_API void ForEachDataLayerInstance(TFunctionRef<bool(UDataLayerInstance*)> Func) const;
 
 	ENGINE_API TArray<const UDataLayerInstance*> GetDataLayerInstances(const TArray<FName>& InDataLayerInstanceNames) const;
 
@@ -110,8 +117,8 @@ public:
 	ENGINE_API void SetDataLayerRuntimeState(const UDataLayerInstance* InDataLayerInstance, EDataLayerRuntimeState InState, bool bIsRecursive = false);
 	ENGINE_API EDataLayerRuntimeState GetDataLayerRuntimeStateByName(FName InDataLayerName) const;
 	ENGINE_API EDataLayerRuntimeState GetDataLayerEffectiveRuntimeStateByName(FName InDataLAyerName) const;
-	const TSet<FName>& GetEffectiveActiveDataLayerNames() const { return EffectiveActiveDataLayerNames; }
-	const TSet<FName>& GetEffectiveLoadedDataLayerNames() const { return EffectiveLoadedDataLayerNames; }
+	const TSet<FName>& GetEffectiveActiveDataLayerNames() const;
+	const TSet<FName>& GetEffectiveLoadedDataLayerNames() const;
 	UFUNCTION(NetMulticast, Reliable)
 	ENGINE_API void OnDataLayerRuntimeStateChanged(const UDataLayerInstance* InDataLayer, EDataLayerRuntimeState InState);
 	static int32 GetDataLayersStateEpoch() { return DataLayersStateEpoch; }
@@ -229,6 +236,9 @@ private:
 	TSet<FName> ActiveDataLayerNames;
 	TSet<FName> LoadedDataLayerNames;
 
+	TSet<FName> LocalActiveDataLayerNames;
+	TSet<FName> LocalLoadedDataLayerNames;
+
 	UPROPERTY(Transient, Replicated, ReplicatedUsing=OnRep_EffectiveActiveDataLayerNames)
 	TArray<FName> RepEffectiveActiveDataLayerNames;
 		
@@ -239,6 +249,12 @@ private:
 	TSet<FName> EffectiveActiveDataLayerNames;
 	TSet<FName> EffectiveLoadedDataLayerNames;
 
+	TSet<FName> LocalEffectiveActiveDataLayerNames;
+	TSet<FName> LocalEffectiveLoadedDataLayerNames;
+
+	mutable TSet<FName> AllEffectiveActiveDataLayerNames;
+	mutable TSet<FName> AllEffectiveLoadedDataLayerNames;
+
 	static ENGINE_API int32 DataLayersStateEpoch;
 
 	static_assert(DATALAYER_TO_INSTANCE_RUNTIME_CONVERSION_ENABLED, "bHasDeprecatedDataLayers is deprecated and needs to be deleted.");
@@ -321,11 +337,11 @@ void AWorldDataLayers::OverwriteDataLayerRuntimeStates(const TArray<IdentifierTy
 			RepLoadedDataLayerNames = LoadedDataLayerNames.Array();
 		}
 
-		ForEachDataLayer([this](class UDataLayerInstance* DataLayer)
+		ForEachDataLayerInstance([this](class UDataLayerInstance* DataLayerInstance)
 		{
-			if (DataLayer && DataLayer->IsRuntime())
+			if (DataLayerInstance->IsRuntime())
 			{
-				ResolveEffectiveRuntimeState(DataLayer, /*bNotifyChange*/false);
+				ResolveEffectiveRuntimeState(DataLayerInstance, /*bNotifyChange*/false);
 			}
 			return true;
 		});
