Additive transform curves are not applied to certain bones
#fix ensure that non bone-animated yet with an additive animation curve _are_ incorporated into compressed animation data

diff --git a/Engine/Source/Runtime/Engine/Private/Animation/AnimCompressionTypes.cpp b/Engine/Source/Runtime/Engine/Private/Animation/AnimCompressionTypes.cpp
--- a/Engine/Source/Runtime/Engine/Private/Animation/AnimCompressionTypes.cpp
+++ b/Engine/Source/Runtime/Engine/Private/Animation/AnimCompressionTypes.cpp
@@ -507,6 +507,22 @@ void FCompressibleAnimData::ResampleAnimationTrackData(const FFrameRate& SampleR
 					TrackData.InternalTrackData.ScaleKeys.SetNumUninitialized(SampledKeys);
 				}
 			}
+
+			const TArray<FTransformCurve>& TransformCurves = DataModelInterface->GetTransformCurves();
+			for (const FTransformCurve& Curve : TransformCurves)
+			{
+				const int32 CurveBoneIndex = ReferenceSkeleton.FindBoneIndex(Curve.GetName());
+				if (CurveBoneIndex != INDEX_NONE && !ResampledTrackData.ContainsByPredicate([CurveBoneIndex](const FBoneAnimationTrack& Track) { return Track.BoneTreeIndex == CurveBoneIndex; }))
+				{
+					FBoneAnimationTrack& TrackData = ResampledTrackData.AddDefaulted_GetRef();					
+					TrackData.Name = Curve.GetName();
+					TrackData.BoneTreeIndex = CurveBoneIndex;
+					
+					TrackData.InternalTrackData.PosKeys.SetNumUninitialized(SampledKeys);
+					TrackData.InternalTrackData.RotKeys.SetNumUninitialized(SampledKeys);
+					TrackData.InternalTrackData.ScaleKeys.SetNumUninitialized(SampledKeys);
+				}
+			}
 			
 			FBlendedCurve Curve;
 			Curve.InitFrom(EvalContext.RequiredBones);
