Moved volumetric volume to RHICmdList.

diff --git a/Engine/Source/Runtime/Engine/Private/PrecomputedVolumetricLightmap.cpp b/Engine/Source/Runtime/Engine/Private/PrecomputedVolumetricLightmap.cpp
--- a/Engine/Source/Runtime/Engine/Private/PrecomputedVolumetricLightmap.cpp
+++ b/Engine/Source/Runtime/Engine/Private/PrecomputedVolumetricLightmap.cpp
@@ -699,9 +699,9 @@ void FPrecomputedVolumetricLightmap::AddToScene(FSceneInterface* Scene, UMapBuil
 		FPrecomputedVolumetricLightmap* Volume = this;
 
 		ENQUEUE_RENDER_COMMAND(SetVolumeDataCommand)
-			([Volume, NewData, Scene](FRHICommandListImmediate& RHICmdList) 
+			([Volume, NewData, Scene] (FRHICommandListBase& RHICmdList)
 			{
-				Volume->SetData(NewData, Scene);
+				Volume->SetData(RHICmdList, NewData, Scene);
 			});
 		Scene->AddPrecomputedVolumetricLightmap(this, bIsPersistentLevel);
 	}
@@ -741,7 +741,7 @@ void FPrecomputedVolumetricLightmap::RemoveFromScene(FSceneInterface* Scene)
 	WorldOriginOffset = FVector::ZeroVector;
 }
 
-void FPrecomputedVolumetricLightmap::SetData(FPrecomputedVolumetricLightmapData* NewData, FSceneInterface* Scene)
+void FPrecomputedVolumetricLightmap::SetData(FRHICommandListBase& RHICmdList, FPrecomputedVolumetricLightmapData* NewData, FSceneInterface* Scene)
 {
 	Data = NewData;
 
@@ -753,7 +753,7 @@ void FPrecomputedVolumetricLightmap::SetData(FPrecomputedVolumetricLightmapData*
 
 		if (Scene->GetFeatureLevel() >= ERHIFeatureLevel::SM5)
 		{
-			Data->InitResource(FRHICommandListImmediate::Get());
+			Data->InitResource(RHICmdList);
 		}
 	}
 }
diff --git a/Engine/Source/Runtime/Engine/Public/PrecomputedVolumetricLightmap.h b/Engine/Source/Runtime/Engine/Public/PrecomputedVolumetricLightmap.h
--- a/Engine/Source/Runtime/Engine/Public/PrecomputedVolumetricLightmap.h
+++ b/Engine/Source/Runtime/Engine/Public/PrecomputedVolumetricLightmap.h
@@ -226,7 +226,7 @@ public:
 
 	ENGINE_API void RemoveFromScene(FSceneInterface* Scene);
 	
-	ENGINE_API void SetData(FPrecomputedVolumetricLightmapData* NewData, FSceneInterface* Scene);
+	ENGINE_API void SetData(FRHICommandListBase& RHICmdList, FPrecomputedVolumetricLightmapData* NewData, FSceneInterface* Scene);
 
 	bool IsAddedToScene() const
 	{
