Added support for direct RHI command list initialization of static mesh buffer resources.

diff --git a/Engine/Source/Runtime/Engine/Private/Components/TextRenderComponent.cpp b/Engine/Source/Runtime/Engine/Private/Components/TextRenderComponent.cpp
--- a/Engine/Source/Runtime/Engine/Private/Components/TextRenderComponent.cpp
+++ b/Engine/Source/Runtime/Engine/Private/Components/TextRenderComponent.cpp
@@ -737,9 +737,9 @@ void FTextRenderSceneProxy::CreateRenderThreadResources(FRHICommandListBase& RHI
 		VertexBuffers.SetOwnerName(Name);
 		IndexBuffer.SetOwnerName(Name);
 #endif
-		VertexBuffers.InitFromDynamicVertex(&VertexFactory, OutVertices);
+		VertexBuffers.InitFromDynamicVertex(RHICmdList, &VertexFactory, OutVertices);
 		// Enqueue initialization of render resources
-		BeginInitResource(&IndexBuffer);
+		IndexBuffer.InitResource(RHICmdList);
 
 #if RHI_RAYTRACING
 		if (bSupportRayTracing)
diff --git a/Engine/Source/Runtime/Engine/Private/StaticMesh.cpp b/Engine/Source/Runtime/Engine/Private/StaticMesh.cpp
--- a/Engine/Source/Runtime/Engine/Private/StaticMesh.cpp
+++ b/Engine/Source/Runtime/Engine/Private/StaticMesh.cpp
@@ -1125,28 +1125,35 @@ void FStaticMeshVertexBuffers::InitModelBuffers(TArray<FModelVertex>& Vertices)
 	}
 }
 
-void FStaticMeshVertexBuffers::InitModelVF(FRenderCommandPipe* RenderCommandPipe, FLocalVertexFactory* VertexFactory)
+void FStaticMeshVertexBuffers::InitModelVF(FRHICommandListBase* RHICmdList, FRenderCommandPipe* RenderCommandPipe, FLocalVertexFactory* VertexFactory)
 {
-	FStaticMeshVertexBuffers* Self = this;
-	ENQUEUE_RENDER_COMMAND(StaticMeshVertexBuffersLegacyBspInit)(RenderCommandPipe,
-		[VertexFactory, Self](FRHICommandListBase& RHICmdList)
+	auto Lambda = [this, VertexFactory](FRHICommandListBase& RHICmdList)
 	{
-		check(Self->PositionVertexBuffer.IsInitialized());
-		check(Self->StaticMeshVertexBuffer.IsInitialized());
+		check(PositionVertexBuffer.IsInitialized());
+		check(StaticMeshVertexBuffer.IsInitialized());
 
 		FLocalVertexFactory::FDataType Data;
-		Self->PositionVertexBuffer.BindPositionVertexBuffer(VertexFactory, Data);
-		Self->StaticMeshVertexBuffer.BindTangentVertexBuffer(VertexFactory, Data);
-		Self->StaticMeshVertexBuffer.BindPackedTexCoordVertexBuffer(VertexFactory, Data);
-		Self->StaticMeshVertexBuffer.BindLightMapVertexBuffer(VertexFactory, Data, 1);
+		PositionVertexBuffer.BindPositionVertexBuffer(VertexFactory, Data);
+		StaticMeshVertexBuffer.BindTangentVertexBuffer(VertexFactory, Data);
+		StaticMeshVertexBuffer.BindPackedTexCoordVertexBuffer(VertexFactory, Data);
+		StaticMeshVertexBuffer.BindLightMapVertexBuffer(VertexFactory, Data, 1);
 		FColorVertexBuffer::BindDefaultColorVertexBuffer(VertexFactory, Data, FColorVertexBuffer::NullBindStride::ZeroForDefaultBufferBind);
 		VertexFactory->SetData(RHICmdList, Data);
 
 		InitOrUpdateResource(RHICmdList, VertexFactory);
-	});
+	};
+
+	if (RHICmdList)
+	{
+		Lambda(*RHICmdList);
+	}
+	else
+	{
+		ENQUEUE_RENDER_COMMAND(StaticMeshVertexBuffersInitModelVF)(RenderCommandPipe, MoveTemp(Lambda));
+	}
 }
 
-void FStaticMeshVertexBuffers::InitWithDummyData(FRenderCommandPipe* RenderCommandPipe, FLocalVertexFactory* VertexFactory, uint32 NumVerticies, uint32 NumTexCoords, uint32 LightMapIndex)
+void FStaticMeshVertexBuffers::InitWithDummyData(FRHICommandListBase* RHICmdList, FRenderCommandPipe* RenderCommandPipe, FLocalVertexFactory* VertexFactory, uint32 NumVerticies, uint32 NumTexCoords, uint32 LightMapIndex)
 {
 	check(NumVerticies);
 	check(NumTexCoords < MAX_STATIC_TEXCOORDS && NumTexCoords > 0);
@@ -1156,27 +1163,34 @@ void FStaticMeshVertexBuffers::InitWithDummyData(FRenderCommandPipe* RenderComma
 	StaticMeshVertexBuffer.Init(NumVerticies, NumTexCoords);
 	ColorVertexBuffer.Init(NumVerticies);
 
-	FStaticMeshVertexBuffers* Self = this;
-	ENQUEUE_RENDER_COMMAND(StaticMeshVertexBuffersLegacyInit)(RenderCommandPipe,
-		[VertexFactory, Self, LightMapIndex](FRHICommandListBase& RHICmdList)
+	auto Lambda = [this, VertexFactory, LightMapIndex](FRHICommandListBase& RHICmdList)
 	{
-		InitOrUpdateResource(RHICmdList, &Self->PositionVertexBuffer);
-		InitOrUpdateResource(RHICmdList, &Self->StaticMeshVertexBuffer);
-		InitOrUpdateResource(RHICmdList, &Self->ColorVertexBuffer);
+		InitOrUpdateResource(RHICmdList, &PositionVertexBuffer);
+		InitOrUpdateResource(RHICmdList, &StaticMeshVertexBuffer);
+		InitOrUpdateResource(RHICmdList, &ColorVertexBuffer);
 
 		FLocalVertexFactory::FDataType Data;
-		Self->PositionVertexBuffer.BindPositionVertexBuffer(VertexFactory, Data);
-		Self->StaticMeshVertexBuffer.BindTangentVertexBuffer(VertexFactory, Data);
-		Self->StaticMeshVertexBuffer.BindPackedTexCoordVertexBuffer(VertexFactory, Data);
-		Self->StaticMeshVertexBuffer.BindLightMapVertexBuffer(VertexFactory, Data, LightMapIndex);
-		Self->ColorVertexBuffer.BindColorVertexBuffer(VertexFactory, Data);
+		PositionVertexBuffer.BindPositionVertexBuffer(VertexFactory, Data);
+		StaticMeshVertexBuffer.BindTangentVertexBuffer(VertexFactory, Data);
+		StaticMeshVertexBuffer.BindPackedTexCoordVertexBuffer(VertexFactory, Data);
+		StaticMeshVertexBuffer.BindLightMapVertexBuffer(VertexFactory, Data, LightMapIndex);
+		ColorVertexBuffer.BindColorVertexBuffer(VertexFactory, Data);
 		VertexFactory->SetData(RHICmdList, Data);
 
 		InitOrUpdateResource(RHICmdList, VertexFactory);
-	});
+	};
+
+	if (RHICmdList)
+	{
+		Lambda(*RHICmdList);
+	}
+	else
+	{
+		ENQUEUE_RENDER_COMMAND(StaticMeshVertexBuffersInitWithDummyData)(RenderCommandPipe, MoveTemp(Lambda));
+	}
 }
 
-void FStaticMeshVertexBuffers::InitFromDynamicVertex(FRenderCommandPipe* RenderCommandPipe, FLocalVertexFactory* VertexFactory, TArray<FDynamicMeshVertex>& Vertices, uint32 NumTexCoords, uint32 LightMapIndex)
+void FStaticMeshVertexBuffers::InitFromDynamicVertex(FRHICommandListBase* RHICmdList, FRenderCommandPipe* RenderCommandPipe, FLocalVertexFactory* VertexFactory, TArray<FDynamicMeshVertex>& Vertices, uint32 NumTexCoords, uint32 LightMapIndex)
 {
 	check(NumTexCoords < MAX_STATIC_TEXCOORDS && NumTexCoords > 0);
 	check(LightMapIndex < NumTexCoords);
@@ -1214,24 +1228,31 @@ void FStaticMeshVertexBuffers::InitFromDynamicVertex(FRenderCommandPipe* RenderC
 		LightMapIndex = 0;
 	}
 
-	FStaticMeshVertexBuffers* Self = this;
-	ENQUEUE_RENDER_COMMAND(StaticMeshVertexBuffersLegacyInit)(RenderCommandPipe,
-		[VertexFactory, Self, LightMapIndex](FRHICommandListBase& RHICmdList)
-		{
-			InitOrUpdateResource(RHICmdList, &Self->PositionVertexBuffer);
-			InitOrUpdateResource(RHICmdList, &Self->StaticMeshVertexBuffer);
-			InitOrUpdateResource(RHICmdList, &Self->ColorVertexBuffer);
+	auto Lambda = [this, VertexFactory, LightMapIndex](FRHICommandListBase& RHICmdList)
+	{
+		InitOrUpdateResource(RHICmdList, &PositionVertexBuffer);
+		InitOrUpdateResource(RHICmdList, &StaticMeshVertexBuffer);
+		InitOrUpdateResource(RHICmdList, &ColorVertexBuffer);
 
-			FLocalVertexFactory::FDataType Data;
-			Self->PositionVertexBuffer.BindPositionVertexBuffer(VertexFactory, Data);
-			Self->StaticMeshVertexBuffer.BindTangentVertexBuffer(VertexFactory, Data);
-			Self->StaticMeshVertexBuffer.BindPackedTexCoordVertexBuffer(VertexFactory, Data);
-			Self->StaticMeshVertexBuffer.BindLightMapVertexBuffer(VertexFactory, Data, LightMapIndex);
-			Self->ColorVertexBuffer.BindColorVertexBuffer(VertexFactory, Data);
-			VertexFactory->SetData(RHICmdList, Data);
-
-			InitOrUpdateResource(RHICmdList, VertexFactory);
-		});
+		FLocalVertexFactory::FDataType Data;
+		PositionVertexBuffer.BindPositionVertexBuffer(VertexFactory, Data);
+		StaticMeshVertexBuffer.BindTangentVertexBuffer(VertexFactory, Data);
+		StaticMeshVertexBuffer.BindPackedTexCoordVertexBuffer(VertexFactory, Data);
+		StaticMeshVertexBuffer.BindLightMapVertexBuffer(VertexFactory, Data, LightMapIndex);
+		ColorVertexBuffer.BindColorVertexBuffer(VertexFactory, Data);
+		VertexFactory->SetData(RHICmdList, Data);
+
+		InitOrUpdateResource(RHICmdList, VertexFactory);
+	};
+
+	if (RHICmdList)
+	{
+		Lambda(*RHICmdList);
+	}
+	else
+	{
+		ENQUEUE_RENDER_COMMAND(StaticMeshVertexBuffersInitFromDynamicVertex)(RenderCommandPipe, MoveTemp(Lambda));
+	}
 };
 
 void FStaticMeshVertexBuffers::SetOwnerName(const FName& OwnerName)
diff --git a/Engine/Source/Runtime/Engine/Public/StaticMeshResources.h b/Engine/Source/Runtime/Engine/Public/StaticMeshResources.h
--- a/Engine/Source/Runtime/Engine/Public/StaticMeshResources.h
+++ b/Engine/Source/Runtime/Engine/Public/StaticMeshResources.h
@@ -315,34 +315,59 @@ struct FStaticMeshVertexBuffers
 	/** The buffer containing the vertex color data. */
 	FColorVertexBuffer ColorVertexBuffer;
 
-	/* This is a temporary function to refactor and convert old code, do not copy this as is and try to build your data as SoA from the beginning.*/
-	void ENGINE_API InitWithDummyData(FRenderCommandPipe* RenderCommandPipe, FLocalVertexFactory* VertexFactory, uint32 NumVertices, uint32 NumTexCoords = 1, uint32 LightMapIndex = 0);
+	void inline InitWithDummyData(FRHICommandListBase& RHICmdList, FLocalVertexFactory* VertexFactory, uint32 NumVertices, uint32 NumTexCoords = 1, uint32 LightMapIndex = 0)
+	{
+		InitWithDummyData(&RHICmdList, nullptr, VertexFactory, NumVertices, NumTexCoords, LightMapIndex);
+	}
+
+	void inline InitWithDummyData(FRenderCommandPipe* RenderCommandPipe, FLocalVertexFactory* VertexFactory, uint32 NumVertices, uint32 NumTexCoords = 1, uint32 LightMapIndex = 0)
+	{
+		InitWithDummyData(nullptr, RenderCommandPipe, VertexFactory, NumVertices, NumTexCoords, LightMapIndex);
+	}
 
 	void inline InitWithDummyData(FLocalVertexFactory* VertexFactory, uint32 NumVertices, uint32 NumTexCoords = 1, uint32 LightMapIndex = 0)
 	{
-		InitWithDummyData(nullptr, VertexFactory, NumVertices, NumTexCoords, LightMapIndex);
+		InitWithDummyData(nullptr, nullptr, VertexFactory, NumVertices, NumTexCoords, LightMapIndex);
 	}
 
-	/* This is a temporary function to refactor and convert old code, do not copy this as is and try to build your data as SoA from the beginning.*/
-	void ENGINE_API InitFromDynamicVertex(FRenderCommandPipe* RenderCommandPipe, FLocalVertexFactory* VertexFactory, TArray<FDynamicMeshVertex>& Vertices, uint32 NumTexCoords = 1, uint32 LightMapIndex = 0);
+	inline void InitFromDynamicVertex(FRHICommandListBase& RHICmdList, FLocalVertexFactory* VertexFactory, TArray<FDynamicMeshVertex>& Vertices, uint32 NumTexCoords = 1, uint32 LightMapIndex = 0)
+	{
+		InitFromDynamicVertex(&RHICmdList, nullptr, VertexFactory, Vertices, NumTexCoords, LightMapIndex);
+	}
+
+	inline void InitFromDynamicVertex(FRenderCommandPipe* RenderCommandPipe, FLocalVertexFactory* VertexFactory, TArray<FDynamicMeshVertex>& Vertices, uint32 NumTexCoords = 1, uint32 LightMapIndex = 0)
+	{
+		InitFromDynamicVertex(nullptr, RenderCommandPipe, VertexFactory, Vertices, NumTexCoords, LightMapIndex);
+	}
 
 	inline void InitFromDynamicVertex(FLocalVertexFactory* VertexFactory, TArray<FDynamicMeshVertex>& Vertices, uint32 NumTexCoords = 1, uint32 LightMapIndex = 0)
 	{
-		InitFromDynamicVertex(nullptr, VertexFactory, Vertices, NumTexCoords, LightMapIndex);
+		InitFromDynamicVertex(nullptr, nullptr, VertexFactory, Vertices, NumTexCoords, LightMapIndex);
 	}
 
-	/* This is a temporary function to refactor and convert old code, do not copy this as is and try to build your data as SoA from the beginning.*/
 	void ENGINE_API InitModelBuffers(TArray<FModelVertex>& Vertices);
 
-	/* This is a temporary function to refactor and convert old code, do not copy this as is and try to build your data as SoA from the beginning.*/
-	void ENGINE_API InitModelVF(FRenderCommandPipe* RenderCommandPipe, FLocalVertexFactory* VertexFactory);
+	inline void ENGINE_API InitModelVF(FRHICommandListBase& RHICmdList, FLocalVertexFactory* VertexFactory)
+	{
+		InitModelVF(&RHICmdList, nullptr, VertexFactory);
+	}
+
+	inline void ENGINE_API InitModelVF(FRenderCommandPipe* RenderCommandPipe, FLocalVertexFactory* VertexFactory)
+	{
+		InitModelVF(nullptr, RenderCommandPipe, VertexFactory);
+	}
 
 	inline void ENGINE_API InitModelVF(FLocalVertexFactory* VertexFactory)
 	{
-		InitModelVF(nullptr, VertexFactory);
+		InitModelVF(nullptr, nullptr, VertexFactory);
 	}
 
 	void ENGINE_API SetOwnerName(const FName& OwnerName);
+
+private:
+	void ENGINE_API InitWithDummyData(FRHICommandListBase* RHICmdList, FRenderCommandPipe* RenderCommandPipe, FLocalVertexFactory* VertexFactory, uint32 NumVertices, uint32 NumTexCoords, uint32 LightMapIndex);
+	void ENGINE_API InitFromDynamicVertex(FRHICommandListBase* RHICmdList, FRenderCommandPipe* RenderCommandPipe, FLocalVertexFactory* VertexFactory, TArray<FDynamicMeshVertex>& Vertices, uint32 NumTexCoords, uint32 LightMapIndex);
+	void ENGINE_API InitModelVF(FRHICommandListBase* RHICmdList, FRenderCommandPipe* RenderCommandPipe, FLocalVertexFactory* VertexFactory);
 };
 
 struct FAdditionalStaticMeshIndexBuffers
