Moved skeletal mesh and GPU skin cache updates off the render thread to a dedicated render command pipe.

diff --git a/Engine/Plugins/MeshPainting/Source/MeshPaintingToolset/Private/MeshPaintSkeletalMeshAdapter.cpp b/Engine/Plugins/MeshPainting/Source/MeshPaintingToolset/Private/MeshPaintSkeletalMeshAdapter.cpp
--- a/Engine/Plugins/MeshPainting/Source/MeshPaintingToolset/Private/MeshPaintSkeletalMeshAdapter.cpp
+++ b/Engine/Plugins/MeshPainting/Source/MeshPaintingToolset/Private/MeshPaintSkeletalMeshAdapter.cpp
@@ -11,7 +11,7 @@
 #include "InterchangeAssetImportData.h"
 #include "InterchangeGenericAssetsPipeline.h"
 #include "InterchangePythonPipelineBase.h"
-#include "RenderingThread.h"
+#include "Rendering/RenderCommandPipes.h"
 
 DEFINE_LOG_CATEGORY_STATIC(LogMeshPaintSkeletalMeshAdapter, Log, All);
 //////////////////////////////////////////////////////////////////////////
@@ -382,7 +382,7 @@ void FMeshPaintSkeletalMeshComponentAdapter::PreEdit()
 		LODData->StaticVertexBuffers.ColorVertexBuffer.InitFromSingleColor(FColor(255, 255, 255, 255), LODData->GetNumVertices());
 		ReferencedSkeletalMesh->SetHasVertexColors(true);
 		ReferencedSkeletalMesh->SetVertexColorGuid(FGuid::NewGuid());
-		BeginInitResource(&LODData->StaticVertexBuffers.ColorVertexBuffer);
+		BeginInitResource(&LODData->StaticVertexBuffers.ColorVertexBuffer, &UE::RenderCommandPipe::SkeletalMesh);
 	}
 	//Make sure we change the import data so the re-import do not replace the new data
 	if (ReferencedSkeletalMesh->GetAssetImportData())
diff --git a/Engine/Source/Editor/MeshPaint/Private/MeshPaintSkeletalMeshAdapter.cpp b/Engine/Source/Editor/MeshPaint/Private/MeshPaintSkeletalMeshAdapter.cpp
--- a/Engine/Source/Editor/MeshPaint/Private/MeshPaintSkeletalMeshAdapter.cpp
+++ b/Engine/Source/Editor/MeshPaint/Private/MeshPaintSkeletalMeshAdapter.cpp
@@ -10,6 +10,7 @@
 #include "Rendering/SkeletalMeshRenderData.h"
 #include "Rendering/SkeletalMeshModel.h"
 #include "Factories/FbxSkeletalMeshImportData.h"
+#include "Rendering/RenderCommandPipes.h"
 
 DEFINE_LOG_CATEGORY_STATIC(LogMeshPaintSkeletalMeshAdapter, Log, All);
 
@@ -455,7 +456,7 @@ void FMeshPaintGeometryAdapterForSkeletalMeshes::PreEdit()
 		LODData->StaticVertexBuffers.ColorVertexBuffer.InitFromSingleColor(FColor(255, 255, 255, 255), LODData->GetNumVertices());
 		ReferencedSkeletalMesh->SetHasVertexColors(true);
 		ReferencedSkeletalMesh->SetVertexColorGuid(FGuid::NewGuid());
-		BeginInitResource(&LODData->StaticVertexBuffers.ColorVertexBuffer);
+		BeginInitResource(&LODData->StaticVertexBuffers.ColorVertexBuffer, &UE::RenderCommandPipe::SkeletalMesh);
 	}
 	//Make sure we change the import data so the re-import do not replace the new data
 	if (ReferencedSkeletalMesh->GetAssetImportData())
diff --git a/Engine/Source/Runtime/Engine/Private/Components/SkeletalMeshComponent.cpp b/Engine/Source/Runtime/Engine/Private/Components/SkeletalMeshComponent.cpp
--- a/Engine/Source/Runtime/Engine/Private/Components/SkeletalMeshComponent.cpp
+++ b/Engine/Source/Runtime/Engine/Private/Components/SkeletalMeshComponent.cpp
@@ -25,6 +25,7 @@
 #include "SkeletalRender.h"
 #include "SkinnedAssetCompiler.h"
 #include "Rendering/SkeletalMeshRenderData.h"
+#include "Rendering/RenderCommandPipes.h"
 
 #include "Logging/MessageLog.h"
 #include "Animation/AnimNode_LinkedInputPose.h"
@@ -4058,8 +4059,8 @@ void USkeletalMeshComponent::SendRenderDynamicData_Concurrent()
 
 		FSkeletalMeshSceneProxy* TargetProxy = (FSkeletalMeshSceneProxy*)SceneProxy;
 
-		ENQUEUE_RENDER_COMMAND(PoseWatchDynamicDataCommand)(
-			[TargetProxy, NewDynamicData](FRHICommandListImmediate& RHICommandList)
+		ENQUEUE_RENDER_COMMAND(PoseWatchDynamicDataCommand)(UE::RenderCommandPipe::SkeletalMesh,
+			[TargetProxy, NewDynamicData]
 			{
 				if (TargetProxy->PoseWatchDynamicData)
 				{
diff --git a/Engine/Source/Runtime/Engine/Private/GPUSkinVertexFactory.cpp b/Engine/Source/Runtime/Engine/Private/GPUSkinVertexFactory.cpp
--- a/Engine/Source/Runtime/Engine/Private/GPUSkinVertexFactory.cpp
+++ b/Engine/Source/Runtime/Engine/Private/GPUSkinVertexFactory.cpp
@@ -242,7 +242,7 @@ static TAutoConsoleVariable<int32> CVarRHICmdDeferSkeletalLockAndFillToRHIThread
 
 static bool DeferSkeletalLockAndFillToRHIThread()
 {
-	return IsRunningRHIInSeparateThread() && CVarRHICmdDeferSkeletalLockAndFillToRHIThread.GetValueOnRenderThread() > 0;
+	return IsRunningRHIInSeparateThread() && CVarRHICmdDeferSkeletalLockAndFillToRHIThread.GetValueOnRenderThread() > 0 && !UE::RenderCommandPipe::IsReplaying();
 }
 
 bool FGPUBaseSkinVertexFactory::FShaderDataType::UpdateBoneData(FRHICommandList& RHICmdList, const TArray<FMatrix44f>& ReferenceToLocalMatrices,
diff --git a/Engine/Source/Runtime/Engine/Private/LevelTick.cpp b/Engine/Source/Runtime/Engine/Private/LevelTick.cpp
--- a/Engine/Source/Runtime/Engine/Private/LevelTick.cpp
+++ b/Engine/Source/Runtime/Engine/Private/LevelTick.cpp
@@ -24,6 +24,7 @@
 #include "Camera/CameraPhotography.h"
 #include "UObject/Stack.h"
 #include "PhysicsEngine/CollisionAnalyzerCapture.h"
+#include "Rendering/RenderCommandPipes.h"
 
 #if !UE_SERVER
 #include "IMediaModule.h"
@@ -973,23 +974,23 @@ struct FSendAllEndOfFrameUpdates
 void BeginSendEndOfFrameUpdatesDrawEvent(FSendAllEndOfFrameUpdates& SendAllEndOfFrameUpdates)
 {
 	BEGIN_DRAW_EVENTF_GAMETHREAD(SendAllEndOfFrameUpdates, SendAllEndOfFrameUpdates.DrawEvent, TEXT("SendAllEndOfFrameUpdates"));
-	
-	ENQUEUE_RENDER_COMMAND(BeginDrawEventCommand)(
-		[GPUSkinCache = SendAllEndOfFrameUpdates.GPUSkinCache](FRHICommandListImmediate&)
+
+	ENQUEUE_RENDER_COMMAND(BeginDrawEventCommand)(UE::RenderCommandPipe::SkeletalMesh,
+		[GPUSkinCache = SendAllEndOfFrameUpdates.GPUSkinCache]
+	{
+		if (GPUSkinCache != nullptr)
 		{
-			if (GPUSkinCache != nullptr)
-			{
-				GPUSkinCache->BeginBatchDispatch();
-			}
-		});
+			GPUSkinCache->BeginBatchDispatch();
+		}
+	});
 }
 
 DECLARE_GPU_STAT(EndOfFrameUpdates);
 DECLARE_GPU_STAT(GPUSkinCacheRayTracingGeometry);
 void EndSendEndOfFrameUpdatesDrawEvent(FSendAllEndOfFrameUpdates& SendAllEndOfFrameUpdates)
 {
-	ENQUEUE_RENDER_COMMAND(EndDrawEventCommand)(
-		[GPUSkinCache = SendAllEndOfFrameUpdates.GPUSkinCache] (FRHICommandListImmediate&)
+	ENQUEUE_RENDER_COMMAND(EndDrawEventCommand)(UE::RenderCommandPipe::SkeletalMesh,
+		[GPUSkinCache = SendAllEndOfFrameUpdates.GPUSkinCache]
 	{
 		if (GPUSkinCache != nullptr)
 		{
diff --git a/Engine/Source/Runtime/Engine/Private/Rendering/MultiSizeIndexContainer.cpp b/Engine/Source/Runtime/Engine/Private/Rendering/MultiSizeIndexContainer.cpp
--- a/Engine/Source/Runtime/Engine/Private/Rendering/MultiSizeIndexContainer.cpp
+++ b/Engine/Source/Runtime/Engine/Private/Rendering/MultiSizeIndexContainer.cpp
@@ -1,10 +1,10 @@
 // Copyright Epic Games, Inc. All Rights Reserved.
 
 #include "Rendering/MultiSizeIndexContainer.h"
+#include "Rendering/RenderCommandPipes.h"
 #include "EngineLogs.h"
 #include "RawIndexBuffer.h"
 #include "Stats/Stats.h"
-#include "RenderingThread.h"
 
 FMultiSizeIndexContainer::~FMultiSizeIndexContainer()
 {
@@ -31,7 +31,7 @@ void FMultiSizeIndexContainer::InitResources()
 	check(IsInGameThread());
 	if (IndexBuffer)
 	{
-		BeginInitResource(IndexBuffer);
+		BeginInitResource(IndexBuffer, &UE::RenderCommandPipe::SkeletalMesh);
 	}
 }
 
@@ -43,7 +43,7 @@ void FMultiSizeIndexContainer::ReleaseResources()
 	check(IsInGameThread());
 	if (IndexBuffer)
 	{
-		BeginReleaseResource(IndexBuffer);
+		BeginReleaseResource(IndexBuffer, &UE::RenderCommandPipe::SkeletalMesh);
 	}
 }
 
diff --git a/Engine/Source/Runtime/Engine/Private/Rendering/RenderCommandPipes.cpp b/Engine/Source/Runtime/Engine/Private/Rendering/RenderCommandPipes.cpp
new file mode 100644
--- /dev/null
+++ b/Engine/Source/Runtime/Engine/Private/Rendering/RenderCommandPipes.cpp
@@ -0,0 +1,5 @@
+// Copyright Epic Games, Inc. All Rights Reserved.
+
+#include "Rendering/RenderCommandPipes.h"
+
+DEFINE_RENDER_COMMAND_PIPE(SkeletalMesh);
\ No newline at end of file
diff --git a/Engine/Source/Runtime/Engine/Private/Rendering/SkinWeightVertexBuffer.cpp b/Engine/Source/Runtime/Engine/Private/Rendering/SkinWeightVertexBuffer.cpp
--- a/Engine/Source/Runtime/Engine/Private/Rendering/SkinWeightVertexBuffer.cpp
+++ b/Engine/Source/Runtime/Engine/Private/Rendering/SkinWeightVertexBuffer.cpp
@@ -8,6 +8,7 @@
 #include "RHIResourceUpdates.h"
 #include "SkeletalMeshLegacyCustomVersions.h"
 #include "UObject/UE5MainStreamObjectVersion.h"
+#include "Rendering/RenderCommandPipes.h"
 
 /*-----------------------------------------------------------------------------
 FSkinWeightLookupVertexBuffer
@@ -806,14 +807,14 @@ void FSkinWeightVertexBuffer::SetOwnerName(const FName& OwnerName)
 
 void FSkinWeightVertexBuffer::BeginInitResources()
 {
-	BeginInitResource(&LookupVertexBuffer);
-	BeginInitResource(&DataVertexBuffer);
+	BeginInitResource(&LookupVertexBuffer, &UE::RenderCommandPipe::SkeletalMesh);
+	BeginInitResource(&DataVertexBuffer, &UE::RenderCommandPipe::SkeletalMesh);
 }
 
 void FSkinWeightVertexBuffer::BeginReleaseResources()
 {
-	BeginReleaseResource(&LookupVertexBuffer);
-	BeginReleaseResource(&DataVertexBuffer);
+	BeginReleaseResource(&LookupVertexBuffer, &UE::RenderCommandPipe::SkeletalMesh);
+	BeginReleaseResource(&DataVertexBuffer, &UE::RenderCommandPipe::SkeletalMesh);
 }
 
 void FSkinWeightVertexBuffer::ReleaseResources()
diff --git a/Engine/Source/Runtime/Engine/Private/SkeletalMesh.cpp b/Engine/Source/Runtime/Engine/Private/SkeletalMesh.cpp
--- a/Engine/Source/Runtime/Engine/Private/SkeletalMesh.cpp
+++ b/Engine/Source/Runtime/Engine/Private/SkeletalMesh.cpp
@@ -44,6 +44,7 @@
 #include "Engine/AssetUserData.h"
 #include "Animation/NodeMappingContainer.h"
 #include "Rendering/SkeletalMeshRenderData.h"
+#include "Rendering/RenderCommandPipes.h"
 #include "AnimationRuntime.h"
 #include "Animation/AnimSequence.h"
 #include "Streaming/SkeletalMeshUpdate.h"
@@ -6807,8 +6808,8 @@ void FSkeletalMeshSceneProxy::UpdateMorphMaterialUsage_GameThread(TArray<UMateri
 			ERHIFeatureLevel::Type InFeatureLevel = GetScene().GetFeatureLevel();
 			FSkeletalMeshSceneProxy* SkelMeshSceneProxy = this;
 			FMaterialRelevance DefaultRelevance = DefaultMaterial->GetRelevance(InFeatureLevel);
-			ENQUEUE_RENDER_COMMAND(UpdateSkelProxyLODSectionElementsCmd)(
-				[InMaterialsToSwap, DefaultMaterial, DefaultRelevance, InFeatureLevel, SkelMeshSceneProxy](FRHICommandList& RHICmdList)
+			ENQUEUE_RENDER_COMMAND(UpdateSkelProxyLODSectionElementsCmd)(UE::RenderCommandPipe::SkeletalMesh,
+				[InMaterialsToSwap, DefaultMaterial, DefaultRelevance, InFeatureLevel, SkelMeshSceneProxy]
 				{
 					for( int32 LodIdx=0; LodIdx < SkelMeshSceneProxy->LODSections.Num(); LodIdx++ )
 					{
diff --git a/Engine/Source/Runtime/Engine/Private/SkeletalMeshComponentPhysics.cpp b/Engine/Source/Runtime/Engine/Private/SkeletalMeshComponentPhysics.cpp
--- a/Engine/Source/Runtime/Engine/Private/SkeletalMeshComponentPhysics.cpp
+++ b/Engine/Source/Runtime/Engine/Private/SkeletalMeshComponentPhysics.cpp
@@ -14,6 +14,7 @@
 #include "AnimationRuntime.h"
 #include "ClothCollisionData.h"
 #include "ClothingSimulationInteractor.h"
+#include "Rendering/RenderCommandPipes.h"
 
 #include "Logging/MessageLog.h"
 #include "CollisionDebugDrawingPublic.h"
@@ -1659,8 +1660,8 @@ void USkeletalMeshComponent::SendRenderDebugPhysics(FPrimitiveSceneProxy* Overri
 			}
 		}
 
-		ENQUEUE_RENDER_COMMAND(SkeletalMesh_SendRenderDebugPhysics)(
-			[UseSceneProxy, DebugMassData](FRHICommandListImmediate& RHICmdList)
+		ENQUEUE_RENDER_COMMAND(SkeletalMesh_SendRenderDebugPhysics)(UE::RenderCommandPipe::SkeletalMesh,
+			[UseSceneProxy, DebugMassData]
 			{
 				UseSceneProxy->SetDebugMassData(DebugMassData);
 			}
diff --git a/Engine/Source/Runtime/Engine/Private/SkeletalMeshLODRenderData.cpp b/Engine/Source/Runtime/Engine/Private/SkeletalMeshLODRenderData.cpp
--- a/Engine/Source/Runtime/Engine/Private/SkeletalMeshLODRenderData.cpp
+++ b/Engine/Source/Runtime/Engine/Private/SkeletalMeshLODRenderData.cpp
@@ -12,6 +12,7 @@
 #include "Interfaces/ITargetPlatform.h"
 #include "Interfaces/ITargetPlatformManagerModule.h"
 #include "Rendering/SkeletalMeshLODModel.h"
+#include "Rendering/RenderCommandPipes.h"
 #include "Serialization/MemoryReader.h"
 #include "SkeletalMeshLegacyCustomVersions.h"
 #include "UObject/Package.h"
@@ -206,7 +207,7 @@ void FSkeletalMeshLODRenderData::InitMorphResources()
 {
 	if (GAllowSkinnedMorphDataStreaming && !MorphTargetVertexInfoBuffers.IsRHIIntialized() && MorphTargetVertexInfoBuffers.IsMorphCPUDataValid() && MorphTargetVertexInfoBuffers.NumTotalBatches > 0)
 	{
-		BeginInitResource(&MorphTargetVertexInfoBuffers);
+		BeginInitResource(&MorphTargetVertexInfoBuffers, &UE::RenderCommandPipe::SkeletalMesh);
 	}
 }
 
@@ -223,9 +224,9 @@ void FSkeletalMeshLODRenderData::InitResources(bool bNeedsVertexColors, int32 LO
 	MultiSizeIndexContainer.InitResources();
 
 	StaticVertexBuffers.PositionVertexBuffer.SetOwnerName(OwnerName);
-	BeginInitResource(&StaticVertexBuffers.PositionVertexBuffer);
+	BeginInitResource(&StaticVertexBuffers.PositionVertexBuffer, &UE::RenderCommandPipe::SkeletalMesh);
 	StaticVertexBuffers.StaticMeshVertexBuffer.SetOwnerName(OwnerName);
-	BeginInitResource(&StaticVertexBuffers.StaticMeshVertexBuffer);
+	BeginInitResource(&StaticVertexBuffers.StaticMeshVertexBuffer, &UE::RenderCommandPipe::SkeletalMesh);
 
 	SkinWeightVertexBuffer.SetOwnerName(OwnerName);
 	SkinWeightVertexBuffer.BeginInitResources();
@@ -234,14 +235,14 @@ void FSkeletalMeshLODRenderData::InitResources(bool bNeedsVertexColors, int32 LO
 	{
 		// Only init the color buffer if the mesh has vertex colors
 		StaticVertexBuffers.ColorVertexBuffer.SetOwnerName(OwnerName);
-		BeginInitResource(&StaticVertexBuffers.ColorVertexBuffer);
+		BeginInitResource(&StaticVertexBuffers.ColorVertexBuffer, &UE::RenderCommandPipe::SkeletalMesh);
 	}
 
 	if (ClothVertexBuffer.GetNumVertices() > 0)
 	{
 		// Only init the clothing buffer if the mesh has clothing data
 		ClothVertexBuffer.SetOwnerName(OwnerName);
-		BeginInitResource(&ClothVertexBuffer);
+		BeginInitResource(&ClothVertexBuffer, &UE::RenderCommandPipe::SkeletalMesh);
 	}
 
 	// We can discard any the cooked DuplicatedVertices here based on runtime settings.
@@ -262,7 +263,7 @@ void FSkeletalMeshLODRenderData::InitResources(bool bNeedsVertexColors, int32 LO
 				// No need to discard CPU data in cooked builds as bNeedsCPUAccess is false (see FDuplicatedVerticesBuffer constructor), 
 				// so it'd be auto-discarded after the RHI has copied the resource data. Keep CPU data when in the editor for geometry operations.
 				RenderSection.DuplicatedVerticesBuffer.SetOwnerName(OwnerName);
-				BeginInitResource(&RenderSection.DuplicatedVerticesBuffer);
+				BeginInitResource(&RenderSection.DuplicatedVerticesBuffer, &UE::RenderCommandPipe::SkeletalMesh);
 			}
 		}
 	}
@@ -278,7 +279,7 @@ void FSkeletalMeshLODRenderData::InitResources(bool bNeedsVertexColors, int32 LO
 	if (!MorphTargetVertexInfoBuffers.IsRHIIntialized() && MorphTargetVertexInfoBuffers.IsMorphCPUDataValid() && MorphTargetVertexInfoBuffers.NumTotalBatches > 0)
 	{
 		MorphTargetVertexInfoBuffers.SetOwnerName(OwnerName);
-		BeginInitResource(&MorphTargetVertexInfoBuffers);
+		BeginInitResource(&MorphTargetVertexInfoBuffers, &UE::RenderCommandPipe::SkeletalMesh);
 	}
 
 	VertexAttributeBuffers.InitResources();
@@ -289,7 +290,7 @@ void FSkeletalMeshLODRenderData::InitResources(bool bNeedsVertexColors, int32 LO
 		if (SourceRayTracingGeometry.RawData.Num() > 0)
 		{
 			SourceRayTracingGeometry.SetOwnerName(OwnerName);
-			BeginInitResource(&SourceRayTracingGeometry);
+			BeginInitResource(&SourceRayTracingGeometry, &UE::RenderCommandPipe::SkeletalMesh);
 		}
 	}
 #endif
@@ -301,16 +302,16 @@ void FSkeletalMeshLODRenderData::ReleaseResources()
 
 	MultiSizeIndexContainer.ReleaseResources();
 
-	BeginReleaseResource(&StaticVertexBuffers.PositionVertexBuffer);
-	BeginReleaseResource(&StaticVertexBuffers.StaticMeshVertexBuffer);
+	BeginReleaseResource(&StaticVertexBuffers.PositionVertexBuffer, &UE::RenderCommandPipe::SkeletalMesh);
+	BeginReleaseResource(&StaticVertexBuffers.StaticMeshVertexBuffer, &UE::RenderCommandPipe::SkeletalMesh);
 	SkinWeightVertexBuffer.BeginReleaseResources();
-	BeginReleaseResource(&StaticVertexBuffers.ColorVertexBuffer);
-	BeginReleaseResource(&ClothVertexBuffer);
+	BeginReleaseResource(&StaticVertexBuffers.ColorVertexBuffer, &UE::RenderCommandPipe::SkeletalMesh);
+	BeginReleaseResource(&ClothVertexBuffer, &UE::RenderCommandPipe::SkeletalMesh);
 	for (FSkelMeshRenderSection& RenderSection : RenderSections)
 	{
-		BeginReleaseResource(&RenderSection.DuplicatedVerticesBuffer);
+		BeginReleaseResource(&RenderSection.DuplicatedVerticesBuffer, &UE::RenderCommandPipe::SkeletalMesh);
 	}
-	BeginReleaseResource(&MorphTargetVertexInfoBuffers);
+	BeginReleaseResource(&MorphTargetVertexInfoBuffers, &UE::RenderCommandPipe::SkeletalMesh);
 
 	DEC_DWORD_STAT_BY(STAT_SkeletalMeshVertexMemory, SkinWeightProfilesData.GetResourcesSize());
 	SkinWeightProfilesData.ReleaseResources();
@@ -319,8 +320,8 @@ void FSkeletalMeshLODRenderData::ReleaseResources()
 #if RHI_RAYTRACING
 	if (IsRayTracingAllowed())
 	{
-		BeginReleaseResource(&SourceRayTracingGeometry);
-		BeginReleaseResource(&StaticRayTracingGeometry);
+		BeginReleaseResource(&SourceRayTracingGeometry, &UE::RenderCommandPipe::SkeletalMesh);
+		BeginReleaseResource(&StaticRayTracingGeometry, &UE::RenderCommandPipe::SkeletalMesh);
 	}
 #endif
 }
diff --git a/Engine/Source/Runtime/Engine/Private/SkeletalMeshRenderData.cpp b/Engine/Source/Runtime/Engine/Private/SkeletalMeshRenderData.cpp
--- a/Engine/Source/Runtime/Engine/Private/SkeletalMeshRenderData.cpp
+++ b/Engine/Source/Runtime/Engine/Private/SkeletalMeshRenderData.cpp
@@ -8,6 +8,7 @@
 #include "Engine/SkinnedAssetCommon.h"
 #include "EngineLogs.h"
 #include "UObject/Package.h"
+#include "Rendering/RenderCommandPipes.h"
 
 #if WITH_EDITOR
 #include "ProfilingDebugging/CookStats.h"
@@ -454,7 +455,7 @@ void FSkeletalMeshRenderData::SyncUVChannelData(const TArray<FSkeletalMaterial>&
 	// to the render thread.
 	if (bInitialized)
 	{
-		ENQUEUE_RENDER_COMMAND(SyncUVChannelData)([this, UpdateData = MoveTemp(UpdateData)](FRHICommandListImmediate& RHICmdList)
+		ENQUEUE_RENDER_COMMAND(SyncUVChannelData)(UE::RenderCommandPipe::SkeletalMesh, [this, UpdateData = MoveTemp(UpdateData)]
 		{
 			Swap(UVChannelDataPerMaterial, *UpdateData.Get());
 		});
@@ -641,8 +642,8 @@ void FSkeletalMeshRenderData::InitResources(bool bNeedsVertexColors, TArray<UMor
 			}
 		}
 
-		ENQUEUE_RENDER_COMMAND(CmdSetSkeletalMeshReadyForStreaming)(
-			[this, Owner](FRHICommandListImmediate&)
+		ENQUEUE_RENDER_COMMAND(CmdSetSkeletalMeshReadyForStreaming)(UE::RenderCommandPipe::SkeletalMesh,
+			[this, Owner]
 		{
 			bReadyForStreaming = true;
 		});
diff --git a/Engine/Source/Runtime/Engine/Private/SkeletalRenderCPUSkin.cpp b/Engine/Source/Runtime/Engine/Private/SkeletalRenderCPUSkin.cpp
--- a/Engine/Source/Runtime/Engine/Private/SkeletalRenderCPUSkin.cpp
+++ b/Engine/Source/Runtime/Engine/Private/SkeletalRenderCPUSkin.cpp
@@ -15,6 +15,7 @@
 #include "RenderCore.h"
 #include "SceneInterface.h"
 #include "Stats/StatsTrace.h"
+#include "Rendering/RenderCommandPipes.h"
 
 #if RHI_RAYTRACING
 #include "Engine/SkinnedAssetCommon.h"
@@ -154,8 +155,8 @@ void FSkeletalMeshObjectCPUSkin::Update(
 		// queue a call to update this data
 		{
 			FSkeletalMeshObjectCPUSkin* MeshObject = this;
-			ENQUEUE_RENDER_COMMAND(SkelMeshObjectUpdateDataCommand)(
-				[MeshObject, FrameNumberToPrepare, RevisionNumber, NewDynamicData](FRHICommandListImmediate& RHICmdList)
+			ENQUEUE_RENDER_COMMAND(SkelMeshObjectUpdateDataCommand)(UE::RenderCommandPipe::SkeletalMesh,
+				[MeshObject, FrameNumberToPrepare, RevisionNumber, NewDynamicData](FRHICommandList& RHICmdList)
 				{
 					FScopeCycleCounter Context(MeshObject->GetStatId());
 					MeshObject->UpdateDynamicData_RenderThread(RHICmdList, NewDynamicData, FrameNumberToPrepare, RevisionNumber);
@@ -296,13 +297,13 @@ void FSkeletalMeshObjectCPUSkin::CacheVertices(int32 LODIndex, bool bForce) cons
 			}
 		}
 
-		BeginUpdateResourceRHI(&MeshLOD.PositionVertexBuffer);
-		BeginUpdateResourceRHI(&MeshLOD.StaticMeshVertexBuffer);
+		BeginUpdateResourceRHI(&MeshLOD.PositionVertexBuffer, &UE::RenderCommandPipe::SkeletalMesh);
+		BeginUpdateResourceRHI(&MeshLOD.StaticMeshVertexBuffer, &UE::RenderCommandPipe::SkeletalMesh);
 
 		const FSkeletalMeshObjectCPUSkin::FSkeletalMeshObjectLOD* MeshLODptr = &MeshLOD;
 		FLocalVertexFactory* VertexFactoryPtr = &MeshLOD.VertexFactory;
-		ENQUEUE_RENDER_COMMAND(UpdateSkeletalMeshCPUSkinVertexFactory)(
-			[VertexFactoryPtr, MeshLODptr](FRHICommandListImmediate& RHICmdList)
+		ENQUEUE_RENDER_COMMAND(UpdateSkeletalMeshCPUSkinVertexFactory)(UE::RenderCommandPipe::SkeletalMesh,
+			[VertexFactoryPtr, MeshLODptr](FRHICommandList& RHICmdList)
 		{
 			FLocalVertexFactory::FDataType Data;
 
@@ -347,8 +348,8 @@ void FSkeletalMeshObjectCPUSkin::FSkeletalMeshObjectLOD::InitResources(FSkelMesh
 		StaticMeshVertexBuffer.SetVertexUV(i, 0, SrcVertexBuf.GetVertexUV(i, 0));
 	}
 
-	BeginInitResource(&PositionVertexBuffer);
-	BeginInitResource(&StaticMeshVertexBuffer);
+	BeginInitResource(&PositionVertexBuffer, &UE::RenderCommandPipe::SkeletalMesh);
+	BeginInitResource(&StaticMeshVertexBuffer, &UE::RenderCommandPipe::SkeletalMesh);
 
 	FSkeletalMeshObjectCPUSkin::FSkeletalMeshObjectLOD* Self = this;
 	FLocalVertexFactory* VertexFactoryPtr = &VertexFactory;
@@ -449,12 +450,12 @@ void FSkeletalMeshObjectCPUSkin::FSkeletalMeshObjectLOD::UpdateSkinWeights(FSkel
  */
 void FSkeletalMeshObjectCPUSkin::FSkeletalMeshObjectLOD::ReleaseResources()
 {	
-	BeginReleaseResource(&VertexFactory);
-	BeginReleaseResource(&PositionVertexBuffer);
-	BeginReleaseResource(&StaticMeshVertexBuffer);
+	BeginReleaseResource(&VertexFactory, &UE::RenderCommandPipe::SkeletalMesh);
+	BeginReleaseResource(&PositionVertexBuffer, &UE::RenderCommandPipe::SkeletalMesh);
+	BeginReleaseResource(&StaticMeshVertexBuffer, &UE::RenderCommandPipe::SkeletalMesh);
 
 #if RHI_RAYTRACING
-	BeginReleaseResource(&RayTracingGeometry);
+	BeginReleaseResource(&RayTracingGeometry, &UE::RenderCommandPipe::SkeletalMesh);
 #endif // RHI_RAYTRACING
 
 	bResourcesInitialized = false;
diff --git a/Engine/Source/Runtime/Engine/Private/SkeletalRenderGPUSkin.cpp b/Engine/Source/Runtime/Engine/Private/SkeletalRenderGPUSkin.cpp
--- a/Engine/Source/Runtime/Engine/Private/SkeletalRenderGPUSkin.cpp
+++ b/Engine/Source/Runtime/Engine/Private/SkeletalRenderGPUSkin.cpp
@@ -10,7 +10,7 @@
 #include "SkeletalRender.h"
 #include "GPUSkinCache.h"
 #include "RayTracingSkinnedGeometry.h"
-#include "RenderingThread.h"
+#include "Rendering/RenderCommandPipes.h"
 #include "ShaderParameterUtils.h"
 #include "SceneInterface.h"
 #include "SkeletalMeshSceneProxy.h"
@@ -143,17 +143,17 @@ void FMorphVertexBufferPool::InitResources(const FName& OwnerName)
 
 	check(!MorphVertexBuffers[0].VertexBufferRHI.IsValid());
 	check(!MorphVertexBuffers[1].VertexBufferRHI.IsValid());
-	BeginInitResource(&MorphVertexBuffers[0]);
+	BeginInitResource(&MorphVertexBuffers[0], &UE::RenderCommandPipe::SkeletalMesh);
 	if (bDoubleBuffer)
 	{
-		BeginInitResource(&MorphVertexBuffers[1]);
+		BeginInitResource(&MorphVertexBuffers[1], &UE::RenderCommandPipe::SkeletalMesh);
 	}
 }
 
 void FMorphVertexBufferPool::ReleaseResources()
 {
-	BeginReleaseResource(&MorphVertexBuffers[0]);
-	BeginReleaseResource(&MorphVertexBuffers[1]);
+	BeginReleaseResource(&MorphVertexBuffers[0], &UE::RenderCommandPipe::SkeletalMesh);
+	BeginReleaseResource(&MorphVertexBuffers[1], &UE::RenderCommandPipe::SkeletalMesh);
 }
 
 SIZE_T FMorphVertexBufferPool::GetResourceSize() const
@@ -281,7 +281,7 @@ void FSkeletalMeshObjectGPUSkin::InitResources(USkinnedMeshComponent* InMeshComp
 #if RHI_RAYTRACING
 	if (bSupportRayTracing)
 	{
-		BeginInitResource(&RayTracingGeometry);
+		BeginInitResource(&RayTracingGeometry, &UE::RenderCommandPipe::SkeletalMesh);
 	}
 #endif
 }
@@ -296,7 +296,7 @@ void FSkeletalMeshObjectGPUSkin::ReleaseResources()
 	ReleaseMorphResources();
 	FSkeletalMeshObjectGPUSkin* MeshObject = this;
 	FGPUSkinCacheEntry** PtrSkinCacheEntry = &SkinCacheEntry;
-	ENQUEUE_RENDER_COMMAND(WaitRHIThreadFenceForDynamicData)(
+	ENQUEUE_RENDER_COMMAND(WaitRHIThreadFenceForDynamicData)(UE::RenderCommandPipe::SkeletalMesh,
 		[MeshObject, PtrSkinCacheEntry, &SkinCacheEntryForRayTracing = SkinCacheEntryForRayTracing](FRHICommandList& RHICmdList)
 		{
 			FGPUSkinCacheEntry*& LocalSkinCacheEntry = *PtrSkinCacheEntry;
@@ -313,14 +313,14 @@ void FSkeletalMeshObjectGPUSkin::ReleaseResources()
 #if RHI_RAYTRACING
 	if (bSupportRayTracing)
 	{
-		BeginReleaseResource(&RayTracingGeometry);
+		BeginReleaseResource(&RayTracingGeometry, &UE::RenderCommandPipe::SkeletalMesh);
 	}
 
 	// Only enqueue when intialized
 	if (RayTracingUpdateQueue != nullptr || RayTracingDynamicVertexBuffer.NumBytes > 0)
 	{
-		ENQUEUE_RENDER_COMMAND(ReleaseRayTracingDynamicVertexBuffer)(
-			[RayTracingUpdateQueue = RayTracingUpdateQueue, RayTracingGeometryPtr = &RayTracingGeometry, &RayTracingDynamicVertexBuffer = RayTracingDynamicVertexBuffer](FRHICommandListImmediate& RHICmdList) mutable
+		ENQUEUE_RENDER_COMMAND(ReleaseRayTracingDynamicVertexBuffer)(UE::RenderCommandPipe::SkeletalMesh,
+			[RayTracingUpdateQueue = RayTracingUpdateQueue, RayTracingGeometryPtr = &RayTracingGeometry, &RayTracingDynamicVertexBuffer = RayTracingDynamicVertexBuffer](FRHICommandList& RHICmdList) mutable
 			{
 				if (RayTracingUpdateQueue != nullptr)
 				{
@@ -419,8 +419,8 @@ void FSkeletalMeshObjectGPUSkin::Update(
 
 	// queue a call to update this data
 	FSkeletalMeshObjectGPUSkin* MeshObject = this;
-	ENQUEUE_RENDER_COMMAND(SkelMeshObjectUpdateDataCommand)(
-		[MeshObject, FrameNumberToPrepare, RevisionNumber, NewDynamicData, GPUSkinCache, Scene](FRHICommandListImmediate& RHICmdList)
+	ENQUEUE_RENDER_COMMAND(SkelMeshObjectUpdateDataCommand)(UE::RenderCommandPipe::SkeletalMesh,
+		[MeshObject, FrameNumberToPrepare, RevisionNumber, NewDynamicData, GPUSkinCache, Scene](FRHICommandList& RHICmdList)
 		{
 			FScopeCycleCounter Context(MeshObject->GetStatId());
 			MeshObject->UpdateDynamicData_RenderThread(GPUSkinCache, RHICmdList, NewDynamicData, Scene, FrameNumberToPrepare, RevisionNumber);
@@ -450,8 +450,8 @@ void FSkeletalMeshObjectGPUSkin::UpdateSkinWeightBuffer(USkinnedMeshComponent* I
 				FGPUSkinCacheEntry* SkinCacheEntryToUpdate = SkinCacheEntry;
 				if (SkinCacheEntryToUpdate)
 				{
-					ENQUEUE_RENDER_COMMAND(UpdateSkinCacheSkinWeightBuffer)(
-						[SkinCacheEntryToUpdate](FRHICommandListImmediate& RHICmdList)
+					ENQUEUE_RENDER_COMMAND(UpdateSkinCacheSkinWeightBuffer)(UE::RenderCommandPipe::SkeletalMesh,
+						[SkinCacheEntryToUpdate](FRHICommandList& RHICmdList)
 					{
 						FGPUSkinCache::UpdateSkinWeightBuffer(SkinCacheEntryToUpdate);
 					});
@@ -459,8 +459,8 @@ void FSkeletalMeshObjectGPUSkin::UpdateSkinWeightBuffer(USkinnedMeshComponent* I
 
 				if (SkinCacheEntryForRayTracing)
 				{
-					ENQUEUE_RENDER_COMMAND(UpdateSkinCacheSkinWeightBuffer)(
-						[SkinCacheEntryForRayTracing = SkinCacheEntryForRayTracing](FRHICommandListImmediate& RHICmdList)
+					ENQUEUE_RENDER_COMMAND(UpdateSkinCacheSkinWeightBuffer)(UE::RenderCommandPipe::SkeletalMesh,
+						[SkinCacheEntryForRayTracing = SkinCacheEntryForRayTracing](FRHICommandList& RHICmdList)
 					{
 						FGPUSkinCache::UpdateSkinWeightBuffer(SkinCacheEntryForRayTracing);
 					});
@@ -1352,8 +1352,8 @@ void FSkeletalMeshObjectGPUSkin::FSkeletalMeshObjectLOD::UpdateSkinWeights(FSkel
 			GetVertexBuffers(VertexBuffers, LODData);
 
 			FSkeletalMeshObjectLOD* Self = this;
-			ENQUEUE_RENDER_COMMAND(UpdateSkinWeightsGPUSkin)(
-				[NewMeshObjectWeightBuffer, VertexBuffers, Self](FRHICommandListImmediate& RHICmdList)
+			ENQUEUE_RENDER_COMMAND(UpdateSkinWeightsGPUSkin)(UE::RenderCommandPipe::SkeletalMesh,
+				[NewMeshObjectWeightBuffer, VertexBuffers, Self](FRHICommandList& RHICmdList)
 			{
 				Self->GPUSkinVertexFactories.UpdateVertexFactoryData(VertexBuffers);
 			});
@@ -1810,7 +1810,7 @@ static FGPUBaseSkinVertexFactory* CreateVertexFactory(TArray<TUniquePtr<FGPUBase
 	FDynamicUpdateVertexFactoryData VertexUpdateData(VertexFactory, InVertexBuffers);
 
 	// update vertex factory components and sync it
-	ENQUEUE_RENDER_COMMAND(InitGPUSkinVertexFactory)(
+	ENQUEUE_RENDER_COMMAND(InitGPUSkinVertexFactory)(UE::RenderCommandPipe::SkeletalMesh,
 		[VertexUpdateData](FRHICommandList& RHICmdList)
 		{
 			FGPUSkinDataType Data;
@@ -1836,7 +1836,7 @@ void UpdateVertexFactory(TArray<TUniquePtr<FGPUBaseSkinVertexFactory>>& VertexFa
 			FDynamicUpdateVertexFactoryData VertexUpdateData(VertexFactory, InVertexBuffers);
 
 			// update vertex factory components and sync it
-			ENQUEUE_RENDER_COMMAND(UpdateGPUSkinVertexFactory)(
+			ENQUEUE_RENDER_COMMAND(UpdateGPUSkinVertexFactory)(UE::RenderCommandPipe::SkeletalMesh,
 				[VertexUpdateData](FRHICommandList& RHICmdList)
 			{
 				FGPUSkinDataType Data;
@@ -1854,7 +1854,7 @@ static void CreatePassthroughVertexFactory(ERHIFeatureLevel::Type InFeatureLevel
 	PassthroughVertexFactories.Add(TUniquePtr<FGPUSkinPassthroughVertexFactory>(NewPassthroughVertexFactory));
 
 	// update vertex factory components and sync it
-	ENQUEUE_RENDER_COMMAND(InitPassthroughGPUSkinVertexFactory)(
+	ENQUEUE_RENDER_COMMAND(InitPassthroughGPUSkinVertexFactory)(UE::RenderCommandPipe::SkeletalMesh,
 		[NewPassthroughVertexFactory, SourceVertexFactory](FRHICommandList& RHICmdList)
 		{
 			FLocalVertexFactory::FDataType Data;
@@ -1889,7 +1889,7 @@ static void CreateVertexFactoryMorph(TArray<TUniquePtr<FGPUBaseSkinVertexFactory
 	FDynamicUpdateVertexFactoryData VertexUpdateData(VertexFactory, InVertexBuffers);
 
 	// update vertex factory components and sync it
-	ENQUEUE_RENDER_COMMAND(InitGPUSkinVertexFactoryMorph)(
+	ENQUEUE_RENDER_COMMAND(InitGPUSkinVertexFactoryMorph)(UE::RenderCommandPipe::SkeletalMesh,
 		[VertexUpdateData](FRHICommandList& RHICmdList)
 		{
 			FGPUSkinDataType Data;
@@ -1914,7 +1914,7 @@ static void UpdateVertexFactoryMorph(TArray<TUniquePtr<FGPUBaseSkinVertexFactory
 			FDynamicUpdateVertexFactoryData VertexUpdateData(VertexFactory, InVertexBuffers);
 
 			// update vertex factory components and sync it
-			ENQUEUE_RENDER_COMMAND(InitGPUSkinVertexFactoryMorph)(
+			ENQUEUE_RENDER_COMMAND(InitGPUSkinVertexFactoryMorph)(UE::RenderCommandPipe::SkeletalMesh,
 				[VertexUpdateData](FRHICommandList& RHICmdList)
 			{
 				FGPUSkinDataType Data;
@@ -1968,7 +1968,7 @@ static void CreateVertexFactoryCloth(TArray<TUniquePtr<FGPUBaseSkinAPEXClothVert
 	FDynamicUpdateVertexFactoryData VertexUpdateData(VertexFactory, InVertexBuffers);
 
 	// update vertex factory components and sync it
-	ENQUEUE_RENDER_COMMAND(InitGPUSkinAPEXClothVertexFactory)(
+	ENQUEUE_RENDER_COMMAND(InitGPUSkinAPEXClothVertexFactory)(UE::RenderCommandPipe::SkeletalMesh,
 		[VertexUpdateData](FRHICommandList& RHICmdList)
 		{
 			FGPUSkinAPEXClothDataType Data;
@@ -1993,7 +1993,7 @@ static void UpdateVertexFactoryCloth(TArray<TUniquePtr<FGPUBaseSkinAPEXClothVert
 			FDynamicUpdateVertexFactoryData VertexUpdateData(VertexFactory, InVertexBuffers);
 
 			// update vertex factory components and sync it
-			ENQUEUE_RENDER_COMMAND(InitGPUSkinAPEXClothVertexFactory)(
+			ENQUEUE_RENDER_COMMAND(InitGPUSkinAPEXClothVertexFactory)(UE::RenderCommandPipe::SkeletalMesh,
 				[VertexUpdateData](FRHICommandList& RHICmdList)
 			{
 				FGPUSkinAPEXClothDataType Data;
@@ -2112,12 +2112,12 @@ void FSkeletalMeshObjectGPUSkin::FVertexFactoryData::ReleaseVertexFactories()
 	// Default factories
 	for( int32 FactoryIdx=0; FactoryIdx < VertexFactories.Num(); FactoryIdx++)
 	{
-		BeginReleaseResource(VertexFactories[FactoryIdx].Get());
+		BeginReleaseResource(VertexFactories[FactoryIdx].Get(), &UE::RenderCommandPipe::SkeletalMesh);
 	}
 
 	for (int32 FactoryIdx = 0; FactoryIdx < PassthroughVertexFactories.Num(); FactoryIdx++)
 	{
-		BeginReleaseResource(PassthroughVertexFactories[FactoryIdx].Get());
+		BeginReleaseResource(PassthroughVertexFactories[FactoryIdx].Get(), &UE::RenderCommandPipe::SkeletalMesh);
 	}
 }
 
@@ -2143,7 +2143,7 @@ void FSkeletalMeshObjectGPUSkin::FVertexFactoryData::ReleaseMorphVertexFactories
 	// Default morph factories
 	for( int32 FactoryIdx=0; FactoryIdx < MorphVertexFactories.Num(); FactoryIdx++ )
 	{
-		BeginReleaseResource(MorphVertexFactories[FactoryIdx].Get());
+		BeginReleaseResource(MorphVertexFactories[FactoryIdx].Get(), &UE::RenderCommandPipe::SkeletalMesh);
 	}
 }
 
@@ -2188,7 +2188,7 @@ void FSkeletalMeshObjectGPUSkin::FVertexFactoryData::ReleaseAPEXClothVertexFacto
 		TUniquePtr<FGPUBaseSkinAPEXClothVertexFactory>& ClothVertexFactory = ClothVertexFactories[FactoryIdx];
 		if (ClothVertexFactory)
 		{
-			BeginReleaseResource(ClothVertexFactory->GetVertexFactory());
+			BeginReleaseResource(ClothVertexFactory->GetVertexFactory(), &UE::RenderCommandPipe::SkeletalMesh);
 		}
 	}
 }
diff --git a/Engine/Source/Runtime/Engine/Private/SkeletalRenderStatic.cpp b/Engine/Source/Runtime/Engine/Private/SkeletalRenderStatic.cpp
--- a/Engine/Source/Runtime/Engine/Private/SkeletalRenderStatic.cpp
+++ b/Engine/Source/Runtime/Engine/Private/SkeletalRenderStatic.cpp
@@ -7,7 +7,7 @@
 #include "SkeletalRenderStatic.h"
 #include "RenderUtils.h"
 #include "Rendering/SkeletalMeshRenderData.h"
-#include "RenderingThread.h"
+#include "Rendering/RenderCommandPipes.h"
 
 #if RHI_RAYTRACING
 #include "Engine/SkinnedAssetCommon.h"
@@ -67,7 +67,7 @@ void FSkeletalMeshObjectStatic::InitResources(USkinnedMeshComponent* InMeshCompo
 					}
 
 					TArray<FSkelMeshRenderSection>* RenderSections = &LODModel.RenderSections;
-					ENQUEUE_RENDER_COMMAND(InitSkeletalRenderStaticRayTracingGeometry)(
+					ENQUEUE_RENDER_COMMAND(InitSkeletalRenderStaticRayTracingGeometry)(UE::RenderCommandPipe::SkeletalMesh,
 						[this, VertexBufferRHI, IndexBufferRHI, VertexBufferStride, TrianglesCount, RenderSections, 
 						LODIndex = LODIndex, 
 						SkelMeshRenderData = SkelLOD.SkelMeshRenderData, 
@@ -144,14 +144,13 @@ void FSkeletalMeshObjectStatic::ReleaseResources()
 
 					if (SkelLOD.SkelMeshRenderData->LODRenderData[LODIndex].NumReferencingStaticSkeletalMeshObjects == 0)
 					{
-						ENQUEUE_RENDER_COMMAND(ResetStaticRayTracingGeometryFlag)(
-							[&bReferencedByStaticSkeletalMeshObjects_RenderThread = SkelLOD.SkelMeshRenderData->LODRenderData[LODIndex].bReferencedByStaticSkeletalMeshObjects_RenderThread](FRHICommandListImmediate& RHICmdList)
+						ENQUEUE_RENDER_COMMAND(ResetStaticRayTracingGeometryFlag)(UE::RenderCommandPipe::SkeletalMesh,
+							[&bReferencedByStaticSkeletalMeshObjects_RenderThread = SkelLOD.SkelMeshRenderData->LODRenderData[LODIndex].bReferencedByStaticSkeletalMeshObjects_RenderThread]
 						{
 							bReferencedByStaticSkeletalMeshObjects_RenderThread = false;
-						}
-						);
+						});
 
-						BeginReleaseResource(&SkelLOD.SkelMeshRenderData->LODRenderData[LODIndex].StaticRayTracingGeometry);
+						BeginReleaseResource(&SkelLOD.SkelMeshRenderData->LODRenderData[LODIndex].StaticRayTracingGeometry, &UE::RenderCommandPipe::SkeletalMesh);
 					}
 				}
 			}
@@ -204,7 +203,7 @@ void FSkeletalMeshObjectStatic::FSkeletalMeshObjectLOD::InitResources(FSkelMeshC
 	FLocalVertexFactory* VertexFactoryPtr = &VertexFactory;
 	FColorVertexBuffer* ColorVertexBufferPtr = ColorVertexBuffer;
 
-	ENQUEUE_RENDER_COMMAND(InitSkeletalMeshStaticSkinVertexFactory)(
+	ENQUEUE_RENDER_COMMAND(InitSkeletalMeshStaticSkinVertexFactory)(UE::RenderCommandPipe::SkeletalMesh,
 		[VertexFactoryPtr, PositionVertexBufferPtr, StaticMeshVertexBufferPtr, ColorVertexBufferPtr](FRHICommandList& RHICmdList)
 		{
 			FLocalVertexFactory::FDataType Data;
@@ -230,10 +229,10 @@ void FSkeletalMeshObjectStatic::FSkeletalMeshObjectLOD::InitResources(FSkelMeshC
  */
 void FSkeletalMeshObjectStatic::FSkeletalMeshObjectLOD::ReleaseResources()
 {	
-	BeginReleaseResource(&VertexFactory);
+	BeginReleaseResource(&VertexFactory, &UE::RenderCommandPipe::SkeletalMesh);
 
 #if RHI_RAYTRACING
-	BeginReleaseResource(&RayTracingGeometry);
+	BeginReleaseResource(&RayTracingGeometry, &UE::RenderCommandPipe::SkeletalMesh);
 #endif // RHI_RAYTRACING
 
 	bResourcesInitialized = false;
diff --git a/Engine/Source/Runtime/Engine/Public/Rendering/RenderCommandPipes.h b/Engine/Source/Runtime/Engine/Public/Rendering/RenderCommandPipes.h
new file mode 100644
--- /dev/null
+++ b/Engine/Source/Runtime/Engine/Public/Rendering/RenderCommandPipes.h
@@ -0,0 +1,7 @@
+// Copyright Epic Games, Inc. All Rights Reserved.
+
+#pragma once
+
+#include "RenderingThread.h"
+
+DECLARE_RENDER_COMMAND_PIPE(SkeletalMesh, ENGINE_API);
\ No newline at end of file
