Fixed errors mirroring curves with multiple entries in a mirror data table.

diff --git a/Engine/Source/Runtime/Engine/Private/Animation/MirrorDataTable.cpp b/Engine/Source/Runtime/Engine/Private/Animation/MirrorDataTable.cpp
--- a/Engine/Source/Runtime/Engine/Private/Animation/MirrorDataTable.cpp
+++ b/Engine/Source/Runtime/Engine/Private/Animation/MirrorDataTable.cpp
@@ -462,7 +462,12 @@ void UMirrorDataTable::FillMirrorArrays()
 			{
 			case  EMirrorRowType::Curve:
 			{
-				CurveToMirrorCurveMap.Add(Value.Name, Value.MirroredName);
+				// curves swap, so only one entry should exist.  For instance, if the table has (Left, Right) and (Right, Left) only add one item
+				const FName* TestMirroredMatch = CurveToMirrorCurveMap.Find(Value.MirroredName);
+				if (TestMirroredMatch == nullptr || *TestMirroredMatch != Value.Name)
+				{
+					CurveToMirrorCurveMap.Add(Value.Name, Value.MirroredName);
+				}
 				break;
 			}
 			case EMirrorRowType::SyncMarker:
