[Backout] - CL27048554, 27048589, 27048615, 27049788, 27049950, 27050314, 27050334
First Original CL Desc
-----------------------------------------------------------------
Moved volumetric volume to RHICmdList.

diff --git a/Engine/Plugins/FX/Niagara/Source/Niagara/Private/NiagaraComponent.cpp b/Engine/Plugins/FX/Niagara/Source/Niagara/Private/NiagaraComponent.cpp
--- a/Engine/Plugins/FX/Niagara/Source/Niagara/Private/NiagaraComponent.cpp
+++ b/Engine/Plugins/FX/Niagara/Source/Niagara/Private/NiagaraComponent.cpp
@@ -43,8 +43,6 @@ DECLARE_CYCLE_STAT(TEXT("Get Dynamic Mesh Elements (RT)"), STAT_NiagaraComponent
 
 DEFINE_LOG_CATEGORY(LogNiagara);
 
-DEFINE_RENDER_COMMAND_PIPE(NiagaraDynamicData);
-
 static int GNiagaraSoloTickEarly = 1;
 static FAutoConsoleVariableRef CVarNiagaraSoloTickEarly(
 	TEXT("fx.Niagara.Solo.TickEarly"),
@@ -2141,8 +2139,8 @@ void UNiagaraComponent::SendRenderDynamicData_Concurrent()
 				SystemInstanceController->GenerateSetDynamicDataCommands(SetDataCommands, *RenderData, *NiagaraProxy);
 			}
 
-			ENQUEUE_RENDER_COMMAND(NiagaraSetDynamicData)(UE::RenderCommandPipe::NiagaraDynamicData,
-				[NiagaraProxy, CommandsRT = MoveTemp(SetDataCommands), NewProxyDynamicData]
+			ENQUEUE_RENDER_COMMAND(NiagaraSetDynamicData)(
+				[NiagaraProxy, CommandsRT = MoveTemp(SetDataCommands), NewProxyDynamicData](FRHICommandListImmediate& RHICmdList)
 				{
 					SCOPE_CYCLE_COUNTER(STAT_NiagaraSetDynamicData);
 					PARTICLE_PERF_STAT_CYCLES_WITH_COUNT_RT(NewProxyDynamicData.PerfStatsContext, RenderUpdate, 1);
diff --git a/Engine/Plugins/MeshPainting/Source/MeshPaintingToolset/Private/MeshPaintSkeletalMeshAdapter.cpp b/Engine/Plugins/MeshPainting/Source/MeshPaintingToolset/Private/MeshPaintSkeletalMeshAdapter.cpp
--- a/Engine/Plugins/MeshPainting/Source/MeshPaintingToolset/Private/MeshPaintSkeletalMeshAdapter.cpp
+++ b/Engine/Plugins/MeshPainting/Source/MeshPaintingToolset/Private/MeshPaintSkeletalMeshAdapter.cpp
@@ -11,7 +11,7 @@
 #include "InterchangeAssetImportData.h"
 #include "InterchangeGenericAssetsPipeline.h"
 #include "InterchangePythonPipelineBase.h"
-#include "Rendering/RenderCommandPipes.h"
+#include "RenderingThread.h"
 
 DEFINE_LOG_CATEGORY_STATIC(LogMeshPaintSkeletalMeshAdapter, Log, All);
 //////////////////////////////////////////////////////////////////////////
@@ -382,7 +382,7 @@ void FMeshPaintSkeletalMeshComponentAdapter::PreEdit()
 		LODData->StaticVertexBuffers.ColorVertexBuffer.InitFromSingleColor(FColor(255, 255, 255, 255), LODData->GetNumVertices());
 		ReferencedSkeletalMesh->SetHasVertexColors(true);
 		ReferencedSkeletalMesh->SetVertexColorGuid(FGuid::NewGuid());
-		BeginInitResource(&LODData->StaticVertexBuffers.ColorVertexBuffer, &UE::RenderCommandPipe::SkeletalMesh);
+		BeginInitResource(&LODData->StaticVertexBuffers.ColorVertexBuffer);
 	}
 	//Make sure we change the import data so the re-import do not replace the new data
 	if (ReferencedSkeletalMesh->GetAssetImportData())
diff --git a/Engine/Plugins/Runtime/CableComponent/Source/CableComponent/Private/CableComponent.cpp b/Engine/Plugins/Runtime/CableComponent/Source/CableComponent/Private/CableComponent.cpp
--- a/Engine/Plugins/Runtime/CableComponent/Source/CableComponent/Private/CableComponent.cpp
+++ b/Engine/Plugins/Runtime/CableComponent/Source/CableComponent/Private/CableComponent.cpp
@@ -14,7 +14,6 @@
 #include "StaticMeshResources.h"
 #include "SceneInterface.h"
 
-DEFINE_RENDER_COMMAND_PIPE(Cable);
 
 static TAutoConsoleVariable<int32> CVarRayTracingCableMeshes(
 	TEXT("r.RayTracing.Geometry.Cable"),
@@ -93,12 +92,12 @@ public:
 		, NumSides(Component->NumSides)
 		, TileMaterial(Component->TileMaterial)
 	{
-		VertexBuffers.InitWithDummyData(&UE::RenderCommandPipe::Cable, &VertexFactory, GetRequiredVertexCount());
+		VertexBuffers.InitWithDummyData(&VertexFactory, GetRequiredVertexCount());
 
 		IndexBuffer.NumIndices = GetRequiredIndexCount();
 
 		// Enqueue initialization of render resource
-		BeginInitResource(&IndexBuffer, &UE::RenderCommandPipe::Cable);
+		BeginInitResource(&IndexBuffer);
 
 		// Grab material
 		Material = Component->GetMaterial(0);
@@ -1044,8 +1043,8 @@ void UCableComponent::SendRenderDynamicData_Concurrent()
 
 		// Enqueue command to send to render thread
 		FCableSceneProxy* CableSceneProxy = (FCableSceneProxy*)SceneProxy;
-		ENQUEUE_RENDER_COMMAND(FSendCableDynamicData)(UE::RenderCommandPipe::Cable,
-			[CableSceneProxy, DynamicData](FRHICommandListBase& RHICmdList)
+		ENQUEUE_RENDER_COMMAND(FSendCableDynamicData)(
+			[CableSceneProxy, DynamicData](FRHICommandListImmediate& RHICmdList)
 		{
 			CableSceneProxy->SetDynamicData_RenderThread(RHICmdList, DynamicData);
 		});
diff --git a/Engine/Source/Editor/MeshPaint/Private/MeshPaintSkeletalMeshAdapter.cpp b/Engine/Source/Editor/MeshPaint/Private/MeshPaintSkeletalMeshAdapter.cpp
--- a/Engine/Source/Editor/MeshPaint/Private/MeshPaintSkeletalMeshAdapter.cpp
+++ b/Engine/Source/Editor/MeshPaint/Private/MeshPaintSkeletalMeshAdapter.cpp
@@ -10,7 +10,6 @@
 #include "Rendering/SkeletalMeshRenderData.h"
 #include "Rendering/SkeletalMeshModel.h"
 #include "Factories/FbxSkeletalMeshImportData.h"
-#include "Rendering/RenderCommandPipes.h"
 
 DEFINE_LOG_CATEGORY_STATIC(LogMeshPaintSkeletalMeshAdapter, Log, All);
 
@@ -456,7 +455,7 @@ void FMeshPaintGeometryAdapterForSkeletalMeshes::PreEdit()
 		LODData->StaticVertexBuffers.ColorVertexBuffer.InitFromSingleColor(FColor(255, 255, 255, 255), LODData->GetNumVertices());
 		ReferencedSkeletalMesh->SetHasVertexColors(true);
 		ReferencedSkeletalMesh->SetVertexColorGuid(FGuid::NewGuid());
-		BeginInitResource(&LODData->StaticVertexBuffers.ColorVertexBuffer, &UE::RenderCommandPipe::SkeletalMesh);
+		BeginInitResource(&LODData->StaticVertexBuffers.ColorVertexBuffer);
 	}
 	//Make sure we change the import data so the re-import do not replace the new data
 	if (ReferencedSkeletalMesh->GetAssetImportData())
diff --git a/Engine/Source/Editor/VirtualTexturingEditor/Private/RuntimeVirtualTextureThumbnailRenderer.cpp b/Engine/Source/Editor/VirtualTexturingEditor/Private/RuntimeVirtualTextureThumbnailRenderer.cpp
--- a/Engine/Source/Editor/VirtualTexturingEditor/Private/RuntimeVirtualTextureThumbnailRenderer.cpp
+++ b/Engine/Source/Editor/VirtualTexturingEditor/Private/RuntimeVirtualTextureThumbnailRenderer.cpp
@@ -95,8 +95,6 @@ void URuntimeVirtualTextureThumbnailRenderer::Draw(UObject* Object, int32 X, int
 	RuntimeVirtualTexture->GetProducerDescription(VTDesc, URuntimeVirtualTexture::FInitSettings(), Transform);
 	const int32 MaxLevel = (int32)FMath::CeilLogTwo(FMath::Max(VTDesc.BlockWidthInTiles, VTDesc.BlockHeightInTiles));
 
-	UE::RenderCommandPipe::FSyncScope SyncScope;
-
 	ENQUEUE_RENDER_COMMAND(BakeStreamingTextureTileCommand)(
 		[Scene, VirtualTextureSceneIndex, MaterialType, RenderTarget, DestBox, Transform, Bounds, MaxLevel](FRHICommandListImmediate& RHICmdList)
 	{
diff --git a/Engine/Source/Runtime/Engine/Private/Components/SkeletalMeshComponent.cpp b/Engine/Source/Runtime/Engine/Private/Components/SkeletalMeshComponent.cpp
--- a/Engine/Source/Runtime/Engine/Private/Components/SkeletalMeshComponent.cpp
+++ b/Engine/Source/Runtime/Engine/Private/Components/SkeletalMeshComponent.cpp
@@ -25,7 +25,6 @@
 #include "SkeletalRender.h"
 #include "SkinnedAssetCompiler.h"
 #include "Rendering/SkeletalMeshRenderData.h"
-#include "Rendering/RenderCommandPipes.h"
 
 #include "Logging/MessageLog.h"
 #include "Animation/AnimNode_LinkedInputPose.h"
@@ -4059,8 +4058,8 @@ void USkeletalMeshComponent::SendRenderDynamicData_Concurrent()
 
 		FSkeletalMeshSceneProxy* TargetProxy = (FSkeletalMeshSceneProxy*)SceneProxy;
 
-		ENQUEUE_RENDER_COMMAND(PoseWatchDynamicDataCommand)(UE::RenderCommandPipe::SkeletalMesh,
-			[TargetProxy, NewDynamicData]
+		ENQUEUE_RENDER_COMMAND(PoseWatchDynamicDataCommand)(
+			[TargetProxy, NewDynamicData](FRHICommandListImmediate& RHICommandList)
 			{
 				if (TargetProxy->PoseWatchDynamicData)
 				{
diff --git a/Engine/Source/Runtime/Engine/Private/GPUSkinVertexFactory.cpp b/Engine/Source/Runtime/Engine/Private/GPUSkinVertexFactory.cpp
--- a/Engine/Source/Runtime/Engine/Private/GPUSkinVertexFactory.cpp
+++ b/Engine/Source/Runtime/Engine/Private/GPUSkinVertexFactory.cpp
@@ -242,7 +242,7 @@ static TAutoConsoleVariable<int32> CVarRHICmdDeferSkeletalLockAndFillToRHIThread
 
 static bool DeferSkeletalLockAndFillToRHIThread()
 {
-	return IsRunningRHIInSeparateThread() && CVarRHICmdDeferSkeletalLockAndFillToRHIThread.GetValueOnRenderThread() > 0 && !UE::RenderCommandPipe::IsReplaying();
+	return IsRunningRHIInSeparateThread() && CVarRHICmdDeferSkeletalLockAndFillToRHIThread.GetValueOnRenderThread() > 0;
 }
 
 bool FGPUBaseSkinVertexFactory::FShaderDataType::UpdateBoneData(FRHICommandList& RHICmdList, const TArray<FMatrix44f>& ReferenceToLocalMatrices,
diff --git a/Engine/Source/Runtime/Engine/Private/LevelTick.cpp b/Engine/Source/Runtime/Engine/Private/LevelTick.cpp
--- a/Engine/Source/Runtime/Engine/Private/LevelTick.cpp
+++ b/Engine/Source/Runtime/Engine/Private/LevelTick.cpp
@@ -24,7 +24,6 @@
 #include "Camera/CameraPhotography.h"
 #include "UObject/Stack.h"
 #include "PhysicsEngine/CollisionAnalyzerCapture.h"
-#include "Rendering/RenderCommandPipes.h"
 
 #if !UE_SERVER
 #include "IMediaModule.h"
@@ -974,23 +973,23 @@ struct FSendAllEndOfFrameUpdates
 void BeginSendEndOfFrameUpdatesDrawEvent(FSendAllEndOfFrameUpdates& SendAllEndOfFrameUpdates)
 {
 	BEGIN_DRAW_EVENTF_GAMETHREAD(SendAllEndOfFrameUpdates, SendAllEndOfFrameUpdates.DrawEvent, TEXT("SendAllEndOfFrameUpdates"));
-
-	ENQUEUE_RENDER_COMMAND(BeginDrawEventCommand)(UE::RenderCommandPipe::SkeletalMesh,
-		[GPUSkinCache = SendAllEndOfFrameUpdates.GPUSkinCache]
-	{
-		if (GPUSkinCache != nullptr)
+	
+	ENQUEUE_RENDER_COMMAND(BeginDrawEventCommand)(
+		[GPUSkinCache = SendAllEndOfFrameUpdates.GPUSkinCache](FRHICommandListImmediate&)
 		{
-			GPUSkinCache->BeginBatchDispatch();
-		}
-	});
+			if (GPUSkinCache != nullptr)
+			{
+				GPUSkinCache->BeginBatchDispatch();
+			}
+		});
 }
 
 DECLARE_GPU_STAT(EndOfFrameUpdates);
 DECLARE_GPU_STAT(GPUSkinCacheRayTracingGeometry);
 void EndSendEndOfFrameUpdatesDrawEvent(FSendAllEndOfFrameUpdates& SendAllEndOfFrameUpdates)
 {
-	ENQUEUE_RENDER_COMMAND(EndDrawEventCommand)(UE::RenderCommandPipe::SkeletalMesh,
-		[GPUSkinCache = SendAllEndOfFrameUpdates.GPUSkinCache]
+	ENQUEUE_RENDER_COMMAND(EndDrawEventCommand)(
+		[GPUSkinCache = SendAllEndOfFrameUpdates.GPUSkinCache] (FRHICommandListImmediate&)
 	{
 		if (GPUSkinCache != nullptr)
 		{
diff --git a/Engine/Source/Runtime/Engine/Private/PrecomputedVolumetricLightmap.cpp b/Engine/Source/Runtime/Engine/Private/PrecomputedVolumetricLightmap.cpp
--- a/Engine/Source/Runtime/Engine/Private/PrecomputedVolumetricLightmap.cpp
+++ b/Engine/Source/Runtime/Engine/Private/PrecomputedVolumetricLightmap.cpp
@@ -699,9 +699,9 @@ void FPrecomputedVolumetricLightmap::AddToScene(FSceneInterface* Scene, UMapBuil
 		FPrecomputedVolumetricLightmap* Volume = this;
 
 		ENQUEUE_RENDER_COMMAND(SetVolumeDataCommand)
-			([Volume, NewData, Scene] (FRHICommandListBase& RHICmdList)
+			([Volume, NewData, Scene](FRHICommandListImmediate& RHICmdList) 
 			{
-				Volume->SetData(RHICmdList, NewData, Scene);
+				Volume->SetData(NewData, Scene);
 			});
 		Scene->AddPrecomputedVolumetricLightmap(this, bIsPersistentLevel);
 	}
@@ -741,7 +741,7 @@ void FPrecomputedVolumetricLightmap::RemoveFromScene(FSceneInterface* Scene)
 	WorldOriginOffset = FVector::ZeroVector;
 }
 
-void FPrecomputedVolumetricLightmap::SetData(FRHICommandListBase& RHICmdList, FPrecomputedVolumetricLightmapData* NewData, FSceneInterface* Scene)
+void FPrecomputedVolumetricLightmap::SetData(FPrecomputedVolumetricLightmapData* NewData, FSceneInterface* Scene)
 {
 	Data = NewData;
 
@@ -753,7 +753,7 @@ void FPrecomputedVolumetricLightmap::SetData(FRHICommandListBase& RHICmdList, FP
 
 		if (Scene->GetFeatureLevel() >= ERHIFeatureLevel::SM5)
 		{
-			Data->InitResource(RHICmdList);
+			Data->InitResource(FRHICommandListImmediate::Get());
 		}
 	}
 }
diff --git a/Engine/Source/Runtime/Engine/Private/Rendering/MultiSizeIndexContainer.cpp b/Engine/Source/Runtime/Engine/Private/Rendering/MultiSizeIndexContainer.cpp
--- a/Engine/Source/Runtime/Engine/Private/Rendering/MultiSizeIndexContainer.cpp
+++ b/Engine/Source/Runtime/Engine/Private/Rendering/MultiSizeIndexContainer.cpp
@@ -1,10 +1,10 @@
 // Copyright Epic Games, Inc. All Rights Reserved.
 
 #include "Rendering/MultiSizeIndexContainer.h"
-#include "Rendering/RenderCommandPipes.h"
 #include "EngineLogs.h"
 #include "RawIndexBuffer.h"
 #include "Stats/Stats.h"
+#include "RenderingThread.h"
 
 FMultiSizeIndexContainer::~FMultiSizeIndexContainer()
 {
@@ -31,7 +31,7 @@ void FMultiSizeIndexContainer::InitResources()
 	check(IsInGameThread());
 	if (IndexBuffer)
 	{
-		BeginInitResource(IndexBuffer, &UE::RenderCommandPipe::SkeletalMesh);
+		BeginInitResource(IndexBuffer);
 	}
 }
 
@@ -43,7 +43,7 @@ void FMultiSizeIndexContainer::ReleaseResources()
 	check(IsInGameThread());
 	if (IndexBuffer)
 	{
-		BeginReleaseResource(IndexBuffer, &UE::RenderCommandPipe::SkeletalMesh);
+		BeginReleaseResource(IndexBuffer);
 	}
 }
 
diff --git a/Engine/Source/Runtime/Engine/Private/Rendering/RenderCommandPipes.cpp b/Engine/Source/Runtime/Engine/Private/Rendering/RenderCommandPipes.cpp
deleted file mode 100644
diff --git a/Engine/Source/Runtime/Engine/Private/Rendering/SkinWeightVertexBuffer.cpp b/Engine/Source/Runtime/Engine/Private/Rendering/SkinWeightVertexBuffer.cpp
--- a/Engine/Source/Runtime/Engine/Private/Rendering/SkinWeightVertexBuffer.cpp
+++ b/Engine/Source/Runtime/Engine/Private/Rendering/SkinWeightVertexBuffer.cpp
@@ -8,7 +8,6 @@
 #include "RHIResourceUpdates.h"
 #include "SkeletalMeshLegacyCustomVersions.h"
 #include "UObject/UE5MainStreamObjectVersion.h"
-#include "Rendering/RenderCommandPipes.h"
 
 /*-----------------------------------------------------------------------------
 FSkinWeightLookupVertexBuffer
@@ -807,14 +806,14 @@ void FSkinWeightVertexBuffer::SetOwnerName(const FName& OwnerName)
 
 void FSkinWeightVertexBuffer::BeginInitResources()
 {
-	BeginInitResource(&LookupVertexBuffer, &UE::RenderCommandPipe::SkeletalMesh);
-	BeginInitResource(&DataVertexBuffer, &UE::RenderCommandPipe::SkeletalMesh);
+	BeginInitResource(&LookupVertexBuffer);
+	BeginInitResource(&DataVertexBuffer);
 }
 
 void FSkinWeightVertexBuffer::BeginReleaseResources()
 {
-	BeginReleaseResource(&LookupVertexBuffer, &UE::RenderCommandPipe::SkeletalMesh);
-	BeginReleaseResource(&DataVertexBuffer, &UE::RenderCommandPipe::SkeletalMesh);
+	BeginReleaseResource(&LookupVertexBuffer);
+	BeginReleaseResource(&DataVertexBuffer);
 }
 
 void FSkinWeightVertexBuffer::ReleaseResources()
diff --git a/Engine/Source/Runtime/Engine/Private/SkeletalMesh.cpp b/Engine/Source/Runtime/Engine/Private/SkeletalMesh.cpp
--- a/Engine/Source/Runtime/Engine/Private/SkeletalMesh.cpp
+++ b/Engine/Source/Runtime/Engine/Private/SkeletalMesh.cpp
@@ -44,7 +44,6 @@
 #include "Engine/AssetUserData.h"
 #include "Animation/NodeMappingContainer.h"
 #include "Rendering/SkeletalMeshRenderData.h"
-#include "Rendering/RenderCommandPipes.h"
 #include "AnimationRuntime.h"
 #include "Animation/AnimSequence.h"
 #include "Streaming/SkeletalMeshUpdate.h"
@@ -6808,8 +6807,8 @@ void FSkeletalMeshSceneProxy::UpdateMorphMaterialUsage_GameThread(TArray<UMateri
 			ERHIFeatureLevel::Type InFeatureLevel = GetScene().GetFeatureLevel();
 			FSkeletalMeshSceneProxy* SkelMeshSceneProxy = this;
 			FMaterialRelevance DefaultRelevance = DefaultMaterial->GetRelevance(InFeatureLevel);
-			ENQUEUE_RENDER_COMMAND(UpdateSkelProxyLODSectionElementsCmd)(UE::RenderCommandPipe::SkeletalMesh,
-				[InMaterialsToSwap, DefaultMaterial, DefaultRelevance, InFeatureLevel, SkelMeshSceneProxy]
+			ENQUEUE_RENDER_COMMAND(UpdateSkelProxyLODSectionElementsCmd)(
+				[InMaterialsToSwap, DefaultMaterial, DefaultRelevance, InFeatureLevel, SkelMeshSceneProxy](FRHICommandList& RHICmdList)
 				{
 					for( int32 LodIdx=0; LodIdx < SkelMeshSceneProxy->LODSections.Num(); LodIdx++ )
 					{
diff --git a/Engine/Source/Runtime/Engine/Private/SkeletalMeshComponentPhysics.cpp b/Engine/Source/Runtime/Engine/Private/SkeletalMeshComponentPhysics.cpp
--- a/Engine/Source/Runtime/Engine/Private/SkeletalMeshComponentPhysics.cpp
+++ b/Engine/Source/Runtime/Engine/Private/SkeletalMeshComponentPhysics.cpp
@@ -14,7 +14,6 @@
 #include "AnimationRuntime.h"
 #include "ClothCollisionData.h"
 #include "ClothingSimulationInteractor.h"
-#include "Rendering/RenderCommandPipes.h"
 
 #include "Logging/MessageLog.h"
 #include "CollisionDebugDrawingPublic.h"
@@ -1660,8 +1659,8 @@ void USkeletalMeshComponent::SendRenderDebugPhysics(FPrimitiveSceneProxy* Overri
 			}
 		}
 
-		ENQUEUE_RENDER_COMMAND(SkeletalMesh_SendRenderDebugPhysics)(UE::RenderCommandPipe::SkeletalMesh,
-			[UseSceneProxy, DebugMassData]
+		ENQUEUE_RENDER_COMMAND(SkeletalMesh_SendRenderDebugPhysics)(
+			[UseSceneProxy, DebugMassData](FRHICommandListImmediate& RHICmdList)
 			{
 				UseSceneProxy->SetDebugMassData(DebugMassData);
 			}
diff --git a/Engine/Source/Runtime/Engine/Private/SkeletalMeshLODRenderData.cpp b/Engine/Source/Runtime/Engine/Private/SkeletalMeshLODRenderData.cpp
--- a/Engine/Source/Runtime/Engine/Private/SkeletalMeshLODRenderData.cpp
+++ b/Engine/Source/Runtime/Engine/Private/SkeletalMeshLODRenderData.cpp
@@ -12,7 +12,6 @@
 #include "Interfaces/ITargetPlatform.h"
 #include "Interfaces/ITargetPlatformManagerModule.h"
 #include "Rendering/SkeletalMeshLODModel.h"
-#include "Rendering/RenderCommandPipes.h"
 #include "Serialization/MemoryReader.h"
 #include "SkeletalMeshLegacyCustomVersions.h"
 #include "UObject/Package.h"
@@ -207,7 +206,7 @@ void FSkeletalMeshLODRenderData::InitMorphResources()
 {
 	if (GAllowSkinnedMorphDataStreaming && !MorphTargetVertexInfoBuffers.IsRHIIntialized() && MorphTargetVertexInfoBuffers.IsMorphCPUDataValid() && MorphTargetVertexInfoBuffers.NumTotalBatches > 0)
 	{
-		BeginInitResource(&MorphTargetVertexInfoBuffers, &UE::RenderCommandPipe::SkeletalMesh);
+		BeginInitResource(&MorphTargetVertexInfoBuffers);
 	}
 }
 
@@ -224,9 +223,9 @@ void FSkeletalMeshLODRenderData::InitResources(bool bNeedsVertexColors, int32 LO
 	MultiSizeIndexContainer.InitResources();
 
 	StaticVertexBuffers.PositionVertexBuffer.SetOwnerName(OwnerName);
-	BeginInitResource(&StaticVertexBuffers.PositionVertexBuffer, &UE::RenderCommandPipe::SkeletalMesh);
+	BeginInitResource(&StaticVertexBuffers.PositionVertexBuffer);
 	StaticVertexBuffers.StaticMeshVertexBuffer.SetOwnerName(OwnerName);
-	BeginInitResource(&StaticVertexBuffers.StaticMeshVertexBuffer, &UE::RenderCommandPipe::SkeletalMesh);
+	BeginInitResource(&StaticVertexBuffers.StaticMeshVertexBuffer);
 
 	SkinWeightVertexBuffer.SetOwnerName(OwnerName);
 	SkinWeightVertexBuffer.BeginInitResources();
@@ -235,14 +234,14 @@ void FSkeletalMeshLODRenderData::InitResources(bool bNeedsVertexColors, int32 LO
 	{
 		// Only init the color buffer if the mesh has vertex colors
 		StaticVertexBuffers.ColorVertexBuffer.SetOwnerName(OwnerName);
-		BeginInitResource(&StaticVertexBuffers.ColorVertexBuffer, &UE::RenderCommandPipe::SkeletalMesh);
+		BeginInitResource(&StaticVertexBuffers.ColorVertexBuffer);
 	}
 
 	if (ClothVertexBuffer.GetNumVertices() > 0)
 	{
 		// Only init the clothing buffer if the mesh has clothing data
 		ClothVertexBuffer.SetOwnerName(OwnerName);
-		BeginInitResource(&ClothVertexBuffer, &UE::RenderCommandPipe::SkeletalMesh);
+		BeginInitResource(&ClothVertexBuffer);
 	}
 
 	// We can discard any the cooked DuplicatedVertices here based on runtime settings.
@@ -263,7 +262,7 @@ void FSkeletalMeshLODRenderData::InitResources(bool bNeedsVertexColors, int32 LO
 				// No need to discard CPU data in cooked builds as bNeedsCPUAccess is false (see FDuplicatedVerticesBuffer constructor), 
 				// so it'd be auto-discarded after the RHI has copied the resource data. Keep CPU data when in the editor for geometry operations.
 				RenderSection.DuplicatedVerticesBuffer.SetOwnerName(OwnerName);
-				BeginInitResource(&RenderSection.DuplicatedVerticesBuffer, &UE::RenderCommandPipe::SkeletalMesh);
+				BeginInitResource(&RenderSection.DuplicatedVerticesBuffer);
 			}
 		}
 	}
@@ -279,7 +278,7 @@ void FSkeletalMeshLODRenderData::InitResources(bool bNeedsVertexColors, int32 LO
 	if (!MorphTargetVertexInfoBuffers.IsRHIIntialized() && MorphTargetVertexInfoBuffers.IsMorphCPUDataValid() && MorphTargetVertexInfoBuffers.NumTotalBatches > 0)
 	{
 		MorphTargetVertexInfoBuffers.SetOwnerName(OwnerName);
-		BeginInitResource(&MorphTargetVertexInfoBuffers, &UE::RenderCommandPipe::SkeletalMesh);
+		BeginInitResource(&MorphTargetVertexInfoBuffers);
 	}
 
 	VertexAttributeBuffers.InitResources();
@@ -290,7 +289,7 @@ void FSkeletalMeshLODRenderData::InitResources(bool bNeedsVertexColors, int32 LO
 		if (SourceRayTracingGeometry.RawData.Num() > 0)
 		{
 			SourceRayTracingGeometry.SetOwnerName(OwnerName);
-			BeginInitResource(&SourceRayTracingGeometry, &UE::RenderCommandPipe::SkeletalMesh);
+			BeginInitResource(&SourceRayTracingGeometry);
 		}
 	}
 #endif
@@ -302,16 +301,16 @@ void FSkeletalMeshLODRenderData::ReleaseResources()
 
 	MultiSizeIndexContainer.ReleaseResources();
 
-	BeginReleaseResource(&StaticVertexBuffers.PositionVertexBuffer, &UE::RenderCommandPipe::SkeletalMesh);
-	BeginReleaseResource(&StaticVertexBuffers.StaticMeshVertexBuffer, &UE::RenderCommandPipe::SkeletalMesh);
+	BeginReleaseResource(&StaticVertexBuffers.PositionVertexBuffer);
+	BeginReleaseResource(&StaticVertexBuffers.StaticMeshVertexBuffer);
 	SkinWeightVertexBuffer.BeginReleaseResources();
-	BeginReleaseResource(&StaticVertexBuffers.ColorVertexBuffer, &UE::RenderCommandPipe::SkeletalMesh);
-	BeginReleaseResource(&ClothVertexBuffer, &UE::RenderCommandPipe::SkeletalMesh);
+	BeginReleaseResource(&StaticVertexBuffers.ColorVertexBuffer);
+	BeginReleaseResource(&ClothVertexBuffer);
 	for (FSkelMeshRenderSection& RenderSection : RenderSections)
 	{
-		BeginReleaseResource(&RenderSection.DuplicatedVerticesBuffer, &UE::RenderCommandPipe::SkeletalMesh);
+		BeginReleaseResource(&RenderSection.DuplicatedVerticesBuffer);
 	}
-	BeginReleaseResource(&MorphTargetVertexInfoBuffers, &UE::RenderCommandPipe::SkeletalMesh);
+	BeginReleaseResource(&MorphTargetVertexInfoBuffers);
 
 	DEC_DWORD_STAT_BY(STAT_SkeletalMeshVertexMemory, SkinWeightProfilesData.GetResourcesSize());
 	SkinWeightProfilesData.ReleaseResources();
@@ -320,8 +319,8 @@ void FSkeletalMeshLODRenderData::ReleaseResources()
 #if RHI_RAYTRACING
 	if (IsRayTracingAllowed())
 	{
-		BeginReleaseResource(&SourceRayTracingGeometry, &UE::RenderCommandPipe::SkeletalMesh);
-		BeginReleaseResource(&StaticRayTracingGeometry, &UE::RenderCommandPipe::SkeletalMesh);
+		BeginReleaseResource(&SourceRayTracingGeometry);
+		BeginReleaseResource(&StaticRayTracingGeometry);
 	}
 #endif
 }
diff --git a/Engine/Source/Runtime/Engine/Private/SkeletalMeshRenderData.cpp b/Engine/Source/Runtime/Engine/Private/SkeletalMeshRenderData.cpp
--- a/Engine/Source/Runtime/Engine/Private/SkeletalMeshRenderData.cpp
+++ b/Engine/Source/Runtime/Engine/Private/SkeletalMeshRenderData.cpp
@@ -8,7 +8,6 @@
 #include "Engine/SkinnedAssetCommon.h"
 #include "EngineLogs.h"
 #include "UObject/Package.h"
-#include "Rendering/RenderCommandPipes.h"
 
 #if WITH_EDITOR
 #include "ProfilingDebugging/CookStats.h"
@@ -455,7 +454,7 @@ void FSkeletalMeshRenderData::SyncUVChannelData(const TArray<FSkeletalMaterial>&
 	// to the render thread.
 	if (bInitialized)
 	{
-		ENQUEUE_RENDER_COMMAND(SyncUVChannelData)(UE::RenderCommandPipe::SkeletalMesh, [this, UpdateData = MoveTemp(UpdateData)]
+		ENQUEUE_RENDER_COMMAND(SyncUVChannelData)([this, UpdateData = MoveTemp(UpdateData)](FRHICommandListImmediate& RHICmdList)
 		{
 			Swap(UVChannelDataPerMaterial, *UpdateData.Get());
 		});
@@ -642,8 +641,8 @@ void FSkeletalMeshRenderData::InitResources(bool bNeedsVertexColors, TArray<UMor
 			}
 		}
 
-		ENQUEUE_RENDER_COMMAND(CmdSetSkeletalMeshReadyForStreaming)(UE::RenderCommandPipe::SkeletalMesh,
-			[this, Owner]
+		ENQUEUE_RENDER_COMMAND(CmdSetSkeletalMeshReadyForStreaming)(
+			[this, Owner](FRHICommandListImmediate&)
 		{
 			bReadyForStreaming = true;
 		});
diff --git a/Engine/Source/Runtime/Engine/Private/SkeletalRenderCPUSkin.cpp b/Engine/Source/Runtime/Engine/Private/SkeletalRenderCPUSkin.cpp
--- a/Engine/Source/Runtime/Engine/Private/SkeletalRenderCPUSkin.cpp
+++ b/Engine/Source/Runtime/Engine/Private/SkeletalRenderCPUSkin.cpp
@@ -15,7 +15,6 @@
 #include "RenderCore.h"
 #include "SceneInterface.h"
 #include "Stats/StatsTrace.h"
-#include "Rendering/RenderCommandPipes.h"
 
 #if RHI_RAYTRACING
 #include "Engine/SkinnedAssetCommon.h"
@@ -155,8 +154,8 @@ void FSkeletalMeshObjectCPUSkin::Update(
 		// queue a call to update this data
 		{
 			FSkeletalMeshObjectCPUSkin* MeshObject = this;
-			ENQUEUE_RENDER_COMMAND(SkelMeshObjectUpdateDataCommand)(UE::RenderCommandPipe::SkeletalMesh,
-				[MeshObject, FrameNumberToPrepare, RevisionNumber, NewDynamicData](FRHICommandList& RHICmdList)
+			ENQUEUE_RENDER_COMMAND(SkelMeshObjectUpdateDataCommand)(
+				[MeshObject, FrameNumberToPrepare, RevisionNumber, NewDynamicData](FRHICommandListImmediate& RHICmdList)
 				{
 					FScopeCycleCounter Context(MeshObject->GetStatId());
 					MeshObject->UpdateDynamicData_RenderThread(RHICmdList, NewDynamicData, FrameNumberToPrepare, RevisionNumber);
@@ -297,13 +296,13 @@ void FSkeletalMeshObjectCPUSkin::CacheVertices(int32 LODIndex, bool bForce) cons
 			}
 		}
 
-		BeginUpdateResourceRHI(&MeshLOD.PositionVertexBuffer, &UE::RenderCommandPipe::SkeletalMesh);
-		BeginUpdateResourceRHI(&MeshLOD.StaticMeshVertexBuffer, &UE::RenderCommandPipe::SkeletalMesh);
+		BeginUpdateResourceRHI(&MeshLOD.PositionVertexBuffer);
+		BeginUpdateResourceRHI(&MeshLOD.StaticMeshVertexBuffer);
 
 		const FSkeletalMeshObjectCPUSkin::FSkeletalMeshObjectLOD* MeshLODptr = &MeshLOD;
 		FLocalVertexFactory* VertexFactoryPtr = &MeshLOD.VertexFactory;
-		ENQUEUE_RENDER_COMMAND(UpdateSkeletalMeshCPUSkinVertexFactory)(UE::RenderCommandPipe::SkeletalMesh,
-			[VertexFactoryPtr, MeshLODptr](FRHICommandList& RHICmdList)
+		ENQUEUE_RENDER_COMMAND(UpdateSkeletalMeshCPUSkinVertexFactory)(
+			[VertexFactoryPtr, MeshLODptr](FRHICommandListImmediate& RHICmdList)
 		{
 			FLocalVertexFactory::FDataType Data;
 
@@ -348,8 +347,8 @@ void FSkeletalMeshObjectCPUSkin::FSkeletalMeshObjectLOD::InitResources(FSkelMesh
 		StaticMeshVertexBuffer.SetVertexUV(i, 0, SrcVertexBuf.GetVertexUV(i, 0));
 	}
 
-	BeginInitResource(&PositionVertexBuffer, &UE::RenderCommandPipe::SkeletalMesh);
-	BeginInitResource(&StaticMeshVertexBuffer, &UE::RenderCommandPipe::SkeletalMesh);
+	BeginInitResource(&PositionVertexBuffer);
+	BeginInitResource(&StaticMeshVertexBuffer);
 
 	FSkeletalMeshObjectCPUSkin::FSkeletalMeshObjectLOD* Self = this;
 	FLocalVertexFactory* VertexFactoryPtr = &VertexFactory;
@@ -450,12 +449,12 @@ void FSkeletalMeshObjectCPUSkin::FSkeletalMeshObjectLOD::UpdateSkinWeights(FSkel
  */
 void FSkeletalMeshObjectCPUSkin::FSkeletalMeshObjectLOD::ReleaseResources()
 {	
-	BeginReleaseResource(&VertexFactory, &UE::RenderCommandPipe::SkeletalMesh);
-	BeginReleaseResource(&PositionVertexBuffer, &UE::RenderCommandPipe::SkeletalMesh);
-	BeginReleaseResource(&StaticMeshVertexBuffer, &UE::RenderCommandPipe::SkeletalMesh);
+	BeginReleaseResource(&VertexFactory);
+	BeginReleaseResource(&PositionVertexBuffer);
+	BeginReleaseResource(&StaticMeshVertexBuffer);
 
 #if RHI_RAYTRACING
-	BeginReleaseResource(&RayTracingGeometry, &UE::RenderCommandPipe::SkeletalMesh);
+	BeginReleaseResource(&RayTracingGeometry);
 #endif // RHI_RAYTRACING
 
 	bResourcesInitialized = false;
diff --git a/Engine/Source/Runtime/Engine/Private/SkeletalRenderGPUSkin.cpp b/Engine/Source/Runtime/Engine/Private/SkeletalRenderGPUSkin.cpp
--- a/Engine/Source/Runtime/Engine/Private/SkeletalRenderGPUSkin.cpp
+++ b/Engine/Source/Runtime/Engine/Private/SkeletalRenderGPUSkin.cpp
@@ -10,7 +10,7 @@
 #include "SkeletalRender.h"
 #include "GPUSkinCache.h"
 #include "RayTracingSkinnedGeometry.h"
-#include "Rendering/RenderCommandPipes.h"
+#include "RenderingThread.h"
 #include "ShaderParameterUtils.h"
 #include "SceneInterface.h"
 #include "SkeletalMeshSceneProxy.h"
@@ -143,17 +143,17 @@ void FMorphVertexBufferPool::InitResources(const FName& OwnerName)
 
 	check(!MorphVertexBuffers[0].VertexBufferRHI.IsValid());
 	check(!MorphVertexBuffers[1].VertexBufferRHI.IsValid());
-	BeginInitResource(&MorphVertexBuffers[0], &UE::RenderCommandPipe::SkeletalMesh);
+	BeginInitResource(&MorphVertexBuffers[0]);
 	if (bDoubleBuffer)
 	{
-		BeginInitResource(&MorphVertexBuffers[1], &UE::RenderCommandPipe::SkeletalMesh);
+		BeginInitResource(&MorphVertexBuffers[1]);
 	}
 }
 
 void FMorphVertexBufferPool::ReleaseResources()
 {
-	BeginReleaseResource(&MorphVertexBuffers[0], &UE::RenderCommandPipe::SkeletalMesh);
-	BeginReleaseResource(&MorphVertexBuffers[1], &UE::RenderCommandPipe::SkeletalMesh);
+	BeginReleaseResource(&MorphVertexBuffers[0]);
+	BeginReleaseResource(&MorphVertexBuffers[1]);
 }
 
 SIZE_T FMorphVertexBufferPool::GetResourceSize() const
@@ -281,7 +281,7 @@ void FSkeletalMeshObjectGPUSkin::InitResources(USkinnedMeshComponent* InMeshComp
 #if RHI_RAYTRACING
 	if (bSupportRayTracing)
 	{
-		BeginInitResource(&RayTracingGeometry, &UE::RenderCommandPipe::SkeletalMesh);
+		BeginInitResource(&RayTracingGeometry);
 	}
 #endif
 }
@@ -296,7 +296,7 @@ void FSkeletalMeshObjectGPUSkin::ReleaseResources()
 	ReleaseMorphResources();
 	FSkeletalMeshObjectGPUSkin* MeshObject = this;
 	FGPUSkinCacheEntry** PtrSkinCacheEntry = &SkinCacheEntry;
-	ENQUEUE_RENDER_COMMAND(WaitRHIThreadFenceForDynamicData)(UE::RenderCommandPipe::SkeletalMesh,
+	ENQUEUE_RENDER_COMMAND(WaitRHIThreadFenceForDynamicData)(
 		[MeshObject, PtrSkinCacheEntry, &SkinCacheEntryForRayTracing = SkinCacheEntryForRayTracing](FRHICommandList& RHICmdList)
 		{
 			FGPUSkinCacheEntry*& LocalSkinCacheEntry = *PtrSkinCacheEntry;
@@ -313,14 +313,14 @@ void FSkeletalMeshObjectGPUSkin::ReleaseResources()
 #if RHI_RAYTRACING
 	if (bSupportRayTracing)
 	{
-		BeginReleaseResource(&RayTracingGeometry, &UE::RenderCommandPipe::SkeletalMesh);
+		BeginReleaseResource(&RayTracingGeometry);
 	}
 
 	// Only enqueue when intialized
 	if (RayTracingUpdateQueue != nullptr || RayTracingDynamicVertexBuffer.NumBytes > 0)
 	{
-		ENQUEUE_RENDER_COMMAND(ReleaseRayTracingDynamicVertexBuffer)(UE::RenderCommandPipe::SkeletalMesh,
-			[RayTracingUpdateQueue = RayTracingUpdateQueue, RayTracingGeometryPtr = &RayTracingGeometry, &RayTracingDynamicVertexBuffer = RayTracingDynamicVertexBuffer](FRHICommandList& RHICmdList) mutable
+		ENQUEUE_RENDER_COMMAND(ReleaseRayTracingDynamicVertexBuffer)(
+			[RayTracingUpdateQueue = RayTracingUpdateQueue, RayTracingGeometryPtr = &RayTracingGeometry, &RayTracingDynamicVertexBuffer = RayTracingDynamicVertexBuffer](FRHICommandListImmediate& RHICmdList) mutable
 			{
 				if (RayTracingUpdateQueue != nullptr)
 				{
@@ -419,8 +419,8 @@ void FSkeletalMeshObjectGPUSkin::Update(
 
 	// queue a call to update this data
 	FSkeletalMeshObjectGPUSkin* MeshObject = this;
-	ENQUEUE_RENDER_COMMAND(SkelMeshObjectUpdateDataCommand)(UE::RenderCommandPipe::SkeletalMesh,
-		[MeshObject, FrameNumberToPrepare, RevisionNumber, NewDynamicData, GPUSkinCache, Scene](FRHICommandList& RHICmdList)
+	ENQUEUE_RENDER_COMMAND(SkelMeshObjectUpdateDataCommand)(
+		[MeshObject, FrameNumberToPrepare, RevisionNumber, NewDynamicData, GPUSkinCache, Scene](FRHICommandListImmediate& RHICmdList)
 		{
 			FScopeCycleCounter Context(MeshObject->GetStatId());
 			MeshObject->UpdateDynamicData_RenderThread(GPUSkinCache, RHICmdList, NewDynamicData, Scene, FrameNumberToPrepare, RevisionNumber);
@@ -450,8 +450,8 @@ void FSkeletalMeshObjectGPUSkin::UpdateSkinWeightBuffer(USkinnedMeshComponent* I
 				FGPUSkinCacheEntry* SkinCacheEntryToUpdate = SkinCacheEntry;
 				if (SkinCacheEntryToUpdate)
 				{
-					ENQUEUE_RENDER_COMMAND(UpdateSkinCacheSkinWeightBuffer)(UE::RenderCommandPipe::SkeletalMesh,
-						[SkinCacheEntryToUpdate](FRHICommandList& RHICmdList)
+					ENQUEUE_RENDER_COMMAND(UpdateSkinCacheSkinWeightBuffer)(
+						[SkinCacheEntryToUpdate](FRHICommandListImmediate& RHICmdList)
 					{
 						FGPUSkinCache::UpdateSkinWeightBuffer(SkinCacheEntryToUpdate);
 					});
@@ -459,8 +459,8 @@ void FSkeletalMeshObjectGPUSkin::UpdateSkinWeightBuffer(USkinnedMeshComponent* I
 
 				if (SkinCacheEntryForRayTracing)
 				{
-					ENQUEUE_RENDER_COMMAND(UpdateSkinCacheSkinWeightBuffer)(UE::RenderCommandPipe::SkeletalMesh,
-						[SkinCacheEntryForRayTracing = SkinCacheEntryForRayTracing](FRHICommandList& RHICmdList)
+					ENQUEUE_RENDER_COMMAND(UpdateSkinCacheSkinWeightBuffer)(
+						[SkinCacheEntryForRayTracing = SkinCacheEntryForRayTracing](FRHICommandListImmediate& RHICmdList)
 					{
 						FGPUSkinCache::UpdateSkinWeightBuffer(SkinCacheEntryForRayTracing);
 					});
@@ -1352,8 +1352,8 @@ void FSkeletalMeshObjectGPUSkin::FSkeletalMeshObjectLOD::UpdateSkinWeights(FSkel
 			GetVertexBuffers(VertexBuffers, LODData);
 
 			FSkeletalMeshObjectLOD* Self = this;
-			ENQUEUE_RENDER_COMMAND(UpdateSkinWeightsGPUSkin)(UE::RenderCommandPipe::SkeletalMesh,
-				[NewMeshObjectWeightBuffer, VertexBuffers, Self](FRHICommandList& RHICmdList)
+			ENQUEUE_RENDER_COMMAND(UpdateSkinWeightsGPUSkin)(
+				[NewMeshObjectWeightBuffer, VertexBuffers, Self](FRHICommandListImmediate& RHICmdList)
 			{
 				Self->GPUSkinVertexFactories.UpdateVertexFactoryData(VertexBuffers);
 			});
@@ -1810,7 +1810,7 @@ static FGPUBaseSkinVertexFactory* CreateVertexFactory(TArray<TUniquePtr<FGPUBase
 	FDynamicUpdateVertexFactoryData VertexUpdateData(VertexFactory, InVertexBuffers);
 
 	// update vertex factory components and sync it
-	ENQUEUE_RENDER_COMMAND(InitGPUSkinVertexFactory)(UE::RenderCommandPipe::SkeletalMesh,
+	ENQUEUE_RENDER_COMMAND(InitGPUSkinVertexFactory)(
 		[VertexUpdateData](FRHICommandList& RHICmdList)
 		{
 			FGPUSkinDataType Data;
@@ -1836,7 +1836,7 @@ void UpdateVertexFactory(TArray<TUniquePtr<FGPUBaseSkinVertexFactory>>& VertexFa
 			FDynamicUpdateVertexFactoryData VertexUpdateData(VertexFactory, InVertexBuffers);
 
 			// update vertex factory components and sync it
-			ENQUEUE_RENDER_COMMAND(UpdateGPUSkinVertexFactory)(UE::RenderCommandPipe::SkeletalMesh,
+			ENQUEUE_RENDER_COMMAND(UpdateGPUSkinVertexFactory)(
 				[VertexUpdateData](FRHICommandList& RHICmdList)
 			{
 				FGPUSkinDataType Data;
@@ -1854,7 +1854,7 @@ static void CreatePassthroughVertexFactory(ERHIFeatureLevel::Type InFeatureLevel
 	PassthroughVertexFactories.Add(TUniquePtr<FGPUSkinPassthroughVertexFactory>(NewPassthroughVertexFactory));
 
 	// update vertex factory components and sync it
-	ENQUEUE_RENDER_COMMAND(InitPassthroughGPUSkinVertexFactory)(UE::RenderCommandPipe::SkeletalMesh,
+	ENQUEUE_RENDER_COMMAND(InitPassthroughGPUSkinVertexFactory)(
 		[NewPassthroughVertexFactory, SourceVertexFactory](FRHICommandList& RHICmdList)
 		{
 			FLocalVertexFactory::FDataType Data;
@@ -1889,7 +1889,7 @@ static void CreateVertexFactoryMorph(TArray<TUniquePtr<FGPUBaseSkinVertexFactory
 	FDynamicUpdateVertexFactoryData VertexUpdateData(VertexFactory, InVertexBuffers);
 
 	// update vertex factory components and sync it
-	ENQUEUE_RENDER_COMMAND(InitGPUSkinVertexFactoryMorph)(UE::RenderCommandPipe::SkeletalMesh,
+	ENQUEUE_RENDER_COMMAND(InitGPUSkinVertexFactoryMorph)(
 		[VertexUpdateData](FRHICommandList& RHICmdList)
 		{
 			FGPUSkinDataType Data;
@@ -1914,7 +1914,7 @@ static void UpdateVertexFactoryMorph(TArray<TUniquePtr<FGPUBaseSkinVertexFactory
 			FDynamicUpdateVertexFactoryData VertexUpdateData(VertexFactory, InVertexBuffers);
 
 			// update vertex factory components and sync it
-			ENQUEUE_RENDER_COMMAND(InitGPUSkinVertexFactoryMorph)(UE::RenderCommandPipe::SkeletalMesh,
+			ENQUEUE_RENDER_COMMAND(InitGPUSkinVertexFactoryMorph)(
 				[VertexUpdateData](FRHICommandList& RHICmdList)
 			{
 				FGPUSkinDataType Data;
@@ -1968,7 +1968,7 @@ static void CreateVertexFactoryCloth(TArray<TUniquePtr<FGPUBaseSkinAPEXClothVert
 	FDynamicUpdateVertexFactoryData VertexUpdateData(VertexFactory, InVertexBuffers);
 
 	// update vertex factory components and sync it
-	ENQUEUE_RENDER_COMMAND(InitGPUSkinAPEXClothVertexFactory)(UE::RenderCommandPipe::SkeletalMesh,
+	ENQUEUE_RENDER_COMMAND(InitGPUSkinAPEXClothVertexFactory)(
 		[VertexUpdateData](FRHICommandList& RHICmdList)
 		{
 			FGPUSkinAPEXClothDataType Data;
@@ -1993,7 +1993,7 @@ static void UpdateVertexFactoryCloth(TArray<TUniquePtr<FGPUBaseSkinAPEXClothVert
 			FDynamicUpdateVertexFactoryData VertexUpdateData(VertexFactory, InVertexBuffers);
 
 			// update vertex factory components and sync it
-			ENQUEUE_RENDER_COMMAND(InitGPUSkinAPEXClothVertexFactory)(UE::RenderCommandPipe::SkeletalMesh,
+			ENQUEUE_RENDER_COMMAND(InitGPUSkinAPEXClothVertexFactory)(
 				[VertexUpdateData](FRHICommandList& RHICmdList)
 			{
 				FGPUSkinAPEXClothDataType Data;
@@ -2112,12 +2112,12 @@ void FSkeletalMeshObjectGPUSkin::FVertexFactoryData::ReleaseVertexFactories()
 	// Default factories
 	for( int32 FactoryIdx=0; FactoryIdx < VertexFactories.Num(); FactoryIdx++)
 	{
-		BeginReleaseResource(VertexFactories[FactoryIdx].Get(), &UE::RenderCommandPipe::SkeletalMesh);
+		BeginReleaseResource(VertexFactories[FactoryIdx].Get());
 	}
 
 	for (int32 FactoryIdx = 0; FactoryIdx < PassthroughVertexFactories.Num(); FactoryIdx++)
 	{
-		BeginReleaseResource(PassthroughVertexFactories[FactoryIdx].Get(), &UE::RenderCommandPipe::SkeletalMesh);
+		BeginReleaseResource(PassthroughVertexFactories[FactoryIdx].Get());
 	}
 }
 
@@ -2143,7 +2143,7 @@ void FSkeletalMeshObjectGPUSkin::FVertexFactoryData::ReleaseMorphVertexFactories
 	// Default morph factories
 	for( int32 FactoryIdx=0; FactoryIdx < MorphVertexFactories.Num(); FactoryIdx++ )
 	{
-		BeginReleaseResource(MorphVertexFactories[FactoryIdx].Get(), &UE::RenderCommandPipe::SkeletalMesh);
+		BeginReleaseResource(MorphVertexFactories[FactoryIdx].Get());
 	}
 }
 
@@ -2188,7 +2188,7 @@ void FSkeletalMeshObjectGPUSkin::FVertexFactoryData::ReleaseAPEXClothVertexFacto
 		TUniquePtr<FGPUBaseSkinAPEXClothVertexFactory>& ClothVertexFactory = ClothVertexFactories[FactoryIdx];
 		if (ClothVertexFactory)
 		{
-			BeginReleaseResource(ClothVertexFactory->GetVertexFactory(), &UE::RenderCommandPipe::SkeletalMesh);
+			BeginReleaseResource(ClothVertexFactory->GetVertexFactory());
 		}
 	}
 }
diff --git a/Engine/Source/Runtime/Engine/Private/SkeletalRenderStatic.cpp b/Engine/Source/Runtime/Engine/Private/SkeletalRenderStatic.cpp
--- a/Engine/Source/Runtime/Engine/Private/SkeletalRenderStatic.cpp
+++ b/Engine/Source/Runtime/Engine/Private/SkeletalRenderStatic.cpp
@@ -7,7 +7,7 @@
 #include "SkeletalRenderStatic.h"
 #include "RenderUtils.h"
 #include "Rendering/SkeletalMeshRenderData.h"
-#include "Rendering/RenderCommandPipes.h"
+#include "RenderingThread.h"
 
 #if RHI_RAYTRACING
 #include "Engine/SkinnedAssetCommon.h"
@@ -67,7 +67,7 @@ void FSkeletalMeshObjectStatic::InitResources(USkinnedMeshComponent* InMeshCompo
 					}
 
 					TArray<FSkelMeshRenderSection>* RenderSections = &LODModel.RenderSections;
-					ENQUEUE_RENDER_COMMAND(InitSkeletalRenderStaticRayTracingGeometry)(UE::RenderCommandPipe::SkeletalMesh,
+					ENQUEUE_RENDER_COMMAND(InitSkeletalRenderStaticRayTracingGeometry)(
 						[this, VertexBufferRHI, IndexBufferRHI, VertexBufferStride, TrianglesCount, RenderSections, 
 						LODIndex = LODIndex, 
 						SkelMeshRenderData = SkelLOD.SkelMeshRenderData, 
@@ -144,13 +144,14 @@ void FSkeletalMeshObjectStatic::ReleaseResources()
 
 					if (SkelLOD.SkelMeshRenderData->LODRenderData[LODIndex].NumReferencingStaticSkeletalMeshObjects == 0)
 					{
-						ENQUEUE_RENDER_COMMAND(ResetStaticRayTracingGeometryFlag)(UE::RenderCommandPipe::SkeletalMesh,
-							[&bReferencedByStaticSkeletalMeshObjects_RenderThread = SkelLOD.SkelMeshRenderData->LODRenderData[LODIndex].bReferencedByStaticSkeletalMeshObjects_RenderThread]
+						ENQUEUE_RENDER_COMMAND(ResetStaticRayTracingGeometryFlag)(
+							[&bReferencedByStaticSkeletalMeshObjects_RenderThread = SkelLOD.SkelMeshRenderData->LODRenderData[LODIndex].bReferencedByStaticSkeletalMeshObjects_RenderThread](FRHICommandListImmediate& RHICmdList)
 						{
 							bReferencedByStaticSkeletalMeshObjects_RenderThread = false;
-						});
+						}
+						);
 
-						BeginReleaseResource(&SkelLOD.SkelMeshRenderData->LODRenderData[LODIndex].StaticRayTracingGeometry, &UE::RenderCommandPipe::SkeletalMesh);
+						BeginReleaseResource(&SkelLOD.SkelMeshRenderData->LODRenderData[LODIndex].StaticRayTracingGeometry);
 					}
 				}
 			}
@@ -203,7 +204,7 @@ void FSkeletalMeshObjectStatic::FSkeletalMeshObjectLOD::InitResources(FSkelMeshC
 	FLocalVertexFactory* VertexFactoryPtr = &VertexFactory;
 	FColorVertexBuffer* ColorVertexBufferPtr = ColorVertexBuffer;
 
-	ENQUEUE_RENDER_COMMAND(InitSkeletalMeshStaticSkinVertexFactory)(UE::RenderCommandPipe::SkeletalMesh,
+	ENQUEUE_RENDER_COMMAND(InitSkeletalMeshStaticSkinVertexFactory)(
 		[VertexFactoryPtr, PositionVertexBufferPtr, StaticMeshVertexBufferPtr, ColorVertexBufferPtr](FRHICommandList& RHICmdList)
 		{
 			FLocalVertexFactory::FDataType Data;
@@ -229,10 +230,10 @@ void FSkeletalMeshObjectStatic::FSkeletalMeshObjectLOD::InitResources(FSkelMeshC
  */
 void FSkeletalMeshObjectStatic::FSkeletalMeshObjectLOD::ReleaseResources()
 {	
-	BeginReleaseResource(&VertexFactory, &UE::RenderCommandPipe::SkeletalMesh);
+	BeginReleaseResource(&VertexFactory);
 
 #if RHI_RAYTRACING
-	BeginReleaseResource(&RayTracingGeometry, &UE::RenderCommandPipe::SkeletalMesh);
+	BeginReleaseResource(&RayTracingGeometry);
 #endif // RHI_RAYTRACING
 
 	bResourcesInitialized = false;
diff --git a/Engine/Source/Runtime/Engine/Public/PrecomputedVolumetricLightmap.h b/Engine/Source/Runtime/Engine/Public/PrecomputedVolumetricLightmap.h
--- a/Engine/Source/Runtime/Engine/Public/PrecomputedVolumetricLightmap.h
+++ b/Engine/Source/Runtime/Engine/Public/PrecomputedVolumetricLightmap.h
@@ -226,7 +226,7 @@ public:
 
 	ENGINE_API void RemoveFromScene(FSceneInterface* Scene);
 	
-	ENGINE_API void SetData(FRHICommandListBase& RHICmdList, FPrecomputedVolumetricLightmapData* NewData, FSceneInterface* Scene);
+	ENGINE_API void SetData(FPrecomputedVolumetricLightmapData* NewData, FSceneInterface* Scene);
 
 	bool IsAddedToScene() const
 	{
diff --git a/Engine/Source/Runtime/Engine/Public/Rendering/RenderCommandPipes.h b/Engine/Source/Runtime/Engine/Public/Rendering/RenderCommandPipes.h
deleted file mode 100644
diff --git a/Engine/Source/Runtime/GeometryFramework/Public/Components/BaseDynamicMeshSceneProxy.h b/Engine/Source/Runtime/GeometryFramework/Public/Components/BaseDynamicMeshSceneProxy.h
--- a/Engine/Source/Runtime/GeometryFramework/Public/Components/BaseDynamicMeshSceneProxy.h
+++ b/Engine/Source/Runtime/GeometryFramework/Public/Components/BaseDynamicMeshSceneProxy.h
@@ -98,6 +98,8 @@ public:
 
 	virtual ~FMeshRenderBufferSet()
 	{
+		check(IsInRenderingThread());
+
 		if (TriangleCount > 0)
 		{
 			PositionVertexBuffer.ReleaseResource();
@@ -391,7 +393,11 @@ protected:
 			return;
 		}
 
-		delete BufferSet;
+		ENQUEUE_RENDER_COMMAND(FMeshRenderBufferSetDestroy)(
+			[BufferSet](FRHICommandListImmediate& RHICmdList)
+		{
+			delete BufferSet;
+		});
 	}
 
 
