Fix redundant delegate registrations on BPGC load/compile (regression).

diff --git a/Engine/Source/Runtime/Engine/Private/BlueprintGeneratedClass.cpp b/Engine/Source/Runtime/Engine/Private/BlueprintGeneratedClass.cpp
--- a/Engine/Source/Runtime/Engine/Private/BlueprintGeneratedClass.cpp
+++ b/Engine/Source/Runtime/Engine/Private/BlueprintGeneratedClass.cpp
@@ -124,23 +124,29 @@ namespace UE::Runtime::Engine::Private
 
 			static bool bIsInitialized = false;
 #if WITH_EDITOR
-			FCoreUObjectDelegates::ReloadCompleteDelegate.AddLambda([](EReloadCompleteReason Reason)
+			static bool bNeedsDelegateRegistration = true;
+			if (bNeedsDelegateRegistration)
 			{
-				switch(Reason)
+				FCoreUObjectDelegates::ReloadCompleteDelegate.AddLambda([](EReloadCompleteReason Reason)
 				{
-				case EReloadCompleteReason::HotReloadManual:
-				case EReloadCompleteReason::HotReloadAutomatic:
-				{
-					// Re-initialize after hot reload in editor context.
-					bIsInitialized = false;
-					PropertiesRequiringCompleteValueInitialization.Reset();
-				}
+					switch (Reason)
+					{
+					case EReloadCompleteReason::HotReloadManual:
+					case EReloadCompleteReason::HotReloadAutomatic:
+					{
+						// Re-initialize after hot reload in editor context.
+						bIsInitialized = false;
+						PropertiesRequiringCompleteValueInitialization.Reset();
+					}
 					break;
 
-				default:
-					break;
-				}
-			});
+					default:
+						break;
+					}
+				});
+
+				bNeedsDelegateRegistration = false;
+			}
 #endif	// WITH_EDITOR
 
 			if (!bIsInitialized && GConfig)
