ISMComponentDescriptor: Add bEmissiveLightSource and bAffectIndirectLightingWhileHidden to descriptor properties

diff --git a/Engine/Source/Runtime/Engine/Private/ISMPartition/ISMComponentDescriptor.cpp b/Engine/Source/Runtime/Engine/Private/ISMPartition/ISMComponentDescriptor.cpp
--- a/Engine/Source/Runtime/Engine/Private/ISMPartition/ISMComponentDescriptor.cpp
+++ b/Engine/Source/Runtime/Engine/Private/ISMPartition/ISMComponentDescriptor.cpp
@@ -89,11 +89,13 @@ void FISMComponentDescriptorBase::InitFrom(const UStaticMeshComponent* Template,
 	OverriddenLightMapRes = Template->OverriddenLightMapRes;
 	CustomDepthStencilValue = Template->CustomDepthStencilValue;
 	bCastShadow = Template->CastShadow;
+	bEmissiveLightSource = Template->bEmissiveLightSource;
 	bCastStaticShadow = Template->bCastStaticShadow;
 	bCastDynamicShadow = Template->bCastDynamicShadow;
 	bCastContactShadow = Template->bCastContactShadow;
 	bCastShadowAsTwoSided = Template->bCastShadowAsTwoSided;
 	bAffectDynamicIndirectLighting = Template->bAffectDynamicIndirectLighting;
+	bAffectDynamicIndirectLightingWhileHidden = Template->bAffectIndirectLightingWhileHidden;
 	bAffectDistanceFieldLighting = Template->bAffectDistanceFieldLighting;
 	bReceivesDecals = Template->bReceivesDecals;
 	bOverrideLightMapRes = Template->bOverrideLightMapRes;
@@ -191,11 +193,13 @@ bool FISMComponentDescriptorBase::operator==(const FISMComponentDescriptorBase&
 	OverriddenLightMapRes == Other.OverriddenLightMapRes &&
 	CustomDepthStencilValue == Other.CustomDepthStencilValue &&
 	bCastShadow == Other.bCastShadow &&
+	bEmissiveLightSource == Other.bEmissiveLightSource &&
 	bCastStaticShadow == Other.bCastStaticShadow &&
 	bCastDynamicShadow == Other.bCastDynamicShadow &&
 	bCastContactShadow == Other.bCastContactShadow &&
 	bCastShadowAsTwoSided == Other.bCastShadowAsTwoSided &&
 	bAffectDynamicIndirectLighting == Other.bAffectDynamicIndirectLighting &&
+	bAffectDynamicIndirectLightingWhileHidden == Other.bAffectDynamicIndirectLightingWhileHidden &&
 	bAffectDistanceFieldLighting == Other.bAffectDistanceFieldLighting &&
 	bReceivesDecals == Other.bReceivesDecals &&
 	bOverrideLightMapRes == Other.bOverrideLightMapRes &&
@@ -302,11 +306,13 @@ void FISMComponentDescriptorBase::InitComponent(UInstancedStaticMeshComponent* I
 	ISMComponent->OverriddenLightMapRes = OverriddenLightMapRes;
 	ISMComponent->CustomDepthStencilValue = CustomDepthStencilValue;
 	ISMComponent->CastShadow = bCastShadow;
+	ISMComponent->bEmissiveLightSource = bEmissiveLightSource;
 	ISMComponent->bCastStaticShadow = bCastStaticShadow;
 	ISMComponent->bCastDynamicShadow = bCastDynamicShadow;
 	ISMComponent->bCastContactShadow = bCastContactShadow;
 	ISMComponent->bCastShadowAsTwoSided = bCastShadowAsTwoSided;
 	ISMComponent->bAffectDynamicIndirectLighting = bAffectDynamicIndirectLighting;
+	ISMComponent->bAffectIndirectLightingWhileHidden = bAffectDynamicIndirectLightingWhileHidden;
 	ISMComponent->bAffectDistanceFieldLighting = bAffectDistanceFieldLighting;
 	ISMComponent->bReceivesDecals = bReceivesDecals;
 	ISMComponent->bOverrideLightMapRes = bOverrideLightMapRes;
diff --git a/Engine/Source/Runtime/Engine/Public/ISMPartition/ISMComponentDescriptor.h b/Engine/Source/Runtime/Engine/Public/ISMPartition/ISMComponentDescriptor.h
--- a/Engine/Source/Runtime/Engine/Public/ISMPartition/ISMComponentDescriptor.h
+++ b/Engine/Source/Runtime/Engine/Public/ISMPartition/ISMComponentDescriptor.h
@@ -92,6 +92,9 @@ public:
 	UPROPERTY(EditAnywhere, Category = "Component Settings")
 	uint8 bCastShadow : 1;
 
+	UPROPERTY(EditAnywhere, Category = "Component Settings")
+	uint8 bEmissiveLightSource : 1;
+		
 	UPROPERTY(EditAnywhere, Category = "Component Settings")
 	uint8 bCastDynamicShadow : 1;
 
@@ -107,6 +110,9 @@ public:
 	UPROPERTY(EditAnywhere, Category = "Component Settings")
 	uint8 bAffectDynamicIndirectLighting : 1;
 
+	UPROPERTY(EditAnywhere, Category = "Component Settings")
+	uint8 bAffectDynamicIndirectLightingWhileHidden : 1;
+
 	UPROPERTY(EditAnywhere, Category = "Component Settings")
 	uint8 bAffectDistanceFieldLighting : 1;
 
