Clear up usage of the DDC flag, and add some logging to try and catch a random failure we're getting.

diff --git a/Engine/Source/Runtime/Engine/Private/TextureDerivedDataTask.cpp b/Engine/Source/Runtime/Engine/Private/TextureDerivedDataTask.cpp
--- a/Engine/Source/Runtime/Engine/Private/TextureDerivedDataTask.cpp
+++ b/Engine/Source/Runtime/Engine/Private/TextureDerivedDataTask.cpp
@@ -1412,7 +1412,6 @@ static void DDC1_FetchAndFillDerivedData(
 	FTexturePlatformData* DerivedData,
 	FString& KeySuffix,
 	bool& bSucceeded,
-	bool& bLoadedFromDDC,
 	bool& bInvalidVirtualTextureCompression,
 	int64& BytesCached
 	)
@@ -1434,6 +1433,7 @@ static void DDC1_FetchAndFillDerivedData(
 	FString FetchOrBuildKeySuffix;
 	GetTextureDerivedDataKeySuffix(Texture, BuildSettingsPerLayerFetchOrBuild.GetData(), FetchOrBuildKeySuffix);
 
+	bool bGotDDCData = false;
 	bool bUsedFetchFirst = false;
 	if (BuildSettingsPerLayerFetchFirst.Num() && !bForceRebuild)
 	{
@@ -1466,8 +1466,8 @@ static void DDC1_FetchAndFillDerivedData(
 			});
 			BlockingOwner.Wait();
 
-			bLoadedFromDDC = !RawDerivedData.IsNull();
-			if (bLoadedFromDDC)
+			bGotDDCData = !RawDerivedData.IsNull();
+			if (bGotDDCData)
 			{
 				// We don't necessarily get the metadata until we force a texture rebuild.
 				if (MetadataBuffer.IsNull() == false)
@@ -1481,7 +1481,7 @@ static void DDC1_FetchAndFillDerivedData(
 		}
 	}
 
-	if (bLoadedFromDDC == false)
+	if (bGotDDCData == false)
 	{
 		// Didn't get the initial fetch, so we're using fetch/build.
 		LocalDerivedDataKeySuffix = MoveTemp(FetchOrBuildKeySuffix);
@@ -1507,8 +1507,8 @@ static void DDC1_FetchAndFillDerivedData(
 		});
 		BlockingOwner.Wait();
 
-		bLoadedFromDDC = !RawDerivedData.IsNull();
-		if (bLoadedFromDDC)
+		bGotDDCData = !RawDerivedData.IsNull();
+		if (bGotDDCData)
 		{
 			// Only read the metadata if we actually got the main data.
 			// We don't necessarily get the metadata until we force a texture rebuild.
@@ -1523,7 +1523,7 @@ static void DDC1_FetchAndFillDerivedData(
 	DerivedData->DerivedDataKey.Emplace<FString>(LocalDerivedDataKey);
 	DerivedData->ResultMetadata = bUsedFetchFirst ? FetchFirstMetadata : FetchOrBuildMetadata;
 
-	if (bLoadedFromDDC)
+	if (bGotDDCData)
 	{
 		const bool bInlineMips = EnumHasAnyFlags(CacheFlags, ETextureCacheFlags::InlineMips);
 		const bool bForDDC = EnumHasAnyFlags(CacheFlags, ETextureCacheFlags::ForDDCBuild);
@@ -1663,8 +1663,6 @@ static void DDC1_FetchAndFillDerivedData(
 				delete DerivedData->VTData;
 				DerivedData->VTData = nullptr;
 			}
-			
-			bLoadedFromDDC = false;
 		}
 	}
 
@@ -1773,15 +1771,14 @@ bool DDC1_BuildTiledClassicTexture(
 	int64 LinearBytesCached = 0;
 	bool bLinearDDCCorrupted = false;
 	bool bLinearSucceeded = false;
-	bool bLinearLoadedFromDDC = false;
 	DDC1_FetchAndFillDerivedData(
 		Texture, TexturePathName, CacheFlags,
 		LinearSettingsPerLayerFetchFirst, FetchFirstMetadata, 
 		LinearSettingsPerLayerFetchOrBuild, FetchOrBuildMetadata,
-		&LinearDerivedData, LinearKeySuffix, bLinearSucceeded, bLinearLoadedFromDDC, bLinearDDCCorrupted, LinearBytesCached);
+		&LinearDerivedData, LinearKeySuffix, bLinearSucceeded, bLinearDDCCorrupted, LinearBytesCached);
 
 	BytesCached = LinearBytesCached;
-	bool bHasLinearDerivedData = bLinearSucceeded && bLinearLoadedFromDDC;
+	bool bHasLinearDerivedData = bLinearSucceeded;
 
 	if (bHasLinearDerivedData == false)
 	{
@@ -1969,7 +1966,11 @@ void FTextureCacheDerivedDataWorker::DoWork()
 
 	DDC1_FetchAndFillDerivedData(
 		/* inputs */ Texture, TexturePathName, CacheFlags, BuildSettingsPerLayerFetchFirst, FetchFirstMetadata, BuildSettingsPerLayerFetchOrBuild, FetchOrBuildMetadata, 
-		/* outputs */ DerivedData, KeySuffix, bSucceeded, bLoadedFromDDC, bInvalidVirtualTextureCompression, BytesCached);
+		/* outputs */ DerivedData, KeySuffix, bSucceeded, bInvalidVirtualTextureCompression, BytesCached);
+	if (bSucceeded)
+	{
+		bLoadedFromDDC = true;
+	}
 
 	if (BuildSettingsPerLayerFetchOrBuild[0].Tiler && !bForVirtualTextureStreamingBuild)
 	{
@@ -2025,6 +2026,22 @@ void FTextureCacheDerivedDataWorker::DoWork()
 		{
 			bSucceeded = false;
 
+			// Excess logging to try and nail down a spurious failure.
+			UE_LOG(LogTexture, Display, TEXT("Texture was not found in DDC and couldn't build as the texture source was unable to load or validate (%s)"), *TexturePathName);
+			int32 TextureDataBlocks = TextureData.Blocks.Num();
+			int32 TextureDataBlocksLayers = TextureDataBlocks > 0 ? TextureData.Blocks[0].MipsPerLayer.Num() : -1;
+			int32 TextureDataBlocksLayerMips = TextureDataBlocksLayers > 0 ? TextureData.Blocks[0].MipsPerLayer[0].Num() : -1;
+
+			UE_LOG(LogTexture, Display, TEXT("Texture Data Blocks: %d Layers: %d Mips: %d"), TextureDataBlocks, TextureDataBlocksLayers, TextureDataBlocksLayerMips);
+			if (CompositeTextureData.IsValid()) // Says IsValid, but means whether or not we _need_ composite texture data.
+			{
+				int32 CompositeTextureDataBlocks = CompositeTextureData.Blocks.Num();
+				int32 CompositeTextureDataBlocksLayers = CompositeTextureDataBlocks > 0 ? CompositeTextureData.Blocks[0].MipsPerLayer.Num() : -1;
+				int32 CompositeTextureDataBlocksLayerMips = CompositeTextureDataBlocksLayers > 0 ? CompositeTextureData.Blocks[0].MipsPerLayer[0].Num() : -1;
+
+				UE_LOG(LogTexture, Display, TEXT("Composite Texture Data Blocks: %d Layers: %d Mips: %d"), CompositeTextureDataBlocks, CompositeTextureDataBlocksLayers, CompositeTextureDataBlocksLayerMips);
+			}
+
 			// bTriedAndFailed = true; // no retry in Finalize ?  @todo ?
 		}
 	}
