PLAY-8961: Crash solved:
Chaos::TAABBTree<FExternalSpatialAccelerationPayload,Chaos::TAABBTreeLeafArray<FExternalSpatialAccelerationPayload,1,double>,1,double>::RemoveElement(FExternalSpatialAccelerationPayload const &) [AABBTree.h:1679]

diff --git a/Engine/Source/Runtime/Engine/Classes/PhysicsEngine/ExternalSpatialAccelerationPayload.h b/Engine/Source/Runtime/Engine/Classes/PhysicsEngine/ExternalSpatialAccelerationPayload.h
--- a/Engine/Source/Runtime/Engine/Classes/PhysicsEngine/ExternalSpatialAccelerationPayload.h
+++ b/Engine/Source/Runtime/Engine/Classes/PhysicsEngine/ExternalSpatialAccelerationPayload.h
@@ -37,6 +37,10 @@ public:
 	ENGINE_API void Serialize(Chaos::FChaosArchive& Ar);
 	bool PrePreQueryFilter(const void* QueryData) const { return false; }
 	bool PrePreSimFilter(const void* SimData) const { return false; }
+	bool IsValid() const 
+	{
+		return CachedUniqueIdx.IsValid(); // Don't use this payload if the UniqueIdx is invalid
+	};
 
 	friend ENGINE_API uint32 GetTypeHash(const FExternalSpatialAccelerationPayload& Payload);
 private:
diff --git a/Engine/Source/Runtime/Engine/Private/PhysicsEngine/ClusterUnionComponent.cpp b/Engine/Source/Runtime/Engine/Private/PhysicsEngine/ClusterUnionComponent.cpp
--- a/Engine/Source/Runtime/Engine/Private/PhysicsEngine/ClusterUnionComponent.cpp
+++ b/Engine/Source/Runtime/Engine/Private/PhysicsEngine/ClusterUnionComponent.cpp
@@ -150,8 +150,11 @@ void UClusterUnionComponent::AddComponentToCluster(UPrimitiveComponent* InCompon
 				const FBox HandleBounds = Interface->GetWorldBounds({ &Object, 1 });
 				FExternalSpatialAccelerationPayload Handle;
 				Handle.Initialize(InComponent, BoneId);
-				AccelerationStructure->UpdateElement(Handle, Chaos::TAABB<Chaos::FReal, 3>{HandleBounds.Min, HandleBounds.Max}, HandleBounds.IsValid != 0);
-				PendingData.AccelerationPayloads.Add(Handle);
+				if (Handle.IsValid())
+				{
+					AccelerationStructure->UpdateElement(Handle, Chaos::TAABB<Chaos::FReal, 3>{HandleBounds.Min, HandleBounds.Max}, HandleBounds.IsValid != 0);
+					PendingData.AccelerationPayloads.Add(Handle);
+				}
 			}
 
 			const FTransform CurrentParticleTransform = Interface->GetTransform(Object);
@@ -204,7 +207,10 @@ void UClusterUnionComponent::RemoveComponentFromCluster(UPrimitiveComponent* InC
 		{
 			for (const FExternalSpatialAccelerationPayload& Payload : PendingData->AccelerationPayloads)
 			{
-				AccelerationStructure->RemoveElement(Payload);
+				if (Payload.IsValid())
+				{
+					AccelerationStructure->RemoveElement(Payload);
+				}
 			}
 		}
 
@@ -239,7 +245,10 @@ void UClusterUnionComponent::RemoveComponentFromCluster(UPrimitiveComponent* InC
 		{
 			for (const FExternalSpatialAccelerationPayload& Payload : ComponentData->CachedAccelerationPayloads)
 			{
-				AccelerationStructure->RemoveElement(Payload);
+				if (Payload.IsValid())
+				{
+					AccelerationStructure->RemoveElement(Payload);
+				}
 			}
 			ComponentData->CachedAccelerationPayloads.Reset();
 		}
@@ -300,7 +309,10 @@ void UClusterUnionComponent::RemoveComponentBonesFromCluster(UPrimitiveComponent
 				FExternalSpatialAccelerationPayload Payload;
 				Payload.Initialize(InComponent, BoneId);
 				PendingData->AccelerationPayloads.Remove(Payload);
-				AccelerationStructure->RemoveElement(Payload);
+				if (Payload.IsValid())
+				{
+					AccelerationStructure->RemoveElement(Payload);
+				}
 			}
 		}
 
@@ -325,7 +337,10 @@ void UClusterUnionComponent::RemoveComponentBonesFromCluster(UPrimitiveComponent
 				FExternalSpatialAccelerationPayload Payload;
 				Payload.Initialize(InComponent, BoneId);
 				ComponentData->CachedAccelerationPayloads.Remove(Payload);
-				AccelerationStructure->RemoveElement(Payload);
+				if (Payload.IsValid())
+				{
+					AccelerationStructure->RemoveElement(Payload);
+				}
 			}
 			ComponentData->BoneIds.Remove(BoneId);
 		}
@@ -909,8 +924,11 @@ void UClusterUnionComponent::HandleAddOrModifiedClusteredComponent(UPrimitiveCom
 				const FBox HandleBounds = Interface->GetWorldBounds({ &PhysicsObject, 1 });
 				FExternalSpatialAccelerationPayload Handle;
 				Handle.Initialize(ChangedComponent, Kvp.Key);
-				AccelerationStructure->UpdateElement(Handle, Chaos::TAABB<Chaos::FReal, 3>{HandleBounds.Min, HandleBounds.Max}, HandleBounds.IsValid != 0);
-				ComponentData.CachedAccelerationPayloads.Add(Handle);
+				if (Handle.IsValid())
+				{
+					AccelerationStructure->UpdateElement(Handle, Chaos::TAABB<Chaos::FReal, 3>{HandleBounds.Min, HandleBounds.Max}, HandleBounds.IsValid != 0);
+					ComponentData.CachedAccelerationPayloads.Add(Handle);
+				}
 			}
 		}
 	}
@@ -925,8 +943,11 @@ void UClusterUnionComponent::HandleAddOrModifiedClusteredComponent(UPrimitiveCom
 			{
 				FExternalSpatialAccelerationPayload Handle;
 				Handle.Initialize(ChangedComponent, BoneId);
-				AccelerationStructure->RemoveElement(Handle);
-				ComponentData.CachedAccelerationPayloads.Remove(Handle);
+				if (Handle.IsValid())
+				{
+					AccelerationStructure->RemoveElement(Handle);
+					ComponentData.CachedAccelerationPayloads.Remove(Handle);
+				}
 			}
 		}
 	}
@@ -940,7 +961,10 @@ void UClusterUnionComponent::HandleRemovedClusteredComponent(TObjectKey<UPrimiti
 		{
 			for (const FExternalSpatialAccelerationPayload& Handle : Data->CachedAccelerationPayloads)
 			{
-				AccelerationStructure->RemoveElement(Handle);
+				if (Handle.IsValid())
+				{
+					AccelerationStructure->RemoveElement(Handle);
+				}
 			}
 		}
 
diff --git a/Engine/Source/Runtime/Experimental/Chaos/Public/Chaos/ISpatialAcceleration.h b/Engine/Source/Runtime/Experimental/Chaos/Public/Chaos/ISpatialAcceleration.h
--- a/Engine/Source/Runtime/Experimental/Chaos/Public/Chaos/ISpatialAcceleration.h
+++ b/Engine/Source/Runtime/Experimental/Chaos/Public/Chaos/ISpatialAcceleration.h
@@ -525,7 +525,7 @@ public:
 	TValue* Find(const TKey& Key)
 	{
 		const int32 Idx = GetUniqueIdx(Key).Idx;
-		if(Idx < Entries.Num() && Entries[Idx].bSet)
+		if(Entries.IsValidIndex(Idx) && Entries[Idx].bSet)
 		{
 			return &Entries[Idx].Value;
 		}
