Fix issue where particles created with probe collision in CreateGeometry were excluding sim collision. This is correct at the UE level, but not at the Chaos level - there must be parity between ChaosInterface::CreateGeometry and FBodyInstance::UpdatePhysicsFilterData. CreateGeometry was missing the details injected by BuildBodyCollisionFlags, which became noticeable after my last change to enable probe collision in CreateGeometry.

diff --git a/Engine/Source/Runtime/Engine/Private/PhysicsEngine/Experimental/ChaosInterfaceUtils.cpp b/Engine/Source/Runtime/Engine/Private/PhysicsEngine/Experimental/ChaosInterfaceUtils.cpp
--- a/Engine/Source/Runtime/Engine/Private/PhysicsEngine/Experimental/ChaosInterfaceUtils.cpp
+++ b/Engine/Source/Runtime/Engine/Private/PhysicsEngine/Experimental/ChaosInterfaceUtils.cpp
@@ -118,6 +118,10 @@ namespace ChaosInterface
 			const ECollisionEnabled::Type CombinedCollisionEnabled = CollisionEnabledIntersection(BodyCollisionEnabled, ShapeCollisionEnabled);
 			CollisionEnabledToFlags(CombinedCollisionEnabled, bEnableQuery, bEnableSim, bEnableProbe);
 
+			// Mirror behavior of FBodyInstance::BuildBodyCollisionFlags.
+			// At the chaos level, if probe is enabled, we must also have sim.
+			bEnableSim |= bEnableProbe;
+
 			// Set the parameters on the shape
 			NewShape->SetSimEnabled(bEnableSim);
 			NewShape->SetQueryEnabled(bEnableQuery);
