#UE WorldPartition fix for loading resetting the linker when loading an actor folder package. No longer loading any packages in PostLoad.

diff --git a/Engine/Source/Runtime/Engine/Classes/Engine/Level.h b/Engine/Source/Runtime/Engine/Classes/Engine/Level.h
--- a/Engine/Source/Runtime/Engine/Classes/Engine/Level.h
+++ b/Engine/Source/Runtime/Engine/Classes/Engine/Level.h
@@ -773,9 +773,9 @@ private:
 	UPROPERTY(Transient)
 	TMap<FString, FActorFolderSet> FolderLabelToActorFolders;
 
-	/** Temporary array containing actor folder objects manually loaded from their external packages (only used while loading the level). */
+	/** Keeps track of whether we've looked for actor folder objects manually loaded from their external packages. */
 	UPROPERTY(Transient)
-	TArray<TObjectPtr<UActorFolder>> LoadedExternalActorFolders;
+	bool bNeedsUpdateForLoadedExternalActorFolders;
 #endif // #if WITH_EDITORONLY_DATA
 
 	enum class ERouteActorInitializationState : uint8
diff --git a/Engine/Source/Runtime/Engine/Private/Level.cpp b/Engine/Source/Runtime/Engine/Private/Level.cpp
--- a/Engine/Source/Runtime/Engine/Private/Level.cpp
+++ b/Engine/Source/Runtime/Engine/Private/Level.cpp
@@ -1125,12 +1125,7 @@ void ULevel::PostLoad()
 	{
 		if (!bWasDuplicated)
 		{
-			// Load all folders for this level
-			FExternalPackageHelper::LoadObjectsFromExternalPackages<UActorFolder>(this, [this](UActorFolder* LoadedFolder)
-			{
-				check(IsValid(LoadedFolder));
-				LoadedExternalActorFolders.Add(LoadedFolder);
-			});
+			bNeedsUpdateForLoadedExternalActorFolders = true;
 		}
 	}
 
@@ -2813,13 +2808,19 @@ void ULevel::FixupActorFolders()
 
 		if (IsUsingExternalObjects())
 		{
-			// At this point, LoadedExternalActorFolders are fully loaded, transfer them to the ActorFolders list.
-			for (UActorFolder* LoadedActorFolder : LoadedExternalActorFolders)
+			if (bNeedsUpdateForLoadedExternalActorFolders)
 			{
-				check(LoadedActorFolder->GetGuid().IsValid());
-				AddActorFolder(LoadedActorFolder);
+				// Load all folders for this level
+				FExternalPackageHelper::LoadObjectsFromExternalPackages<UActorFolder>(this, [this](UActorFolder* LoadedFolder)
+				{
+					check(IsValid(LoadedFolder));
+					check(LoadedFolder->GetGuid().IsValid());
+					// At this point, this ActorFolders is fully loaded, transfer it to the ActorFolders list.
+					AddActorFolder(LoadedFolder);
+				});
+				
+				bNeedsUpdateForLoadedExternalActorFolders = false;
 			}
-			LoadedExternalActorFolders.Empty();
 		}
 		else
 		{
diff --git a/Engine/Source/Runtime/Engine/Private/WorldPartition/WorldPartitionLevelHelper.cpp b/Engine/Source/Runtime/Engine/Private/WorldPartition/WorldPartitionLevelHelper.cpp
--- a/Engine/Source/Runtime/Engine/Private/WorldPartition/WorldPartitionLevelHelper.cpp
+++ b/Engine/Source/Runtime/Engine/Private/WorldPartition/WorldPartitionLevelHelper.cpp
@@ -490,7 +490,7 @@ bool FWorldPartitionLevelHelper::LoadActors(UWorld* InOuterWorld, ULevel* InDest
 				// Make sure Source level actor folder fixup was called
 				if (ContainerWorld->PersistentLevel->IsUsingActorFolders())
 				{ 
-					if (!ContainerWorld->PersistentLevel->LoadedExternalActorFolders.IsEmpty())
+					if (ContainerWorld->PersistentLevel->bNeedsUpdateForLoadedExternalActorFolders)
 					{
 						ContainerWorld->PersistentLevel->bFixupActorFoldersAtLoad = false;
 						ContainerWorld->PersistentLevel->FixupActorFolders();
