Add a simple guard to avoid reentrancy in ULevel::FixupActorFolders.

diff --git a/Engine/Source/Runtime/Engine/Classes/Engine/Level.h b/Engine/Source/Runtime/Engine/Classes/Engine/Level.h
--- a/Engine/Source/Runtime/Engine/Classes/Engine/Level.h
+++ b/Engine/Source/Runtime/Engine/Classes/Engine/Level.h
@@ -776,6 +776,10 @@ private:
 	/** Keeps track of whether we've looked for actor folder objects manually loaded from their external packages. */
 	UPROPERTY(Transient)
 	bool bNeedsUpdateForLoadedExternalActorFolders;
+
+	/** Keeps track of whether we are already fixing up actor folders or not. */
+	UPROPERTY(Transient)
+	bool bIsFixingUpActorFolders;
 #endif // #if WITH_EDITORONLY_DATA
 
 	enum class ERouteActorInitializationState : uint8
diff --git a/Engine/Source/Runtime/Engine/Private/Level.cpp b/Engine/Source/Runtime/Engine/Private/Level.cpp
--- a/Engine/Source/Runtime/Engine/Private/Level.cpp
+++ b/Engine/Source/Runtime/Engine/Private/Level.cpp
@@ -2802,6 +2802,12 @@ void ULevel::FixupActorFolders()
 		return;
 	}
 
+	TGuardValue<bool> IsFixingUpActorFoldersGuard(bIsFixingUpActorFolders, true);	
+	if (*IsFixingUpActorFoldersGuard)
+	{
+		return;
+	}
+
 	if (IsUsingActorFolders())
 	{
 		TGuardValue<bool> FixingActorFolders(GIsFixingActorFolders, true);
@@ -4553,4 +4559,4 @@ void ULevel::OnMultipleLSAsPopupDismissed()
 
 #endif // WITH_EDITOR
 
-#undef LOCTEXT_NAMESPACE
\ No newline at end of file
+#undef LOCTEXT_NAMESPACE
