non-blocking traces/overlaps should also trace into cluster unions

diff --git a/Engine/Source/Runtime/Engine/Private/Collision/SceneQuery.cpp b/Engine/Source/Runtime/Engine/Private/Collision/SceneQuery.cpp
--- a/Engine/Source/Runtime/Engine/Private/Collision/SceneQuery.cpp
+++ b/Engine/Source/Runtime/Engine/Private/Collision/SceneQuery.cpp
@@ -911,13 +911,14 @@ bool GeomOverlapMultiImp(const UWorld* World, const FPhysicsGeometry& Geom, cons
 							if (UClusterUnionComponent* ClusterUnion = Cast<UClusterUnionComponent>(OriginalOverlap.GetComponent()))
 							{
 								Result.bIsClusterUnion = true;
-								Result.bHit = ClusterUnion->OverlapComponentWithResult(GeomPose.GetTranslation(), GeomPose.GetRotation(), Geom, TraceChannel, Params, ResponseParams, ObjectParams, NewOverlaps);
+								ClusterUnion->OverlapComponentWithResult(GeomPose.GetTranslation(), GeomPose.GetRotation(), Geom, TraceChannel, Params, ResponseParams, ObjectParams, NewOverlaps);
+								Result.bHit = !NewOverlaps.IsEmpty();
 							}
 
 							return Result;
 						};
 
-						if (bHaveBlockingHit && Params.bTraceIntoSubComponents)
+						if (!OutOverlaps.IsEmpty() && Params.bTraceIntoSubComponents)
 						{
 							TArray<FOverlapResult> AllNewOverlaps;
 							TArray<int32> ClusterUnionIndices;
