Fixed queued animation notifies not being fired when running a full anim graph update on server.

If we are doing a full graph update on the server, this means we can have notifies originating from both graph AND montages, however, any queued notifies will not fire right after ticking the character pose unless the full anim graph is updated in the game thread. Therefore, when doing parallel anim graph update we have to manually dispatch notifies here, otherwise, only the latest packet's notifies will be triggered in the case where we receive multiple move packets per frame. Lastly, we need to reset the bNeedsQueuedAnimNotifies variable to not prevent notifies from firing during post evaluation.

diff --git a/Engine/Source/Runtime/Engine/Classes/Components/SkeletalMeshComponent.h b/Engine/Source/Runtime/Engine/Classes/Components/SkeletalMeshComponent.h
--- a/Engine/Source/Runtime/Engine/Classes/Components/SkeletalMeshComponent.h
+++ b/Engine/Source/Runtime/Engine/Classes/Components/SkeletalMeshComponent.h
@@ -942,6 +942,9 @@ PRAGMA_DISABLE_DEPRECATION_WARNINGS
 	const TArray<UAnimInstance*>& GetLinkedAnimInstances() const { return LinkedInstances; }
 PRAGMA_ENABLE_DEPRECATION_WARNINGS
 
+	/** If true, next time that ConditionallyDispatchQueuedAnimEvents() is called it will trigger any queued anim notifies available. */
+	ENGINE_API void AllowQueuedAnimEventsNextDispatch();
+	
 private:
 PRAGMA_DISABLE_DEPRECATION_WARNINGS
 	TArray<TObjectPtr<UAnimInstance>>& GetLinkedAnimInstances() { return LinkedInstances; }
diff --git a/Engine/Source/Runtime/Engine/Private/Components/CharacterMovementComponent.cpp b/Engine/Source/Runtime/Engine/Private/Components/CharacterMovementComponent.cpp
--- a/Engine/Source/Runtime/Engine/Private/Components/CharacterMovementComponent.cpp
+++ b/Engine/Source/Runtime/Engine/Private/Components/CharacterMovementComponent.cpp
@@ -7,6 +7,7 @@
 
 #include "GameFramework/CharacterMovementComponent.h"
 #include "Animation/AnimMontage.h"
+#include "Animation/AnimInstance.h"
 #include "EngineStats.h"
 #include "AI/NavigationSystemBase.h"
 #include "AI/Navigation/NavigationDataInterface.h"
@@ -10082,11 +10083,26 @@ void UCharacterMovementComponent::MoveAutonomous
 		}
 		// TODO: SaveBaseLocation() in case tick moves us?
 
-		if (!CharacterMovementCVars::EnableQueuedAnimEventsOnServer || CharacterOwner->GetMesh()->ShouldOnlyTickMontages(DeltaTime))
+		USkeletalMeshComponent* OwnerMesh = CharacterOwner->GetMesh();
+		check(OwnerMesh != nullptr)
+		
+		if (CharacterMovementCVars::EnableQueuedAnimEventsOnServer)
+		{
+			if (OwnerMesh->VisibilityBasedAnimTickOption <= EVisibilityBasedAnimTickOption::AlwaysTickPose && OwnerMesh->GetAnimInstance()->NeedsUpdate())
+			{
+				// If we are doing a full graph update on the server but its doing a parallel update,
+				// trigger events right away since these are notifies queued from the montage update and we could be receiving multiple ServerMoves per frame.
+				OwnerMesh->ConditionallyDispatchQueuedAnimEvents();
+				OwnerMesh->AllowQueuedAnimEventsNextDispatch();
+			}
+		}
+		else
 		{
-			// If we're not doing a full anim graph update on the server, 
-			// trigger events right away, as we could be receiving multiple ServerMoves per frame.
-			CharacterOwner->GetMesh()->ConditionallyDispatchQueuedAnimEvents();
+			// Revert back to old behavior if wanted/needed.
+			if (OwnerMesh->ShouldOnlyTickMontages(DeltaTime))
+			{
+				OwnerMesh->ConditionallyDispatchQueuedAnimEvents();
+			}
 		}
 	}
 
diff --git a/Engine/Source/Runtime/Engine/Private/Components/SkeletalMeshComponent.cpp b/Engine/Source/Runtime/Engine/Private/Components/SkeletalMeshComponent.cpp
--- a/Engine/Source/Runtime/Engine/Private/Components/SkeletalMeshComponent.cpp
+++ b/Engine/Source/Runtime/Engine/Private/Components/SkeletalMeshComponent.cpp
@@ -3208,6 +3208,11 @@ void USkeletalMeshComponent::ResetLinkedAnimInstances()
 	LinkedInstances.Reset();
 }
 
+void USkeletalMeshComponent::AllowQueuedAnimEventsNextDispatch()
+{
+	bNeedsQueuedAnimEventsDispatched = true;
+}
+
 UAnimInstance* USkeletalMeshComponent::GetLinkedAnimGraphInstanceByTag(FName InName) const
 {
 	if(AnimScriptInstance)
