PLAY-10177 Held object replication

Perform the placement sweep from above the pawns head and towards the front of the pawn and then down towards the ground.
Also adjust the UMeshComponent bounds radius used in the sweep to better represent the size of the object.
Also re-implemented ETeleportType::ResetPhysics on SetActorLocation when placing an object since else it could fail to set the location and end up clipping with the pawn.

PhysicsObjectPreset assets are edited to enabled Predictive Interpolation physics replication.

diff --git a/Engine/Source/Runtime/Engine/Private/ActorReplication.cpp b/Engine/Source/Runtime/Engine/Private/ActorReplication.cpp
--- a/Engine/Source/Runtime/Engine/Private/ActorReplication.cpp
+++ b/Engine/Source/Runtime/Engine/Private/ActorReplication.cpp
@@ -193,7 +193,24 @@ void AActor::OnRep_ReplicatedMovement()
 #if !(UE_BUILD_SHIPPING || UE_BUILD_TEST)
 		if (CVarDrawDebugRepMovement->GetInt() > 0)
 		{
-			DrawDebugCapsule(GetWorld(), LocalRepMovement.Location, GetSimpleCollisionHalfHeight(), GetSimpleCollisionRadius(), LocalRepMovement.Rotation.Quaternion(), FColor(100, 255, 100), false, 1.f);
+			FColor DebugColor = FColor::Green;
+			if (LocalRepMovement.bRepPhysics)
+			{
+				switch (GetPhysicsReplicationMode())
+				{
+				case EPhysicsReplicationMode::PredictiveInterpolation:
+					DebugColor = FColor::Yellow;
+					break;
+				case EPhysicsReplicationMode::Resimulation:
+					DebugColor = FColor::Red;
+					break;
+				case EPhysicsReplicationMode::Default:
+				default:
+					DebugColor = FColor::Cyan;
+					break;
+				}
+			}
+			DrawDebugCapsule(GetWorld(), LocalRepMovement.Location, GetSimpleCollisionHalfHeight(), GetSimpleCollisionRadius(), LocalRepMovement.Rotation.Quaternion(), DebugColor, false, 1.f);
 		}
 #endif
 
