PR #10751: Fixed crash in FPreviewScene if CVS.bDefaultLighting=false

diff --git a/Engine/Source/Runtime/Engine/Private/PreviewScene.cpp b/Engine/Source/Runtime/Engine/Private/PreviewScene.cpp
--- a/Engine/Source/Runtime/Engine/Private/PreviewScene.cpp
+++ b/Engine/Source/Runtime/Engine/Private/PreviewScene.cpp
@@ -217,52 +217,72 @@ void FPreviewScene::ClearLineBatcher()
 /** Accessor for finding the current direction of the preview scene's DirectionalLight. */
 FRotator FPreviewScene::GetLightDirection()
 {
-	return DirectionalLight->GetComponentTransform().GetUnitAxis( EAxis::X ).Rotation();
+	if (DirectionalLight != NULL)
+	{
+		return DirectionalLight->GetComponentTransform().GetUnitAxis( EAxis::X ).Rotation();
+	}
+
+	return FRotator::ZeroRotator;
 }
 
 /** Function for modifying the current direction of the preview scene's DirectionalLight. */
 void FPreviewScene::SetLightDirection(const FRotator& InLightDir)
 {
+	if (DirectionalLight != NULL)
+	{
 #if WITH_EDITOR
-	DirectionalLight->PreEditChange(NULL);
+		DirectionalLight->PreEditChange(NULL);
 #endif // WITH_EDITOR
-	DirectionalLight->SetAbsolute(true, true, true);
-	DirectionalLight->SetRelativeRotation(InLightDir);
+		DirectionalLight->SetAbsolute(true, true, true);
+		DirectionalLight->SetRelativeRotation(InLightDir);
 #if WITH_EDITOR
-	DirectionalLight->PostEditChange();
+		DirectionalLight->PostEditChange();
 #endif // WITH_EDITOR
+	}
 }
 
 void FPreviewScene::SetLightBrightness(float LightBrightness)
 {
+	if (DirectionalLight != NULL)
+	{
 #if WITH_EDITOR
-	DirectionalLight->PreEditChange(NULL);
+		DirectionalLight->PreEditChange(NULL);
 #endif // WITH_EDITOR
-	DirectionalLight->Intensity = LightBrightness;
+		DirectionalLight->Intensity = LightBrightness;
 #if WITH_EDITOR
-	DirectionalLight->PostEditChange();
+		DirectionalLight->PostEditChange();
 #endif // WITH_EDITOR
+	}
 }
 
 void FPreviewScene::SetLightColor(const FColor& LightColor)
 {
+	if (DirectionalLight != NULL)
+	{
 #if WITH_EDITOR
-	DirectionalLight->PreEditChange(NULL);
+		DirectionalLight->PreEditChange(NULL);
 #endif // WITH_EDITOR
-	DirectionalLight->LightColor = LightColor;
+		DirectionalLight->LightColor = LightColor;
 #if WITH_EDITOR
-	DirectionalLight->PostEditChange();
+		DirectionalLight->PostEditChange();
 #endif // WITH_EDITOR
+	}
 }
 
 void FPreviewScene::SetSkyBrightness(float SkyBrightness)
 {
-	SkyLight->SetIntensity(SkyBrightness);
+	if (SkyLight != NULL)
+	{
+		SkyLight->SetIntensity(SkyBrightness);
+	}
 }
 
 void FPreviewScene::SetSkyCubemap(UTextureCube* Cubemap)
 {
-	SkyLight->SetCubemap(Cubemap);
+	if (SkyLight != NULL)
+	{
+		SkyLight->SetCubemap(Cubemap);
+	}
 }
 
 void FPreviewScene::LoadSettings(const TCHAR* Section)
diff --git a/Engine/Source/Runtime/Engine/Public/PreviewScene.h b/Engine/Source/Runtime/Engine/Public/PreviewScene.h
--- a/Engine/Source/Runtime/Engine/Public/PreviewScene.h
+++ b/Engine/Source/Runtime/Engine/Public/PreviewScene.h
@@ -112,8 +112,8 @@ public:
 	ENGINE_API void LoadSettings(const TCHAR* Section);
 	ENGINE_API void SaveSettings(const TCHAR* Section);
 
-	class UDirectionalLightComponent* DirectionalLight;
-	class USkyLightComponent* SkyLight;
+	class UDirectionalLightComponent* DirectionalLight = nullptr;;
+	class USkyLightComponent* SkyLight = nullptr;;
 
 private:
 
