PR #10753: Expose UCameraModifier::bPendingDisable

diff --git a/Engine/Source/Runtime/Engine/Classes/Camera/CameraModifier.h b/Engine/Source/Runtime/Engine/Classes/Camera/CameraModifier.h
--- a/Engine/Source/Runtime/Engine/Classes/Camera/CameraModifier.h
+++ b/Engine/Source/Runtime/Engine/Classes/Camera/CameraModifier.h
@@ -122,6 +122,10 @@ public:
 	UFUNCTION(BlueprintCallable, Category = CameraModifier)
 	ENGINE_API virtual bool IsDisabled() const;
 	
+	/** @return Returns true if modifier is pending disable, false otherwise. */
+	UFUNCTION(BlueprintCallable, Category = CameraModifier)
+	ENGINE_API virtual bool IsPendingDisable() const;
+	
 	/** @return Returns the actor the camera is currently viewing. */
 	UFUNCTION(BlueprintCallable, Category = CameraModifier)
 	ENGINE_API virtual AActor* GetViewTarget() const;
diff --git a/Engine/Source/Runtime/Engine/Private/Camera/CameraModifier.cpp b/Engine/Source/Runtime/Engine/Private/Camera/CameraModifier.cpp
--- a/Engine/Source/Runtime/Engine/Private/Camera/CameraModifier.cpp
+++ b/Engine/Source/Runtime/Engine/Private/Camera/CameraModifier.cpp
@@ -93,6 +93,11 @@ bool UCameraModifier::IsDisabled() const
 	return bDisabled;
 }
 
+bool UCameraModifier::IsPendingDisable() const
+{
+	return bPendingDisable;
+}
+
 AActor* UCameraModifier::GetViewTarget() const
 {
 	return CameraOwner ? CameraOwner->GetViewTarget() : nullptr;
