Enable per instance LOD on ISM by default.
Change from old behavior to LOD all instances on the whole ISM bounds.

diff --git a/Engine/Source/Runtime/Engine/Private/InstancedStaticMesh.cpp b/Engine/Source/Runtime/Engine/Private/InstancedStaticMesh.cpp
--- a/Engine/Source/Runtime/Engine/Private/InstancedStaticMesh.cpp
+++ b/Engine/Source/Runtime/Engine/Private/InstancedStaticMesh.cpp
@@ -2559,6 +2559,8 @@ UInstancedStaticMeshComponent::UInstancedStaticMeshComponent(const FObjectInitia
 	bDisallowMeshPaintPerInstance = true;
 	bMultiBodyOverlap = true;
 
+	bUseGpuLodSelection = true;
+
 #if STATS
 	{
 		UObject const* StatObject = this->AdditionalStatObject();
