PR #10712: Make the RemoveInstances method fail fast in case the InstancesToRemove array is empty

diff --git a/Engine/Source/Runtime/Engine/Private/InstancedStaticMesh.cpp b/Engine/Source/Runtime/Engine/Private/InstancedStaticMesh.cpp
--- a/Engine/Source/Runtime/Engine/Private/InstancedStaticMesh.cpp
+++ b/Engine/Source/Runtime/Engine/Private/InstancedStaticMesh.cpp
@@ -3902,6 +3902,11 @@ bool UInstancedStaticMeshComponent::RemoveInstances(const TArray<int32>& Instanc
 
 bool UInstancedStaticMeshComponent::RemoveInstances(const TArray<int32>& InstancesToRemove, bool bInstanceArrayAlreadySortedInReverseOrder)
 {
+	if(InstancesToRemove.IsEmpty())
+	{
+		return false;
+	}
+	
 	auto RemoveInstancedWithSortedArray = [this](const TArray<int32>& SortedInstancesToRemove) -> bool
 	{
 		if (!PerInstanceSMData.IsValidIndex(SortedInstancesToRemove[0]) || !PerInstanceSMData.IsValidIndex(SortedInstancesToRemove.Last()))
