Add bUseGpuLodSelection to ISM component descriptors.

diff --git a/Engine/Source/Runtime/Engine/Private/ISMPartition/ISMComponentDescriptor.cpp b/Engine/Source/Runtime/Engine/Private/ISMPartition/ISMComponentDescriptor.cpp
--- a/Engine/Source/Runtime/Engine/Private/ISMPartition/ISMComponentDescriptor.cpp
+++ b/Engine/Source/Runtime/Engine/Private/ISMPartition/ISMComponentDescriptor.cpp
@@ -125,6 +125,7 @@ void FISMComponentDescriptorBase::InitFrom(const UStaticMeshComponent* Template,
 	{
 		InstanceStartCullDistance = ISMTemplate->InstanceStartCullDistance;
 		InstanceEndCullDistance = ISMTemplate->InstanceEndCullDistance;
+		bUseGpuLodSelection = ISMTemplate->bUseGpuLodSelection;
 
 		// HISM Specific
 		if (const UHierarchicalInstancedStaticMeshComponent* HISMTemplate = Cast<UHierarchicalInstancedStaticMeshComponent>(Template))
@@ -212,6 +213,7 @@ bool FISMComponentDescriptorBase::operator==(const FISMComponentDescriptorBase&
 	bVisibleInRayTracing == Other.bVisibleInRayTracing &&
 	bEvaluateWorldPositionOffset == Other.bEvaluateWorldPositionOffset &&
 	bReverseCulling == Other.bReverseCulling &&
+	bUseGpuLodSelection == Other.bUseGpuLodSelection &&
 	bUseDefaultCollision == Other.bUseDefaultCollision &&
 	bGenerateOverlapEvents == Other.bGenerateOverlapEvents &&
 	WorldPositionOffsetDisableDistance == Other.WorldPositionOffsetDisableDistance &&
@@ -324,6 +326,7 @@ void FISMComponentDescriptorBase::InitComponent(UInstancedStaticMeshComponent* I
 	ISMComponent->bVisibleInRayTracing = bVisibleInRayTracing;
 	ISMComponent->bEvaluateWorldPositionOffset = bEvaluateWorldPositionOffset;
 	ISMComponent->bReverseCulling = bReverseCulling;
+	ISMComponent->bUseGpuLodSelection = bUseGpuLodSelection;
 	ISMComponent->bUseDefaultCollision = bUseDefaultCollision;
 	ISMComponent->SetGenerateOverlapEvents(bGenerateOverlapEvents);
 	ISMComponent->WorldPositionOffsetDisableDistance = WorldPositionOffsetDisableDistance;
diff --git a/Engine/Source/Runtime/Engine/Public/ISMPartition/ISMComponentDescriptor.h b/Engine/Source/Runtime/Engine/Public/ISMPartition/ISMComponentDescriptor.h
--- a/Engine/Source/Runtime/Engine/Public/ISMPartition/ISMComponentDescriptor.h
+++ b/Engine/Source/Runtime/Engine/Public/ISMPartition/ISMComponentDescriptor.h
@@ -152,6 +152,9 @@ public:
 	UPROPERTY(EditAnywhere, Category = "Component Settings")
 	uint8 bReverseCulling : 1;
 
+	UPROPERTY(EditAnywhere, Category = "Component Settings")
+	uint8 bUseGpuLodSelection : 1;
+
 #if WITH_EDITORONLY_DATA
 	UPROPERTY(EditAnywhere, Category = "Component Settings")
 	uint8 bIncludeInHLOD : 1;
