merged 5 texture render resource tasks into 2

diff --git a/Engine/Source/Runtime/Engine/Private/Texture.cpp b/Engine/Source/Runtime/Engine/Private/Texture.cpp
--- a/Engine/Source/Runtime/Engine/Private/Texture.cpp
+++ b/Engine/Source/Runtime/Engine/Private/Texture.cpp
@@ -253,6 +253,8 @@ void UTexture::ReleaseResource()
 {
 	if (PrivateResource)
 	{
+		check(!IsInActualRenderingThread() && !IsInRHIThread());
+
 		UnlinkStreaming();
 
 		// When using PlatformData, the resource shouldn't be released before it is initialized to prevent threading issues
@@ -264,14 +266,14 @@ void UTexture::ReleaseResource()
 
 		CachedSRRState.Clear();
 
-		FTextureResource* ToDelete = PrivateResource;
 		// Free the resource.
-		SetResource(nullptr);
-		ENQUEUE_RENDER_COMMAND(DeleteResource)([ToDelete](FRHICommandListImmediate& RHICmdList)
+		ENQUEUE_RENDER_COMMAND(DeleteResource)([this, ToDelete = PrivateResource](FRHICommandListImmediate& RHICmdList)
 		{
+			PrivateResourceRenderThread = nullptr;
 			ToDelete->ReleaseResource();
 			delete ToDelete;
 		});
+		PrivateResource = nullptr;
 	}
 }
 
@@ -344,7 +346,8 @@ void UTexture::UpdateResource()
 			}
 		}
 
-		SetResource(NewResource);
+		check(!IsInActualRenderingThread() && !IsInRHIThread());
+
 		if (NewResource)
 		{
 			LLM_SCOPE(ELLMTag::Textures);
@@ -359,17 +362,24 @@ void UTexture::UpdateResource()
 				}
 			}
 
+			PrivateResource = NewResource;
+			NewResource->SetOwnerName(FName(GetPathName()));
+
 			// Init the texture reference, which needs to be set from a render command, since TextureReference.TextureReferenceRHI is gamethread coherent.
 			ENQUEUE_RENDER_COMMAND(SetTextureReference)([this, NewResource](FRHICommandListImmediate& RHICmdList)
 			{
+				PrivateResourceRenderThread = NewResource;
 				NewResource->SetTextureReference(TextureReference.TextureReferenceRHI);
+				NewResource->InitResource(RHICmdList);
 			});
 
-			NewResource->SetOwnerName(FName(GetPathName()));
-			BeginInitResource(NewResource);
 			// Now that the resource is ready for streaming, bind it to the streamer.
 			LinkStreaming();
 		}
+		else
+		{
+			SetResource(NewResource);
+		}
 	}
 }
 
