InstancedStaticMesh: Fix regression caused by CL 27024751
- Make sure to only follow new serialization path when Archetype is a Blueprint Archetype

#rnx

diff --git a/Engine/Source/Runtime/Engine/Private/InstancedStaticMesh.cpp b/Engine/Source/Runtime/Engine/Private/InstancedStaticMesh.cpp
--- a/Engine/Source/Runtime/Engine/Private/InstancedStaticMesh.cpp
+++ b/Engine/Source/Runtime/Engine/Private/InstancedStaticMesh.cpp
@@ -3478,7 +3478,7 @@ void UInstancedStaticMeshComponent::Serialize(FArchive& Ar)
 
 	// Inherit properties when bEditableWhenInherited == true (when the component isn't a template and we are persisting data)
 	const UInstancedStaticMeshComponent* Archetype = Cast<UInstancedStaticMeshComponent>(GetArchetype());
-	const bool bInheritSkipSerializationProperties = !bEditableWhenInherited && Archetype && Ar.IsPersistent() && !IsTemplate();
+	const bool bInheritSkipSerializationProperties = !bEditableWhenInherited && Archetype && Archetype->IsInBlueprint() && Ar.IsPersistent() && !IsTemplate();
 	
 
 	// Check if we need have SkipSerialization property data to load/save
