[Backout] - CL27095182
Original CL Desc
-----------------------------------------------------------------
#UE WorldPartition fix for loading resetting the linker when loading an actor folder package. No longer loading any packages in PostLoad.

diff --git a/Engine/Source/Runtime/Engine/Classes/Engine/Level.h b/Engine/Source/Runtime/Engine/Classes/Engine/Level.h
--- a/Engine/Source/Runtime/Engine/Classes/Engine/Level.h
+++ b/Engine/Source/Runtime/Engine/Classes/Engine/Level.h
@@ -773,13 +773,9 @@ private:
 	UPROPERTY(Transient)
 	TMap<FString, FActorFolderSet> FolderLabelToActorFolders;
 
-	/** Keeps track of whether we've looked for actor folder objects manually loaded from their external packages. */
+	/** Temporary array containing actor folder objects manually loaded from their external packages (only used while loading the level). */
 	UPROPERTY(Transient)
-	bool bNeedsUpdateForLoadedExternalActorFolders;
-
-	/** Keeps track of whether we are already fixing up actor folders or not. */
-	UPROPERTY(Transient)
-	bool bIsFixingUpActorFolders;
+	TArray<TObjectPtr<UActorFolder>> LoadedExternalActorFolders;
 #endif // #if WITH_EDITORONLY_DATA
 
 	enum class ERouteActorInitializationState : uint8
diff --git a/Engine/Source/Runtime/Engine/Private/Level.cpp b/Engine/Source/Runtime/Engine/Private/Level.cpp
--- a/Engine/Source/Runtime/Engine/Private/Level.cpp
+++ b/Engine/Source/Runtime/Engine/Private/Level.cpp
@@ -1125,7 +1125,12 @@ void ULevel::PostLoad()
 	{
 		if (!bWasDuplicated)
 		{
-			bNeedsUpdateForLoadedExternalActorFolders = true;
+			// Load all folders for this level
+			FExternalPackageHelper::LoadObjectsFromExternalPackages<UActorFolder>(this, [this](UActorFolder* LoadedFolder)
+			{
+				check(IsValid(LoadedFolder));
+				LoadedExternalActorFolders.Add(LoadedFolder);
+			});
 		}
 	}
 
@@ -2802,31 +2807,19 @@ void ULevel::FixupActorFolders()
 		return;
 	}
 
-	TGuardValue<bool> IsFixingUpActorFoldersGuard(bIsFixingUpActorFolders, true);	
-	if (*IsFixingUpActorFoldersGuard)
-	{
-		return;
-	}
-
 	if (IsUsingActorFolders())
 	{
 		TGuardValue<bool> FixingActorFolders(GIsFixingActorFolders, true);
 
 		if (IsUsingExternalObjects())
 		{
-			if (bNeedsUpdateForLoadedExternalActorFolders)
+			// At this point, LoadedExternalActorFolders are fully loaded, transfer them to the ActorFolders list.
+			for (UActorFolder* LoadedActorFolder : LoadedExternalActorFolders)
 			{
-				// Load all folders for this level
-				FExternalPackageHelper::LoadObjectsFromExternalPackages<UActorFolder>(this, [this](UActorFolder* LoadedFolder)
-				{
-					check(IsValid(LoadedFolder));
-					check(LoadedFolder->GetGuid().IsValid());
-					// At this point, this ActorFolders is fully loaded, transfer it to the ActorFolders list.
-					AddActorFolder(LoadedFolder);
-				});
-				
-				bNeedsUpdateForLoadedExternalActorFolders = false;
+				check(LoadedActorFolder->GetGuid().IsValid());
+				AddActorFolder(LoadedActorFolder);
 			}
+			LoadedExternalActorFolders.Empty();
 		}
 		else
 		{
diff --git a/Engine/Source/Runtime/Engine/Private/WorldPartition/WorldPartitionLevelHelper.cpp b/Engine/Source/Runtime/Engine/Private/WorldPartition/WorldPartitionLevelHelper.cpp
--- a/Engine/Source/Runtime/Engine/Private/WorldPartition/WorldPartitionLevelHelper.cpp
+++ b/Engine/Source/Runtime/Engine/Private/WorldPartition/WorldPartitionLevelHelper.cpp
@@ -490,7 +490,7 @@ bool FWorldPartitionLevelHelper::LoadActors(UWorld* InOuterWorld, ULevel* InDest
 				// Make sure Source level actor folder fixup was called
 				if (ContainerWorld->PersistentLevel->IsUsingActorFolders())
 				{ 
-					if (ContainerWorld->PersistentLevel->bNeedsUpdateForLoadedExternalActorFolders)
+					if (!ContainerWorld->PersistentLevel->LoadedExternalActorFolders.IsEmpty())
 					{
 						ContainerWorld->PersistentLevel->bFixupActorFoldersAtLoad = false;
 						ContainerWorld->PersistentLevel->FixupActorFolders();
