Fix bug where the Sound Cue Enveloper Node incorrectly calculates the Current Loop Count

diff --git a/Engine/Source/Runtime/Engine/Private/SoundNodeEnveloper.cpp b/Engine/Source/Runtime/Engine/Private/SoundNodeEnveloper.cpp
--- a/Engine/Source/Runtime/Engine/Private/SoundNodeEnveloper.cpp
+++ b/Engine/Source/Runtime/Engine/Private/SoundNodeEnveloper.cpp
@@ -87,7 +87,7 @@ void USoundNodeEnveloper::ParseNodes( FAudioDevice* AudioDevice, const UPTRINT N
 		}
 
 		float LoopDuration = LoopEnd - LoopStart;
-		int32 CurrentLoopCount = (int32)(PlayTime - LoopStart) / FMath::Min(UE_SMALL_NUMBER, LoopDuration);
+		int32 CurrentLoopCount = (int32)(PlayTime - LoopStart) / FMath::Max(UE_SMALL_NUMBER, LoopDuration);
 		PlayTime -= CurrentLoopCount * LoopDuration;
 
 		if( CurrentLoopCount == LoopCount && !bLoopIndefinitely && LoopCount != 0 )
