Fixed crash when doing full graph update and dispatching anim events in MoveAutonomous() due to AnimInstance being able to be null.

#rnx

diff --git a/Engine/Source/Runtime/Engine/Private/Components/CharacterMovementComponent.cpp b/Engine/Source/Runtime/Engine/Private/Components/CharacterMovementComponent.cpp
--- a/Engine/Source/Runtime/Engine/Private/Components/CharacterMovementComponent.cpp
+++ b/Engine/Source/Runtime/Engine/Private/Components/CharacterMovementComponent.cpp
@@ -10088,12 +10088,15 @@ void UCharacterMovementComponent::MoveAutonomous
 		
 		if (CharacterMovementCVars::EnableQueuedAnimEventsOnServer)
 		{
-			if (OwnerMesh->VisibilityBasedAnimTickOption <= EVisibilityBasedAnimTickOption::AlwaysTickPose && OwnerMesh->GetAnimInstance()->NeedsUpdate())
+			if (const UAnimInstance* AnimInstance = OwnerMesh->GetAnimInstance())
 			{
-				// If we are doing a full graph update on the server but its doing a parallel update,
-				// trigger events right away since these are notifies queued from the montage update and we could be receiving multiple ServerMoves per frame.
-				OwnerMesh->ConditionallyDispatchQueuedAnimEvents();
-				OwnerMesh->AllowQueuedAnimEventsNextDispatch();
+				if (OwnerMesh->VisibilityBasedAnimTickOption <= EVisibilityBasedAnimTickOption::AlwaysTickPose && AnimInstance->NeedsUpdate())
+				{
+					// If we are doing a full graph update on the server but its doing a parallel update,
+					// trigger events right away since these are notifies queued from the montage update and we could be receiving multiple ServerMoves per frame.
+					OwnerMesh->ConditionallyDispatchQueuedAnimEvents();
+					OwnerMesh->AllowQueuedAnimEventsNextDispatch();
+				}
 			}
 		}
 		else
