Fix for wrong skeletal meshes motion vectors

#rnx
#tests fix tested on all platforms

diff --git a/Engine/Source/Runtime/Engine/Private/SkeletalMesh.cpp b/Engine/Source/Runtime/Engine/Private/SkeletalMesh.cpp
--- a/Engine/Source/Runtime/Engine/Private/SkeletalMesh.cpp
+++ b/Engine/Source/Runtime/Engine/Private/SkeletalMesh.cpp
@@ -6054,7 +6054,7 @@ void FSkeletalMeshSceneProxy::GetMeshElementsConditionallySelectable(const TArra
 	{
 		return;
 	}	
-	MeshObject->PreGDMECallback(ViewFamily.Scene->GetGPUSkinCache(), ViewFamily.FrameNumber);
+	MeshObject->PreGDMECallback(ViewFamily.Scene->GetGPUSkinCache(), ViewFamily.FrameCounter);
 
 	const FEngineShowFlags& EngineShowFlags = ViewFamily.EngineShowFlags;
 
