Make texture object picker filter only for textures that have valid source.

diff --git a/Engine/Source/Runtime/Engine/Classes/Engine/Texture.h b/Engine/Source/Runtime/Engine/Classes/Engine/Texture.h
--- a/Engine/Source/Runtime/Engine/Classes/Engine/Texture.h
+++ b/Engine/Source/Runtime/Engine/Classes/Engine/Texture.h
@@ -1248,7 +1248,7 @@ public:
 	 * Make sure the normal map has at least as many mips as this texture.
 	 */
 	UE_DEPRECATED(5.3, "Use GetCompositeTexture() and SetCompositeTexture() instead.")
-	UPROPERTY(EditAnywhere, BlueprintReadWrite, Category=Compositing, meta = (AllowPrivateAccess), Setter = SetCompositeTexture, Getter = GetCompositeTexture)
+	UPROPERTY(EditAnywhere, BlueprintReadWrite, Category=Compositing, meta = (AllowPrivateAccess, RequiredAssetDataTags = "IsSourceValid=True"), Setter = SetCompositeTexture, Getter = GetCompositeTexture)
 	TObjectPtr<class UTexture> CompositeTexture;
 
 private:
diff --git a/Engine/Source/Runtime/Engine/Classes/Engine/Texture2DArray.h b/Engine/Source/Runtime/Engine/Classes/Engine/Texture2DArray.h
--- a/Engine/Source/Runtime/Engine/Classes/Engine/Texture2DArray.h
+++ b/Engine/Source/Runtime/Engine/Classes/Engine/Texture2DArray.h
@@ -88,7 +88,7 @@ public:
 
 #if WITH_EDITORONLY_DATA
 	/** Add Textures*/
-	UPROPERTY(EditAnywhere, Category = Source2D, meta = (DisplayName = "Source Textures", EditCondition = bSourceGeneratedFromSourceTexturesArray, EditConditionHides, HideEditConditionToggle))
+	UPROPERTY(EditAnywhere, Category = Source2D, meta = (DisplayName = "Source Textures", EditCondition = bSourceGeneratedFromSourceTexturesArray, EditConditionHides, HideEditConditionToggle, RequiredAssetDataTags = "IsSourceValid=True"))
 	TArray<TObjectPtr<UTexture2D>> SourceTextures;
 
 	/** 
diff --git a/Engine/Source/Runtime/Engine/Classes/Engine/TextureCubeArray.h b/Engine/Source/Runtime/Engine/Classes/Engine/TextureCubeArray.h
--- a/Engine/Source/Runtime/Engine/Classes/Engine/TextureCubeArray.h
+++ b/Engine/Source/Runtime/Engine/Classes/Engine/TextureCubeArray.h
@@ -62,7 +62,7 @@ public:
 
 #if WITH_EDITORONLY_DATA
 	/** Add Textures*/
-	UPROPERTY(EditAnywhere, Category = SourceCube, meta = (DisplayName = "Source Textures", EditCondition = bSourceGeneratedFromSourceTexturesArray, EditConditionHides, HideEditConditionToggle))
+	UPROPERTY(EditAnywhere, Category = SourceCube, meta = (DisplayName = "Source Textures", EditCondition = bSourceGeneratedFromSourceTexturesArray, EditConditionHides, HideEditConditionToggle, RequiredAssetDataTags = "IsSourceValid=True"))
 	TArray<TObjectPtr<UTextureCube>> SourceTextures;
 
 	/**
diff --git a/Engine/Source/Runtime/Engine/Private/Texture.cpp b/Engine/Source/Runtime/Engine/Private/Texture.cpp
--- a/Engine/Source/Runtime/Engine/Private/Texture.cpp
+++ b/Engine/Source/Runtime/Engine/Private/Texture.cpp
@@ -1279,6 +1279,8 @@ void UTexture::GetAssetRegistryTags(TArray<FAssetRegistryTag>& OutTags) const
 		StaticEnum<ETextureSourceFormat>()->GetDisplayNameTextByValue(Source.GetFormat()).ToString(),
 		FAssetRegistryTag::TT_Alphabetical));
 	
+	OutTags.Add(FAssetRegistryTag("IsSourceValid", Source.IsValid() ? TEXT("True") : TEXT("False"), FAssetRegistryTag::TT_Alphabetical));
+
 	Super::GetAssetRegistryTags(OutTags);
 }
 #endif
