Shader Compiler:  Job cache pending list thread safety fix.  In LinkTailAtomic, we need to set Job.NextLink before the atomic tail swap, as there is a chance of another thread updating Job.NextLink before the function completes, causing items to be lost from the list, leading to downstream asserts in GetPendingJobs due to NumPendingJobs not matching the number of items in the list.  Similar to the existing code in LinkHeadAtomic.

#rnx

diff --git a/Engine/Source/Runtime/Engine/Private/ShaderCompiler/ShaderCompiler.cpp b/Engine/Source/Runtime/Engine/Private/ShaderCompiler/ShaderCompiler.cpp
--- a/Engine/Source/Runtime/Engine/Private/ShaderCompiler/ShaderCompiler.cpp
+++ b/Engine/Source/Runtime/Engine/Private/ShaderCompiler/ShaderCompiler.cpp
@@ -771,9 +771,12 @@ private:
 	 */
 	FORCEINLINE static void LinkTailAtomic(FShaderCommonCompileJob& Job, FShaderCommonCompileJob**& Tail)
 	{
+		// It's important that NextLink is set before the InterlockedExchange, as a subsequent Tail pointer exchange could write
+		// another item and need to update NextLink for this item before this function completes.
+		Job.NextLink = nullptr;
+
 		FShaderCommonCompileJob** OldTail = (FShaderCommonCompileJob**)FPlatformAtomics::InterlockedExchange((PTRINT*)&Tail, (PTRINT)&Job.NextLink);
 		Job.PrevLink = OldTail;
-		Job.NextLink = nullptr;
 
 		// Update previous tail's next pointer (or OldTail may be pointing at Head if list was empty)
 		*OldTail = (FShaderCommonCompileJob*)&Job;
