UModel Cook Determinism: Calculate lighting guid deterministically, and only update it when the UModel is saved because it is now too expensive to run on every call to Modify.
#rnx

diff --git a/Engine/Source/Runtime/Engine/Private/Model.cpp b/Engine/Source/Runtime/Engine/Private/Model.cpp
--- a/Engine/Source/Runtime/Engine/Private/Model.cpp
+++ b/Engine/Source/Runtime/Engine/Private/Model.cpp
@@ -5,11 +5,14 @@
 =============================================================================*/
 
 #include "Model.h"
+
 #include "Containers/TransArray.h"
 #include "EngineUtils.h"
 #include "Engine/Polys.h"
+#include "Hash/Blake3.h"
 #include "StaticLighting.h"
 #include "UObject/GarbageCollectionSchema.h"
+#include "UObject/ObjectSaveContext.h"
 
 float UModel::BSPTexelScale = 100.0f;
 
@@ -43,6 +46,37 @@ bool FBspSurf::IsHiddenEdAtStartup() const
 /*-----------------------------------------------------------------------------
 	Struct serializers.
 -----------------------------------------------------------------------------*/
+#if WITH_EDITOR
+template <typename T>
+static void ModelUpdateHash(FBlake3& Builder, const TObjectPtr<T>& Object)
+{
+	// We don't merge other UObject data into the lighting guid, the lighting guid only handles properties of the UModel changing
+	// just record the path for each UObject
+	FString PathName = Object.GetPathName();
+	if (PathName.IsEmpty())
+	{
+		Builder.Update(TEXT(""), 0);
+	}
+	else
+	{
+		Builder.Update(*PathName, PathName.Len() * sizeof(PathName[0]));
+	}
+};
+static void ModelUpdateHash(FBlake3& Builder, UObject* Object)
+{
+	// We don't merge other UObject data into the lighting guid, the lighting guid only handles properties of the UModel changing
+	// just record the path for each UObject
+	FString PathName = Object ? Object->GetPathName() : TEXT("");
+	if (PathName.IsEmpty())
+	{
+		Builder.Update(TEXT(""), 0);
+	}
+	else
+	{
+		Builder.Update(*PathName, PathName.Len() * sizeof(PathName[0]));
+	}
+};
+#endif
 
 FArchive& operator<<( FArchive& Ar, FBspSurf& Surf )
 {
@@ -68,6 +102,23 @@ FArchive& operator<<( FArchive& Ar, FBspSurf& Surf )
 	return Ar;
 }
 
+#if WITH_EDITOR
+void UpdateHash(FBlake3& Builder, const FBspSurf& Surf)
+{
+	ModelUpdateHash(Builder, Surf.Material);
+	Builder.Update(&Surf.PolyFlags, sizeof(Surf.PolyFlags));
+	Builder.Update(&Surf.pBase, sizeof(Surf.pBase));
+	Builder.Update(&Surf.vNormal, sizeof(Surf.vNormal));
+	Builder.Update(&Surf.vTextureU, sizeof(Surf.vTextureU));
+	Builder.Update(&Surf.vTextureV, sizeof(Surf.vTextureV));
+	Builder.Update(&Surf.iBrushPoly, sizeof(Surf.iBrushPoly));
+	ModelUpdateHash(Builder, Surf.Actor);
+	Builder.Update(&Surf.Plane, sizeof(Surf.Plane));
+	Builder.Update(&Surf.LightMapScale, sizeof(Surf.LightMapScale));
+	Builder.Update(&Surf.iLightmassIndex, sizeof(Surf.iLightmassIndex));
+}
+#endif
+
 void FBspSurf::AddReferencedObjects( FReferenceCollector& Collector )
 {
 	Collector.AddReferencedObject( Material );
@@ -400,6 +451,19 @@ void UModel::PostLoad()
 	}
 }
 
+void UModel::PreSave(FObjectPreSaveContext SaveContext)
+{
+	Super::PreSave(SaveContext);
+#if WITH_EDITOR
+	if (!SaveContext.IsProceduralSave())
+	{
+		// Reconstruct the lighting guid every time the model is saved by the user in editor.
+		// ConstructLightingGuid is deterministic so this will not cause spurious changes.
+		LightingGuid = ConstructLightingGuid();
+	}
+#endif
+}
+
 #if WITH_EDITOR
 void UModel::PostEditUndo()
 {
@@ -561,8 +625,7 @@ bool UModel::Modify( bool bAlwaysMarkDirty/*=false*/ )
 {
 	bool bSavedToTransactionBuffer = Super::Modify(bAlwaysMarkDirty);
 
-	// make a new guid whenever this model changes
-	LightingGuid = FGuid::NewGuid();
+	// We do not reconstruct the LightingGuid on every change, we reconstruct it (deterministically) when saved
 
 	// Modify all child objects.
 	if( Polys )
@@ -572,7 +635,97 @@ bool UModel::Modify( bool bAlwaysMarkDirty/*=false*/ )
 
 	return bSavedToTransactionBuffer;
 }
-#endif
+
+FGuid UModel::ConstructLightingGuid() const
+{
+	TRACE_CPUPROFILER_EVENT_SCOPE(UModel::ConstructLightingGuid);
+	FBlake3 Builder;
+
+	// Add in all variables but only variables that could affect lighting. This includes all of the geometry and materials.
+	// It does not include properties in other UObjects (we copy their path only) because the LightingGuid only
+	// needs to cover data on *this.
+	// It does not include cached data that only optimizes operations that is redundant with other geometry data.
+	// It does not include transient data.
+
+	ModelUpdateHash(Builder, Polys);
+	if (!Nodes.IsEmpty())
+	{
+		static_assert(alignof(FBspNode) <= 1 || sizeof(FBspNode) % alignof(FBspNode) == 0, "We rely on zero padding in arrays");
+		checkf(Nodes.Num() < 2 || (int64)&Nodes[1] - (int64)&Nodes[0] == sizeof(Nodes[0]), TEXT("We rely on zero padding in arrays"));
+		Builder.Update(Nodes.GetData(), Nodes.Num() * sizeof(Nodes[0]));
+	}
+	if (!Verts.IsEmpty())
+	{
+		static_assert(alignof(FVert) <= 1 || sizeof(FVert) % alignof(FVert) == 0, "We rely on zero padding in arrays");
+		checkf(Verts.Num() < 2 || (int64)&Verts[1] - (int64)&Verts[0] == sizeof(Verts[0]), TEXT("We rely on zero padding in arrays"));
+		Builder.Update(Verts.GetData(), Verts.Num() * sizeof(Verts[0]));
+	}
+	if (!Vectors.IsEmpty())
+	{
+		static_assert(alignof(FVector3f) <= 1 || sizeof(FVector3f) % alignof(FVector3f) == 0, "We rely on zero padding in arrays");
+		checkf(Vectors.Num() < 2 || (int64)&Vectors[1] - (int64)&Vectors[0] == sizeof(Vectors[0]), TEXT("We rely on zero padding in arrays"));
+		Builder.Update(Vectors.GetData(), Vectors.Num() * sizeof(Vectors[0]));
+	}
+	if (!Points.IsEmpty())
+	{
+		static_assert(alignof(FVector3f) <= 1 || sizeof(FVector3f) % alignof(FVector3f) == 0, "We rely on zero padding in arrays");
+		checkf(Points.Num() < 2 || (int64)&Points[1] - (int64)&Points[0] == sizeof(Points[0]), TEXT("We rely on zero padding in arrays"));
+		Builder.Update(Points.GetData(), Points.Num() * sizeof(Points[0]));
+	}
+	for (const FBspSurf& Value : Surfs)
+	{
+		UpdateHash(Builder, Value);
+	}
+	if (!LeafHulls.IsEmpty())
+	{
+		static_assert(alignof(int32) <= 1 || sizeof(int32) % alignof(int32) == 0, "We rely on zero padding in arrays");
+		checkf(LeafHulls.Num() < 2 || (int64)&LeafHulls[1] - (int64)&LeafHulls[0] == sizeof(LeafHulls[0]), TEXT("We rely on zero padding in arrays"));
+		Builder.Update(LeafHulls.GetData(), LeafHulls.Num() * sizeof(LeafHulls[0]));
+	}
+	for (const FLeaf& Value : Leaves)
+	{
+		Builder.Update(&Value, sizeof(Value));
+	}
+	for (const FLightmassPrimitiveSettings& Value : LightmassSettings)
+	{
+		Builder.Update(&Value, sizeof(Value));
+	}
+	for (const TPair<UMaterialInterface*, TUniquePtr<FRawIndexBuffer16or32>>& Pair : MaterialIndexBuffers)
+	{
+		ModelUpdateHash(Builder, Pair.Key);
+		if (Pair.Value)
+		{
+			for (const uint32& Value : Pair.Value->Indices)
+			{
+				Builder.Update(&Value, sizeof(Value));
+			}
+		}
+	}
+	UpdateHash(Builder, VertexBuffer);
+	// ReleaseResourcesFence - Not needed, transient runtime rendering 
+	// InvalidSurfaces - Not needed, editor operations support 
+	// bOnlyRebuildMaterialIndexBuffers - Not needed, does not impact lighting results
+	// bInvalidForStaticLighting - Not needed, does not impact lighting results
+	// NumUniqueVertices - Not needed, redundant with this->Verts
+	// LightingGuid - Not needed, it's the thing we're calculated and is redundant with all the other data
+	// NodeGroups - Not needed, redundant with this->Nodes
+	// CachedMappings - Not needed, redundant with this->Surfs
+	// NumIncompleteNodeGroups - Not needed, editor operations support
+	// LightingLevel - Not needed, redundant with this->Nodes
+	// OwnerLocationWhenLastBuilt - Not needed, redundant with this->Verts
+	// OwnerRotationWhenLastBuilt - Not needed, redundant with this->Verts
+	// OwnerScaleWhenLastBuilt - Not needed, redundant with this->Verts
+	// bCachedOwnerTransformValid - Not needed, editor operations support
+	// RootOutside  - Not needed, editor operations support
+	// Linked - Not needed, editor operations support
+	// NumSharedSides - Not needed, editor operations support
+	// Bounds - Not needed, redundant with this->Verts
+
+	FBlake3Hash Hash = Builder.Finalize();
+	uint32* HashBytes = (uint32*)Hash.GetBytes();
+	return FGuid(HashBytes[0], HashBytes[1], HashBytes[2], HashBytes[3]);
+}
+#endif // WITH_EDITOR
 
 //
 // Empty the contents of a model.
diff --git a/Engine/Source/Runtime/Engine/Private/ModelRender.cpp b/Engine/Source/Runtime/Engine/Private/ModelRender.cpp
--- a/Engine/Source/Runtime/Engine/Private/ModelRender.cpp
+++ b/Engine/Source/Runtime/Engine/Private/ModelRender.cpp
@@ -24,6 +24,7 @@
 #include "Engine/Brush.h"
 #include "Components/ModelComponent.h"
 #include "SceneInterface.h"
+#include "Hash/Blake3.h"
 
 namespace
 {
@@ -87,6 +88,18 @@ FArchive& operator<<(FArchive& Ar,FModelVertexBuffer& B)
 	return Ar;
 }
 
+#if WITH_EDITOR
+void UpdateHash(FBlake3& Builder, const FModelVertexBuffer& B)
+{
+	if (!B.Vertices.IsEmpty())
+	{
+		static_assert(alignof(FModelVertex) <= 1 || sizeof(FModelVertex) % alignof(FModelVertex) == 0, "We rely on zero padding in arrays");
+		checkf(B.Vertices.Num() < 2 || (int64)&B.Vertices[1] - (int64)&B.Vertices[0] == sizeof(B.Vertices[0]), TEXT("We rely on zero padding in arrays"));
+		Builder.Update(B.Vertices.GetData(), B.Vertices.Num() * sizeof(B.Vertices[0]));
+	}
+	// B.Buffers - Not needed, transient runtime rendering 
+}
+#endif
 /*-----------------------------------------------------------------------------
 UModelComponent
 -----------------------------------------------------------------------------*/
diff --git a/Engine/Source/Runtime/Engine/Public/Model.h b/Engine/Source/Runtime/Engine/Public/Model.h
--- a/Engine/Source/Runtime/Engine/Public/Model.h
+++ b/Engine/Source/Runtime/Engine/Public/Model.h
@@ -23,6 +23,7 @@
 
 class AActor;
 class ABrush;
+class FBlake3;
 class FMeshMapBuildData;
 class ULevel;
 class ULightComponent;
@@ -339,6 +340,9 @@ public:
 	* @param B - data to serialize
 	*/
 	friend FArchive& operator<<(FArchive& Ar,FModelVertexBuffer& B);
+#if WITH_EDITOR
+	friend void UpdateHash(FBlake3& Builder, const FModelVertexBuffer& B);
+#endif
 
 private:
 	UModel* Model;
@@ -472,6 +476,7 @@ public:
 	// UObject interface.
 	virtual void Serialize( FArchive& Ar ) override;
 	virtual void PostLoad() override;
+	virtual void PreSave(FObjectPreSaveContext SaveContext) override;
 #if WITH_EDITOR
 	virtual void PostEditUndo() override;
 	virtual bool Modify(bool bAlwaysMarkDirty = false) override;
@@ -600,7 +605,7 @@ public:
 	friend class AVolume;
 
 private:
-
+	FGuid ConstructLightingGuid() const;
 };
 
 /**
