Remove assert related to mobile GPUScene and PSO pre-caching, its no longer needed

diff --git a/Engine/Source/Runtime/Engine/Private/LocalVertexFactory.cpp b/Engine/Source/Runtime/Engine/Private/LocalVertexFactory.cpp
--- a/Engine/Source/Runtime/Engine/Private/LocalVertexFactory.cpp
+++ b/Engine/Source/Runtime/Engine/Private/LocalVertexFactory.cpp
@@ -375,9 +375,7 @@ void FLocalVertexFactory::GetVertexElements(ERHIFeatureLevel::Type FeatureLevel,
 
 	if (UseGPUScene(GMaxRHIShaderPlatform, GMaxRHIFeatureLevel))
 	{
-		// For ES3.1 attribute ID needs to be done differently
-		check(FeatureLevel > ERHIFeatureLevel::ES3_1);
-		Elements.Add(FVertexElement(VertexStreams.Num(), 0, VET_UInt, 13, 0, true));
+		Elements.Add(FVertexElement(VertexStreams.Num(), 0, VET_UInt, 13, sizeof(uint32), true));
 	}
 }
 
