World Partition HLOD generation: moved HLOD layer resolving from HLOD generation pass into actor descriptor views creation pass, so we can perform error checking later on.

diff --git a/Engine/Plugins/Editor/WorldPartitionHLODUtilities/Source/Private/WorldPartition/HLOD/Utilities/WorldPartitionHLODUtilities.cpp b/Engine/Plugins/Editor/WorldPartitionHLODUtilities/Source/Private/WorldPartition/HLOD/Utilities/WorldPartitionHLODUtilities.cpp
--- a/Engine/Plugins/Editor/WorldPartitionHLODUtilities/Source/Private/WorldPartition/HLOD/Utilities/WorldPartitionHLODUtilities.cpp
+++ b/Engine/Plugins/Editor/WorldPartitionHLODUtilities/Source/Private/WorldPartition/HLOD/Utilities/WorldPartitionHLODUtilities.cpp
@@ -114,8 +114,7 @@ TArray<AWorldPartitionHLOD*> FWorldPartitionHLODUtilities::CreateHLODActors(FHLO
 				continue;
 			}
 
-			UHLODLayer* HLODLayer = UHLODLayer::GetHLODLayer(ActorDescView, InCreationParams.WorldPartition);
-			if (HLODLayer)
+			if (UHLODLayer* HLODLayer = Cast<UHLODLayer>(ActorDescView.GetHLODLayer().TryLoad()))
 			{
 				TSet<FHLODSubActor>& SubActors = SubActorsPerHLODLayer.FindOrAdd(HLODLayer);
 				AddSubActor(ActorDescView, ActorInstance, SubActors);
diff --git a/Engine/Source/Runtime/Engine/Private/WorldPartition/HLOD/HLODLayer.cpp b/Engine/Source/Runtime/Engine/Private/WorldPartition/HLOD/HLODLayer.cpp
--- a/Engine/Source/Runtime/Engine/Private/WorldPartition/HLOD/HLODLayer.cpp
+++ b/Engine/Source/Runtime/Engine/Private/WorldPartition/HLOD/HLODLayer.cpp
@@ -32,53 +32,6 @@ UHLODLayer::UHLODLayer(const FObjectInitializer& ObjectInitializer)
 {}
 
 #if WITH_EDITOR
-UHLODLayer* UHLODLayer::GetHLODLayer(const AActor* InActor)
-{
-	if (UHLODLayer* HLODLayer = InActor->GetHLODLayer())
-	{
-		return HLODLayer;
-	}
-
-	// Only fallback to the default HLODLayer for the first level of HLOD
-	bool bIsHLOD0 = !InActor->IsA<AWorldPartitionHLOD>();
-	if (bIsHLOD0) 
-	{
-		// Fallback to the world partition default HLOD layer
-		if (UWorldPartition* WorldPartition = InActor->GetWorld()->GetWorldPartition())
-		{
-			return WorldPartition->GetDefaultHLODLayer();
-		}
-	}
-
-	return nullptr;
-}
-
-UHLODLayer* UHLODLayer::GetHLODLayer(const FWorldPartitionActorDescView& InActorDesc, const UWorldPartition* InWorldPartition)
-{
-	check(InWorldPartition);
-
-	const FSoftObjectPath HLODLayerPath = InActorDesc.GetHLODLayer();
-	if (UHLODLayer* HLODLayer = Cast<UHLODLayer>(HLODLayerPath.TryLoad()))
-	{
-		return HLODLayer;
-	}
-
-	// Only fallback to the default HLODLayer for the first level of HLOD
-	bool bIsHLOD0 = !InActorDesc.GetActorNativeClass()->IsChildOf<AWorldPartitionHLOD>();
-	if (bIsHLOD0)
-	{
-		// Fallback to the world partition default HLOD layer
-		return InWorldPartition->GetDefaultHLODLayer();
-	}
-
-	return nullptr;
-}
-
-UHLODLayer* UHLODLayer::GetHLODLayer(const FWorldPartitionActorDesc& InActorDesc, const UWorldPartition* InWorldPartition)
-{
-	return GetHLODLayer(FWorldPartitionActorDescView(&InActorDesc), InWorldPartition);
-}
-
 bool UHLODLayer::DoesRequireWarmup() const
 {
 	IWorldPartitionHLODUtilities* WPHLODUtilities = FModuleManager::Get().LoadModuleChecked<IWorldPartitionHLODUtilitiesModule>("WorldPartitionHLODUtilities").GetUtilities();
diff --git a/Engine/Source/Runtime/Engine/Private/WorldPartition/RuntimeSpatialHash/RuntimeSpatialHashHLOD.cpp b/Engine/Source/Runtime/Engine/Private/WorldPartition/RuntimeSpatialHash/RuntimeSpatialHashHLOD.cpp
--- a/Engine/Source/Runtime/Engine/Private/WorldPartition/RuntimeSpatialHash/RuntimeSpatialHashHLOD.cpp
+++ b/Engine/Source/Runtime/Engine/Private/WorldPartition/RuntimeSpatialHash/RuntimeSpatialHashHLOD.cpp
@@ -285,7 +285,7 @@ static TMap<UHLODLayer*, int32> GatherHLODLayers(const IStreamingGenerationConte
 			{
 				if (ActorDescView.GetActorIsHLODRelevant())
 				{
-					UHLODLayer* HLODLayer = UHLODLayer::GetHLODLayer(ActorDescView, WorldPartition);
+					UHLODLayer* HLODLayer = Cast<UHLODLayer>(ActorDescView.GetHLODLayer().TryLoad());
 		
 					// If layer was already encountered, no need to process it again
 					if (!HLODLayersLevel.Contains(HLODLayer))
diff --git a/Engine/Source/Runtime/Engine/Private/WorldPartition/WorldPartitionActorDescView.cpp b/Engine/Source/Runtime/Engine/Private/WorldPartition/WorldPartitionActorDescView.cpp
--- a/Engine/Source/Runtime/Engine/Private/WorldPartition/WorldPartitionActorDescView.cpp
+++ b/Engine/Source/Runtime/Engine/Private/WorldPartition/WorldPartitionActorDescView.cpp
@@ -93,6 +93,11 @@ bool FWorldPartitionActorDescView::GetActorIsHLODRelevant() const
 
 FSoftObjectPath FWorldPartitionActorDescView::GetHLODLayer() const
 {
+	if (RuntimedHLODLayer.IsSet())
+	{
+		return RuntimedHLODLayer.GetValue();
+	}
+
 	return ActorDesc->GetHLODLayer();
 }
 
@@ -265,7 +270,11 @@ void FWorldPartitionActorDescView::SetForcedNonSpatiallyLoaded()
 
 void FWorldPartitionActorDescView::SetForcedNoRuntimeGrid()
 {
-	bIsForcedNoRuntimeGrid = true;	
+	if (!bIsForcedNoRuntimeGrid)
+	{
+		bIsForcedNoRuntimeGrid = true;	
+		UE_LOG(LogWorldPartition, Verbose, TEXT("Actor '%s' runtime grid invalidated"), *GetActorLabelOrName().ToString());
+	}
 }
 
 void FWorldPartitionActorDescView::SetForcedNoDataLayers()
@@ -292,6 +301,11 @@ void FWorldPartitionActorDescView::SetEditorReferences(const TArray<FGuid>& InEd
 	EditorReferences = InEditorReferences;
 }
 
+void FWorldPartitionActorDescView::SetRuntimeHLODLayer(const FSoftObjectPath& InHLODLayer)
+{
+	RuntimedHLODLayer = InHLODLayer;
+}
+
 AActor* FWorldPartitionActorDescView::GetActor() const
 {
 	return ActorDesc->GetActor();
diff --git a/Engine/Source/Runtime/Engine/Private/WorldPartition/WorldPartitionStreamingGeneration.cpp b/Engine/Source/Runtime/Engine/Private/WorldPartition/WorldPartitionStreamingGeneration.cpp
--- a/Engine/Source/Runtime/Engine/Private/WorldPartition/WorldPartitionStreamingGeneration.cpp
+++ b/Engine/Source/Runtime/Engine/Private/WorldPartition/WorldPartitionStreamingGeneration.cpp
@@ -28,6 +28,7 @@
 #include "WorldPartition/ErrorHandling/WorldPartitionStreamingGenerationLogErrorHandler.h"
 #include "WorldPartition/ErrorHandling/WorldPartitionStreamingGenerationMapCheckErrorHandler.h"
 #include "WorldPartition/HLOD/HLODActor.h"
+#include "WorldPartition/HLOD/HLODLayer.h"
 #include "WorldPartition/WorldPartitionSubsystem.h"
 #include "HAL/FileManager.h"
 
@@ -340,6 +341,16 @@ class FWorldPartitionStreamingGenerator
 		}
 	}
 
+	void ResolveHLODLayer(FWorldPartitionActorDescView& ActorDescView, const FSoftObjectPath& DefaultHLODLayer)
+	{
+		// Only assign the default layer to actors that don't have a valid HLOD layer set. HLOD actors will have their 
+		// parent HLOD layer set during HLOD generation.
+		if (!ActorDescView.GetHLODLayer().IsValid())
+		{
+			ActorDescView.SetRuntimeHLODLayer(DefaultHLODLayer);
+		}
+	}
+
 	void ResolveParentView(FWorldPartitionActorDescView& ActorDescView, const FActorDescViewMap& ActorDescViewMap)
 	{
 		if (FGuid ParentGuid = ActorDescView.GetParentActor(); ParentGuid.IsValid())
@@ -633,6 +644,7 @@ class FWorldPartitionStreamingGenerator
 			ResolveRuntimeSpatiallyLoaded(ActorDescView);
 			ResolveRuntimeGrid(ActorDescView);
 			ResolveRuntimeDataLayers(ActorDescView, ContainerCollectionDescriptor.ActorDescViewMap);
+			ResolveHLODLayer(ActorDescView, FSoftObjectPath(WorldPartitionContext->GetDefaultHLODLayer()));
 			ResolveParentView(ActorDescView, ContainerCollectionDescriptor.ActorDescViewMap);
 		};
 
diff --git a/Engine/Source/Runtime/Engine/Public/WorldPartition/HLOD/HLODLayer.h b/Engine/Source/Runtime/Engine/Public/WorldPartition/HLOD/HLODLayer.h
--- a/Engine/Source/Runtime/Engine/Public/WorldPartition/HLOD/HLODLayer.h
+++ b/Engine/Source/Runtime/Engine/Public/WorldPartition/HLOD/HLODLayer.h
@@ -42,10 +42,6 @@ class UHLODLayer : public UObject
 	
 #if WITH_EDITOR
 public:
-	static ENGINE_API UHLODLayer* GetHLODLayer(const AActor* InActor);
-	static ENGINE_API UHLODLayer* GetHLODLayer(const FWorldPartitionActorDesc& InActorDesc, const UWorldPartition* InWorldPartition);
-	static ENGINE_API UHLODLayer* GetHLODLayer(const FWorldPartitionActorDescView& InActorDescView, const UWorldPartition* InWorldPartition);
-
 	/** Get the default engine HLOD layers setup */
 	static ENGINE_API UHLODLayer* GetEngineDefaultHLODLayersSetup();
 
diff --git a/Engine/Source/Runtime/Engine/Public/WorldPartition/WorldPartitionActorDescView.h b/Engine/Source/Runtime/Engine/Public/WorldPartition/WorldPartitionActorDescView.h
--- a/Engine/Source/Runtime/Engine/Public/WorldPartition/WorldPartitionActorDescView.h
+++ b/Engine/Source/Runtime/Engine/Public/WorldPartition/WorldPartitionActorDescView.h
@@ -84,6 +84,8 @@ public:
 	ENGINE_API void SetRuntimeReferences(const TArray<FGuid>& InRuntimeReferences);
 	ENGINE_API void SetEditorReferences(const TArray<FGuid>& InEditorReferences);
 	ENGINE_API void SetDataLayerInstanceNames(const TArray<FName>& InDataLayerInstanceNames);
+	
+	ENGINE_API void SetRuntimeHLODLayer(const FSoftObjectPath& InHLODLayer);
 
 	ENGINE_API AActor* GetActor() const;
 
@@ -111,10 +113,11 @@ protected:
 	const FWorldPartitionActorDescView* ParentView;
 	bool bIsForcedNonSpatiallyLoaded;
 	bool bIsForcedNoRuntimeGrid;
-	bool bIsForcedNoDataLayers;	
+	bool bIsForcedNoDataLayers;
 	TOptional<TArray<FName>> ResolvedDataLayerInstanceNames;
 	TOptional<TArray<FName>> RuntimeDataLayerInstanceNames;
 	TOptional<TArray<FGuid>> RuntimeReferences;
+	TOptional<FSoftObjectPath> RuntimedHLODLayer;
 	TArray<FGuid> EditorReferences;
 };
 #endif
