Disable Downscale setting in the Texture Editor UI when its value is ignored during the texture build.

diff --git a/Engine/Source/Runtime/Engine/Classes/Engine/Texture.h b/Engine/Source/Runtime/Engine/Classes/Engine/Texture.h
--- a/Engine/Source/Runtime/Engine/Classes/Engine/Texture.h
+++ b/Engine/Source/Runtime/Engine/Classes/Engine/Texture.h
@@ -1343,16 +1343,22 @@ PRAGMA_ENABLE_DEPRECATION_WARNINGS
 	UPROPERTY(EditAnywhere, BlueprintReadWrite, Category=LevelOfDetail, meta=(DisplayName="Texture Group"), AssetRegistrySearchable)
 	TEnumAsByte<enum TextureGroup> LODGroup;
 
+	/** This function is used to control access to the Downscale and DownscaleOptions properties from the Texture Editor UI
+	 * in order to make it more clear to the user whether these properties will or will not be used when building the texture.
+	 */
+	UFUNCTION()
+	virtual bool AreDownscalePropertiesEditable() const { return false; }
+
 	/** Downscale source texture, applied only to 2d textures without mips 
-	 * 0.0 - use scale value from texture group
+	 * < 1.0 - use scale value from texture group
 	 * 1.0 - do not scale texture
 	 * > 1.0 - scale texure
 	 */
-	UPROPERTY(EditAnywhere, Category=LevelOfDetail, AdvancedDisplay, meta=(ClampMin="0.0", ClampMax="8.0"))
+	UPROPERTY(EditAnywhere, Category=LevelOfDetail, AdvancedDisplay, meta=(ClampMin="0.0", ClampMax="8.0", EditCondition=AreDownscalePropertiesEditable))
 	FPerPlatformFloat Downscale;
 
 	/** Texture downscaling options */
-	UPROPERTY(EditAnywhere, Category=LevelOfDetail, AdvancedDisplay)
+	UPROPERTY(EditAnywhere, Category=LevelOfDetail, AdvancedDisplay, meta=(EditCondition=AreDownscalePropertiesEditable))
 	ETextureDownscaleOptions DownscaleOptions;
 
 	/** 
diff --git a/Engine/Source/Runtime/Engine/Classes/Engine/Texture2D.h b/Engine/Source/Runtime/Engine/Classes/Engine/Texture2D.h
--- a/Engine/Source/Runtime/Engine/Classes/Engine/Texture2D.h
+++ b/Engine/Source/Runtime/Engine/Classes/Engine/Texture2D.h
@@ -322,6 +322,12 @@ public:
 	* For game code access to the cpu copy image, use GetCPUCopy().
 	*/
 	ENGINE_API UTexture2D* GetCPUCopyTexture();
+
+	/**
+	 * Returns true if the Downscale and DownscaleOptions properties should be editable in the UI.
+	 * Currently downscaling is only supported for 2d textures without mipmaps.
+	 */
+	virtual bool AreDownscalePropertiesEditable() const override { return MipGenSettings == TMGS_NoMipmaps || MipGenSettings == TMGS_FromTextureGroup; }
 #endif
 
 	friend struct FRenderAssetStreamingManager;
diff --git a/Engine/Source/Runtime/Engine/Private/TextureDerivedData.cpp b/Engine/Source/Runtime/Engine/Private/TextureDerivedData.cpp
--- a/Engine/Source/Runtime/Engine/Private/TextureDerivedData.cpp
+++ b/Engine/Source/Runtime/Engine/Private/TextureDerivedData.cpp
@@ -1003,13 +1003,29 @@ static void GetTextureBuildSettings(
 	}
 
 	OutBuildSettings.Downscale = 1.0f;
+
+	float Downscale;
+	ETextureDownscaleOptions DownscaleOptions;
+	TextureLODSettings.GetDownscaleOptions(Texture, TargetPlatform, Downscale, DownscaleOptions);
+
 	// Downscale only allowed if NoMipMaps, 2d, and not VT
 	//	silently does nothing otherwise
 	if (! bVirtualTextureStreaming &&
 		MipGenSettings == TMGS_NoMipmaps && 
 		Texture.IsA(UTexture2D::StaticClass()))	// TODO: support more texture types
 	{
-		TextureLODSettings.GetDownscaleOptions(Texture, TargetPlatform, OutBuildSettings.Downscale, (ETextureDownscaleOptions&)OutBuildSettings.DownscaleOptions);
+		OutBuildSettings.Downscale = Downscale;
+		OutBuildSettings.DownscaleOptions = (uint8)DownscaleOptions;
+	}
+	// only show a warning for textures where Downscale setting would have effect if it was used
+	else if (Downscale != 1.f)
+	{
+		UE_LOG(LogTexture, Warning, TEXT("Downscale setting of %f was not used when building texture %s%s."), Downscale, *Texture.GetName(),
+			bVirtualTextureStreaming ? TEXT(" because it is using virtual texturing") :
+			MipGenSettings != TMGS_NoMipmaps ? TEXT(" because it is using mipmaps") :
+			!Texture.IsA(UTexture2D::StaticClass()) ? TEXT(" because it is only supported for 2D textures") :
+			TEXT("")
+		);
 	}
 	
 	// For virtual texturing we take the address mode into consideration
