Migrate: implement FLevelInstanceActorDesc::TransferWorldData

#rnx

diff --git a/Engine/Source/Runtime/Engine/Private/WorldPartition/LevelInstance/LevelInstanceActorDesc.cpp b/Engine/Source/Runtime/Engine/Private/WorldPartition/LevelInstance/LevelInstanceActorDesc.cpp
--- a/Engine/Source/Runtime/Engine/Private/WorldPartition/LevelInstance/LevelInstanceActorDesc.cpp
+++ b/Engine/Source/Runtime/Engine/Private/WorldPartition/LevelInstance/LevelInstanceActorDesc.cpp
@@ -233,6 +233,14 @@ void FLevelInstanceActorDesc::TransferFrom(const FWorldPartitionActorDesc* From)
 	}
 }
 
+void FLevelInstanceActorDesc::TransferWorldData(const FWorldPartitionActorDesc* From)
+{
+	FWorldPartitionActorDesc::TransferWorldData(From);
+
+	FLevelInstanceActorDesc* FromLevelInstanceActorDesc = (FLevelInstanceActorDesc*)From;
+	LevelInstanceTransform = FromLevelInstanceActorDesc->LevelInstanceTransform;
+}
+
 void FLevelInstanceActorDesc::Serialize(FArchive& Ar)
 {
 	FWorldPartitionActorDesc::Serialize(Ar);
diff --git a/Engine/Source/Runtime/Engine/Public/WorldPartition/LevelInstance/LevelInstanceActorDesc.h b/Engine/Source/Runtime/Engine/Public/WorldPartition/LevelInstance/LevelInstanceActorDesc.h
--- a/Engine/Source/Runtime/Engine/Public/WorldPartition/LevelInstance/LevelInstanceActorDesc.h
+++ b/Engine/Source/Runtime/Engine/Public/WorldPartition/LevelInstance/LevelInstanceActorDesc.h
@@ -37,6 +37,7 @@ protected:
 	ENGINE_API virtual void Init(const FWorldPartitionActorDescInitData& DescData) override;
 	ENGINE_API virtual bool Equals(const FWorldPartitionActorDesc* Other) const override;
 	ENGINE_API virtual void TransferFrom(const FWorldPartitionActorDesc* From) override;
+	ENGINE_API virtual void TransferWorldData(const FWorldPartitionActorDesc* From) override;
 	virtual uint32 GetSizeOf() const override { return sizeof(FLevelInstanceActorDesc); }
 	ENGINE_API virtual void Serialize(FArchive& Ar) override;
 	ENGINE_API virtual void SetContainer(UActorDescContainer* InContainer, UWorld* InWorldContext) override;
