World Partition: fix crash when validating changelists for unloaded worlds.

diff --git a/Engine/Source/Runtime/Engine/Private/WorldPartition/WorldPartitionStreamingGeneration.cpp b/Engine/Source/Runtime/Engine/Private/WorldPartition/WorldPartitionStreamingGeneration.cpp
--- a/Engine/Source/Runtime/Engine/Private/WorldPartition/WorldPartitionStreamingGeneration.cpp
+++ b/Engine/Source/Runtime/Engine/Private/WorldPartition/WorldPartitionStreamingGeneration.cpp
@@ -644,7 +644,7 @@ class FWorldPartitionStreamingGenerator
 			ResolveRuntimeSpatiallyLoaded(ActorDescView);
 			ResolveRuntimeGrid(ActorDescView);
 			ResolveRuntimeDataLayers(ActorDescView, ContainerCollectionDescriptor.ActorDescViewMap);
-			ResolveHLODLayer(ActorDescView, FSoftObjectPath(WorldPartitionContext->GetDefaultHLODLayer()));
+			ResolveHLODLayer(ActorDescView, WorldPartitionContext ? FSoftObjectPath(WorldPartitionContext->GetDefaultHLODLayer()) : FSoftObjectPath());
 			ResolveParentView(ActorDescView, ContainerCollectionDescriptor.ActorDescViewMap);
 		};
 
