Changed fx.detailedcsvstats to emit CSV metadata that we can use to automatically generate an FX report type.

#rnx

diff --git a/Engine/Source/Runtime/Engine/Private/Particles/ParticlePerfStats.cpp b/Engine/Source/Runtime/Engine/Private/Particles/ParticlePerfStats.cpp
--- a/Engine/Source/Runtime/Engine/Private/Particles/ParticlePerfStats.cpp
+++ b/Engine/Source/Runtime/Engine/Private/Particles/ParticlePerfStats.cpp
@@ -1125,6 +1125,7 @@ void OnDetailedCSVStatsEnabledChanged(IConsoleVariable* Variable)
 	if (CSVProfiler)
 	{
 		CSVProfiler->EnableCategoryByIndex(CSV_CATEGORY_INDEX(Particles), bDetailedCSVStats);
+		CSV_METADATA(TEXT("DetailedFXStats"), bDetailedCSVStats ? TEXT("1") : TEXT("0"));
 	}
 
 	if (bDetailedCSVStats)
