-Fix potential unity build issue

diff --git a/Engine/Source/Runtime/Engine/Private/Components/StaticMeshComponent.cpp b/Engine/Source/Runtime/Engine/Private/Components/StaticMeshComponent.cpp
--- a/Engine/Source/Runtime/Engine/Private/Components/StaticMeshComponent.cpp
+++ b/Engine/Source/Runtime/Engine/Private/Components/StaticMeshComponent.cpp
@@ -2335,7 +2335,7 @@ const Nanite::FResources* UStaticMeshComponent::GetNaniteResources() const
 namespace Nanite
 {
 	template<class T> 
-	bool ShouldCreateNaniteProxy(const T& Component, FMaterialAudit* OutNaniteMaterials = nullptr);
+	bool ShouldCreateNaniteProxy(const T& Component, FMaterialAudit* OutNaniteMaterials);
 
 	template<class T> 
 	bool HasValidNaniteData(const T& Component)
@@ -2348,7 +2348,7 @@ namespace Nanite
 	bool UseNaniteOverrideMaterials(const T& Component, bool bDoingMaterialAudit) 
 	{
 		// Check for valid data on this SMC and support for Nanite material overrides
-		return (bDoingMaterialAudit || ShouldCreateNaniteProxy(Component)) && GEnableNaniteMaterialOverrides != 0;
+		return (bDoingMaterialAudit || ShouldCreateNaniteProxy(Component, nullptr)) && GEnableNaniteMaterialOverrides != 0;
 	}
 }
 
diff --git a/Engine/Source/Runtime/Engine/Private/StaticMeshRender.cpp b/Engine/Source/Runtime/Engine/Private/StaticMeshRender.cpp
--- a/Engine/Source/Runtime/Engine/Private/StaticMeshRender.cpp
+++ b/Engine/Source/Runtime/Engine/Private/StaticMeshRender.cpp
@@ -2557,7 +2557,7 @@ FLODMask FStaticMeshSceneProxy::GetLODMask(const FSceneView* View) const
 namespace Nanite
 {
 	template<class T> 
-	bool ShouldCreateNaniteProxy(const T& Component, FMaterialAudit* OutNaniteMaterials = nullptr)
+	bool ShouldCreateNaniteProxy(const T& Component, FMaterialAudit* OutNaniteMaterials)
 	{
 		// Whether or not to allow Nanite for this component
 	#if WITH_EDITORONLY_DATA
