[materials] re-add ODSC guard for material shadermap DDC queries. in practice compiling the shaders is faster than performing this query, so it's not having the intended effect.

diff --git a/Engine/Source/Runtime/Engine/Private/Materials/MaterialShared.cpp b/Engine/Source/Runtime/Engine/Private/Materials/MaterialShared.cpp
--- a/Engine/Source/Runtime/Engine/Private/Materials/MaterialShared.cpp
+++ b/Engine/Source/Runtime/Engine/Private/Materials/MaterialShared.cpp
@@ -2702,8 +2702,11 @@ bool FMaterial::CacheShaders(const FMaterialShaderMapId& ShaderMapId, EShaderPla
 			// If we have an incomplete shadermap, continue with it to prevent creation of duplicate shadermaps for the same ShaderMapId
 			if (!ShaderMap && !FPlatformProperties::RequiresCookedData())
 			{
-				TRefCountPtr<FMaterialShaderMap> LoadedShaderMap;
-				CacheShadersPending = FMaterialShaderMap::BeginLoadFromDerivedDataCache(this, ShaderMapId, Platform, TargetPlatform, LoadedShaderMap, DDCKeyHash);
+				if (bSkipCompilationOnPostLoad == false || IsRequiredComplete())
+				{
+					TRefCountPtr<FMaterialShaderMap> LoadedShaderMap;
+					CacheShadersPending = FMaterialShaderMap::BeginLoadFromDerivedDataCache(this, ShaderMapId, Platform, TargetPlatform, LoadedShaderMap, DDCKeyHash);
+				}
 			}
 
 			check(!ShaderMap || ShaderMap->GetFrozenContentSize() > 0u);
