Optimize anim blueprint context menu construction, mostly by avoiding asset registry tag gathering, but also by reducing temporary allocations

diff --git a/Engine/Source/Editor/BlueprintGraph/Private/BlueprintActionFilter.cpp b/Engine/Source/Editor/BlueprintGraph/Private/BlueprintActionFilter.cpp
--- a/Engine/Source/Editor/BlueprintGraph/Private/BlueprintActionFilter.cpp
+++ b/Engine/Source/Editor/BlueprintGraph/Private/BlueprintActionFilter.cpp
@@ -914,7 +914,7 @@ static bool BlueprintActionFilterImpl::IsAssetPermissionNotGranted(FBlueprintAct
 	{
 		if (Filter.AssetReferenceFilter.IsValid())
 		{
-			bool const bIsAccessible = Filter.AssetReferenceFilter->PassesFilter(FAssetData(Owner));
+			bool const bIsAccessible = Filter.AssetReferenceFilter->PassesFilter(FAssetData(Owner, FAssetData::ECreationFlags::SkipAssetRegistryTagsGathering));
 			if (!bIsAccessible)
 			{
 				bIsFilteredOut = true;
diff --git a/Engine/Source/Runtime/Engine/Private/Animation/AnimCurveMetaData.cpp b/Engine/Source/Runtime/Engine/Private/Animation/AnimCurveMetaData.cpp
--- a/Engine/Source/Runtime/Engine/Private/Animation/AnimCurveMetaData.cpp
+++ b/Engine/Source/Runtime/Engine/Private/Animation/AnimCurveMetaData.cpp
@@ -1,4 +1,4 @@
-﻿// Copyright Epic Games, Inc. All Rights Reserved.
+// Copyright Epic Games, Inc. All Rights Reserved.
 
 #include "Animation/AnimCurveMetadata.h"
 #include "Animation/Skeleton.h"
@@ -43,12 +43,12 @@ void UAnimCurveMetaData::GetAssetRegistryTags(TArray<FAssetRegistryTag>& OutTags
 	// The delimiter is necessary so we can distinguish between data with no curves and old data, as the asset registry
 	// strips tags that have empty values 
 	TStringBuilder<256> CurvesBuilder;
-	CurvesBuilder.Append(USkeleton::CurveTagDelimiter);	
+	CurvesBuilder << USkeleton::CurveTagDelimiter;
 
 	ForEachCurveMetaData([&CurvesBuilder](FName InCurveName, const FCurveMetaData&)
 	{
-		CurvesBuilder.Append(InCurveName.ToString());
-		CurvesBuilder.Append(USkeleton::CurveTagDelimiter);
+		CurvesBuilder << InCurveName;
+		CurvesBuilder << USkeleton::CurveTagDelimiter;
 	});
 
 	OutTags.Add(FAssetRegistryTag(USkeleton::CurveNameTag, CurvesBuilder.ToString(), FAssetRegistryTag::TT_Hidden));
diff --git a/Engine/Source/Runtime/Engine/Private/Animation/Skeleton.cpp b/Engine/Source/Runtime/Engine/Private/Animation/Skeleton.cpp
--- a/Engine/Source/Runtime/Engine/Private/Animation/Skeleton.cpp
+++ b/Engine/Source/Runtime/Engine/Private/Animation/Skeleton.cpp
@@ -140,7 +140,7 @@ bool USkeleton::IsCompatibleForEditor(const FAssetData& AssetData, const TCHAR*
 bool USkeleton::IsCompatibleForEditor(const FString& SkeletonAssetString) const
 {
 	// First check against itself.
-	const FString SkeletonString = FAssetData(this).GetExportTextName();
+	const FString SkeletonString = FAssetData(this, FAssetData::ECreationFlags::SkipAssetRegistryTagsGathering).GetExportTextName();
 	if (SkeletonString == SkeletonAssetString)
 	{
 		return true;
