- Fix incorrect task tag for PSC gamethread CNC tasks

diff --git a/Engine/Source/Runtime/Engine/Private/Particles/ParticleComponents.cpp b/Engine/Source/Runtime/Engine/Private/Particles/ParticleComponents.cpp
--- a/Engine/Source/Runtime/Engine/Private/Particles/ParticleComponents.cpp
+++ b/Engine/Source/Runtime/Engine/Private/Particles/ParticleComponents.cpp
@@ -1959,7 +1959,7 @@ void UParticleEmitter::CacheEmitterModuleInfo()
 
 float UParticleEmitter::GetQualityLevelSpawnRateMult()
 {
-	int32 EffectsQuality = Scalability::GetEffectsQualityDirect(IsInGameThread());
+	int32 EffectsQuality = Scalability::GetEffectsQualityDirect(IsInGameThread() || IsInParallelGameThread());
 	int32 ReferenceLevel = CVarQLSpawnRateReferenceLevel.GetValueOnAnyThread(true);
 	float Level = (ReferenceLevel - EffectsQuality);
 	float Q = FMath::Pow(QualityLevelSpawnRateScale, Level);
diff --git a/Engine/Source/Runtime/Engine/Private/Particles/ParticleSystemManager.cpp b/Engine/Source/Runtime/Engine/Private/Particles/ParticleSystemManager.cpp
--- a/Engine/Source/Runtime/Engine/Private/Particles/ParticleSystemManager.cpp
+++ b/Engine/Source/Runtime/Engine/Private/Particles/ParticleSystemManager.cpp
@@ -113,7 +113,8 @@ public:
 
 	void DoTask(ENamedThreads::Type CurrentThread, const FGraphEventRef& MyCompletionGraphEvent)
 	{
-		FTaskTagScope Scope(ETaskTag::EParallelRenderingThread);
+		//FTaskTagScope Scope(ETaskTag::EParallelRenderingThread);
+		FTaskTagScope Scope(ETaskTag::EParallelGameThread);
 		SCOPE_CYCLE_COUNTER(STAT_PSCMan_AsyncBatch);
 
 // 		FString Ticked;
