-Removed incorrect check
-Fix for Project.Maps.AllInPIE.Maps.Framework.AI.WorldPartition.TM-AI_WorldPartition failing

diff --git a/Engine/Source/Runtime/Engine/Private/StaticMeshRender.cpp b/Engine/Source/Runtime/Engine/Private/StaticMeshRender.cpp
--- a/Engine/Source/Runtime/Engine/Private/StaticMeshRender.cpp
+++ b/Engine/Source/Runtime/Engine/Private/StaticMeshRender.cpp
@@ -1147,7 +1147,7 @@ HHitProxy* FStaticMeshSceneProxy::CreateHitProxies(IPrimitiveComponent* Componen
 	if ( ComponentInterface->GetOwner() )
 	{
 		// Sanity check for a case we'll not be handling anymore
-		check (! (ComponentInterface->GetUObject<UBrushComponent>() || ComponentInterface->GetOwner<ABrush>()));
+		check (! (ComponentInterface->GetUObject<UBrushComponent>()));
 
 		// Generate separate hit proxies for each sub mesh, so that we can perform hit tests against each section for applying materials
 		// to each one.
