-Always setup our ptr to the OwnerLastRenderTimer in UPrimitiveComponent::CreateRenderState_Concurrent, owner can change without corresponding call to OnRegister()

diff --git a/Engine/Source/Runtime/Engine/Private/Components/PrimitiveComponent.cpp b/Engine/Source/Runtime/Engine/Private/Components/PrimitiveComponent.cpp
--- a/Engine/Source/Runtime/Engine/Private/Components/PrimitiveComponent.cpp
+++ b/Engine/Source/Runtime/Engine/Private/Components/PrimitiveComponent.cpp
@@ -569,6 +569,10 @@ void UPrimitiveComponent::CreateRenderState_Concurrent(FRegisterComponentContext
 		CachedMaxDrawDistance = bNeverCull ? 0.f : LDMaxDrawDistance;
 	}
 
+	// Always setup our ptr to the OwnerLastRenderTimer for rendering time feedback from the renderer
+	// The owner can change after calls to OnRegister so we must resynchronize this value
+	SceneData.OwnerLastRenderTimePtr = FActorLastRenderTime::GetPtr(GetOwner());
+
 	Super::CreateRenderState_Concurrent(Context);
 
 	UpdateBounds();
