World Partition HLODs: added -ReportOnly option to setup HLOD actors and only report changes needed.

diff --git a/Engine/Source/Editor/UnrealEd/Private/WorldPartition/WorldPartitionHLODsBuilder.cpp b/Engine/Source/Editor/UnrealEd/Private/WorldPartition/WorldPartitionHLODsBuilder.cpp
--- a/Engine/Source/Editor/UnrealEd/Private/WorldPartition/WorldPartitionHLODsBuilder.cpp
+++ b/Engine/Source/Editor/UnrealEd/Private/WorldPartition/WorldPartitionHLODsBuilder.cpp
@@ -169,6 +169,7 @@ UWorldPartitionHLODsBuilder::UWorldPartitionHLODsBuilder(const FObjectInitialize
 
 	bDistributedBuild = HasParam("DistributedBuild");
 	bForceBuild = HasParam("RebuildHLODs");
+	bReportOnly = HasParam("ReportOnly");
 
 	GetParamValue("BuildManifest=", BuildManifest);
 	GetParamValue("BuilderIdx=", BuilderIdx);
@@ -331,24 +332,27 @@ bool UWorldPartitionHLODsBuilder::RunInternal(UWorld* InWorld, const FCellInfo&
 		bRet = SetupHLODActors();
 	}
 
-	if (bRet && ShouldRunStep(EHLODBuildStep::HLOD_Build))
+	if (!bReportOnly)
 	{
-		bRet = BuildHLODActors();
-	}
+		if (bRet && ShouldRunStep(EHLODBuildStep::HLOD_Build))
+		{
+			bRet = BuildHLODActors();
+		}
 
-	if (bRet && ShouldRunStep(EHLODBuildStep::HLOD_Delete))
-	{
-		bRet = DeleteHLODActors();
-	}
+		if (bRet && ShouldRunStep(EHLODBuildStep::HLOD_Delete))
+		{
+			bRet = DeleteHLODActors();
+		}
 
-	if (bRet && ShouldRunStep(EHLODBuildStep::HLOD_Finalize))
-	{
-		bRet = SubmitHLODActors();
-	}
+		if (bRet && ShouldRunStep(EHLODBuildStep::HLOD_Finalize))
+		{
+			bRet = SubmitHLODActors();
+		}
 
-	if (bRet && ShouldRunStep(EHLODBuildStep::HLOD_Stats))
-	{
-		bRet = DumpStats();
+		if (bRet && ShouldRunStep(EHLODBuildStep::HLOD_Stats))
+		{
+			bRet = DumpStats();
+		}
 	}
 
 	WorldPartition = nullptr;
@@ -379,8 +383,11 @@ bool UWorldPartitionHLODsBuilder::SetupHLODActors()
 			GEngine->OnActorFolderAdded().Remove(ActorFolderAddedDelegateHandle);
 		};
 
-		const bool bCreateActorsOnly = true;
-		WorldPartition->GenerateHLOD(SourceControlHelper, bCreateActorsOnly);
+		UWorldPartition::FSetupHLODActorsParams SetupHLODActorsParams = UWorldPartition::FSetupHLODActorsParams()
+			.SetSourceControlHelper(SourceControlHelper)
+			.SetReportOnly(bReportOnly);
+
+		WorldPartition->SetupHLODActors(SetupHLODActorsParams);
 
 		// When performing a distributed build, ensure our work folder is empty
 		if (IsDistributedBuild())
diff --git a/Engine/Source/Editor/UnrealEd/Public/WorldPartition/WorldPartitionHLODsBuilder.h b/Engine/Source/Editor/UnrealEd/Public/WorldPartition/WorldPartitionHLODsBuilder.h
--- a/Engine/Source/Editor/UnrealEd/Public/WorldPartition/WorldPartitionHLODsBuilder.h
+++ b/Engine/Source/Editor/UnrealEd/Public/WorldPartition/WorldPartitionHLODsBuilder.h
@@ -124,6 +124,7 @@ private:
 
 	bool bDistributedBuild;
 	bool bForceBuild;
+	bool bReportOnly;
 	FString BuildManifest;
 	int32 BuilderIdx;
 	int32 BuilderCount;
diff --git a/Engine/Source/Runtime/Engine/Private/WorldPartition/RuntimeSpatialHash/RuntimeSpatialHashHLOD.cpp b/Engine/Source/Runtime/Engine/Private/WorldPartition/RuntimeSpatialHash/RuntimeSpatialHashHLOD.cpp
--- a/Engine/Source/Runtime/Engine/Private/WorldPartition/RuntimeSpatialHash/RuntimeSpatialHashHLOD.cpp
+++ b/Engine/Source/Runtime/Engine/Private/WorldPartition/RuntimeSpatialHash/RuntimeSpatialHashHLOD.cpp
@@ -4,7 +4,6 @@
 #include "WorldPartition/WorldPartitionStreamingGeneration.h"
 
 #if WITH_EDITOR
-
 #include "HAL/PlatformFile.h"
 #include "WorldPartition/RuntimeSpatialHash/RuntimeSpatialHashGridHelper.h"
 
@@ -16,27 +15,21 @@
 #include "WorldPartition/HLOD/IWorldPartitionHLODUtilitiesModule.h"
 #include "WorldPartition/DataLayer/DataLayerManager.h"
 #include "WorldPartition/DataLayer/DataLayerUtils.h"
-
 #include "WorldPartition/WorldPartition.h"
 #include "WorldPartition/WorldPartitionHelpers.h"
 #include "WorldPartition/ContentBundle/ContentBundleActivationScope.h"
 
 #include "UObject/GCObjectScopeGuard.h"
+#include "UObject/SavePackage.h"
 #include "Engine/Engine.h"
+#include "EngineUtils.h"
+
 #include "HAL/PlatformFileManager.h"
 #include "Misc/ScopedSlowTask.h"
 #include "Modules/ModuleManager.h"
 
-
-#include "EngineUtils.h"
-
-
-#include "UObject/SavePackage.h"
-#include <limits> // IWYU pragma: keep
-
 DEFINE_LOG_CATEGORY_STATIC(LogWorldPartitionRuntimeSpatialHashHLOD, Log, All);
 
-
 static void SavePackage(UPackage* Package, ISourceControlHelper* SourceControlHelper)
 {
 	if (SourceControlHelper)
@@ -118,7 +111,7 @@ static void GameTick(UWorld* InWorld)
 	}
 }
 
-static TArray<FGuid> GenerateHLODsForGrid(UWorldPartition* WorldPartition, const IStreamingGenerationContext* StreamingGenerationContext, const FSpatialHashRuntimeGrid& RuntimeGrid, uint32 HLODLevel, FHLODCreationContext& Context, ISourceControlHelper* SourceControlHelper, bool bCreateActorsOnly, const TArray<const IStreamingGenerationContext::FActorSetInstance*>& ActorSetInstances)
+static TArray<FGuid> GenerateHLODActorsForGrid(UWorldPartition* WorldPartition, const IStreamingGenerationContext* StreamingGenerationContext, const UWorldPartition::FSetupHLODActorsParams& Params, const FSpatialHashRuntimeGrid& RuntimeGrid, uint32 HLODLevel, FHLODCreationContext& Context, ISourceControlHelper* SourceControlHelper, const TArray<const IStreamingGenerationContext::FActorSetInstance*>& ActorSetInstances, TArray<FName>& NewActors, TArray<FName>& DirtyActors)
 {
 	const IStreamingGenerationContext::FActorSetContainer* MainActorSetContainer = StreamingGenerationContext->GetMainWorldContainer();
 	const FBox WorldBounds = StreamingGenerationContext->GetWorldBounds();
@@ -126,13 +119,13 @@ static TArray<FGuid> GenerateHLODsForGrid(UWorldPartition* WorldPartition, const
 	const FSquare2DGridHelper PartitionedActors = GetPartitionedActors(WorldBounds, RuntimeGrid, ActorSetInstances);
 	const FSquare2DGridHelper::FGridLevel::FGridCell& AlwaysLoadedCell = PartitionedActors.GetAlwaysLoadedCell();
 
-	auto ShouldGenerateHLODs = [&AlwaysLoadedCell](const FSquare2DGridHelper::FGridLevel::FGridCell& GridCell, const FSquare2DGridHelper::FGridLevel::FGridCellDataChunk& GridCellDataChunk)
+	auto ShouldGenerateHLODActors = [&AlwaysLoadedCell](const FSquare2DGridHelper::FGridLevel::FGridCell& GridCell, const FSquare2DGridHelper::FGridLevel::FGridCellDataChunk& GridCellDataChunk)
 	{
 		const bool bIsCellAlwaysLoaded = &GridCell == &AlwaysLoadedCell;
 		const bool bChunkHasActors = !GridCellDataChunk.GetActorSetInstances().IsEmpty();
 
-		const bool bShouldGenerateHLODs = !bIsCellAlwaysLoaded && bChunkHasActors;
-		return bShouldGenerateHLODs;
+		const bool bShouldGenerateHLODActors = !bIsCellAlwaysLoaded && bChunkHasActors;
+		return bShouldGenerateHLODActors;
 	};
 
 	UWorld* OuterWorld = WorldPartition->GetTypedOuter<UWorld>();
@@ -143,8 +136,8 @@ static TArray<FGuid> GenerateHLODsForGrid(UWorldPartition* WorldPartition, const
 	{
 		for (const FSquare2DGridHelper::FGridLevel::FGridCellDataChunk& GridCellDataChunk : GridCell.GetDataChunks())
 		{
-			const bool bShouldGenerateHLODs = ShouldGenerateHLODs(GridCell, GridCellDataChunk);
-			if (bShouldGenerateHLODs)
+			const bool bShouldGenerateHLODActors = ShouldGenerateHLODActors(GridCell, GridCellDataChunk);
+			if (bShouldGenerateHLODActors)
 			{
 				NbCellsToProcess++;
 			}
@@ -172,8 +165,8 @@ static TArray<FGuid> GenerateHLODsForGrid(UWorldPartition* WorldPartition, const
 			// Fake tick
 			GameTick(WorldPartition->GetWorld());
 
-			const bool bShouldGenerateHLODs = ShouldGenerateHLODs(GridCell, GridCellDataChunk);
-			if (bShouldGenerateHLODs)
+			const bool bShouldGenerateHLODActors = ShouldGenerateHLODActors(GridCell, GridCellDataChunk);
+			if (bShouldGenerateHLODActors)
 			{
 				SlowTask.EnterProgressFrame(1);
 
@@ -187,8 +180,7 @@ static TArray<FGuid> GenerateHLODsForGrid(UWorldPartition* WorldPartition, const
 				UE_LOG(LogWorldPartitionRuntimeSpatialHashHLOD, Display, TEXT("[%d / %d] Processing cell %s..."), (int32)SlowTask.CompletedWork + 1, (int32)SlowTask.TotalAmountOfWork, *CellName);
 
 				// We know the cell's 2D bound but must figure out the bounds in Z
-				FBox CellBounds = FBox(FVector(CellBounds2D.Min,  std::numeric_limits<double>::max()), 
-									   FVector(CellBounds2D.Max, -std::numeric_limits<double>::max()));
+				FBox CellBounds = FBox(FVector(CellBounds2D.Min, MAX_dbl), FVector(CellBounds2D.Max, -MAX_dbl));
 
 				TArray<IStreamingGenerationContext::FActorInstance> ActorInstances;
 				for (const IStreamingGenerationContext::FActorSetInstance* ActorSetInstance : GridCellDataChunk.GetActorSetInstances())
@@ -231,19 +223,22 @@ static TArray<FGuid> GenerateHLODsForGrid(UWorldPartition* WorldPartition, const
 						if (CellHLODActorPackage->HasAnyPackageFlags(PKG_NewlyCreated))
 						{
 							NewCellHLODActors.Add(CellHLODActor);
+							NewActors.Add(CellHLODActor->GetFName());
 						}
-					}
-
-					for (AWorldPartitionHLOD* CellHLODActor : CellHLODActors)
-					{
-						if (!bCreateActorsOnly)
+						else if (CellHLODActorPackage->IsDirty())
 						{
-							CellHLODActor->BuildHLOD();
+							DirtyActors.Add(CellHLODActor->GetFName());
 						}
+					}
 
-						if (CellHLODActor->GetPackage()->IsDirty())
+					if (!Params.bReportOnly)
+					{
+						for (AWorldPartitionHLOD* CellHLODActor : CellHLODActors)
 						{
-							SavePackage(CellHLODActor->GetPackage(), SourceControlHelper);
+							if (CellHLODActor->GetPackage()->IsDirty())
+							{
+								SavePackage(CellHLODActor->GetPackage(), SourceControlHelper);
+							}
 						}
 					}
 
@@ -430,9 +425,20 @@ static void UpdateHLODGridsActors(UWorld* World, const UActorDescContainer* Cont
 	}
 }
 
-bool UWorldPartitionRuntimeSpatialHash::GenerateHLOD(ISourceControlHelper* SourceControlHelper, const IStreamingGenerationContext* StreamingGenerationContext, bool bCreateActorsOnly) const
+bool UWorldPartitionRuntimeSpatialHash::SetupHLODActors(const IStreamingGenerationContext* StreamingGenerationContext, const UWorldPartition::FSetupHLODActorsParams& Params) const
 {
-	TRACE_CPUPROFILER_EVENT_SCOPE(UWorldPartitionRuntimeSpatialHash::GenerateHLOD);
+	TRACE_CPUPROFILER_EVENT_SCOPE(UWorldPartitionRuntimeSpatialHash::SetupHLODActors);
+
+	ISourceControlHelper* SourceControlHelper = Params.SourceControlHelper;
+
+	TArray<FName> NewActors;
+	TArray<FName> DirtyActors;
+	TArray<FName> DeletedActors;
+
+	if (Params.bReportOnly)
+	{
+		LogWorldPartitionRuntimeSpatialHashHLOD.SetVerbosity(ELogVerbosity::Verbose);
+	}
 
 	if (!Grids.Num())
 	{
@@ -480,7 +486,7 @@ bool UWorldPartitionRuntimeSpatialHash::GenerateHLOD(ISourceControlHelper* Sourc
 		if (!FoundIndex)
 		{
 			//@todo_ow should this be done upstream?
-			UE_LOG(LogWorldPartition, Error, TEXT("Invalid partition grid '%s' referenced by actor cluster"), *ActorSetInstance.RuntimeGrid.ToString());
+			UE_LOG(LogWorldPartitionRuntimeSpatialHashHLOD, Error, TEXT("Invalid partition grid '%s' referenced by actor cluster"), *ActorSetInstance.RuntimeGrid.ToString());
 		}
 
 		int32 GridIndex = FoundIndex ? *FoundIndex : 0;
@@ -493,27 +499,30 @@ bool UWorldPartitionRuntimeSpatialHash::GenerateHLOD(ISourceControlHelper* Sourc
 
 	const UDataLayerManager* DataLayerManager = WorldPartition->GetDataLayerManager();
 
-	auto GenerateHLODs = [&GridsHLODActors, &ActorDescViews, &MainActorSetContainer, StreamingGenerationContext, WorldPartition, DataLayerManager, &Context, SourceControlHelper, bCreateActorsOnly]
+	auto GenerateHLODActors = [&GridsHLODActors, &ActorDescViews, &MainActorSetContainer, StreamingGenerationContext, &Params, WorldPartition, DataLayerManager, &Context, SourceControlHelper, &NewActors, &DirtyActors]
 	(const FSpatialHashRuntimeGrid& RuntimeGrid, uint32 HLODLevel, const TArray<const IStreamingGenerationContext::FActorSetInstance*>& ActorSetInstances)
 	{
 		// Generate HLODs for this grid
-		TArray<FGuid> HLODActors = GenerateHLODsForGrid(WorldPartition, StreamingGenerationContext, RuntimeGrid, HLODLevel, Context, SourceControlHelper, bCreateActorsOnly, ActorSetInstances);
+		TArray<FGuid> HLODActors = GenerateHLODActorsForGrid(WorldPartition, StreamingGenerationContext, Params, RuntimeGrid, HLODLevel, Context, SourceControlHelper, ActorSetInstances, NewActors, DirtyActors);
 
 		for (const FGuid& HLODActorGuid : HLODActors)
 		{
-			FWorldPartitionActorDesc* HLODActorDesc = WorldPartition->GetActorDesc(HLODActorGuid);
-			check(HLODActorDesc);
-			FWorldPartitionActorDescView* ActorDescView = ActorDescViews.Emplace_GetRef(MakeUnique<FWorldPartitionActorDescView>(HLODActorDesc)).Get();
-			TArray<FName> RuntimeDataLayerInstanceNames;
-			if (FDataLayerUtils::ResolveRuntimeDataLayerInstanceNames(DataLayerManager, *ActorDescView, *MainActorSetContainer->ActorDescViewMap, RuntimeDataLayerInstanceNames))
+			if (!Params.bReportOnly)
 			{
-				ActorDescView->SetRuntimeDataLayerInstanceNames(RuntimeDataLayerInstanceNames);
-			}
-			FActorDescViewMap* NonConstActorDescViewMap = const_cast<FActorDescViewMap*>(MainActorSetContainer->ActorDescViewMap);			
-			NonConstActorDescViewMap->Emplace(HLODActorGuid, *ActorDescView);
+				FWorldPartitionActorDesc* HLODActorDesc = WorldPartition->GetActorDesc(HLODActorGuid);
+				check(HLODActorDesc);
+				FWorldPartitionActorDescView* ActorDescView = ActorDescViews.Emplace_GetRef(MakeUnique<FWorldPartitionActorDescView>(HLODActorDesc)).Get();
+				TArray<FName> RuntimeDataLayerInstanceNames;
+				if (FDataLayerUtils::ResolveRuntimeDataLayerInstanceNames(DataLayerManager, *ActorDescView, *MainActorSetContainer->ActorDescViewMap, RuntimeDataLayerInstanceNames))
+				{
+					ActorDescView->SetRuntimeDataLayerInstanceNames(RuntimeDataLayerInstanceNames);
+				}
+				FActorDescViewMap* NonConstActorDescViewMap = const_cast<FActorDescViewMap*>(MainActorSetContainer->ActorDescViewMap);			
+				NonConstActorDescViewMap->Emplace(HLODActorGuid, *ActorDescView);
 
-			const FWorldPartitionActorDescView& HLODActorDescView = MainActorSetContainer->ActorDescViewMap->FindByGuidChecked(HLODActorGuid);
-			GridsHLODActors.FindOrAdd(HLODActorDescView.GetRuntimeGrid()).Add(HLODActorGuid);
+				const FWorldPartitionActorDescView& HLODActorDescView = MainActorSetContainer->ActorDescViewMap->FindByGuidChecked(HLODActorGuid);
+				GridsHLODActors.FindOrAdd(HLODActorDescView.GetRuntimeGrid()).Add(HLODActorGuid);
+			}
 		}
 	};
 
@@ -521,7 +530,7 @@ bool UWorldPartitionRuntimeSpatialHash::GenerateHLOD(ISourceControlHelper* Sourc
 	for (int32 GridIndex = 0; GridIndex < Grids.Num(); GridIndex++)
 	{
 		// Generate HLODs for this grid - retrieve actors from the world partition
-		GenerateHLODs(Grids[GridIndex], 0, GridActorSetInstances[GridIndex]);
+		GenerateHLODActors(Grids[GridIndex], 0, GridActorSetInstances[GridIndex]);
 	}
 
 	// Now, go on and create HLOD actors from HLOD grids (HLOD 1-N)
@@ -566,7 +575,7 @@ bool UWorldPartitionRuntimeSpatialHash::GenerateHLOD(ISourceControlHelper* Sourc
 
 		// Generate HLODs for this grid - retrieve actors from our ValidHLODActors map
 		// We can't rely on actor descs for newly created HLOD actors
-		GenerateHLODs(HLODGrids[HLODGridName], HLODLayersLevels[HLODGrids[HLODGridName].HLODLayer] + 1, HLODActorSetInstancePtrs);
+		GenerateHLODActors(HLODGrids[HLODGridName], HLODLayersLevels[HLODGrids[HLODGridName].HLODLayer] + 1, HLODActorSetInstancePtrs);
 	}
 
 	auto DeleteHLODActor = [&SourceControlHelper, WorldPartition](FWorldPartitionHandle ActorHandle)
@@ -580,11 +589,34 @@ bool UWorldPartitionRuntimeSpatialHash::GenerateHLOD(ISourceControlHelper* Sourc
 	// Destroy all unreferenced HLOD actors
 	for (const auto& HLODActorPair : Context.HLODActorDescs)
 	{
-		DeleteHLODActor(HLODActorPair.Value);
+		DeletedActors.Add(HLODActorPair.Key);
+
+		if (!Params.bReportOnly)
+		{
+			DeleteHLODActor(HLODActorPair.Value);
+		}
+	}
+
+	if (!Params.bReportOnly)
+	{
+		// Create/destroy HLOD grid actors
+		UpdateHLODGridsActors(GetWorld(), MainContainerCollection->GetMainActorDescContainer(), HLODGrids, SourceControlHelper);
 	}
 
-	// Create/destroy HLOD grid actors
-	UpdateHLODGridsActors(GetWorld(), MainContainerCollection->GetMainActorDescContainer(), HLODGrids, SourceControlHelper);
+	auto DumpActorsStats = [](const TCHAR* StatType, const TArray<FName>& Actors)
+	{
+		UE_LOG(LogWorldPartitionRuntimeSpatialHashHLOD, Display, TEXT("\t%s actors (%d):"), StatType, Actors.Num());
+		for (FName ActorName : Actors)
+		{
+			UE_LOG(LogWorldPartitionRuntimeSpatialHashHLOD, Display, TEXT("\t\t%s:"), *ActorName.ToString());
+		}
+	};
+
+	UE_LOG(LogWorldPartitionRuntimeSpatialHashHLOD, Display, TEXT("SetupHLODActors stats:"));
+
+	DumpActorsStats(TEXT("New"), NewActors);
+	DumpActorsStats(TEXT("Dirty"), DirtyActors);
+	DumpActorsStats(TEXT("Deleted"), DeletedActors);
 
 	return true;
 }
diff --git a/Engine/Source/Runtime/Engine/Private/WorldPartition/WorldPartitionStreamingGeneration.cpp b/Engine/Source/Runtime/Engine/Private/WorldPartition/WorldPartitionStreamingGeneration.cpp
--- a/Engine/Source/Runtime/Engine/Private/WorldPartition/WorldPartitionStreamingGeneration.cpp
+++ b/Engine/Source/Runtime/Engine/Private/WorldPartition/WorldPartitionStreamingGeneration.cpp
@@ -1373,9 +1373,9 @@ URuntimeHashExternalStreamingObjectBase* UWorldPartition::FlushStreamingToExtern
 	return ExternalStreamingObject;
 }
 
-void UWorldPartition::GenerateHLOD(ISourceControlHelper* SourceControlHelper, bool bCreateActorsOnly)
+void UWorldPartition::SetupHLODActors(const FSetupHLODActorsParams& Params)
 {
-	ForEachActorDescContainer([this, &SourceControlHelper, bCreateActorsOnly](UActorDescContainer* InActorDescContainer)
+	ForEachActorDescContainer([this, &Params](UActorDescContainer* InActorDescContainer)
 	{
 		FStreamingGenerationLogErrorHandler LogErrorHandler;
 		FWorldPartitionStreamingGenerator::FWorldPartitionStreamingGeneratorParams StreamingGeneratorParams = FWorldPartitionStreamingGenerator::FWorldPartitionStreamingGeneratorParams()
@@ -1393,7 +1393,7 @@ void UWorldPartition::GenerateHLOD(ISourceControlHelper* SourceControlHelper, bo
 		FHierarchicalLogArchive HierarchicalLogAr(*LogFileAr);
 		StreamingGenerator.DumpStateLog(HierarchicalLogAr);
 
-		RuntimeHash->GenerateHLOD(SourceControlHelper, StreamingGenerator.GetStreamingGenerationContext(ActorDescCollection), bCreateActorsOnly);
+		RuntimeHash->SetupHLODActors(StreamingGenerator.GetStreamingGenerationContext(ActorDescCollection), Params);
 	});	
 }
 
diff --git a/Engine/Source/Runtime/Engine/Public/WorldPartition/WorldPartition.h b/Engine/Source/Runtime/Engine/Public/WorldPartition/WorldPartition.h
--- a/Engine/Source/Runtime/Engine/Public/WorldPartition/WorldPartition.h
+++ b/Engine/Source/Runtime/Engine/Public/WorldPartition/WorldPartition.h
@@ -289,7 +289,23 @@ public:
 	
 	UHLODLayer* GetDefaultHLODLayer() const { return DefaultHLODLayer; }
 	void SetDefaultHLODLayer(UHLODLayer* InDefaultHLODLayer) { DefaultHLODLayer = InDefaultHLODLayer; }
-	ENGINE_API void GenerateHLOD(ISourceControlHelper* SourceControlHelper, bool bCreateActorsOnly);
+
+	/* Struct of optional parameters passed to SetupHLODActors function. */
+	struct FSetupHLODActorsParams
+	{
+		FSetupHLODActorsParams()
+			: SourceControlHelper(nullptr)
+			, bReportOnly(false)
+		{}
+
+		ISourceControlHelper* SourceControlHelper;
+		bool bReportOnly;
+
+		FSetupHLODActorsParams& SetSourceControlHelper(ISourceControlHelper* InSourceControlHelper) { SourceControlHelper = InSourceControlHelper; return *this; }
+		FSetupHLODActorsParams& SetReportOnly(bool bInReportOnly) { bReportOnly = bInReportOnly; return *this; }
+	};
+
+	ENGINE_API void SetupHLODActors(const FSetupHLODActorsParams& Params);
 
 	// Debugging
 	ENGINE_API void DrawRuntimeHashPreview();
diff --git a/Engine/Source/Runtime/Engine/Public/WorldPartition/WorldPartitionRuntimeHash.h b/Engine/Source/Runtime/Engine/Public/WorldPartition/WorldPartitionRuntimeHash.h
--- a/Engine/Source/Runtime/Engine/Public/WorldPartition/WorldPartitionRuntimeHash.h
+++ b/Engine/Source/Runtime/Engine/Public/WorldPartition/WorldPartitionRuntimeHash.h
@@ -89,7 +89,7 @@ class UWorldPartitionRuntimeHash : public UObject
 	ENGINE_API virtual TArray<UWorldPartitionRuntimeCell*> GetAlwaysLoadedCells() const;
 	ENGINE_API virtual bool GenerateStreaming(class UWorldPartitionStreamingPolicy* StreamingPolicy, const IStreamingGenerationContext* StreamingGenerationContext, TArray<FString>* OutPackagesToGenerate);
 	ENGINE_API virtual void FlushStreaming();
-	virtual bool GenerateHLOD(ISourceControlHelper* SourceControlHelper, const IStreamingGenerationContext* StreamingGenerationContext, bool bCreateActorsOnly) const { return false; }
+	virtual bool SetupHLODActors(const IStreamingGenerationContext* StreamingGenerationContext, const UWorldPartition::FSetupHLODActorsParams& Params) const { return false; }
 	virtual bool IsValidGrid(FName GridName) const { return false; }
 	virtual void DrawPreview() const {}
 
diff --git a/Engine/Source/Runtime/Engine/Public/WorldPartition/WorldPartitionRuntimeSpatialHash.h b/Engine/Source/Runtime/Engine/Public/WorldPartition/WorldPartitionRuntimeSpatialHash.h
--- a/Engine/Source/Runtime/Engine/Public/WorldPartition/WorldPartitionRuntimeSpatialHash.h
+++ b/Engine/Source/Runtime/Engine/Public/WorldPartition/WorldPartitionRuntimeSpatialHash.h
@@ -276,7 +276,7 @@ public:
 	ENGINE_API virtual void SetDefaultValues() override;
 	virtual bool SupportsHLODs() const override { return true; }
 	ENGINE_API virtual void FlushStreaming() override;
-	ENGINE_API virtual bool GenerateHLOD(ISourceControlHelper* SourceControlHelper, const IStreamingGenerationContext* StreamingGenerationContext, bool bCreateActorsOnly) const override;
+	ENGINE_API virtual bool SetupHLODActors(const IStreamingGenerationContext* StreamingGenerationContext, const UWorldPartition::FSetupHLODActorsParams& Params) const override;
 	ENGINE_API virtual bool IsValidGrid(FName GridName) const override;
 	ENGINE_API virtual void DrawPreview() const override;
 
