UEngine::LoadMap & UGameEngine::PreExit: Make sure to call SetShouldForceUnloadStreamingLevels(true)/FlushLevelStreaming(EFlushLevelStreamingType::Visibility) before RouteEndPlay which can Mark plugin packages as garbage

#rnx

diff --git a/Engine/Source/Runtime/Engine/Private/GameEngine.cpp b/Engine/Source/Runtime/Engine/Private/GameEngine.cpp
--- a/Engine/Source/Runtime/Engine/Private/GameEngine.cpp
+++ b/Engine/Source/Runtime/Engine/Private/GameEngine.cpp
@@ -1236,6 +1236,13 @@ void UGameEngine::PreExit()
 			// Shut down any existing game connections
 			ShutdownWorldNetDriver(World);
 
+			// Force mark all streaming levels for stream out
+			World->bIsLevelStreamingFrozen = false;
+			World->SetShouldForceUnloadStreamingLevels(true);
+
+			// Make sure there are no pending visibility requests.
+			World->FlushLevelStreaming(EFlushLevelStreamingType::Visibility);
+						
 			for (FActorIterator ActorIt(World); ActorIt; ++ActorIt)
 			{
 				ActorIt->RouteEndPlay(EEndPlayReason::Quit);
@@ -1246,7 +1253,6 @@ void UGameEngine::PreExit()
 				World->GetGameInstance()->Shutdown();
 			}
 
-			World->FlushLevelStreaming(EFlushLevelStreamingType::Visibility);
 			World->CleanupWorld();
 		}
 	}
diff --git a/Engine/Source/Runtime/Engine/Private/UnrealEngine.cpp b/Engine/Source/Runtime/Engine/Private/UnrealEngine.cpp
--- a/Engine/Source/Runtime/Engine/Private/UnrealEngine.cpp
+++ b/Engine/Source/Runtime/Engine/Private/UnrealEngine.cpp
@@ -14915,6 +14915,10 @@ bool UEngine::LoadMap( FWorldContext& WorldContext, FURL URL, class UPendingNetG
 		// Clean up networking
 		ShutdownWorldNetDriver(WorldContext.World());
 
+		// Force mark all streaming levels for stream out
+		WorldContext.World()->bIsLevelStreamingFrozen = false;
+		WorldContext.World()->SetShouldForceUnloadStreamingLevels(true);
+
 		// Make sure there are no pending visibility requests.
 		WorldContext.World()->FlushLevelStreaming(EFlushLevelStreamingType::Visibility);
 
diff --git a/Engine/Source/Runtime/Engine/Private/World.cpp b/Engine/Source/Runtime/Engine/Private/World.cpp
--- a/Engine/Source/Runtime/Engine/Private/World.cpp
+++ b/Engine/Source/Runtime/Engine/Private/World.cpp
@@ -2051,6 +2051,10 @@ void UWorld::InitWorld(const InitializationValues IVS)
 		return;
 	}
 
+	// Reset flags in case of world reuse
+	bIsLevelStreamingFrozen = false;
+	bShouldForceUnloadStreamingLevels = false;
+
 	FCoreUObjectDelegates::GetPostGarbageCollect().AddUObject(this, &UWorld::OnPostGC);
 
 	InitializeSubsystems();
