Don't show Level Instance properties on the World Settings

diff --git a/Engine/Source/Runtime/Engine/Classes/GameFramework/WorldSettings.h b/Engine/Source/Runtime/Engine/Classes/GameFramework/WorldSettings.h
--- a/Engine/Source/Runtime/Engine/Classes/GameFramework/WorldSettings.h
+++ b/Engine/Source/Runtime/Engine/Classes/GameFramework/WorldSettings.h
@@ -457,7 +457,7 @@ struct FBroadphaseSettings
 /**
  * Actor containing all script accessible world properties.
  */
-UCLASS(config=game, hidecategories=(Actor, Advanced, Display, Events, Object, Attachment, Info, Input, Blueprint, Layers, Tags, Replication), showcategories=(Rendering, WorldPartition, "Input|MouseInput", "Input|TouchInput"), notplaceable, MinimalAPI)
+UCLASS(config=game, hidecategories=(Actor, Advanced, Display, Events, Object, Attachment, Info, Input, Blueprint, Layers, Tags, Replication, LevelInstance), showcategories=(Rendering, WorldPartition, "Input|MouseInput", "Input|TouchInput"), notplaceable, MinimalAPI)
 class AWorldSettings : public AInfo, public IInterface_AssetUserData
 {
 	GENERATED_UCLASS_BODY()
