implement a GetComponentVelocity override for cluster unions

diff --git a/Engine/Source/Runtime/Engine/Classes/PhysicsEngine/ClusterUnionComponent.h b/Engine/Source/Runtime/Engine/Classes/PhysicsEngine/ClusterUnionComponent.h
--- a/Engine/Source/Runtime/Engine/Classes/PhysicsEngine/ClusterUnionComponent.h
+++ b/Engine/Source/Runtime/Engine/Classes/PhysicsEngine/ClusterUnionComponent.h
@@ -296,6 +296,7 @@ public:
 	ENGINE_API virtual void OnReceiveReplicatedState(const FVector X, const FQuat R, const FVector V, const FVector W) override;
 	ENGINE_API virtual void OnUpdateTransform(EUpdateTransformFlags UpdateTransformFlags, ETeleportType Teleport) override;
 	ENGINE_API virtual FBoxSphereBounds CalcBounds(const FTransform& LocalToWorld) const override;
+	ENGINE_API virtual FVector GetComponentVelocity() const override;
 	//~ End USceneComponent Interface
 
 	//~ Begin IPhysicsComponent Interface.
diff --git a/Engine/Source/Runtime/Engine/Private/PhysicsEngine/ClusterUnionComponent.cpp b/Engine/Source/Runtime/Engine/Private/PhysicsEngine/ClusterUnionComponent.cpp
--- a/Engine/Source/Runtime/Engine/Private/PhysicsEngine/ClusterUnionComponent.cpp
+++ b/Engine/Source/Runtime/Engine/Private/PhysicsEngine/ClusterUnionComponent.cpp
@@ -691,6 +691,18 @@ FBoxSphereBounds UClusterUnionComponent::CalcBounds(const FTransform& LocalToWor
 	return CachedLocalBounds.TransformBy(LocalToWorld);
 }
 
+FVector UClusterUnionComponent::GetComponentVelocity() const
+{
+	if (!PhysicsProxy)
+	{
+		return Super::GetComponentVelocity();
+	}
+
+	Chaos::FPhysicsObjectHandle PhysicsObject = PhysicsProxy->GetPhysicsObjectHandle();
+	FLockedReadPhysicsObjectExternalInterface Interface = FPhysicsObjectExternalInterface::LockRead({&PhysicsObject, 1});
+	return Interface->GetV(PhysicsObject);
+}
+
 bool UClusterUnionComponent::ShouldCreatePhysicsState() const
 {
 	return true;
