-Renamed IsEditorOnly to IsOwnerEditorOnly
-Initialized with the correct value, to restore previous behavior

diff --git a/Engine/Source/Runtime/Engine/Private/PrimitiveSceneProxy.cpp b/Engine/Source/Runtime/Engine/Private/PrimitiveSceneProxy.cpp
--- a/Engine/Source/Runtime/Engine/Private/PrimitiveSceneProxy.cpp
+++ b/Engine/Source/Runtime/Engine/Private/PrimitiveSceneProxy.cpp
@@ -319,9 +319,10 @@ FPrimitiveSceneProxyDesc::FPrimitiveSceneProxyDesc(const UPrimitiveComponent* In
 				ActorOwners.Add(CurrentOwner);
 			}
 		}
+		
+		bIsOwnerEditorOnly = InComponent->GetOwner()->IsEditorOnly(); 
 	}
 	bSupportsWorldPositionOffsetVelocity = InComponent->SupportsWorldPositionOffsetVelocity();
-	bIsEditorOnly = InComponent->IsEditorOnly(); 
 	bIsInstancedStaticMesh = Cast<UInstancedStaticMeshComponent>(InComponent) != nullptr; 
 
 	Mobility = InComponent->Mobility;;
@@ -576,7 +577,7 @@ FPrimitiveSceneProxy::FPrimitiveSceneProxy(const FPrimitiveSceneProxyDesc& InPro
 #if WITH_EDITOR
 		// cache the actor's group membership
 		HiddenEditorViews = InProxyDesc.GetHiddenEditorViews();
-		DrawInAnyEditMode = InProxyDesc.IsEditorOnly();
+		DrawInAnyEditMode = InProxyDesc.IsOwnerEditorOnly();
 		bIsFoliage = InProxyDesc.IsOwnedByFoliage();
 #endif
 	}	
diff --git a/Engine/Source/Runtime/Engine/Public/PrimitiveSceneProxyDesc.h b/Engine/Source/Runtime/Engine/Public/PrimitiveSceneProxyDesc.h
--- a/Engine/Source/Runtime/Engine/Public/PrimitiveSceneProxyDesc.h
+++ b/Engine/Source/Runtime/Engine/Public/PrimitiveSceneProxyDesc.h
@@ -58,7 +58,7 @@ struct FPrimitiveSceneProxyDesc
 		bIsHidden = false;
 		bIsHiddenEd = false;
 		bSupportsWorldPositionOffsetVelocity = true;
-		bIsEditorOnly = false;
+		bIsOwnerEditorOnly = false;
 		bIsInstancedStaticMesh = false;
 		bHoldout = false;
 
@@ -128,7 +128,7 @@ struct FPrimitiveSceneProxyDesc
 	uint32 bIsHidden : 1;
 	uint32 bIsHiddenEd : 1;
 	uint32 bSupportsWorldPositionOffsetVelocity : 1;
-	uint32 bIsEditorOnly : 1;
+	uint32 bIsOwnerEditorOnly : 1;
 	uint32 bIsInstancedStaticMesh : 1;
 	uint32 bHasStaticLighting : 1;
 	uint32 bHasValidSettingsForStaticLighting : 1;
@@ -210,7 +210,7 @@ struct FPrimitiveSceneProxyDesc
 	FString GetPathName() const { return Component->GetPathName(); }
 	
 	bool IsHidden() const { return bIsHidden; }
-	bool IsEditorOnly() const { return bIsEditorOnly; }
+	bool IsOwnerEditorOnly() const { return bIsOwnerEditorOnly; }
 
 #if WITH_EDITOR
 	bool IsHiddenEd() const { return bIsHiddenEd; }
