Separated out change: Using SoundWave CompressedData GUID instead of ObjectKeys for StreamCache keys.

diff --git a/Engine/Source/Runtime/Engine/Classes/Sound/SoundWave.h b/Engine/Source/Runtime/Engine/Classes/Sound/SoundWave.h
--- a/Engine/Source/Runtime/Engine/Classes/Sound/SoundWave.h
+++ b/Engine/Source/Runtime/Engine/Classes/Sound/SoundWave.h
@@ -1294,6 +1294,7 @@ public:
 
 	ENGINE_API void OverrideRuntimeFormat(const FName& InRuntimeFormat);
 
+	const FGuid& GetGUID() const { return WaveGuid; }
 	const FName& GetFName() const { return NameCached; }
 	const FName& GetPackageName() const { return PackageNameCached; }
 	const FName& GetRuntimeFormat() const { return RuntimeFormat; }
@@ -1398,6 +1399,7 @@ private:
 	TArray<FSoundWaveCuePoint> CuePoints;
 	TArray<FSoundWaveCuePoint> LoopRegions;
 	ESoundAssetCompressionType SoundAssetCompressionType;
+	FGuid WaveGuid;
 	
 	float SampleRate = 0;
 	float Duration = 0;
diff --git a/Engine/Source/Runtime/Engine/Private/AudioStreamingCache.cpp b/Engine/Source/Runtime/Engine/Private/AudioStreamingCache.cpp
--- a/Engine/Source/Runtime/Engine/Private/AudioStreamingCache.cpp
+++ b/Engine/Source/Runtime/Engine/Private/AudioStreamingCache.cpp
@@ -231,11 +231,11 @@ FAudioChunkCache::FChunkKey::FChunkKey(const FSoundWavePtr& InSoundWave, uint32
 	if (InSoundWave.IsValid())
 	{
 		SoundWaveName = InSoundWave->GetFName();
-		ObjectKey = InSoundWave->GetFObjectKey();
+		ObjectKey = InSoundWave->GetGUID();
 	}
 }
 
-FAudioChunkCache::FChunkKey::FChunkKey(const FName& InSoundWaveName , const FObjectKey& InSoundWaveObjectKey , uint32 InChunkIndex
+FAudioChunkCache::FChunkKey::FChunkKey(const FName& InSoundWaveName , const FGuid& InSoundWaveObjectKey , uint32 InChunkIndex
 #if WITH_EDITOR
 		, uint32 InChunkRevision
 #endif // #if WITH_EDITOR
@@ -636,7 +636,7 @@ void FCachedAudioStreamingManager::AddReferenceToChunk(const FAudioChunkHandle&
 
 	const FAudioChunkCache::FChunkKey ChunkKey(
 		  InHandle.CorrespondingWaveName
-		, InHandle.CorrespondingWaveObjectKey
+		, InHandle.CorrespondingWaveGuid
 		, ((uint32)InHandle.ChunkIndex)
 #if WITH_EDITOR
 		, InHandle.ChunkRevision
@@ -654,7 +654,7 @@ void FCachedAudioStreamingManager::RemoveReferenceToChunk(const FAudioChunkHandl
 
 	const FAudioChunkCache::FChunkKey ChunkKey(
 		  InHandle.CorrespondingWaveName
-		, InHandle.CorrespondingWaveObjectKey
+		, InHandle.CorrespondingWaveGuid
 		, ((uint32)InHandle.ChunkIndex)
 #if WITH_EDITOR
 		, InHandle.ChunkRevision
@@ -1525,30 +1525,7 @@ FAudioChunkCache::FCacheElement* FAudioChunkCache::EvictLeastRecentChunk(bool bB
 	return CacheElement;
 }
 
-TArray<FObjectKey> FAudioChunkCache::GetLeastRecentlyUsedRetainedSoundWaves(int32 NumSoundWavesToRetrieve)
-{
-	// Start at the least recent element, then crawl our way up the LRU cache, 
-	// adding object keys for elements as we go.
-	FCacheElement* CacheElement = LeastRecentElement;
-
-	TArray<FObjectKey> SoundWavesToRelease;
-
-	// In order to avoid cycles, we always leave at least two chunks in the cache.
-	const FCacheElement* ElementToStopAt = MostRecentElement;
-
-	while (CacheElement && CacheElement != ElementToStopAt && SoundWavesToRelease.Num() < NumSoundWavesToRetrieve)
-	{
-		// If the least recent chunk is evictable, evict it.
-		if (CacheElement->IsInUse())
-		{
-			SoundWavesToRelease.Add(CacheElement->Key.ObjectKey);
-		}
 
-		CacheElement = CacheElement->MoreRecentElement;
-	}
-
-	return SoundWavesToRelease;
-}
 
 static FAutoConsoleTaskPriority CPrio_ClearAudioChunkCacheReadRequest(
 	TEXT("TaskGraph.TaskPriorities.ClearAudioChunkCacheReadRequest"),
diff --git a/Engine/Source/Runtime/Engine/Private/ContentStreaming.cpp b/Engine/Source/Runtime/Engine/Private/ContentStreaming.cpp
--- a/Engine/Source/Runtime/Engine/Private/ContentStreaming.cpp
+++ b/Engine/Source/Runtime/Engine/Private/ContentStreaming.cpp
@@ -1466,7 +1466,7 @@ FAudioChunkHandle::FAudioChunkHandle(const uint8* InData, uint32 NumBytes, const
 		TSharedPtr<FSoundWaveData> SoundWaveData = InSoundWave->GetSoundWaveData();
 		if (SoundWaveData.IsValid())
 		{
-			CorrespondingWaveObjectKey = SoundWaveData->GetFObjectKey();
+			CorrespondingWaveGuid = SoundWaveData->GetGUID();
 		}
 	}
 }
@@ -1494,7 +1494,7 @@ FAudioChunkHandle& FAudioChunkHandle::operator=(FAudioChunkHandle&& Other)
 	CachedData = Other.CachedData;
 	CachedDataNumBytes = Other.CachedDataNumBytes;
 	CorrespondingWaveName = Other.CorrespondingWaveName;
-	CorrespondingWaveObjectKey = Other.CorrespondingWaveObjectKey;
+	CorrespondingWaveGuid = Other.CorrespondingWaveGuid;
 	ChunkIndex = Other.ChunkIndex;
 #if WITH_EDITOR
 	CorrespondingWave = MoveTemp(Other.CorrespondingWave);
@@ -1506,7 +1506,7 @@ FAudioChunkHandle& FAudioChunkHandle::operator=(FAudioChunkHandle&& Other)
 	Other.CachedData = nullptr;
 	Other.CachedDataNumBytes = 0;
 	Other.CorrespondingWaveName = FName();
-	Other.CorrespondingWaveObjectKey = FObjectKey();
+	Other.CorrespondingWaveGuid = FGuid();
 	Other.ChunkIndex = INDEX_NONE;
 #if WITH_EDITOR
 	Other.CorrespondingWave = nullptr;
@@ -1527,7 +1527,7 @@ FAudioChunkHandle& FAudioChunkHandle::operator=(const FAudioChunkHandle& Other)
 	CachedData = Other.CachedData;
 	CachedDataNumBytes = Other.CachedDataNumBytes;
 	CorrespondingWaveName = Other.CorrespondingWaveName;
-	CorrespondingWaveObjectKey = Other.CorrespondingWaveObjectKey;
+	CorrespondingWaveGuid = Other.CorrespondingWaveGuid;
 	ChunkIndex = Other.ChunkIndex;
 #if WITH_EDITOR
 	CorrespondingWave = Other.CorrespondingWave;
diff --git a/Engine/Source/Runtime/Engine/Private/SoundWave.cpp b/Engine/Source/Runtime/Engine/Private/SoundWave.cpp
--- a/Engine/Source/Runtime/Engine/Private/SoundWave.cpp
+++ b/Engine/Source/Runtime/Engine/Private/SoundWave.cpp
@@ -160,6 +160,7 @@ void FSoundWaveData::InitializeDataFromSoundWave(USoundWave& InWave)
 	NumChannels = InWave.NumChannels;
 
 	NumFrames = (int32)(Duration * (float)SampleRate);
+	WaveGuid = InWave.CompressedDataGuid;
 
 	// update shared flags
 	bIsLooping = InWave.IsLooping();
diff --git a/Engine/Source/Runtime/Engine/Public/AudioStreamingCache.h b/Engine/Source/Runtime/Engine/Public/AudioStreamingCache.h
--- a/Engine/Source/Runtime/Engine/Public/AudioStreamingCache.h
+++ b/Engine/Source/Runtime/Engine/Public/AudioStreamingCache.h
@@ -38,7 +38,7 @@ public:
 
 		FChunkKey(
 			  const FName& InSoundWaveName
-			, const FObjectKey& InSoundWaveObjectKey
+			, const FGuid& InSoundWaveGuid
 			, uint32 InChunkIndex
 #if WITH_EDITOR
 			, uint32 InChunkRevision = 0
@@ -58,7 +58,7 @@ public:
 
 	public:
 		FName SoundWaveName = FName();
-		FObjectKey ObjectKey = FObjectKey();
+		FGuid ObjectKey = FGuid();
 		uint32 ChunkIndex = INDEX_NONE;
 
 #if WITH_EDITOR
@@ -231,6 +231,8 @@ private:
 
 		// if true, 
 		bool bWasCacheMiss;
+		bool bWasLoadedFromInlineChunk = false;
+		bool bWasInlinedButUnloaded = false;
 
 		FCacheElementDebugInfo()
 			: NumTotalChunks(0)
@@ -252,6 +254,8 @@ private:
 			LoadingBehavior = ESoundWaveLoadingBehavior::Uninitialized;
 			bWasCacheMiss = false;
 			AverageLocationInCacheWhenNeeded = 0.0f;
+			bWasLoadedFromInlineChunk = false;
+			bWasInlinedButUnloaded = false;
 		}
 
 	};
diff --git a/Engine/Source/Runtime/Engine/Public/ContentStreaming.h b/Engine/Source/Runtime/Engine/Public/ContentStreaming.h
--- a/Engine/Source/Runtime/Engine/Public/ContentStreaming.h
+++ b/Engine/Source/Runtime/Engine/Public/ContentStreaming.h
@@ -139,7 +139,7 @@ private:
 	int32 CachedDataNumBytes;
 
 	FName CorrespondingWaveName;
-	FObjectKey CorrespondingWaveObjectKey;
+	FGuid CorrespondingWaveGuid;
 
 	// The index of this chunk in the sound wave's full set of chunks of compressed audio.
 	int32 ChunkIndex;
