Bink Latency Fixes (Owner Loading behavior tagging,Inline first chunk of audio optionally with Retained audio, Changed MONO_PCM size to 480)

diff --git a/Engine/Source/Runtime/AudioCodecEngine/Private/DecoderInputBackCompat.cpp b/Engine/Source/Runtime/AudioCodecEngine/Private/DecoderInputBackCompat.cpp
--- a/Engine/Source/Runtime/AudioCodecEngine/Private/DecoderInputBackCompat.cpp
+++ b/Engine/Source/Runtime/AudioCodecEngine/Private/DecoderInputBackCompat.cpp
@@ -20,7 +20,7 @@ namespace Audio
 			TUniquePtr<ICompressedAudioInfo> InfoInstance;
 			FName Format = Wave->GetRuntimeFormat();
 			IAudioInfoFactory* Factory = IAudioInfoFactoryRegistry::Get().Find(Format);
-			if (!ensure(Factory))
+			if (!ensureMsgf(Factory, TEXT("Format '%s' not found for wave '%s'"),*Format.ToString(), *Wave->GetPackageName().ToString()))
 			{
 				return nullptr;
 			}
diff --git a/Engine/Source/Runtime/AudioPlatformConfiguration/Private/AudioCompressionSettings.cpp b/Engine/Source/Runtime/AudioPlatformConfiguration/Private/AudioCompressionSettings.cpp
--- a/Engine/Source/Runtime/AudioPlatformConfiguration/Private/AudioCompressionSettings.cpp
+++ b/Engine/Source/Runtime/AudioPlatformConfiguration/Private/AudioCompressionSettings.cpp
@@ -56,8 +56,9 @@ void FPlatformAudioCookOverrides::GetHashSuffix(const FPlatformAudioCookOverride
 	FPCU::AppendHash(OutSuffix, TEXT("QMOD"), InOverrides->CompressionQualityModifier);
 	FPCU::AppendHash(OutSuffix, TEXT("CQLT"), InOverrides->SoundCueCookQualityIndex);
 	FPCU::AppendHash(OutSuffix, TEXT("ASTH"), InOverrides->AutoStreamingThreshold);
-	FPCU::AppendHash(OutSuffix, TEXT("INLC"), InOverrides->bInlineStreamedAudioChunks);
-
+	FPCU::AppendHash(OutSuffix, TEXT("INLC"), InOverrides->bInlineFirstAudioChunk);
+	FPCU::AppendHash(OutSuffix, TEXT("LCK1"), InOverrides->LengthOfFirstAudioChunkInSecs);
+	
 	// FAudioStreamCachingSettings
 	FPCU::AppendHash(OutSuffix, TEXT("CSZE"), InOverrides->StreamCachingSettings.CacheSizeKB);
 	FPCU::AppendHash(OutSuffix, TEXT("LCF"), InOverrides->StreamCachingSettings.bForceLegacyStreamChunking);
diff --git a/Engine/Source/Runtime/AudioPlatformConfiguration/Public/AudioCompressionSettings.h b/Engine/Source/Runtime/AudioPlatformConfiguration/Public/AudioCompressionSettings.h
--- a/Engine/Source/Runtime/AudioPlatformConfiguration/Public/AudioCompressionSettings.h
+++ b/Engine/Source/Runtime/AudioPlatformConfiguration/Public/AudioCompressionSettings.h
@@ -80,8 +80,13 @@ struct FPlatformAudioCookOverrides
 	// If StreamCaching is set to true, this will be used 
 	float AutoStreamingThreshold;
 
-	// Whether to put streamed audio chunks inline in the Pak file or not (only matters if bUseStreamCaching is true)
-	bool bInlineStreamedAudioChunks;
+	// Wether to inline the first "Audio" chunk, which is typically chunk 1. (Only on assets marked retain-on-load with a size of audio in secs set) 
+	bool bInlineFirstAudioChunk = false;
+
+	// This will decide how much data to put in the first audio chunk. Anything <= 0 will be ignored.
+	// Must be combined with bInlineFirstAudioChunk, this will decide how much data to put in the first chunk.
+	// NOTE: This is platform default and can be overriden by each asset or soundclass.
+	float LengthOfFirstAudioChunkInSecs = 0.f;
 
 	// If Load On Demand is enabled, these settings are used to determine chunks and cache sizes.
 	FAudioStreamCachingSettings StreamCachingSettings;
@@ -90,7 +95,8 @@ struct FPlatformAudioCookOverrides
 		: bResampleForDevice(false)
 		, CompressionQualityModifier(1.0f)
 		, AutoStreamingThreshold(0.0f)
-		, bInlineStreamedAudioChunks(false)
+		, bInlineFirstAudioChunk(false)
+		, LengthOfFirstAudioChunkInSecs(0.f)
 	{
 		PlatformSampleRates.Add(ESoundwaveSampleRateSettings::Max, 48000);
 		PlatformSampleRates.Add(ESoundwaveSampleRateSettings::High, 32000);
diff --git a/Engine/Source/Runtime/Engine/Classes/Sound/SoundClass.h b/Engine/Source/Runtime/Engine/Classes/Sound/SoundClass.h
--- a/Engine/Source/Runtime/Engine/Classes/Sound/SoundClass.h
+++ b/Engine/Source/Runtime/Engine/Classes/Sound/SoundClass.h
@@ -16,6 +16,7 @@
 #include "EdGraph/EdGraph.h"
 #endif // WITH_EDITOR
 #include "SoundWaveLoadingBehavior.h"
+#include "PerPlatformProperties.h"
 
 #include "SoundClass.generated.h"
 
@@ -130,9 +131,15 @@ struct FSoundClassProperties
 	TEnumAsByte<EAudioOutputTarget::Type> OutputTarget;
 
 	/** Specifies how and when compressed audio data is loaded for asset if stream caching is enabled. */
-	UPROPERTY(EditAnywhere, BlueprintReadOnly, Category = Loading, meta = (DisplayName = "Loading Behavior Override"))
+	UPROPERTY(EditAnywhere, BlueprintReadOnly, Category = Loading, meta = (DisplayName = "Loading Behavior Override"), AssetRegistrySearchable)
 	ESoundWaveLoadingBehavior LoadingBehavior;
 
+#if WITH_EDITORONLY_DATA
+   	/** How much audio to add to First Audio Chunk (in seconds), (only for RetainOnLoad) */
+   	UPROPERTY(EditAnywhere, Category = Loading, meta = (UIMin = 0, UIMax = 10, EditCondition = "LoadingBehavior == ESoundWaveLoadingBehavior::RetainOnLoad"), DisplayName="Size of First Audio Chunk (seconds)")
+   	FPerPlatformFloat SizeOfFirstAudioChunkInSeconds = 0.0f;
+#endif //WITH_EDITORONLY_DATA
+
 	/** Default output submix of referencing sounds. If unset, falls back to the 'Master Submix' as set in the 'Audio' category of Project Settings. 
 	  * (Unavailable if legacy 'Output to Master EQ Submix' is set) */
 	UPROPERTY(EditAnywhere, BlueprintReadOnly, Category = Submix, meta = (EditCondition = "!bApplyEffects"))
diff --git a/Engine/Source/Runtime/Engine/Classes/Sound/SoundWave.h b/Engine/Source/Runtime/Engine/Classes/Sound/SoundWave.h
--- a/Engine/Source/Runtime/Engine/Classes/Sound/SoundWave.h
+++ b/Engine/Source/Runtime/Engine/Classes/Sound/SoundWave.h
@@ -73,6 +73,9 @@ struct FStreamedAudioChunk
 	/**  returns false if data retrieval failed */
 	bool GetCopy(void** OutChunkData);
 
+	/* Moves the memory out of the byte bulk-data. (Does copy with Discard original) */
+	FBulkDataBuffer<uint8> MoveOutAsBuffer();
+
 	/** Size of the chunk of data in bytes including zero padding */
 	int32 DataSize = 0;
 
@@ -97,6 +100,9 @@ public:
 	/** True if this chunk was loaded from a cooked package. */
 	bool bLoadedFromCookedPackage = false;
 
+	/** If marked true, will attempt to inline this chunk. */
+	bool bInlineChunk = false;
+
 	/**
 	 * Place chunk data in the derived data cache associated with the provided
 	 * key.
@@ -155,7 +161,7 @@ struct FStreamedAudioPlatformData
 	ENGINE_API void Serialize(FArchive& Ar, class USoundWave* Owner);
 
 #if WITH_EDITORONLY_DATA
-	ENGINE_API void Cache(class USoundWave& InSoundWave, const FPlatformAudioCookOverrides* CompressionOverrides, FName AudioFormatName, uint32 InFlags);
+	ENGINE_API void Cache(class USoundWave& InSoundWave, const FPlatformAudioCookOverrides* CompressionOverrides, FName AudioFormatName, uint32 InFlags, const ITargetPlatform* InTargetPlatform=nullptr);
 	ENGINE_API void FinishCache();
 	ENGINE_API bool IsFinishedCache() const;
 	ENGINE_API bool IsAsyncWorkComplete() const;
@@ -182,7 +188,6 @@ private:
 	 * @param SerializedData Serialized data resulting from DDC.GetAsynchronousResults or DDC.GetSynchronous.
 	 * @param ChunkToDeserializeInto is the chunk to fill with the deserialized data.
 	 * @param ChunkIndex is the index of the chunk in this instance of FStreamedAudioPlatformData.
-	 * @param bCachedChunk is true if the chunk was successfully cached, false otherwise.
 	 * @param OutChunkData is a pointer to a pointer to populate with the chunk itself, or if pointing to nullptr, returns an allocated buffer.
 	 * @returns the size of the chunk loaded in bytes, or zero if the chunk didn't load.
 	 */
@@ -598,9 +603,9 @@ public:
 	 */
 	ENGINE_API ESoundWaveLoadingBehavior GetLoadingBehavior(bool bCheckSoundClasses = true) const;
 
-	/** Use this to override how much audio data is loaded when this USoundWave is loaded. */
-	UPROPERTY(EditAnywhere, AdvancedDisplay, Category = "Loading")
-	int32 InitialChunkSize;
+	/** Please use size of First Chunk in Seconds. */
+	UPROPERTY(AdvancedDisplay, meta=(DeprecatedProperty))
+	int32 InitialChunkSize_DEPRECATED;
 
 #if WITH_EDITOR
 	ENGINE_API const FWaveTransformUObjectConfiguration& GetTransformationChainConfig() const;
@@ -714,6 +719,13 @@ public:
 	UPROPERTY(EditAnywhere, Category = "Loading", meta = (DisplayName = "Loading Behavior Override"))
 	mutable ESoundWaveLoadingBehavior LoadingBehavior;
 
+#if WITH_EDITORONLY_DATA
+public:
+   	/** How much audio to add to First Audio Chunk (in seconds), (only for RetainOnLoad) */
+   	UPROPERTY(EditAnywhere, Category = Loading, meta = (UIMin = 0, UIMax = 10, EditCondition = "LoadingBehavior == ESoundWaveLoadingBehavior::RetainOnLoad"), DisplayName="Size of First Audio Chunk (seconds)")
+   	FPerPlatformFloat SizeOfFirstAudioChunkInSeconds = 0.0f;
+#endif //WITH_EDITOR_ONLY_DATA
+
 	/** A localized version of the text that is actually spoken phonetically in the audio. */
 	UPROPERTY(meta = (DeprecatedProperty, DeprecationMessage = "Use Subtitles instead."))
 	FString SpokenText_DEPRECATED;
@@ -823,6 +835,8 @@ public:
 	*/
 	static ENGINE_API ITargetPlatform* GetRunningPlatform();
 
+	static ENGINE_API ESoundWaveLoadingBehavior GetDefaultLoadingBehavior();
+
 	/** Async worker that decompresses the audio data on a different thread */
 	typedef FAsyncTask< class FAsyncAudioDecompressWorker > FAsyncAudioDecompress;	// Forward declare typedef
 	FAsyncAudioDecompress* AudioDecompressor;
@@ -895,6 +909,8 @@ public:
 	ENGINE_API bool IsAsyncWorkComplete() const;
 
 	ENGINE_API void PostImport();
+	
+	ENGINE_API virtual void PreSave(FObjectPreSaveContext SaveContext);
 
 private:
 	friend class FSoundWaveCompilingManager;
@@ -1090,18 +1106,18 @@ public:
 	ENGINE_API bool IsLoadedFromCookedData() const;
 #endif //WITH_EDITOR
 
-	ENGINE_API virtual void BeginGetCompressedData(FName Format, const FPlatformAudioCookOverrides* CompressionOverrides);
+	ENGINE_API virtual void BeginGetCompressedData(FName Format, const FPlatformAudioCookOverrides* CompressionOverrides, const ITargetPlatform* InTargetPlatform);
 
 	/**
 	 * Gets the compressed data from derived data cache for the specified platform
 	 * Warning, the returned pointer isn't valid after we add new formats
 	 *
 	 * @param Format	format of compressed data
-	 * @param PlatformName optional name of platform we are getting compressed data for.
 	 * @param CompressionOverrides optional platform compression overrides
 	 * @return	compressed data, if it could be obtained
 	 */
-	ENGINE_API virtual FByteBulkData* GetCompressedData(FName Format, const FPlatformAudioCookOverrides* CompressionOverrides = GetPlatformCompressionOverridesForCurrentPlatform());
+	ENGINE_API virtual FByteBulkData* GetCompressedData(FName Format, const FPlatformAudioCookOverrides* CompressionOverrides = GetPlatformCompressionOverridesForCurrentPlatform(),
+		const ITargetPlatform* InTargetPlatform = GetRunningPlatform());
 
 	/**
 	 * Change the guid and flush all compressed data
@@ -1187,10 +1203,23 @@ public:
 #if WITH_EDITORONLY_DATA
 
 #if WITH_EDITOR
+	
+	/*
+	* Used to determine the loading behavior that's applied by the owner of this wave. Will determine
+	* the most appropriate loading behavior if there are multiple owners. If there is no owner "Unintialized"
+	* will be returned. The results are cached in TMap below.
+	*/
+	ISoundWaveLoadingBehaviorUtil::FClassData GetOwnerLoadingBehavior(const ITargetPlatform* InTargetPlatform) const;
+	
+	mutable TMap<FName, ISoundWaveLoadingBehaviorUtil::FClassData> OwnerLoadingBehaviorCache;
+	mutable FCriticalSection OwnerLoadingBehaviorCacheCS;
+			
 	/*
 	* Returns a sample rate if there is a specific sample rate override for this platform, -1.0 otherwise.
 	*/
 	ENGINE_API float GetSampleRateForTargetPlatform(const ITargetPlatform* TargetPlatform);
+	
+	ENGINE_API float GetSizeOfFirstAudioChunkInSeconds(const ITargetPlatform* InTargetPlatform) const;
 
 	/**
 	 * Begins caching platform data in the background for the platform requested
@@ -1354,6 +1383,7 @@ public:
 
 #if WITH_EDITORONLY_DATA
  	ENGINE_API FString GetDerivedDataKey() const;
+	ENGINE_API FPerPlatformFloat GetSizeOfFirstAudioChunkInSeconds() const { return SizeOfFirstAudioChunkInSeconds; }
 #endif // #if WITH_EDITORONLY_DATA
 
 	int32 GetResourceSize() { return ResourceSize; }
@@ -1388,6 +1418,11 @@ private:
 
 	ESoundWaveLoadingBehavior LoadingBehavior = ESoundWaveLoadingBehavior::Uninitialized;
 
+#if WITH_EDITORONLY_DATA
+	// Set by CacheInheritedLoadingBehavior after traversing the Soundclass heirarchy 
+	FPerPlatformFloat SizeOfFirstAudioChunkInSeconds = 0.f;
+#endif //WITH_EDITORONLY_DATA
+	
 #if WITH_EDITOR
 	std::atomic<int32> CurrentChunkRevision;
 #endif // #if WITH_EDITOR
@@ -1398,7 +1433,7 @@ private:
 	FObjectKey SoundWaveKeyCached;
 	TArray<FSoundWaveCuePoint> CuePoints;
 	TArray<FSoundWaveCuePoint> LoopRegions;
-	ESoundAssetCompressionType SoundAssetCompressionType;
+	ESoundAssetCompressionType SoundAssetCompressionType = ESoundAssetCompressionType::BinkAudio;
 	FGuid WaveGuid;
 	
 	float SampleRate = 0;
diff --git a/Engine/Source/Runtime/Engine/Classes/Sound/SoundWaveLoadingBehavior.cpp b/Engine/Source/Runtime/Engine/Classes/Sound/SoundWaveLoadingBehavior.cpp
--- a/Engine/Source/Runtime/Engine/Classes/Sound/SoundWaveLoadingBehavior.cpp
+++ b/Engine/Source/Runtime/Engine/Classes/Sound/SoundWaveLoadingBehavior.cpp
@@ -1,6 +1,13 @@
 // Copyright Epic Games, Inc. All Rights Reserved.
 #include "SoundWaveLoadingBehavior.h"
 
+#include "AssetRegistry/AssetRegistryModule.h"
+#include "AssetRegistry/IAssetRegistry.h"
+#include "GenericPlatform/GenericPlatformMisc.h"
+#include "Interfaces/ITargetPlatform.h"
+#include "SoundClass.h"
+#include "SoundCue.h"
+#include "SoundWave.h"
 
 const TCHAR* EnumToString(ESoundWaveLoadingBehavior InCurrentState)
 {
@@ -22,3 +29,259 @@ const TCHAR* EnumToString(ESoundWaveLoadingBehavior InCurrentState)
 	ensure(false);
 	return TEXT("Unknown");
 }
+
+
+#if WITH_EDITOR
+
+static int32 SoundWaveLoadingBehaviorUtil_CacheAllOnStartup = 0;
+FAutoConsoleVariableRef CVAR_SoundWaveLoadingBehaviorUtil_CacheAllOnStartup(
+	TEXT("au.editor.SoundWaveOwnerLoadingBehaviorCacheOnStartup"),
+	SoundWaveLoadingBehaviorUtil_CacheAllOnStartup,
+	TEXT("Disables searching the asset registry on startup of the singleton. Otherwise it will incrementally fill cache"),
+	ECVF_Default
+);
+
+static int32 SoundWaveLoadingBehaviorUtil_Enable = 1;
+FAutoConsoleVariableRef CVAR_SoundWaveLoadingBehaviorUtil_Enable(
+	TEXT("au.editor.SoundWaveOwnerLoadingBehaviorEnable"),
+	SoundWaveLoadingBehaviorUtil_Enable,
+	TEXT("Enables or disables the Soundwave owner loading behavior tagging"),
+	ECVF_Default
+);
+
+class FSoundWaveLoadingBehaviorUtil : public ISoundWaveLoadingBehaviorUtil
+{
+public:
+	FSoundWaveLoadingBehaviorUtil()	
+		: AssetRegistry(FModuleManager::LoadModuleChecked<FAssetRegistryModule>(TEXT("AssetRegistry")).Get())
+	{
+		if (SoundWaveLoadingBehaviorUtil_CacheAllOnStartup)
+		{
+			CacheAllClassLoadingBehaviors();
+		}
+	}
+	virtual ~FSoundWaveLoadingBehaviorUtil() override = default;
+
+private:
+	void CacheAllClassLoadingBehaviors()
+	{
+		TRACE_CPUPROFILER_EVENT_SCOPE(LoadAndCacheAllSoundClassLoadingBehaviors);
+
+		ensureMsgf(!AssetRegistry.IsSearchAsync() || !AssetRegistry.IsLoadingAssets(), 
+		           TEXT("Function must not be called until after cook has started and waited on the AssetRegistry already."));
+		
+		AssetRegistry.SearchAllAssets(true);
+
+		TArray<FAssetData> SoundClasses;
+		AssetRegistry.GetAssetsByClass(USoundClass::StaticClass()->GetClassPathName(), SoundClasses, true);
+		
+		for (const FAssetData& i : SoundClasses)
+		{
+			LoadAndCacheClass(i);
+		}
+	}
+
+	FClassData WalkClassHierarchy(USoundClass* InClass) const
+	{
+		FClassData Behavior(InClass);
+		while (InClass->ParentClass && InClass->Properties.LoadingBehavior == ESoundWaveLoadingBehavior::Inherited)
+		{
+			InClass = InClass->ParentClass;
+			if (InClass->HasAnyFlags(RF_NeedLoad))
+			{
+				InClass->GetLinker()->Preload(InClass);
+			}
+			if (InClass->HasAnyFlags(RF_NeedPostLoad))
+			{
+				InClass->ConditionalPostLoad();
+			}
+			Behavior = FClassData(InClass);
+		}
+		
+		// If we failed to find anything other than inherited, use the default which is cvar'd.
+		if (Behavior.LoadingBehavior == ESoundWaveLoadingBehavior::Inherited ||
+			Behavior.LoadingBehavior == ESoundWaveLoadingBehavior::Uninitialized )
+		{
+			Behavior.LoadingBehavior = USoundWave::GetDefaultLoadingBehavior();
+			Behavior.LengthOfFirstChunkInSeconds = 0;
+		}
+		return Behavior;
+	}
+
+	FClassData LoadAndCacheClass(const FAssetData& InAssetData) const
+	{
+		if (USoundClass* SoundClass = Cast<USoundClass>(InAssetData.GetAsset()))
+		{
+			FClassData Result = WalkClassHierarchy(SoundClass);
+			CacheClassLoadingBehaviors.Add(InAssetData.PackageName,Result);
+			return Result;
+		}
+		return {};
+	}
+	
+	virtual FClassData FindOwningLoadingBehavior(const USoundWave* InWave, const ITargetPlatform* InTargetPlatform) const 
+	{
+		TRACE_CPUPROFILER_EVENT_SCOPE(FindOwningLoadingBehavior);
+
+		// This code: Given a wave, finds all cues that references it. (reverse lookup)
+		// Then finds the SoundClasses those cues use, traverses the heirarchy (from lookup) to determine the loading behavior.
+		// Then stack ranks the most important behavior. (RetainOnLoad (Highest), PrimeOnLoad (Medium), LoadOnDemand (Lowest))
+		// Which ever wins, we also capture the "SizeOfFirstChunk" to use for that wave.
+
+		if (!InWave)
+		{
+			return {};
+		}
+		
+		const UPackage* WavePackage = InWave->GetPackage();
+		if (!WavePackage)
+		{
+			return {};
+		}
+
+		TArray<FName> SoundWaveReferencerNames;
+		if (!AssetRegistry.GetReferencers(WavePackage->GetFName(), SoundWaveReferencerNames))
+		{
+			return {};
+		}
+
+		if (SoundWaveReferencerNames.IsEmpty())
+		{
+			return {};
+		}
+
+		// Filter on SoundCues.
+		FARFilter Filter;
+		Filter.ClassPaths.Add(USoundCue::StaticClass()->GetClassPathName());
+		Filter.bRecursiveClasses = true;
+		Filter.PackageNames = SoundWaveReferencerNames;
+		TArray<FAssetData> ReferencingSoundCueAssetDataArray;
+		if (!AssetRegistry.GetAssets(Filter, ReferencingSoundCueAssetDataArray))
+		{
+			return {};
+		}
+
+		if (ReferencingSoundCueAssetDataArray.IsEmpty())
+		{
+			return {};
+		}
+
+		FClassData MostImportantLoadingBehavior;
+		for (const FAssetData& CueAsset : ReferencingSoundCueAssetDataArray)
+		{
+			// Query for class references from the Cue instead of loading and opening it.
+			TArray<FName> SoundCueReferences;
+			if (!AssetRegistry.GetDependencies(CueAsset.PackageName, SoundCueReferences))
+			{
+				UE_LOG(LogAudio, Warning, TEXT("Failed to query SoundCue '%s' for it's dependencies."), 
+				       *CueAsset.PackagePath.ToString() );
+				continue;
+			}
+
+			if (SoundCueReferences.Num() == 0)
+			{
+				continue;
+			}
+
+			// Filter for Classes.
+			FARFilter ClassFilter;
+			ClassFilter.ClassPaths.Add(USoundClass::StaticClass()->GetClassPathName());
+			ClassFilter.PackageNames = SoundCueReferences;
+			TArray<FAssetData> ReferencedSoundClasses;
+			if (!AssetRegistry.GetAssets(ClassFilter, ReferencedSoundClasses))
+			{
+				UE_LOG(LogAudio, Warning, TEXT("Failed to filter for Soundclasses from the SoundCue dependencies for '%s'"), *CueAsset.PackagePath.ToString());
+				continue;
+			}			
+
+			// Should have a single SoundClass, we hope, otherwise ignore.
+			if (ReferencedSoundClasses.Num() != 1)
+			{
+				UE_CLOG(ReferencedSoundClasses.Num() > 1, LogAudio, Warning, 
+				        TEXT("More than one Soundclass refereneced from this cue '%s'"), *CueAsset.PackagePath.ToString() );
+				continue;
+			}
+			
+			// Look up this classes loading behavior in our cache.
+			const FAssetData& Class = ReferencedSoundClasses[0];
+			FClassData CacheLoadingBehavior;
+			{
+				FScopeLock Lock(&CacheCS);
+				if (const FClassData* Found = CacheClassLoadingBehaviors.Find(Class.PackageName))
+				{
+					CacheLoadingBehavior = *Found;
+				}
+				else
+				{
+					CacheLoadingBehavior = LoadAndCacheClass(Class);
+				}
+			}
+
+			// Compare if this is more important
+			if (MostImportantLoadingBehavior.CompareGreater(CacheLoadingBehavior, InTargetPlatform))
+			{
+				MostImportantLoadingBehavior = CacheLoadingBehavior;
+			}
+		}
+		
+		// Return the most important one we found.
+		return MostImportantLoadingBehavior;
+	}
+	
+	// Make cache here.
+	IAssetRegistry& AssetRegistry;
+	mutable TMap<FName, FClassData> CacheClassLoadingBehaviors;
+	mutable FCriticalSection CacheCS;
+};
+
+ISoundWaveLoadingBehaviorUtil* ISoundWaveLoadingBehaviorUtil::Get()
+{
+	// Cvar disable system if necessary
+	if (!SoundWaveLoadingBehaviorUtil_Enable)
+	{
+		return nullptr;
+	}
+	
+	static bool bAllowOutsideOfCookCommandlet = FParse::Param(FCommandLine::Get(), TEXT("AllowSoundWaveOwnerLoadingBehaviorInEditor"));
+	
+	// Only run while the cooker is active.
+	if (!IsRunningCookCommandlet() && !bAllowOutsideOfCookCommandlet)
+	{
+		return nullptr;
+	}
+	
+	static FSoundWaveLoadingBehaviorUtil Instance;
+	return &Instance;
+}
+
+ISoundWaveLoadingBehaviorUtil::FClassData::FClassData(const USoundClass* InClass)
+	: LoadingBehavior(InClass->Properties.LoadingBehavior)
+	, LengthOfFirstChunkInSeconds(InClass->Properties.SizeOfFirstAudioChunkInSeconds)
+{
+}
+
+bool ISoundWaveLoadingBehaviorUtil::FClassData::CompareGreater(const FClassData& InOther, const ITargetPlatform* InPlatform) const
+{
+	if (InOther.LoadingBehavior != LoadingBehavior)
+	{
+		// If loading behavior enum is less, it's more important.
+		return InOther.LoadingBehavior < LoadingBehavior;
+	}
+	else 
+	{
+		// If we are using Retain, use one with the higher Length.
+		if (LoadingBehavior == ESoundWaveLoadingBehavior::RetainOnLoad)
+		{
+			const float Length = LengthOfFirstChunkInSeconds.GetValueForPlatform(*InPlatform->PlatformName());
+			const float OtherLength = InOther.LengthOfFirstChunkInSeconds.GetValueForPlatform(*InPlatform->PlatformName());
+				
+			if (OtherLength > Length)
+			{
+				return true;
+			}
+		}
+		return false;
+	}
+}
+
+#endif //WITH_EDITOR
diff --git a/Engine/Source/Runtime/Engine/Classes/Sound/SoundWaveLoadingBehavior.h b/Engine/Source/Runtime/Engine/Classes/Sound/SoundWaveLoadingBehavior.h
--- a/Engine/Source/Runtime/Engine/Classes/Sound/SoundWaveLoadingBehavior.h
+++ b/Engine/Source/Runtime/Engine/Classes/Sound/SoundWaveLoadingBehavior.h
@@ -4,6 +4,10 @@
 #include "CoreTypes.h"
 #include "UObject/ObjectMacros.h"
 
+#if WITH_EDITOR
+#include "PerPlatformProperties.h"
+#endif //WITH_EDITOR
+
 #include "SoundWaveLoadingBehavior.generated.h"
 
 /**
@@ -32,4 +36,31 @@ enum class ESoundWaveLoadingBehavior : uint8
 	Uninitialized = 0xff UMETA(Hidden)
 };
 
-const TCHAR* EnumToString(ESoundWaveLoadingBehavior InCurrentState);
+ENGINE_API const TCHAR* EnumToString(ESoundWaveLoadingBehavior InCurrentState);
+
+#if WITH_EDITOR
+
+class USoundClass;
+
+class ISoundWaveLoadingBehaviorUtil
+{
+public:
+	virtual ~ISoundWaveLoadingBehaviorUtil() = default;
+
+	struct FClassData
+	{
+		FClassData() = default;
+		FClassData(const USoundClass* InClass);
+
+		bool CompareGreater(const FClassData&, const ITargetPlatform*) const;
+
+		ESoundWaveLoadingBehavior LoadingBehavior = ESoundWaveLoadingBehavior::Uninitialized;
+		FPerPlatformFloat LengthOfFirstChunkInSeconds;
+	};
+
+	static ISoundWaveLoadingBehaviorUtil* Get();	
+
+	virtual FClassData FindOwningLoadingBehavior(const USoundWave*, const ITargetPlatform*) const = 0;
+};
+
+#endif //WITH_EDITOR
diff --git a/Engine/Source/Runtime/Engine/Classes/Sound/SoundWaveProcedural.h b/Engine/Source/Runtime/Engine/Classes/Sound/SoundWaveProcedural.h
--- a/Engine/Source/Runtime/Engine/Classes/Sound/SoundWaveProcedural.h
+++ b/Engine/Source/Runtime/Engine/Classes/Sound/SoundWaveProcedural.h
@@ -63,8 +63,8 @@ public:
 	//~ Begin USoundWave Interface.
 	ENGINE_API virtual int32 GeneratePCMData(uint8* PCMData, const int32 SamplesNeeded) override;
 	ENGINE_API virtual bool HasCompressedData(FName Format, ITargetPlatform* TargetPlatform) const override;
-	ENGINE_API virtual void BeginGetCompressedData(FName Format, const FPlatformAudioCookOverrides* CompressionOverrides) override;
-	ENGINE_API virtual FByteBulkData* GetCompressedData(FName Format, const FPlatformAudioCookOverrides* CompressionOverrides = nullptr) override;
+	ENGINE_API virtual void BeginGetCompressedData(FName Format, const FPlatformAudioCookOverrides* CompressionOverrides, const ITargetPlatform* InTargetPlatform) override;
+	ENGINE_API virtual FByteBulkData* GetCompressedData(FName Format, const FPlatformAudioCookOverrides* CompressionOverrides, const ITargetPlatform* InTargetPlatform) override;
 	ENGINE_API virtual void InitAudioResource( FByteBulkData& CompressedData ) override;
 	ENGINE_API virtual bool InitAudioResource(FName Format) override;
 	ENGINE_API virtual int32 GetResourceSizeForFormat(FName Format) override;
diff --git a/Engine/Source/Runtime/Engine/Private/AudioCompressionSettingsUtils.cpp b/Engine/Source/Runtime/Engine/Private/AudioCompressionSettingsUtils.cpp
--- a/Engine/Source/Runtime/Engine/Private/AudioCompressionSettingsUtils.cpp
+++ b/Engine/Source/Runtime/Engine/Private/AudioCompressionSettingsUtils.cpp
@@ -107,7 +107,15 @@ void CacheAudioCookOverrides(FPlatformAudioCookOverrides& OutOverrides, const TC
 		OutOverrides.SoundCueCookQualityIndex = SoundCueQualityIndex;
 	}
 
-	GConfig->GetBool(*CategoryName, TEXT("bInlineStreamedAudioChunks"), OutOverrides.bInlineStreamedAudioChunks, GEngineIni);
+	if (bool bInlineFirstAudioChunk = false; PlatformFile->GetBool(*CategoryName, TEXT("bInlineFirstAudioChunk"), bInlineFirstAudioChunk))
+	{
+		OutOverrides.bInlineFirstAudioChunk = bInlineFirstAudioChunk;
+	}
+
+	if (float LengthOfFirstAudioChunkInSeconds = 0.f; PlatformFile->GetFloat(*CategoryName, TEXT("LengthOfFirstAudioChunkInSeconds"), LengthOfFirstAudioChunkInSeconds))
+	{
+		OutOverrides.LengthOfFirstAudioChunkInSecs = LengthOfFirstAudioChunkInSeconds;
+	}
 
 	/** Memory Load On Demand Settings */
 	// Cache size:
diff --git a/Engine/Source/Runtime/Engine/Private/AudioDecompress.cpp b/Engine/Source/Runtime/Engine/Private/AudioDecompress.cpp
--- a/Engine/Source/Runtime/Engine/Private/AudioDecompress.cpp
+++ b/Engine/Source/Runtime/Engine/Private/AudioDecompress.cpp
@@ -6,6 +6,7 @@
 #include "Engine/Engine.h"
 #include "Interfaces/IAudioFormat.h"
 #include "ADPCMAudioInfo.h"
+#include "AudioStreamingCache.h"
 #include "Misc/CoreStats.h"
 #include "Misc/ScopeRWLock.h"
 #include "Stats/StatsTrace.h"
@@ -27,11 +28,11 @@ IStreamedCompressedInfo::IStreamedCompressedInfo()
 	, LastPCMOffset(0)
 	, bStoringEndOfFile(false)
 	, CurrentChunkIndex(0)
-	, bPrintChunkFailMessage(true)
 	, SrcBufferPadding(0)
 	, StreamSeekBlockIndex(INDEX_NONE)
 	, StreamSeekBlockOffset(0)
 {
+	StartTimeInCycles = FPlatformTime::Cycles64();
 }
 
 uint32 IStreamedCompressedInfo::Read(void *OutBuffer, uint32 DataSize)
@@ -184,7 +185,7 @@ bool IStreamedCompressedInfo::StreamCompressedInfoInternal(const FSoundWaveProxy
 		return false;
 	}
 
-	// Get the first chunk of audio data (should always be loaded)
+	// Get the zeroth chunk of data (should always be loaded)
 	CurrentChunkIndex = 0;
 	uint32 ChunkSize = 0;
 	const uint8* FirstChunk = GetLoadedChunk(StreamingSoundWave, CurrentChunkIndex, ChunkSize);
@@ -229,7 +230,7 @@ bool IStreamedCompressedInfo::StreamCompressedData(uint8* Destination, bool bLoo
 	SCOPED_NAMED_EVENT(IStreamedCompressedInfo_StreamCompressedData, FColor::Blue);
 
 	SCOPE_CYCLE_COUNTER(STAT_AudioStreamedDecompressTime);
-
+	
 	UE_LOG(LogAudio, Log, TEXT("Streaming compressed data from SoundWave'%s' - Chunk=%d\tNumChunks=%d\tOffset=%d\tChunkSize=%d\tLooping=%s\tLastPCMOffset=%d\tContainsEOF=%s" ), 
 		*StreamingSoundWave->GetFName().ToString(), 	
 		CurrentChunkIndex, 
@@ -249,7 +250,15 @@ bool IStreamedCompressedInfo::StreamCompressedData(uint8* Destination, bool bLoo
 	{
 		uint32 ChunkSize = 0;
 		const uint8* NewlySeekedChunk = GetLoadedChunk(StreamingSoundWave, StreamSeekBlockIndex, ChunkSize);
-		UE_LOG(LogAudio, Log, TEXT("Seek block request: %d / %d (%s)"), StreamSeekBlockIndex.load(), StreamSeekBlockOffset, SrcBufferData == nullptr ? TEXT("present") : TEXT("missing"));
+		UE_LOG(LogAudio, Log, TEXT("Seek request for (%s): (%s) Chunk=%d (%s), Offset=%d, OffsetInAudioFrames=%u"),
+			*StreamingSoundWave->GetFName().ToString(),
+			StreamSeekToAudioFrames != INDEX_NONE ? TEXT("Using streaming seek-tables") : TEXT("Using chunk/offset pair"),
+			StreamSeekBlockIndex.load(),
+			NewlySeekedChunk != nullptr ? TEXT("cache hit") : TEXT("cache miss"),
+			StreamSeekBlockOffset,
+			StreamSeekToAudioFrames
+		);
+		
 		if (NewlySeekedChunk == nullptr)
 		{
 			// After a seek we're likely to need to wait a bit for the chunk to get in to memory.
@@ -314,7 +323,6 @@ bool IStreamedCompressedInfo::StreamCompressedData(uint8* Destination, bool bLoo
 		SrcBufferData = GetLoadedChunk(StreamingSoundWave, CurrentChunkIndex, ChunkSize);
 		if (SrcBufferData)
 		{
-			bPrintChunkFailMessage = true;
 			SrcBufferDataSize = ChunkSize;
 			SrcBufferOffset = CurrentChunkIndex == 0 ? AudioDataOffset : 0;
 
@@ -327,12 +335,48 @@ bool IStreamedCompressedInfo::StreamCompressedData(uint8* Destination, bool bLoo
 		}
 		else
 		{
-			// Still not loaded, zero remainder of current buffer
-			if (bPrintChunkFailMessage)
+			// Still not loaded, zero remainder of current buffer.
+			// For Chunk 1
+			//  For Load on Demand, this is expected, because it will likely wait for the first chunk to load.
+			//  For Prime on Load, it could happen if the load is still pending for the First Chunk.
+			//  For Retain on Load, it could also happen if the retained data has been purged for the cache. (or we're in the editor)
+			// For All other Chunks other than > 1
+			//  An underrun (starvation) has occured. Warn the user.
+			const ESoundWaveLoadingBehavior Behavior = StreamingSoundWave->GetLoadingBehavior();
+			const float TimeSoFar = FPlatformTime::ToSeconds64(FPlatformTime::Cycles64() - StartTimeInCycles);
+			if (CurrentChunkIndex == 1)
 			{
-				UE_LOG(LogAudio, Verbose, TEXT("Chunk %d not loaded from streaming manager for SoundWave '%s'. Likely due to stall on game thread."), CurrentChunkIndex, *StreamingSoundWave->GetFName().ToString());
-				bPrintChunkFailMessage = false;
+				switch(Behavior)
+				{
+				case ESoundWaveLoadingBehavior::RetainOnLoad:
+				{	// We don't inline in the editor, so don't warn.
+					const bool bIsEditor = GEngine ? GEngine->IsEditor() : false;
+					UE_CLOG(!bIsEditor, LogAudioStreamCaching, Verbose, TEXT("First Audio Chunk (%s) not loaded and we are 'RetainOnLoad'. Likely because the stream cache has purged it. FailCount=%d"),
+						*StreamingSoundWave->GetFName().ToString(), ++PrintChunkFailMessageCount);
+					break;
+				}
+				case ESoundWaveLoadingBehavior::PrimeOnLoad:
+				{
+					UE_LOG(LogAudioStreamCaching, Verbose, TEXT("First Audio Chunk (%s) not loaded and we are 'PrimeOnLoad'. Likely because the stream cache is still loading it. Attempt=%d"),
+						*StreamingSoundWave->GetFName().ToString(), ++PrintChunkFailMessageCount);
+					break;
+				}
+				case ESoundWaveLoadingBehavior::LoadOnDemand:
+				{
+					UE_LOG(LogAudioStreamCaching, VeryVerbose, TEXT("First Audio Chunk (%s) not loaded and we are 'LoadOnDemand'. This is expected while we load the first chunk. Attempt=%d"),
+						*StreamingSoundWave->GetFName().ToString(), ++PrintChunkFailMessageCount);
+					break;
+				}
+				default:
+					break;
+				}
+			}
+			else // If we're mid play back, and we hit this point. We've underrun, log a warning.
+			{
+				UE_LOG(LogAudioStreamCaching, Warning, TEXT("Chunk %d not yet loaded for playing SoundWave '%s'. Attempt=%d, LoadBehavior=%s, Latency=%2.2f secs, LikelyReason: Not loaded fast enough, IO Saturation." ),
+					CurrentChunkIndex, *StreamingSoundWave->GetFName().ToString(), ++PrintChunkFailMessageCount, EnumToString(Behavior), TimeSoFar);
 			}
+			
 			ZeroBuffer(Destination + RawPCMOffset, BufferSize - RawPCMOffset);
 			OutNumBytesStreamed = BufferSize - RawPCMOffset;
 			return false;
@@ -353,9 +397,11 @@ bool IStreamedCompressedInfo::StreamCompressedData(uint8* Destination, bool bLoo
 		// Decompress the next compression frame of audio (many samples) into the PCM buffer
 		int32 DecodedSamples = DecompressToPCMBuffer(/*Unused*/ 0);
 
-
 		if (DecodedSamples < 0)
 		{
+			UE_LOG(LogAudioStreamCaching, Warning, TEXT("Zero pad buffer Chunk=%d, Wave=%s, Reason=Decoder returned negative samples."),
+				CurrentChunkIndex, *StreamingSoundWave->GetFName().ToString());
+
 			LastPCMByteSize = 0;
 			ZeroBuffer(Destination + RawPCMOffset, BufferSize - RawPCMOffset);
 			OutNumBytesStreamed = BufferSize - RawPCMOffset;
@@ -415,10 +461,14 @@ bool IStreamedCompressedInfo::StreamCompressedData(uint8* Destination, bool bLoo
 						ensureMsgf(FStreamedAudioChunkSeekTable::Parse(SrcBufferData, SrcBufferDataSize, SrcBufferOffset, GetCurrentSeekTable()), 
 							TEXT("Failed to parse seektable in '%s' chunk=%d"), *StreamingSoundWave->GetFName().ToString(), CurrentChunkIndex);
 					}
-					UE_CLOG(PreviousChunkIndex != CurrentChunkIndex, LogAudio, Log, TEXT("Changed current chunk '%s' from %d to %d, Offset %d"), *StreamingSoundWave->GetFName().ToString(), PreviousChunkIndex, CurrentChunkIndex, SrcBufferOffset);
+					UE_CLOG(PreviousChunkIndex != CurrentChunkIndex, LogAudio, Log, TEXT("Changed current chunk '%s' from %d to %d, Offset %d"),
+						*StreamingSoundWave->GetFName().ToString(), PreviousChunkIndex, CurrentChunkIndex, SrcBufferOffset);
 				}
 				else
 				{
+					UE_LOG(LogAudioStreamCaching, Warning, TEXT("Chunk %d not yet loaded for playing SoundWave '%s'. LikelyReason: Not loaded fast enough, IO Saturation."),
+						CurrentChunkIndex, *StreamingSoundWave->GetFName().ToString());
+
 					SrcBufferDataSize = 0;
 					RawPCMOffset += ZeroBuffer(Destination + RawPCMOffset, BufferSize - RawPCMOffset);
 				}
diff --git a/Engine/Source/Runtime/Engine/Private/AudioDerivedData.cpp b/Engine/Source/Runtime/Engine/Private/AudioDerivedData.cpp
--- a/Engine/Source/Runtime/Engine/Private/AudioDerivedData.cpp
+++ b/Engine/Source/Runtime/Engine/Private/AudioDerivedData.cpp
@@ -22,6 +22,7 @@
 #include "DSP/MultichannelBuffer.h"
 #include "Containers/UnrealString.h"
 #include "Sound/StreamedAudioChunkSeekTable.h"
+#include "AudioDecompress.h"
 
 DEFINE_LOG_CATEGORY_STATIC(LogAudioDerivedData, Log, All);
 
@@ -166,8 +167,8 @@ Derived data key generation.
 
 // If you want to bump this version, generate a new guid using
 // VS->Tools->Create GUID and paste it here. https://www.guidgen.com works too.
-#define AUDIO_DERIVEDDATA_VER				TEXT("78D1D5180FCF4D3FB7E8B2A4799AD116")				// Change this if you want to bump all audio formats at once
-#define STREAMEDAUDIO_DERIVEDDATA_VER		TEXT("1A69BFC09E8146E08A1B7FD9BD5B9BA4")				// This depends on the above key, but will regenerate all streaming chunk data derived from the compressed audio
+#define AUDIO_DERIVEDDATA_VER				TEXT("E0B90A86877043398B7D2C0963526C72")				// Change this if you want to bump all audio formats at once
+#define STREAMEDAUDIO_DERIVEDDATA_VER		TEXT("D092D8E4F0024390B5D31F4EE3C4A2C5")				// This depends on the above key, but will regenerate all streaming chunk data derived from the compressed audio
 
 #if WITH_EDITORONLY_DATA
 
@@ -184,7 +185,7 @@ const TCHAR* LexToString(const ESoundwaveSampleRateSettings Enum)
 	return TEXT("<Unknown ESoundwaveSampleRateSettings>");
 }
 
-static FString GetSoundWaveHash(const USoundWave& InWave)
+static FString GetSoundWaveHash(const USoundWave& InWave, const ITargetPlatform* InTargetPlatform)
 {
 	// Hash the parts of the SoundWave that can affect the compressed data. It doesn't hurt to do this. 
 	// Typically the GUID will change if compressed data changes, but some settings can affect 
@@ -195,6 +196,8 @@ static FString GetSoundWaveHash(const USoundWave& InWave)
 	FPCU::AppendHash(SoundWaveHash, TEXT("QLT"), InWave.GetCompressionQuality());
 	FPCU::AppendHash(SoundWaveHash, TEXT("CHN"), InWave.NumChannels);
 	FPCU::AppendHash(SoundWaveHash, TEXT("SRQ"), InWave.SampleRateQuality);
+	FPCU::AppendHash(SoundWaveHash, TEXT("CK1"), InWave.GetSizeOfFirstAudioChunkInSeconds(InTargetPlatform));
+
 	return SoundWaveHash;
 }
 
@@ -208,6 +211,7 @@ static void GetStreamedAudioDerivedDataKeySuffix(
 	const USoundWave& SoundWave,
 	FName AudioFormatName,
 	const FPlatformAudioCookOverrides* CompressionOverrides,
+	const ITargetPlatform* InTargetPlatform,
 	FString& OutKeySuffix
 	)
 {
@@ -232,7 +236,7 @@ static void GetStreamedAudioDerivedDataKeySuffix(
 		FPlatformAudioCookOverrides::GetHashSuffix(CompressionOverrides, AudioFormatNameString);
 	}
 
-	FString SoundWaveHash = GetSoundWaveHash(SoundWave);
+	const FString SoundWaveHash = GetSoundWaveHash(SoundWave, InTargetPlatform);
 		
 	// build the key
 	OutKeySuffix = FString::Printf(TEXT("%s_%d_%s_%s"),
@@ -257,7 +261,6 @@ static void GetStreamedAudioDerivedDataKeySuffix(
  */
 static void GetStreamedAudioDerivedDataKeyFromSuffix(const FString& KeySuffix, FString& OutKey)
 {
-	static UE::DerivedData::FCacheBucket LegacyBucket(TEXTVIEW("LegacySTREAMEDAUDIO"), TEXTVIEW("Audio"));
 	OutKey = FDerivedDataCacheInterface::BuildCacheKey(
 		TEXT("STREAMEDAUDIO"),
 		STREAMEDAUDIO_DERIVEDDATA_VER,
@@ -295,11 +298,12 @@ static void GetStreamedAudioDerivedDataKey(
 	const USoundWave& SoundWave,
 	FName AudioFormatName,
 	const FPlatformAudioCookOverrides* CompressionOverrides,
+	const ITargetPlatform* InTargetPlatform,
 	FString& OutKey
 	)
 {
 	FString KeySuffix;
-	GetStreamedAudioDerivedDataKeySuffix(SoundWave, AudioFormatName, CompressionOverrides, KeySuffix);
+	GetStreamedAudioDerivedDataKeySuffix(SoundWave, AudioFormatName, CompressionOverrides, InTargetPlatform, KeySuffix);
 	GetStreamedAudioDerivedDataKeyFromSuffix(KeySuffix, OutKey);
 }
 
@@ -435,7 +439,15 @@ class FStreamedAudioCacheDerivedDataWorker : public FNonAbandonableTask
 	/** If the wave file is streaming */
 	bool bIsStreaming = false;
 	/** Initial chunk size */
-	int32 InitialChunkSize = 0;
+	int32 ZerothChunkSizeSoundWaveOverride = 0;
+	/* Cached platform specific sample rate */
+	uint32 SampleRate = 0;
+	/* Cached platform specific channel count */
+	int32 NumChannels = 0;
+	/* Cached platform specific number of AudioFrames */
+	uint64 NumOfAudioFrames = 0;
+	/* Size in Audio Frames of First Audio Chunk */
+	uint32 SizeOfFirstAudioChunkInFrames = 0;
 
 	bool GetCompressedData(TArray<uint8>& OutData)
 	{
@@ -516,20 +528,37 @@ class FStreamedAudioCacheDerivedDataWorker : public FNonAbandonableTask
 		OutChunkOffsets.Add(ChunkTable.FindTime(0));
 	}
 
-	bool SplitUsingSeekTable(const IAudioFormat::FSeekTable& InTable, const TArray<uint8>& InSrcBuffer, TArray<TArray<uint8>>& OutBuffers, const int32 InFirstChunkMaxSize, const int32 InMaxChunkSize, TArray<uint32>& OutChunkOffsets)
+	bool SplitUsingSeekTable(const IAudioFormat::FSeekTable& InTable, const TArray<uint8>& InSrcBuffer,
+	                         TArray<TArray<uint8>>& OutBuffers,
+	                         const int32 InFirstChunkMaxSize,const int32 InMaxChunkSize,
+	                         TArray<uint32>& OutChunkOffsets,
+	                         const uint32 InFirstAudioChunkMaxLengthInAudioFrames,
+	                        const uint32 InMaxChunkLengthInFrames) const
 	{
-		const TArray<uint32>& Offsets = InTable.Offsets;
+		// Typical layout:
+		// Chunk 0 - Header (always inlined)
+		// Chunk 1 - Audio (optionally inlined, containing n secs of audio).
+		// Chunk N - Audio - Sized using Max Chunk size.
 		
+		const TArray<uint32>& Offsets = InTable.Offsets;
+		const TArray<uint32>& Times = InTable.Times;
+
+		// Table needs to have equal number of time entries as offsets.
+		if (!ensure(Offsets.Num() == Times.Num()))
+		{
+			return false;
+		}
+		 
 		// Reject bad inputs.
 		if (Offsets.IsEmpty() || InSrcBuffer.IsEmpty() || InFirstChunkMaxSize <= 0 || InMaxChunkSize <=0)
 		{
 			return false;
 		}
 
-		EChunkSeekTableMode Mode = DetermineSeekTableMode(InTable);
+		const EChunkSeekTableMode Mode = DetermineSeekTableMode(InTable);
 
 		uint8 const* Source = InSrcBuffer.GetData();
-		uint32 SourceLen = InSrcBuffer.Num();
+		const uint32 SourceLen = InSrcBuffer.Num();
 		uint8 const* SourceEnd = Source + SourceLen;
 
 		uint8 const* ChunkStart = Source;
@@ -537,8 +566,10 @@ class FStreamedAudioCacheDerivedDataWorker : public FNonAbandonableTask
 
 		int32 BlockCount = 0;
 		int32 CurrentBlock = 0;
+		int32 CurrentAudioFrameCount = 0;
 
 		uint32 Budget = InFirstChunkMaxSize;
+		uint32 AudioFrameBudget = InFirstAudioChunkMaxLengthInAudioFrames;
 
 		while (Current < SourceEnd)
 		{	
@@ -554,7 +585,17 @@ class FStreamedAudioCacheDerivedDataWorker : public FNonAbandonableTask
 			{
 				return false;
 			}
-			
+
+			// Tally the duration of this block.
+			// Note we don't know the size of the very last block, so we set it 0. For this purpose that's fine.
+			const uint32 BlockStartTime = Times[CurrentBlock];
+			const uint32 BlockEndTime = CurrentBlock + 1 < Times.Num()-1 ? Times[CurrentBlock + 1] : BlockStartTime;
+			if (!ensureMsgf(BlockStartTime <= BlockEndTime, TEXT("Malformed Table: Block times out of order. BlockStartTime=%u, BlockEndTime=%u, CurrentBlock=%d"), BlockStartTime, BlockEnd,  CurrentBlock))
+			{
+				return false;
+			}
+			const uint32 BlockLengthAudioFrames = BlockEndTime - BlockStartTime;
+						
 			// Current chunk stats.
 			const int32 CurrentTableSize = FStreamedAudioChunkSeekTable::CalcSize(BlockCount, Mode);
 			checkf(CurrentTableSize >= 0, TEXT("CurrentTableSize=%d"), CurrentTableSize);
@@ -566,8 +607,12 @@ class FStreamedAudioCacheDerivedDataWorker : public FNonAbandonableTask
 			const int32 TableDelta = NewTableSize - CurrentTableSize;
 			checkf(TableDelta > 0, TEXT("Table should grow when we add a new entry: TableDelta=%d"), TableDelta);
 			
-			// Accumulated enough?
-			if (CurrentChunkSize + BlockSize + TableDelta > Budget)
+			// Accumulated enough? (bytes or audio frames).
+			check(Budget > 0);
+			const bool bOverSize = (CurrentChunkSize + BlockSize + TableDelta) > Budget;
+			const bool bOverLength = AudioFrameBudget > 0 && CurrentAudioFrameCount > 0 && (CurrentAudioFrameCount + BlockLengthAudioFrames > AudioFrameBudget);
+			
+			if (bOverSize || bOverLength)
 			{
 				// can't add this chunk, emit.
 				TArray<uint8> Chunk(ChunkStart, Current - ChunkStart);
@@ -575,21 +620,33 @@ class FStreamedAudioCacheDerivedDataWorker : public FNonAbandonableTask
 
 				CreateChunkSeektableAndPrefix(Mode, InTable, ChunkStart-Source, Current-Source, Chunk, OutChunkOffsets, CurrentTableSize, BlockCount, AudioSize, Budget);
 		
-				UE_LOG(LogAudio, Verbose, TEXT("Adding Chunk %d: Blocks=%d (%u bytes), SeekTableEntries=%d (%d bytes), ChunkSize=%u bytes, PercentFull=%2.2f, Remaining=%u bytes"),
-					OutBuffers.Num(), BlockCount, AudioSize, BlockCount, CurrentTableSize, Chunk.Num(), ((float)Chunk.Num() / Budget) * 100.f, Budget-Chunk.Num());
+				UE_LOG(LogAudio, Verbose, TEXT("Adding Chunk %d because (%s): Blocks=%d (%u bytes), AudioFrames=%u (%2.2f seconds), SeekTableEntries=%d (%d bytes), ChunkSize=%u bytes, PercentFull=%2.2f, Remaining=%u bytes"),
+					OutBuffers.Num(), bOverSize ? TEXT("Enough Bytes") : TEXT("Enough Samples"), BlockCount, AudioSize, CurrentAudioFrameCount,
+					(float)CurrentAudioFrameCount / SampleRate, BlockCount, CurrentTableSize, Chunk.Num(), ((float)Chunk.Num() / Budget) * 100.f, Budget-Chunk.Num()
+					);
 
 				OutBuffers.Add(Chunk);
 
+				// Reset our pointers for a new chunk.
 				ChunkStart = Current;
+				BlockCount = 0;
+				CurrentAudioFrameCount = 0;
+				
+				// Reset size budget. (note different budget after chunk 0).
 				Budget = InMaxChunkSize;
-				BlockCount = 0;				
-			
+
+				// Reset length budget. (note different budget for chunk 1, the audio chunk).
+				// The first "Audio" chunk is considered to be Chunk 1				
+				const bool bIsFirstAudioChunk = OutBuffers.Num() == 1;
+				AudioFrameBudget = bIsFirstAudioChunk ? InFirstAudioChunkMaxLengthInAudioFrames : InMaxChunkLengthInFrames;
+				
 				// retry.
 				continue;
 			}
 				
 			// Include this block.
 			Current += BlockSize;
+			CurrentAudioFrameCount += BlockLengthAudioFrames;
 			BlockCount++;
 			CurrentBlock++;
 
@@ -605,8 +662,10 @@ class FStreamedAudioCacheDerivedDataWorker : public FNonAbandonableTask
 			const int32 CurrentTableSize = FStreamedAudioChunkSeekTable::CalcSize(BlockCount, Mode);
 			CreateChunkSeektableAndPrefix(Mode, InTable, ChunkStart - Source, Current - Source, Chunk, OutChunkOffsets, CurrentTableSize, BlockCount, AudioSize, Budget);
 
-			UE_LOG(LogAudio, Verbose, TEXT("Adding Chunk %d: Blocks=%d (%u bytes), SeekTableEntries=%d (%d bytes), ChunkSize=%u bytes, PercentFull=%2.2f, Remaining=%u bytes"),
-				OutBuffers.Num(), BlockCount, AudioSize, BlockCount, CurrentTableSize, Chunk.Num(), ((float)Chunk.Num() / Budget) * 100.f, Budget - Chunk.Num());
+			UE_LOG(LogAudio, Verbose, TEXT("Adding FINAL Chunk %d because (%s): Blocks=%d (%u bytes), AudioFrames=%u (%2.2f seconds), SeekTableEntries=%d (%d bytes), ChunkSize=%u bytes, PercentFull=%2.2f, Remaining=%u bytes"),
+				OutBuffers.Num(), TEXT("EOF"), BlockCount, AudioSize, CurrentAudioFrameCount,
+				(float)CurrentAudioFrameCount / SampleRate, BlockCount, CurrentTableSize, Chunk.Num(), ((float)Chunk.Num() / Budget) * 100.f, Budget-Chunk.Num()
+				);
 
 			OutBuffers.Add(Chunk);
 		}
@@ -614,12 +673,11 @@ class FStreamedAudioCacheDerivedDataWorker : public FNonAbandonableTask
 		return true;
 	}
 
-	bool SplitDataForStreaming(const IAudioFormat* InFormat, const TArray<uint8>& InSrcBuffer, TArray<TArray<uint8>>& OutBuffers, const int32 InFirstChunkMaxSize, const int32 InMaxChunkSize, TArray<uint32>& OutChunkOffsets)
+	bool SplitDataForStreaming(const IAudioFormat* InFormat, const TArray<uint8>& InSrcBuffer, TArray<TArray<uint8>>& OutBuffers, const int32 InFirstChunkMaxSize, const int32 InMaxChunkSize, TArray<uint32>& OutChunkOffsets,
+		const uint32 InFirstChunkMaxLengthAudioFrames, const uint32 InMaxChunkLengthAudioFrames) const
 	{
-		bool bRequiresSeekTable = InFormat->RequiresStreamingSeekTable();	
-
 		// Split using the seek-table if we have one.
-		if (bRequiresSeekTable)
+		if (const bool bRequiresSeekTable = InFormat->RequiresStreamingSeekTable())
 		{
 			// Because the extract call modifies the source buffer inline (to remove the embedded seek table) we must keep a copy as if
 			// the call was to fail to split the calling code must contend with the original.
@@ -630,7 +688,16 @@ class FStreamedAudioCacheDerivedDataWorker : public FNonAbandonableTask
 			{
 				return false;
 			}
-			return SplitUsingSeekTable(Table, InSrcBufferCopy,OutBuffers,InFirstChunkMaxSize,InMaxChunkSize, OutChunkOffsets);
+
+			// A limitation of our pipeline requires that we decode a minimum of MONO_PCM_BUFFER_SAMPLES (8k) for our first chunk.
+			// So don't allow our first chunk size in frames be less than that.
+			int32 FirstChunkMaxLengthAudioFrames = InFirstChunkMaxLengthAudioFrames;
+			if (FirstChunkMaxLengthAudioFrames > 0)
+			{
+				FirstChunkMaxLengthAudioFrames = FMath::Max(MONO_PCM_BUFFER_SAMPLES, FirstChunkMaxLengthAudioFrames);
+			}
+
+			return SplitUsingSeekTable(Table, InSrcBufferCopy,OutBuffers,InFirstChunkMaxSize,InMaxChunkSize, OutChunkOffsets, FirstChunkMaxLengthAudioFrames, InMaxChunkLengthAudioFrames);
 		}
 
 		// otherwise... ask the format to split.
@@ -668,16 +735,16 @@ class FStreamedAudioCacheDerivedDataWorker : public FNonAbandonableTask
 				// Set the ideal chunk size to be 256k to optimize for data reads on console.
 				int32 MaxChunkSizeForCurrentWave = 256 * 1024;
 				
-				// By default, the first chunk's max size is the same as the other chunks.
-				int32 FirstChunkSize = MaxChunkSizeForCurrentWave;
+				// By default, the zeroth chunk's max size is the same as the other chunks.
+				int32 ZerothChunkSize = MaxChunkSizeForCurrentWave;
 
 				const int32 MinimumChunkSize = AudioFormat->GetMinimumSizeForInitialChunk(AudioFormatName, CompressedBuffer);
 				const bool bForceLegacyStreamChunking = bIsStreaming && CompressionOverrides && CompressionOverrides->StreamCachingSettings.bForceLegacyStreamChunking;
 
-				// If the initial chunk  for this sound wave was overridden, use that:
-				if (InitialChunkSize > 0)
+				// If the zeroth size for this sound wave was overridden, use that:
+				if (ZerothChunkSizeSoundWaveOverride > 0)
 				{
-					FirstChunkSize = FMath::Max(MinimumChunkSize, InitialChunkSize);
+					ZerothChunkSize = FMath::Max(MinimumChunkSize, ZerothChunkSizeSoundWaveOverride);
 				}
 				else
 				{
@@ -693,12 +760,12 @@ class FStreamedAudioCacheDerivedDataWorker : public FNonAbandonableTask
 							LegacyZerothChunkSize = MaxChunkSizeForCurrentWave;
 						}
 
-						FirstChunkSize = LegacyZerothChunkSize;
+						ZerothChunkSize = LegacyZerothChunkSize;
 					}
 					else
 					{
 						// Otherwise if we're using Audio Stream Caching, the first chunk should be as small as possible:
-						FirstChunkSize = MinimumChunkSize;
+						ZerothChunkSize = MinimumChunkSize;
 					}
 				}
 
@@ -707,27 +774,46 @@ class FStreamedAudioCacheDerivedDataWorker : public FNonAbandonableTask
 					// Use the chunk size for this duration:
 					MaxChunkSizeForCurrentWave = FPlatformCompressionUtilities::GetMaxChunkSizeForCookOverrides(CompressionOverrides);
 
-					// observe the override chunk size now that we have set the FirstChunkSize
+					// observe the override chunk size now that we have set the 
 					const int32 MaxChunkSizeOverrideBytes = CompressionOverrides->StreamCachingSettings.MaxChunkSizeOverrideKB * 1024;
 					if (MaxChunkSizeOverrideBytes > 0)
 					{
 						MaxChunkSizeForCurrentWave = FMath::Min(MaxChunkSizeOverrideBytes, MaxChunkSizeForCurrentWave);
 					}
 
-					UE_LOG(LogAudio, Display, TEXT("Chunk sizes for %s: Chunk0=%d bytes, Chunk1-N=%dk"), *SoundWaveFullName, FirstChunkSize, MaxChunkSizeForCurrentWave>>10);
 				}
 				
-				check(FirstChunkSize != 0 && MaxChunkSizeForCurrentWave != 0);
+
+				check(ZerothChunkSize != 0 && MaxChunkSizeForCurrentWave != 0);
+
+				// If platform is configured to so, optionally inline the first Audio chunk and size it with the specific requirements.
+				// Value of zero will do nothing.
+				const bool bInlineFirstAudioChunk = CompressionOverrides->bInlineFirstAudioChunk && !bForceLegacyStreamChunking;
+				uint32 MaxLengthOfFirstAudioChunkInFrames = MAX_uint32;
+				if (bInlineFirstAudioChunk)
+				{
+					MaxLengthOfFirstAudioChunkInFrames = SizeOfFirstAudioChunkInFrames;
+				}
 
 				TArray<uint32> ChunkOffsets;
-				if (SplitDataForStreaming(AudioFormat, CompressedBuffer, ChunkBuffers, FirstChunkSize, MaxChunkSizeForCurrentWave, ChunkOffsets))
+				if (SplitDataForStreaming(AudioFormat, CompressedBuffer, ChunkBuffers, ZerothChunkSize,
+				                          MaxChunkSizeForCurrentWave, ChunkOffsets, MaxLengthOfFirstAudioChunkInFrames,
+				                          MAX_uint32))
 				{
+					UE_LOG(LogAudio, Display, TEXT("Chunk stats for (%s: Duration=%2.2f secs, Channels=%d), Settings(FirstChunkSize=%u frames, InlineFirst=%s, MaxChunkSize=%d), Chunks=%d, Chunk0=%d bytes, Chunk1=%d bytes, ChunkN=%d)"),
+						*SoundWaveFullName,(float)NumOfAudioFrames / SampleRate, NumChannels,
+						MaxLengthOfFirstAudioChunkInFrames, ToCStr(LexToString(bInlineFirstAudioChunk)), MaxChunkSizeForCurrentWave,
+						ChunkBuffers.Num(),
+						ChunkBuffers.IsValidIndex(0) ? ChunkBuffers[0].Num() : 0,
+						ChunkBuffers.IsValidIndex(1) ? ChunkBuffers[1].Num() : 0,
+						ChunkBuffers.IsValidIndex(2) ? ChunkBuffers[2].Num() : 0
+						);
+					
 					if (ChunkBuffers.Num() > 32)
 					{
 						UE_LOG(LogAudio, Display, TEXT("Sound Wave %s is very large, requiring %d chunks."), *SoundWaveFullName, ChunkBuffers.Num());
 					}
 
-					int32 TotalAudioBytesExcludingSeektable = 0;
 					if (ChunkBuffers.Num() > 0)
 					{
 						// The zeroth chunk should not be zero-padded.
@@ -776,6 +862,12 @@ class FStreamedAudioCacheDerivedDataWorker : public FNonAbandonableTask
 						NewChunk->AudioDataSize = AudioDataSize;
 						NewChunk->DataSize = AudioDataSize + ZeroPadBytes;
 						NewChunk->SeekOffsetInAudioFrames = ChunkOffsets.IsValidIndex(ChunkIndex) ? ChunkOffsets[ChunkIndex] : INDEX_NONE;
+
+						// If this is the first chunk of Audio, ask that we inline it when we serialize.
+						// NOTE we should only do this if we've been given a size greater than 0 to put there
+						// otherwise the chunk will use the normal size boundary.
+						const bool bIsFirstAudioChunk = NewChunk->SeekOffsetInAudioFrames == 0;
+						NewChunk->bInlineChunk = bIsFirstAudioChunk && bInlineFirstAudioChunk && SizeOfFirstAudioChunkInFrames > 0;
 						
 						if (NewChunk->BulkData.IsLocked())
 						{
@@ -854,7 +946,8 @@ public:
 		USoundWave* InSoundWave,
 		const FPlatformAudioCookOverrides* InCompressionOverrides,
 		FName InAudioFormatName,
-		uint32 InCacheFlags
+		uint32 InCacheFlags,
+		const ITargetPlatform* InTargetPlatform=nullptr
 		)
 		: DerivedData(InDerivedData)
 		, SoundWavePath(InSoundWave->GetPathName())
@@ -864,14 +957,26 @@ public:
 		, CompressionOverrides(InCompressionOverrides)
 		, bIsProcedural(InSoundWave->IsA<USoundWaveProcedural>())
 		, bIsStreaming(InSoundWave->bStreaming)
-		, InitialChunkSize(InSoundWave->InitialChunkSize)
+		, ZerothChunkSizeSoundWaveOverride(InSoundWave->InitialChunkSize_DEPRECATED)
 	{
 		// Gather all USoundWave object inputs to avoid race-conditions that could result when touching the UObject from another thread.
-		GetStreamedAudioDerivedDataKeySuffix(*InSoundWave, AudioFormatName, CompressionOverrides, KeySuffix);
-		FName PlatformSpecificFormat = InSoundWave->GetPlatformSpecificFormat(InAudioFormatName, CompressionOverrides);
+		GetStreamedAudioDerivedDataKeySuffix(*InSoundWave, AudioFormatName, CompressionOverrides, InTargetPlatform, KeySuffix);
+		const FName PlatformSpecificFormat = InSoundWave->GetPlatformSpecificFormat(InAudioFormatName, CompressionOverrides);
+
+		// Cache num channels (only used for logging currently).
+		NumChannels = InSoundWave->NumChannels;
+
+		// Cache the sample-rate (making sure to use the compression overrides version for this platform).
+		const float OverrideSampleRate = InSoundWave->GetSampleRateForCompressionOverrides(InCompressionOverrides);
+		SampleRate = OverrideSampleRate > 0 ? static_cast<uint32>(OverrideSampleRate)
+			: static_cast<uint32>(InSoundWave->GetSampleRateForCurrentPlatform());
+
+		// Determine amount to inline in Audio frames. (careful to round for determinism reasons).
+		const float SizeOfFirstAudioChunkInSeconds = InSoundWave->GetSizeOfFirstAudioChunkInSeconds(InTargetPlatform);
+		SizeOfFirstAudioChunkInFrames = FMath::CeilToInt(static_cast<float>(SampleRate * SizeOfFirstAudioChunkInSeconds));
 
  		// Fetch compressed data directly from the DDC to ensure thread-safety. Will be async if the compressor is thread-safe.
-		FDerivedAudioDataCompressor* DeriveAudioData = new FDerivedAudioDataCompressor(InSoundWave, InAudioFormatName, PlatformSpecificFormat, CompressionOverrides);
+		FDerivedAudioDataCompressor* DeriveAudioData = new FDerivedAudioDataCompressor(InSoundWave, InAudioFormatName, PlatformSpecificFormat, CompressionOverrides, InTargetPlatform);
 		CompressedAudioHandle = GetDerivedDataCacheRef().GetAsynchronous(DeriveAudioData);
 	}
 
@@ -971,20 +1076,22 @@ struct FStreamedAudioAsyncCacheDerivedDataTask : public FAsyncTask<FStreamedAudi
 		USoundWave* InSoundWave,
 		const FPlatformAudioCookOverrides* CompressionSettings,
 		FName InAudioFormatName,
-		uint32 InCacheFlags
+		uint32 InCacheFlags,
+		const ITargetPlatform* InTargetPlatform=nullptr
 		)
 		: FAsyncTask<FStreamedAudioCacheDerivedDataWorker>(
 			InDerivedData,
 			InSoundWave,
 			CompressionSettings,
 			InAudioFormatName,
-			InCacheFlags
+			InCacheFlags,
+			InTargetPlatform
 			)
 	{
 	}
 };
 
-void FStreamedAudioPlatformData::Cache(USoundWave& InSoundWave, const FPlatformAudioCookOverrides* CompressionOverrides, FName AudioFormatName,  uint32 InFlags)
+void FStreamedAudioPlatformData::Cache(USoundWave& InSoundWave, const FPlatformAudioCookOverrides* CompressionOverrides, FName AudioFormatName, uint32 InFlags, const ITargetPlatform* InTargetPlatform)
 {
 	TRACE_CPUPROFILER_EVENT_SCOPE(FStreamedAudioPlatformData::Cache);
 
@@ -1001,7 +1108,7 @@ void FStreamedAudioPlatformData::Cache(USoundWave& InSoundWave, const FPlatformA
 
 	bool bForceRebuild = (Flags & EStreamedAudioCacheFlags::ForceRebuild) != 0;
 	bool bAsync = (Flags & EStreamedAudioCacheFlags::Async) != 0;
-	GetStreamedAudioDerivedDataKey(InSoundWave, AudioFormatName, CompressionOverrides, DerivedDataKey);
+	GetStreamedAudioDerivedDataKey(InSoundWave, AudioFormatName, CompressionOverrides, InTargetPlatform, DerivedDataKey);
 
 	if (bAsync && !bForceRebuild && FSoundWaveCompilingManager::Get().IsAsyncCompilationAllowed(&InSoundWave))
 	{
@@ -1011,7 +1118,7 @@ void FStreamedAudioPlatformData::Cache(USoundWave& InSoundWave, const FPlatformA
 		{
 			FWriteScopeLock AsyncTaskScope(AsyncTaskLock.Get());
 			check(AsyncTask == nullptr);
-			AsyncTask = new FStreamedAudioAsyncCacheDerivedDataTask(this, &InSoundWave, CompressionOverrides, AudioFormatName, Flags);
+			AsyncTask = new FStreamedAudioAsyncCacheDerivedDataTask(this, &InSoundWave, CompressionOverrides, AudioFormatName, Flags, InTargetPlatform);
 			int64 RequiredMemory = -1; // @todo RequiredMemory
 			AsyncTask->StartBackgroundTask(SoundWaveThreadPool, BasePriority, EQueuedWorkFlags::DoNotRunInsideBusyWait, RequiredMemory, TEXT("AudioDerivedData") );
 		}
@@ -1034,7 +1141,7 @@ void FStreamedAudioPlatformData::Cache(USoundWave& InSoundWave, const FPlatformA
 	}
 	else
 	{
-		FStreamedAudioCacheDerivedDataWorker Worker(this, &InSoundWave, CompressionOverrides, AudioFormatName, Flags);
+		FStreamedAudioCacheDerivedDataWorker Worker(this, &InSoundWave, CompressionOverrides, AudioFormatName, Flags, InTargetPlatform);
 		{
 			COOK_STAT(auto Timer = AudioCookStats::UsageStats.TimeSyncWork());
 			Worker.DoWork();
@@ -1372,7 +1479,9 @@ void FStreamedAudioPlatformData::Serialize(FArchive& Ar, USoundWave* Owner)
 
 	if (Ar.IsLoading())
 	{
+		check(!AudioFormat.IsNone());
 		check(NumChunks >= 0);
+
 		Chunks.Empty(NumChunks);
 		for (int32 ChunkIndex = 0; ChunkIndex < NumChunks; ++ChunkIndex)
 		{
@@ -1484,7 +1593,7 @@ struct FAudioCookInputs
 	UE::Serialization::FEditorBulkData& BulkData;
 #endif
 
-	FAudioCookInputs(USoundWave* InSoundWave, FName InBaseFormat, FName InHashFormat, const FPlatformAudioCookOverrides* InCookOverrides)
+	FAudioCookInputs(USoundWave* InSoundWave, FName InBaseFormat, FName InHashFormat, const FPlatformAudioCookOverrides* InCookOverrides, const ITargetPlatform* InTargetPlatform)
 		: SoundName(InSoundWave->GetName())
 		, BaseFormat(InBaseFormat)
 		, HashedFormat(InHashFormat)
@@ -1497,7 +1606,7 @@ struct FAudioCookInputs
 		, bIsSoundWaveProcedural(InSoundWave->IsA<USoundWaveProcedural>())
 		, CompressionQuality(InSoundWave->GetCompressionQuality())
 		, CompressionQualityModifier(InCookOverrides ? InCookOverrides->CompressionQualityModifier : 1.0f)
-		, SoundWaveHash(GetSoundWaveHash(*InSoundWave))
+		, SoundWaveHash(GetSoundWaveHash(*InSoundWave, InTargetPlatform))
 		, WaveTransformations(InSoundWave->CreateTransformations())
 		, BulkDataCriticalSection(InSoundWave->RawDataCriticalSection)
 		, BulkData(InSoundWave->RawData)
@@ -1934,8 +2043,8 @@ static void CookSurroundWave(const FAudioCookInputs& Inputs,  TArray<uint8>& Out
 
 #endif // WITH_EDITORONLY_DATA
 
-FDerivedAudioDataCompressor::FDerivedAudioDataCompressor(USoundWave* InSoundNode, FName InBaseFormat, FName InHashedFormat, const FPlatformAudioCookOverrides* InCompressionOverrides)
-	: CookInputs(MakeUnique<FAudioCookInputs>(InSoundNode, InBaseFormat, InHashedFormat, InCompressionOverrides))
+FDerivedAudioDataCompressor::FDerivedAudioDataCompressor(USoundWave* InSoundNode, FName InBaseFormat, FName InHashedFormat, const FPlatformAudioCookOverrides* InCompressionOverrides, const ITargetPlatform* InTargetPlatform)
+	: CookInputs(MakeUnique<FAudioCookInputs>(InSoundNode, InBaseFormat, InHashedFormat, InCompressionOverrides, InTargetPlatform))
 {
 }
 
@@ -2036,14 +2145,14 @@ void USoundWave::SerializeCookedPlatformData(FArchive& Ar)
 		const FPlatformAudioCookOverrides* CompressionOverrides = FPlatformCompressionUtilities::GetCookOverrides(*Ar.CookingTarget()->IniPlatformName());
 		FString DerivedDataKey;
 
-		GetStreamedAudioDerivedDataKeySuffix(*this, PlatformFormat, CompressionOverrides, DerivedDataKey);
+		GetStreamedAudioDerivedDataKeySuffix(*this, PlatformFormat, CompressionOverrides, Ar.CookingTarget(), DerivedDataKey);
 
 		FStreamedAudioPlatformData *PlatformDataToSave = CookedPlatformData.FindRef(DerivedDataKey);
 
 		if (PlatformDataToSave == NULL)
 		{
 			PlatformDataToSave = new FStreamedAudioPlatformData();
-			PlatformDataToSave->Cache(*this, CompressionOverrides, PlatformFormat, EStreamedAudioCacheFlags::InlineChunks | EStreamedAudioCacheFlags::Async);
+			PlatformDataToSave->Cache(*this, CompressionOverrides, PlatformFormat, EStreamedAudioCacheFlags::InlineChunks | EStreamedAudioCacheFlags::Async, Ar.CookingTarget());
 
 			CookedPlatformData.Add(DerivedDataKey, PlatformDataToSave);
 		}
@@ -2078,12 +2187,12 @@ void USoundWave::CachePlatformData(bool bAsyncCache)
 	FString DerivedDataKey;
 	FName AudioFormat = GetWaveFormatForRunningPlatform(*this);
 	const FPlatformAudioCookOverrides* CompressionOverrides = GetCookOverridesForRunningPlatform();
-	GetStreamedAudioDerivedDataKey(*this, AudioFormat, CompressionOverrides, DerivedDataKey);
+	GetStreamedAudioDerivedDataKey(*this, AudioFormat, CompressionOverrides, GetRunningPlatform(),  DerivedDataKey);
 
 	if (SoundWaveDataPtr->RunningPlatformData.DerivedDataKey != DerivedDataKey)
 	{
 		const uint32 CacheFlags = bAsyncCache ? (EStreamedAudioCacheFlags::Async | EStreamedAudioCacheFlags::AllowAsyncBuild) : EStreamedAudioCacheFlags::None;
-		SoundWaveDataPtr->RunningPlatformData.Cache(*this, CompressionOverrides, AudioFormat, CacheFlags);
+		SoundWaveDataPtr->RunningPlatformData.Cache(*this, CompressionOverrides, AudioFormat, CacheFlags, GetRunningPlatform());
 	}
 }
 
@@ -2125,7 +2234,7 @@ void USoundWave::BeginCacheForCookedPlatformData(const ITargetPlatform *TargetPl
 
 			// find format data by comparing derived data keys.
 			FString DerivedDataKey;
-			GetStreamedAudioDerivedDataKeySuffix(*this, PlatformFormat, CompressionOverrides, DerivedDataKey);
+			GetStreamedAudioDerivedDataKeySuffix(*this, PlatformFormat, CompressionOverrides, TargetPlatform, DerivedDataKey);
 
 			FStreamedAudioPlatformData *PlatformData = CookedPlatformData.FindRef(DerivedDataKey);
 
@@ -2136,14 +2245,15 @@ void USoundWave::BeginCacheForCookedPlatformData(const ITargetPlatform *TargetPl
 					*this,
 					CompressionOverrides,
 					PlatformFormat,
-					CacheFlags
+					CacheFlags,
+					TargetPlatform
 					);
 				CookedPlatformData.Add(DerivedDataKey, PlatformData);
 			}
 		}
 		else
 		{
-			BeginGetCompressedData(PlatformFormat, CompressionOverrides);
+			BeginGetCompressedData(PlatformFormat, CompressionOverrides, TargetPlatform);
 		}
 	}
 
@@ -2165,7 +2275,7 @@ bool USoundWave::IsCachedCookedPlatformDataLoaded( const ITargetPlatform* Target
 		{
 			// find format data by comparing derived data keys.
 			FString DerivedDataKey;
-			GetStreamedAudioDerivedDataKeySuffix(*this, PlatformFormat, CompressionOverrides, DerivedDataKey);
+			GetStreamedAudioDerivedDataKeySuffix(*this, PlatformFormat, CompressionOverrides, TargetPlatform, DerivedDataKey);
 
 			FStreamedAudioPlatformData *PlatformData = CookedPlatformData.FindRef(DerivedDataKey);
 			if (PlatformData == nullptr)
@@ -2224,7 +2334,7 @@ void USoundWave::ClearCachedCookedPlatformData( const ITargetPlatform* TargetPla
 
 		// find format data by comparing derived data keys.
 		FString DerivedDataKey;
-		GetStreamedAudioDerivedDataKeySuffix(*this, PlatformFormat, CompressionOverrides, DerivedDataKey);
+		GetStreamedAudioDerivedDataKeySuffix(*this, PlatformFormat, CompressionOverrides, TargetPlatform, DerivedDataKey);
 
 
 		if ( CookedPlatformData.Contains(DerivedDataKey) )
@@ -2273,7 +2383,7 @@ void USoundWave::FinishCachePlatformData()
 		FString DerivedDataKey;
 		FName AudioFormat = GetWaveFormatForRunningPlatform(*this);
 		const FPlatformAudioCookOverrides* CompressionOverrides = GetCookOverridesForRunningPlatform();
-		GetStreamedAudioDerivedDataKey(*this, AudioFormat, CompressionOverrides, DerivedDataKey);
+		GetStreamedAudioDerivedDataKey(*this, AudioFormat,CompressionOverrides, GetRunningPlatform(), DerivedDataKey);
 
 		UE_CLOG(SoundWaveDataPtr->RunningPlatformData.DerivedDataKey != DerivedDataKey, LogAudio, Warning, TEXT("Audio was cooked with the DDC key %s but should've had the DDC key %s. the cook overrides/codec used may be incorrect."), *SoundWaveDataPtr->RunningPlatformData.DerivedDataKey, *DerivedDataKey);
 	}
@@ -2289,7 +2399,8 @@ void USoundWave::ForceRebuildPlatformData()
 		*this,
 		CompressionOverrides,
 		GetWaveFormatForRunningPlatform(*this),
-		EStreamedAudioCacheFlags::ForceRebuild
+		EStreamedAudioCacheFlags::ForceRebuild,
+		GetRunningPlatform()
 		);
 }
 #endif //WITH_EDITORONLY_DATA
diff --git a/Engine/Source/Runtime/Engine/Private/AudioStreamingCache.cpp b/Engine/Source/Runtime/Engine/Private/AudioStreamingCache.cpp
--- a/Engine/Source/Runtime/Engine/Private/AudioStreamingCache.cpp
+++ b/Engine/Source/Runtime/Engine/Private/AudioStreamingCache.cpp
@@ -1571,7 +1571,7 @@ void FAudioChunkCache::KickOffAsyncLoad(FCacheElement* CacheElement, const FChun
 {
 	check(CacheElement);
 
-	const FStreamedAudioChunk* Chunk = CacheElement->GetChunk(InKey.ChunkIndex);
+	FStreamedAudioChunk* Chunk = CacheElement->GetChunk(InKey.ChunkIndex);
 
 	if (nullptr == Chunk)
 	{
@@ -1692,7 +1692,7 @@ void FAudioChunkCache::KickOffAsyncLoad(FCacheElement* CacheElement, const FChun
 			CacheElement->bIsLoaded = true;
 
 #if DEBUG_STREAM_CACHE
-			CacheElement->DebugInfo.TimeToLoad = (FPlatformTime::Seconds() - CacheElement->DebugInfo.TimeLoadStarted) * 1000.0f;
+			CacheElement->DebugInfo.TimeToLoad = FPlatformTime::ToMilliseconds64(FPlatformTime::Cycles64() - CacheElement->DebugInfo.TimeLoadStarted);
 #endif
 			const EAudioChunkLoadResult LoadResult = bWasCancelled ? EAudioChunkLoadResult::Interrupted : EAudioChunkLoadResult::Completed;
 			ExecuteOnLoadCompleteCallback(LoadResult, OnLoadCompleted, CallbackThread);
@@ -1701,21 +1701,53 @@ void FAudioChunkCache::KickOffAsyncLoad(FCacheElement* CacheElement, const FChun
 		};
 
 #if DEBUG_STREAM_CACHE
-		CacheElement->DebugInfo.TimeLoadStarted = FPlatformTime::Seconds();
+		CacheElement->DebugInfo.TimeLoadStarted = FPlatformTime::Cycles64();
 #endif
 
 		CacheElement->ReadRequest = nullptr;
-		IBulkDataIORequest* LocalReadRequest = Chunk->BulkData.CreateStreamingRequest(0, ChunkDataSize, AsyncIOPriority | AIOP_FLAG_DONTCACHE, &AsyncFileCallBack, CacheElement->ChunkData);
-		if (!LocalReadRequest)
+		if (Chunk->BulkData.IsBulkDataLoaded())
 		{
-			UE_LOG(LogAudioStreamCaching, Error, TEXT("Chunk load in audio LRU cache failed."));
-			OnLoadCompleted(EAudioChunkLoadResult::ChunkOutOfBounds);
-			NumberOfLoadsInFlight.Decrement();
+			// If this chunk has been inlined and loaded, move out the data into our newly allocated block.
+			const FBulkDataBuffer<uint8> ChunkMemory = Chunk->MoveOutAsBuffer();
+			
+			// Copy and delete to be sure we pay back the LLM and use our newly allocated version.
+			check(CacheElement->ChunkDataSize <= ChunkMemory.GetView().Num());
+			FMemory::Memcpy(CacheElement->ChunkData, ChunkMemory.GetView().GetData(), ChunkMemory.GetView().Num());
+
+#if DEBUG_STREAM_CACHE
+			UE_LOG(LogAudioStreamCaching, Verbose, TEXT("Loading Inlined Chunk: %s, %d, TimeToLoad=%2.2fms"), *InKey.SoundWaveName.ToString(),
+				InKey.ChunkIndex, FPlatformTime::ToMilliseconds64(FPlatformTime::Cycles64() - CacheElement->DebugInfo.TimeLoadStarted));
+
+			CacheElement->DebugInfo.bWasLoadedFromInlineChunk = true;
+			UE_LOG(LogAudioStreamCaching, VeryVerbose, TEXT("COPY+DISCARD %s - Bulk=0x%p"), *InKey.SoundWaveName.ToString(), &Chunk->BulkData);
+#endif //DEBUG_STREAM_CACHE
+			
+			// Fire the callback (this will mark it load completed etc).
+			AsyncFileCallBack(false, nullptr);
 		}
-		else if (FPlatformAtomics::InterlockedCompareExchangePointer((void* volatile*)&CacheElement->ReadRequest, LocalReadRequest, nullptr) == (void*)0x1)
+		else
 		{
-			// The request is completed before we can store it. Just delete it
-			TGraphTask<FClearAudioChunkCacheReadRequestTask>::CreateTask().ConstructAndDispatchWhenReady(LocalReadRequest);
+			UE_LOG(LogAudioStreamCaching, VeryVerbose, TEXT("DISK %s - Bulk=0x%p"), *InKey.SoundWaveName.ToString(), &Chunk->BulkData)
+
+#if DEBUG_STREAM_CACHE
+			CacheElement->DebugInfo.bWasInlinedButUnloaded = Chunk->BulkData.IsInlined() || Chunk->BulkData.GetBulkDataFlags() & BULKDATA_ForceInlinePayload;
+			UE_CLOG(CacheElement->DebugInfo.bWasInlinedButUnloaded,LogAudioStreamCaching, Log, TEXT("IO LOAD FOR INLINE %s - Bulk=0x%p"), *InKey.SoundWaveName.ToString(), &Chunk->BulkData);
+#endif //DEBUG_STREAM_CACHE
+			
+			UE_LOG(LogAudioStreamCaching, Verbose, TEXT("Loading Chunk: %s, %d"), *InKey.SoundWaveName.ToString(), InKey.ChunkIndex);
+
+			IBulkDataIORequest* LocalReadRequest = Chunk->BulkData.CreateStreamingRequest(0, ChunkDataSize, AsyncIOPriority | AIOP_FLAG_DONTCACHE, &AsyncFileCallBack, CacheElement->ChunkData);
+			if (!LocalReadRequest)
+			{
+				UE_LOG(LogAudioStreamCaching, Error, TEXT("Chunk load in audio LRU cache failed."));
+				OnLoadCompleted(EAudioChunkLoadResult::ChunkOutOfBounds);
+				NumberOfLoadsInFlight.Decrement();
+			}
+			else if (FPlatformAtomics::InterlockedCompareExchangePointer((void* volatile*)&CacheElement->ReadRequest, LocalReadRequest, nullptr) == (void*)0x1)
+			{
+				// The request is completed before we can store it. Just delete it
+				TGraphTask<FClearAudioChunkCacheReadRequestTask>::CreateTask().ConstructAndDispatchWhenReady(LocalReadRequest);
+			}
 		}
 	}
 }
@@ -1975,6 +2007,7 @@ TPair<int, int> FAudioChunkCache::DebugDisplayLegacy(UWorld* World, FViewport* V
 		bool bLoadingBehaviorExternallyOverriden = false;
 		bool bWasCacheMiss = false;
 		bool bIsStaleChunk = false;
+		bool bWasLoadedInlined = false;
 
 #if DEBUG_STREAM_CACHE
 		NumTotalChunks = CurrentElement->DebugInfo.NumTotalChunks;
@@ -1984,6 +2017,7 @@ TPair<int, int> FAudioChunkCache::DebugDisplayLegacy(UWorld* World, FViewport* V
 		LoadingBehavior = CurrentElement->DebugInfo.LoadingBehavior;
 		bLoadingBehaviorExternallyOverriden = CurrentElement->DebugInfo.bLoadingBehaviorExternallyOverriden;
 		bWasCacheMiss = CurrentElement->DebugInfo.bWasCacheMiss;
+		bWasLoadedInlined = CurrentElement->DebugInfo.bWasLoadedFromInlineChunk;
 #endif
 
 #if WITH_EDITOR
@@ -1993,7 +2027,7 @@ TPair<int, int> FAudioChunkCache::DebugDisplayLegacy(UWorld* World, FViewport* V
 
 		const bool bWasTrimmed = CurrentElement->ChunkDataSize == 0;
 
-		FString ElementInfo = *FString::Printf(TEXT("%4i. Size: %6.2f KB   Chunk: %d of %d   Request Count: %d    Average Index: %6.2f  Number of Handles Retaining Chunk: %d     Chunk Load Time(in ms): %6.4fms      Loading Behavior: %s%s      Name: %s Notes: %s %s"),
+		FString ElementInfo = *FString::Printf(TEXT("%4i. Size: %6.2f KB   Chunk: %d of %d   Request Count: %d    Average Index: %6.2f  Number of Handles Retaining Chunk: %d     Chunk Load Time(in ms): %6.4fms      Loading Behavior: %s%s      Name: %s Notes: %s %s %s"),
 			Index,
 			CurrentElement->ChunkDataSize / 1024.0f,
 			CurrentElement->Key.ChunkIndex,
@@ -2006,7 +2040,8 @@ TPair<int, int> FAudioChunkCache::DebugDisplayLegacy(UWorld* World, FViewport* V
 			bLoadingBehaviorExternallyOverriden ? TEXT("*") : TEXT(""),
 			bWasTrimmed ? TEXT("TRIMMED CHUNK") : *CurrentElement->Key.SoundWaveName.ToString(),
 			bWasCacheMiss ? TEXT("(Cache Miss!)") : TEXT(""),
-			bIsStaleChunk ? TEXT("(Stale Chunk)") : TEXT("")
+			bIsStaleChunk ? TEXT("(Stale Chunk)") : TEXT(""),
+			bWasLoadedInlined ? TEXT("(Inlined)") : TEXT("")
 		);
 
 		// Since there's a lot of info here,
@@ -2110,6 +2145,7 @@ FString FAudioChunkCache::DebugPrint()
 		bool bLoadingBehaviorExternallyOverriden = false;
 		bool bWasCacheMiss = false;
 		bool bIsStaleChunk = false;
+		bool bWasLoadedInlined = false;
 
 #if DEBUG_STREAM_CACHE
 		NumTotalChunks = CurrentElement->DebugInfo.NumTotalChunks;
@@ -2119,6 +2155,7 @@ FString FAudioChunkCache::DebugPrint()
 		LoadingBehavior = CurrentElement->DebugInfo.LoadingBehavior;
 		bLoadingBehaviorExternallyOverriden = CurrentElement->DebugInfo.bLoadingBehaviorExternallyOverriden;
 		bWasCacheMiss = CurrentElement->DebugInfo.bWasCacheMiss;
+		bWasLoadedInlined = CurrentElement->DebugInfo.bWasLoadedFromInlineChunk;
 #endif
 
 #if WITH_EDITOR
@@ -2128,7 +2165,7 @@ FString FAudioChunkCache::DebugPrint()
 
 		const bool bWasTrimmed = CurrentElement->ChunkDataSize == 0;
 
-		FString ElementInfo = *FString::Printf(TEXT("%4i.\t, %6.2f\t, %d of %d\t, %d\t, %6.2f\t, %d\t,  %6.4f\t, %s\t, %s%s, %s %s %s"),
+		FString ElementInfo = *FString::Printf(TEXT("%4i.\t, %6.2f\t, %d of %d\t, %d\t, %6.2f\t, %d\t,  %6.4f\t, %s\t, %s%s, %s %s %s %s"),
 			Index,
 			CurrentElement->ChunkDataSize / 1024.0f,
 			CurrentElement->Key.ChunkIndex,
@@ -2142,7 +2179,8 @@ FString FAudioChunkCache::DebugPrint()
 			bLoadingBehaviorExternallyOverriden ? TEXT("*") : TEXT(""),
 			bWasCacheMiss ? TEXT("(Cache Miss!)") : TEXT(""),
 			bIsStaleChunk ? TEXT("(Stale Chunk)") : TEXT(""),
-			CurrentElement->IsLoadInProgress() ? TEXT("(Loading In Progress)") : TEXT("")
+			CurrentElement->IsLoadInProgress() ? TEXT("(Loading In Progress)") : TEXT(""),
+			bWasLoadedInlined ? TEXT("(Inlined)") : TEXT("(Disk)")
 		);
 
 		if (!bWasTrimmed)
diff --git a/Engine/Source/Runtime/Engine/Private/SoundWave.cpp b/Engine/Source/Runtime/Engine/Private/SoundWave.cpp
--- a/Engine/Source/Runtime/Engine/Private/SoundWave.cpp
+++ b/Engine/Source/Runtime/Engine/Private/SoundWave.cpp
@@ -30,16 +30,14 @@
 #include "UObject/UE5MainStreamObjectVersion.h"
 #include "Algo/BinarySearch.h"
 #include "Templates/UnrealTemplate.h"
+#include "UObject/ObjectSaveContext.h"
 
-#if WITH_EDITOR
-#endif // WITH_EDITOR
-
-static int32 SoundWaveDefaultLoadingBehaviorCVar = 0;
+static int32 SoundWaveDefaultLoadingBehaviorCVar = static_cast<int32>(ESoundWaveLoadingBehavior::LoadOnDemand);
 FAutoConsoleVariableRef CVarSoundWaveDefaultLoadingBehavior(
 	TEXT("au.streamcache.SoundWaveDefaultLoadingBehavior"),
 	SoundWaveDefaultLoadingBehaviorCVar,
 	TEXT("This can be set to define the default behavior when a USoundWave is loaded.\n")
-	TEXT("0: Default (load on demand), 1: Retain audio data on load, 2: prime audio data on load, 3: load on demand (No audio data is loaded until a USoundWave is played or primed)."),
+	TEXT("1: Retain audio data on load, 2: prime audio data on load, 3: load on demand (No audio data is loaded until a USoundWave is played or primed)."),
 	ECVF_Default);
 
 static int32 ForceNonStreamingInEditorCVar = 0;
@@ -82,6 +80,28 @@ FAutoConsoleVariableRef CvarAllowReverbForMultichannelSources(
 	TEXT("0: Disable, >0: Enable"),
 	ECVF_Default);
 
+namespace SoundWave_Private
+{
+	static ESoundWaveLoadingBehavior GetDefaultLoadingBehaviorCVar()
+	{
+		// query the default loading behavior CVar
+		ensureMsgf(SoundWaveDefaultLoadingBehaviorCVar >= static_cast<int32>(ESoundWaveLoadingBehavior::RetainOnLoad) &&
+			SoundWaveDefaultLoadingBehaviorCVar <= static_cast<int32>(ESoundWaveLoadingBehavior::LoadOnDemand),
+			TEXT("Invalid default loading behavior CVar value=%d:%s. Use value 1 (retain on load), 2 (prime on load) or 3 (load on demand)."),
+			SoundWaveDefaultLoadingBehaviorCVar, EnumToString(static_cast<ESoundWaveLoadingBehavior>(SoundWaveDefaultLoadingBehaviorCVar)));
+
+		// Clamp the value to make sure its in range.
+		const ESoundWaveLoadingBehavior DefaultLoadingBehavior = static_cast<ESoundWaveLoadingBehavior>(
+			FMath::Clamp<int32>(
+				SoundWaveDefaultLoadingBehaviorCVar,
+				static_cast<int32>(ESoundWaveLoadingBehavior::RetainOnLoad),
+				static_cast<int32>(ESoundWaveLoadingBehavior::LoadOnDemand)
+			));
+
+		return DefaultLoadingBehavior;
+	}
+}
+
 
 #if !UE_BUILD_SHIPPING
 static void DumpBakedAnalysisData(const TArray<FString>& Args)
@@ -146,8 +166,8 @@ void FSoundWaveData::InitializeDataFromSoundWave(USoundWave& InWave)
 
 	// cache the runtime format for the wave
 	// note if this fails, it will ensure, and keep the "FSoundWaveProxy_InvalidFormat" as its value.
-	FName FoundFormat = FindRuntimeFormat(InWave);
-	if (!FoundFormat.IsNone())
+	const FName FoundFormat = FindRuntimeFormat(InWave);
+	if (ensure(!FoundFormat.IsNone()))
 	{
 		RuntimeFormat = FoundFormat;
 	}
@@ -656,6 +676,11 @@ ITargetPlatform* USoundWave::GetRunningPlatform()
 	}
 }
 
+ENGINE_API ESoundWaveLoadingBehavior USoundWave::GetDefaultLoadingBehavior() 
+{
+	return SoundWave_Private::GetDefaultLoadingBehaviorCVar();
+}
+
 /*-----------------------------------------------------------------------------
 	FStreamedAudioChunk
 -----------------------------------------------------------------------------*/
@@ -663,36 +688,32 @@ ITargetPlatform* USoundWave::GetRunningPlatform()
 void FStreamedAudioChunk::Serialize(FArchive& Ar, UObject* Owner, int32 ChunkIndex)
 {
 	DECLARE_SCOPE_CYCLE_COUNTER( TEXT("FStreamedAudioChunk::Serialize"), STAT_StreamedAudioChunk_Serialize, STATGROUP_LoadTime );
-	bool bShouldInlineAudioChunk = false;
-
-	const ITargetPlatform* CookingTarget = Ar.CookingTarget();
-	if (CookingTarget != nullptr)
-	{
-		const FPlatformAudioCookOverrides* Overrides = FPlatformCompressionUtilities::GetCookOverrides(*CookingTarget->IniPlatformName());
-		check(Overrides);
-		bShouldInlineAudioChunk = Overrides->bInlineStreamedAudioChunks;
-	}
 
 	enum 
 	{
 		IsCooked			 = 1 << 0,
-		HasSeekOffset		 = 1 << 1
+		HasSeekOffset		 = 1 << 1,
+		IsInlined			 = 1 << 2,
 	};
 	
 	// Bit-pack flags into a single uint32, instead of a bool.
 	uint32 Flags = 0;
 	if (Ar.IsSaving())
 	{
+#if WITH_EDITORONLY_DATA
+		// Chunk 0 is always inlined, others are optionally inlined.
+		const bool bShouldInline = ChunkIndex == 0 || bInlineChunk;
+		Flags |= bShouldInline ? IsInlined : 0;
+#endif //WITH_EDITORONLY_DATA
 		Flags |= Ar.IsCooking() ? IsCooked : 0;
 		Flags |= SeekOffsetInAudioFrames != INDEX_NONE ? HasSeekOffset : 0;
 	}
 
 	Ar << Flags;
 
-	// ChunkIndex 0 is always inline payload, all other chunks are streamed.
 	if (Ar.IsSaving())
 	{
-		if (ChunkIndex == 0 || (ChunkIndex == 1 && bShouldInlineAudioChunk))
+		if (Flags & IsInlined)
 		{
 			BulkData.SetBulkDataFlags(BULKDATA_ForceInlinePayload);
 		}
@@ -738,9 +759,10 @@ void FStreamedAudioChunk::Serialize(FArchive& Ar, UObject* Owner, int32 ChunkInd
 		Ar << DerivedDataKey;
 	}
 
-	if (Ar.IsLoading() && (Flags & IsCooked))
+	if (Ar.IsLoading())
 	{
-		bLoadedFromCookedPackage = true;
+		bLoadedFromCookedPackage = !!(Flags & IsCooked);
+		bInlineChunk = !!(Flags & IsInlined);
 	}
 #endif // #if WITH_EDITORONLY_DATA
 }
@@ -778,6 +800,19 @@ bool FStreamedAudioChunk::GetCopy(void** OutChunkData)
 	return false;
 }
 
+FBulkDataBuffer<uint8> FStreamedAudioChunk::MoveOutAsBuffer() 
+{
+	if (ensure(AudioDataSize > 0))
+	{
+		if (ensure(AudioDataSize <= BulkData.GetElementCount()))
+		{
+			const int64 RequestedSize = FMath::Min(AudioDataSize, BulkData.GetElementCount());
+			return BulkData.GetCopyAsBuffer(RequestedSize, /* Discard Internal Copy */ true);
+		}
+	}
+	return {};
+}
+
 #if WITH_EDITORONLY_DATA
 uint32 FStreamedAudioChunk::StoreInDerivedDataCache(const FString& InDerivedDataKey, const FStringView& SoundWaveName)
 {
@@ -989,9 +1024,44 @@ void USoundWave::Serialize( FArchive& Ar )
 	DECLARE_SCOPE_CYCLE_COUNTER( TEXT("USoundWave::Serialize"), STAT_SoundWave_Serialize, STATGROUP_LoadTime );
 
 	Super::Serialize( Ar );
+	
+	// Bit-pack.
+	enum
+	{
+		CookedFlag					= 1 << 0,
+		HasOwnerLoadingBehaviorFlag	= 1 << 1,
+		
+		LoadingBehaviorShift		= 2,
+		LoadingBehaviorMask			= 0b00000111,
+	};
+
+	// Pack flags into uint32
+	uint32 Flags = 0;
+	if (Ar.IsSaving())
+	{
+		if (Ar.IsCooking())
+		{
+			Flags |= CookedFlag;
+
+#if WITH_EDITOR
+			// If we are going to save the owners loading behavior
+			const ISoundWaveLoadingBehaviorUtil::FClassData OwnerBehavior = GetOwnerLoadingBehavior(Ar.CookingTarget());
+			if (OwnerBehavior.LoadingBehavior != ESoundWaveLoadingBehavior::Uninitialized)
+			{
+				UE_LOG(LogAudio, Display, TEXT("Saving Owners Loading Behavior: %s, %s"), 
+				       EnumToString(OwnerBehavior.LoadingBehavior), *GetFullName());
+				
+				Flags |= HasOwnerLoadingBehaviorFlag;
+				Flags |= (uint32)OwnerBehavior.LoadingBehavior << LoadingBehaviorShift;
+			}
+#endif //WITH_EDITOR
+
+		}
+	}
+
+	Ar << Flags;
 
-	bool bCooked = Ar.IsCooking();
-	Ar << bCooked;
+	const bool bCooked = Flags & CookedFlag;
 
 	if (FPlatformProperties::RequiresCookedData() && !bCooked && Ar.IsLoading())
 	{
@@ -1159,6 +1229,13 @@ void USoundWave::Serialize( FArchive& Ar )
 
 	Ar << CompressedDataGuid;
 
+	// Use our owners loading behavior if we've been cooked to do so.
+	if (Ar.IsLoading() && (Flags & CookedFlag) && (Flags & HasOwnerLoadingBehaviorFlag) && LoadingBehavior != ESoundWaveLoadingBehavior::ForceInline)
+	{
+		const ESoundWaveLoadingBehavior OwnerBehavior = (ESoundWaveLoadingBehavior)((Flags >> LoadingBehaviorShift) & LoadingBehaviorMask);
+		LoadingBehavior = OwnerBehavior;
+	}
+
 #if WITH_EDITORONLY_DATA	
 	// If we're converting to FEditorBulkData, we need to first serialize the CompressedDataGuid, as we use it as our key.
 	if (bBulkDataConverted)
@@ -1582,7 +1659,7 @@ FName USoundWave::GetPlatformSpecificFormat(FName Format, const FPlatformAudioCo
 	return PlatformSpecificFormat;
 }
 
-void USoundWave::BeginGetCompressedData(FName Format, const FPlatformAudioCookOverrides* CompressionOverrides)
+void USoundWave::BeginGetCompressedData(FName Format, const FPlatformAudioCookOverrides* CompressionOverrides, const ITargetPlatform* InTargetPlatform)
 {
 	TRACE_CPUPROFILER_EVENT_SCOPE(USoundWave::BeginGetCompressedData);
 	check(SoundWaveDataPtr);
@@ -1604,7 +1681,7 @@ void USoundWave::BeginGetCompressedData(FName Format, const FPlatformAudioCookOv
 		{
 			COOK_STAT(auto Timer = SoundWaveCookStats::UsageStats.TimeSyncWork());
 			COOK_STAT(Timer.TrackCyclesOnly());
-			FDerivedAudioDataCompressor* DeriveAudioData = new FDerivedAudioDataCompressor(this, Format, PlatformSpecificFormat, CompressionOverrides);
+			FDerivedAudioDataCompressor* DeriveAudioData = new FDerivedAudioDataCompressor(this, Format, PlatformSpecificFormat, CompressionOverrides, InTargetPlatform);
 			uint32 GetHandle = GetDerivedDataCacheRef().GetAsynchronous(DeriveAudioData);
 			AsyncLoadingDataFormats.Add(PlatformSpecificFormat, GetHandle);
 		}
@@ -1642,7 +1719,7 @@ bool USoundWave::IsLoadedFromCookedData() const
 
 #endif
 
-FByteBulkData* USoundWave::GetCompressedData(FName Format, const FPlatformAudioCookOverrides* CompressionOverrides)
+FByteBulkData* USoundWave::GetCompressedData(FName Format, const FPlatformAudioCookOverrides* CompressionOverrides, const ITargetPlatform* InTargetPlatform )
 {
 	check(SoundWaveDataPtr);
 
@@ -1681,7 +1758,7 @@ FByteBulkData* USoundWave::GetCompressedData(FName Format, const FPlatformAudioC
 			else
 #endif
 			{
-				FDerivedAudioDataCompressor* DeriveAudioData = new FDerivedAudioDataCompressor(this, Format, PlatformSpecificFormat, CompressionOverrides);
+				FDerivedAudioDataCompressor* DeriveAudioData = new FDerivedAudioDataCompressor(this, Format, PlatformSpecificFormat, CompressionOverrides, InTargetPlatform);
 				bGetSuccessful = GetDerivedDataCacheRef().GetSynchronous(DeriveAudioData, OutData, &bDataWasBuilt);
 			}
 
@@ -1827,6 +1904,12 @@ void USoundWave::PostLoad()
 		}
 #endif // #if WITH_EDITORONLY_DATA
 
+		// update shared state
+		if (!HasAnyFlags(RF_ClassDefaultObject) && FApp::CanEverRenderAudio())
+		{
+			SoundWaveDataPtr->InitializeDataFromSoundWave(*this);
+		}
+		
 		ESoundWaveLoadingBehavior ActualLoadingBehavior = GetLoadingBehavior();
 
 		if (!Proxy.IsValid() && ActualLoadingBehavior != ESoundWaveLoadingBehavior::ForceInline)
@@ -1889,17 +1972,12 @@ void USoundWave::PostLoad()
 		{
 			if (Platform->AllowAudioVisualData())
 			{
-				BeginGetCompressedData(Platform->GetWaveFormat(this), FPlatformCompressionUtilities::GetCookOverrides(*Platform->IniPlatformName()));
+				BeginGetCompressedData(Platform->GetWaveFormat(this), FPlatformCompressionUtilities::GetCookOverrides(*Platform->IniPlatformName()), Platform);
 			}
 		}
 	}
 
-	// update shared state
-	if (!HasAnyFlags(RF_ClassDefaultObject) && FApp::CanEverRenderAudio())
-	{
-		SoundWaveDataPtr->InitializeDataFromSoundWave(*this);
-	}
-
+	
 	// We don't precache default objects and we don't precache in the Editor as the latter will
 	// most likely cause us to run out of memory.
 	if (!GIsEditor && !IsTemplate(RF_ClassDefaultObject) && GEngine)
@@ -2168,6 +2246,42 @@ void USoundWave::InvalidateSoundWaveIfNeccessary()
 	}
 }
 
+void USoundWave::PreSave(FObjectPreSaveContext InSaveContext)
+{
+	if (InSaveContext.IsCooking())
+	{
+		// Populate our cache ahead of Serialize, where we can't do any AssetRegistry queries safely.
+		GetOwnerLoadingBehavior(InSaveContext.GetTargetPlatform());
+	}
+	Super::PreSave(InSaveContext);
+}
+
+ISoundWaveLoadingBehaviorUtil::FClassData USoundWave::GetOwnerLoadingBehavior(const ITargetPlatform* InTargetPlatform) const
+{
+	// Use {} or Name_None in the case of a missing platform.
+	const FName PlatformName = InTargetPlatform ? *InTargetPlatform->PlatformName() : FName();
+
+	// Lock for entirity of call.
+	FScopeLock Lock(&OwnerLoadingBehaviorCacheCS);
+	
+	// We have it in our cache?
+	if (const ISoundWaveLoadingBehaviorUtil::FClassData* Cached = OwnerLoadingBehaviorCache.Find(PlatformName))
+	{
+		return *Cached;
+	}
+
+	// ... otherwise, do lookup
+	if (const ISoundWaveLoadingBehaviorUtil* Util = ISoundWaveLoadingBehaviorUtil::Get())
+	{
+		const ISoundWaveLoadingBehaviorUtil::FClassData Data = Util->FindOwningLoadingBehavior(this, InTargetPlatform);
+		OwnerLoadingBehaviorCache.Add(PlatformName, Data);
+		return Data;
+	}
+	
+	// Not set.
+	return {};
+}
+
 float USoundWave::GetSampleRateForTargetPlatform(const ITargetPlatform* TargetPlatform)
 {
 	const FPlatformAudioCookOverrides* Overrides = FPlatformCompressionUtilities::GetCookOverrides(*TargetPlatform->IniPlatformName());
@@ -2181,6 +2295,58 @@ float USoundWave::GetSampleRateForTargetPlatform(const ITargetPlatform* TargetPl
 	}
 }
 
+#if WITH_EDITOR
+
+float USoundWave::GetSizeOfFirstAudioChunkInSeconds(const ITargetPlatform* InPlatform) const
+{
+#if WITH_EDITORONLY_DATA
+	
+	ensure(InPlatform);
+
+	// We need to walk up the hierarchy of sound classes to first determine our FirstAudioChunkSize.
+	if (SoundWaveDataPtr->LoadingBehavior == ESoundWaveLoadingBehavior::Uninitialized)
+	{
+		CacheInheritedLoadingBehavior();
+	}
+
+	// Use owner loading behavior if we have it.
+	ISoundWaveLoadingBehaviorUtil::FClassData Behavior = GetOwnerLoadingBehavior(InPlatform);
+	if (Behavior.LoadingBehavior == ESoundWaveLoadingBehavior::Uninitialized)
+	{
+		// Otherwise use what we've been tagged
+		Behavior.LoadingBehavior = SoundWaveDataPtr->LoadingBehavior;
+		Behavior.LengthOfFirstChunkInSeconds = SoundWaveDataPtr->SizeOfFirstAudioChunkInSeconds;
+	}
+
+	// We only allow setting the first chunk size if we're planning on retaining it.
+	// This could be relaxed later once we have a better format that wouldn't cause a huge delta
+	// on tuning the size (all assets etc). By using just retained, we expect a smaller set.
+	if (Behavior.LoadingBehavior == ESoundWaveLoadingBehavior::RetainOnLoad)
+	{
+		// See if we have an override for this platform on the SoundClass hierarchy.
+		const FName PlatformName = InPlatform ? InPlatform->GetPlatformInfo().IniPlatformName : FName();
+		const float PlatformSize = Behavior.LengthOfFirstChunkInSeconds.GetValueForPlatform(PlatformName);
+
+		if (PlatformSize > 0)
+		{
+			return PlatformSize;
+		}
+
+		// If that's not set to anything fallback to setting on cook overrides, if there's anything set.
+		if (const FPlatformAudioCookOverrides* CompressionOverrides = FPlatformCompressionUtilities::GetCookOverrides(*PlatformName.ToString()))
+		{
+			return CompressionOverrides->LengthOfFirstAudioChunkInSecs;
+		}
+	}
+	
+#endif //WITH_EDITORONLY_DATA
+
+	// Nothing.
+	return 0.f;
+}
+
+#endif //WITH_EDTIOR
+
 void USoundWave::LogBakedData()
 {
 	const FString AnalysisPathName = *(FPaths::ProjectLogDir() + TEXT("BakedAudioAnalysisData/"));
@@ -2594,8 +2760,9 @@ void USoundWave::PostEditChangeProperty(FPropertyChangedEvent& PropertyChangedEv
 	static const FName StreamingFName = GET_MEMBER_NAME_CHECKED(USoundWave, bStreaming);
 	static const FName SoundAssetCompressionTypeFName = GET_MEMBER_NAME_CHECKED(USoundWave, SoundAssetCompressionType);
 	static const FName LoadingBehaviorFName = GET_MEMBER_NAME_CHECKED(USoundWave, LoadingBehavior);
-	static const FName InitialChunkSizeFName = GET_MEMBER_NAME_CHECKED(USoundWave, InitialChunkSize);
+	static const FName InitialChunkSizeFName = GET_MEMBER_NAME_CHECKED(USoundWave, InitialChunkSize_DEPRECATED);
 	static const FName TransformationsFName = GET_MEMBER_NAME_CHECKED(USoundWave, Transformations);
+	static const FName InlinedAudioInSecondsFName = GET_MEMBER_NAME_CHECKED(USoundWave, SizeOfFirstAudioChunkInSeconds);
 
 	// force proxy state to be up to date
 	SoundWaveDataPtr->InitializeDataFromSoundWave(*this);
@@ -2614,7 +2781,7 @@ void USoundWave::PostEditChangeProperty(FPropertyChangedEvent& PropertyChangedEv
 		if (PropertyChangedEvent.ChangeType != EPropertyChangeType::Interactive)
 		{
 			// Regenerate on save any compressed sound formats or if analysis needs to be re-done
-			if (Name == LoadingBehaviorFName)
+			if (Name == LoadingBehaviorFName || Name == InlinedAudioInSecondsFName)
 			{
 				// Update and cache new loading behavior if it has changed. This
 				// must be called before a new FSoundWaveProxy is created. 
@@ -2640,6 +2807,7 @@ void USoundWave::PostEditChangeProperty(FPropertyChangedEvent& PropertyChangedEv
 				|| Name == SampleRateFName
 				|| Name == StreamingFName
 				|| Name == LoadingBehaviorFName
+				|| Name == InlinedAudioInSecondsFName
 				|| Name == InitialChunkSizeFName
 				|| Name == TransformationsFName)
 			{
@@ -4091,6 +4259,9 @@ void USoundWave::CacheInheritedLoadingBehavior() const
 		if (SoundWaveDataPtr->LoadingBehavior == ESoundWaveLoadingBehavior::Uninitialized)
 		{
 			SoundWaveDataPtr->LoadingBehavior = LoadingBehavior;
+#if WITH_EDITORONLY_DATA
+			SoundWaveDataPtr->SizeOfFirstAudioChunkInSeconds = SizeOfFirstAudioChunkInSeconds;
+#endif //WITH_EDITORONLY_DATA
 		}
 	}
  	else if (SoundWaveDataPtr->bLoadingBehaviorOverridden)
@@ -4104,20 +4275,22 @@ void USoundWave::CacheInheritedLoadingBehavior() const
 
 		USoundClass* CurrentSoundClass = GetSoundClass();
 		ESoundWaveLoadingBehavior SoundClassLoadingBehavior = ESoundWaveLoadingBehavior::Inherited;
+		FPerPlatformFloat SoundClassSizeOfFirstAudioChunkInSeconds = 0;
 
 		// Recurse through this sound class's parents until we find an override.
 		while (SoundClassLoadingBehavior == ESoundWaveLoadingBehavior::Inherited && CurrentSoundClass != nullptr)
 		{
 			SoundClassLoadingBehavior = CurrentSoundClass->Properties.LoadingBehavior;
+#if WITH_EDITORONLY_DATA
+			SoundClassSizeOfFirstAudioChunkInSeconds = CurrentSoundClass->Properties.SizeOfFirstAudioChunkInSeconds;
+#endif //WITH_EDITORONLY_DATA			
 			CurrentSoundClass = CurrentSoundClass->ParentClass;
 		}
 
 		// If we could not find an override in the sound class hierarchy, use the loading behavior defined by our cvar.
 		if (SoundClassLoadingBehavior == ESoundWaveLoadingBehavior::Inherited)
 		{
-			// query the default loading behavior CVar
-			ensureAlwaysMsgf(SoundWaveDefaultLoadingBehaviorCVar >= 0 && SoundWaveDefaultLoadingBehaviorCVar < 4, TEXT("Invalid default loading behavior CVar value. Use value 0, 1, 2 or 3."));
-			ESoundWaveLoadingBehavior DefaultLoadingBehavior = (ESoundWaveLoadingBehavior)FMath::Clamp<int32>(SoundWaveDefaultLoadingBehaviorCVar, 0, (int32)ESoundWaveLoadingBehavior::LoadOnDemand);
+			const ESoundWaveLoadingBehavior DefaultLoadingBehavior = SoundWave_Private::GetDefaultLoadingBehaviorCVar();
 
 			// override this loading behavior w/ our default
 			SoundClassLoadingBehavior = DefaultLoadingBehavior;
@@ -4125,6 +4298,9 @@ void USoundWave::CacheInheritedLoadingBehavior() const
 		}
 
 		SoundWaveDataPtr->LoadingBehavior = SoundClassLoadingBehavior;
+#if WITH_EDITORONLY_DATA		
+		SoundWaveDataPtr->SizeOfFirstAudioChunkInSeconds = SoundClassSizeOfFirstAudioChunkInSeconds;
+#endif //WITH_EDITORONLY_DATA		
 	}
 }
 
@@ -4167,9 +4343,9 @@ ESoundWaveLoadingBehavior USoundWave::GetLoadingBehavior(bool bCheckSoundClasses
 		}
 		else
 		{
-			// Otherwise, use the loading behavior defined by our cvar.
-			ensureAlwaysMsgf(SoundWaveDefaultLoadingBehaviorCVar >= 0 && SoundWaveDefaultLoadingBehaviorCVar < 4, TEXT("Invalid default loading behavior CVar value. Use value 0, 1, 2 or 3."));
-			return (ESoundWaveLoadingBehavior)FMath::Clamp<int32>(SoundWaveDefaultLoadingBehaviorCVar, 0, (int32)ESoundWaveLoadingBehavior::LoadOnDemand);
+			// Use the cvar.
+			const ESoundWaveLoadingBehavior DefaultSoundWaveBehavior = SoundWave_Private::GetDefaultLoadingBehaviorCVar();
+			return DefaultSoundWaveBehavior;
 		}
 	}
 	else if (SoundWaveDataPtr->LoadingBehavior == ESoundWaveLoadingBehavior::Uninitialized)
@@ -4177,6 +4353,8 @@ ESoundWaveLoadingBehavior USoundWave::GetLoadingBehavior(bool bCheckSoundClasses
 		CacheInheritedLoadingBehavior();
 	}
 
+	// We should have resolved by now.
+	ensure(SoundWaveDataPtr->LoadingBehavior != ESoundWaveLoadingBehavior::Inherited);
 	return SoundWaveDataPtr->LoadingBehavior;
 }
 
@@ -4463,3 +4641,4 @@ bool USoundWave::HasError() const
 	check(SoundWaveDataPtr);
 	return SoundWaveDataPtr->HasError();
 }
+
diff --git a/Engine/Source/Runtime/Engine/Private/SoundWaveProcedural.cpp b/Engine/Source/Runtime/Engine/Private/SoundWaveProcedural.cpp
--- a/Engine/Source/Runtime/Engine/Private/SoundWaveProcedural.cpp
+++ b/Engine/Source/Runtime/Engine/Private/SoundWaveProcedural.cpp
@@ -157,12 +157,12 @@ bool USoundWaveProcedural::HasCompressedData(FName Format, ITargetPlatform* Targ
 	return false;
 }
 
-void USoundWaveProcedural::BeginGetCompressedData(FName Format, const FPlatformAudioCookOverrides* CompressionOverrides)
+void USoundWaveProcedural::BeginGetCompressedData(FName Format, const FPlatformAudioCookOverrides* CompressionOverrides, const ITargetPlatform* InTargetPlatform)
 {
 	// SoundWaveProcedural does not have compressed data and should generally not be asked about it
 }
 
-FByteBulkData* USoundWaveProcedural::GetCompressedData(FName Format, const FPlatformAudioCookOverrides* CompressionOverrides)
+FByteBulkData* USoundWaveProcedural::GetCompressedData(FName Format, const FPlatformAudioCookOverrides* CompressionOverrides, const ITargetPlatform* InTargetPlatform )
 {
 	// SoundWaveProcedural does not have compressed data and should generally not be asked about it
 	return nullptr;
diff --git a/Engine/Source/Runtime/Engine/Public/AudioDecompress.h b/Engine/Source/Runtime/Engine/Public/AudioDecompress.h
--- a/Engine/Source/Runtime/Engine/Public/AudioDecompress.h
+++ b/Engine/Source/Runtime/Engine/Public/AudioDecompress.h
@@ -14,9 +14,10 @@
 #include "HAL/LowLevelMemTracker.h"
 #include "ContentStreaming.h"
 
-// 186ms of 44.1KHz data
-// 372ms of 22KHz data
-#define MONO_PCM_BUFFER_SAMPLES		8192
+// 100ms of 48KHz data
+// 108ms of 44.1KHz data
+// 218ms of 22KHz data
+#define MONO_PCM_BUFFER_SAMPLES		(4800)
 #define MONO_PCM_BUFFER_SIZE		( MONO_PCM_BUFFER_SAMPLES * sizeof( int16 ) )
 
 struct FSoundQualityInfo;
@@ -319,7 +320,9 @@ protected:
 	/** The current chunk index in the streamed chunks. */
 	int32 CurrentChunkIndex;
 	/** Whether or not to print the chunk fail message. */
-	bool bPrintChunkFailMessage;
+	int32 PrintChunkFailMessageCount = 0;
+	/** A counter of when we started the last request (for gauging latency) */
+	uint64 StartTimeInCycles = 0;
 	/** Number of bytes of padding used, overridden in some implementations. Defaults to 0. */
 	uint32 SrcBufferPadding;
 	/** Chunk Handle to ensure that this chunk of streamed audio is not deleted while we are using it. */
diff --git a/Engine/Source/Runtime/Engine/Public/AudioDerivedData.h b/Engine/Source/Runtime/Engine/Public/AudioDerivedData.h
--- a/Engine/Source/Runtime/Engine/Public/AudioDerivedData.h
+++ b/Engine/Source/Runtime/Engine/Public/AudioDerivedData.h
@@ -16,7 +16,7 @@ private:
 	TUniquePtr<FAudioCookInputs> CookInputs;
 
 public:
-	FDerivedAudioDataCompressor(USoundWave* InSoundNode, FName InBaseFormat, FName InHashedFormat, const FPlatformAudioCookOverrides* InCompressionOverrides);
+	FDerivedAudioDataCompressor(USoundWave* InSoundNode, FName InBaseFormat, FName InHashedFormat, const FPlatformAudioCookOverrides* InCompressionOverrides, const ITargetPlatform* InTargetPlatform=nullptr);
 
 	virtual const TCHAR* GetPluginName() const override
 	{
