PR #10775: Added UWorld::GetGameStateChecked()

diff --git a/Engine/Source/Runtime/Engine/Classes/Engine/World.h b/Engine/Source/Runtime/Engine/Classes/Engine/World.h
--- a/Engine/Source/Runtime/Engine/Classes/Engine/World.h
+++ b/Engine/Source/Runtime/Engine/Classes/Engine/World.h
@@ -3557,6 +3557,13 @@ public:
 		return Cast<T>(GameState);
 	}
 
+	/** Returns the current GameState instance cast to the template type, asserting that it is of the correct type. */
+	template<class T>
+	T* GetGameStateChecked() const
+	{
+		return CastChecked<T>(GameState);
+	}
+
 	/** Returns the current GameState instance. */
 	AGameStateBase* GetGameState() const { return GameState; }
 
