Fix CIS nounity error in SoundWaveLoadingBehavior

diff --git a/Engine/Source/Runtime/Engine/Classes/Sound/SoundWaveLoadingBehavior.cpp b/Engine/Source/Runtime/Engine/Classes/Sound/SoundWaveLoadingBehavior.cpp
--- a/Engine/Source/Runtime/Engine/Classes/Sound/SoundWaveLoadingBehavior.cpp
+++ b/Engine/Source/Runtime/Engine/Classes/Sound/SoundWaveLoadingBehavior.cpp
@@ -1,13 +1,18 @@
 // Copyright Epic Games, Inc. All Rights Reserved.
 #include "SoundWaveLoadingBehavior.h"
 
+#include "Audio.h"
+#include "AssetRegistry/AssetData.h"
 #include "AssetRegistry/AssetRegistryModule.h"
 #include "AssetRegistry/IAssetRegistry.h"
 #include "GenericPlatform/GenericPlatformMisc.h"
 #include "Interfaces/ITargetPlatform.h"
+#include "Misc/CommandLine.h"
 #include "SoundClass.h"
 #include "SoundCue.h"
 #include "SoundWave.h"
+#include "UObject/LinkerLoad.h"
+
 
 const TCHAR* EnumToString(ESoundWaveLoadingBehavior InCurrentState)
 {
diff --git a/Engine/Source/Runtime/Engine/Classes/Sound/SoundWaveLoadingBehavior.h b/Engine/Source/Runtime/Engine/Classes/Sound/SoundWaveLoadingBehavior.h
--- a/Engine/Source/Runtime/Engine/Classes/Sound/SoundWaveLoadingBehavior.h
+++ b/Engine/Source/Runtime/Engine/Classes/Sound/SoundWaveLoadingBehavior.h
@@ -40,7 +40,9 @@ ENGINE_API const TCHAR* EnumToString(ESoundWaveLoadingBehavior InCurrentState);
 
 #if WITH_EDITOR
 
+class USoundWave;
 class USoundClass;
+class ITargetPlatform;
 
 class ISoundWaveLoadingBehaviorUtil
 {
