Enabling audio attenuation settings in restricted mode editor

diff --git a/Engine/Source/Runtime/Engine/Classes/Sound/SoundAttenuation.h b/Engine/Source/Runtime/Engine/Classes/Sound/SoundAttenuation.h
--- a/Engine/Source/Runtime/Engine/Classes/Sound/SoundAttenuation.h
+++ b/Engine/Source/Runtime/Engine/Classes/Sound/SoundAttenuation.h
@@ -8,6 +8,7 @@
 #endif
 #include "IAudioParameterInterfaceRegistry.h"
 #include "AudioLinkSettingsAbstract.h"
+#include "SoundAttenuationEditorSettings.h"
 
 #include "SoundAttenuation.generated.h"
 
@@ -509,6 +510,13 @@ struct FSoundAttenuationSettings : public FBaseAttenuationSettings
 		, PriorityAttenuationDistanceMax(UE_REAL_TO_FLOAT(AttenuationShapeExtents.X) + FalloffDistance)
 		, ManualPriorityAttenuation(1.0f)
 	{
+#if WITH_EDITOR
+		if (const USoundAttenuationEditorSettings* SoundAttenuationEditorSettings = GetDefault<USoundAttenuationEditorSettings>())
+		{
+			bEnableReverbSend = SoundAttenuationEditorSettings->bEnableReverbSend;
+			bEnableSendToAudioLink = SoundAttenuationEditorSettings->bEnableSendToAudioLink;
+		}
+#endif // WITH_EDITOR
 	}
 
 	ENGINE_API bool operator==(const FSoundAttenuationSettings& Other) const;
diff --git a/Engine/Source/Runtime/Engine/Classes/Sound/SoundAttenuationEditorSettings.h b/Engine/Source/Runtime/Engine/Classes/Sound/SoundAttenuationEditorSettings.h
new file mode 100644
--- /dev/null
+++ b/Engine/Source/Runtime/Engine/Classes/Sound/SoundAttenuationEditorSettings.h
@@ -0,0 +1,23 @@
+// Copyright Epic Games, Inc. All Rights Reserved.
+
+#pragma once
+
+#include "SoundAttenuationEditorSettings.generated.h"
+
+/**
+ * Implements Editor settings for sound attenuation assets and is used when switching editor modes.
+ */
+UCLASS(MinimalAPI)
+class USoundAttenuationEditorSettings : public UObject
+{
+	GENERATED_BODY()
+
+public:
+	/** Enables adjusting reverb sends based on distance. */
+	UPROPERTY()
+	bool bEnableReverbSend = true;
+
+	/** Enables/Disables AudioLink on all sources using this attenuation */
+	UPROPERTY()
+	bool bEnableSendToAudioLink = true;
+};
diff --git a/Engine/Source/Runtime/MovieSceneTracks/Public/Sections/MovieSceneAudioSection.h b/Engine/Source/Runtime/MovieSceneTracks/Public/Sections/MovieSceneAudioSection.h
--- a/Engine/Source/Runtime/MovieSceneTracks/Public/Sections/MovieSceneAudioSection.h
+++ b/Engine/Source/Runtime/MovieSceneTracks/Public/Sections/MovieSceneAudioSection.h
@@ -273,7 +273,7 @@ private:
 	bool bOverrideAttenuation;
 
 	/** The attenuation settings to use. */
-	UPROPERTY( EditAnywhere, Category="Attenuation" )
+	UPROPERTY( EditAnywhere, Category="Attenuation", meta = (EditCondition = "bOverrideAttenuation") )
 	TObjectPtr<class USoundAttenuation> AttenuationSettings;
 
 	/** Called when subtitles are sent to the SubtitleManager.  Set this delegate if you want to hijack the subtitles for other purposes */
