World Partition Streaming Generation

RuntimeHash
- Added the possibility for runtime hash implementations to perform a validation on HLOD layers and report errors if the setup is invalid (always valid for the current runtime hash implementation).

RuntimeHashSet
- Wrapped input and output generation parameters into separate structs FGenerateStreamingParams and FGenerateStreamingResult to minimize interface changes in the future.
- Added support for multiple HLOD layers for runtime hashset implementation.

diff --git a/Engine/Plugins/Editor/DataValidation/Source/DataValidation/Private/WorldPartitionChangelistValidator.cpp b/Engine/Plugins/Editor/DataValidation/Source/DataValidation/Private/WorldPartitionChangelistValidator.cpp
--- a/Engine/Plugins/Editor/DataValidation/Source/DataValidation/Private/WorldPartitionChangelistValidator.cpp
+++ b/Engine/Plugins/Editor/DataValidation/Source/DataValidation/Private/WorldPartitionChangelistValidator.cpp
@@ -446,5 +446,16 @@ void UWorldPartitionChangelistValidator::OnInvalidActorFilterReference(const FWo
 	// Not a validation error
 }
 
-#undef LOCTEXT_NAMESPACE
+void UWorldPartitionChangelistValidator::OnInvalidHLODLayer(const FWorldPartitionActorDescView& ActorDescView)
+{
+	if (Filter(ActorDescView))
+	{
+		FText CurrentError = FText::Format(LOCTEXT("DataValidation.Changelist.WorldPartition.InvalidHLODLayer", "Actor {0} has an invalid HLOD Layer {1}"),
+											FText::FromString(GetFullActorName(ActorDescView)), 
+											FText::FromString(ActorDescView.GetHLODLayer().ToString()));
+
+		Errors->Add(CurrentError);
+	}
+}
 
+#undef LOCTEXT_NAMESPACE
\ No newline at end of file
diff --git a/Engine/Plugins/Editor/DataValidation/Source/DataValidation/Private/WorldPartitionChangelistValidator.h b/Engine/Plugins/Editor/DataValidation/Source/DataValidation/Private/WorldPartitionChangelistValidator.h
--- a/Engine/Plugins/Editor/DataValidation/Source/DataValidation/Private/WorldPartitionChangelistValidator.h
+++ b/Engine/Plugins/Editor/DataValidation/Source/DataValidation/Private/WorldPartitionChangelistValidator.h
@@ -47,4 +47,5 @@ protected:
 	virtual void OnActorNeedsResave(const FWorldPartitionActorDescView& ActorDescView) override;
 	virtual void OnLevelInstanceInvalidWorldAsset(const FWorldPartitionActorDescView& ActorDescView, FName WorldAsset, ELevelInstanceInvalidReason Reason) override;
 	virtual void OnInvalidActorFilterReference(const FWorldPartitionActorDescView& ActorDescView, const FWorldPartitionActorDescView& ReferenceActorDescView) override;
+	virtual void OnInvalidHLODLayer(const FWorldPartitionActorDescView& ActorDescView) override;
 };
diff --git a/Engine/Source/Runtime/Engine/Private/WorldPartition/ErrorHandling/WorldPartitionStreamingGenerationLogErrorHandler.cpp b/Engine/Source/Runtime/Engine/Private/WorldPartition/ErrorHandling/WorldPartitionStreamingGenerationLogErrorHandler.cpp
--- a/Engine/Source/Runtime/Engine/Private/WorldPartition/ErrorHandling/WorldPartitionStreamingGenerationLogErrorHandler.cpp
+++ b/Engine/Source/Runtime/Engine/Private/WorldPartition/ErrorHandling/WorldPartitionStreamingGenerationLogErrorHandler.cpp
@@ -96,4 +96,8 @@ void FStreamingGenerationLogErrorHandler::OnInvalidActorFilterReference(const FW
 	UE_ASSET_LOG_ACTORDESCVIEW(LogWorldPartition, Log, ReferenceActorDescView, TEXT("Actor %s will not be filtered out because it is referenced by Actor %s not part of the filter"), *GetActorName(ReferenceActorDescView), *GetActorName(ActorDescView));
 }
 
+void FStreamingGenerationLogErrorHandler::OnInvalidHLODLayer(const FWorldPartitionActorDescView& ActorDescView)
+{
+	UE_ASSET_LOG_ACTORDESCVIEW(LogWorldPartition, Log, ActorDescView, TEXT("Actor %s has an invalid HLOD layer %s"), *GetActorName(ActorDescView), *ActorDescView.GetHLODLayer().ToString());
+}
 #endif
diff --git a/Engine/Source/Runtime/Engine/Private/WorldPartition/ErrorHandling/WorldPartitionStreamingGenerationProxyErrorHandler.cpp b/Engine/Source/Runtime/Engine/Private/WorldPartition/ErrorHandling/WorldPartitionStreamingGenerationProxyErrorHandler.cpp
--- a/Engine/Source/Runtime/Engine/Private/WorldPartition/ErrorHandling/WorldPartitionStreamingGenerationProxyErrorHandler.cpp
+++ b/Engine/Source/Runtime/Engine/Private/WorldPartition/ErrorHandling/WorldPartitionStreamingGenerationProxyErrorHandler.cpp
@@ -67,4 +67,9 @@ void FStreamingGenerationProxyErrorHandler::OnInvalidActorFilterReference(const
 {
 	InnerErrorHandler->OnInvalidActorFilterReference(ActorDescView, ReferenceActorDescView);
 }
+
+void FStreamingGenerationProxyErrorHandler::OnInvalidHLODLayer(const FWorldPartitionActorDescView& ActorDescView)
+{
+	InnerErrorHandler->OnInvalidHLODLayer(ActorDescView);
+}
 #endif
diff --git a/Engine/Source/Runtime/Engine/Private/WorldPartition/ErrorHandling/WorldPartitionStreamingGenerationTokenizedMessageErrorHandler.cpp b/Engine/Source/Runtime/Engine/Private/WorldPartition/ErrorHandling/WorldPartitionStreamingGenerationTokenizedMessageErrorHandler.cpp
--- a/Engine/Source/Runtime/Engine/Private/WorldPartition/ErrorHandling/WorldPartitionStreamingGenerationTokenizedMessageErrorHandler.cpp
+++ b/Engine/Source/Runtime/Engine/Private/WorldPartition/ErrorHandling/WorldPartitionStreamingGenerationTokenizedMessageErrorHandler.cpp
@@ -210,6 +210,18 @@ void ITokenizedMessageErrorHandler::OnInvalidActorFilterReference(const FWorldPa
 	HandleTokenizedMessage(MoveTemp(Message));
 }
 
+void ITokenizedMessageErrorHandler::OnInvalidHLODLayer(const FWorldPartitionActorDescView& ActorDescView)
+{
+	TSharedRef<FTokenizedMessage> Message = FTokenizedMessage::Create(EMessageSeverity::Error);
+	Message->AddToken(FTextToken::Create(LOCTEXT("TokenMessage_WorldPartition_Actor", "Actor")))
+		->AddToken(FActorToken::Create(ActorDescView.GetActorSoftPath().ToString(), ActorDescView.GetGuid(), FText::FromString(GetActorName(ActorDescView))))
+		->AddToken(FTextToken::Create(LOCTEXT("TokenMessage_WorldPartition_HaveInvalidHLODLayer", "has an invalid HLOD layer")))
+		->AddToken(FAssetNameToken::Create(ActorDescView.GetHLODLayer().ToString()))
+		->AddToken(FMapErrorToken::Create(TEXT("WorldPartition_InvalidActorHLODLayer_CheckForErrors")));
+
+	HandleTokenizedMessage(MoveTemp(Message));
+}
+
 #undef LOCTEXT_NAMESPACE
 
 #endif
diff --git a/Engine/Source/Runtime/Engine/Private/WorldPartition/RuntimeHashSet/RuntimePartitionLHGrid.cpp b/Engine/Source/Runtime/Engine/Private/WorldPartition/RuntimeHashSet/RuntimePartitionLHGrid.cpp
--- a/Engine/Source/Runtime/Engine/Private/WorldPartition/RuntimeHashSet/RuntimePartitionLHGrid.cpp
+++ b/Engine/Source/Runtime/Engine/Private/WorldPartition/RuntimeHashSet/RuntimePartitionLHGrid.cpp
@@ -84,7 +84,7 @@ bool URuntimePartitionLHGrid::IsValidGrid(FName GridName) const
 	return GridNameList.Num() == 1;
 }
 
-bool URuntimePartitionLHGrid::GenerateStreaming(const TArray<const IStreamingGenerationContext::FActorSetInstance*>& ActorSetInstances, TArray<FCellDesc>& OutRuntimeCellDescs)
+bool URuntimePartitionLHGrid::GenerateStreaming(const FGenerateStreamingParams& InParams, FGenerateStreamingResult& OutResult)
 {
 	UWorldPartition* WorldPartition = GetTypedOuter<UWorldPartition>();
 	UWorld* World = WorldPartition->GetWorld();
@@ -92,7 +92,7 @@ bool URuntimePartitionLHGrid::GenerateStreaming(const TArray<const IStreamingGen
 	const bool bIsMainWorldPartition = (World == OuterWorld);
 
 	TArray<IStreamingGenerationContext::FActorInstance> CellActorInstances;
-	if (PopulateCellActorInstances(ActorSetInstances, bIsMainWorldPartition, false, CellActorInstances))
+	if (PopulateCellActorInstances(*InParams.ActorSetInstances, bIsMainWorldPartition, false, CellActorInstances))
 	{
 		TMap<FCellCoord, TArray<IStreamingGenerationContext::FActorInstance>> SubLevelsActorInstances;
 		for (const IStreamingGenerationContext::FActorInstance& ActorInstance : CellActorInstances)
@@ -104,7 +104,7 @@ bool URuntimePartitionLHGrid::GenerateStreaming(const TArray<const IStreamingGen
 
 		for (auto& [CellCoord, SubLevelActorSetInstances] : SubLevelsActorInstances)
 		{
-			OutRuntimeCellDescs.Emplace(CreateCellDesc(CellCoord.ToString(), true, SubLevelActorSetInstances[0].ActorSetInstance->ContentBundleID, CellCoord.Level, SubLevelActorSetInstances));
+			OutResult.RuntimeCellDescs.Emplace(CreateCellDesc(CellCoord.ToString(), true, SubLevelActorSetInstances[0].ActorSetInstance->ContentBundleID, CellCoord.Level, SubLevelActorSetInstances));
 		}
 	}
 
diff --git a/Engine/Source/Runtime/Engine/Private/WorldPartition/RuntimeHashSet/RuntimePartitionLevelStreaming.cpp b/Engine/Source/Runtime/Engine/Private/WorldPartition/RuntimeHashSet/RuntimePartitionLevelStreaming.cpp
--- a/Engine/Source/Runtime/Engine/Private/WorldPartition/RuntimeHashSet/RuntimePartitionLevelStreaming.cpp
+++ b/Engine/Source/Runtime/Engine/Private/WorldPartition/RuntimeHashSet/RuntimePartitionLevelStreaming.cpp
@@ -15,7 +15,7 @@ bool URuntimePartitionLevelStreaming::IsValidGrid(FName GridName) const
 	return GridNameList.Num() && (GridNameList.Num() <= 2);
 }
 
-bool URuntimePartitionLevelStreaming::GenerateStreaming(const TArray<const IStreamingGenerationContext::FActorSetInstance*>& ActorSetInstances, TArray<FCellDesc>& OutRuntimeCellDescs)
+bool URuntimePartitionLevelStreaming::GenerateStreaming(const FGenerateStreamingParams& InParams, FGenerateStreamingResult& OutResult)
 {
 	UWorldPartition* WorldPartition = GetTypedOuter<UWorldPartition>();
 	UWorld* World = WorldPartition->GetWorld();
@@ -23,7 +23,7 @@ bool URuntimePartitionLevelStreaming::GenerateStreaming(const TArray<const IStre
 	const bool bIsMainWorldPartition = (World == OuterWorld);
 
 	TArray<IStreamingGenerationContext::FActorInstance> CellActorInstances;
-	if (PopulateCellActorInstances(ActorSetInstances, bIsMainWorldPartition, false, CellActorInstances))
+	if (PopulateCellActorInstances(*InParams.ActorSetInstances, bIsMainWorldPartition, false, CellActorInstances))
 	{
 		TMap<FName, TArray<IStreamingGenerationContext::FActorInstance>> SubLevelsActorInstances;
 		for (const IStreamingGenerationContext::FActorInstance& ActorInstance : CellActorInstances)
@@ -59,7 +59,7 @@ bool URuntimePartitionLevelStreaming::GenerateStreaming(const TArray<const IStre
 
 		for (auto& [SubLevelName, SubLevelActorSetInstances] : SubLevelsActorInstances)
 		{
-			OutRuntimeCellDescs.Emplace(CreateCellDesc(SubLevelName.ToString(), true, SubLevelActorSetInstances[0].ActorSetInstance->ContentBundleID, 0, SubLevelActorSetInstances));
+			OutResult.RuntimeCellDescs.Emplace(CreateCellDesc(SubLevelName.ToString(), true, SubLevelActorSetInstances[0].ActorSetInstance->ContentBundleID, 0, SubLevelActorSetInstances));
 		}
 	}
 
diff --git a/Engine/Source/Runtime/Engine/Private/WorldPartition/RuntimeHashSet/RuntimePartitionPersistent.cpp b/Engine/Source/Runtime/Engine/Private/WorldPartition/RuntimeHashSet/RuntimePartitionPersistent.cpp
--- a/Engine/Source/Runtime/Engine/Private/WorldPartition/RuntimeHashSet/RuntimePartitionPersistent.cpp
+++ b/Engine/Source/Runtime/Engine/Private/WorldPartition/RuntimeHashSet/RuntimePartitionPersistent.cpp
@@ -4,7 +4,7 @@
 #include "WorldPartition/WorldPartitionStreamingGenerationContext.h"
 
 #if WITH_EDITOR
-bool URuntimePartitionPersistent::GenerateStreaming(const TArray<const IStreamingGenerationContext::FActorSetInstance*>& ActorSetInstances, TArray<FCellDesc>& OutRuntimeCellDescs)
+bool URuntimePartitionPersistent::GenerateStreaming(const FGenerateStreamingParams& InParams, FGenerateStreamingResult& OutResult)
 {
 	UWorldPartition* WorldPartition = GetTypedOuter<UWorldPartition>();
 	UWorld* World = WorldPartition->GetWorld();
@@ -12,10 +12,10 @@ bool URuntimePartitionPersistent::GenerateStreaming(const TArray<const IStreamin
 	const bool bIsMainWorldPartition = (World == OuterWorld);
 
 	TArray<IStreamingGenerationContext::FActorInstance> CellActorInstances;
-	if (PopulateCellActorInstances(ActorSetInstances, bIsMainWorldPartition, true, CellActorInstances))
+	if (PopulateCellActorInstances(*InParams.ActorSetInstances, bIsMainWorldPartition, true, CellActorInstances))
 	{
 		const FString CellName(TEXT("Persistent"));
-		OutRuntimeCellDescs.Emplace(CreateCellDesc(CellName, false, CellActorInstances[0].ActorSetInstance->ContentBundleID, 0, CellActorInstances));
+		OutResult.RuntimeCellDescs.Emplace(CreateCellDesc(CellName, false, CellActorInstances[0].ActorSetInstance->ContentBundleID, 0, CellActorInstances));
 	}
 
 	return true;
diff --git a/Engine/Source/Runtime/Engine/Private/WorldPartition/RuntimeHashSet/WorldPartitionRuntimeHashSet.cpp b/Engine/Source/Runtime/Engine/Private/WorldPartition/RuntimeHashSet/WorldPartitionRuntimeHashSet.cpp
--- a/Engine/Source/Runtime/Engine/Private/WorldPartition/RuntimeHashSet/WorldPartitionRuntimeHashSet.cpp
+++ b/Engine/Source/Runtime/Engine/Private/WorldPartition/RuntimeHashSet/WorldPartitionRuntimeHashSet.cpp
@@ -3,22 +3,7 @@
 #include "WorldPartition/RuntimeHashSet/WorldPartitionRuntimeHashSet.h"
 #include "WorldPartition/RuntimeHashSet/RuntimePartition.h"
 #include "WorldPartition/RuntimeHashSet/RuntimePartitionPersistent.h"
-#include "WorldPartition/ContentBundle/ContentBundleWorldSubsystem.h"
-#include "WorldPartition/ContentBundle/ContentBundleDescriptor.h"
-#include "WorldPartition/ContentBundle/ContentBundleBase.h"
-#include "WorldPartition/DataLayer/DataLayerManager.h"
-#include "WorldPartition/WorldPartitionStreamingGenerationContext.h"
-#include "WorldPartition/WorldPartitionRuntimeSpatialHash.h"
-#include "WorldPartition/WorldPartitionStreamingPolicy.h"
-#include "WorldPartition/WorldPartitionStreamingSource.h"
-#include "WorldPartition/WorldPartitionDraw2DContext.h"
-#include "WorldPartition/WorldPartitionDebugHelper.h"
-#include "WorldPartition/DataLayer/DataLayersID.h"
 #include "WorldPartition/HLOD/HLODLayer.h"
-#include "Algo/ForEach.h"
-#include "Algo/Transform.h"
-#include "Math/TransformCalculus2D.h"
-#include "Misc/ArchiveMD5.h"
 
 FAutoConsoleCommand WorldPartitionRuntimeHashSetEnable(
 	TEXT("wp.Editor.WorldPartitionRuntimeHashSet.Enable"),
@@ -29,36 +14,20 @@ FAutoConsoleCommand WorldPartitionRuntimeHashSetEnable(
 	})
 );
 
-static int32 GShowRuntimeHashSetDebugDisplayLevel = 0;
-static FAutoConsoleVariableRef CVarShowRuntimeHashSetDebugDisplayLevel(
-	TEXT("wp.Runtime.HashSet.ShowDebugDisplayLevel"),
-	GShowRuntimeHashSetDebugDisplayLevel,
-	TEXT("Used to choose which level to display when showing runtime partitions."));
-
-static int32 GShowRuntimeHashSetDebugDisplayLevelCount = 1;
-static FAutoConsoleVariableRef CVarShowRuntimeHashSetDebugDisplayLevelCount(
-	TEXT("wp.Runtime.HashSet.ShowDebugDisplayLevelCount"),
-	GShowRuntimeHashSetDebugDisplayLevelCount,
-	TEXT("Used to choose how many levels to display when showing runtime partitions."));
-
-static int32 GShowRuntimeHashSetDebugDisplayMode = 0;
-static FAutoConsoleVariableRef CVarShowRuntimeHashSetDebugDisplayMode(
-	TEXT("wp.Runtime.HashSet.ShowDebugDisplayMode"),
-	GShowRuntimeHashSetDebugDisplayMode,
-	TEXT("Used to choose what mode to display when showing runtime partitions (0=Level Streaming State, 1=Data Layers, 2=Content Bundles)."));
-
 #if WITH_EDITOR
 void FRuntimePartitionDesc::UpdateHLODPartitionLayers()
 {
-	if (!Class || !MainLayer || !HLODLayer)
+	if (!Class || !MainLayer)
 	{
 		HLODSetups.Empty();
+		return;
 	}
-	else
+
+	for (FRuntimePartitionHLODSetup& HLODSetup : HLODSetups)
 	{
 		TSet<const UHLODLayer*> VisitedHLODLayers;
 
-		const UHLODLayer* CurHLODLayer = HLODLayer;
+		const UHLODLayer* CurHLODLayer = HLODSetup.HLODLayer;
 		while (CurHLODLayer)
 		{
 			const int32 HLODSetupIndex = VisitedHLODLayers.Num();
@@ -67,32 +36,33 @@ void FRuntimePartitionDesc::UpdateHLODPartitionLayers()
 			VisitedHLODLayers.Add(CurHLODLayer, &bHLODLayerWasAlreadyInSet);
 			if (bHLODLayerWasAlreadyInSet)
 			{
+				// Circular reference
 				break;
 			}
 
-			if (!HLODSetups.IsValidIndex(HLODSetupIndex))
+			if (!HLODSetup.PartitionLayers.IsValidIndex(HLODSetupIndex))
 			{
-				HLODSetups.AddDefaulted();
+				HLODSetup.PartitionLayers.AddDefaulted();
 			}
 
-			FRuntimePartitionHLODSetup& HLODSetup = HLODSetups[HLODSetupIndex];
+			FRuntimePartitionHLODSetupLayer& HLODSetupLayer = HLODSetup.PartitionLayers[HLODSetupIndex];
 
 			const bool bHLODLayerMatches = HLODSetup.HLODLayer == CurHLODLayer;
 			const UClass* ExpectedHLODPartitionClass = CurHLODLayer->IsSpatiallyLoaded() ? MainLayer->GetClass() : URuntimePartitionPersistent::StaticClass();
-			const bool bHasValidPartitionLayer = HLODSetup.PartitionLayer && (HLODSetup.PartitionLayer->GetClass() == ExpectedHLODPartitionClass);
+			const bool bHasValidPartitionLayer = HLODSetupLayer.PartitionLayer && (HLODSetupLayer.PartitionLayer->GetClass() == ExpectedHLODPartitionClass);
 			
 			if (!bHLODLayerMatches || !bHasValidPartitionLayer)
 			{
-				HLODSetup.HLODLayer = CurHLODLayer;
-				HLODSetup.PartitionLayer = CurHLODLayer->IsSpatiallyLoaded() ? DuplicateObject<URuntimePartition>(MainLayer, MainLayer->GetOuter()) : NewObject<URuntimePartition>(MainLayer->GetOuter(), ExpectedHLODPartitionClass);
-				HLODSetup.PartitionLayer->Name = CurHLODLayer->GetFName();
-				HLODSetup.PartitionLayer->bIsHLODSetup = true;
+				HLODSetupLayer.HLODLayer = CurHLODLayer;
+				HLODSetupLayer.PartitionLayer = CurHLODLayer->IsSpatiallyLoaded() ? DuplicateObject<URuntimePartition>(MainLayer, MainLayer->GetOuter()) : NewObject<URuntimePartition>(MainLayer->GetOuter(), ExpectedHLODPartitionClass);
+				HLODSetupLayer.PartitionLayer->Name = CurHLODLayer->GetFName();
+				HLODSetupLayer.PartitionLayer->bIsHLODSetup = true;
 			}
 
 			CurHLODLayer = CurHLODLayer->GetParentLayer();	
 		}
 
-		HLODSetups.SetNum(VisitedHLODLayers.Num());
+		HLODSetup.PartitionLayers.SetNum(VisitedHLODLayers.Num());
 	}
 }
 #endif
@@ -173,562 +143,83 @@ void UWorldPartitionRuntimeHashSet::PostLoad()
 }
 
 #if WITH_EDITOR
-static EWorldPartitionRuntimeCellVisualizeMode GetStreamingCellVisualizeMode()
-{
-	const EWorldPartitionRuntimeCellVisualizeMode VisualizeMode = FWorldPartitionDebugHelper::IsRuntimeSpatialHashCellStreamingPriorityShown() ? EWorldPartitionRuntimeCellVisualizeMode::StreamingPriority : EWorldPartitionRuntimeCellVisualizeMode::StreamingStatus;
-	return VisualizeMode;
-}
-
-static TMap<FName, FColor> GetDataLayerDebugColors(const UWorldPartition* InWorldPartition)
-{
-	TMap<FName, FColor> DebugColors;
-	if (const UDataLayerManager* DataLayerManager = InWorldPartition->GetDataLayerManager())
-	{
-		DataLayerManager->ForEachDataLayerInstance([&DebugColors](UDataLayerInstance* DataLayerInstance)
-		{
-			DebugColors.Add(DataLayerInstance->GetDataLayerFName(), DataLayerInstance->GetDebugColor());
-			return true;
-		});
-	}
-	return DebugColors;
-}
+void UWorldPartitionRuntimeHashSet::SetDefaultValues()
+{}
 
-bool UWorldPartitionRuntimeHashSet::Draw2D(FWorldPartitionDraw2DContext& DrawContext) const
+void UWorldPartitionRuntimeHashSet::FlushStreaming()
 {
-	TRACE_CPUPROFILER_EVENT_SCOPE(UWorldPartitionRuntimeHashSet::Draw2D);
-
-	const UWorldPartition* WorldPartition = GetOuterUWorldPartition();
-	const TArray<FWorldPartitionStreamingSource>& Sources = WorldPartition->GetStreamingSources();
+	Super::FlushStreaming();
 	
-	if (!Sources.Num())
-	{
-		return false;
-	}
+	check(PersistentPartitionDesc.Class);
+	PersistentPartitionDesc.Class = nullptr;
+	PersistentPartitionDesc.Name = NAME_None;
+	PersistentPartitionDesc.MainLayer = nullptr;
 
-	TMap<FName, TArray<const FRuntimePartitionStreamingData*>> FilteredStreamingObjects;
-	ForEachStreamingData([&FilteredStreamingObjects](const FRuntimePartitionStreamingData& StreamingData)
-	{
-		if (StreamingData.StreamingCells.Num())
-		{
-			if (FWorldPartitionDebugHelper::IsDebugRuntimeHashGridShown(StreamingData.Name))
-			{
-				FilteredStreamingObjects.FindOrAdd(StreamingData.Name).Add(&StreamingData);
-			}
-		}
-	});
+	RuntimeStreamingData.Empty();
+}
 
-	if (!FilteredStreamingObjects.Num())
+bool UWorldPartitionRuntimeHashSet::IsValidGrid(FName GridName) const
+{
+	if (!RuntimePartitions.Num())
 	{
 		return false;
 	}
 
-	const FTransform WorldPartitionTransform = WorldPartition->GetInstanceTransform();
-	const FTransform2D LocalWPToGlobal(FQuat2D(FMath::DegreesToRadians(WorldPartitionTransform.Rotator().Yaw)), FVector2D(WorldPartitionTransform.GetLocation()));
-	const FTransform2D GlobalToLocalWP = LocalWPToGlobal.Inverse();
-	const FBox2D& WorldRegion = DrawContext.GetWorldRegion();
-	const FVector2D WorldRegionSize = WorldRegion.GetSize();
-	TArray<FVector2D> Points;
-	Points.Add(GlobalToLocalWP.TransformPoint(WorldRegion.Min));
-	Points.Add(GlobalToLocalWP.TransformPoint(WorldRegion.Min + FVector2D(WorldRegionSize.X, 0)));
-	Points.Add(GlobalToLocalWP.TransformPoint(WorldRegion.Min + WorldRegionSize));
-	Points.Add(GlobalToLocalWP.TransformPoint(WorldRegion.Min + FVector2D(0, WorldRegionSize.Y)));
-	const FBox2D LocalRegion(Points);
-
-	const FBox2D& CanvasRegion = DrawContext.GetCanvasRegion();
-	const int32 GridScreenWidthDivider = DrawContext.IsDetailedMode() ? FilteredStreamingObjects.Num() : 1;
-	const float GridScreenWidthShrinkSize = GridScreenWidthDivider > 1 ? 20.f : 0.f;
-	const float CanvasMaxScreenWidth = CanvasRegion.GetSize().X;
-	const float GridMaxScreenWidth = CanvasMaxScreenWidth / GridScreenWidthDivider;
-	const float GridEffectiveScreenWidth = FMath::Min(GridMaxScreenWidth, CanvasRegion.GetSize().Y) - GridScreenWidthShrinkSize;
-	const FVector2D PartitionCanvasSize = FVector2D(CanvasRegion.GetSize().GetMin());
-	const FVector2D GridScreenExtent = FVector2D(GridEffectiveScreenWidth, GridEffectiveScreenWidth);
-	const FVector2D GridScreenHalfExtent = 0.5f * GridScreenExtent;
-	const FVector2D GridScreenInitialOffset = CanvasRegion.Min;
-
-	auto DrawStreamingData = [this](const FRuntimePartitionStreamingData* StreamingData, const FBox2D& Region2D, const FBox2D& GridScreenBounds, TFunctionRef<FVector2D(const FVector2D&, bool)> InWorldToScreen, FWorldPartitionDraw2DContext& DrawContext)
-	{
-		const UWorldPartition* WorldPartition = GetOuterUWorldPartition();
-		const TArray<FWorldPartitionStreamingSource>& Sources = WorldPartition->GetStreamingSources();
-
-		const FBox Region(FVector(Region2D.Min.X, Region2D.Min.Y, 0), FVector(Region2D.Max.X, Region2D.Max.Y, 0));
-		const UWorld* OwningWorld = WorldPartition->GetWorld();
-		const EWorldPartitionRuntimeCellVisualizeMode VisualizeMode = GetStreamingCellVisualizeMode();
-		const UContentBundleManager* ContentBundleManager = OwningWorld->ContentBundleManager;
-		TMap<FName, FColor> DataLayerDebugColors = GetDataLayerDebugColors(WorldPartition);
-
-		auto WorldToScreen = [&](const FVector2D& Pos, bool bIsLocal = true)
-		{
-			return InWorldToScreen(Pos, bIsLocal);
-		};
-
-		TArray<const UWorldPartitionRuntimeCell*> FilteredCells;
-		const FBox Region3D(FVector(Region.Min.X, Region.Min.Y, -HALF_WORLD_MAX), FVector(Region.Max.X, Region.Max.Y, HALF_WORLD_MAX));		
-		StreamingData->SpatialIndex->ForEachIntersectingElement(Region3D, [&FilteredCells](UWorldPartitionRuntimeCell* Cell)
-		{
-			UWorldPartitionRuntimeCellDataSpatialHashSet* RuntimeCellDataHashSet = CastChecked<UWorldPartitionRuntimeCellDataSpatialHashSet>(Cell->RuntimeCellData);
-
-			if ((RuntimeCellDataHashSet->Level >= GShowRuntimeHashSetDebugDisplayLevel) && (RuntimeCellDataHashSet->Level < (GShowRuntimeHashSetDebugDisplayLevel + GShowRuntimeHashSetDebugDisplayLevelCount)))
-			{
-				FilteredCells.Add(Cell);
-			}
-		});
-
-		if (FilteredCells.Num())
-		{
-			for (const UWorldPartitionRuntimeCell* Cell : FilteredCells)
-			{
-				const FVector2D CellBoundsSize = FVector2D(Cell->GetCellBounds().GetSize());
-				const FVector2D CellBoundsMin = FVector2D(Cell->GetCellBounds().Min);
-
-				float CellOpacity = 0.0f;;
-				TArray<FLinearColor> CellColors;
-
-				switch (GShowRuntimeHashSetDebugDisplayMode)
-				{
-				case 0:
-					CellColors.Add(Cell->GetDebugColor(VisualizeMode));
-					CellOpacity = 0.25f / FMath::Max<float>(GShowRuntimeHashSetDebugDisplayLevelCount, 1);
-					break;
-				case 1:
-					if (DataLayerDebugColors.Num() && Cell->GetDataLayers().Num())
-					{
-						for (const FName& DataLayer : Cell->GetDataLayers())
-						{
-							CellColors.Add(DataLayerDebugColors[DataLayer]);
-						}
-						CellOpacity = 0.67f;
-					}
-					break;
-				case 2:
-					if (ContentBundleManager && Cell->GetContentBundleID().IsValid())
-					{
-						if (const FContentBundleBase* ContentBundle = ContentBundleManager->GetContentBundle(OwningWorld, Cell->GetContentBundleID()))
-						{
-							check(ContentBundle->GetDescriptor());
-							CellColors.Add(ContentBundle->GetDescriptor()->GetDebugColor());
-							CellOpacity = 0.67f;
-						}
-					}
-					break;
-				}
-
-				if (CellColors.IsEmpty())
-				{
-					CellColors.Add(FLinearColor::White);
-					CellOpacity = 0.1f;
-				}
-
-				FVector2D::FReal BoundsOffsetX = 0;
-				const FVector2D::FReal BoundsOffsetStepX = CellBoundsSize.X / CellColors.Num();
-				for (const FLinearColor& CellColor : CellColors)
-				{
-					const FVector2D EffectiveCellBoundsMin(CellBoundsMin.X + BoundsOffsetX, CellBoundsMin.Y);
-					const FVector2D EffectiveCellBoundsSize(BoundsOffsetStepX, CellBoundsSize.Y);
-					DrawContext.LocalDrawTile(GridScreenBounds, EffectiveCellBoundsMin, EffectiveCellBoundsSize, CellColor.CopyWithNewOpacity(CellOpacity), WorldToScreen);
-					BoundsOffsetX += BoundsOffsetStepX;
-				}
-				
-				DrawContext.LocalDrawBox(GridScreenBounds, CellBoundsMin, CellBoundsSize, FLinearColor::Black, 1, WorldToScreen);
-			}
-		}
-
-		// Draw X/Y Axis
-		if (DrawContext.GetDrawGridAxis())
-		{
-			DrawContext.PushDrawSegment(GridScreenBounds, WorldToScreen(FVector2D(-1638400.f, 0.f), false), WorldToScreen(FVector2D(1638400.f, 0.f), false), FLinearColor::Red, 3);
-			DrawContext.PushDrawSegment(GridScreenBounds, WorldToScreen(FVector2D(0.f, -1638400.f), false), WorldToScreen(FVector2D(0.f, 1638400.f), false), FLinearColor::Green, 3);
-		}
-
-		// Draw Grid Bounds
-		if (DrawContext.GetDrawGridBounds())
-		{
-			const FBox2D Bounds = GridScreenBounds.ExpandBy(FVector2D(10));
-			const FVector2D Size = GridScreenBounds.GetSize();
-			DrawContext.PushDrawBox(Bounds, GridScreenBounds.Min, GridScreenBounds.Min + FVector2D(Size.X, 0), GridScreenBounds.Max, GridScreenBounds.Min + FVector2D(0, Size.Y), FLinearColor::White, 1);
-		}
-
-		// Draw Streaming Sources
-		for (const FWorldPartitionStreamingSource& Source : Sources)
-		{
-			const FColor Color = Source.GetDebugColor();
-			// @todo_jfd
-			const FSoftObjectPath HLODLayer;
-			Source.ForEachShape(StreamingData->LoadingRange, StreamingData->Name, HLODLayer, true, [&Color, &WorldToScreen, &GridScreenBounds, &DrawContext, this](const FSphericalSector& Shape)
-			{
-				check(!Shape.IsNearlyZero())
-
-				// Spherical Sector
-				const FVector2D Center2D(FVector2D(Shape.GetCenter()));
-				const FSphericalSector::FReal Angle = Shape.GetAngle();
-				const int32 MaxSegments = FMath::Max(4, FMath::CeilToInt(64 * Angle / 360.f));
-				const float AngleIncrement = Angle / MaxSegments;
-				const FVector2D Axis = FVector2D(Shape.GetAxis());
-				const FVector Startup = FRotator(0, -0.5f * Angle, 0).RotateVector(Shape.GetScaledAxis());
-
-				FVector2D LineStart = FVector2D(Startup);
-				if (!Shape.IsSphere())
-				{
-					// Draw sector start axis
-					DrawContext.PushDrawSegment(GridScreenBounds, WorldToScreen(Center2D), WorldToScreen(Center2D + LineStart), Color, 2);
-				}
-				// Draw sector Arc
-				for (int32 i = 1; i <= MaxSegments; i++)
-				{
-					FVector2D LineEnd = FVector2D(FRotator(0, AngleIncrement * i, 0).RotateVector(Startup));
-					DrawContext.PushDrawSegment(GridScreenBounds, WorldToScreen(Center2D + LineStart), WorldToScreen(Center2D + LineEnd), Color, 2);
-					LineStart = LineEnd;
-				}
-				// If sphere, close circle, else draw sector end axis
-				DrawContext.PushDrawSegment(GridScreenBounds, WorldToScreen(Center2D + LineStart), WorldToScreen(Center2D + (Shape.IsSphere() ? FVector2D(Startup) : FVector2D::ZeroVector)), Color, 2);
-
-				// Draw direction vector
-				DrawContext.PushDrawSegment(GridScreenBounds, WorldToScreen(Center2D), WorldToScreen(Center2D + Axis * Shape.GetRadius()), Color, 2);
-			});
-		}
-	};
-
-	FBox GridsShapeBounds(ForceInit);
-	FBox2D GridsBounds(ForceInit);
-	int32 GridIndex = 0;
-	for (auto& [Name, StreamingDataList] : FilteredStreamingObjects)
-	{
-		for (const FWorldPartitionStreamingSource& Source : Sources)
-		{
-			// @todo_jfd
-			const FSoftObjectPath HLODLayer;
-			Source.ForEachShape(StreamingDataList[0]->LoadingRange, Name, HLODLayer, true, [&GridsShapeBounds](const FSphericalSector& Shape) { GridsShapeBounds += Shape.CalcBounds(); });
-		}
-
-		const FVector2D GridReferenceWorldPos = FVector2D(WorldRegion.GetCenter());
-		const FVector2D WorldRegionExtent = FVector2D(WorldRegion.GetExtent().GetMax());
-		const FVector2D GridScreenOffset = GridScreenInitialOffset + ((float)GridIndex * FVector2D(GridMaxScreenWidth, 0.f)) + GridScreenHalfExtent + FVector2D(GridScreenWidthShrinkSize * 0.5f);
-		const FVector2D WorldToScreenScale = GridScreenHalfExtent / WorldRegionExtent;
-		const FBox2D GridScreenBounds(GridScreenOffset - GridScreenHalfExtent, GridScreenOffset + GridScreenHalfExtent);
-
-		auto WorldToScreen = [&](const FVector2D& LocalWorldPos, bool bIsLocalWorldPos = true)
-		{
-			FVector2D GlocalPos = bIsLocalWorldPos ? LocalWPToGlobal.TransformPoint(LocalWorldPos) : LocalWorldPos;
-			return (WorldToScreenScale * (GlocalPos - GridReferenceWorldPos)) + GridScreenOffset;
-		};
-
-		for (const FRuntimePartitionStreamingData* StreamingData : StreamingDataList)
-		{
-			DrawStreamingData(StreamingData, LocalRegion, GridScreenBounds, WorldToScreen, DrawContext);
-		}
-
-		GridsBounds += GridScreenBounds;
-
-		if (DrawContext.IsDetailedMode())
-		{
-			FVector2D GridInfoPos = GridScreenOffset - GridScreenHalfExtent;
-			FWorldPartitionCanvasMultiLineText MultiLineText;
-			MultiLineText.Emplace(UWorld::RemovePIEPrefix(FPaths::GetBaseFilename(WorldPartition->GetPackage()->GetName())), FLinearColor::White);
-			FString GridInfoText = FString::Printf(TEXT("%s | %d m"), *Name.ToString(), int32(StreamingDataList[0]->LoadingRange * 0.01f));
-			MultiLineText.Emplace(GridInfoText, FLinearColor::Yellow);
-			FWorldPartitionCanvasMultiLineTextItem Item(GridInfoPos, MultiLineText);
-			DrawContext.PushDrawText(Item);
-			++GridIndex;
-		}
-	}
-
-	FBox2D DesiredWorldBounds(ForceInit);
-	if (GridsShapeBounds.IsValid)
+	// The None grid name will always map to the first runtime partition in the list
+	if (GridName.IsNone())
 	{
-		// Convert to 2D
-		FBox2D GridsShapeBounds2D = FBox2D(FVector2D(GridsShapeBounds.Min.X, GridsShapeBounds.Min.Y), FVector2D(GridsShapeBounds.Max.X, GridsShapeBounds.Max.Y));
-		FVector2D CenterGlobalPos = LocalWPToGlobal.TransformPoint(GridsShapeBounds2D.GetCenter());
-		// Use max extent of X/Y
-		FVector2D Extent = FVector2D(GridsShapeBounds2D.GetExtent().GetMax());
-		// Transform to global space
-		GridsShapeBounds2D = FBox2D(CenterGlobalPos - Extent, CenterGlobalPos + Extent);
-		// Expand by 10% computed bounds
-		DesiredWorldBounds = GridsShapeBounds2D.ExpandBy(GridsShapeBounds2D.GetExtent() * 0.1f);
+		return true;
 	}
-	DrawContext.SetDesiredWorldBounds(DesiredWorldBounds);
-	DrawContext.SetUsedCanvasBounds(GridsBounds);
-
-	return true;
-}
 
-void UWorldPartitionRuntimeHashSet::SetDefaultValues()
-{
-}
-
-bool UWorldPartitionRuntimeHashSet::SupportsHLODs() const
-{
+	// Parse the potentially dot separated grid name to identiy the associated runtime partition
+	const TArray<FName> GridNameList = ParseGridName(GridName);
 	for (const FRuntimePartitionDesc& RuntimePartitionDesc : RuntimePartitions)
 	{
-		if (RuntimePartitionDesc.MainLayer)
+		if (RuntimePartitionDesc.Name == GridNameList[0])
 		{
-			if (RuntimePartitionDesc.MainLayer->SupportsHLODs())
-			{
-				return true;
-			}
+			return RuntimePartitionDesc.MainLayer && RuntimePartitionDesc.MainLayer->IsValidGrid(GridName);
 		}
 	}
 
 	return false;
 }
 
-bool UWorldPartitionRuntimeHashSet::GenerateStreaming(UWorldPartitionStreamingPolicy* StreamingPolicy, const IStreamingGenerationContext* StreamingGenerationContext, TArray<FString>* OutPackagesToGenerate)
+bool UWorldPartitionRuntimeHashSet::IsValidHLODLayer(FName GridName, const FSoftObjectPath& HLODLayerPath) const
 {
-	verify(Super::GenerateStreaming(StreamingPolicy, StreamingGenerationContext, OutPackagesToGenerate));
-
-	UWorldPartition* WorldPartition = GetOuterUWorldPartition();
-	UWorld* World = WorldPartition->GetWorld();
-	UWorld* OuterWorld = GetTypedOuter<UWorld>();
-	const bool bIsMainWorldPartition = (World == OuterWorld);
-
-	check(!PersistentPartitionDesc.Class);
-	PersistentPartitionDesc.Class = URuntimePartitionPersistent::StaticClass();
-	PersistentPartitionDesc.Name = NAME_PersistentLevel;
-	PersistentPartitionDesc.MainLayer = NewObject<URuntimePartition>(this, URuntimePartitionPersistent::StaticClass(), NAME_None);
-	PersistentPartitionDesc.MainLayer->Name = NAME_PersistentLevel;
-
-	//
-	// Split actor sets into their corresponding runtime partition implementation
-	//
-	TMap<FName, const FRuntimePartitionDesc*> NameToRuntimePartitionDescMap;
-
-	NameToRuntimePartitionDescMap.Add(NAME_None, &RuntimePartitions[0]);				// Actors with RuntimeGrid=None will be assigned to the default partition
-	NameToRuntimePartitionDescMap.Add(NAME_PersistentLevel, &PersistentPartitionDesc);	// Non-spatially loaded actors will be assigned to the persistent partition
-
-	for (const FRuntimePartitionDesc& RuntimePartitionDesc : RuntimePartitions)
+	if (!RuntimePartitions.Num())
 	{
-		NameToRuntimePartitionDescMap.Add(RuntimePartitionDesc.Name, &RuntimePartitionDesc);
+		return false;
 	}
 
-	TMap<URuntimePartition*, TArray<const IStreamingGenerationContext::FActorSetInstance*>> RuntimePartitionsToActorSetMap;
-	StreamingGenerationContext->ForEachActorSetInstance([this, &NameToRuntimePartitionDescMap, &RuntimePartitionsToActorSetMap](const IStreamingGenerationContext::FActorSetInstance& ActorSetInstance)
-	{
-		const TArray<FName> ActorSetRuntimeGrid = ActorSetInstance.bIsSpatiallyLoaded ? ParseGridName(ActorSetInstance.RuntimeGrid) : TArray<FName>({ NAME_PersistentLevel });
-
-		if (const FRuntimePartitionDesc** RuntimePartitionDesc = NameToRuntimePartitionDescMap.Find(ActorSetRuntimeGrid[0]))
-		{
-			RuntimePartitionsToActorSetMap.FindOrAdd((*RuntimePartitionDesc)->MainLayer).Add(&ActorSetInstance);
-		}
-	});
-
-	struct FCellDescInstance : public URuntimePartition::FCellDesc
-	{
-		FCellDescInstance(const URuntimePartition::FCellDesc& InCellDesc, URuntimePartition* InSourcePartition, const TArray<const UDataLayerInstance*>& InDataLayers)
-			: URuntimePartition::FCellDesc(InCellDesc)
-			, SourcePartition(InSourcePartition)
-			, DataLayers(InDataLayers)
-		{}
-
-		URuntimePartition* SourcePartition;
-		TArray<const UDataLayerInstance*> DataLayers;
-	};
-
-	//
-	// Generate runtime partitions streaming data
-	//
-	TArray<FCellDescInstance> RuntimeCellDescsInstances;
+	if (const UHLODLayer* HLODLayer = Cast<UHLODLayer>(HLODLayerPath.ResolveObject()))
 	{
-		TSet<FName> CellDescsNames;
-		for (auto [RuntimePartition, ActorSetInstances] : RuntimePartitionsToActorSetMap)
+		// The None grid name will always map to the first runtime partition in the list
+		int32 RuntimePartitionIndex = GridName.IsNone() ? 0 : INDEX_NONE;
+		
+		if (RuntimePartitionIndex == INDEX_NONE)
 		{
-			// Gather runtime partition cell descs
-			TArray<URuntimePartition::FCellDesc> RuntimeCellDescs;
-			if (!RuntimePartition->GenerateStreaming(ActorSetInstances, RuntimeCellDescs))
-			{
-				return false;
-			}
-
-			// Split cell descs into data layers
-			for (const URuntimePartition::FCellDesc& RuntimeCellDesc : RuntimeCellDescs)
+			// Parse the potentially dot separated grid name to identiy the associated runtime partition
+			const TArray<FName> GridNameList = ParseGridName(GridName);
+			for (const FRuntimePartitionDesc& RuntimePartitionDesc : RuntimePartitions)
 			{
-				TMap<FDataLayersID, FCellDescInstance> RuntimeCellDescsInstancesSet;
-
-				bool bCellNameExists;
-				CellDescsNames.Add(RuntimeCellDesc.Name, &bCellNameExists);
-				check(!bCellNameExists);
-
-				for (const IStreamingGenerationContext::FActorInstance& ActorInstance : RuntimeCellDesc.ActorInstances)
+				if (RuntimePartitionDesc.Name == GridNameList[0])
 				{
-					const FDataLayersID DataLayersID(ActorInstance.ActorSetInstance->DataLayers);
-					FCellDescInstance* CellDescInstance = RuntimeCellDescsInstancesSet.Find(DataLayersID);
-
-					if (!CellDescInstance)
-					{
-						CellDescInstance = &RuntimeCellDescsInstancesSet.Emplace(DataLayersID, FCellDescInstance(RuntimeCellDesc, RuntimePartition, ActorInstance.ActorSetInstance->DataLayers));
-						CellDescInstance->ActorInstances.Empty();
-					}
-
-					CellDescInstance->ActorInstances.Add(ActorInstance);
+					break;
 				}
-
-				Algo::Transform(RuntimeCellDescsInstancesSet, RuntimeCellDescsInstances, [](const auto& Value) { return Value.Value; });
+				RuntimePartitionIndex++;
 			}
 		}
-	}
 
-	//
-	// Create and populate streaming object
-	//
-	check(RuntimeStreamingData.IsEmpty());
-
-	// Generate runtime cells
-	auto CreateRuntimeCellFromCellDesc = [this](const URuntimePartition::FCellDesc& CellDesc, const TArray<const UDataLayerInstance*>& DataLayers, TSubclassOf<UWorldPartitionRuntimeCell> CellClass, TSubclassOf<UWorldPartitionRuntimeCellData> CellDataClass)
-	{
-		FString CellObjectName;
-		FGuid CellGuid;
+		if (RuntimePartitionIndex == INDEX_NONE)
 		{
-			UWorld* OuterWorld = GetTypedOuter<UWorld>();
-			check(OuterWorld);
-
-			FString WorldName = FPackageName::GetShortName(OuterWorld->GetPackage());
-
-			CellObjectName = FString::Printf(TEXT("%s_%s"), *WorldName, *CellDesc.Name.ToString());
-
-			const FDataLayersID DataLayersID(DataLayers);
-			if (DataLayersID.GetHash())
-			{
-				CellObjectName += FString::Printf(TEXT("_d%X"), DataLayersID.GetHash());
-			}
-
-			if (CellDesc.ContentBundleID.IsValid())
-			{
-				CellObjectName += FString::Printf(TEXT("_c%s"), *UContentBundleDescriptor::GetContentBundleCompactString(CellDesc.ContentBundleID));
-			}
-
-			if (!IsRunningCookCommandlet() && OuterWorld->IsGameWorld())
-			{
-				FString SourceWorldPath;
-				FString InstancedWorldPath;
-				if (OuterWorld->GetSoftObjectPathMapping(SourceWorldPath, InstancedWorldPath))
-				{
-					const FTopLevelAssetPath SourceAssetPath(SourceWorldPath);
-					WorldName = FPackageName::GetShortName(SourceAssetPath.GetPackageName());
-						
-					InstancedWorldPath = UWorld::RemovePIEPrefix(InstancedWorldPath);
-
-					const FString SourcePackageName = SourceAssetPath.GetPackageName().ToString();
-					const FTopLevelAssetPath InstanceAssetPath(InstancedWorldPath);
-					const FString InstancePackageName = InstanceAssetPath.GetPackageName().ToString();
-
-					if (int32 Index = InstancePackageName.Find(SourcePackageName); Index != INDEX_NONE)
-					{
-						CellObjectName += FString::Printf(TEXT("_i%s"), *InstancePackageName.Mid(Index + SourcePackageName.Len()));
-					}
-				}
-			}
-
-			FArchiveMD5 ArMD5;
-			ArMD5 << CellObjectName;
-			CellGuid = ArMD5.GetGuidFromHash();
-			check(CellGuid.IsValid());
+			return false;
 		}
 
-		UWorldPartitionRuntimeCell* RuntimeCell = Super::CreateRuntimeCell(CellClass, CellDataClass, CellObjectName, TEXT(""));
-
-		RuntimeCell->SetIsAlwaysLoaded(!CellDesc.bIsSpatiallyLoaded);
-		RuntimeCell->SetDataLayers(DataLayers);
-		RuntimeCell->SetContentBundleUID(CellDesc.ContentBundleID);
-		RuntimeCell->SetPriority(CellDesc.Priority);
-		RuntimeCell->SetClientOnlyVisible(CellDesc.bClientOnlyVisible);
-		RuntimeCell->SetBlockOnSlowLoading(CellDesc.bBlockOnSlowStreaming);
-		RuntimeCell->SetIsHLOD(false);
-		RuntimeCell->SetGuid(CellGuid);
-
-		UWorldPartitionRuntimeCellDataSpatialHashSet* RuntimeCellDataHashSet = CastChecked<UWorldPartitionRuntimeCellDataSpatialHashSet>(RuntimeCell->RuntimeCellData);
-		RuntimeCellDataHashSet->DebugName = CellObjectName;
-		RuntimeCellDataHashSet->Level = CellDesc.Level;
-
-		return RuntimeCell;
-	};
-
-	TArray<UWorldPartitionRuntimeCell*> RuntimeCells;
-	TMap<URuntimePartition*, FRuntimePartitionStreamingData> RuntimePartitionsStreamingData;
-	for (const FCellDescInstance& CellDescInstance : RuntimeCellDescsInstances)
-	{
-		UWorldPartitionRuntimeCell* RuntimeCell = RuntimeCells.Emplace_GetRef(CreateRuntimeCellFromCellDesc(CellDescInstance, CellDescInstance.DataLayers, StreamingPolicy->GetRuntimeCellClass(), UWorldPartitionRuntimeCellDataSpatialHashSet::StaticClass()));
-		PopulateRuntimeCell(RuntimeCell, CellDescInstance.ActorInstances, nullptr);
-
-		// Override the cell bounds if the runtime partition provided one
-		if (CellDescInstance.Bounds.IsValid)
+		for (const FRuntimePartitionHLODSetup& HLODSetup : RuntimePartitions[RuntimePartitionIndex].HLODSetups)
 		{
-			RuntimeCell->RuntimeCellData->ContentBounds = CellDescInstance.Bounds;
-		}
-
-		// Create partition streaming data
-		FRuntimePartitionStreamingData& StreamingData = RuntimePartitionsStreamingData.FindOrAdd(CellDescInstance.SourcePartition);
-
-		StreamingData.Name = CellDescInstance.SourcePartition->Name;
-		StreamingData.LoadingRange = CellDescInstance.SourcePartition->LoadingRange;
-
-		if (RuntimeCell->IsAlwaysLoaded())
-		{
-			StreamingData.NonStreamingCells.Add(RuntimeCell);
-		}
-		else
-		{
-			StreamingData.StreamingCells.Add(RuntimeCell);
-		}
-	}
-
-	//
-	// Finalize streaming object
-	//
-	for (auto& [Partition, StreamingData] : RuntimePartitionsStreamingData)
-	{
-		RuntimeStreamingData.Emplace(MoveTemp(StreamingData));
-	}
-
-	//
-	// Output generated packages
-	//
-	if (OutPackagesToGenerate)
-	{
-		for (UWorldPartitionRuntimeCell* RuntimeCell : RuntimeCells)
-		{
-			// Always loaded cell actors are transfered to World's Persistent Level (see UWorldPartitionRuntimeSpatialHash::PopulateGeneratorPackageForCook)
-			if (RuntimeCell->GetActorCount() && !RuntimeCell->IsAlwaysLoaded())
+			if (HLODSetup.HLODLayer == HLODLayer)
 			{
-				const FString PackageRelativePath = RuntimeCell->GetPackageNameToCreate();
-				check(!PackageRelativePath.IsEmpty());
-
-				OutPackagesToGenerate->Add(PackageRelativePath);
-
-				// Map relative package to StreamingCell for PopulateGeneratedPackageForCook/PopulateGeneratorPackageForCook/GetCellForPackage
-				PackagesToGenerateForCook.Add(PackageRelativePath, RuntimeCell);
-			}
-		}
-	}
-
-	return true;
-}
-
-void UWorldPartitionRuntimeHashSet::FlushStreaming()
-{
-	Super::FlushStreaming();
-	
-	check(PersistentPartitionDesc.Class);
-	PersistentPartitionDesc.Class = nullptr;
-	PersistentPartitionDesc.Name = NAME_None;
-	PersistentPartitionDesc.MainLayer = nullptr;
-
-	RuntimeStreamingData.Empty();
-}
-
-bool UWorldPartitionRuntimeHashSet::IsValidGrid(FName GridName) const
-{
-	// The None grid name will always map to the first runtime partition in the list
-	if (GridName.IsNone())
-	{
-		return true;
-	}
-
-	// Parse the potentially dot separated grid name to identiy the associated runtime partition
-	const TArray<FName> GridNameList = ParseGridName(GridName);
-	for (const FRuntimePartitionDesc& RuntimePartitionDesc : RuntimePartitions)
-	{
-		if (RuntimePartitionDesc.Name == GridNameList[0])
-		{
-			if (RuntimePartitionDesc.MainLayer)
-			{
-				return RuntimePartitionDesc.MainLayer->IsValidGrid(GridName);
+				return true;
 			}
 		}
 	}
@@ -746,28 +237,6 @@ TArray<UWorldPartitionRuntimeCell*> UWorldPartitionRuntimeHashSet::GetAlwaysLoad
 	return AlwaysLoadedCells;
 }
 
-void UWorldPartitionRuntimeHashSet::DumpStateLog(FHierarchicalLogArchive& Ar) const
-{
-	Super::DumpStateLog(Ar);
-
-	Ar.Printf(TEXT("----------------------------------------------------------------------------------------------------------------"));
-	Ar.Printf(TEXT("%s - Runtime Hash Set"), *GetWorld()->GetName());
-	Ar.Printf(TEXT("----------------------------------------------------------------------------------------------------------------"));
-
-	TArray<const UWorldPartitionRuntimeCell*> StreamingCells;
-	ForEachStreamingCells([&StreamingCells](const UWorldPartitionRuntimeCell* StreamingCell) { StreamingCells.Add(StreamingCell); return true; });
-				
-	StreamingCells.Sort([this](const UWorldPartitionRuntimeCell& A, const UWorldPartitionRuntimeCell& B) { return A.GetFName().LexicalLess(B.GetFName()); });
-
-	for (const UWorldPartitionRuntimeCell* StreamingCell : StreamingCells)
-	{
-		FHierarchicalLogArchive::FIndentScope CellIndentScope = Ar.PrintfIndent(TEXT("Content of Cell %s (%s)"), *StreamingCell->GetDebugName(), *StreamingCell->GetName());
-		StreamingCell->DumpStateLog(Ar);
-	}
-
-	Ar.Printf(TEXT(""));
-}
-
 TArray<FName> UWorldPartitionRuntimeHashSet::ParseGridName(FName GridName)
 {
 	TArray<FString> GridNameList;
@@ -1077,4 +546,4 @@ void UWorldPartitionRuntimeHashSet::ForEachStreamingData(TFunctionRef<void(const
 			}
 		}
 	}
-}
+}
\ No newline at end of file
diff --git a/Engine/Source/Runtime/Engine/Private/WorldPartition/RuntimeHashSet/WorldPartitionRuntimeHashSetDebugDraw.cpp b/Engine/Source/Runtime/Engine/Private/WorldPartition/RuntimeHashSet/WorldPartitionRuntimeHashSetDebugDraw.cpp
new file mode 100644
--- /dev/null
+++ b/Engine/Source/Runtime/Engine/Private/WorldPartition/RuntimeHashSet/WorldPartitionRuntimeHashSetDebugDraw.cpp
@@ -0,0 +1,319 @@
+// Copyright Epic Games, Inc. All Rights Reserved.
+
+#include "WorldPartition/RuntimeHashSet/WorldPartitionRuntimeHashSet.h"
+#include "WorldPartition/ContentBundle/ContentBundleWorldSubsystem.h"
+#include "WorldPartition/ContentBundle/ContentBundleDescriptor.h"
+#include "WorldPartition/ContentBundle/ContentBundleBase.h"
+#include "WorldPartition/DataLayer/DataLayerManager.h"
+#include "WorldPartition/WorldPartitionDraw2DContext.h"
+#include "WorldPartition/WorldPartitionDebugHelper.h"
+#include "DrawDebugHelpers.h"
+#include "Math/TransformCalculus2D.h"
+
+static int32 GShowRuntimeHashSetDebugDisplayLevel = 0;
+static FAutoConsoleVariableRef CVarShowRuntimeHashSetDebugDisplayLevel(
+	TEXT("wp.Runtime.HashSet.ShowDebugDisplayLevel"),
+	GShowRuntimeHashSetDebugDisplayLevel,
+	TEXT("Used to choose which level to display when showing runtime partitions."));
+
+static int32 GShowRuntimeHashSetDebugDisplayLevelCount = 1;
+static FAutoConsoleVariableRef CVarShowRuntimeHashSetDebugDisplayLevelCount(
+	TEXT("wp.Runtime.HashSet.ShowDebugDisplayLevelCount"),
+	GShowRuntimeHashSetDebugDisplayLevelCount,
+	TEXT("Used to choose how many levels to display when showing runtime partitions."));
+
+static int32 GShowRuntimeHashSetDebugDisplayMode = 0;
+static FAutoConsoleVariableRef CVarShowRuntimeHashSetDebugDisplayMode(
+	TEXT("wp.Runtime.HashSet.ShowDebugDisplayMode"),
+	GShowRuntimeHashSetDebugDisplayMode,
+	TEXT("Used to choose what mode to display when showing runtime partitions (0=Level Streaming State, 1=Data Layers, 2=Content Bundles)."));
+
+static EWorldPartitionRuntimeCellVisualizeMode GetStreamingCellVisualizeMode()
+{
+	const EWorldPartitionRuntimeCellVisualizeMode VisualizeMode = FWorldPartitionDebugHelper::IsRuntimeSpatialHashCellStreamingPriorityShown() ? EWorldPartitionRuntimeCellVisualizeMode::StreamingPriority : EWorldPartitionRuntimeCellVisualizeMode::StreamingStatus;
+	return VisualizeMode;
+}
+
+static TMap<FName, FColor> GetDataLayerDebugColors(const UWorldPartition* InWorldPartition)
+{
+	TMap<FName, FColor> DebugColors;
+	if (const UDataLayerManager* DataLayerManager = InWorldPartition->GetDataLayerManager())
+	{
+		DataLayerManager->ForEachDataLayerInstance([&DebugColors](UDataLayerInstance* DataLayerInstance)
+		{
+			DebugColors.Add(DataLayerInstance->GetDataLayerFName(), DataLayerInstance->GetDebugColor());
+			return true;
+		});
+	}
+	return DebugColors;
+}
+
+bool UWorldPartitionRuntimeHashSet::Draw2D(FWorldPartitionDraw2DContext& DrawContext) const
+{
+	TRACE_CPUPROFILER_EVENT_SCOPE(UWorldPartitionRuntimeHashSet::Draw2D);
+
+	const UWorldPartition* WorldPartition = GetOuterUWorldPartition();
+	const TArray<FWorldPartitionStreamingSource>& Sources = WorldPartition->GetStreamingSources();
+	
+	if (!Sources.Num())
+	{
+		return false;
+	}
+
+	TMap<FName, TArray<const FRuntimePartitionStreamingData*>> FilteredStreamingObjects;
+	ForEachStreamingData([&FilteredStreamingObjects](const FRuntimePartitionStreamingData& StreamingData)
+	{
+		if (StreamingData.StreamingCells.Num())
+		{
+			if (FWorldPartitionDebugHelper::IsDebugRuntimeHashGridShown(StreamingData.Name))
+			{
+				FilteredStreamingObjects.FindOrAdd(StreamingData.Name).Add(&StreamingData);
+			}
+		}
+	});
+
+	if (!FilteredStreamingObjects.Num())
+	{
+		return false;
+	}
+
+	const FTransform WorldPartitionTransform = WorldPartition->GetInstanceTransform();
+	const FTransform2D LocalWPToGlobal(FQuat2D(FMath::DegreesToRadians(WorldPartitionTransform.Rotator().Yaw)), FVector2D(WorldPartitionTransform.GetLocation()));
+	const FTransform2D GlobalToLocalWP = LocalWPToGlobal.Inverse();
+	const FBox2D& WorldRegion = DrawContext.GetWorldRegion();
+	const FVector2D WorldRegionSize = WorldRegion.GetSize();
+	TArray<FVector2D> Points;
+	Points.Add(GlobalToLocalWP.TransformPoint(WorldRegion.Min));
+	Points.Add(GlobalToLocalWP.TransformPoint(WorldRegion.Min + FVector2D(WorldRegionSize.X, 0)));
+	Points.Add(GlobalToLocalWP.TransformPoint(WorldRegion.Min + WorldRegionSize));
+	Points.Add(GlobalToLocalWP.TransformPoint(WorldRegion.Min + FVector2D(0, WorldRegionSize.Y)));
+	const FBox2D LocalRegion(Points);
+
+	const FBox2D& CanvasRegion = DrawContext.GetCanvasRegion();
+	const int32 GridScreenWidthDivider = DrawContext.IsDetailedMode() ? FilteredStreamingObjects.Num() : 1;
+	const float GridScreenWidthShrinkSize = GridScreenWidthDivider > 1 ? 20.f : 0.f;
+	const float CanvasMaxScreenWidth = CanvasRegion.GetSize().X;
+	const float GridMaxScreenWidth = CanvasMaxScreenWidth / GridScreenWidthDivider;
+	const float GridEffectiveScreenWidth = FMath::Min(GridMaxScreenWidth, CanvasRegion.GetSize().Y) - GridScreenWidthShrinkSize;
+	const FVector2D PartitionCanvasSize = FVector2D(CanvasRegion.GetSize().GetMin());
+	const FVector2D GridScreenExtent = FVector2D(GridEffectiveScreenWidth, GridEffectiveScreenWidth);
+	const FVector2D GridScreenHalfExtent = 0.5f * GridScreenExtent;
+	const FVector2D GridScreenInitialOffset = CanvasRegion.Min;
+
+	auto DrawStreamingData = [this, &WorldPartitionTransform](const FRuntimePartitionStreamingData* StreamingData, const FBox2D& Region2D, const FBox2D& GridScreenBounds, TFunctionRef<FVector2D(const FVector2D&, bool)> InWorldToScreen, FWorldPartitionDraw2DContext& DrawContext)
+	{
+		const UWorldPartition* WorldPartition = GetOuterUWorldPartition();
+		const TArray<FWorldPartitionStreamingSource>& Sources = WorldPartition->GetStreamingSources();
+
+		const FBox Region(FVector(Region2D.Min.X, Region2D.Min.Y, 0), FVector(Region2D.Max.X, Region2D.Max.Y, 0));
+		const UWorld* OwningWorld = WorldPartition->GetWorld();
+		const EWorldPartitionRuntimeCellVisualizeMode VisualizeMode = GetStreamingCellVisualizeMode();
+		const UContentBundleManager* ContentBundleManager = OwningWorld->ContentBundleManager;
+		TMap<FName, FColor> DataLayerDebugColors = GetDataLayerDebugColors(WorldPartition);
+
+		auto WorldToScreen = [&](const FVector2D& Pos, bool bIsLocal = true)
+		{
+			return InWorldToScreen(Pos, bIsLocal);
+		};
+
+		TArray<const UWorldPartitionRuntimeCell*> FilteredCells;
+		const FBox Region3D(FVector(Region.Min.X, Region.Min.Y, -HALF_WORLD_MAX), FVector(Region.Max.X, Region.Max.Y, HALF_WORLD_MAX));		
+		StreamingData->SpatialIndex->ForEachIntersectingElement(Region3D, [&FilteredCells](UWorldPartitionRuntimeCell* Cell)
+		{
+			UWorldPartitionRuntimeCellDataSpatialHashSet* RuntimeCellDataHashSet = CastChecked<UWorldPartitionRuntimeCellDataSpatialHashSet>(Cell->RuntimeCellData);
+
+			if ((RuntimeCellDataHashSet->Level >= GShowRuntimeHashSetDebugDisplayLevel) && (RuntimeCellDataHashSet->Level < (GShowRuntimeHashSetDebugDisplayLevel + GShowRuntimeHashSetDebugDisplayLevelCount)))
+			{
+				FilteredCells.Add(Cell);
+			}
+		});
+
+		if (FilteredCells.Num())
+		{
+			for (const UWorldPartitionRuntimeCell* Cell : FilteredCells)
+			{
+				const FVector2D CellBoundsSize = FVector2D(Cell->GetCellBounds().GetSize());
+				const FVector2D CellBoundsMin = FVector2D(Cell->GetCellBounds().Min);
+
+				float CellOpacity = 0.0f;;
+				TArray<FLinearColor> CellColors;
+
+				switch (GShowRuntimeHashSetDebugDisplayMode)
+				{
+				case 0:
+					CellColors.Add(Cell->GetDebugColor(VisualizeMode));
+					CellOpacity = 0.25f / FMath::Max<float>(GShowRuntimeHashSetDebugDisplayLevelCount, 1);
+					break;
+				case 1:
+					if (DataLayerDebugColors.Num() && Cell->GetDataLayers().Num())
+					{
+						for (const FName& DataLayer : Cell->GetDataLayers())
+						{
+							CellColors.Add(DataLayerDebugColors[DataLayer]);
+						}
+						CellOpacity = 0.67f;
+					}
+					break;
+				case 2:
+					if (ContentBundleManager && Cell->GetContentBundleID().IsValid())
+					{
+						if (const FContentBundleBase* ContentBundle = ContentBundleManager->GetContentBundle(OwningWorld, Cell->GetContentBundleID()))
+						{
+							check(ContentBundle->GetDescriptor());
+							CellColors.Add(ContentBundle->GetDescriptor()->GetDebugColor());
+							CellOpacity = 0.67f;
+						}
+					}
+					break;
+				}
+
+				if (CellColors.IsEmpty())
+				{
+					CellColors.Add(FLinearColor::White);
+					CellOpacity = 0.1f;
+				}
+
+				FVector2D::FReal BoundsOffsetX = 0;
+				const FVector2D::FReal BoundsOffsetStepX = CellBoundsSize.X / CellColors.Num();
+				for (const FLinearColor& CellColor : CellColors)
+				{
+					const FVector2D EffectiveCellBoundsMin(CellBoundsMin.X + BoundsOffsetX, CellBoundsMin.Y);
+					const FVector2D EffectiveCellBoundsSize(BoundsOffsetStepX, CellBoundsSize.Y);
+					DrawContext.LocalDrawTile(GridScreenBounds, EffectiveCellBoundsMin, EffectiveCellBoundsSize, CellColor.CopyWithNewOpacity(CellOpacity), WorldToScreen);
+					BoundsOffsetX += BoundsOffsetStepX;
+				}
+				
+				DrawContext.LocalDrawBox(GridScreenBounds, CellBoundsMin, CellBoundsSize, FLinearColor::Black, 1, WorldToScreen);
+
+				// Draw Cell using its debug color
+				const FBox Box(Cell->GetCellBounds());
+				const FVector BoxCenter(Box.GetCenter());
+				const FColor BoxColor(CellColors[0].CopyWithNewOpacity(CellOpacity).ToFColor(true));
+				const FVector CellPos = WorldPartitionTransform.TransformPosition(BoxCenter);
+				DrawDebugSolidBox(OwningWorld, Box, BoxColor, WorldPartitionTransform, false, -1.f, 255);				
+				DrawDebugBox(OwningWorld, CellPos, Box.GetExtent(), WorldPartitionTransform.GetRotation(), BoxColor.WithAlpha(255), false, -1.f, 255, 10.f);
+			}
+		}
+
+		// Draw X/Y Axis
+		if (DrawContext.GetDrawGridAxis())
+		{
+			DrawContext.PushDrawSegment(GridScreenBounds, WorldToScreen(FVector2D(-1638400.f, 0.f), false), WorldToScreen(FVector2D(1638400.f, 0.f), false), FLinearColor::Red, 3);
+			DrawContext.PushDrawSegment(GridScreenBounds, WorldToScreen(FVector2D(0.f, -1638400.f), false), WorldToScreen(FVector2D(0.f, 1638400.f), false), FLinearColor::Green, 3);
+		}
+
+		// Draw Grid Bounds
+		if (DrawContext.GetDrawGridBounds())
+		{
+			const FBox2D Bounds = GridScreenBounds.ExpandBy(FVector2D(10));
+			const FVector2D Size = GridScreenBounds.GetSize();
+			DrawContext.PushDrawBox(Bounds, GridScreenBounds.Min, GridScreenBounds.Min + FVector2D(Size.X, 0), GridScreenBounds.Max, GridScreenBounds.Min + FVector2D(0, Size.Y), FLinearColor::White, 1);
+		}
+
+		// Draw Streaming Sources
+		for (const FWorldPartitionStreamingSource& Source : Sources)
+		{
+			const FColor Color = Source.GetDebugColor();
+			// @todo_jfd
+			const FSoftObjectPath HLODLayer;
+			Source.ForEachShape(StreamingData->LoadingRange, StreamingData->Name, HLODLayer, true, [&Color, &WorldToScreen, &GridScreenBounds, &DrawContext, this](const FSphericalSector& Shape)
+			{
+				check(!Shape.IsNearlyZero())
+
+				// Spherical Sector
+				const FVector2D Center2D(FVector2D(Shape.GetCenter()));
+				const FSphericalSector::FReal Angle = Shape.GetAngle();
+				const int32 MaxSegments = FMath::Max(4, FMath::CeilToInt(64 * Angle / 360.f));
+				const float AngleIncrement = Angle / MaxSegments;
+				const FVector2D Axis = FVector2D(Shape.GetAxis());
+				const FVector Startup = FRotator(0, -0.5f * Angle, 0).RotateVector(Shape.GetScaledAxis());
+
+				FVector2D LineStart = FVector2D(Startup);
+				if (!Shape.IsSphere())
+				{
+					// Draw sector start axis
+					DrawContext.PushDrawSegment(GridScreenBounds, WorldToScreen(Center2D), WorldToScreen(Center2D + LineStart), Color, 2);
+				}
+				// Draw sector Arc
+				for (int32 i = 1; i <= MaxSegments; i++)
+				{
+					FVector2D LineEnd = FVector2D(FRotator(0, AngleIncrement * i, 0).RotateVector(Startup));
+					DrawContext.PushDrawSegment(GridScreenBounds, WorldToScreen(Center2D + LineStart), WorldToScreen(Center2D + LineEnd), Color, 2);
+					LineStart = LineEnd;
+				}
+				// If sphere, close circle, else draw sector end axis
+				DrawContext.PushDrawSegment(GridScreenBounds, WorldToScreen(Center2D + LineStart), WorldToScreen(Center2D + (Shape.IsSphere() ? FVector2D(Startup) : FVector2D::ZeroVector)), Color, 2);
+
+				// Draw direction vector
+				DrawContext.PushDrawSegment(GridScreenBounds, WorldToScreen(Center2D), WorldToScreen(Center2D + Axis * Shape.GetRadius()), Color, 2);
+			});
+		}
+	};
+
+	FBox GridsShapeBounds(ForceInit);
+	FBox2D GridsBounds(ForceInit);
+	int32 GridIndex = 0;
+	for (auto& [Name, StreamingDataList] : FilteredStreamingObjects)
+	{
+		for (const FWorldPartitionStreamingSource& Source : Sources)
+		{
+			// @todo_jfd
+			const FSoftObjectPath HLODLayer;
+			Source.ForEachShape(StreamingDataList[0]->LoadingRange, Name, HLODLayer, true, [&GridsShapeBounds](const FSphericalSector& Shape) { GridsShapeBounds += Shape.CalcBounds(); });
+		}
+
+		const FVector2D GridReferenceWorldPos = FVector2D(WorldRegion.GetCenter());
+		const FVector2D WorldRegionExtent = FVector2D(WorldRegion.GetExtent().GetMax());
+		const FVector2D GridScreenOffset = GridScreenInitialOffset + ((float)GridIndex * FVector2D(GridMaxScreenWidth, 0.f)) + GridScreenHalfExtent + FVector2D(GridScreenWidthShrinkSize * 0.5f);
+		const FVector2D WorldToScreenScale = GridScreenHalfExtent / WorldRegionExtent;
+		const FBox2D GridScreenBounds(GridScreenOffset - GridScreenHalfExtent, GridScreenOffset + GridScreenHalfExtent);
+
+		auto WorldToScreen = [&](const FVector2D& LocalWorldPos, bool bIsLocalWorldPos = true)
+		{
+			FVector2D GlocalPos = bIsLocalWorldPos ? LocalWPToGlobal.TransformPoint(LocalWorldPos) : LocalWorldPos;
+			return (WorldToScreenScale * (GlocalPos - GridReferenceWorldPos)) + GridScreenOffset;
+		};
+
+		for (const FRuntimePartitionStreamingData* StreamingData : StreamingDataList)
+		{
+			DrawStreamingData(StreamingData, LocalRegion, GridScreenBounds, WorldToScreen, DrawContext);
+		}
+
+		GridsBounds += GridScreenBounds;
+
+		if (DrawContext.IsDetailedMode())
+		{
+			FVector2D GridInfoPos = GridScreenOffset - GridScreenHalfExtent;
+			FWorldPartitionCanvasMultiLineText MultiLineText;
+			MultiLineText.Emplace(UWorld::RemovePIEPrefix(FPaths::GetBaseFilename(WorldPartition->GetPackage()->GetName())), FLinearColor::White);
+			FString GridInfoText = FString::Printf(TEXT("%s | %d m"), *Name.ToString(), int32(StreamingDataList[0]->LoadingRange * 0.01f));
+			MultiLineText.Emplace(GridInfoText, FLinearColor::Yellow);
+			FWorldPartitionCanvasMultiLineTextItem Item(GridInfoPos, MultiLineText);
+			DrawContext.PushDrawText(Item);
+			++GridIndex;
+		}
+	}
+
+	FBox2D DesiredWorldBounds(ForceInit);
+	if (GridsShapeBounds.IsValid)
+	{
+		// Convert to 2D
+		FBox2D GridsShapeBounds2D = FBox2D(FVector2D(GridsShapeBounds.Min.X, GridsShapeBounds.Min.Y), FVector2D(GridsShapeBounds.Max.X, GridsShapeBounds.Max.Y));
+		FVector2D CenterGlobalPos = LocalWPToGlobal.TransformPoint(GridsShapeBounds2D.GetCenter());
+		// Use max extent of X/Y
+		FVector2D Extent = FVector2D(GridsShapeBounds2D.GetExtent().GetMax());
+		// Transform to global space
+		GridsShapeBounds2D = FBox2D(CenterGlobalPos - Extent, CenterGlobalPos + Extent);
+		// Expand by 10% computed bounds
+		DesiredWorldBounds = GridsShapeBounds2D.ExpandBy(GridsShapeBounds2D.GetExtent() * 0.1f);
+	}
+	DrawContext.SetDesiredWorldBounds(DesiredWorldBounds);
+	DrawContext.SetUsedCanvasBounds(GridsBounds);
+
+	return true;
+}
+
+void UWorldPartitionRuntimeHashSet::Draw3D(const TArray<FWorldPartitionStreamingSource>& Sources) const
+{
+	return;
+}
\ No newline at end of file
diff --git a/Engine/Source/Runtime/Engine/Private/WorldPartition/RuntimeHashSet/WorldPartitionRuntimeHashSetHLODGeneration.cpp b/Engine/Source/Runtime/Engine/Private/WorldPartition/RuntimeHashSet/WorldPartitionRuntimeHashSetHLODGeneration.cpp
new file mode 100644
--- /dev/null
+++ b/Engine/Source/Runtime/Engine/Private/WorldPartition/RuntimeHashSet/WorldPartitionRuntimeHashSetHLODGeneration.cpp
@@ -0,0 +1,27 @@
+// Copyright Epic Games, Inc. All Rights Reserved.
+
+#include "WorldPartition/RuntimeHashSet/WorldPartitionRuntimeHashSet.h"
+#include "WorldPartition/RuntimeHashSet/RuntimePartition.h"
+
+#if WITH_EDITOR
+bool UWorldPartitionRuntimeHashSet::SupportsHLODs() const
+{
+	for (const FRuntimePartitionDesc& RuntimePartitionDesc : RuntimePartitions)
+	{
+		if (RuntimePartitionDesc.MainLayer)
+		{
+			if (RuntimePartitionDesc.MainLayer->SupportsHLODs())
+			{
+				return true;
+			}
+		}
+	}
+
+	return false;
+}
+
+bool UWorldPartitionRuntimeHashSet::SetupHLODActors(const IStreamingGenerationContext* StreamingGenerationContext, const UWorldPartition::FSetupHLODActorsParams& Params) const
+{
+	return false;
+}
+#endif
\ No newline at end of file
diff --git a/Engine/Source/Runtime/Engine/Private/WorldPartition/RuntimeHashSet/WorldPartitionRuntimeHashSetStreamingGeneration.cpp b/Engine/Source/Runtime/Engine/Private/WorldPartition/RuntimeHashSet/WorldPartitionRuntimeHashSetStreamingGeneration.cpp
new file mode 100644
--- /dev/null
+++ b/Engine/Source/Runtime/Engine/Private/WorldPartition/RuntimeHashSet/WorldPartitionRuntimeHashSetStreamingGeneration.cpp
@@ -0,0 +1,267 @@
+// Copyright Epic Games, Inc. All Rights Reserved.
+
+#include "WorldPartition/RuntimeHashSet/WorldPartitionRuntimeHashSet.h"
+#include "WorldPartition/RuntimeHashSet/RuntimePartition.h"
+#include "WorldPartition/RuntimeHashSet/RuntimePartitionPersistent.h"
+#include "WorldPartition/ContentBundle/ContentBundleDescriptor.h"
+#include "WorldPartition/WorldPartitionStreamingPolicy.h"
+#include "WorldPartition/DataLayer/DataLayersID.h"
+#include "Misc/ArchiveMD5.h"
+
+#if WITH_EDITOR
+bool UWorldPartitionRuntimeHashSet::GenerateStreaming(UWorldPartitionStreamingPolicy* StreamingPolicy, const IStreamingGenerationContext* StreamingGenerationContext, TArray<FString>* OutPackagesToGenerate)
+{
+	verify(Super::GenerateStreaming(StreamingPolicy, StreamingGenerationContext, OutPackagesToGenerate));
+
+	UWorldPartition* WorldPartition = GetOuterUWorldPartition();
+	UWorld* World = WorldPartition->GetWorld();
+	UWorld* OuterWorld = GetTypedOuter<UWorld>();
+	const bool bIsMainWorldPartition = (World == OuterWorld);
+
+	check(!PersistentPartitionDesc.Class);
+	PersistentPartitionDesc.Class = URuntimePartitionPersistent::StaticClass();
+	PersistentPartitionDesc.Name = NAME_PersistentLevel;
+	PersistentPartitionDesc.MainLayer = NewObject<URuntimePartition>(this, URuntimePartitionPersistent::StaticClass(), NAME_None);
+	PersistentPartitionDesc.MainLayer->Name = NAME_PersistentLevel;
+
+	//
+	// Split actor sets into their corresponding runtime partition implementation
+	//
+	TMap<FName, const FRuntimePartitionDesc*> NameToRuntimePartitionDescMap;
+
+	NameToRuntimePartitionDescMap.Add(NAME_None, &RuntimePartitions[0]);				// Actors with RuntimeGrid=None will be assigned to the default partition
+	NameToRuntimePartitionDescMap.Add(NAME_PersistentLevel, &PersistentPartitionDesc);	// Non-spatially loaded actors will be assigned to the persistent partition
+
+	for (const FRuntimePartitionDesc& RuntimePartitionDesc : RuntimePartitions)
+	{
+		NameToRuntimePartitionDescMap.Add(RuntimePartitionDesc.Name, &RuntimePartitionDesc);
+	}
+
+	TMap<URuntimePartition*, TArray<const IStreamingGenerationContext::FActorSetInstance*>> RuntimePartitionsToActorSetMap;
+	StreamingGenerationContext->ForEachActorSetInstance([this, &NameToRuntimePartitionDescMap, &RuntimePartitionsToActorSetMap](const IStreamingGenerationContext::FActorSetInstance& ActorSetInstance)
+	{
+		const TArray<FName> ActorSetRuntimeGrid = ActorSetInstance.bIsSpatiallyLoaded ? ParseGridName(ActorSetInstance.RuntimeGrid) : TArray<FName>({ NAME_PersistentLevel });
+
+		if (const FRuntimePartitionDesc** RuntimePartitionDesc = NameToRuntimePartitionDescMap.Find(ActorSetRuntimeGrid[0]))
+		{
+			RuntimePartitionsToActorSetMap.FindOrAdd((*RuntimePartitionDesc)->MainLayer).Add(&ActorSetInstance);
+		}
+	});
+
+	struct FCellDescInstance : public URuntimePartition::FCellDesc
+	{
+		FCellDescInstance(const URuntimePartition::FCellDesc& InCellDesc, URuntimePartition* InSourcePartition, const TArray<const UDataLayerInstance*>& InDataLayers)
+			: URuntimePartition::FCellDesc(InCellDesc)
+			, SourcePartition(InSourcePartition)
+			, DataLayers(InDataLayers)
+		{}
+
+		URuntimePartition* SourcePartition;
+		TArray<const UDataLayerInstance*> DataLayers;
+	};
+
+	//
+	// Generate runtime partitions streaming data
+	//
+	TArray<FCellDescInstance> RuntimeCellDescsInstances;
+	{
+		TSet<FName> CellDescsNames;
+		for (auto [RuntimePartition, ActorSetInstances] : RuntimePartitionsToActorSetMap)
+		{
+			// Gather runtime partition cell descs
+			URuntimePartition::FGenerateStreamingParams GenerateStreamingParams;
+			GenerateStreamingParams.ActorSetInstances = &ActorSetInstances;
+
+			URuntimePartition::FGenerateStreamingResult GenerateStreamingResult;
+
+			if (!RuntimePartition->GenerateStreaming(GenerateStreamingParams, GenerateStreamingResult))
+			{
+				return false;
+			}
+
+			// Split cell descs into data layers
+			for (const URuntimePartition::FCellDesc& RuntimeCellDesc : GenerateStreamingResult.RuntimeCellDescs)
+			{
+				TMap<FDataLayersID, FCellDescInstance> RuntimeCellDescsInstancesSet;
+
+				bool bCellNameExists;
+				CellDescsNames.Add(RuntimeCellDesc.Name, &bCellNameExists);
+				check(!bCellNameExists);
+
+				for (const IStreamingGenerationContext::FActorInstance& ActorInstance : RuntimeCellDesc.ActorInstances)
+				{
+					const FDataLayersID DataLayersID(ActorInstance.ActorSetInstance->DataLayers);
+					FCellDescInstance* CellDescInstance = RuntimeCellDescsInstancesSet.Find(DataLayersID);
+
+					if (!CellDescInstance)
+					{
+						CellDescInstance = &RuntimeCellDescsInstancesSet.Emplace(DataLayersID, FCellDescInstance(RuntimeCellDesc, RuntimePartition, ActorInstance.ActorSetInstance->DataLayers));
+						CellDescInstance->ActorInstances.Empty();
+					}
+
+					CellDescInstance->ActorInstances.Add(ActorInstance);
+				}
+
+				Algo::Transform(RuntimeCellDescsInstancesSet, RuntimeCellDescsInstances, [](const auto& Value) { return Value.Value; });
+			}
+		}
+	}
+
+	//
+	// Create and populate streaming object
+	//
+	check(RuntimeStreamingData.IsEmpty());
+
+	// Generate runtime cells
+	auto CreateRuntimeCellFromCellDesc = [this](const URuntimePartition::FCellDesc& CellDesc, const TArray<const UDataLayerInstance*>& DataLayers, TSubclassOf<UWorldPartitionRuntimeCell> CellClass, TSubclassOf<UWorldPartitionRuntimeCellData> CellDataClass)
+	{
+		FString CellObjectName;
+		FGuid CellGuid;
+		{
+			UWorld* OuterWorld = GetTypedOuter<UWorld>();
+			check(OuterWorld);
+
+			FString WorldName = FPackageName::GetShortName(OuterWorld->GetPackage());
+
+			CellObjectName = FString::Printf(TEXT("%s_%s"), *WorldName, *CellDesc.Name.ToString());
+
+			const FDataLayersID DataLayersID(DataLayers);
+			if (DataLayersID.GetHash())
+			{
+				CellObjectName += FString::Printf(TEXT("_d%X"), DataLayersID.GetHash());
+			}
+
+			if (CellDesc.ContentBundleID.IsValid())
+			{
+				CellObjectName += FString::Printf(TEXT("_c%s"), *UContentBundleDescriptor::GetContentBundleCompactString(CellDesc.ContentBundleID));
+			}
+
+			if (!IsRunningCookCommandlet() && OuterWorld->IsGameWorld())
+			{
+				FString SourceWorldPath;
+				FString InstancedWorldPath;
+				if (OuterWorld->GetSoftObjectPathMapping(SourceWorldPath, InstancedWorldPath))
+				{
+					const FTopLevelAssetPath SourceAssetPath(SourceWorldPath);
+					WorldName = FPackageName::GetShortName(SourceAssetPath.GetPackageName());
+						
+					InstancedWorldPath = UWorld::RemovePIEPrefix(InstancedWorldPath);
+
+					const FString SourcePackageName = SourceAssetPath.GetPackageName().ToString();
+					const FTopLevelAssetPath InstanceAssetPath(InstancedWorldPath);
+					const FString InstancePackageName = InstanceAssetPath.GetPackageName().ToString();
+
+					if (int32 Index = InstancePackageName.Find(SourcePackageName); Index != INDEX_NONE)
+					{
+						CellObjectName += FString::Printf(TEXT("_i%s"), *InstancePackageName.Mid(Index + SourcePackageName.Len()));
+					}
+				}
+			}
+
+			FArchiveMD5 ArMD5;
+			ArMD5 << CellObjectName;
+			CellGuid = ArMD5.GetGuidFromHash();
+			check(CellGuid.IsValid());
+		}
+
+		UWorldPartitionRuntimeCell* RuntimeCell = Super::CreateRuntimeCell(CellClass, CellDataClass, CellObjectName, TEXT(""));
+
+		RuntimeCell->SetIsAlwaysLoaded(!CellDesc.bIsSpatiallyLoaded);
+		RuntimeCell->SetDataLayers(DataLayers);
+		RuntimeCell->SetContentBundleUID(CellDesc.ContentBundleID);
+		RuntimeCell->SetPriority(CellDesc.Priority);
+		RuntimeCell->SetClientOnlyVisible(CellDesc.bClientOnlyVisible);
+		RuntimeCell->SetBlockOnSlowLoading(CellDesc.bBlockOnSlowStreaming);
+		RuntimeCell->SetIsHLOD(false);
+		RuntimeCell->SetGuid(CellGuid);
+
+		UWorldPartitionRuntimeCellDataSpatialHashSet* RuntimeCellDataHashSet = CastChecked<UWorldPartitionRuntimeCellDataSpatialHashSet>(RuntimeCell->RuntimeCellData);
+		RuntimeCellDataHashSet->DebugName = CellObjectName;
+		RuntimeCellDataHashSet->Level = CellDesc.Level;
+
+		return RuntimeCell;
+	};
+
+	TArray<UWorldPartitionRuntimeCell*> RuntimeCells;
+	TMap<URuntimePartition*, FRuntimePartitionStreamingData> RuntimePartitionsStreamingData;
+	for (const FCellDescInstance& CellDescInstance : RuntimeCellDescsInstances)
+	{
+		UWorldPartitionRuntimeCell* RuntimeCell = RuntimeCells.Emplace_GetRef(CreateRuntimeCellFromCellDesc(CellDescInstance, CellDescInstance.DataLayers, StreamingPolicy->GetRuntimeCellClass(), UWorldPartitionRuntimeCellDataSpatialHashSet::StaticClass()));
+		PopulateRuntimeCell(RuntimeCell, CellDescInstance.ActorInstances, nullptr);
+
+		// Override the cell bounds if the runtime partition provided one
+		if (CellDescInstance.Bounds.IsValid)
+		{
+			RuntimeCell->RuntimeCellData->ContentBounds = CellDescInstance.Bounds;
+		}
+
+		// Create partition streaming data
+		FRuntimePartitionStreamingData& StreamingData = RuntimePartitionsStreamingData.FindOrAdd(CellDescInstance.SourcePartition);
+
+		StreamingData.Name = CellDescInstance.SourcePartition->Name;
+		StreamingData.LoadingRange = CellDescInstance.SourcePartition->LoadingRange;
+
+		if (RuntimeCell->IsAlwaysLoaded())
+		{
+			StreamingData.NonStreamingCells.Add(RuntimeCell);
+		}
+		else
+		{
+			StreamingData.StreamingCells.Add(RuntimeCell);
+		}
+	}
+
+	//
+	// Finalize streaming object
+	//
+	for (auto& [Partition, StreamingData] : RuntimePartitionsStreamingData)
+	{
+		RuntimeStreamingData.Emplace(MoveTemp(StreamingData));
+	}
+
+	//
+	// Output generated packages
+	//
+	if (OutPackagesToGenerate)
+	{
+		for (UWorldPartitionRuntimeCell* RuntimeCell : RuntimeCells)
+		{
+			// Always loaded cell actors are transfered to World's Persistent Level (see UWorldPartitionRuntimeSpatialHash::PopulateGeneratorPackageForCook)
+			if (RuntimeCell->GetActorCount() && !RuntimeCell->IsAlwaysLoaded())
+			{
+				const FString PackageRelativePath = RuntimeCell->GetPackageNameToCreate();
+				check(!PackageRelativePath.IsEmpty());
+
+				OutPackagesToGenerate->Add(PackageRelativePath);
+
+				// Map relative package to StreamingCell for PopulateGeneratedPackageForCook/PopulateGeneratorPackageForCook/GetCellForPackage
+				PackagesToGenerateForCook.Add(PackageRelativePath, RuntimeCell);
+			}
+		}
+	}
+
+	return true;
+}
+
+void UWorldPartitionRuntimeHashSet::DumpStateLog(FHierarchicalLogArchive& Ar) const
+{
+	Super::DumpStateLog(Ar);
+
+	Ar.Printf(TEXT("----------------------------------------------------------------------------------------------------------------"));
+	Ar.Printf(TEXT("%s - Runtime Hash Set"), *GetWorld()->GetName());
+	Ar.Printf(TEXT("----------------------------------------------------------------------------------------------------------------"));
+
+	TArray<const UWorldPartitionRuntimeCell*> StreamingCells;
+	ForEachStreamingCells([&StreamingCells](const UWorldPartitionRuntimeCell* StreamingCell) { StreamingCells.Add(StreamingCell); return true; });
+				
+	StreamingCells.Sort([this](const UWorldPartitionRuntimeCell& A, const UWorldPartitionRuntimeCell& B) { return A.GetFName().LexicalLess(B.GetFName()); });
+
+	for (const UWorldPartitionRuntimeCell* StreamingCell : StreamingCells)
+	{
+		FHierarchicalLogArchive::FIndentScope CellIndentScope = Ar.PrintfIndent(TEXT("Content of Cell %s (%s)"), *StreamingCell->GetDebugName(), *StreamingCell->GetName());
+		StreamingCell->DumpStateLog(Ar);
+	}
+
+	Ar.Printf(TEXT(""));
+}
+#endif
\ No newline at end of file
diff --git a/Engine/Source/Runtime/Engine/Private/WorldPartition/WorldPartitionActorDescView.cpp b/Engine/Source/Runtime/Engine/Private/WorldPartition/WorldPartitionActorDescView.cpp
--- a/Engine/Source/Runtime/Engine/Private/WorldPartition/WorldPartitionActorDescView.cpp
+++ b/Engine/Source/Runtime/Engine/Private/WorldPartition/WorldPartitionActorDescView.cpp
@@ -16,6 +16,7 @@ FWorldPartitionActorDescView::FWorldPartitionActorDescView(const FWorldPartition
 	, bIsForcedNonSpatiallyLoaded(false)
 	, bIsForcedNoRuntimeGrid(false)
 	, bIsForcedNoDataLayers(false)	
+	, bIsForceNoHLODLayer(false)
 {}
 
 const FGuid& FWorldPartitionActorDescView::GetGuid() const
@@ -93,6 +94,11 @@ bool FWorldPartitionActorDescView::GetActorIsHLODRelevant() const
 
 FSoftObjectPath FWorldPartitionActorDescView::GetHLODLayer() const
 {
+	if (bIsForceNoHLODLayer)
+	{
+		return FSoftObjectPath();
+	}
+	
 	if (RuntimedHLODLayer.IsSet())
 	{
 		return RuntimedHLODLayer.GetValue();
@@ -301,6 +307,15 @@ void FWorldPartitionActorDescView::SetEditorReferences(const TArray<FGuid>& InEd
 	EditorReferences = InEditorReferences;
 }
 
+void FWorldPartitionActorDescView::SetForcedNoHLODLayer()
+{
+	if (!bIsForceNoHLODLayer)
+	{
+		bIsForceNoHLODLayer = true;
+		UE_LOG(LogWorldPartition, Verbose, TEXT("Actor '%s' HLOD layer invalidated"), *GetActorLabelOrName().ToString());
+	}
+}
+
 void FWorldPartitionActorDescView::SetRuntimeHLODLayer(const FSoftObjectPath& InHLODLayer)
 {
 	RuntimedHLODLayer = InHLODLayer;
diff --git a/Engine/Source/Runtime/Engine/Private/WorldPartition/WorldPartitionStreamingGeneration.cpp b/Engine/Source/Runtime/Engine/Private/WorldPartition/WorldPartitionStreamingGeneration.cpp
--- a/Engine/Source/Runtime/Engine/Private/WorldPartition/WorldPartitionStreamingGeneration.cpp
+++ b/Engine/Source/Runtime/Engine/Private/WorldPartition/WorldPartitionStreamingGeneration.cpp
@@ -1001,6 +1001,20 @@ PRAGMA_ENABLE_DEPRECATION_WARNINGS
 
 						NbErrorsDetected++;
 					}
+					
+					if (ActorDescView.GetHLODLayer().IsValid() && !IsValidHLODLayer(PerInstanceData.RuntimeGrid, ActorDescView.GetHLODLayer()))
+					{
+						if (PassType == EPassType::ErrorReporting)
+						{
+							ErrorHandler->OnInvalidHLODLayer(ActorDescView);
+						}
+						else
+						{
+							ActorDescView.SetForcedNoHLODLayer();
+						}
+
+						NbErrorsDetected++;
+					}
 				}
 
 				if (TSet<FGuid>* FilteredActors = ContainerFilteredActors.Find(ContainerCollectionInstanceDescriptor.ID))
@@ -1057,6 +1071,7 @@ public:
 		TMap<FGuid, const UActorDescContainer*> ActorGuidsToContainerMap;
 		TArray<TSubclassOf<AActor>> FilteredClasses;
 		TFunction<bool(FName)> IsValidGrid;
+		TFunction<bool(FName GridName, const FSoftObjectPath&)> IsValidHLODLayer;
 
 		FWorldPartitionStreamingGeneratorParams& SetWorldPartitionContext(const UWorldPartition* InWorldPartitionContext) { WorldPartitionContext = InWorldPartitionContext; return *this; }
 		FWorldPartitionStreamingGeneratorParams& SetModifiedActorsDescList(FActorDescList* InModifiedActorsDescList) { ModifiedActorsDescList = InModifiedActorsDescList; return *this; }
@@ -1065,6 +1080,7 @@ public:
 		FWorldPartitionStreamingGeneratorParams& SetActorGuidsToContainerMap(const TMap<FGuid, const UActorDescContainer*>& InActorGuidsToContainerMap) { ActorGuidsToContainerMap = InActorGuidsToContainerMap; return *this; }
 		FWorldPartitionStreamingGeneratorParams& SetFilteredClasses(const TArray<TSubclassOf<AActor>>& InFilteredClasses) { FilteredClasses = InFilteredClasses; return *this; }
 		FWorldPartitionStreamingGeneratorParams& SetIsValidGrid(TFunction<bool(FName)> InIsValidGrid) { IsValidGrid = InIsValidGrid; return *this; }
+		FWorldPartitionStreamingGeneratorParams& SetIsValidHLODLayer(TFunction<bool(FName GridName, const FSoftObjectPath&)> InIsValidHLODLayer) { IsValidHLODLayer = InIsValidHLODLayer; return *this; }
 
 		inline static FStreamingGenerationNullErrorHandler NullErrorHandler;
 	};
@@ -1075,6 +1091,7 @@ public:
 		, ModifiedActorsDescList(Params.ModifiedActorsDescList)
 		, FilteredClasses(Params.FilteredClasses)
 		, IsValidGrid(Params.IsValidGrid)
+		, IsValidHLODLayer(Params.IsValidHLODLayer)
 		, ErrorHandler(Params.ErrorHandler)
 		, ActorGuidsToContainerMap(Params.ActorGuidsToContainerMap)
 	{
@@ -1219,6 +1236,7 @@ private:
 	FActorDescList* ModifiedActorsDescList;
 	TArray<TSubclassOf<AActor>> FilteredClasses;
 	TFunction<bool(FName)> IsValidGrid;
+	TFunction<bool(FName, const FSoftObjectPath&)> IsValidHLODLayer;
 	IStreamingGenerationErrorHandler* ErrorHandler;
 
 	// Maps a level to its descriptor
@@ -1296,18 +1314,14 @@ bool UWorldPartition::GenerateContainerStreaming(const FGenerateStreamingParams&
 
 	UE_SCOPED_TIMER(*FString::Printf(TEXT("GenerateStreaming for '%s'"), *ContainerShortName), LogWorldPartition, Display);
 
-	// Dump state log
+	// Dump streaming generation log
 	TUniquePtr<FArchive> LogFileAr;
 	TUniquePtr<FHierarchicalLogArchive> HierarchicalLogAr;
 
-	bool bAllowStreamingGenerationLog = true;
-	IWorldPartitionEditorModule& WorldPartitionEditorModule = FModuleManager::LoadModuleChecked<IWorldPartitionEditorModule>("WorldPartitionEditor");
-	if (bIsPIE && !WorldPartitionEditorModule.GetEnableStreamingGenerationLogOnPIE())
-	{
-		bAllowStreamingGenerationLog = false;
-	}
+	const bool bIsStreamingGenerationLogAllowed = !bIsPIE || FModuleManager::LoadModuleChecked<IWorldPartitionEditorModule>("WorldPartitionEditor").GetEnableStreamingGenerationLogOnPIE();
+	const bool bIsStreamingGenerationLogRelevant = IsMainWorldPartition() && (!GIsBuildMachine || GIsAutomationTesting || IsRunningCookCommandlet());
 
-	if (bAllowStreamingGenerationLog && IsMainWorldPartition() && (!GIsBuildMachine || GIsAutomationTesting || IsRunningCookCommandlet()))
+	if (bIsStreamingGenerationLogAllowed && bIsStreamingGenerationLogRelevant)
 	{
 		TStringBuilder<256> StateLogSuffix;
 		StateLogSuffix += bIsPIE ? TEXT("PIE") : (IsRunningGame() ? TEXT("Game") : (IsRunningCookCommandlet() ? TEXT("Cook") : (GIsAutomationTesting ? TEXT("UnitTest") : TEXT("Manual"))));
@@ -1323,6 +1337,7 @@ bool UWorldPartition::GenerateContainerStreaming(const FGenerateStreamingParams&
 		.SetWorldPartitionContext(this)
 		.SetModifiedActorsDescList( ModifiedActorsDescList)
 		.SetIsValidGrid([this](FName GridName) { return RuntimeHash->IsValidGrid(GridName); })
+		.SetIsValidHLODLayer([this](FName GridName, const FSoftObjectPath& HLODLayerPath) { return RuntimeHash->IsValidHLODLayer(GridName, HLODLayerPath); })
 		.SetErrorHandler(StreamingGenerationErrorHandlerOverride ? (*StreamingGenerationErrorHandlerOverride)(ErrorHandler) : ErrorHandler)
 		.SetEnableStreaming(IsStreamingEnabled());
 
@@ -1383,7 +1398,8 @@ void UWorldPartition::SetupHLODActors(const FSetupHLODActorsParams& Params)
 			.SetErrorHandler(StreamingGenerationErrorHandlerOverride ? (*StreamingGenerationErrorHandlerOverride)(&LogErrorHandler) : &LogErrorHandler)
 			.SetEnableStreaming(IsStreamingEnabled())
 			.SetFilteredClasses({ AWorldPartitionHLOD::StaticClass() })
-			.SetIsValidGrid([this](FName GridName) { return RuntimeHash->IsValidGrid(GridName); });
+			.SetIsValidGrid([this](FName GridName) { return RuntimeHash->IsValidGrid(GridName); })
+			.SetIsValidHLODLayer([this](FName GridName, const FSoftObjectPath& HLODLayerPath) { return RuntimeHash->IsValidHLODLayer(GridName, HLODLayerPath); });
 
 		FWorldPartitionStreamingGenerator StreamingGenerator(StreamingGeneratorParams);
 		FStreamingGenerationActorDescCollection ActorDescCollection{InActorDescContainer};
@@ -1406,6 +1422,7 @@ void UWorldPartition::CheckForErrors(IStreamingGenerationErrorHandler* ErrorHand
 		.SetModifiedActorsDescList(&ModifiedActorDescList)
 		.SetErrorHandler(StreamingGenerationErrorHandlerOverride ? (*StreamingGenerationErrorHandlerOverride)(ErrorHandler) : ErrorHandler)
 		.SetIsValidGrid([this](FName GridName) { return RuntimeHash->IsValidGrid(GridName); })
+		.SetIsValidHLODLayer([this](FName GridName, const FSoftObjectPath& HLODLayerPath) { return RuntimeHash->IsValidHLODLayer(GridName, HLODLayerPath); })
 		.SetEnableStreaming(IsStreamingEnabled());
 
 	ForEachActorDescContainer([&StreamingGeneratorParams](const UActorDescContainer* InActorDescContainer)
@@ -1546,7 +1563,8 @@ void UWorldPartition::CheckForErrors(const FCheckForErrorsParams& Params)
 		.SetWorldPartitionContext(Params.ActorDescContainer->GetWorldPartition())
 		.SetModifiedActorsDescList(!Params.ActorDescContainer->IsTemplateContainer() ? &ModifiedActorDescList : nullptr)
 		.SetErrorHandler(StreamingGenerationErrorHandlerOverride ? (*StreamingGenerationErrorHandlerOverride)(Params.ErrorHandler) : Params.ErrorHandler)
-		.SetIsValidGrid([](FName GridName) { return true; })
+		.SetIsValidGrid([](FName) { return true; })
+		.SetIsValidHLODLayer([](FName, const FSoftObjectPath&) { return true; })
 		.SetEnableStreaming(Params.bEnableStreaming)
 		.SetActorGuidsToContainerMap(Params.ActorGuidsToContainerMap);
 
diff --git a/Engine/Source/Runtime/Engine/Public/WorldPartition/ErrorHandling/WorldPartitionStreamingGenerationErrorHandler.h b/Engine/Source/Runtime/Engine/Public/WorldPartition/ErrorHandling/WorldPartitionStreamingGenerationErrorHandler.h
--- a/Engine/Source/Runtime/Engine/Public/WorldPartition/ErrorHandling/WorldPartitionStreamingGenerationErrorHandler.h
+++ b/Engine/Source/Runtime/Engine/Public/WorldPartition/ErrorHandling/WorldPartitionStreamingGenerationErrorHandler.h
@@ -85,8 +85,16 @@ public:
 
 	virtual void OnLevelInstanceInvalidWorldAsset(const FWorldPartitionActorDescView& ActorDescView, FName WorldAsset, ELevelInstanceInvalidReason Reason) = 0;
 
+	/**
+	 * Called when an actor descriptor references another actor with a different set of actor filter.
+	 */
 	virtual void OnInvalidActorFilterReference(const FWorldPartitionActorDescView& ActorDescView, const FWorldPartitionActorDescView& ReferenceActorDescView) = 0;
 
+	/**
+	 * Called when an actor descriptor has an invalid HLOD layer.
+	 */
+	virtual void OnInvalidHLODLayer(const FWorldPartitionActorDescView& ActorDescView) = 0;
+
 	// Helpers
 	static ENGINE_API FString GetActorName(const FWorldPartitionActorDescView& ActorDescView);
 	static ENGINE_API FString GetFullActorName(const FWorldPartitionActorDescView& ActorDescView);
diff --git a/Engine/Source/Runtime/Engine/Public/WorldPartition/ErrorHandling/WorldPartitionStreamingGenerationLogErrorHandler.h b/Engine/Source/Runtime/Engine/Public/WorldPartition/ErrorHandling/WorldPartitionStreamingGenerationLogErrorHandler.h
--- a/Engine/Source/Runtime/Engine/Public/WorldPartition/ErrorHandling/WorldPartitionStreamingGenerationLogErrorHandler.h
+++ b/Engine/Source/Runtime/Engine/Public/WorldPartition/ErrorHandling/WorldPartitionStreamingGenerationLogErrorHandler.h
@@ -20,5 +20,6 @@ public:
 	ENGINE_API virtual void OnActorNeedsResave(const FWorldPartitionActorDescView& ActorDescView) override;
 	ENGINE_API virtual void OnLevelInstanceInvalidWorldAsset(const FWorldPartitionActorDescView& ActorDescView, FName WorldAsset, ELevelInstanceInvalidReason Reason) override;
 	ENGINE_API virtual void OnInvalidActorFilterReference(const FWorldPartitionActorDescView& ActorDescView, const FWorldPartitionActorDescView& ReferenceActorDescView) override;
+	ENGINE_API virtual void OnInvalidHLODLayer(const FWorldPartitionActorDescView& ActorDescView) override;
 };
 #endif
diff --git a/Engine/Source/Runtime/Engine/Public/WorldPartition/ErrorHandling/WorldPartitionStreamingGenerationNullErrorHandler.h b/Engine/Source/Runtime/Engine/Public/WorldPartition/ErrorHandling/WorldPartitionStreamingGenerationNullErrorHandler.h
--- a/Engine/Source/Runtime/Engine/Public/WorldPartition/ErrorHandling/WorldPartitionStreamingGenerationNullErrorHandler.h
+++ b/Engine/Source/Runtime/Engine/Public/WorldPartition/ErrorHandling/WorldPartitionStreamingGenerationNullErrorHandler.h
@@ -20,5 +20,6 @@ public:
 	virtual void OnActorNeedsResave(const FWorldPartitionActorDescView& ActorDescView) override {};
 	virtual void OnLevelInstanceInvalidWorldAsset(const FWorldPartitionActorDescView& ActorDescView, FName WorldAsset, ELevelInstanceInvalidReason Reason) override {};
 	virtual void OnInvalidActorFilterReference(const FWorldPartitionActorDescView& ActorDescView, const FWorldPartitionActorDescView& ReferenceActorDescView) override {};
+	virtual void OnInvalidHLODLayer(const FWorldPartitionActorDescView& ActorDescView) override {}
 };
 #endif
diff --git a/Engine/Source/Runtime/Engine/Public/WorldPartition/ErrorHandling/WorldPartitionStreamingGenerationProxyErrorHandler.h b/Engine/Source/Runtime/Engine/Public/WorldPartition/ErrorHandling/WorldPartitionStreamingGenerationProxyErrorHandler.h
--- a/Engine/Source/Runtime/Engine/Public/WorldPartition/ErrorHandling/WorldPartitionStreamingGenerationProxyErrorHandler.h
+++ b/Engine/Source/Runtime/Engine/Public/WorldPartition/ErrorHandling/WorldPartitionStreamingGenerationProxyErrorHandler.h
@@ -24,6 +24,7 @@ public:
 	ENGINE_API virtual void OnActorNeedsResave(const FWorldPartitionActorDescView& ActorDescView) override;
 	ENGINE_API virtual void OnLevelInstanceInvalidWorldAsset(const FWorldPartitionActorDescView& ActorDescView, FName WorldAsset, ELevelInstanceInvalidReason Reason) override;
 	ENGINE_API virtual void OnInvalidActorFilterReference(const FWorldPartitionActorDescView& ActorDescView, const FWorldPartitionActorDescView& ReferenceActorDescView) override;
+	ENGINE_API virtual void OnInvalidHLODLayer(const FWorldPartitionActorDescView& ActorDescView) override;
 
 protected:
 	IStreamingGenerationErrorHandler* InnerErrorHandler;
diff --git a/Engine/Source/Runtime/Engine/Public/WorldPartition/ErrorHandling/WorldPartitionStreamingGenerationTokenizedMessageErrorHandler.h b/Engine/Source/Runtime/Engine/Public/WorldPartition/ErrorHandling/WorldPartitionStreamingGenerationTokenizedMessageErrorHandler.h
--- a/Engine/Source/Runtime/Engine/Public/WorldPartition/ErrorHandling/WorldPartitionStreamingGenerationTokenizedMessageErrorHandler.h
+++ b/Engine/Source/Runtime/Engine/Public/WorldPartition/ErrorHandling/WorldPartitionStreamingGenerationTokenizedMessageErrorHandler.h
@@ -21,6 +21,7 @@ class ITokenizedMessageErrorHandler : public IStreamingGenerationErrorHandler
 	ENGINE_API virtual void OnActorNeedsResave(const FWorldPartitionActorDescView& ActorDescView) override;
 	ENGINE_API virtual void OnLevelInstanceInvalidWorldAsset(const FWorldPartitionActorDescView& ActorDescView, FName WorldAsset, ELevelInstanceInvalidReason Reason) override;
 	ENGINE_API virtual void OnInvalidActorFilterReference(const FWorldPartitionActorDescView& ActorDescView, const FWorldPartitionActorDescView& ReferenceActorDescView) override;
+	ENGINE_API virtual void OnInvalidHLODLayer(const FWorldPartitionActorDescView& ActorDescView) override;
 
 protected:
 	virtual void HandleTokenizedMessage(TSharedRef<FTokenizedMessage>&& ErrorMessage) = 0;
@@ -37,5 +38,4 @@ protected:
 private:
 	TArray<TSharedRef<FTokenizedMessage>> ErrorMessages;
 };
-
 #endif
diff --git a/Engine/Source/Runtime/Engine/Public/WorldPartition/RuntimeHashSet/RuntimePartition.h b/Engine/Source/Runtime/Engine/Public/WorldPartition/RuntimeHashSet/RuntimePartition.h
--- a/Engine/Source/Runtime/Engine/Public/WorldPartition/RuntimeHashSet/RuntimePartition.h
+++ b/Engine/Source/Runtime/Engine/Public/WorldPartition/RuntimeHashSet/RuntimePartition.h
@@ -33,10 +33,20 @@ public:
 		TArray<IStreamingGenerationContext::FActorInstance> ActorInstances;
 	};
 
+	struct FGenerateStreamingParams
+	{
+		const TArray<const IStreamingGenerationContext::FActorSetInstance*>* ActorSetInstances;
+	};
+
+	struct FGenerateStreamingResult
+	{
+		TArray<FCellDesc> RuntimeCellDescs;
+	};
+
 	virtual void SetDefaultValues();
 	virtual bool SupportsHLODs() const PURE_VIRTUAL(URuntimePartition::SupportsHLODs, return false;);
 	virtual bool IsValidGrid(FName InGridName) const PURE_VIRTUAL(URuntimePartition::IsValidGrid, return false;);
-	virtual bool GenerateStreaming(const TArray<const IStreamingGenerationContext::FActorSetInstance*>& InActorSetInstances, TArray<FCellDesc>& OutRuntimeCellDescs) PURE_VIRTUAL(URuntimePartition::GenerateStreaming, return false;);
+	virtual bool GenerateStreaming(const FGenerateStreamingParams& InParams, FGenerateStreamingResult& OutResult) PURE_VIRTUAL(URuntimePartition::GenerateStreaming, return false;);
 #endif
 
 #if WITH_EDITORONLY_DATA
diff --git a/Engine/Source/Runtime/Engine/Public/WorldPartition/RuntimeHashSet/RuntimePartitionLHGrid.h b/Engine/Source/Runtime/Engine/Public/WorldPartition/RuntimeHashSet/RuntimePartitionLHGrid.h
--- a/Engine/Source/Runtime/Engine/Public/WorldPartition/RuntimeHashSet/RuntimePartitionLHGrid.h
+++ b/Engine/Source/Runtime/Engine/Public/WorldPartition/RuntimeHashSet/RuntimePartitionLHGrid.h
@@ -20,7 +20,7 @@ public:
 	virtual void SetDefaultValues() override;
 	virtual bool SupportsHLODs() const override;
 	virtual bool IsValidGrid(FName GridName) const override;
-	virtual bool GenerateStreaming(const TArray<const IStreamingGenerationContext::FActorSetInstance*>& ActorSetInstances, TArray<FCellDesc>& OutRuntimeCellDescs) override;
+	virtual bool GenerateStreaming(const FGenerateStreamingParams& InParams, FGenerateStreamingResult& OutResult) override;
 	//~ End URuntimePartition interface
 #endif
 
diff --git a/Engine/Source/Runtime/Engine/Public/WorldPartition/RuntimeHashSet/RuntimePartitionLevelStreaming.h b/Engine/Source/Runtime/Engine/Public/WorldPartition/RuntimeHashSet/RuntimePartitionLevelStreaming.h
--- a/Engine/Source/Runtime/Engine/Public/WorldPartition/RuntimeHashSet/RuntimePartitionLevelStreaming.h
+++ b/Engine/Source/Runtime/Engine/Public/WorldPartition/RuntimeHashSet/RuntimePartitionLevelStreaming.h
@@ -15,7 +15,7 @@ public:
 	//~ Begin URuntimePartition interface
 	virtual bool SupportsHLODs() const override;
 	virtual bool IsValidGrid(FName GridName) const override;
-	virtual bool GenerateStreaming(const TArray<const IStreamingGenerationContext::FActorSetInstance*>& ActorSetInstances, TArray<FCellDesc>& OutRuntimeCellDescs) override;
+	virtual bool GenerateStreaming(const FGenerateStreamingParams& InParams, FGenerateStreamingResult& OutResult) override;
 	//~ End URuntimePartition interface
 #endif
 
diff --git a/Engine/Source/Runtime/Engine/Public/WorldPartition/RuntimeHashSet/RuntimePartitionPersistent.h b/Engine/Source/Runtime/Engine/Public/WorldPartition/RuntimeHashSet/RuntimePartitionPersistent.h
--- a/Engine/Source/Runtime/Engine/Public/WorldPartition/RuntimeHashSet/RuntimePartitionPersistent.h
+++ b/Engine/Source/Runtime/Engine/Public/WorldPartition/RuntimeHashSet/RuntimePartitionPersistent.h
@@ -15,7 +15,7 @@ public:
 	//~ Begin URuntimePartition interface
 	virtual bool SupportsHLODs() const override { return false; }
 	virtual bool IsValidGrid(FName GridName) const override { return true; }
-	virtual bool GenerateStreaming(const TArray<const IStreamingGenerationContext::FActorSetInstance*>& ActorSetInstances, TArray<FCellDesc>& OutRuntimeCellDescs) override;
+	virtual bool GenerateStreaming(const FGenerateStreamingParams& InParams, FGenerateStreamingResult& OutResult) override;
 	//~ End URuntimePartition interface
 #endif
 };
diff --git a/Engine/Source/Runtime/Engine/Public/WorldPartition/RuntimeHashSet/WorldPartitionRuntimeHashSet.h b/Engine/Source/Runtime/Engine/Public/WorldPartition/RuntimeHashSet/WorldPartitionRuntimeHashSet.h
--- a/Engine/Source/Runtime/Engine/Public/WorldPartition/RuntimeHashSet/WorldPartitionRuntimeHashSet.h
+++ b/Engine/Source/Runtime/Engine/Public/WorldPartition/RuntimeHashSet/WorldPartitionRuntimeHashSet.h
@@ -15,7 +15,7 @@ using FStaticSpatialIndexType = TStaticSpatialIndexRTree<TObjectPtr<UWorldPartit
 
 /** Holds settings for an HLOD layer for a particular partition class. */
 USTRUCT()
-struct FRuntimePartitionHLODSetup
+struct FRuntimePartitionHLODSetupLayer
 {
 	GENERATED_USTRUCT_BODY()
 
@@ -28,6 +28,21 @@ struct FRuntimePartitionHLODSetup
 	TObjectPtr<URuntimePartition> PartitionLayer;
 };
 
+/** Holds an HLOD setup for a particular partition class. */
+USTRUCT()
+struct FRuntimePartitionHLODSetup
+{
+	GENERATED_USTRUCT_BODY()
+
+	/** Associated HLOD Layer objects */
+	UPROPERTY(EditAnywhere, Category = RuntimeSettings, Meta = (EditCondition = "Class != nullptr", HideEditConditionToggle, NoResetToDefault, ForceInlineRow))
+	TObjectPtr<const UHLODLayer> HLODLayer;
+
+	/** HLOD setup, one for each layers in the hierarchy */
+	UPROPERTY(VisibleAnywhere, Category = RuntimeSettings, EditFixedSize, Meta = (EditCondition = "HLODLayer != nullptr", HideEditConditionToggle, ForceInlineRow))
+	TArray<FRuntimePartitionHLODSetupLayer> PartitionLayers;
+};
+
 /** Holds settings for a runtime partition instance. */
 USTRUCT()
 struct FRuntimePartitionDesc
@@ -47,12 +62,8 @@ struct FRuntimePartitionDesc
 	UPROPERTY(VisibleAnywhere, Category = RuntimeSettings, Instanced, Meta = (EditCondition = "Class != nullptr", HideEditConditionToggle, NoResetToDefault, TitleProperty = "Name"))
 	TObjectPtr<URuntimePartition> MainLayer;
 
-	/** Associated HLOD Layer object */
-	UPROPERTY(EditAnywhere, Category = RuntimeSettings, Meta = (EditCondition = "Class != nullptr", HideEditConditionToggle, NoResetToDefault, ForceInlineRow))
-	TObjectPtr<const UHLODLayer> HLODLayer;
-
 	/** HLOD setups used by this partition, one for each layers in the hierarchy */
-	UPROPERTY(VisibleAnywhere, Category = RuntimeSettings, EditFixedSize, Meta = (EditCondition = "HLODLayer != nullptr", HideEditConditionToggle, ForceInlineRow))
+	UPROPERTY(EditAnywhere, Category = RuntimeSettings, Meta = (EditCondition = "Class != nullptr", HideEditConditionToggle, ForceInlineRow))
 	TArray<FRuntimePartitionHLODSetup> HLODSetups;
 #endif
 
@@ -135,15 +146,18 @@ class UWorldPartitionRuntimeHashSet : public UWorldPartitionRuntimeHash
 	//~ End UObject Interface
 
 public:
-#if WITH_EDITOR
 	ENGINE_API virtual bool Draw2D(FWorldPartitionDraw2DContext& DrawContext) const override;
+	ENGINE_API virtual void Draw3D(const TArray<FWorldPartitionStreamingSource>& Sources) const override;
 
+#if WITH_EDITOR
 	// Streaming generation interface
 	ENGINE_API virtual void SetDefaultValues() override;
 	ENGINE_API virtual bool SupportsHLODs() const override;
+	ENGINE_API virtual bool SetupHLODActors(const IStreamingGenerationContext* StreamingGenerationContext, const UWorldPartition::FSetupHLODActorsParams& Params) const override;
 	ENGINE_API virtual bool GenerateStreaming(class UWorldPartitionStreamingPolicy* StreamingPolicy, const IStreamingGenerationContext* StreamingGenerationContext, TArray<FString>* OutPackagesToGenerate) override;
 	ENGINE_API virtual void FlushStreaming() override;
 	ENGINE_API virtual bool IsValidGrid(FName GridName) const;
+	ENGINE_API virtual bool IsValidHLODLayer(FName GridName, const FSoftObjectPath& HLODLayerPath) const;
 	ENGINE_API virtual TArray<UWorldPartitionRuntimeCell*> GetAlwaysLoadedCells() const override;
 	ENGINE_API virtual void DumpStateLog(FHierarchicalLogArchive& Ar) const override;
 
diff --git a/Engine/Source/Runtime/Engine/Public/WorldPartition/WorldPartitionActorDescView.h b/Engine/Source/Runtime/Engine/Public/WorldPartition/WorldPartitionActorDescView.h
--- a/Engine/Source/Runtime/Engine/Public/WorldPartition/WorldPartitionActorDescView.h
+++ b/Engine/Source/Runtime/Engine/Public/WorldPartition/WorldPartitionActorDescView.h
@@ -85,6 +85,7 @@ public:
 	ENGINE_API void SetEditorReferences(const TArray<FGuid>& InEditorReferences);
 	ENGINE_API void SetDataLayerInstanceNames(const TArray<FName>& InDataLayerInstanceNames);
 	
+	ENGINE_API void SetForcedNoHLODLayer();
 	ENGINE_API void SetRuntimeHLODLayer(const FSoftObjectPath& InHLODLayer);
 
 	ENGINE_API AActor* GetActor() const;
@@ -114,6 +115,7 @@ protected:
 	bool bIsForcedNonSpatiallyLoaded;
 	bool bIsForcedNoRuntimeGrid;
 	bool bIsForcedNoDataLayers;
+	bool bIsForceNoHLODLayer;
 	TOptional<TArray<FName>> ResolvedDataLayerInstanceNames;
 	TOptional<TArray<FName>> RuntimeDataLayerInstanceNames;
 	TOptional<TArray<FGuid>> RuntimeReferences;
diff --git a/Engine/Source/Runtime/Engine/Public/WorldPartition/WorldPartitionRuntimeHash.h b/Engine/Source/Runtime/Engine/Public/WorldPartition/WorldPartitionRuntimeHash.h
--- a/Engine/Source/Runtime/Engine/Public/WorldPartition/WorldPartitionRuntimeHash.h
+++ b/Engine/Source/Runtime/Engine/Public/WorldPartition/WorldPartitionRuntimeHash.h
@@ -91,6 +91,7 @@ class UWorldPartitionRuntimeHash : public UObject
 	ENGINE_API virtual void FlushStreaming();
 	virtual bool SetupHLODActors(const IStreamingGenerationContext* StreamingGenerationContext, const UWorldPartition::FSetupHLODActorsParams& Params) const { return false; }
 	virtual bool IsValidGrid(FName GridName) const { return false; }
+	virtual bool IsValidHLODLayer(FName GridName, const FSoftObjectPath& HLODLayerPath) const { return false; }
 	virtual void DrawPreview() const {}
 
 	virtual URuntimeHashExternalStreamingObjectBase* StoreToExternalStreamingObject(UObject* StreamingObjectOuter, FName StreamingObjectName) { return nullptr; }
diff --git a/Engine/Source/Runtime/Engine/Public/WorldPartition/WorldPartitionRuntimeSpatialHash.h b/Engine/Source/Runtime/Engine/Public/WorldPartition/WorldPartitionRuntimeSpatialHash.h
--- a/Engine/Source/Runtime/Engine/Public/WorldPartition/WorldPartitionRuntimeSpatialHash.h
+++ b/Engine/Source/Runtime/Engine/Public/WorldPartition/WorldPartitionRuntimeSpatialHash.h
@@ -278,6 +278,7 @@ public:
 	ENGINE_API virtual void FlushStreaming() override;
 	ENGINE_API virtual bool SetupHLODActors(const IStreamingGenerationContext* StreamingGenerationContext, const UWorldPartition::FSetupHLODActorsParams& Params) const override;
 	ENGINE_API virtual bool IsValidGrid(FName GridName) const override;
+	ENGINE_API virtual bool IsValidHLODLayer(FName GridName, const FSoftObjectPath& HLODLayerPath) const override { return true; }
 	ENGINE_API virtual void DrawPreview() const override;
 
 	ENGINE_API virtual URuntimeHashExternalStreamingObjectBase* StoreToExternalStreamingObject(UObject* StreamingObjectOuter, FName StreamingObjectName) override;
