Shader Compiler:  Proper fix for issues with corrupted pending job list.

The bug was caused by a second attempt to submit the same job at a higher priority, in between the time the job was created and added to the pending list.  The job ends up getting linked a second time by the creation logic, corrupting the linked lists.  The fix is simply to keep write lock active when submitting new jobs until the end of the function where the job is added to the pending list, so the addition of the job to the pending list completes before any other submits can occur.  Keeping the write lock active is more efficient than reacquiring a read lock anyway (there's no processing between the locks, and a cheap linked list insertion beyond that), so this actually saves performance.  For further potential safety, the other place where a read lock is used (forcing a job to run locally that was previously assigned to an external worker) was upgraded to write lock.  This is a rare case, which also doesn't happen during cooks.

Besides the fix, some changes were made to make the code overall less fragile for future maintenance.  LinkJobWithPriority / UnlinkJobWithPriority now exclusively handle setting PendingPriority on the job, rather than setting it externally.  PendingPriority should match the pending list the job is in (or be EShaderCompileJobPriority::None if not in a pending list), and NumPendingJobs should match the count of items in a pending linked list.  With this change, all three are updated in tandem in those functions.  FShaderJobCache::GetPendingJobs was changed to just call UnlinkJobWithPriority as well (updating NumPendingJobs in the regular way), rather than doing its own separate count adjustment outside the loop.  Asserts will now detect linking of jobs that were already linked, unlinking jobs that are not linked, or jobs mysteriously in the wrong pending queue, so we can hopefully detect bugs earlier when an invalid list operation occurs, rather than later when the pending queue is processed.

#rnx

diff --git a/Engine/Source/Runtime/Engine/Private/ShaderCompiler/ShaderCompiler.cpp b/Engine/Source/Runtime/Engine/Private/ShaderCompiler/ShaderCompiler.cpp
--- a/Engine/Source/Runtime/Engine/Private/ShaderCompiler/ShaderCompiler.cpp
+++ b/Engine/Source/Runtime/Engine/Private/ShaderCompiler/ShaderCompiler.cpp
@@ -733,6 +733,7 @@ private:
 	FORCEINLINE void LinkAfter(FShaderCommonCompileJob& Job, FShaderCommonCompileJob* After)
 	{
 		checkSlow(After != NULL);
+		check(Job.NextLink == nullptr && Job.PrevLink == nullptr);
 
 		Job.PrevLink = &After->NextLink;
 		Job.NextLink = *Job.PrevLink;
@@ -753,6 +754,8 @@ private:
 	 */
 	FORCEINLINE void LinkHeadAtomic(FShaderCommonCompileJob& Job, FShaderCommonCompileJob*& Head)
 	{
+		check(Job.NextLink == nullptr && Job.PrevLink == nullptr);
+
 		// It's important that PrevLink is set before the InterlockedExchange, as a subsequent Head pointer exchange could write
 		// another item and need to update PrevLink for this item before this function completes.
 		Job.PrevLink = &Head;
@@ -771,9 +774,7 @@ private:
 	 */
 	FORCEINLINE static void LinkTailAtomic(FShaderCommonCompileJob& Job, FShaderCommonCompileJob**& Tail)
 	{
-		// It's important that NextLink is set before the InterlockedExchange, as a subsequent Tail pointer exchange could write
-		// another item and need to update NextLink for this item before this function completes.
-		Job.NextLink = nullptr;
+		check(Job.NextLink == nullptr && Job.PrevLink == nullptr);
 
 		FShaderCommonCompileJob** OldTail = (FShaderCommonCompileJob**)FPlatformAtomics::InterlockedExchange((PTRINT*)&Tail, (PTRINT)&Job.NextLink);
 		Job.PrevLink = OldTail;
@@ -782,21 +783,25 @@ private:
 		*OldTail = (FShaderCommonCompileJob*)&Job;
 	}
 
-	/** Links job into linked list of given priority */
-	FORCEINLINE void LinkJobWithPriority(FShaderCommonCompileJob& Job, int32 PriorityIndex)
+	/** Links job into linked list with its given Priority */
+	FORCEINLINE void LinkJobWithPriority(FShaderCommonCompileJob& Job)
 	{
+		int32 PriorityIndex = (int32)Job.Priority;
 		check((uint32)PriorityIndex < (uint32)NumShaderCompileJobPriorities);
+		check(Job.PendingPriority == EShaderCompileJobPriority::None);
 		NumPendingJobs[PriorityIndex]++;
 #if UE_SHADERCOMPILER_FIFO_JOB_EXECUTION
 		LinkTailAtomic(Job, PendingJobsTail[PriorityIndex]);
 #else
 		LinkHeadAtomic(Job, PendingJobsHead[PriorityIndex]);
 #endif
+		Job.PendingPriority = Job.Priority;
 	}
 
-	/** Unlinks job from linked list of given priority */
-	FORCEINLINE void UnlinkJobWithPriority(FShaderCommonCompileJob& Job, int32 PriorityIndex)
+	/** Unlinks job from linked list with its current PendingPriority */
+	FORCEINLINE void UnlinkJobWithPriority(FShaderCommonCompileJob& Job)
 	{
+		int32 PriorityIndex = (int32)Job.PendingPriority;
 		check((uint32)PriorityIndex < (uint32)NumShaderCompileJobPriorities);
 		check(NumPendingJobs[PriorityIndex] > 0);
 		NumPendingJobs[PriorityIndex]--;
@@ -805,6 +810,7 @@ private:
 #else
 		Unlink(Job);
 #endif
+		Job.PendingPriority = EShaderCompileJobPriority::None;
 	}
 
 	/** From FShaderCompilingManager, guards access to FShaderMapCompileResults written in ProcessFinishedJob */
@@ -1262,14 +1268,12 @@ void FShaderJobCache::InternalSetPriority(FShaderCommonCompileJob* Job, EShaderC
 		if (Job->PendingPriority != EShaderCompileJobPriority::None)
 		{
 			// Job hasn't started yet, move it to the pending list for the new priority
-			const int32 PrevPriorityIndex = (int32)Job->PendingPriority;
 			check(Job->PendingPriority == Job->Priority);
-			UnlinkJobWithPriority(*Job, PrevPriorityIndex);
+			UnlinkJobWithPriority(*Job);
 
 			ensure(!ShaderCompiler::IsJobCacheEnabled() || Job->bInputHashSet);
-			LinkJobWithPriority(*Job, PriorityIndex);
 			Job->Priority = InPriority;
-			Job->PendingPriority = InPriority;
+			LinkJobWithPriority(*Job);
 			
 			return;
 		}
@@ -1294,11 +1298,9 @@ void FShaderJobCache::InternalSetPriority(FShaderCommonCompileJob* Job, EShaderC
 			TEXT("If shader jobs cache is enabled, cloned job should have the same input hash as the original, and it doesn't.")
 			);
 
-		// Need read lock to call LinkJobWithPriority (uses wait free LinkHeadAtomic / LinkTailAtomic -- works fine for multiple producer threads)
-		FReadScopeLock Locker(JobLock);
-		LinkJobWithPriority(*NewJob, PriorityIndex);
-		NewJob->PendingPriority = InPriority;
+		FWriteScopeLock Locker(JobLock);
 		NumOutstandingJobs.Increment();
+		LinkJobWithPriority(*NewJob);
 
 		//UE_LOG(LogShaderCompilers, Display, TEXT("Submitted duplicate 'ForceLocal' shader compile job to replace existing XGE job"));
 	}
@@ -1350,9 +1352,7 @@ int32 FShaderJobCache::RemoveAllPendingJobsWithId(uint32 InId)
 								// Add it as pending at the appropriate priority
 								GShaderCompilerStats->RegisterNewPendingJob(*DuplicateJob);
 
-								DuplicateJob->PendingPriority = DuplicateJob->Priority;
-
-								LinkJobWithPriority(*DuplicateJob, (int32)DuplicateJob->Priority);
+								LinkJobWithPriority(*DuplicateJob);
 							}
 
 							// DuplicateJob will be nullptr if there was no duplicate to promote
@@ -1360,8 +1360,8 @@ int32 FShaderJobCache::RemoveAllPendingJobsWithId(uint32 InId)
 						}
 					}
 
-					UnlinkJobWithPriority(Job, PriorityIndex);
-					Job.PendingPriority = EShaderCompileJobPriority::None;
+					check((int32)Job.PendingPriority == PriorityIndex);
+					UnlinkJobWithPriority(Job);
 					RemoveJob(&Job);
 					++NumRemoved;
 				}
@@ -1387,8 +1387,8 @@ int32 FShaderJobCache::RemoveAllPendingJobsWithId(uint32 InId)
 					}
 
 					// Duplicate jobs are in their own list, not one of the priority lists, so don't use UnlinkJobWithPriority
+					check(DuplicateJob->PendingPriority == EShaderCompileJobPriority::None);
 					Unlink(*DuplicateJob);
-					DuplicateJob->PendingPriority = EShaderCompileJobPriority::None;
 					RemoveJob(DuplicateJob);
 					++NumRemoved;
 
@@ -1417,9 +1417,13 @@ void FShaderJobCache::SubmitJob(FShaderCommonCompileJobPtr Job)
 
 	const int32 PriorityIndex = (int32)Job->Priority;
 	bool bNewJob = true;
+	bool bCacheEnabled = false;
+
 	// check caches unless we're running in validation mode (which runs _all_ jobs and compares hashes of outputs)
 	if (ShaderCompiler::IsJobCacheEnabled() && !ShaderCompiler::IsJobCacheDebugValidateEnabled())
 	{
+		bCacheEnabled = true;
+
 		const FShaderCommonCompileJob::FInputHash& InputHash = Job->GetInputHash();
 
 		const bool bCheckDDC = !(Job->bIsDefaultMaterial || Job->bIsGlobalShader);
@@ -1441,6 +1445,7 @@ void FShaderJobCache::SubmitJob(FShaderCommonCompileJobPtr Job)
 			// Need to release the lock before calling ProcessFinishedJob, as mentioned above (and it's also good for performance to 
 			// release the lock before the relatively costly "SerializeOutput" call).
 			JobLock.WriteUnlock();
+			bNewJob = false;
 
 			UE_LOG(LogShaderCompilers, UE_SHADERCACHE_LOG_LEVEL, TEXT("There is already a cached job with the ihash %s, processing the new one immediately."), *LexToString(InputHash));
 			FMemoryReaderView MemReader(*ExistingOutput);
@@ -1448,7 +1453,6 @@ void FShaderJobCache::SubmitJob(FShaderCommonCompileJobPtr Job)
 
 			// finish the job instantly
 			Job->TimeTaskSubmitJobsStall += ProcessFinishedJob(Job, true);
-			bNewJob = false;
 		}
 		else
 		{
@@ -1471,6 +1475,7 @@ void FShaderJobCache::SubmitJob(FShaderCommonCompileJobPtr Job)
 					*WaitListHead = Job;
 				}
 				AddDuplicateJob(Job);
+				JobLock.WriteUnlock();
 				bNewJob = false;
 			}
 			else
@@ -1478,8 +1483,6 @@ void FShaderJobCache::SubmitJob(FShaderCommonCompileJobPtr Job)
 				// track new jobs so we can dedupe them
 				JobData.JobInFlight = Job;
 			}
-
-			JobLock.WriteUnlock();
 		}
 	}
 
@@ -1489,14 +1492,15 @@ void FShaderJobCache::SubmitJob(FShaderCommonCompileJobPtr Job)
 		GShaderCompilerStats->RegisterNewPendingJob(*Job);
 		ensure(!ShaderCompiler::IsJobCacheEnabled() || Job->bInputHashSet);
 
-		Job->PendingPriority = Job->Priority;
+		// If cache is disabled, we skipped the code that grabs the write lock above, so we need to do it here, before modifying the pending queue
+		if (bCacheEnabled == false)
+		{
+			JobLock.WriteLock();
+		}
 
-		// Only need a read lock here, as LinkJobWithPriority does an atomic link operation safe across multiple producer threads
-		double StallStart = FPlatformTime::Seconds();
-		FReadScopeLock Locker(JobLock);
-		Job->TimeTaskSubmitJobsStall += FPlatformTime::Seconds() - StallStart;
+		LinkJobWithPriority(*Job);
 
-		LinkJobWithPriority(*Job, PriorityIndex);
+		JobLock.WriteUnlock();
 	}
 }
 
@@ -1692,7 +1696,7 @@ int32 FShaderJobCache::GetPendingJobs(EShaderCompilerWorkerType InWorkerType, ES
 	int32 NumPendingJobsOfPriority = 0;
 	{
 		FReadScopeLock Locker(JobLock);
-		NumPendingJobsOfPriority = NumPendingJobs[PriorityIndex];
+		NumPendingJobsOfPriority = NumPendingJobs[PriorityIndex].load();
 	}
 
 	if (NumPendingJobsOfPriority < MinNumJobs)
@@ -1704,7 +1708,7 @@ int32 FShaderJobCache::GetPendingJobs(EShaderCompilerWorkerType InWorkerType, ES
 	FWriteScopeLock Locker(JobLock);
 
 	// there was a time window before we checked and then acquired the write lock - make sure the number is still sufficient
-	NumPendingJobsOfPriority = NumPendingJobs[PriorityIndex];
+	NumPendingJobsOfPriority = NumPendingJobs[PriorityIndex].load();
 	if (NumPendingJobsOfPriority < MinNumJobs)
 	{
 		// Not enough jobs
@@ -1712,13 +1716,13 @@ int32 FShaderJobCache::GetPendingJobs(EShaderCompilerWorkerType InWorkerType, ES
 	}
 	
 	OutJobs.Reserve(OutJobs.Num() + FMath::Min(MaxNumJobs, NumPendingJobsOfPriority));
-	int32 NumJobs = FMath::Min(MaxNumJobs, NumPendingJobs[PriorityIndex].load());
+	int32 NumJobs = FMath::Min(MaxNumJobs, NumPendingJobsOfPriority);
 	FShaderCommonCompileJobIterator It(PendingJobsHead[PriorityIndex]);
 	// Randomize job selection by randomly skipping over jobs while traversing the list.
 	// Say, we need to pick 3 jobs out of 5 total. We can skip over 2 jobs in total, e.g. like this:
 	// pick one (4 more to go and we need to get 2 of 4), skip one (3 more to go, picking 2 out of 3), pick one (2 more to go, picking 1 of 2), skip one, pick one.
 	// It is possible that we won't skip at all and instead pick consequential jobs
-	int32 MaxJobsWeCanSkipOver = NumPendingJobs[PriorityIndex] - NumJobs;
+	int32 MaxJobsWeCanSkipOver = NumPendingJobsOfPriority - NumJobs;
 	for (int32 i = 0; i < NumJobs; ++i)
 	{
 		FShaderCommonCompileJob& Job = *It;
@@ -1731,14 +1735,9 @@ int32 FShaderJobCache::GetPendingJobs(EShaderCompilerWorkerType InWorkerType, ES
 
 		It.Next();
 
-		// Don't use UnlinkJobWithPriority here, as that updates NumPendingJobs, interfering with skipping logic -- NumPendingJobs is updated at the end of the loop
-#if UE_SHADERCOMPILER_FIFO_JOB_EXECUTION
-		UnlinkWithTail(Job, PendingJobsTail[PriorityIndex]);
-#else
-		Unlink(Job);
-#endif
+		check((int32)Job.PendingPriority == PriorityIndex);
+		UnlinkJobWithPriority(Job);
 
-		Job.PendingPriority = EShaderCompileJobPriority::None;
 		Job.CurrentWorker = InWorkerType;
 		OutJobs.Add(&Job);
 
@@ -1753,11 +1752,10 @@ int32 FShaderJobCache::GetPendingJobs(EShaderCompilerWorkerType InWorkerType, ES
 				--MaxJobsWeCanSkipOver;
 			}
 			checkf(MaxJobsWeCanSkipOver >= 0, TEXT("We skipped over too many jobs"));
-			checkf(MaxJobsWeCanSkipOver <= NumPendingJobs[PriorityIndex] - i, TEXT("Number of jobs to skip should stay less or equal than the number of nodes to go"));
+			checkf(MaxJobsWeCanSkipOver <= NumPendingJobsOfPriority - i, TEXT("Number of jobs to skip should stay less or equal than the number of nodes to go"));
 		}
 	}
 
-	NumPendingJobs[PriorityIndex] -= NumJobs;
 	return NumJobs;
 }
 
