SubObject Replication:
* Make sure new subobject bunches are always reliable even if the subobject created an object replicator earlier than expected.

PR #9850: Fixed an issue where Object requested an unreliable RPC before the ObjectReplicator was sent for the first time

diff --git a/Engine/Source/Runtime/Engine/Private/DataChannel.cpp b/Engine/Source/Runtime/Engine/Private/DataChannel.cpp
--- a/Engine/Source/Runtime/Engine/Private/DataChannel.cpp
+++ b/Engine/Source/Runtime/Engine/Private/DataChannel.cpp
@@ -5306,12 +5306,12 @@ bool UActorChannel::WriteSubObjectInBunch(UObject* Obj, FOutBunch& Bunch, FRepli
 		Connection->Driver->GuidCache->AssignNewNetGUID_Server( Obj );	//Make sure it gets a NetGUID so that it is now 'supported'
 	}
 
-	bool NewSubobject = false;
-	
 	FReplicationFlags ObjRepFlags = RepFlags;
 	TSharedRef<FObjectReplicator>& ObjectReplicator = !bFoundReplicator ? CreateReplicator(Obj) : *FoundReplicator;
 
-	if (!bFoundReplicator || bNewToReplay)
+	const bool bIsNewSubObject = (ObjectReplicator->bSentSubObjectCreation == false) || bNewToReplay;
+
+	if (bIsNewSubObject)
 	{
 		// This is the first time replicating this subobject
 		// This bunch should be reliable and we should always return true
@@ -5319,7 +5319,8 @@ bool UActorChannel::WriteSubObjectInBunch(UObject* Obj, FOutBunch& Bunch, FRepli
 		// (this will ensure the content header chunk is sent which is all we care about
 		// to spawn this on the client).
 		Bunch.bReliable = true;
-		NewSubobject = true;
+		ObjectReplicator->bSentSubObjectCreation = true;
+
 		if (UE::Net::bEnableNetInitialSubObjects)
 		{
 			ObjRepFlags.bNetInitial = true;
@@ -5327,7 +5328,7 @@ bool UActorChannel::WriteSubObjectInBunch(UObject* Obj, FOutBunch& Bunch, FRepli
 	}
 	UE_NET_TRACE_OBJECT_SCOPE(ObjectReplicator->ObjectNetGUID, Bunch, GetTraceCollector(Bunch), ENetTraceVerbosity::Trace);
 	bool bWroteSomething = ObjectReplicator.Get().ReplicateProperties(Bunch, ObjRepFlags);
-	if (NewSubobject && !bWroteSomething)
+	if (bIsNewSubObject && !bWroteSomething)
 	{
 		// Write empty payload to force object creation
 		FNetBitWriter EmptyPayload;
diff --git a/Engine/Source/Runtime/Engine/Private/DataReplication.cpp b/Engine/Source/Runtime/Engine/Private/DataReplication.cpp
--- a/Engine/Source/Runtime/Engine/Private/DataReplication.cpp
+++ b/Engine/Source/Runtime/Engine/Private/DataReplication.cpp
@@ -323,6 +323,7 @@ FObjectReplicator::FObjectReplicator()
 	, bOpenAckCalled(false)
 	, bForceUpdateUnmapped(false)
 	, bHasReplicatedProperties(false)
+	, bSentSubObjectCreation(false)
 	, bSupportsFastArrayDelta(false)
 	, bCanUseNonDirtyOptimization(false)
 	, bDirtyForReplay(true)
@@ -602,6 +603,7 @@ void FObjectReplicator::CleanUp()
 	Connection					= nullptr;
 	RemoteFunctions				= nullptr;
 	bHasReplicatedProperties	= false;
+	bSentSubObjectCreation		= false;
 	bOpenAckCalled				= false;
 
 	RepState = nullptr;
diff --git a/Engine/Source/Runtime/Engine/Public/Net/DataReplication.h b/Engine/Source/Runtime/Engine/Public/Net/DataReplication.h
--- a/Engine/Source/Runtime/Engine/Public/Net/DataReplication.h
+++ b/Engine/Source/Runtime/Engine/Public/Net/DataReplication.h
@@ -272,6 +272,9 @@ public:
 
 	/** Whether or not we've already replicated properties this frame. */
 	uint32 bHasReplicatedProperties : 1;
+	
+	/** Is this a subobject replicator and have we sent the creation data yet. */
+    uint32 bSentSubObjectCreation : 1;
 
 private:
 
