Foliage: Fix use cases where moving existing foliage with an active DataLayer context would make that moved foliage part of the DataLayer

- Add Optional FDataLayerEditorContext param to FActorPartitionGetParams
- Add a FForEachRelevantActorParams so that UActorPartitionSubSystem::ForEachRelevantActor can also be filtered by Guid/DataLayer Editor Context
- This will allow upcoming change of the Move tool so that it can move Foliage from one DL context to another.

#rnx

diff --git a/Engine/Source/Editor/FoliageEdit/Private/FoliageEdMode.cpp b/Engine/Source/Editor/FoliageEdit/Private/FoliageEdMode.cpp
--- a/Engine/Source/Editor/FoliageEdit/Private/FoliageEdMode.cpp
+++ b/Engine/Source/Editor/FoliageEdit/Private/FoliageEdMode.cpp
@@ -704,7 +704,7 @@ void FEdModeFoliage::ForEachFoliageInfo(UWorld* InWorld, const UFoliageType* Fol
 	if (UActorPartitionSubsystem* ActorPartitionSubsystem = InWorld->GetSubsystem<UActorPartitionSubsystem>())
 	{	
 		const FBox BrushSphereBounds(BrushSphere.Center - BrushSphere.W, BrushSphere.Center + BrushSphere.W);
-		ActorPartitionSubsystem->ForEachRelevantActor(AInstancedFoliageActor::StaticClass(), BrushSphereBounds, IFAOperation);
+		ActorPartitionSubsystem->ForEachRelevantActor(UActorPartitionSubsystem::FForEachRelevantActorParams(AInstancedFoliageActor::StaticClass(), BrushSphereBounds, IFAOperation));
 	}
 }
 
diff --git a/Engine/Source/Runtime/Engine/Private/ActorPartition/ActorPartitionSubsystem.cpp b/Engine/Source/Runtime/Engine/Private/ActorPartition/ActorPartitionSubsystem.cpp
--- a/Engine/Source/Runtime/Engine/Private/ActorPartition/ActorPartitionSubsystem.cpp
+++ b/Engine/Source/Runtime/Engine/Private/ActorPartition/ActorPartitionSubsystem.cpp
@@ -32,6 +32,16 @@ FActorPartitionGetParams::FActorPartitionGetParams(const TSubclassOf<APartitionA
 {
 }
 
+FActorPartitionGetParams::FActorPartitionGetParams(const TSubclassOf<APartitionActor>& InActorClass, bool bInCreate, TFunctionRef<void(APartitionActor*)> InActorCreated)
+	: ActorClass(InActorClass)
+	, bCreate(bInCreate)
+	, LocationHint(FVector::Zero())
+	, LevelHint(nullptr)
+	, bBoundsSearch(true)
+	, ActorCreatedCallback(InActorCreated)
+{
+}
+
 void FActorPartitionGridHelper::ForEachIntersectingCell(const TSubclassOf<APartitionActor>& InActorClass, const FBox& InBounds, ULevel* InLevel, TFunctionRef<bool(const UActorPartitionSubsystem::FCellCoord&, const FBox&)> InOperation, uint32 InGridSize)
 {
 	const uint32 GridSize = InGridSize > 0 ? InGridSize : InActorClass->GetDefaultObject<APartitionActor>()->GetDefaultGridSize(InLevel->GetWorld());
@@ -138,11 +148,11 @@ public:
 	}
 
 
-	void ForEachRelevantActor(const TSubclassOf<APartitionActor>& InActorClass, const FBox& IntersectionBounds, TFunctionRef<bool(APartitionActor*)>InOperation) const override
+	void ForEachRelevantActor(const UActorPartitionSubsystem::FForEachRelevantActorParams& InParams) const override
 	{
-		for (TActorIterator<APartitionActor> It(World, InActorClass); It; ++It)
+		for (TActorIterator<APartitionActor> It(World, InParams.ActorClass); It; ++It)
 		{
-			if (!InOperation(*It))
+			if (!InParams.Operation(*It))
 			{
 				return;
 			}
@@ -289,21 +299,25 @@ public:
 		return FoundActor;
 	}
 
-	void ForEachRelevantActor(const TSubclassOf<APartitionActor>& InActorClass, const FBox& IntersectionBounds, TFunctionRef<bool(APartitionActor*)>InOperation) const override
+	void ForEachRelevantActor(const UActorPartitionSubsystem::FForEachRelevantActorParams& InParams) const override
 	{
 		UActorPartitionSubsystem* ActorSubsystem = World->GetSubsystem<UActorPartitionSubsystem>();
-		FActorPartitionGridHelper::ForEachIntersectingCell(InActorClass, IntersectionBounds, World->PersistentLevel, [&ActorSubsystem, &InActorClass, &IntersectionBounds, &InOperation](const UActorPartitionSubsystem::FCellCoord& InCellCoord, const FBox& InCellBounds) {
+		FActorPartitionGridHelper::ForEachIntersectingCell(InParams.ActorClass, InParams.IntersectionBounds, World->PersistentLevel, [&ActorSubsystem, &InParams](const UActorPartitionSubsystem::FCellCoord& InCellCoord, const FBox& InCellBounds) {
 
-			if (InCellBounds.Intersect(IntersectionBounds))
+			if (InCellBounds.Intersect(InParams.IntersectionBounds))
 			{
-				const bool bCreate = false;
-				if (auto PartitionActor = ActorSubsystem->GetActor(InActorClass, InCellCoord, bCreate))
+				FActorPartitionGetParams GetParams(InParams.ActorClass);
+				GetParams.bBoundsSearch = true;
+				GetParams.GuidHint = InParams.GridGuid;
+				GetParams.DataLayerEditorContext = InParams.DataLayerEditorContext;
+
+				if (auto PartitionActor = ActorSubsystem->GetActor(GetParams, InCellCoord))
 				{
-					return InOperation(PartitionActor);
+					return InParams.Operation(PartitionActor);
 				}
 			}
 			return true;
-			});
+		});
 	}
 };
 
@@ -345,7 +359,15 @@ void UActorPartitionSubsystem::ForEachRelevantActor(const TSubclassOf<APartition
 {
 	if (ActorPartition)
 	{
-		ActorPartition->ForEachRelevantActor(InActorClass, IntersectionBounds, InOperation);
+		ActorPartition->ForEachRelevantActor(FForEachRelevantActorParams(InActorClass, IntersectionBounds, InOperation));
+	}
+}
+
+void UActorPartitionSubsystem::ForEachRelevantActor(const UActorPartitionSubsystem::FForEachRelevantActorParams& InParams) const
+{
+	if (ActorPartition)
+	{
+		ActorPartition->ForEachRelevantActor(InParams);
 	}
 }
 
@@ -390,34 +412,53 @@ void UActorPartitionSubsystem::UninitializeActorPartition()
 APartitionActor* UActorPartitionSubsystem::GetActor(const FActorPartitionGetParams& GetParams)
 {
 	FCellCoord CellCoord = ActorPartition->GetActorPartitionHash(GetParams);
-	return GetActor(GetParams.ActorClass, CellCoord, GetParams.bCreate, GetParams.GuidHint, GetParams.GridSize, GetParams.bBoundsSearch, GetParams.ActorCreatedCallback);
+	return GetActor(GetParams, CellCoord);
 }
 
 APartitionActor* UActorPartitionSubsystem::GetActor(const TSubclassOf<APartitionActor>& InActorClass, const FCellCoord& InCellCoords, bool bInCreate, const FGuid& InGuid, uint32 InGridSize, bool bInBoundsSearch, TFunctionRef<void(APartitionActor*)> InActorCreated)
 {
-	const uint32 GridSize = InGridSize > 0 ? InGridSize : InActorClass->GetDefaultObject<APartitionActor>()->GetDefaultGridSize(GetWorld());
+	FActorPartitionGetParams GetParams(InActorClass, bInCreate, nullptr, FVector::Zero(), InGridSize, InGuid, bInBoundsSearch, InActorCreated);
+	return GetActor(GetParams, InCellCoords);
+}
+
+APartitionActor* UActorPartitionSubsystem::GetActor(const FActorPartitionGetParams& GetParams, const FCellCoord& InCellCoords)
+{
+	const uint32 GridSize = GetParams.GridSize > 0 ? GetParams.GridSize : GetParams.ActorClass->GetDefaultObject<APartitionActor>()->GetDefaultGridSize(GetWorld());
 	
 	UWorld* World = GetWorld();
 	const UDataLayerManager* DataLayerManager = UDataLayerManager::GetDataLayerManager(World);
+	uint32 DataLayerEditorContextHash = FDataLayerEditorContext::EmptyHash;
+	TArray<const UDataLayerInstance*> DataLayerInstances;
 	
-	auto WrappedInActorCreated = [DataLayerManager, InActorCreated](APartitionActor* PartitionActor)
+	if (DataLayerManager)
 	{
-		if(DataLayerManager)
+		if (GetParams.DataLayerEditorContext.IsSet())
 		{
-			for (UDataLayerInstance* DataLayer : DataLayerManager->GetActorEditorContextDataLayers())
-			{
-				PartitionActor->AddDataLayer(DataLayer);
-			}
+			DataLayerInstances = DataLayerManager->GetDataLayerInstances(GetParams.DataLayerEditorContext->GetDataLayerInstanceNames());
+			DataLayerEditorContextHash = GetParams.DataLayerEditorContext->GetHash();
+		}
+		else
+		{
+			DataLayerInstances.Append(DataLayerManager->GetActorEditorContextDataLayers());
+			DataLayerEditorContextHash = DataLayerManager->GetDataLayerEditorContextHash();
+		}
+	}
+
+	auto WrappedInActorCreated = [&DataLayerInstances, ActorCreated = GetParams.ActorCreatedCallback](APartitionActor* PartitionActor)
+	{
+		for (const UDataLayerInstance* DataLayer : DataLayerInstances)
+		{
+			PartitionActor->AddDataLayer(DataLayer);
 		}
-		InActorCreated(PartitionActor);
+		ActorCreated(PartitionActor);
 	};
 
-	FActorPartitionIdentifier ActorPartitionId(InActorClass, InGuid, DataLayerManager ? DataLayerManager->GetDataLayerEditorContextHash() : FDataLayerEditorContext::EmptyHash);
+	FActorPartitionIdentifier ActorPartitionId(GetParams.ActorClass, GetParams.GuidHint, DataLayerEditorContextHash);
 	TMap<FActorPartitionIdentifier, TWeakObjectPtr<APartitionActor>>* ActorsPerId = PartitionedActors.Find(InCellCoords);
 	APartitionActor* FoundActor = nullptr;
 	if (!ActorsPerId)
 	{
-		FoundActor = ActorPartition->GetActor(ActorPartitionId, bInCreate, InCellCoords, GridSize, bInBoundsSearch, WrappedInActorCreated);
+		FoundActor = ActorPartition->GetActor(ActorPartitionId, GetParams.bCreate, InCellCoords, GridSize, GetParams.bBoundsSearch, WrappedInActorCreated);
 		if (FoundActor)
 		{
 			PartitionedActors.Add(InCellCoords).Add(ActorPartitionId, FoundActor);
@@ -428,7 +469,7 @@ APartitionActor* UActorPartitionSubsystem::GetActor(const TSubclassOf<APartition
 		TWeakObjectPtr<APartitionActor>* ActorPtr = ActorsPerId->Find(ActorPartitionId);
 		if (!ActorPtr || !ActorPtr->IsValid())
 		{
-			FoundActor = ActorPartition->GetActor(ActorPartitionId, bInCreate, InCellCoords, GridSize, bInBoundsSearch, WrappedInActorCreated);
+			FoundActor = ActorPartition->GetActor(ActorPartitionId, GetParams.bCreate, InCellCoords, GridSize, GetParams.bBoundsSearch, WrappedInActorCreated);
 			if (FoundActor)
 			{
 				if (!ActorPtr)
diff --git a/Engine/Source/Runtime/Engine/Public/ActorPartition/ActorPartitionSubsystem.h b/Engine/Source/Runtime/Engine/Public/ActorPartition/ActorPartitionSubsystem.h
--- a/Engine/Source/Runtime/Engine/Public/ActorPartition/ActorPartitionSubsystem.h
+++ b/Engine/Source/Runtime/Engine/Public/ActorPartition/ActorPartitionSubsystem.h
@@ -7,6 +7,7 @@
 #include "Misc/HashBuilder.h"
 #include "Misc/Guid.h"
 #include "ActorPartition/PartitionActor.h"
+#include "WorldPartition/DataLayer/DataLayerEditorContext.h"
 #include "ActorPartitionSubsystem.generated.h"
 
 class FBaseActorPartition;
@@ -19,6 +20,7 @@ class UWorldPartition;
 struct FActorPartitionGetParams
 {
 	ENGINE_API FActorPartitionGetParams(const TSubclassOf<APartitionActor>& InActorClass, bool bInCreate, ULevel* InLevelHint, const FVector& InLocationHint, uint32 InGridSize = 0, const FGuid& InGuidHint = FGuid(), bool bInBoundsSearch = true, TFunctionRef<void(APartitionActor*)> InActorCreated = [](APartitionActor*) {});
+	ENGINE_API FActorPartitionGetParams(const TSubclassOf<APartitionActor>& InActorClass, bool bInCreate = false, TFunctionRef<void(APartitionActor*)> InActorCreated = [](APartitionActor*) {});
 
 	/* Class of Actor we are getting from the subsystem. */
 	TSubclassOf<APartitionActor> ActorClass;
@@ -43,6 +45,9 @@ struct FActorPartitionGetParams
 
 	/* If set, a callback to use if an actor is created. */
 	TFunctionRef<void(APartitionActor*)> ActorCreatedCallback;
+
+	/* Optional DataLayer editor context. If none provided, the current editor context will be used */
+	TOptional<FDataLayerEditorContext> DataLayerEditorContext;
 };
 
 /**
@@ -154,7 +159,8 @@ public:
 	};
 
 #if WITH_EDITOR
-	ENGINE_API APartitionActor* GetActor(const FActorPartitionGetParams& GetParam);
+	ENGINE_API APartitionActor* GetActor(const FActorPartitionGetParams& GetParams, const FCellCoord& InCellCoords);
+	ENGINE_API APartitionActor* GetActor(const FActorPartitionGetParams& GetParams);
 
 	/**
 	 * Returns a matching actor based on the parameters being provided
@@ -171,7 +177,36 @@ public:
 	ENGINE_API virtual void Initialize(FSubsystemCollectionBase& Collection) override;
 	ENGINE_API virtual void Deinitialize() override;
 
+	struct FForEachRelevantActorParams
+	{
+		FForEachRelevantActorParams(const TSubclassOf<APartitionActor>& InActorClass, const FBox& InIntersectionBounds, TFunctionRef<bool(APartitionActor*)> InOperation)
+			: ActorClass(InActorClass)
+			, IntersectionBounds(InIntersectionBounds)
+			, Operation(InOperation)
+		{
+
+		}
+
+		/* Actor class to search for */
+		TSubclassOf<APartitionActor> ActorClass;
+
+		/* Bounds to search */
+		FBox IntersectionBounds;
+		
+		/* Search for a specific grid guid */
+		FGuid GridGuid;
+
+		/* Search for a specific data layer editor context */
+		TOptional<FDataLayerEditorContext> DataLayerEditorContext;
+		
+		TFunctionRef<bool(APartitionActor*)> Operation;
+	};
+
+	UE_DEPRECATED(5.4, "Use ForEachRelevantActor taking FForEachRelevantActorParams parameter instead")
 	ENGINE_API void ForEachRelevantActor(const TSubclassOf<APartitionActor>& InActorClass, const FBox& IntersectionBounds, TFunctionRef<bool(APartitionActor*)>InOperation) const;
+	
+	ENGINE_API void ForEachRelevantActor(const FForEachRelevantActorParams& InParams) const;
+
 #endif
 	ENGINE_API bool IsLevelPartition() const;
 
@@ -203,7 +238,7 @@ public:
 
 	virtual UActorPartitionSubsystem::FCellCoord GetActorPartitionHash(const FActorPartitionGetParams& GetParams) const = 0;
 	virtual APartitionActor* GetActor(const FActorPartitionIdentifier& InActorPartitionId, bool bInCreate, const UActorPartitionSubsystem::FCellCoord& InCellCoord, uint32 InGridSize, bool bInBoundsSearch, TFunctionRef<void(APartitionActor*)> InActorCreated) = 0;
-	virtual void ForEachRelevantActor(const TSubclassOf<APartitionActor>& InActorClass, const FBox& IntersectionBounds, TFunctionRef<bool(APartitionActor*)>InOperation) const = 0;
+	virtual void ForEachRelevantActor(const UActorPartitionSubsystem::FForEachRelevantActorParams& InParams) const = 0;
 
 	DECLARE_EVENT_OneParam(FBaseActorPartition, FOnActorPartitionHashInvalidated, const UActorPartitionSubsystem::FCellCoord&);
 	FOnActorPartitionHashInvalidated& GetOnActorPartitionHashInvalidated() { return OnActorPartitionHashInvalidated; }
diff --git a/Engine/Source/Runtime/Foliage/Private/InstancedFoliage.cpp b/Engine/Source/Runtime/Foliage/Private/InstancedFoliage.cpp
--- a/Engine/Source/Runtime/Foliage/Private/InstancedFoliage.cpp
+++ b/Engine/Source/Runtime/Foliage/Private/InstancedFoliage.cpp
@@ -3107,20 +3107,19 @@ const FFoliageInfo* AInstancedFoliageActor::FindInfo(const UFoliageType* InType)
 
 
 #if WITH_EDITOR
-AInstancedFoliageActor* AInstancedFoliageActor::Get(UWorld* InWorld, bool bCreateIfNone, ULevel* InLevelHint, const FVector& InLocationHint)
+AInstancedFoliageActor* AInstancedFoliageActor::Get(UWorld* InWorld, bool bCreateIfNone, ULevel* InLevelHint, const FVector& InLocationHint, const TOptional<FDataLayerEditorContext>& InDataLayerEditorContext)
 {
 	UActorPartitionSubsystem* ActorPartitionSubsystem = InWorld->GetSubsystem<UActorPartitionSubsystem>();
-
-	return Cast<AInstancedFoliageActor>(
-		ActorPartitionSubsystem->GetActor(
-			FActorPartitionGetParams(
-				AInstancedFoliageActor::StaticClass(), 
-				bCreateIfNone, 
-				InLevelHint, 
-				InLocationHint
-			)
-		)
+	FActorPartitionGetParams GetParams(
+		AInstancedFoliageActor::StaticClass(),
+		bCreateIfNone,
+		InLevelHint,
+		InLocationHint
 	);
+
+	GetParams.DataLayerEditorContext = InDataLayerEditorContext;
+
+	return Cast<AInstancedFoliageActor>(ActorPartitionSubsystem->GetActor(GetParams));
 }
 
 
@@ -3215,6 +3214,7 @@ namespace FoliagePartitioningUtils
 
 		bool bMovedInstances = true;
 		AInstancedFoliageActor* TargetIFA = nullptr;
+		FDataLayerEditorContext DataLayerEditorContext(SourceIFA->GetWorld(), SourceIFA->GetDataLayerInstanceNames());
 
 		TSet<int32>* InstanceSet = GetInstanceSet();
 		while (InstanceSet && InstanceSet->Num() > 0 && bMovedInstances)
@@ -3226,7 +3226,7 @@ namespace FoliagePartitioningUtils
 			for (int32 InstanceIdx : *InstanceSet)
 			{
 				FFoliageInstance& Instance = FoliageInfo.Instances[InstanceIdx];
-				AInstancedFoliageActor* NewIFA = AInstancedFoliageActor::Get(SourceIFA->GetWorld(), true, SourceIFA->GetLevel(), Instance.Location);
+				AInstancedFoliageActor* NewIFA = AInstancedFoliageActor::Get(SourceIFA->GetWorld(), true, SourceIFA->GetLevel(), Instance.Location, DataLayerEditorContext);
 				if ((TargetIFA == nullptr || TargetIFA == NewIFA) && NewIFA != SourceIFA)
 				{
 					TargetIFA = NewIFA;
@@ -3566,6 +3566,8 @@ void AInstancedFoliageActor::MoveInstancesToNewComponent(UPrimitiveComponent* In
 	// If Modify was called on this IFA
 	bool bModified = false;
 	
+	FDataLayerEditorContext DataLayerEditorContext(IFAWorld, GetDataLayerInstanceNames());
+
 	for (auto& Pair : FoliageInfos)
 	{			
 		FFoliageInfo& Info = *Pair.Value;
@@ -3588,7 +3590,7 @@ void AInstancedFoliageActor::MoveInstancesToNewComponent(UPrimitiveComponent* In
 				{
 					FFoliageInstance& InstanceToMove = Info.Instances[InstanceIndex];
 					const bool bCreate = true;
-					if (AInstancedFoliageActor* TargetIFA = AInstancedFoliageActor::Get(IFAWorld, bCreate, IFALevel, InstanceToMove.Location))
+					if (AInstancedFoliageActor* TargetIFA = AInstancedFoliageActor::Get(IFAWorld, bCreate, IFALevel, InstanceToMove.Location, DataLayerEditorContext))
 					{
 						// Call Modify only once
 						if (!bModified)
@@ -5660,7 +5662,10 @@ void AInstancedFoliageActor::HandleFoliageInstancePostMove(const FFoliageInstanc
 	// Verify, and re-instance any existing typed elements if required
 	{
 		const FFoliageInstance& FoliageInstance = InstanceId.GetInstanceChecked();
-		AInstancedFoliageActor* TargetIFA = AInstancedFoliageActor::Get(GetWorld(), /*bCreateIfNone*/true, GetLevel(), FoliageInstance.Location);
+
+		FDataLayerEditorContext DataLayerEditorContext(GetWorld(), GetDataLayerInstanceNames());
+
+		AInstancedFoliageActor* TargetIFA = AInstancedFoliageActor::Get(GetWorld(), /*bCreateIfNone*/true, GetLevel(), FoliageInstance.Location, DataLayerEditorContext);
 		if (TargetIFA != this)
 		{
 			// Add the new instance first, as we need both instances available for the re-instance to work
diff --git a/Engine/Source/Runtime/Foliage/Public/InstancedFoliageActor.h b/Engine/Source/Runtime/Foliage/Public/InstancedFoliageActor.h
--- a/Engine/Source/Runtime/Foliage/Public/InstancedFoliageActor.h
+++ b/Engine/Source/Runtime/Foliage/Public/InstancedFoliageActor.h
@@ -13,6 +13,7 @@
 #include "InstancedFoliage.h"
 #include "InstancedFoliageCustomVersion.h"
 #include "ISMPartition/ISMPartitionActor.h"
+#include "WorldPartition/DataLayer/DataLayerEditorContext.h"
 
 #include "InstancedFoliageActor.generated.h"
 
@@ -170,8 +171,9 @@ public:
 	 * @param bCreateIfNone					Create if doesnt already exist
 	 * @param InLevelHint					Level hint for foliage instance creation
 	 * @param InLocationHint				Location hint for foliage instance creation
+	 * @param InDataLayerEditorContext		Optional DataLayer editor context. If none provided, the current editor context will be used
 	 */
-	static FOLIAGE_API AInstancedFoliageActor* Get(UWorld* InWorld, bool bCreateIfNone, ULevel* InLevelHint = nullptr, const FVector& InLocationHint = FVector(ForceInitToZero));
+	static FOLIAGE_API AInstancedFoliageActor* Get(UWorld* InWorld, bool bCreateIfNone, ULevel* InLevelHint = nullptr, const FVector& InLocationHint = FVector(ForceInitToZero), const TOptional<FDataLayerEditorContext>& InDataLayerEditorContext = TOptional<FDataLayerEditorContext>());
 
 	static FOLIAGE_API AInstancedFoliageActor* GetDefault(UWorld* InWorld);
 
