Optimize ULevel::AddLoadedActors / RemoveLoadedActors.
Saves over 5 seconds of loading/unloading time in a 32k actors map.

diff --git a/Engine/Source/Runtime/Engine/Private/Level.cpp b/Engine/Source/Runtime/Engine/Private/Level.cpp
--- a/Engine/Source/Runtime/Engine/Private/Level.cpp
+++ b/Engine/Source/Runtime/Engine/Private/Level.cpp
@@ -813,27 +813,40 @@ void ULevel::AddLoadedActors(const TArray<AActor*>& ActorList, const FTransform*
 		check(Actor->GetLevel() == this);
 		check(IsValidChecked(Actor));
 
-		int32 ActorIndex;
-		if (!Actors.Find(Actor, ActorIndex))
-		{
-			Actors.Add(Actor);
-			ActorsForGC.Add(Actor);
-			ActorsQueue.Add(Actor);
+		ActorsQueue.Add(Actor);
 
-			// Handle child actors
-			Actor->ForEachComponent<UChildActorComponent>(false, [this, &QueueActor](UChildActorComponent* ChildActorComponent)
+		// Handle child actors
+		Actor->ForEachComponent<UChildActorComponent>(false, [this, &QueueActor](UChildActorComponent* ChildActorComponent)
+		{
+			if (AActor* ChildActor = ChildActorComponent->GetChildActor())
 			{
-				if (AActor* ChildActor = ChildActorComponent->GetChildActor())
-				{
-					QueueActor(ChildActor);
-				}
-			});
-		}
+				QueueActor(ChildActor);
+			}
+		});
 	};
 
 	for (AActor* Actor : ActorList)
 	{
-		QueueActor(Actor);
+		// Actors that ask the level to keep a reference are already in the Actors array.
+		if (!Actor->ShouldLevelKeepRefIfExternal())
+		{
+			QueueActor(Actor);
+		}
+	}
+
+	TSet<AActor*> ActorsSet;
+	Algo::CopyIf(Actors, ActorsSet, [&](AActor* Actor) { return IsValid(Actor); });
+
+	for (AActor* Actor : ActorsQueue)
+	{
+		if (ActorsSet.Contains(Actor))
+		{
+			UE_LOG(LogLevel, Warning, TEXT("Duplicated actor %s (%s) found in level %s"), *Actor->GetName(), *Actor->GetActorNameOrLabel(), *GetPackage()->GetName());
+			continue;
+		}
+
+		Actors.Add(Actor);
+		ActorsForGC.Add(Actor);
 	}
 
 	FScopedSlowTask SlowTask(ActorsQueue.Num() * 3, LOCTEXT("RegisteringActors", "Registering actors..."));
@@ -922,7 +935,7 @@ void ULevel::RemoveLoadedActor(AActor* Actor, const FTransform* TransformToRemov
 
 void ULevel::RemoveLoadedActors(const TArray<AActor*>& ActorList, const FTransform* TransformToRemove)
 {
-	TArray<AActor*> ActorsQueue;
+	TSet<AActor*> ActorsQueue;
 	ActorsQueue.Reserve(ActorList.Num());
 
 	TFunction<void(AActor* Actor)> QueueActor = [this, &ActorsQueue, &QueueActor](AActor* Actor)
@@ -940,13 +953,9 @@ void ULevel::RemoveLoadedActors(const TArray<AActor*>& ActorList, const FTransfo
 			}
 		});
 
-		int32 ActorIndex;
-		// temporarily downgraded to ensure while an issue is fixed
-		if (ensure(Actors.Find(Actor, ActorIndex))) 
+		// Actors that ask the level to keep a reference should stay in the Actors array.
+		if (!Actor->ShouldLevelKeepRefIfExternal())
 		{
-			Actors[ActorIndex] = nullptr;
-			ActorsForGC.Remove(Actor);
-
 			ActorsQueue.Add(Actor);
 		}
 	};
@@ -956,10 +965,36 @@ void ULevel::RemoveLoadedActors(const TArray<AActor*>& ActorList, const FTransfo
 		QueueActor(Actor);
 	}
 
+	auto RemoveActorsFromList = [](TArray<TObjectPtr<AActor>>& ActorList, const TSet<AActor*>& ActorsToRemove)
+	{
+		int32 NumActorsRemoved = 0;
+
+		for (TObjectPtr<AActor>& Actor : ActorList)
+		{
+			if (Actor && ActorsToRemove.Contains(Actor))
+			{
+				Actor = nullptr;
+				if (++NumActorsRemoved == ActorsToRemove.Num())
+				{
+					break;
+				}
+			}
+		}
+
+		if (NumActorsRemoved)
+		{
+			ActorList.Remove(nullptr);
+		}
+	};
+
+	// Go through the actors lists once and remove queued actors
+	RemoveActorsFromList(Actors, ActorsQueue);
+	RemoveActorsFromList(ActorsForGC, ActorsQueue);
+
 	FScopedSlowTask SlowTask(ActorsQueue.Num(), LOCTEXT("UnregisteringActors", "Unregistering actors..."));
 	SlowTask.MakeDialogDelayed(1.0f);
 
-	OnLoadedActorRemovedFromLevelPreEvent.Broadcast(ActorsQueue);
+	OnLoadedActorRemovedFromLevelPreEvent.Broadcast(ActorsQueue.Array());
 
 	for (AActor* Actor : ActorsQueue)
 	{
@@ -4552,4 +4587,4 @@ void ULevel::OnMultipleLSAsPopupDismissed()
 
 #endif // WITH_EDITOR
 
-#undef LOCTEXT_NAMESPACE
+#undef LOCTEXT_NAMESPACE
\ No newline at end of file
