cluster union component line trace/sweep/overlap functions need to return true even in the case of a non-blocking hit

diff --git a/Engine/Source/Runtime/Engine/Private/PhysicsEngine/ClusterUnionComponent.cpp b/Engine/Source/Runtime/Engine/Private/PhysicsEngine/ClusterUnionComponent.cpp
--- a/Engine/Source/Runtime/Engine/Private/PhysicsEngine/ClusterUnionComponent.cpp
+++ b/Engine/Source/Runtime/Engine/Private/PhysicsEngine/ClusterUnionComponent.cpp
@@ -1076,11 +1076,12 @@ DECLARE_CYCLE_STAT(TEXT("UClusterUnionComponent::LineTraceComponentMulti"), STAT
 bool UClusterUnionComponent::LineTraceComponent(TArray<FHitResult>& OutHit, const FVector Start, const FVector End, ECollisionChannel TraceChannel, const struct FCollisionQueryParams& Params, const struct FCollisionResponseParams& ResponseParams, const struct FCollisionObjectQueryParams& ObjectParams)
 {
 	SCOPE_CYCLE_COUNTER(STAT_ClusterUnionComponentLineTraceComponentMulti);
+	OutHit.Reset();
 	if (AccelerationStructure)
 	{
-		return FGenericRaycastPhysicsInterfaceUsingSpatialAcceleration<IExternalSpatialAcceleration>::RaycastMulti(*AccelerationStructure, GetWorld(), OutHit, Start, End, TraceChannel, Params, ResponseParams, ObjectParams);
+		FGenericRaycastPhysicsInterfaceUsingSpatialAcceleration<IExternalSpatialAcceleration>::RaycastMulti(*AccelerationStructure, GetWorld(), OutHit, Start, End, TraceChannel, Params, ResponseParams, ObjectParams);
+		return !OutHit.IsEmpty();
 	}
-	OutHit.Reset();
 
 	VisitAllCurrentChildComponentsForCollision(TraceChannel, Params, ResponseParams, ObjectParams,
 		[&Start, &End, TraceChannel, &Params, &ResponseParams, &ObjectParams, &OutHit](UPrimitiveComponent* Component)
@@ -1102,13 +1103,13 @@ DECLARE_CYCLE_STAT(TEXT("UClusterUnionComponent::SweepComponentMulti"), STAT_Clu
 bool UClusterUnionComponent::SweepComponent(TArray<FHitResult>& OutHit, const FVector Start, const FVector End, const FQuat& ShapeWorldRotation, const FPhysicsGeometry& Geometry, ECollisionChannel TraceChannel, const struct FCollisionQueryParams& Params, const struct FCollisionResponseParams& ResponseParams, const struct FCollisionObjectQueryParams& ObjectParams)
 {
 	SCOPE_CYCLE_COUNTER(STAT_ClusterUnionComponentSweepComponentMulti);
+	OutHit.Reset();
 	if (AccelerationStructure)
 	{
-		return FGenericGeomPhysicsInterfaceUsingSpatialAcceleration<IExternalSpatialAcceleration, FPhysicsGeometry>::GeomSweepMulti(*AccelerationStructure, GetWorld(), Geometry, ShapeWorldRotation, OutHit, Start, End, TraceChannel, Params, ResponseParams, ObjectParams);
+		FGenericGeomPhysicsInterfaceUsingSpatialAcceleration<IExternalSpatialAcceleration, FPhysicsGeometry>::GeomSweepMulti(*AccelerationStructure, GetWorld(), Geometry, ShapeWorldRotation, OutHit, Start, End, TraceChannel, Params, ResponseParams, ObjectParams);
+		return !OutHit.IsEmpty();
 	}
 
-	OutHit.Reset();
-
 	VisitAllCurrentChildComponentsForCollision(TraceChannel, Params, ResponseParams, ObjectParams,
 		[&Geometry, &Start, &End, &ShapeWorldRotation, TraceChannel, &Params, &ResponseParams, &ObjectParams, &OutHit](UPrimitiveComponent* Component)
 		{
@@ -1129,13 +1130,15 @@ DECLARE_CYCLE_STAT(TEXT("UClusterUnionComponent::LineTraceComponentSingle"), STA
 bool UClusterUnionComponent::LineTraceComponent(FHitResult& OutHit, const FVector Start, const FVector End, ECollisionChannel TraceChannel, const struct FCollisionQueryParams& Params, const struct FCollisionResponseParams& ResponseParams, const struct FCollisionObjectQueryParams& ObjectParams)
 {
 	SCOPE_CYCLE_COUNTER(STAT_ClusterUnionComponentLineTraceComponentSingle);
+
+	OutHit.Distance = TNumericLimits<float>::Max();
 	if (AccelerationStructure)
 	{
-		return FGenericRaycastPhysicsInterfaceUsingSpatialAcceleration<IExternalSpatialAcceleration>::RaycastSingle(*AccelerationStructure, GetWorld(), OutHit, Start, End, TraceChannel, Params, ResponseParams, ObjectParams);
+		FGenericRaycastPhysicsInterfaceUsingSpatialAcceleration<IExternalSpatialAcceleration>::RaycastSingle(*AccelerationStructure, GetWorld(), OutHit, Start, End, TraceChannel, Params, ResponseParams, ObjectParams);
+		return OutHit.HasValidHitObjectHandle();
 	}
 
 	bool bHasHit = false;
-	OutHit.Distance = TNumericLimits<float>::Max();
 
 	VisitAllCurrentChildComponentsForCollision(TraceChannel, Params, ResponseParams, ObjectParams,
 		[&Start, &End, TraceChannel, &Params, &ResponseParams, &ObjectParams, &bHasHit, &OutHit](UPrimitiveComponent* Component)
@@ -1158,13 +1161,14 @@ DECLARE_CYCLE_STAT(TEXT("UClusterUnionComponent::SweepComponentSingle"), STAT_Cl
 bool UClusterUnionComponent::SweepComponent(FHitResult& OutHit, const FVector Start, const FVector End, const FQuat& ShapeWorldRotation, const FPhysicsGeometry& Geometry, ECollisionChannel TraceChannel, const struct FCollisionQueryParams& Params, const struct FCollisionResponseParams& ResponseParams, const struct FCollisionObjectQueryParams& ObjectParams)
 {
 	SCOPE_CYCLE_COUNTER(STAT_ClusterUnionComponentSweepComponentSingle);
+	OutHit.Distance = TNumericLimits<float>::Max();
 	if (AccelerationStructure)
 	{
-		return FGenericGeomPhysicsInterfaceUsingSpatialAcceleration<IExternalSpatialAcceleration, FPhysicsGeometry>::GeomSweepSingle(*AccelerationStructure, GetWorld(), Geometry, ShapeWorldRotation, OutHit, Start, End, TraceChannel, Params, ResponseParams, ObjectParams);
+		FGenericGeomPhysicsInterfaceUsingSpatialAcceleration<IExternalSpatialAcceleration, FPhysicsGeometry>::GeomSweepSingle(*AccelerationStructure, GetWorld(), Geometry, ShapeWorldRotation, OutHit, Start, End, TraceChannel, Params, ResponseParams, ObjectParams);
+		return OutHit.HasValidHitObjectHandle();
 	}
 
 	bool bHasHit = false;
-	OutHit.Distance = TNumericLimits<float>::Max();
 
 	VisitAllCurrentChildComponentsForCollision(TraceChannel, Params, ResponseParams, ObjectParams,
 		[&Geometry, &Start, &End, &ShapeWorldRotation, TraceChannel, &Params, &ResponseParams, &ObjectParams, &bHasHit, &OutHit](UPrimitiveComponent* Component)
@@ -1189,7 +1193,8 @@ bool UClusterUnionComponent::OverlapComponentWithResult(const FVector& Pos, cons
 	SCOPE_CYCLE_COUNTER(STAT_ClusterUnionComponentOverlapComponentWithResult);
 	if (AccelerationStructure)
 	{
-		return FGenericGeomPhysicsInterfaceUsingSpatialAcceleration<IExternalSpatialAcceleration, FPhysicsGeometry>::GeomOverlapMulti(*AccelerationStructure, GetWorld(), Geometry, Pos, Rot, OutOverlap, TraceChannel, Params, ResponseParams, ObjectParams);
+		FGenericGeomPhysicsInterfaceUsingSpatialAcceleration<IExternalSpatialAcceleration, FPhysicsGeometry>::GeomOverlapMulti(*AccelerationStructure, GetWorld(), Geometry, Pos, Rot, OutOverlap, TraceChannel, Params, ResponseParams, ObjectParams);
+		return !OutOverlap.IsEmpty();
 	}
 
 	bool bHasOverlap = false;
