World Partition: removed BlueprintAuthorityOnly on data layers activation functions now that we support client-only data layers.

diff --git a/Engine/Source/Runtime/Engine/Public/WorldPartition/DataLayer/DataLayerInstance.h b/Engine/Source/Runtime/Engine/Public/WorldPartition/DataLayer/DataLayerInstance.h
--- a/Engine/Source/Runtime/Engine/Public/WorldPartition/DataLayer/DataLayerInstance.h
+++ b/Engine/Source/Runtime/Engine/Public/WorldPartition/DataLayer/DataLayerInstance.h
@@ -181,7 +181,6 @@ public:
 
 private:
 	ENGINE_API void AddChild(UDataLayerInstance* DataLayer);
-	ENGINE_API bool SetRuntimeState(EDataLayerRuntimeState InState, bool bInIsRecursive = false) const;
 	friend class UDataLayerManager;
 	friend class FDataLayerUtils;
 
diff --git a/Engine/Source/Runtime/Engine/Public/WorldPartition/DataLayer/DataLayerManager.h b/Engine/Source/Runtime/Engine/Public/WorldPartition/DataLayer/DataLayerManager.h
--- a/Engine/Source/Runtime/Engine/Public/WorldPartition/DataLayer/DataLayerManager.h
+++ b/Engine/Source/Runtime/Engine/Public/WorldPartition/DataLayer/DataLayerManager.h
@@ -59,10 +59,10 @@ public:
 	UFUNCTION(BlueprintCallable, Category = DataLayers)
 	ENGINE_API const UDataLayerInstance* GetDataLayerInstanceFromName(const FName& InDataLayerInstanceName) const;
 
-	UFUNCTION(BlueprintCallable, Category = DataLayers, BlueprintAuthorityOnly)
+	UFUNCTION(BlueprintCallable, Category = DataLayers)
 	ENGINE_API bool SetDataLayerInstanceRuntimeState(const UDataLayerInstance* InDataLayerInstance, EDataLayerRuntimeState InState, bool bInIsRecursive = false);
 
-	UFUNCTION(BlueprintCallable, Category = DataLayers, BlueprintAuthorityOnly)
+	UFUNCTION(BlueprintCallable, Category = DataLayers)
 	ENGINE_API bool SetDataLayerRuntimeState(const UDataLayerAsset* InDataLayerAsset, EDataLayerRuntimeState InState, bool bInIsRecursive = false);
 
 	UFUNCTION(BlueprintCallable, Category = DataLayers)
