Add MinimalAPI specification to HV component so it can be created in the Niagara Baker

diff --git a/Engine/Source/Runtime/Engine/Classes/Components/HeterogeneousVolumeComponent.h b/Engine/Source/Runtime/Engine/Classes/Components/HeterogeneousVolumeComponent.h
--- a/Engine/Source/Runtime/Engine/Classes/Components/HeterogeneousVolumeComponent.h
+++ b/Engine/Source/Runtime/Engine/Classes/Components/HeterogeneousVolumeComponent.h
@@ -16,7 +16,7 @@ class USparseVolumeTexture;
 /**
  * A component that represents a heterogeneous volume.
  */
-UCLASS(Blueprintable, ClassGroup = (Rendering, Common), hidecategories = (Object, Activation, "Components|Activation"), ShowCategories = (Mobility), editinlinenew, meta = (BlueprintSpawnableComponent))
+UCLASS(Blueprintable, ClassGroup = (Rendering, Common), hidecategories = (Object, Activation, "Components|Activation"), ShowCategories = (Mobility), editinlinenew, meta = (BlueprintSpawnableComponent),MinimalAPI)
 class UHeterogeneousVolumeComponent : public UMeshComponent
 {
 	GENERATED_UCLASS_BODY()
