Bump up streaming priority for streamcache reads priority to normal from below normal.

diff --git a/Engine/Source/Runtime/Engine/Private/AudioStreamingCache.cpp b/Engine/Source/Runtime/Engine/Private/AudioStreamingCache.cpp
--- a/Engine/Source/Runtime/Engine/Private/AudioStreamingCache.cpp
+++ b/Engine/Source/Runtime/Engine/Private/AudioStreamingCache.cpp
@@ -53,7 +53,7 @@ FAutoConsoleVariableRef CVarAlwaysLogCacheMisses(
 	TEXT("0: Don't log cache misses until au.streamcaching.StartProfiling is called. 1: Always log cache misses."),
 	ECVF_Default);
 
-static int32 ReadRequestPriorityCVar = 2;
+static int32 ReadRequestPriorityCVar = 1;
 FAutoConsoleVariableRef CVarReadRequestPriority(
 	TEXT("au.streamcaching.ReadRequestPriority"),
 	ReadRequestPriorityCVar,
