-Fixed CastShadow field being overriden in FPrimitiveSceneProxyDesc after its initialized from the UPrimitiveComponent
    -Caused the 'Cast Shadow' property to behave inconsistently
    -Also likely caused performance regressions
-Used a TOptionnal instead of bool + struct combo for the FCollisionResponseContainer

diff --git a/Engine/Source/Runtime/Engine/Private/InstancedStaticMesh.cpp b/Engine/Source/Runtime/Engine/Private/InstancedStaticMesh.cpp
--- a/Engine/Source/Runtime/Engine/Private/InstancedStaticMesh.cpp
+++ b/Engine/Source/Runtime/Engine/Private/InstancedStaticMesh.cpp
@@ -1724,8 +1724,15 @@ FInstancedStaticMeshVFLooseUniformShaderParametersRef FInstancedStaticMeshSceneP
 }
 
 FInstancedStaticMeshSceneProxyDesc::FInstancedStaticMeshSceneProxyDesc(const UInstancedStaticMeshComponent* InComponent)
-	: FStaticMeshSceneProxyDesc(InComponent)	  
+	: FInstancedStaticMeshSceneProxyDesc()	  
 {
+	InitializeFrom(InComponent);
+}
+
+void FInstancedStaticMeshSceneProxyDesc::InitializeFrom(const UInstancedStaticMeshComponent* InComponent)
+{
+	FStaticMeshSceneProxyDesc::InitializeFrom(InComponent);
+
 	PerInstanceSMData = InComponent->PerInstanceSMData;
 	PerInstanceRenderData = InComponent->PerInstanceRenderData;;
 	PerInstanceSMCustomData = InComponent->PerInstanceSMCustomData;
diff --git a/Engine/Source/Runtime/Engine/Private/PrimitiveSceneProxy.cpp b/Engine/Source/Runtime/Engine/Private/PrimitiveSceneProxy.cpp
--- a/Engine/Source/Runtime/Engine/Private/PrimitiveSceneProxy.cpp
+++ b/Engine/Source/Runtime/Engine/Private/PrimitiveSceneProxy.cpp
@@ -255,6 +255,11 @@ static FRenderBounds PadLocalRenderBounds(const FRenderBounds& InBounds, const F
 
 FPrimitiveSceneProxyDesc::FPrimitiveSceneProxyDesc(const UPrimitiveComponent* InComponent)
 	: FPrimitiveSceneProxyDesc()
+{
+	InitializeFrom(InComponent);
+}
+
+void FPrimitiveSceneProxyDesc::InitializeFrom(const UPrimitiveComponent* InComponent)
 {
 	CastShadow = InComponent->CastShadow;
 	bReceivesDecals = InComponent->bReceivesDecals;
diff --git a/Engine/Source/Runtime/Engine/Private/StaticMeshRender.cpp b/Engine/Source/Runtime/Engine/Private/StaticMeshRender.cpp
--- a/Engine/Source/Runtime/Engine/Private/StaticMeshRender.cpp
+++ b/Engine/Source/Runtime/Engine/Private/StaticMeshRender.cpp
@@ -2610,62 +2610,67 @@ bool FStaticMeshSceneProxyDesc::ShouldCreateNaniteProxy(Nanite::FMaterialAudit*
 }
 
 
-FStaticMeshSceneProxyDesc::FStaticMeshSceneProxyDesc(const UStaticMeshComponent* Component)
-	: FPrimitiveSceneProxyDesc(Component)
+FStaticMeshSceneProxyDesc::FStaticMeshSceneProxyDesc(const UStaticMeshComponent* InComponent)
+	: FStaticMeshSceneProxyDesc()
+{	
+	InitializeFrom(InComponent);
+}
+
+void FStaticMeshSceneProxyDesc::InitializeFrom(const UStaticMeshComponent* InComponent)
 {
-	InitDefaults();
+	FPrimitiveSceneProxyDesc::InitializeFrom(InComponent);	
 
-	StaticMesh = Component->GetStaticMesh();
-	OverrideMaterials = const_cast<UStaticMeshComponent*>(Component)->OverrideMaterials;	
-	OverlayMaterial = Component->GetOverlayMaterial();
-	OverlayMaterialMaxDrawDistance = Component->GetOverlayMaterialMaxDrawDistance();
+	StaticMesh = InComponent->GetStaticMesh();
+	OverrideMaterials = const_cast<UStaticMeshComponent*>(InComponent)->OverrideMaterials;	
+	OverlayMaterial = InComponent->GetOverlayMaterial();
+	OverlayMaterialMaxDrawDistance = InComponent->GetOverlayMaterialMaxDrawDistance();
 
-	ForcedLodModel = Component->ForcedLodModel ;
-	MinLOD = Component->MinLOD ;
-	WorldPositionOffsetDisableDistance = Component->WorldPositionOffsetDisableDistance ;	
-	bReverseCulling = Component->bReverseCulling ;
+	ForcedLodModel = InComponent->ForcedLodModel ;
+	MinLOD = InComponent->MinLOD ;
+	WorldPositionOffsetDisableDistance = InComponent->WorldPositionOffsetDisableDistance ;	
+	bReverseCulling = InComponent->bReverseCulling ;
 #if STATICMESH_ENABLE_DEBUG_RENDERING
-	bDrawMeshCollisionIfComplex = Component->bDrawMeshCollisionIfComplex ;
-	bDrawMeshCollisionIfSimple = Component->bDrawMeshCollisionIfSimple ;
+	bDrawMeshCollisionIfComplex = InComponent->bDrawMeshCollisionIfComplex ;
+	bDrawMeshCollisionIfSimple = InComponent->bDrawMeshCollisionIfSimple ;
 #endif
-	bEvaluateWorldPositionOffset = Component->bEvaluateWorldPositionOffset ;
-	bOverrideMinLOD = Component->bOverrideMinLOD ;
-	bCastDistanceFieldIndirectShadow = Component->bCastDistanceFieldIndirectShadow ;
-	bOverrideDistanceFieldSelfShadowBias = Component->bOverrideDistanceFieldSelfShadowBias ;
-	bEvaluateWorldPositionOffsetInRayTracing = Component->bEvaluateWorldPositionOffsetInRayTracing ;	
-	bSortTriangles = Component->bSortTriangles ;
+	bEvaluateWorldPositionOffset = InComponent->bEvaluateWorldPositionOffset ;
+	bOverrideMinLOD = InComponent->bOverrideMinLOD ;
+	bCastDistanceFieldIndirectShadow = InComponent->bCastDistanceFieldIndirectShadow ;
+	bOverrideDistanceFieldSelfShadowBias = InComponent->bOverrideDistanceFieldSelfShadowBias ;
+	bEvaluateWorldPositionOffsetInRayTracing = InComponent->bEvaluateWorldPositionOffsetInRayTracing ;	
+	bSortTriangles = InComponent->bSortTriangles ;
 #if WITH_EDITOR
-	bDisplayNaniteFallbackMesh = Component->bDisplayNaniteFallbackMesh ;
+	bDisplayNaniteFallbackMesh = InComponent->bDisplayNaniteFallbackMesh ;
 #endif
-	bDisallowNanite = Component->bDisallowNanite ;
-	bForceDisableNanite = Component->bForceDisableNanite ;
-	bForceNaniteForMasked = Component->bForceNaniteForMasked ;
-	DistanceFieldSelfShadowBias = Component->DistanceFieldSelfShadowBias ;
-	DistanceFieldIndirectShadowMinVisibility = Component->DistanceFieldIndirectShadowMinVisibility ;
-	StaticLightMapResolution = Component->GetStaticLightMapResolution();
-	LightmapType = Component->LightmapType;
+	bDisallowNanite = InComponent->bDisallowNanite ;
+	bForceDisableNanite = InComponent->bForceDisableNanite ;
+	bForceNaniteForMasked = InComponent->bForceNaniteForMasked ;
+	DistanceFieldSelfShadowBias = InComponent->DistanceFieldSelfShadowBias ;
+	DistanceFieldIndirectShadowMinVisibility = InComponent->DistanceFieldIndirectShadowMinVisibility ;
+	StaticLightMapResolution = InComponent->GetStaticLightMapResolution();
+	LightmapType = InComponent->LightmapType;
 
 #if WITH_EDITORONLY_DATA
-	StreamingDistanceMultiplier = Component->StreamingDistanceMultiplier;
-	MaterialStreamingRelativeBoxes = const_cast<UStaticMeshComponent*>(Component)->MaterialStreamingRelativeBoxes;
-	SectionIndexPreview = Component->SectionIndexPreview;
-	MaterialIndexPreview = Component->MaterialIndexPreview;
-	SelectedEditorMaterial = Component->SelectedEditorMaterial;
-	SelectedEditorSection = Component->SelectedEditorSection;
-
-	TextureStreamingTransformScale = Component->GetTextureStreamingTransformScale();	
+	StreamingDistanceMultiplier = InComponent->StreamingDistanceMultiplier;
+	MaterialStreamingRelativeBoxes = const_cast<UStaticMeshComponent*>(InComponent)->MaterialStreamingRelativeBoxes;
+	SectionIndexPreview = InComponent->SectionIndexPreview;
+	MaterialIndexPreview = InComponent->MaterialIndexPreview;
+	SelectedEditorMaterial = InComponent->SelectedEditorMaterial;
+	SelectedEditorSection = InComponent->SelectedEditorSection;
+
+	TextureStreamingTransformScale = InComponent->GetTextureStreamingTransformScale();	
 #endif
 
-	NaniteResources = Component->GetNaniteResources();
-	BodySetup = const_cast<UStaticMeshComponent*>(Component)->GetBodySetup();
+	NaniteResources = InComponent->GetNaniteResources();
+	BodySetup = const_cast<UStaticMeshComponent*>(InComponent)->GetBodySetup();
 
-	LODData = const_cast<UStaticMeshComponent*>(Component)->LODData;
+	LODData = const_cast<UStaticMeshComponent*>(InComponent)->LODData;
 
-	WireframeColor = Component->GetWireframeColor();
-	LODParentPrimitive = Component->GetLODParentPrimitive();	
+	WireframeColor = InComponent->GetWireframeColor();
+	LODParentPrimitive = InComponent->GetLODParentPrimitive();	
 
-	SetMaterialRelevance(Component->GetMaterialRelevance(World->GetFeatureLevel()));
-	SetCollisionResponseToChannels(Component->GetCollisionResponseToChannels());
+	SetMaterialRelevance(InComponent->GetMaterialRelevance(World->GetFeatureLevel()));
+	SetCollisionResponseToChannels(InComponent->GetCollisionResponseToChannels());
 }
 
 FPrimitiveSceneProxy* UStaticMeshComponent::CreateStaticMeshSceneProxy(Nanite::FMaterialAudit& NaniteMaterials, bool bCreateNanite)
diff --git a/Engine/Source/Runtime/Engine/Public/InstancedStaticMeshSceneProxyDesc.h b/Engine/Source/Runtime/Engine/Public/InstancedStaticMeshSceneProxyDesc.h
--- a/Engine/Source/Runtime/Engine/Public/InstancedStaticMeshSceneProxyDesc.h
+++ b/Engine/Source/Runtime/Engine/Public/InstancedStaticMeshSceneProxyDesc.h
@@ -11,6 +11,7 @@ struct FInstancedStaticMeshSceneProxyDesc : public FStaticMeshSceneProxyDesc
 {		
 	FInstancedStaticMeshSceneProxyDesc() = default;
 	ENGINE_API FInstancedStaticMeshSceneProxyDesc(const UInstancedStaticMeshComponent*);
+	void InitializeFrom(const UInstancedStaticMeshComponent*);
 
 	TArrayView<const FInstancedStaticMeshInstanceData> PerInstanceSMData;	
 	TSharedPtr<FPerInstanceRenderData, ESPMode::ThreadSafe> PerInstanceRenderData;
diff --git a/Engine/Source/Runtime/Engine/Public/PrimitiveSceneProxyDesc.h b/Engine/Source/Runtime/Engine/Public/PrimitiveSceneProxyDesc.h
--- a/Engine/Source/Runtime/Engine/Public/PrimitiveSceneProxyDesc.h
+++ b/Engine/Source/Runtime/Engine/Public/PrimitiveSceneProxyDesc.h
@@ -76,6 +76,8 @@ struct FPrimitiveSceneProxyDesc
 
 	ENGINE_API FPrimitiveSceneProxyDesc(const UPrimitiveComponent*);
 
+	void InitializeFrom(const UPrimitiveComponent*);
+	
 	virtual ~FPrimitiveSceneProxyDesc() = default;
 
 	uint32 CastShadow : 1;
diff --git a/Engine/Source/Runtime/Engine/Public/StaticMeshSceneProxyDesc.h b/Engine/Source/Runtime/Engine/Public/StaticMeshSceneProxyDesc.h
--- a/Engine/Source/Runtime/Engine/Public/StaticMeshSceneProxyDesc.h
+++ b/Engine/Source/Runtime/Engine/Public/StaticMeshSceneProxyDesc.h
@@ -17,11 +17,6 @@ class UBodySetup;
 struct FStaticMeshSceneProxyDesc : public FPrimitiveSceneProxyDesc
 {
 	FStaticMeshSceneProxyDesc()
-	{
-		InitDefaults();
-	}	
-
-	void InitDefaults()
 	{
 		CastShadow = true;
 		bUseAsOccluder = true;
@@ -46,9 +41,10 @@ struct FStaticMeshSceneProxyDesc : public FPrimitiveSceneProxyDesc
 		bForceNaniteForMasked = false;
 
 		bUseProvidedMaterialRelevance = false;
-		bUseProvidedCollisionResponseContainer = false;
 	}
+	
 	ENGINE_API FStaticMeshSceneProxyDesc(const UStaticMeshComponent*);
+	void InitializeFrom(const UStaticMeshComponent*);
 
 	UStaticMesh* StaticMesh = nullptr;
 	TArrayView<TObjectPtr<UMaterialInterface>>	OverrideMaterials;
@@ -119,7 +115,7 @@ struct FStaticMeshSceneProxyDesc : public FPrimitiveSceneProxyDesc
 		return nullptr;
 	}
 
-	FCollisionResponseContainer CollisionResponseContainer;
+	TOptional<FCollisionResponseContainer> CollisionResponseContainer;
 	
 	UObject* LODParentPrimitive = nullptr;
 	UObject* GetLODParentPrimitive() const { return LODParentPrimitive; }
@@ -174,15 +170,15 @@ struct FStaticMeshSceneProxyDesc : public FPrimitiveSceneProxyDesc
 	{ 
 		if (&InContainer != &FCollisionResponseContainer::GetDefaultResponseContainer())
 		{
-			CollisionResponseContainer = InContainer; 
-			bUseProvidedCollisionResponseContainer = true;
-		}		
+			CollisionResponseContainer = InContainer;
+		}
 	}
+
 	const FCollisionResponseContainer& GetCollisionResponseToChannels() const
 	{
-		if (bUseProvidedCollisionResponseContainer)
+		if (CollisionResponseContainer.IsSet())
 		{
-			return CollisionResponseContainer;
+			return *CollisionResponseContainer;
 		}
 
 		return FCollisionResponseContainer::GetDefaultResponseContainer();
