Fix missing SpecularLUT resource transition.

diff --git a/Engine/Source/Runtime/Engine/Private/Rendering/SpecularProfile.cpp b/Engine/Source/Runtime/Engine/Private/Rendering/SpecularProfile.cpp
--- a/Engine/Source/Runtime/Engine/Private/Rendering/SpecularProfile.cpp
+++ b/Engine/Source/Runtime/Engine/Private/Rendering/SpecularProfile.cpp
@@ -384,6 +384,11 @@ IPooledRenderTarget* FSpecularProfileTextureManager::GetAtlasTexture(FRDGBuilder
 
 			SpecularProfileEntries[LayerIt].CachedResolution = Texture ? Texture->TextureReferenceRHI->GetDesc().Extent : FIntPoint(0,0);
 		}
+
+		// Transit texture to SRV for letting the non-RDG resource to be bound later by the various passes
+		FRDGExternalAccessQueue ExternalAccessQueue;
+		ExternalAccessQueue.Add(SpecularProfileTexture);
+		ExternalAccessQueue.Submit(GraphBuilder);
 	}
 	return GSpecularProfileTextureAtlas;
 }
