Fix missing include.

diff --git a/Engine/Source/Runtime/Engine/Private/WorldPartition/RuntimeHashSet/WorldPartitionRuntimeHashSetDebugDraw.cpp b/Engine/Source/Runtime/Engine/Private/WorldPartition/RuntimeHashSet/WorldPartitionRuntimeHashSetDebugDraw.cpp
--- a/Engine/Source/Runtime/Engine/Private/WorldPartition/RuntimeHashSet/WorldPartitionRuntimeHashSetDebugDraw.cpp
+++ b/Engine/Source/Runtime/Engine/Private/WorldPartition/RuntimeHashSet/WorldPartitionRuntimeHashSetDebugDraw.cpp
@@ -9,6 +9,7 @@
 #include "WorldPartition/WorldPartitionDebugHelper.h"
 #include "DrawDebugHelpers.h"
 #include "Math/TransformCalculus2D.h"
+#include "Misc/Paths.h"
 
 static int32 GShowRuntimeHashSetDebugDisplayLevel = 0;
 static FAutoConsoleVariableRef CVarShowRuntimeHashSetDebugDisplayLevel(
