World Partition: use a deferred context when loading/unloading always loaded actors in PIE to make a single call to ULevel::AddLoadedActors/RemoveLoadedActors.
This affects mainly always loaded HLOD actors that aren't loaded in the editor but are in PIE.

diff --git a/Engine/Source/Runtime/Engine/Private/WorldPartition/WorldPartitionRuntimeHash.cpp b/Engine/Source/Runtime/Engine/Private/WorldPartition/WorldPartitionRuntimeHash.cpp
--- a/Engine/Source/Runtime/Engine/Private/WorldPartition/WorldPartitionRuntimeHash.cpp
+++ b/Engine/Source/Runtime/Engine/Private/WorldPartition/WorldPartitionRuntimeHash.cpp
@@ -115,7 +115,10 @@ void UWorldPartitionRuntimeHash::OnEndPlay()
 	ForceExternalActorLevelReference(/*bForceExternalActorLevelReferenceForPIE*/false);
 
 	// Release references (will unload actors that were not already loaded in the Editor)
-	AlwaysLoadedActorsForPIE.Empty();
+	{
+		FWorldPartitionLoadingContext::FDeferred LoadingContext;
+		AlwaysLoadedActorsForPIE.Empty();
+	}
 
 	ModifiedActorDescListForPIE.Empty();
 }
@@ -130,62 +133,60 @@ void UWorldPartitionRuntimeHash::FlushStreaming()
 	PackagesToGenerateForCook.Empty();
 }
 
-// In PIE, Always loaded cell is not generated. Instead, always loaded actors will be added to AlwaysLoadedActorsForPIE.
-// This will trigger loading/registration of these actors in the PersistentLevel (if not already loaded).
-// Then, duplication of world for PIE will duplicate only these actors. 
-// When stopping PIE, WorldPartition will release these FWorldPartitionReferences which 
-// will unload actors that were not already loaded in the non PIE world.
-bool UWorldPartitionRuntimeHash::ConditionalRegisterAlwaysLoadedActorsForPIE(const IStreamingGenerationContext::FActorSetInstance* ActorSetInstance, bool bIsMainWorldPartition, bool bIsMainContainer, bool bIsCellAlwaysLoaded)
+bool UWorldPartitionRuntimeHash::PopulateCellActorInstances(const TArray<const IStreamingGenerationContext::FActorSetInstance*>& ActorSetInstances, bool bIsMainWorldPartition, bool bIsCellAlwaysLoaded, TArray<IStreamingGenerationContext::FActorInstance>& OutCellActorInstances)
 {
-	if (bIsMainWorldPartition && bIsMainContainer && bIsCellAlwaysLoaded && !IsRunningCookCommandlet())
+	// In PIE, Always loaded cell is not generated. Instead, always loaded actors will be added to AlwaysLoadedActorsForPIE.
+	// This will trigger loading/registration of these actors in the PersistentLevel (if not already loaded).
+	// Then, duplication of world for PIE will duplicate only these actors. 
+	// When stopping PIE, WorldPartition will release these FWorldPartitionReferences which 
+	// will unload actors that were not already loaded in the non PIE world.
+	TArray<FWorldPartitionReference> AlwaysLoadedReferences;
 	{
-		ActorSetInstance->ForEachActor([this, ActorSetInstance](const FGuid& ActorGuid)
-		{
-			IStreamingGenerationContext::FActorInstance ActorInstance(ActorGuid, ActorSetInstance);
-			const FWorldPartitionActorDescView& ActorDescView = ActorInstance.GetActorDescView();
-
-			// This will load the actor if it isn't already loaded
-			FWorldPartitionReference Reference(GetOuterUWorldPartition(), ActorDescView.GetGuid());
-
-			if (AActor* AlwaysLoadedActor = FindObject<AActor>(nullptr, *ActorDescView.GetActorSoftPath().ToString()))
-			{
-				AlwaysLoadedActorsForPIE.Emplace(Reference, AlwaysLoadedActor);
-			}
-		});
+		FWorldPartitionLoadingContext::FDeferred LoadingContext;
 
-		return true;
-	}
-
-	return false;
-}
+		const bool bForceLoadAlwaysLoadedReferences = bIsMainWorldPartition && bIsCellAlwaysLoaded && !IsRunningCookCommandlet();
 
-bool UWorldPartitionRuntimeHash::PopulateCellActorInstances(const TArray<const IStreamingGenerationContext::FActorSetInstance*>& ActorSetInstances, bool bIsMainWorldPartition, bool bIsCellAlwaysLoaded, TArray<IStreamingGenerationContext::FActorInstance>& OutCellActorInstances)
-{
-	for (const IStreamingGenerationContext::FActorSetInstance* ActorSetInstance : ActorSetInstances)
-	{
-		// Instanced world partition, ContainerID is the main container, but it's not the main world partition,
-		// so the always loaded actors don't be part of the process of ForceExternalActorLevelReference/AlwaysLoadedActorsForPIE.
-		// In PIE, always loaded actors of an instanced world partition will go in the always loaded cell.
-		if (!ConditionalRegisterAlwaysLoadedActorsForPIE(ActorSetInstance, bIsMainWorldPartition, ActorSetInstance->ContainerID.IsMainContainer(), bIsCellAlwaysLoaded))
+		for (const IStreamingGenerationContext::FActorSetInstance* ActorSetInstance : ActorSetInstances)
 		{
-			ActorSetInstance->ForEachActor([this, ActorSetInstance, &OutCellActorInstances](const FGuid& ActorGuid)
+			ActorSetInstance->ForEachActor([this, ActorSetInstance, &AlwaysLoadedReferences, &OutCellActorInstances, bForceLoadAlwaysLoadedReferences](const FGuid& ActorGuid)
 			{
-				// Actors that return true to ShouldLevelKeepRefIfExternal will always be part of the partitioned persistent level of a
-				// world partition. In PIE, for an instanced world partition, we don't want this actor to be both in the persistent level
-				// and also part of the always loaded cell level.
-				//
-				// @todo_ow: We need to implement PIE always loaded actors of instanced world partitions to be part
-				//			 of the persistent level and get rid of the always loaded cell (to have the same behavior
-				//			 as non-instanced world partition and as cooked world partition).
 				IStreamingGenerationContext::FActorInstance ActorInstance(ActorGuid, ActorSetInstance);
-				if (!CastChecked<AActor>(ActorInstance.GetActorDescView().GetActorNativeClass()->GetDefaultObject())->ShouldLevelKeepRefIfExternal())
+				const FWorldPartitionActorDescView& ActorDescView = ActorInstance.GetActorDescView();
+
+				// Instanced world partition, ContainerID is the main container, but it's not the main world partition,
+				// so the always loaded actors don't be part of the process of ForceExternalActorLevelReference/AlwaysLoadedActorsForPIE.
+				// In PIE, always loaded actors of an instanced world partition will go in the always loaded cell.
+				if (bForceLoadAlwaysLoadedReferences && ActorSetInstance->ContainerID.IsMainContainer())
+				{
+					// This will load the actor if it isn't already loaded, when the deferred context ends.
+					AlwaysLoadedReferences.Emplace(FWorldPartitionReference(GetOuterUWorldPartition(), ActorDescView.GetGuid()));
+				}
+				else
 				{
-					OutCellActorInstances.Emplace(ActorInstance);
+					// Actors that return true to ShouldLevelKeepRefIfExternal will always be part of the partitioned persistent level of a
+					// world partition. In PIE, for an instanced world partition, we don't want this actor to be both in the persistent level
+					// and also part of the always loaded cell level.
+					//
+					// @todo_ow: We need to implement PIE always loaded actors of instanced world partitions to be part
+					//			 of the persistent level and get rid of the always loaded cell (to have the same behavior
+					//			 as non-instanced world partition and as cooked world partition).
+					if (!CastChecked<AActor>(ActorDescView.GetActorNativeClass()->GetDefaultObject())->ShouldLevelKeepRefIfExternal())
+					{
+						OutCellActorInstances.Emplace(ActorInstance);
+					}
 				}
 			});
 		}
 	}
 
+	for (FWorldPartitionReference& Reference : AlwaysLoadedReferences)
+	{
+		if (AActor* AlwaysLoadedActor = FindObject<AActor>(nullptr, *Reference->GetActorSoftPath().ToString()))
+		{
+			AlwaysLoadedActorsForPIE.Emplace(Reference, AlwaysLoadedActor);
+		}
+	}
+
 	return OutCellActorInstances.Num() > 0;
 }
 
diff --git a/Engine/Source/Runtime/Engine/Public/WorldPartition/WorldPartitionRuntimeHash.h b/Engine/Source/Runtime/Engine/Public/WorldPartition/WorldPartitionRuntimeHash.h
--- a/Engine/Source/Runtime/Engine/Public/WorldPartition/WorldPartitionRuntimeHash.h
+++ b/Engine/Source/Runtime/Engine/Public/WorldPartition/WorldPartitionRuntimeHash.h
@@ -109,7 +109,6 @@ class UWorldPartitionRuntimeHash : public UObject
 	ENGINE_API void OnEndPlay();
 
 protected:
-	ENGINE_API bool ConditionalRegisterAlwaysLoadedActorsForPIE(const IStreamingGenerationContext::FActorSetInstance* ActorSetInstance, bool bIsMainWorldPartition, bool bIsMainContainer, bool bIsCellAlwaysLoaded);
 	ENGINE_API bool PopulateCellActorInstances(const TArray<const IStreamingGenerationContext::FActorSetInstance*>& ActorSetInstances, bool bIsMainWorldPartition, bool bIsCellAlwaysLoaded, TArray<IStreamingGenerationContext::FActorInstance>& OutCellActorInstances);
 	ENGINE_API void PopulateRuntimeCell(UWorldPartitionRuntimeCell* RuntimeCell, const TArray<IStreamingGenerationContext::FActorInstance>& ActorInstances, TArray<FString>* OutPackagesToGenerate);
 #endif
