Crash at ' FD3D12CommandContext::RHICopyToStagingBuffer'
#fix ensure that SkinWeightProfile streaming does not add a nullptr default-profile buffer entry (if statically overridden) -
#misc could not reproduce the same callstack as the JIRA

diff --git a/Engine/Source/Runtime/Engine/Private/Animation/SkinWeightProfile.cpp b/Engine/Source/Runtime/Engine/Private/Animation/SkinWeightProfile.cpp
--- a/Engine/Source/Runtime/Engine/Private/Animation/SkinWeightProfile.cpp
+++ b/Engine/Source/Runtime/Engine/Private/Animation/SkinWeightProfile.cpp
@@ -528,9 +528,10 @@ void FSkinWeightProfilesData::ReleaseResources()
 	ProfileNameToBuffer.GenerateValueArray(Buffers);
 	ProfileNameToBuffer.Empty();
 
-	// Never release a default buffer
-	if (bDefaultOverriden)
+	// Never release a default _dynamic_ buffer
+	if (bDefaultOverriden && !bStaticOverriden)
 	{
+		ensure(DefaultOverrideSkinWeightBuffer != nullptr);
 		Buffers.Remove(DefaultOverrideSkinWeightBuffer);
 		ProfileNameToBuffer.Add(DefaultProfileName, DefaultOverrideSkinWeightBuffer);
 	}
