Shader compiler:  Only process complete material shader maps in FShaderCompilingManager::ProcessCompiledShaderMaps when cooking.  Partially complete material shader maps are useful in the editor for rendering while compiling continues, but that doesn't apply in the cooker.  Saves 1.8% on overall cook time in CitySample.

#rnx

diff --git a/Engine/Source/Runtime/Engine/Private/ShaderCompiler/ShaderCompiler.cpp b/Engine/Source/Runtime/Engine/Private/ShaderCompiler/ShaderCompiler.cpp
--- a/Engine/Source/Runtime/Engine/Private/ShaderCompiler/ShaderCompiler.cpp
+++ b/Engine/Source/Runtime/Engine/Private/ShaderCompiler/ShaderCompiler.cpp
@@ -6090,7 +6090,6 @@ void FShaderCompilingManager::ProcessCompiledShaderMaps(
 				}
 			}
 
-			FMaterialShaderMap* ShaderMapToUseForRendering = nullptr;
 			if (bSuccess)
 			{
 				int32 JobIndex = 0;
@@ -6106,12 +6105,6 @@ void FShaderCompilingManager::ProcessCompiledShaderMaps(
 					}
 					FinishedJobs.RemoveAt(0, JobIndex);
 				}
-
-				// Make a clone of the compiling shader map to use for rendering
-				// This will allow rendering to proceed with the clone, while async compilation continues to potentially update the compiling shader map
-				double StartTime = FPlatformTime::Seconds();
-				ShaderMapToUseForRendering = CompilingShaderMap->AcquireFinalizedClone();
-				TimeBudget -= (FPlatformTime::Seconds() - StartTime);
 			}
 
 			if (!bSuccess || FinishedJobs.Num() == 0)
@@ -6119,6 +6112,8 @@ void FShaderCompilingManager::ProcessCompiledShaderMaps(
 				ShaderMapResultIter.RemoveCurrent();
 			}
 
+			FMaterialShaderMap* ShaderMapToUseForRendering = nullptr;
+
 #if DEBUG_INFINITESHADERCOMPILE
 			UE_LOG(LogTemp, Display, TEXT("Finished compile of shader map 0x%08X%08X"), (int)((int64)(ShaderMap.GetReference()) >> 32), (int)((int64)(ShaderMap.GetReference())));
 #endif
@@ -6145,8 +6140,24 @@ void FShaderCompilingManager::ProcessCompiledShaderMaps(
 				}
 				else if (bSuccess)
 				{
-					MaterialsToUpdate.Add(Material, ShaderMapToUseForRendering);
-					if (ShaderMapToUseForRendering->IsComplete(Material, true))
+					bool bIsComplete = CompilingShaderMap->IsComplete(Material, true);
+
+					// If running a cook, only process complete shader maps, as there's no rendering of partially complete shader maps to worry about.
+					if (bIsComplete || IsRunningCookCommandlet() == false)
+					{
+						if (ShaderMapToUseForRendering == nullptr)
+						{
+							// Make a clone of the compiling shader map to use for rendering
+							// This will allow rendering to proceed with the clone, while async compilation continues to potentially update the compiling shader map
+							double StartTime = FPlatformTime::Seconds();
+							ShaderMapToUseForRendering = CompilingShaderMap->AcquireFinalizedClone();
+							TimeBudget -= (FPlatformTime::Seconds() - StartTime);
+						}
+
+						MaterialsToUpdate.Add(Material, ShaderMapToUseForRendering);
+					}
+
+					if (bIsComplete)
 					{
 						bReleaseCompilingId = true;
 					}
@@ -6159,7 +6170,7 @@ void FShaderCompilingManager::ProcessCompiledShaderMaps(
 					{
 						UE_LOG(LogShaderCompilers, Warning, TEXT("Warnings while compiling Material %s for platform %s:"),
 							*Material->GetDebugName(),
-							*LegacyShaderPlatformToShaderFormat(ShaderMapToUseForRendering->GetShaderPlatform()).ToString());
+							*LegacyShaderPlatformToShaderFormat(CompilingShaderMap->GetShaderPlatform()).ToString());
 						for (int32 ErrorIndex = 0; ErrorIndex < Errors.Num(); ErrorIndex++)
 						{
 							UE_LOG(LogShaders, Warning, TEXT("  %s"), *Errors[ErrorIndex]);
@@ -6223,9 +6234,8 @@ void FShaderCompilingManager::ProcessCompiledShaderMaps(
 			{
 				CompilingShaderMap->bCompiledSuccessfully = bSuccess;
 				CompilingShaderMap->bCompilationFinalized = true;
-				if (bSuccess)
+				if (ShaderMapToUseForRendering)
 				{
-					check(ShaderMapToUseForRendering);
 					ShaderMapToUseForRendering->bCompiledSuccessfully = true;
 					ShaderMapToUseForRendering->bCompilationFinalized = true;
 					if (ShaderMapToUseForRendering->bIsPersistent)
