ActorEditorContext
- Add FScopedActorEditorContextFromActor which pushes a new context based on the provided actor (will be used in next Foliage CL)
- Implemented through IActorEditorContextClient(s) with EActorEditorContextAction::InitializeContextFromActor
- Remove Adding of DataLayers in ActorPartitionSubsystem as this will be done by the DataLayerEditorSubsystem through the active editor context
- Properly support ContextBundles in ActorPartitionSubsystem to generate unique name
- Cleanup APartitionActor naming API + Labels

#rnx

diff --git a/Engine/Source/Editor/DataLayerEditor/Private/DataLayer/DataLayerEditorSubsystem.cpp b/Engine/Source/Editor/DataLayerEditor/Private/DataLayer/DataLayerEditorSubsystem.cpp
--- a/Engine/Source/Editor/DataLayerEditor/Private/DataLayer/DataLayerEditorSubsystem.cpp
+++ b/Engine/Source/Editor/DataLayerEditor/Private/DataLayer/DataLayerEditorSubsystem.cpp
@@ -314,6 +314,15 @@ void UDataLayerEditorSubsystem::OnExecuteActorEditorContextAction(UWorld* InWorl
 			DataLayerManager->PopActorEditorContext();
 			BroadcastDataLayerChanged(EDataLayerAction::Reset, NULL, NAME_None);
 			break;
+		case EActorEditorContextAction::InitializeContextFromActor:
+			for (const UDataLayerInstance* DataLayerInstance : InActor->GetDataLayerInstances())
+			{
+				if (DataLayerInstance->CanBeInActorEditorContext())
+				{
+					const_cast<UDataLayerInstance*>(DataLayerInstance)->AddToActorEditorContext();
+				}
+			}
+			break;
 		}
 	}
 }	
diff --git a/Engine/Source/Editor/UnrealEd/Private/EditorActorFolders.cpp b/Engine/Source/Editor/UnrealEd/Private/EditorActorFolders.cpp
--- a/Engine/Source/Editor/UnrealEd/Private/EditorActorFolders.cpp
+++ b/Engine/Source/Editor/UnrealEd/Private/EditorActorFolders.cpp
@@ -611,6 +611,9 @@ void FActorFolders::OnExecuteActorEditorContextAction(UWorld* InWorld, const EAc
 				(*Folders)->PopActorEditorContext();
 			}
 			break;
+		case EActorEditorContextAction::InitializeContextFromActor:
+			SetActorEditorContextFolder(*(InActor->GetWorld()), InActor->GetFolder());
+			break;
 	}
 }
 
diff --git a/Engine/Source/Editor/UnrealEd/Private/Subsystems/ActorEditorContextSubsystem.cpp b/Engine/Source/Editor/UnrealEd/Private/Subsystems/ActorEditorContextSubsystem.cpp
--- a/Engine/Source/Editor/UnrealEd/Private/Subsystems/ActorEditorContextSubsystem.cpp
+++ b/Engine/Source/Editor/UnrealEd/Private/Subsystems/ActorEditorContextSubsystem.cpp
@@ -134,6 +134,21 @@ void UActorEditorContextSubsystem::PopContext()
 	ActorEditorContextSubsystemChanged.Broadcast();
 }
 
+void UActorEditorContextSubsystem::InitializeContextFromActor(AActor* Actor)
+{
+	UWorld* World = GetWorld();
+	if (Clients.IsEmpty() || !World || !Actor)
+	{
+		return;
+	}
+
+	for (IActorEditorContextClient* Client : Clients)
+	{
+		Client->OnExecuteActorEditorContextAction(World, EActorEditorContextAction::InitializeContextFromActor, Actor);
+	}
+	ActorEditorContextSubsystemChanged.Broadcast();
+}
+
 TArray<IActorEditorContextClient*> UActorEditorContextSubsystem::GetDisplayableClients() const
 {
 	TArray<IActorEditorContextClient*> DisplayableClients;
diff --git a/Engine/Source/Editor/UnrealEd/Private/WorldPartition/WorldPartitionFoliageBuilder.cpp b/Engine/Source/Editor/UnrealEd/Private/WorldPartition/WorldPartitionFoliageBuilder.cpp
--- a/Engine/Source/Editor/UnrealEd/Private/WorldPartition/WorldPartitionFoliageBuilder.cpp
+++ b/Engine/Source/Editor/UnrealEd/Private/WorldPartition/WorldPartitionFoliageBuilder.cpp
@@ -124,8 +124,8 @@ bool UWorldPartitionFoliageBuilder::RunInternal(UWorld* World, const FCellInfo&
 			const FDataLayerEditorContext DataLayerEditorContext(World, ActorDesc->GetDataLayerInstanceNames());
 			const FActorPartitionIdentifier ActorPartitionIdentifier(ActorDesc->GetActorNativeClass(), IFA->GetGridGuid(), DataLayerEditorContext.GetHash());
 			const bool bShouldIncludeGridSizeInName = ActorPartitionIdentifier.GetClass()->GetDefaultObject<APartitionActor>()->ShouldIncludeGridSizeInName(World, ActorPartitionIdentifier);
-			const FString ExpectedActorName = APartitionActor::GetActorName(World, ActorPartitionIdentifier.GetClass(), ActorPartitionIdentifier.GetGridGuid(), ActorPartitionIdentifier, PartitionActorDesc->GridSize, 
-				PartitionActorDesc->GridIndexX, PartitionActorDesc->GridIndexY, PartitionActorDesc->GridIndexZ, DataLayerEditorContext.GetHash());
+			const FString ExpectedActorName = APartitionActor::GetActorName(World, ActorPartitionIdentifier, PartitionActorDesc->GridSize, 
+				PartitionActorDesc->GridIndexX, PartitionActorDesc->GridIndexY, PartitionActorDesc->GridIndexZ);
 
 			if (ActorDesc->GetActorName() != ExpectedActorName)
 			{
diff --git a/Engine/Source/Editor/UnrealEd/Public/IActorEditorContextClient.h b/Engine/Source/Editor/UnrealEd/Public/IActorEditorContextClient.h
--- a/Engine/Source/Editor/UnrealEd/Public/IActorEditorContextClient.h
+++ b/Engine/Source/Editor/UnrealEd/Public/IActorEditorContextClient.h
@@ -26,7 +26,8 @@ enum class EActorEditorContextAction
 	ApplyContext,
 	ResetContext,
 	PushContext,
-	PopContext
+	PopContext,
+	InitializeContextFromActor,
 };
 
 struct IActorEditorContextClient
diff --git a/Engine/Source/Editor/UnrealEd/Public/Subsystems/ActorEditorContextSubsystem.h b/Engine/Source/Editor/UnrealEd/Public/Subsystems/ActorEditorContextSubsystem.h
--- a/Engine/Source/Editor/UnrealEd/Public/Subsystems/ActorEditorContextSubsystem.h
+++ b/Engine/Source/Editor/UnrealEd/Public/Subsystems/ActorEditorContextSubsystem.h
@@ -33,6 +33,7 @@ public:
 	UNREALED_API void ResetContext(IActorEditorContextClient* Client);
 	UNREALED_API void PushContext();
 	UNREALED_API void PopContext();
+	UNREALED_API void InitializeContextFromActor(AActor* Actor);
 	UNREALED_API TArray<IActorEditorContextClient*> GetDisplayableClients() const;
 	FOnActorEditorContextSubsystemChanged& OnActorEditorContextSubsystemChanged() { return ActorEditorContextSubsystemChanged; }
 
diff --git a/Engine/Source/Editor/WorldPartitionEditor/Private/WorldPartition/ContentBundle/ContentBundleEditorSubsystem.cpp b/Engine/Source/Editor/WorldPartitionEditor/Private/WorldPartition/ContentBundle/ContentBundleEditorSubsystem.cpp
--- a/Engine/Source/Editor/WorldPartitionEditor/Private/WorldPartition/ContentBundle/ContentBundleEditorSubsystem.cpp
+++ b/Engine/Source/Editor/WorldPartitionEditor/Private/WorldPartition/ContentBundle/ContentBundleEditorSubsystem.cpp
@@ -1,6 +1,6 @@
 // Copyright Epic Games, Inc. All Rights Reserved.
 #include "WorldPartition/ContentBundle/ContentBundleEditorSubsystem.h"
-
+#include "WorldPartition/ContentBundle/ContentBundleEngineSubsystem.h"
 #include "WorldPartition/ContentBundle/ContentBundle.h"
 #include "WorldPartition/ContentBundle/ContentBundleDescriptor.h"
 #include "WorldPartition/ContentBundle/ContentBundle.h"
@@ -30,6 +30,11 @@ void UContentBundleEditingSubmodule::DoDenitialize()
 	}
 
 	EditingContentBundleGuid.Invalidate();
+	if (GEngine)
+	{
+		UContentBundleEngineSubsystem::Get()->SetEditingContentBundleGuid(EditingContentBundleGuid);
+	}
+
 	EditingContentBundlesStack.Empty();
 }
 
@@ -54,6 +59,8 @@ void UContentBundleEditingSubmodule::PostEditUndo()
 		{
 			StartEditing(NewEditingContentBundle);
 		}
+
+		UContentBundleEngineSubsystem::Get()->SetEditingContentBundleGuid(EditingContentBundleGuid);
 	}
 
 	PreUndoRedoEditingContentBundleGuid.Invalidate();
@@ -81,6 +88,17 @@ void UContentBundleEditingSubmodule::OnExecuteActorEditorContextAction(UWorld* I
 	case EActorEditorContextAction::PopContext:
 		PopContentBundleEditing();
 		break;
+	case EActorEditorContextAction::InitializeContextFromActor:
+		{
+			if (InActor->GetContentBundleGuid().IsValid())
+			{
+				if (TSharedPtr<FContentBundleEditor> EditingContentBundle = GetEditorContentBundle(InActor->GetContentBundleGuid()))
+				{
+					ActivateContentBundleEditing(EditingContentBundle);
+				}
+			}
+		}
+		break;
 	}
 }
 
@@ -163,6 +181,7 @@ bool UContentBundleEditingSubmodule::ActivateContentBundleEditing(TSharedPtr<FCo
 			DeactivateCurrentContentBundleEditing();
 		}
 		EditingContentBundleGuid = ContentBundleEditor->GetDescriptor()->GetGuid();
+		UContentBundleEngineSubsystem::Get()->SetEditingContentBundleGuid(EditingContentBundleGuid);
 		StartEditing(ContentBundleEditor);
 		return true;
 	}
@@ -186,6 +205,7 @@ bool UContentBundleEditingSubmodule::DeactivateContentBundleEditing(TSharedPtr<F
 	{
 		check(EditingContentBundleGuid == ContentBundleEditor->GetDescriptor()->GetGuid());
 		EditingContentBundleGuid.Invalidate();
+		UContentBundleEngineSubsystem::Get()->SetEditingContentBundleGuid(EditingContentBundleGuid);
 		StopEditing(ContentBundleEditor);
 		return true;
 	}
diff --git a/Engine/Source/Runtime/Engine/Private/ActorEditorContext/ScopedActorEditorContextFromActor.cpp b/Engine/Source/Runtime/Engine/Private/ActorEditorContext/ScopedActorEditorContextFromActor.cpp
new file mode 100644
--- /dev/null
+++ b/Engine/Source/Runtime/Engine/Private/ActorEditorContext/ScopedActorEditorContextFromActor.cpp
@@ -0,0 +1,19 @@
+// Copyright Epic Games, Inc. All Rights Reserved.
+
+#if WITH_EDITOR
+
+#include "ActorEditorContext/ScopedActorEditorContextFromActor.h"
+#include "Subsystems/ActorEditorContextSubsystem.h"
+
+FScopedActorEditorContextFromActor::FScopedActorEditorContextFromActor(AActor* InActor)
+{
+	UActorEditorContextSubsystem::Get()->PushContext();
+	UActorEditorContextSubsystem::Get()->InitializeContextFromActor(InActor);
+}
+
+FScopedActorEditorContextFromActor::~FScopedActorEditorContextFromActor()
+{
+	UActorEditorContextSubsystem::Get()->PopContext();
+}
+
+#endif
\ No newline at end of file
diff --git a/Engine/Source/Runtime/Engine/Private/ActorPartition/ActorPartitionSubsystem.cpp b/Engine/Source/Runtime/Engine/Private/ActorPartition/ActorPartitionSubsystem.cpp
--- a/Engine/Source/Runtime/Engine/Private/ActorPartition/ActorPartitionSubsystem.cpp
+++ b/Engine/Source/Runtime/Engine/Private/ActorPartition/ActorPartitionSubsystem.cpp
@@ -12,6 +12,7 @@
 
 #if WITH_EDITOR
 #include "WorldPartition/DataLayer/DataLayerManager.h"
+#include "WorldPartition/ContentBundle/ContentBundleEngineSubsystem.h"
 #endif
 
 #include UE_INLINE_GENERATED_CPP_BY_NAME(ActorPartitionSubsystem)
@@ -20,6 +21,21 @@ DEFINE_LOG_CATEGORY_STATIC(LogActorPartitionSubsystem, All, All);
 
 #if WITH_EDITOR
 
+// @todo_ow: All this can be converted back to DataLayerEditorContextHash when EDL replaces ContentBundles
+namespace FActorPartitionContextHash
+{
+	uint32 Get(const FGuid& InContentBundleGuid, uint32 InDataLayerEditorContextHash)
+	{
+		return InContentBundleGuid.IsValid() ? FCrc::TypeCrc32(InContentBundleGuid, InDataLayerEditorContextHash) : InDataLayerEditorContextHash;
+	}
+
+	uint32 Get(const FGuid& InContentBundleGuid, UWorld* InWorld, const TArray<FName>& InDataLayerInstanceNames)
+	{
+		FDataLayerEditorContext DataLayerEditorContext(InWorld, InDataLayerInstanceNames);
+		return Get(InContentBundleGuid, DataLayerEditorContext.GetHash());
+	}
+};
+
 FActorPartitionGetParams::FActorPartitionGetParams(const TSubclassOf<APartitionActor>& InActorClass, bool bInCreate, ULevel* InLevelHint, const FVector& InLocationHint, uint32 InGridSize, const FGuid& InGuidHint, bool bInBoundsSearch, TFunctionRef<void(APartitionActor*)> InActorCreated)
 	: ActorClass(InActorClass)
 	, bCreate(bInCreate)
@@ -192,16 +208,19 @@ public:
 	{
 		APartitionActor* FoundActor = nullptr;
 		bool bUnloadedActorExists = false;
-		auto FindActor = [&FoundActor, &bUnloadedActorExists, InCellCoord, InActorPartitionId, InGridSize, ThisWorld = World](const FWorldPartitionActorDesc* ActorDesc)
+		FGuid ContentBundleGuid = UContentBundleEngineSubsystem::Get()->GetEditingContentBundleGuid();
+
+		auto FindActor = [&FoundActor, &bUnloadedActorExists, &ContentBundleGuid, InCellCoord, InActorPartitionId, InGridSize, ThisWorld = World](const FWorldPartitionActorDesc* ActorDesc)
 		{
 			check(ActorDesc->GetActorNativeClass()->IsChildOf(InActorPartitionId.GetClass()));
 			FPartitionActorDesc* PartitionActorDesc = (FPartitionActorDesc*)ActorDesc;
+
 			if ((PartitionActorDesc->GridIndexX == InCellCoord.X) &&
 				(PartitionActorDesc->GridIndexY == InCellCoord.Y) &&
 				(PartitionActorDesc->GridIndexZ == InCellCoord.Z) &&
 				(PartitionActorDesc->GridSize == InGridSize) &&
 				(PartitionActorDesc->GridGuid == InActorPartitionId.GetGridGuid()) &&
-				(FDataLayerEditorContext(ThisWorld, PartitionActorDesc->GetDataLayerInstanceNames()).GetHash() == InActorPartitionId.GetDataLayerEditorContextHash()))
+				(FActorPartitionContextHash::Get(PartitionActorDesc->GetContentBundleGuid(), ThisWorld, PartitionActorDesc->GetDataLayerInstanceNames()) == InActorPartitionId.GetContextHash()))
 			{
 				AActor* DescActor = ActorDesc->GetActor();
 
@@ -241,7 +260,7 @@ public:
 				
 		if (!FoundActor && bInCreate)
 		{
-			const FString ActorName = APartitionActor::GetActorName(World, InActorPartitionId.GetClass(), InActorPartitionId.GetGridGuid(), InActorPartitionId, InGridSize, InCellCoord.X, InCellCoord.Y, InCellCoord.Z, InActorPartitionId.GetDataLayerEditorContextHash());
+			const FString ActorName = APartitionActor::GetActorName(World, InActorPartitionId, InGridSize, InCellCoord.X, InCellCoord.Y, InCellCoord.Z);
 
 			FActorSpawnParameters SpawnParams;
 			SpawnParams.OverrideLevel = InCellCoord.Level;
@@ -270,19 +289,7 @@ public:
 			
 			InActorCreated(FoundActor);
 
-			// Once actor is created, update its label
-			TStringBuilderWithBuffer<TCHAR, NAME_SIZE> ActorLabelBuilder;
-			ActorLabelBuilder += FString::Printf(TEXT("%s"), *InActorPartitionId.GetClass()->GetName());
-			if (FoundActor->ShouldIncludeGridSizeInLabel())
-			{
-				ActorLabelBuilder += FString::Printf(TEXT("_%u"), InGridSize);
-			}
-			ActorLabelBuilder += FString::Printf(TEXT("_%d_%d_%d"), InCellCoord.X, InCellCoord.Y, InCellCoord.Z);
-			if (InActorPartitionId.GetDataLayerEditorContextHash() != FDataLayerEditorContext::EmptyHash)
-			{
-				ActorLabelBuilder += FString::Printf(TEXT("_%X"), InActorPartitionId.GetDataLayerEditorContextHash());
-			}
-			FoundActor->SetActorLabel(*ActorLabelBuilder);
+			APartitionActor::SetLabelForActor(FoundActor, InActorPartitionId, InGridSize, InCellCoord.X, InCellCoord.Y, InCellCoord.Z);
 		}
 
 		check(FoundActor || !bInCreate);
@@ -398,26 +405,15 @@ APartitionActor* UActorPartitionSubsystem::GetActor(const TSubclassOf<APartition
 	const uint32 GridSize = InGridSize > 0 ? InGridSize : InActorClass->GetDefaultObject<APartitionActor>()->GetDefaultGridSize(GetWorld());
 	
 	UWorld* World = GetWorld();
+	FGuid ContentBundleGuid = UContentBundleEngineSubsystem::Get()->GetEditingContentBundleGuid();
 	const UDataLayerManager* DataLayerManager = UDataLayerManager::GetDataLayerManager(World);
-	
-	auto WrappedInActorCreated = [DataLayerManager, InActorCreated](APartitionActor* PartitionActor)
-	{
-		if(DataLayerManager)
-		{
-			for (UDataLayerInstance* DataLayer : DataLayerManager->GetActorEditorContextDataLayers())
-			{
-				PartitionActor->AddDataLayer(DataLayer);
-			}
-		}
-		InActorCreated(PartitionActor);
-	};
-
-	FActorPartitionIdentifier ActorPartitionId(InActorClass, InGuid, DataLayerManager ? DataLayerManager->GetDataLayerEditorContextHash() : FDataLayerEditorContext::EmptyHash);
+	const uint32 DataLayerEditorContextHash = DataLayerManager ? DataLayerManager->GetDataLayerEditorContextHash() : FDataLayerEditorContext::EmptyHash;
+	FActorPartitionIdentifier ActorPartitionId(InActorClass, InGuid, FActorPartitionContextHash::Get(ContentBundleGuid, DataLayerEditorContextHash));
 	TMap<FActorPartitionIdentifier, TWeakObjectPtr<APartitionActor>>* ActorsPerId = PartitionedActors.Find(InCellCoords);
 	APartitionActor* FoundActor = nullptr;
 	if (!ActorsPerId)
 	{
-		FoundActor = ActorPartition->GetActor(ActorPartitionId, bInCreate, InCellCoords, GridSize, bInBoundsSearch, WrappedInActorCreated);
+		FoundActor = ActorPartition->GetActor(ActorPartitionId, bInCreate, InCellCoords, GridSize, bInBoundsSearch, InActorCreated);
 		if (FoundActor)
 		{
 			PartitionedActors.Add(InCellCoords).Add(ActorPartitionId, FoundActor);
@@ -428,7 +424,7 @@ APartitionActor* UActorPartitionSubsystem::GetActor(const TSubclassOf<APartition
 		TWeakObjectPtr<APartitionActor>* ActorPtr = ActorsPerId->Find(ActorPartitionId);
 		if (!ActorPtr || !ActorPtr->IsValid())
 		{
-			FoundActor = ActorPartition->GetActor(ActorPartitionId, bInCreate, InCellCoords, GridSize, bInBoundsSearch, WrappedInActorCreated);
+			FoundActor = ActorPartition->GetActor(ActorPartitionId, bInCreate, InCellCoords, GridSize, bInBoundsSearch, InActorCreated);
 			if (FoundActor)
 			{
 				if (!ActorPtr)
diff --git a/Engine/Source/Runtime/Engine/Private/ActorPartition/PartitionActor.cpp b/Engine/Source/Runtime/Engine/Private/ActorPartition/PartitionActor.cpp
--- a/Engine/Source/Runtime/Engine/Private/ActorPartition/PartitionActor.cpp
+++ b/Engine/Source/Runtime/Engine/Private/ActorPartition/PartitionActor.cpp
@@ -1,7 +1,7 @@
 // Copyright Epic Games, Inc. All Rights Reserved.
 
 #include "ActorPartition/PartitionActor.h"
-
+#include "ActorPartition/ActorPartitionSubsystem.h"
 #include UE_INLINE_GENERATED_CPP_BY_NAME(PartitionActor)
 
 #if WITH_EDITOR
@@ -71,33 +71,55 @@ void APartitionActor::SetGridSize(uint32 InGridSize)
 }
 PRAGMA_ENABLE_DEPRECATION_WARNINGS
 
-FString APartitionActor::GetActorName(UWorld* World, const UClass* Class, const FGuid& GridGuid, const FActorPartitionIdentifier& ActorPartitionId, uint32 GridSize, int32 CellCoordsX, int32 CellCoordsY, int32 CellCoordsZ, uint32 DataLayerEditorContext)
+FString APartitionActor::GetActorName(UWorld* World, const UClass* Class, const FGuid& GridGuid, const FActorPartitionIdentifier& ActorPartitionId, uint32 GridSize, int32 CellCoordsX, int32 CellCoordsY, int32 CellCoordsZ, uint32 ContextHash)
+{
+	return GetActorName(World, ActorPartitionId, GridSize, CellCoordsX, CellCoordsY, CellCoordsZ);
+}
+
+FString APartitionActor::GetActorName(UWorld* World, const FActorPartitionIdentifier& ActorPartitionId, uint32 GridSize,  int32 CellCoordsX, int32 CellCoordsY, int32 CellCoordsZ)
 {
 	TStringBuilderWithBuffer<TCHAR, NAME_SIZE> ActorNameBuilder;
 
-	ActorNameBuilder += Class->GetName();
+	ActorNameBuilder += ActorPartitionId.GetClass()->GetName();
 	ActorNameBuilder += TEXT("_");
 
-	if (GridGuid.IsValid())
+	if (ActorPartitionId.GetGridGuid().IsValid())
 	{
-		ActorNameBuilder += GridGuid.ToString(EGuidFormats::Base36Encoded);
+		ActorNameBuilder += ActorPartitionId.GetGridGuid().ToString(EGuidFormats::Base36Encoded);
 		ActorNameBuilder += TEXT("_");
 	}
 
-	if (Class->GetDefaultObject<APartitionActor>()->ShouldIncludeGridSizeInName(World, ActorPartitionId))
+	if (ActorPartitionId.GetClass()->GetDefaultObject<APartitionActor>()->ShouldIncludeGridSizeInName(World, ActorPartitionId))
 	{
 		ActorNameBuilder += FString::Printf(TEXT("%d_"), GridSize);
 	}
 
 	ActorNameBuilder += FString::Printf(TEXT("%d_%d_%d"), CellCoordsX, CellCoordsY, CellCoordsZ);
 
-	if (DataLayerEditorContext != FDataLayerEditorContext::EmptyHash)
+	if (ActorPartitionId.GetContextHash() != FActorPartitionIdentifier::EmptyContextHash)
 	{
-		ActorNameBuilder += FString::Printf(TEXT("_%X"), DataLayerEditorContext);
+		ActorNameBuilder += FString::Printf(TEXT("_%X"), ActorPartitionId.GetContextHash());
 	}
 
 	return ActorNameBuilder.ToString();
 }
+
+void APartitionActor::SetLabelForActor(APartitionActor* Actor, const FActorPartitionIdentifier& ActorPartitionId, uint32 GridSize, int32 CellCoordsX, int32 CellCoordsY, int32 CellCoordsZ)
+{
+	// Once actor is created, update its label
+	TStringBuilderWithBuffer<TCHAR, NAME_SIZE> ActorLabelBuilder;
+	ActorLabelBuilder += FString::Printf(TEXT("%s"), *ActorPartitionId.GetClass()->GetName());
+	if (Actor->ShouldIncludeGridSizeInLabel())
+	{
+		ActorLabelBuilder += FString::Printf(TEXT("_%u"), GridSize);
+	}
+	ActorLabelBuilder += FString::Printf(TEXT("_%d_%d_%d"), CellCoordsX, CellCoordsY, CellCoordsZ);
+	if (ActorPartitionId.GetContextHash() != FActorPartitionIdentifier::EmptyContextHash)
+	{
+		ActorLabelBuilder += FString::Printf(TEXT("_%X"), ActorPartitionId.GetContextHash());
+	}
+	Actor->SetActorLabel(*ActorLabelBuilder);
+}
 #endif
 
 #undef LOCTEXT_NAMESPACE
diff --git a/Engine/Source/Runtime/Engine/Public/ActorEditorContext/ScopedActorEditorContextFromActor.h b/Engine/Source/Runtime/Engine/Public/ActorEditorContext/ScopedActorEditorContextFromActor.h
new file mode 100644
--- /dev/null
+++ b/Engine/Source/Runtime/Engine/Public/ActorEditorContext/ScopedActorEditorContextFromActor.h
@@ -0,0 +1,21 @@
+// Copyright Epic Games, Inc. All Rights Reserved.
+
+#pragma once
+
+#if WITH_EDITOR
+
+#include "CoreFwd.h"
+
+class AActor;
+
+/**
+ * Pushes a new context initialized from the provided actor.
+ */
+class ENGINE_API FScopedActorEditorContextFromActor
+{
+public:
+	FScopedActorEditorContextFromActor(AActor* InActor);
+	~FScopedActorEditorContextFromActor();
+};
+
+#endif
\ No newline at end of file
diff --git a/Engine/Source/Runtime/Engine/Public/ActorPartition/ActorPartitionSubsystem.h b/Engine/Source/Runtime/Engine/Public/ActorPartition/ActorPartitionSubsystem.h
--- a/Engine/Source/Runtime/Engine/Public/ActorPartition/ActorPartitionSubsystem.h
+++ b/Engine/Source/Runtime/Engine/Public/ActorPartition/ActorPartitionSubsystem.h
@@ -51,27 +51,27 @@ struct FActorPartitionGetParams
 
 struct FActorPartitionIdentifier
 {
-	FActorPartitionIdentifier(UClass* InClass, const FGuid& InGridGuid, const uint32& InDataLayerEditorContextHash)
+	FActorPartitionIdentifier(UClass* InClass, const FGuid& InGridGuid, const uint32& InContextHash)
 		: Class(InClass)
 		, GridGuid(InGridGuid)
-		, DataLayerEditorContextHash(InDataLayerEditorContextHash)
+		, ContextHash(InContextHash)
 	{}
 
-	bool operator == (const FActorPartitionIdentifier& Other) const { return Class == Other.Class && GridGuid == Other.GridGuid && DataLayerEditorContextHash == Other.DataLayerEditorContextHash; }
+	bool operator == (const FActorPartitionIdentifier& Other) const { return Class == Other.Class && GridGuid == Other.GridGuid && ContextHash == Other.ContextHash; }
 
 	friend uint32 GetTypeHash(const FActorPartitionIdentifier& Id)
 	{
-		return GetTypeHash(Id.Class.Get()->GetName()) ^ GetTypeHash(Id.GridGuid) ^ GetTypeHash(Id.DataLayerEditorContextHash);
+		return GetTypeHash(Id.Class.Get()->GetName()) ^ GetTypeHash(Id.GridGuid) ^ GetTypeHash(Id.ContextHash);
 	}
 
 	const TSubclassOf<APartitionActor>& GetClass() const { return Class; }
 	const FGuid& GetGridGuid() const { return GridGuid; }
-	uint32 GetDataLayerEditorContextHash() const { return DataLayerEditorContextHash; }
-
+	uint32 GetContextHash() const { return ContextHash; }
+	static const uint32 EmptyContextHash = 0;
 private:
 	TSubclassOf<APartitionActor> Class;
 	FGuid GridGuid;
-	uint32 DataLayerEditorContextHash;
+	uint32 ContextHash;
 };
 
 #endif
diff --git a/Engine/Source/Runtime/Engine/Public/ActorPartition/PartitionActor.h b/Engine/Source/Runtime/Engine/Public/ActorPartition/PartitionActor.h
--- a/Engine/Source/Runtime/Engine/Public/ActorPartition/PartitionActor.h
+++ b/Engine/Source/Runtime/Engine/Public/ActorPartition/PartitionActor.h
@@ -41,7 +41,11 @@ public:
 	UPROPERTY()
 	uint32 GridSize;
 
-	static ENGINE_API FString GetActorName(UWorld* World, const UClass* Class, const FGuid& GridGuid, const FActorPartitionIdentifier& ActorPartitionId, uint32 GridSize, int32 CellCoordsX, int32 CellCoordsY, int32 CellCoordsZ, uint32 DataLayerEditorContext);
+	UE_DEPRECATED(5.4, "Use version that takes FActorPartitionIdentifier as second parameter")
+	static ENGINE_API FString GetActorName(UWorld* World, const UClass* Class, const FGuid& GridGuid, const FActorPartitionIdentifier& ActorPartitionId, uint32 GridSize, int32 CellCoordsX, int32 CellCoordsY, int32 CellCoordsZ, uint32 ContextHash);
+
+	static ENGINE_API FString GetActorName(UWorld* World, const FActorPartitionIdentifier& ActorPartitionId, uint32 GridSize, int32 CellCoordsX, int32 CellCoordsY, int32 CellCoordsZ);
+	static ENGINE_API void SetLabelForActor(APartitionActor* PartitionActor, const FActorPartitionIdentifier& ActorPartitionId, uint32 GridSize, int32 CellCoordsX, int32 CellCoordsY, int32 CellCoordsZ);
 #endif
 };
 
diff --git a/Engine/Source/Runtime/Engine/Public/WorldPartition/ContentBundle/ContentBundleEngineSubsystem.h b/Engine/Source/Runtime/Engine/Public/WorldPartition/ContentBundle/ContentBundleEngineSubsystem.h
--- a/Engine/Source/Runtime/Engine/Public/WorldPartition/ContentBundle/ContentBundleEngineSubsystem.h
+++ b/Engine/Source/Runtime/Engine/Public/WorldPartition/ContentBundle/ContentBundleEngineSubsystem.h
@@ -43,6 +43,10 @@ public:
 
 	static ENGINE_API UContentBundleEngineSubsystem* Get();
 
+#if WITH_EDITOR
+	void SetEditingContentBundleGuid(FGuid InGuid) { EditingContentBundleGuid = InGuid; }
+	FGuid GetEditingContentBundleGuid() { return EditingContentBundleGuid; }
+#endif
 private:
 	ENGINE_API void OnPreWorldInitialization(UWorld* World, const UWorld::InitializationValues IVS);
 	ENGINE_API void OnWorldPostCleanup(UWorld* World, bool bSessionEnded, bool bCleanupResources);
@@ -57,4 +61,9 @@ private:
 
 	UPROPERTY(Transient)
 	TObjectPtr<UContentBundleTypeFactory> ContentBundleTypeFactory;
+
+#if WITH_EDITOR
+	// @todo_ow: This will no longer be needed once we move ContentBundles to External Data Layers but currently the ActorPartitionSubsystem needs to get the current ContentBundleGuid to build a proper actor name
+	FGuid EditingContentBundleGuid;
+#endif
 };
