Fix actors with custom depth not being occluded when r.CustomDepth==0. The occlusion code had no knowledge of whether or not custom depth was enabled globally, it only checked bRenderCustomDepth on a per-proxy basis. Adding a check of the cvar in ShouldRenderCustomDepth() fixes that. Also added a callback, OnCustomDepthChanged, that recreates all the proxies when custom depth changes so culling updates immediately when toggling r.CustomDepth.

diff --git a/Engine/Source/Runtime/Engine/Private/PrimitiveSceneProxy.cpp b/Engine/Source/Runtime/Engine/Private/PrimitiveSceneProxy.cpp
--- a/Engine/Source/Runtime/Engine/Private/PrimitiveSceneProxy.cpp
+++ b/Engine/Source/Runtime/Engine/Private/PrimitiveSceneProxy.cpp
@@ -21,6 +21,7 @@
 #include "PrimitiveInstanceUpdateCommand.h"
 #include "NaniteSceneProxy.h" // TODO: PROG_RASTER
 #include "ComponentRecreateRenderStateContext.h"
+#include "ComponentReregisterContext.h"
 #include "DataDrivenShaderPlatformInfo.h"
 #include "SceneInterface.h"
 #include "PrimitiveUniformShaderParametersBuilder.h"
@@ -106,6 +107,40 @@ static TAutoConsoleVariable<int32> CVarApproximateOcclusionQueries(
 	ECVF_RenderThreadSafe
 );
 
+static int32 GCustomDepthMode = 1;
+static TAutoConsoleVariable<int32> CVarCustomDepth(
+	TEXT("r.CustomDepth"),
+	GCustomDepthMode,
+	TEXT("0: feature is disabled\n")
+	TEXT("1: feature is enabled, texture is created on demand\n")
+	TEXT("2: feature is enabled, texture is not released until required (should be the project setting if the feature should not stall)\n")
+	TEXT("3: feature is enabled, stencil writes are enabled, texture is not released until required (should be the project setting if the feature should not stall)"),
+	ECVF_RenderThreadSafe);
+
+void OnCustomDepthChanged(IConsoleVariable* Var)
+{
+	// Easiest way to update all static scene proxies is to recreate them
+	FlushRenderingCommands();
+	FGlobalComponentReregisterContext ReregisterContext;
+	GCustomDepthMode = CVarCustomDepth.GetValueOnAnyThread();
+}
+
+void InitCustomDepth()
+{
+	CVarCustomDepth.AsVariable()->SetOnChangedCallback(FConsoleVariableDelegate::CreateStatic(&OnCustomDepthChanged));
+}
+
+bool FPrimitiveSceneProxy::ShouldRenderCustomDepth() const
+{
+	const bool bCustomDepthEnabled = GCustomDepthMode != 0;
+	return bCustomDepthEnabled && bRenderCustomDepth;
+}
+
+ECustomDepthMode GetCustomDepthMode()
+{
+	return (ECustomDepthMode)GCustomDepthMode;
+}
+
 bool IsAllowingApproximateOcclusionQueries()
 {
 	return CVarApproximateOcclusionQueries.GetValueOnAnyThread() != 0;
diff --git a/Engine/Source/Runtime/Engine/Private/UnrealEngine.cpp b/Engine/Source/Runtime/Engine/Private/UnrealEngine.cpp
--- a/Engine/Source/Runtime/Engine/Private/UnrealEngine.cpp
+++ b/Engine/Source/Runtime/Engine/Private/UnrealEngine.cpp
@@ -256,6 +256,7 @@ IMPLEMENT_MODULE( FEngineModule, Engine );
 #define LOCTEXT_NAMESPACE "UnrealEngine"
 
 void InitScreenPercentage();
+void InitCustomDepth();
 
 void OnChangeEngineCVarRequiringRecreateRenderState(IConsoleVariable* Var)
 {
@@ -282,6 +283,7 @@ void FEngineModule::StartupModule()
 	ResumeTextureStreamingRenderTasks = &ResumeRenderAssetStreaming;
 
 	InitScreenPercentage();
+	InitCustomDepth();
 
 	FParticleSystemWorldManager::OnStartup();
 
diff --git a/Engine/Source/Runtime/Engine/Public/PrimitiveSceneProxy.h b/Engine/Source/Runtime/Engine/Public/PrimitiveSceneProxy.h
--- a/Engine/Source/Runtime/Engine/Public/PrimitiveSceneProxy.h
+++ b/Engine/Source/Runtime/Engine/Public/PrimitiveSceneProxy.h
@@ -138,6 +138,19 @@ public:
 	{}
 };
 
+enum class ECustomDepthMode : uint32
+{
+	// Custom depth is disabled.
+	Disabled,
+
+	// Custom depth is enabled.
+	Enabled,
+
+	// Custom depth is enabled and uses stencil.
+	EnabledWithStencil,
+};
+
+ENGINE_API extern ECustomDepthMode GetCustomDepthMode();
 extern bool IsOptimizedWPO();
 ENGINE_API extern bool ShouldOptimizedWPOAffectNonNaniteShaderSelection();
 extern bool IsAllowingApproximateOcclusionQueries();
@@ -650,7 +663,7 @@ public:
 	inline bool IsEditingLevelInstanceChild() const { return bLevelInstanceEditingState; }
 	inline bool IsSelected() const { return IsParentSelected() || IsIndividuallySelected(); }
 	inline bool WantsSelectionOutline() const { return bWantsSelectionOutline; }
-	inline bool ShouldRenderCustomDepth() const { return bRenderCustomDepth; }
+	ENGINE_API bool ShouldRenderCustomDepth() const;
 	inline bool IsVisibleInSceneCaptureOnly() const { return bVisibleInSceneCaptureOnly; }
 	inline bool IsHiddenInSceneCapture() const { return bHiddenInSceneCapture; }
 	inline uint8 GetCustomDepthStencilValue() const { return CustomDepthStencilValue; }
diff --git a/Engine/Source/Runtime/Renderer/Private/CustomDepthRendering.cpp b/Engine/Source/Runtime/Renderer/Private/CustomDepthRendering.cpp
--- a/Engine/Source/Runtime/Renderer/Private/CustomDepthRendering.cpp
+++ b/Engine/Source/Runtime/Renderer/Private/CustomDepthRendering.cpp
@@ -11,15 +11,6 @@
 #include "MeshPassProcessor.inl"
 #include "UnrealEngine.h"
 
-static TAutoConsoleVariable<int32> CVarCustomDepth(
-	TEXT("r.CustomDepth"),
-	1,
-	TEXT("0: feature is disabled\n")
-	TEXT("1: feature is enabled, texture is created on demand\n")
-	TEXT("2: feature is enabled, texture is not released until required (should be the project setting if the feature should not stall)\n")
-	TEXT("3: feature is enabled, stencil writes are enabled, texture is not released until required (should be the project setting if the feature should not stall)"),
-	ECVF_RenderThreadSafe);
-
 static TAutoConsoleVariable<int32> CVarCustomDepthOrder(
 	TEXT("r.CustomDepth.Order"),
 	2,
@@ -53,17 +44,6 @@ ECustomDepthPassLocation GetCustomDepthPassLocation(EShaderPlatform Platform)
 	return bCustomDepthBeforeBasePase ? ECustomDepthPassLocation::BeforeBasePass : ECustomDepthPassLocation::AfterBasePass;
 }
 
-ECustomDepthMode GetCustomDepthMode()
-{
-	switch (CVarCustomDepth.GetValueOnAnyThread())
-	{
-	case 1: // Fallthrough.
-	case 2: return ECustomDepthMode::Enabled;
-	case 3: return ECustomDepthMode::EnabledWithStencil;
-	}
-	return ECustomDepthMode::Disabled;
-}
-
 bool IsCustomDepthPassWritingStencil()
 {
 	return GetCustomDepthMode() == ECustomDepthMode::EnabledWithStencil;
@@ -73,7 +53,7 @@ FCustomDepthTextures FCustomDepthTextures::Create(FRDGBuilder& GraphBuilder, FIn
 {
 	const ECustomDepthMode CustomDepthMode = GetCustomDepthMode();
 
-	if (!IsCustomDepthPassEnabled())
+	if (CustomDepthMode == ECustomDepthMode::Disabled)
 	{
 		return {};
 	}
diff --git a/Engine/Source/Runtime/Renderer/Private/CustomDepthRendering.h b/Engine/Source/Runtime/Renderer/Private/CustomDepthRendering.h
--- a/Engine/Source/Runtime/Renderer/Private/CustomDepthRendering.h
+++ b/Engine/Source/Runtime/Renderer/Private/CustomDepthRendering.h
@@ -19,26 +19,6 @@ enum class ECustomDepthPassLocation : uint32
 // Returns the location in the frame where custom depth is rendered.
 extern ECustomDepthPassLocation GetCustomDepthPassLocation(EShaderPlatform Platform);
 
-enum class ECustomDepthMode : uint32
-{
-	// Custom depth is disabled.
-	Disabled,
-
-	// Custom depth is enabled.
-	Enabled,
-
-	// Custom depth is enabled and uses stencil.
-	EnabledWithStencil,
-};
-
-// The custom depth mode currently configured.
-extern ECustomDepthMode GetCustomDepthMode();
-
-inline bool IsCustomDepthPassEnabled()
-{
-	return GetCustomDepthMode() != ECustomDepthMode::Disabled;
-}
-
 struct FCustomDepthTextures
 {
 	static FCustomDepthTextures Create(FRDGBuilder& GraphBuilder, FIntPoint CustomDepthExtent, EShaderPlatform ShaderPlatform);
