cluster union needs to set BodyInstance.OwnerComponent so that it can be registered for collision events

diff --git a/Engine/Source/Runtime/Engine/Private/PhysicsEngine/ClusterUnionComponent.cpp b/Engine/Source/Runtime/Engine/Private/PhysicsEngine/ClusterUnionComponent.cpp
--- a/Engine/Source/Runtime/Engine/Private/PhysicsEngine/ClusterUnionComponent.cpp
+++ b/Engine/Source/Runtime/Engine/Private/PhysicsEngine/ClusterUnionComponent.cpp
@@ -571,6 +571,9 @@ void UClusterUnionComponent::OnCreatePhysicsState()
 		return;
 	}
 
+	// We need to set some properties on the BodyInstance to get things to play nice.
+	BodyInstance.OwnerComponent = this;
+
 	// If we're not actually playing/needing this to simulate (e.g. in the editor) there should be no reason to create this proxy.
 	const bool bValidWorld = GetWorld() && (GetWorld()->IsGameWorld() || GetWorld()->IsPreviewWorld());
 	if (!bValidWorld)
