Check if RenderDoc Plugin is loaded before capturing frame trace

#rnx

diff --git a/Engine/Source/Runtime/Engine/Private/Tests/AutomationCommon.cpp b/Engine/Source/Runtime/Engine/Private/Tests/AutomationCommon.cpp
--- a/Engine/Source/Runtime/Engine/Private/Tests/AutomationCommon.cpp
+++ b/Engine/Source/Runtime/Engine/Private/Tests/AutomationCommon.cpp
@@ -21,6 +21,7 @@
 #include "Tests/AutomationTestSettings.h"
 #include "GameMapsSettings.h"
 #include "IRenderCaptureProvider.h"
+#include "Modules/ModuleManager.h"
 #include "Algo/Accumulate.h"
 
 #if WITH_AUTOMATION_TESTS
@@ -158,7 +159,9 @@ namespace AutomationCommon
 	{
 		TArray<uint8> FrameTrace;
 
-		if (CVarAutomationAllowFrameTraceCapture.GetValueOnGameThread() != 0 && IRenderCaptureProvider::IsAvailable())
+		if (CVarAutomationAllowFrameTraceCapture.GetValueOnGameThread() != 0
+			&& IRenderCaptureProvider::IsAvailable()
+			&& FModuleManager::Get().IsModuleLoaded("RenderDocPlugin"))
 		{
 			const FString MapAndTest = MapOrContext / FPaths::MakeValidFileName(TestName, TEXT('_'));
 			FString ScreenshotName = GetScreenshotName(MapAndTest);
