Make sure that Character is initialized with appropriate default gravity. Prevent instances where simulated proxies could attempt to set zero'd gravity direction.
#rnx

diff --git a/Engine/Source/Runtime/Engine/Private/Character.cpp b/Engine/Source/Runtime/Engine/Private/Character.cpp
--- a/Engine/Source/Runtime/Engine/Private/Character.cpp
+++ b/Engine/Source/Runtime/Engine/Private/Character.cpp
@@ -114,6 +114,7 @@ ACharacter::ACharacter(const FObjectInitializer& ObjectInitializer)
 	}
 
 	BaseRotationOffset = FQuat::Identity;
+	ReplicatedGravityDirection = UCharacterMovementComponent::DefaultGravityDirection;
 }
 
 void ACharacter::PostInitializeComponents()
