character movement should attempt to base on the cluster union - not the underlyin geometry collection

diff --git a/Engine/Source/Runtime/Engine/Private/Components/CharacterMovementComponent.cpp b/Engine/Source/Runtime/Engine/Private/Components/CharacterMovementComponent.cpp
--- a/Engine/Source/Runtime/Engine/Private/Components/CharacterMovementComponent.cpp
+++ b/Engine/Source/Runtime/Engine/Private/Components/CharacterMovementComponent.cpp
@@ -6593,7 +6593,13 @@ void UCharacterMovementComponent::ComputeFloorDist(const FVector& CapsuleLocatio
 	}
 
 	bool bBlockingHit = false;
+
 	FCollisionQueryParams QueryParams(SCENE_QUERY_STAT(ComputeFloorDist), false, CharacterOwner);
+	// Having a character base on a component within a cluster union will cause replication problems.
+	// The issue is that ACharacter::SetBase() gets a GeometryCollectionComponent passed to it when standing on the DynamicPlatform
+	// and that GC is never simulating, and since it's not simulating it's stopping the based movement flow there for simulated proxies.
+	QueryParams.bTraceIntoSubComponents = false;
+
 	FCollisionResponseParams ResponseParam;
 	InitCollisionParams(QueryParams, ResponseParam);
 	const ECollisionChannel CollisionChannel = UpdatedComponent->GetCollisionObjectType();
