implement various physical properties accessor functions on the primitive component in terms of the underlying physics object handle

- UPrimitiveComponent::GetPhysicsLinearVelocity
- UPrimitiveComponent::GetPhysicsLinearVelocityAtPoint
- UPrimitiveComponent::GetAngularVelocityInRadians
- UPrimitiveComponent::GetCenterOfMass

diff --git a/Engine/Source/Runtime/Engine/Private/PrimitiveComponentPhysics.cpp b/Engine/Source/Runtime/Engine/Private/PrimitiveComponentPhysics.cpp
--- a/Engine/Source/Runtime/Engine/Private/PrimitiveComponentPhysics.cpp
+++ b/Engine/Source/Runtime/Engine/Private/PrimitiveComponentPhysics.cpp
@@ -277,6 +277,11 @@ FVector UPrimitiveComponent::GetPhysicsLinearVelocity(FName BoneName)
 	{
 		return BI->GetUnrealWorldVelocity();
 	}
+	else if (Chaos::FConstPhysicsObjectHandle Handle = GetPhysicsObjectByName(BoneName))
+	{
+		FLockedReadPhysicsObjectExternalInterface Interface = FPhysicsObjectExternalInterface::LockRead(Handle);
+		return Interface->GetV(Handle);
+	}
 	return FVector(0,0,0);
 }
 
@@ -286,6 +291,11 @@ FVector UPrimitiveComponent::GetPhysicsLinearVelocityAtPoint(FVector Point, FNam
 	{
 		return BI->GetUnrealWorldVelocityAtPoint(Point);
 	}
+	else if (Chaos::FConstPhysicsObjectHandle Handle = GetPhysicsObjectByName(BoneName))
+	{
+		FLockedReadPhysicsObjectExternalInterface Interface = FPhysicsObjectExternalInterface::LockRead(Handle);
+		return Interface->GetVAtPoint(Handle, Point);
+	}
 	return FVector(0, 0, 0);
 }
 
@@ -314,11 +324,15 @@ void UPrimitiveComponent::SetPhysicsMaxAngularVelocityInRadians(float NewMaxAngV
 
 FVector UPrimitiveComponent::GetPhysicsAngularVelocityInRadians(FName BoneName) const
 {
-	FBodyInstance* const BI = GetBodyInstance(BoneName);
-	if(BI != NULL)
+	if (FBodyInstance* const BI = GetBodyInstance(BoneName))
 	{
 		return BI->GetUnrealWorldAngularVelocityInRadians();
 	}
+	else if (Chaos::FConstPhysicsObjectHandle Handle = GetPhysicsObjectByName(BoneName))
+	{
+		FLockedReadPhysicsObjectExternalInterface Interface = FPhysicsObjectExternalInterface::LockRead(Handle);
+		return Interface->GetW(Handle);
+	}
 	return FVector(0,0,0);
 }
 
@@ -328,6 +342,11 @@ FVector UPrimitiveComponent::GetCenterOfMass(FName BoneName) const
 	{
 		return ComponentBodyInstance->GetCOMPosition();
 	}
+	else if (Chaos::FConstPhysicsObjectHandle Handle = GetPhysicsObjectByName(BoneName))
+	{
+		FLockedReadPhysicsObjectExternalInterface Interface = FPhysicsObjectExternalInterface::LockRead(Handle);
+		return Interface->GetWorldCoM(Handle);
+	}
 
 	return FVector::ZeroVector;
 }
diff --git a/Engine/Source/Runtime/Experimental/Chaos/Private/Chaos/PhysicsObjectInterface.cpp b/Engine/Source/Runtime/Experimental/Chaos/Private/Chaos/PhysicsObjectInterface.cpp
--- a/Engine/Source/Runtime/Experimental/Chaos/Private/Chaos/PhysicsObjectInterface.cpp
+++ b/Engine/Source/Runtime/Experimental/Chaos/Private/Chaos/PhysicsObjectInterface.cpp
@@ -252,6 +252,26 @@ namespace Chaos
 		return FVector::Zero();
 	}
 
+	template<EThreadContext Id>
+	FVector FReadPhysicsObjectInterface<Id>::GetVAtPoint(const FConstPhysicsObjectHandle Object, const FVector& Point)
+	{
+		if (!Object)
+		{
+			return FVector::Zero();
+		}
+
+		if (TThreadParticle<Id>* Particle = Object->GetParticle<Id>())
+		{
+			if (TThreadKinematicParticle<Id>* Kinematic = Particle->CastToKinematicParticle())
+			{
+				const FVector CenterOfMass = GetWorldCoM(Object);
+				const FVector Diff = Point - CenterOfMass;
+				return Kinematic->V() - FVector::CrossProduct(Diff, Kinematic->W());
+			}
+		}
+
+		return FVector::Zero();
+	}
 	template<EThreadContext Id>
 	FVector FReadPhysicsObjectInterface<Id>::GetW(const FConstPhysicsObjectHandle Object)
 	{
diff --git a/Engine/Source/Runtime/Experimental/Chaos/Public/Chaos/ParticleHandleFwd.h b/Engine/Source/Runtime/Experimental/Chaos/Public/Chaos/ParticleHandleFwd.h
--- a/Engine/Source/Runtime/Experimental/Chaos/Public/Chaos/ParticleHandleFwd.h
+++ b/Engine/Source/Runtime/Experimental/Chaos/Public/Chaos/ParticleHandleFwd.h
@@ -119,6 +119,9 @@ namespace Chaos
 	template<EThreadContext Id>
 	using TThreadParticle = std::conditional_t<Id == EThreadContext::External, FGeometryParticle, FGeometryParticleHandle>;
 
+	template<EThreadContext Id>
+	using TThreadKinematicParticle = std::conditional_t<Id == EThreadContext::External, FKinematicGeometryParticle, FKinematicGeometryParticleHandle>;
+
 	template<EThreadContext Id>
 	using TThreadRigidParticle = std::conditional_t<Id == EThreadContext::External, FPBDRigidParticle, FPBDRigidParticleHandle>;
 }
diff --git a/Engine/Source/Runtime/Experimental/Chaos/Public/Chaos/PhysicsObjectInterface.h b/Engine/Source/Runtime/Experimental/Chaos/Public/Chaos/PhysicsObjectInterface.h
--- a/Engine/Source/Runtime/Experimental/Chaos/Public/Chaos/PhysicsObjectInterface.h
+++ b/Engine/Source/Runtime/Experimental/Chaos/Public/Chaos/PhysicsObjectInterface.h
@@ -68,6 +68,7 @@ namespace Chaos
 		CHAOS_API FVector GetWorldCoM(const FConstPhysicsObjectHandle Object);
 		CHAOS_API FQuat GetR(const FConstPhysicsObjectHandle Object);
 		CHAOS_API FVector GetV(const FConstPhysicsObjectHandle Object);
+		CHAOS_API FVector GetVAtPoint(const FConstPhysicsObjectHandle Object, const FVector& Point);
 		CHAOS_API FVector GetW(const FConstPhysicsObjectHandle Object);
 		CHAOS_API FSpatialAccelerationIdx GetSpatialIndex(const FConstPhysicsObjectHandle Object);
 
