Fix warning spam from SoundWaveLoadingBehavior

diff --git a/Engine/Source/Runtime/Engine/Classes/Sound/SoundWaveLoadingBehavior.cpp b/Engine/Source/Runtime/Engine/Classes/Sound/SoundWaveLoadingBehavior.cpp
--- a/Engine/Source/Runtime/Engine/Classes/Sound/SoundWaveLoadingBehavior.cpp
+++ b/Engine/Source/Runtime/Engine/Classes/Sound/SoundWaveLoadingBehavior.cpp
@@ -198,19 +198,11 @@ private:
 				UE_LOG(LogAudio, Warning, TEXT("Failed to filter for Soundclasses from the SoundCue dependencies for '%s'"), *CueAsset.PackagePath.ToString());
 				continue;
 			}			
-
-			// Should have a single SoundClass, we hope, otherwise ignore.
-			if (ReferencedSoundClasses.Num() != 1)
-			{
-				UE_CLOG(ReferencedSoundClasses.Num() > 1, LogAudio, Warning, 
-				        TEXT("More than one Soundclass refereneced from this cue '%s'"), *CueAsset.PackagePath.ToString() );
-				continue;
-			}
 			
-			// Look up this classes loading behavior in our cache.
-			const FAssetData& Class = ReferencedSoundClasses[0];
-			FClassData CacheLoadingBehavior;
+			// If there's more than one, rank them.
+			for (const FAssetData& Class : ReferencedSoundClasses)
 			{
+				FClassData CacheLoadingBehavior;
 				FScopeLock Lock(&CacheCS);
 				if (const FClassData* Found = CacheClassLoadingBehaviors.Find(Class.PackageName))
 				{
@@ -220,12 +212,12 @@ private:
 				{
 					CacheLoadingBehavior = LoadAndCacheClass(Class);
 				}
-			}
 
-			// Compare if this is more important
-			if (MostImportantLoadingBehavior.CompareGreater(CacheLoadingBehavior, InTargetPlatform))
-			{
-				MostImportantLoadingBehavior = CacheLoadingBehavior;
+				// Compare if this is more important
+				if (MostImportantLoadingBehavior.CompareGreater(CacheLoadingBehavior, InTargetPlatform))
+				{
+					MostImportantLoadingBehavior = CacheLoadingBehavior;
+				}
 			}
 		}
 		
