additional safeguards to protect against the scenario where the cluster union actor/simulating component replicates after the UClusterUnionReplicatedProxyComponent

diff --git a/Engine/Source/Runtime/Engine/Private/PhysicsEngine/ClusterUnionReplicatedProxyComponent.cpp b/Engine/Source/Runtime/Engine/Private/PhysicsEngine/ClusterUnionReplicatedProxyComponent.cpp
--- a/Engine/Source/Runtime/Engine/Private/PhysicsEngine/ClusterUnionReplicatedProxyComponent.cpp
+++ b/Engine/Source/Runtime/Engine/Private/PhysicsEngine/ClusterUnionReplicatedProxyComponent.cpp
@@ -2,6 +2,7 @@
 
 #include "PhysicsEngine/ClusterUnionReplicatedProxyComponent.h"
 
+#include "Engine/World.h"
 #include "Net/Core/PushModel/PushModel.h"
 #include "Net/UnrealNetwork.h"
 #include "PhysicsEngine/ClusterUnionComponent.h"
@@ -113,9 +114,26 @@ void UClusterUnionReplicatedProxyComponent::PostRepNotifies()
 {
 	UActorComponent::PostRepNotifies();
 
+	if (IsPendingDeletion())
+	{
+		return;
+	}
+
 	const bool bIsValid = ParentClusterUnion.IsValid() && ChildClusteredComponent.IsValid() && !ParticleBoneIds.IsEmpty();
 	if (!bIsValid)
 	{
+		// If any of the above 3 variables isn't valid yet then it probably means they haven't been replicated yet (either by the cluster union proxy component or externally in the case of actors/components).
+		// This should be a relatively rare situation - but could cause the client to never attempt to add into the cluster union it if happens.
+		// If we're still waiting on some replication - then wait until the next natural call to PostRepNotifies.
+		// Otherwise, force a call to PostRepNotifies on the next tick.
+		const bool bIsStillWaitingOnReplication = !bNetUpdateParentClusterUnion || !bNetUpdateChildClusteredComponent || !bNetUpdateParticleBoneIds || !bNetUpdateParticleChildToParents;
+		if (!bIsStillWaitingOnReplication)
+		{
+			if (UWorld* World = GetWorld())
+			{
+				World->GetTimerManager().SetTimerForNextTick(this, &UClusterUnionReplicatedProxyComponent::PostRepNotifies);
+			}
+		}
 		return;
 	}
 
@@ -130,7 +148,7 @@ void UClusterUnionReplicatedProxyComponent::PostRepNotifies()
 		bNetUpdateParticleBoneIds = false;
 	}
 
-	if (bIsValid && bNetUpdateParticleChildToParents && ParticleBoneIds.Num() == ParticleChildToParents.Num())
+	if (bNetUpdateParticleChildToParents && ParticleBoneIds.Num() == ParticleChildToParents.Num())
 	{
 		// This particular bit can't happen until *after* we add the component to the cluster union. There's an additional deferral
 		// in AddComponentToCluster that we have to wait for.
