Introduced Material Translation DDC caching.

This change implements a caching mechanism for material translation results. After translating a material, the intermediate results are now stored in the DDC (Derived Data Cache) for future retrieval. When the material translator is requested to translate a material, it first checks the DDC to see if th combination of specified base material, parameter sets, and platform-specific data is already present in the cache. If this key is found, the intermediate results are pulled and returned, avoiding expensive material translation and decoupling material translation from material shader generation.

The significant advantage of this change is that modifications to the material template source file or any dependent shader file no longer trigger a material retranslation. The material shader can still be generated correctly, using the cached translation results, which ensures that material string generation always uses up-to-date input files.

diff --git a/Engine/Shaders/Private/MaterialTemplate.ush b/Engine/Shaders/Private/MaterialTemplate.ush
--- a/Engine/Shaders/Private/MaterialTemplate.ush
+++ b/Engine/Shaders/Private/MaterialTemplate.ush
@@ -2,6 +2,10 @@
 
 /**
  * MaterialTemplate.usf: Filled in by FHLSLMaterialTranslator::GetMaterialShaderCode for each material being compiled.
+ 
+  $TemplateVersion{1}
+	Bump this number up when you make a change to this file that requires the material
+	translation DDC cache to be invalidated, forcing all materials to re-translate.
  */
 
 #ifndef ENABLE_NEW_HLSL_GENERATOR
@@ -75,13 +79,13 @@
 
 //////////////////////////////////////////////////////////////////////////
 
-#define NUM_MATERIAL_TEXCOORDS_VERTEX %s
-#define NUM_MATERIAL_TEXCOORDS %s
-#define NUM_CUSTOM_VERTEX_INTERPOLATORS %s
-#define NUM_TEX_COORD_INTERPOLATORS %s
+#define NUM_MATERIAL_TEXCOORDS_VERTEX %{num_material_texcoords_vertex}
+#define NUM_MATERIAL_TEXCOORDS %{num_material_texcoords}
+#define NUM_CUSTOM_VERTEX_INTERPOLATORS %{num_custom_vertex_interpolators}
+#define NUM_TEX_COORD_INTERPOLATORS %{num_tex_coord_interpolators}
 
 // Vertex interpolators offsets definition
-%s
+%{vertex_interpolators_offsets_definition}
 
 #if NUM_VIRTUALTEXTURE_SAMPLES || LIGHTMAP_VT_ENABLED
 	#include "/Engine/Private/VirtualTextureCommon.ush"
@@ -293,17 +297,17 @@ float4 GetDynamicParameter(FMaterialParticleParameters Parameters, float4 Defaul
 }
 
 /** Material declarations */
-%s
+%{material_declarations}
 
 /** FMaterialAttributes utilities */
-%s
+%{material_attributes_utilities}
 
 /** 
  * Parameters calculated from the pixel material inputs.
  */
 struct FPixelMaterialInputs
 {
-%s
+%{pixel_material_inputs}
 
 #if TEMPLATE_USES_STRATA
 FStrataData GetFrontStrataData()
@@ -2721,7 +2725,7 @@ float4 MaterialExpressionPathTracingBufferTextureLookup(FMaterialPixelParameters
 #define SwizzleDeriv4(_V, _MASK) ConstructFloatDeriv4(_V.Value._MASK, _V.Ddx._MASK, _V.Ddy._MASK)
 
 // Uniform material expressions.
-%s
+%{uniform_material_expressions}
 
 #if USE_ANALYTIC_DERIVATIVES && TEXTURE_SAMPLE_DEBUG
 MaterialFloat4 DebugTextureCommon(const int Mode, float2 UV, MaterialFloat2 DDX, MaterialFloat2 DDY, MaterialFloat Scale)
@@ -2844,7 +2848,7 @@ half3 GetMaterialEmissive(FPixelMaterialInputs PixelMaterialInputs)
 
 half3 GetMaterialEmissiveForCS(FMaterialPixelParameters Parameters)
 {
-%s;
+%{get_material_emissive_for_cs};
 }
 
 // Shading Model is an uint and represents a SHADINGMODELID_* in ShadingCommon.ush 
@@ -2917,55 +2921,55 @@ half GetMaterialAnisotropy(FPixelMaterialInputs PixelMaterialInputs)
 
 half GetMaterialTranslucencyDirectionalLightingIntensity()
 {
-%s;
+%{get_material_translucency_directional_lighting_intensity};
 }
 
 half GetMaterialTranslucentShadowDensityScale()
 {
-%s;
+%{get_material_translucent_shadow_density_scale};
 }
 
 half GetMaterialTranslucentSelfShadowDensityScale()
 {
-%s;
+%{get_material_translucent_self_shadow_density_scale};
 }
 
 half GetMaterialTranslucentSelfShadowSecondDensityScale()
 {
-%s;
+%{get_material_translucent_self_shadow_second_density_scale};
 }
 
 half GetMaterialTranslucentSelfShadowSecondOpacity()
 {
-%s;
+%{get_material_translucent_self_shadow_second_opacity};
 }
 
 half GetMaterialTranslucentBackscatteringExponent()
 {
-%s;
+%{get_material_translucent_backscattering_exponent};
 }
 
 half3 GetMaterialTranslucentMultipleScatteringExtinction()
 {
-%s;
+%{get_material_translucent_multiple_scattering_extinction};
 }
 
 // This is the clip value constant that is defined in the material (range 0..1)
 // Use GetMaterialMask() to get the Material Mask combined with this.
 half GetMaterialOpacityMaskClipValue()
 {
-%s;
+%{get_material_opacity_mask_clip_value};
 }
 
 #if USES_DISPLACEMENT && IS_NANITE_PASS
 float GetMaterialDisplacementMagnitude()
 {
-%s;
+%{get_material_displacement_magnitude};
 }
 
 float GetMaterialDisplacementCenter()
 {
-%s;
+%{get_material_displacement_center};
 }
 
 float GetMaterialDisplacement(FPixelMaterialInputs PixelMaterialInputs)
@@ -3047,7 +3051,7 @@ float3 ClampWorldPositionOffset(FMaterialVertexParameters Parameters, float3 InO
 
 float3 GetMaterialWorldPositionOffsetRaw(FMaterialVertexParameters Parameters)
 {
-%s;
+%{get_material_world_position_offset_raw};
 }
 
 float3 GetMaterialWorldPositionOffset(FMaterialVertexParameters Parameters)
@@ -3062,7 +3066,7 @@ float3 GetMaterialWorldPositionOffset(FMaterialVertexParameters Parameters)
 
 float3 GetMaterialPreviousWorldPositionOffsetRaw(FMaterialVertexParameters Parameters)
 {
-%s;
+%{get_material_previous_world_position_offset_raw};
 }
 
 float3 GetMaterialPreviousWorldPositionOffset(FMaterialVertexParameters Parameters)
@@ -3090,12 +3094,12 @@ half4 GetMaterialSubsurfaceData(FPixelMaterialInputs PixelMaterialInputs)
 
 half GetMaterialCustomData0(in out FMaterialPixelParameters Parameters)
 {
-%s;
+%{get_material_custom_data0};
 }
 
 half GetMaterialCustomData1(in out FMaterialPixelParameters Parameters)
 {
-%s;
+%{get_material_custom_data1};
 }
 
 half GetMaterialAmbientOcclusionRaw(FPixelMaterialInputs PixelMaterialInputs)
@@ -3152,12 +3156,12 @@ FMaterialRefractionData GetMaterialRefraction(FPixelMaterialInputs PixelMaterial
 #if NUM_TEX_COORD_INTERPOLATORS
 void GetMaterialCustomizedUVs(FMaterialVertexParameters Parameters, inout float2 OutTexCoords[NUM_TEX_COORD_INTERPOLATORS])
 {
-%s
+%{get_material_customized_u_vs}
 }
 
 void GetCustomInterpolators(FMaterialVertexParameters Parameters, inout float2 OutTexCoords[NUM_TEX_COORD_INTERPOLATORS])
 {
-%s
+%{get_custom_interpolators}
 }
 #endif
 
@@ -3239,7 +3243,7 @@ float GetRoughnessFromNormalCurvature(FMaterialPixelParameters InMaterialParamet
 }
 
 // EvaluateMaterialAttributes
-%s
+%{evaluate_material_attributes}
 
 // Be very, very careful about changing CalcPixelMaterialInputs() or CalcPixelMaterialInputsAnalyticDerivatives(). They apply the same basic calculation,
 // but CalcPixelMaterialInputs() uses finite difference derivates from hardware whereas CalcPixelMaterialInputsAnalyticDerivatives() uses analytic
@@ -3247,9 +3251,9 @@ float GetRoughnessFromNormalCurvature(FMaterialPixelParameters InMaterialParamet
 void CalcPixelMaterialInputs(in out FMaterialPixelParameters Parameters, in out FPixelMaterialInputs PixelMaterialInputs)
 {
 	// Initial calculations (required for Normal)
-%s
+%{calc_pixel_material_inputs_initial_calculations}
 	// The Normal is a special case as it might have its own expressions and also be used to calculate other inputs, so perform the assignment here
-%s
+%{calc_pixel_material_inputs_normal}
 
 #if TEMPLATE_USES_STRATA
 	Parameters.StrataPixelFootprint = StrataGetPixelFootprint(Parameters.WorldPosition_CamRelative, GetRoughnessFromNormalCurvature(Parameters));
@@ -3293,7 +3297,7 @@ void CalcPixelMaterialInputs(in out FMaterialPixelParameters Parameters, in out
 #endif // !PARTICLE_SPRITE_FACTORY
 
 	// Now the rest of the inputs
-%s
+%{calc_pixel_material_inputs_other_inputs}
 
 #if MATERIAL_USES_ANISOTROPY
 	Parameters.WorldTangent = CalculateAnisotropyTangent(Parameters, PixelMaterialInputs);
@@ -3309,9 +3313,9 @@ void CalcPixelMaterialInputs(in out FMaterialPixelParameters Parameters, in out
 void CalcPixelMaterialInputsAnalyticDerivatives(in out FMaterialPixelParameters Parameters, in out FPixelMaterialInputs PixelMaterialInputs)
 {
 	// Initial calculations (required for Normal)
-%s
+%{calc_pixel_material_inputs_analytic_derivatives_initial}
 	// The Normal is a special case as it might have its own expressions and also be used to calculate other inputs, so perform the assignment here
-%s
+%{calc_pixel_material_inputs_analytic_derivatives_normal}
 
 #if TEMPLATE_USES_STRATA
 	Parameters.StrataPixelFootprint = StrataGetPixelFootprint(Parameters.WorldPosition_CamRelative, GetRoughnessFromNormalCurvature(Parameters));
@@ -3355,7 +3359,7 @@ void CalcPixelMaterialInputsAnalyticDerivatives(in out FMaterialPixelParameters
 #endif // !PARTICLE_SPRITE_FACTORY
 
 	// Now the rest of the inputs
-%s
+%{calc_pixel_material_inputs_analytic_derivatives_other_inputs}
 
 #if MATERIAL_USES_ANISOTROPY
 	Parameters.WorldTangent = CalculateAnisotropyTangent(Parameters, PixelMaterialInputs);
@@ -3368,7 +3372,7 @@ void CalcPixelMaterialInputsAnalyticDerivatives(in out FMaterialPixelParameters
 
 // Programmatically set the line number after all the material inputs which have a variable number of line endings
 // This allows shader error line numbers after this point to be the same regardless of which material is being compiled
-#line %s
+#line %{line_number}
 
 void ClipLODTransition(float2 SvPosition, float DitherFactor)
 {
diff --git a/Engine/Source/Runtime/Engine/Private/Materials/HLSLMaterialTranslator.cpp b/Engine/Source/Runtime/Engine/Private/Materials/HLSLMaterialTranslator.cpp
--- a/Engine/Source/Runtime/Engine/Private/Materials/HLSLMaterialTranslator.cpp
+++ b/Engine/Source/Runtime/Engine/Private/Materials/HLSLMaterialTranslator.cpp
@@ -36,6 +36,7 @@
 #include "Materials/MaterialExpressionViewProperty.h"
 #include "Materials/MaterialExpressionVolumetricAdvancedMaterialOutput.h"
 #include "Materials/MaterialExpressionWorldPosition.h"
+#include "Materials/MaterialSourceTemplate.h"
 #include "Materials/StrataMaterial.h"
 #include "ParameterCollection.h"
 #include "RenderUtils.h"
@@ -47,6 +48,12 @@
 #include <tuple>
 
 #if WITH_EDITORONLY_DATA
+#include "Serialization/ObjectWriter.h"
+#include "Serialization/ObjectReader.h"
+#include "Serialization/BufferArchive.h"
+#include "DerivedDataCache.h"
+#include "DerivedDataRequestOwner.h"
+#include "Serialization/ObjectAndNameAsStringProxyArchive.h"
 #include "Materials/MaterialExpressionStrata.h"
 #include "ShaderPlatformCachedIniValue.h"
 #endif
@@ -56,13 +63,19 @@
 namespace MaterialTranslatorCookStats
 {
 	static int32 MaterialTranslateCalls = 0;
+	static int32 MaterialTranslateCacheHits = 0;
+	static double MaterialTranslateAndSourceGenTimeSec = 0.0f;
 	static double MaterialTranslateTimeSec = 0.0f;
+	static double MaterialTranslateCachingTimeSec = 0.0f;
 
 	static FCookStatsManager::FAutoRegisterCallback RegisterCookStats([](FCookStatsManager::AddStatFuncRef AddStat)
 		{
 			AddStat(TEXT("Material"), FCookStatsManager::CreateKeyValueArray(
 				TEXT("MaterialTranslateCalls"), MaterialTranslateCalls,
-				TEXT("MaterialTranslateTimeSec"), MaterialTranslateTimeSec
+				TEXT("MaterialTranslateCacheHits"), MaterialTranslateCacheHits,
+				TEXT("MaterialTranslateAndSourceGenTimeSec"), MaterialTranslateAndSourceGenTimeSec,
+				TEXT("MaterialTranslateTimeSec"), MaterialTranslateTimeSec,
+				TEXT("MaterialTranslateCachingTimeSec"), MaterialTranslateCachingTimeSec
 			));
 		});
 }
@@ -134,6 +147,17 @@ static FAutoConsoleVariableRef CVarPedanticErrorChecksEnabled(
 	GPedanticErrorChecksEnabled,
 	TEXT("Enables material compilation pedantic error checking"));
 
+static TAutoConsoleVariable<bool> CVarJobMaterialTranslateDDCEnable(
+	TEXT("r.Material.TranslateDDCEnable"),
+	false,
+	TEXT("Whether to enable material translation DDC cachingk.\n"));
+
+static const int32 MaterialTranslationDDCVersion = 10;
+UE::DerivedData::FCacheBucket MaterialTranslationDDCBucket = UE::DerivedData::FCacheBucket(TEXT("MaterialTranslation"));
+UE::DerivedData::FValueId MaterialCompilationOutputId = UE::DerivedData::FValueId::FromName("FHLSLMaterialTranslator_MaterialCompilationOutput");
+UE::DerivedData::FValueId MaterialResultsOutputId = UE::DerivedData::FValueId::FromName("FHLSLMaterialTranslator_Results");
+UE::DerivedData::FValueId EnvironmentDefinesId = UE::DerivedData::FValueId::FromName("FHLSLMaterialTranslator_EnvironmentDefines");
+
 static inline bool IsAnalyticDerivEnabled()
 {
 	return GAnalyticDerivEnabled != 0;
@@ -333,7 +357,6 @@ FHLSLMaterialTranslator::FHLSLMaterialTranslator(FMaterial* InMaterial,
 ,	Platform(InPlatform)
 ,	QualityLevel(InQualityLevel)
 ,	FeatureLevel(InFeatureLevel)
-,	MaterialTemplateLineNumber(INDEX_NONE)
 ,	NextSymbolIndex(INDEX_NONE)
 ,	NextVertexInterpolatorIndex(0)
 ,	CurrentCustomVertexInterpolatorOffset(0)
@@ -832,12 +855,30 @@ bool FHLSLMaterialTranslator::Translate()
 #endif
 
 	COOK_STAT(MaterialTranslatorCookStats::MaterialTranslateCalls++);
-	COOK_STAT(FScopedDurationTimer DurationTimer(MaterialTranslatorCookStats::MaterialTranslateTimeSec));
+	COOK_STAT(FScopedDurationTimer DurationTimer1(MaterialTranslatorCookStats::MaterialTranslateTimeSec));
+	COOK_STAT(FScopedDurationTimer DurationTimer2(MaterialTranslatorCookStats::MaterialTranslateAndSourceGenTimeSec));
 
 	STAT(double HLSLTranslateTime = 0);
 	{
 		SCOPE_SECONDS_COUNTER(HLSLTranslateTime);
 
+		// First try to see if this translation run has already been done and the results are in the DDC cache.
+		bool bInCache = false;
+		FIoHash DDCKeyHash;
+		if (!Material->IsPreview())
+		{
+			COOK_STAT(FScopedDurationTimer CachingTimer(MaterialTranslatorCookStats::MaterialTranslateCachingTimeSec));
+
+			DDCKeyHash = ComputeMaterialTranslationDDCKeyHash();
+			bInCache = QueryDDCCachedTranslationResults(DDCKeyHash);
+			if (bInCache)
+			{
+				COOK_STAT(MaterialTranslatorCookStats::MaterialTranslateCacheHits++);
+				return true;
+			}
+		}
+
+		// No cache hit, continue translating the material
 		check(ScopeStack.Num() == 0);
 		bSuccess = true;
 
@@ -1600,41 +1641,6 @@ bool FHLSLMaterialTranslator::Translate()
 			}
 		}
 
-		// The code chunk corresponding to FullySimplifiedStrataFrontMaterialCodeChunk have already been written as part of MP_FrontMaterial.
-		// Here we get the FullySimplifiedStrataFrontMaterialTranslatedCodeChunks representing the variable storing the final fully simplified StrataData.
-		if(bStrataEnabled)
-		{
-			uint32 PropertyId = MP_FrontMaterial;
-		
-			if (PropertyId == MP_MaterialAttributes || PropertyId == MP_Normal || PropertyId == MP_CustomOutput)
-			{
-				continue;
-			}
-		
-			const EShaderFrequency PropertyShaderFrequency = FMaterialAttributeDefinitionMap::GetShaderFrequency((EMaterialProperty)PropertyId);
-		
-			int32 StartChunk = 0;
-			if (PropertyShaderFrequency == NormalShaderFrequency && SharedPixelProperties[PropertyId])
-			{
-				// When processing shared properties, do not generate the code before the Normal was generated as those are already handled
-				StartChunk = NormalCodeChunkEnd;
-			}
-		
-			// Reduce definition statements that don't contribute to the function's return value.
-			// @todo-lh: This should be expanded to a general reduction, but is currently only intended to fix an FXC internal compiler error reported in UE-117831
-			const bool bReduceAfterReturnValue = (PropertyId == MP_WorldPositionOffset || PropertyId == CompiledMP_PrevWorldPositionOffset || PropertyId == MP_Displacement);
-		
-			GetFixedParameterCode(
-				StartChunk,
-				SharedPropertyCodeChunks[PropertyShaderFrequency].Num(),
-				FullySimplifiedStrataFrontMaterialCodeChunk,								//Chunk[PropertyId],
-				SharedPropertyCodeChunks[PropertyShaderFrequency],
-				FullySimplifiedStrataFrontMaterialTranslatedCodeChunkDefinitions,
-				FullySimplifiedStrataFrontMaterialTranslatedCodeChunks,
-				Variation,
-				bReduceAfterReturnValue); 
-		}
-
 
 		for (uint32 PropertyId = MP_MAX; PropertyId < CompiledMP_MAX; ++PropertyId)
 		{
@@ -1908,354 +1914,544 @@ bool FHLSLMaterialTranslator::Translate()
 			}
 		}
 
-		LoadShaderSourceFileChecked(TEXT("/Engine/Private/MaterialTemplate.ush"), GetShaderPlatform(), MaterialTemplate);
-
-		// Find the string index of the '#line' statement in MaterialTemplate.usf
-		const int32 LineIndex = MaterialTemplate.Find(TEXT("#line"), ESearchCase::CaseSensitive);
-		check(LineIndex != INDEX_NONE);
-
-		// Count line endings before the '#line' statement
-		MaterialTemplateLineNumber = INDEX_NONE;
-		int32 StartPosition = LineIndex + 1;
-		do 
-		{
-			MaterialTemplateLineNumber++;
-			// Using \n instead of LINE_TERMINATOR as not all of the lines are terminated consistently
-			// Subtract one from the last found line ending index to make sure we skip over it
-			StartPosition = MaterialTemplate.Find(TEXT("\n"), ESearchCase::CaseSensitive, ESearchDir::FromEnd, StartPosition - 1);
-		} 
-		while (StartPosition != INDEX_NONE);
-		check(MaterialTemplateLineNumber != INDEX_NONE);
-		// At this point MaterialTemplateLineNumber is one less than the line number of the '#line' statement
-		// For some reason we have to add 2 more to the #line value to get correct error line numbers from D3DXCompileShader
-		MaterialTemplateLineNumber += 3;
-
 		MaterialCompilationOutput.UniformExpressionSet.SetParameterCollections(ParameterCollections);
 
-		// This will be created shortly after the Translate call in FMaterial::BeginCompileShaderMap()
-		//MaterialCompilationOutput.UniformExpressionSet.CreateBufferStruct();
-
 		// Store the number of unique VT samples
 		MaterialCompilationOutput.EstimatedNumVirtualTextureLookups = NumVtSamples;
-	}
-	ClearAllFunctionStacks();
-		
-	INC_FLOAT_STAT_BY(STAT_ShaderCompiling_HLSLTranslation,(float)HLSLTranslateTime);
-	return bSuccess;
-}
 
-void FHLSLMaterialTranslator::ValidateShadingModelsForFeatureLevel(const FMaterialShadingModelField& ShadingModels)
-{
-	if (FeatureLevel <= ERHIFeatureLevel::ES3_1)
-	{
-		const TArray<EMaterialShadingModel>& InvalidShadingModels = {};
-		for (EMaterialShadingModel InvalidShadingModel : InvalidShadingModels)
+		if (bSuccess)
 		{
-			if (ShadingModels.HasShadingModel(InvalidShadingModel))
-			{
-				FString FeatureLevelName;
-				GetFeatureLevelName(FeatureLevel, FeatureLevelName);
-
-				FString ShadingModelName;
-				const UEnum* EnumPtr = FindObject<UEnum>(nullptr, TEXT("/Script/Engine.EMaterialShadingModel"), true);
-				if (EnumPtr)
-				{
-					ShadingModelName = EnumPtr->GetNameStringByValue(InvalidShadingModel);
-				}
+			// Prepare the list of environment defines and material template string parameters
+			PrepareEnvironmentDefines();
+			PrepareMaterialSourceStringParameters();
 
-				Errorf(TEXT("ShadingModel %s not supported in feature level %s"), *ShadingModelName, *FeatureLevelName);
+			// If not a preview material, push the results to the DDC cache.
+			if (!Material->IsPreview())
+			{
+				PushResultsToDDCCache(DDCKeyHash);
 			}
 		}
 	}
+	ClearAllFunctionStacks();
+		
+	INC_FLOAT_STAT_BY(STAT_ShaderCompiling_HLSLTranslation,(float)HLSLTranslateTime);
+	
+	return bSuccess;
 }
 
-void FHLSLMaterialTranslator::GetMaterialEnvironment(EShaderPlatform InPlatform, FShaderCompilerEnvironment& OutEnvironment)
+FIoHash FHLSLMaterialTranslator::ComputeMaterialTranslationDDCKeyHash()
 {
-	const bool bStrataEnabled = Strata::IsStrataEnabled();
-	bool bMaterialRequestsDualSourceBlending = false;
+	FSHA1 Hasher = FMaterialSourceTemplate::Get().GetTemplateHash(GetShaderPlatform());
+	Hasher.Update((const uint8*)&MaterialTranslationDDCVersion, sizeof(MaterialTranslationDDCVersion));
 
-	if (bNeedsParticlePosition || Material->ShouldGenerateSphericalParticleNormals() || bUsesSphericalParticleOpacity)
+	for (const FStaticSwitchParameter& StaticSwitchParameter : StaticParameters.StaticSwitchParameters)
 	{
-		OutEnvironment.SetDefine(TEXT("NEEDS_PARTICLE_POSITION"), 1);
+		StaticSwitchParameter.UpdateHash(Hasher);
 	}
 
-	if (bNeedsParticleVelocity || Material->IsUsedWithNiagaraMeshParticles())
+	for (const FStaticComponentMaskParameter& StaticComponentMaskParameter : StaticParameters.EditorOnly.StaticComponentMaskParameters)
 	{
-		OutEnvironment.SetDefine(TEXT("NEEDS_PARTICLE_VELOCITY"), 1);
+		StaticComponentMaskParameter.UpdateHash(Hasher);
 	}
 
-	if (DynamicParticleParameterMask)
+	for (const FStaticTerrainLayerWeightParameter& StaticTerrainLayerWeightParameter : StaticParameters.EditorOnly.TerrainLayerWeightParameters)
 	{
-		OutEnvironment.SetDefine(TEXT("USE_DYNAMIC_PARAMETERS"), 1);
-		OutEnvironment.SetDefine(TEXT("DYNAMIC_PARAMETERS_MASK"), DynamicParticleParameterMask);
+		StaticTerrainLayerWeightParameter.UpdateHash(Hasher);
 	}
 
-	if (bNeedsParticleTime)
-	{
-		OutEnvironment.SetDefine(TEXT("NEEDS_PARTICLE_TIME"), 1);
-	}
+	FGuid BaseMaterialId = Material->GetMaterialInterface()->GetBaseMaterial()->StateId;
+	Hasher.Update((const uint8*)&BaseMaterialId, sizeof(BaseMaterialId));
 
-	if (bUsesParticleMotionBlur)
-	{
-		OutEnvironment.SetDefine(TEXT("USES_PARTICLE_MOTION_BLUR"), 1);
-	}
+	FString QualityLevelString;
+	GetMaterialQualityLevelName(QualityLevel, QualityLevelString);
+	Hasher.UpdateWithString(*QualityLevelString, QualityLevelString.Len());
+	Hasher.Update((const uint8*)&FeatureLevel, sizeof(FeatureLevel));
+	StrataCompilationConfig.UpdateHash(Hasher);
 
-	if (bNeedsParticleRandom)
-	{
-		OutEnvironment.SetDefine(TEXT("NEEDS_PARTICLE_RANDOM"), 1);
-	}
+	FSHAHash ShaHash;
+	Hasher.Final();
+	Hasher.GetHash(&ShaHash.Hash[0]);
+	return FIoHash::HashBuffer(&ShaHash, sizeof(ShaHash));
+}
 
-	if (bUsesSphericalParticleOpacity)
+bool FHLSLMaterialTranslator::QueryDDCCachedTranslationResults(const FIoHash& DDCKeyHash)
+{
+	if (!CVarJobMaterialTranslateDDCEnable.GetValueOnAnyThread())
 	{
-		OutEnvironment.SetDefine(TEXT("SPHERICAL_PARTICLE_OPACITY"), TEXT("1"));
+		return false;
 	}
 
-	if (bUsesParticleSubUVs)
-	{
-		OutEnvironment.SetDefine(TEXT("USE_PARTICLE_SUBUVS"), TEXT("1"));
-	}
+	// Try fetching the translation results from the DDC using the generated key hash
+	FSharedBuffer MaterialCompilationOutputBuffer;
+	FSharedBuffer TranslationResultsBuffer;
+	FSharedBuffer EnvironmentDefinesBuffer;
 
-	if (bUsesLightmapUVs)
-	{
-		OutEnvironment.SetDefine(TEXT("LIGHTMAP_UV_ACCESS"),TEXT("1"));
-	}
+	UE::DerivedData::FRequestOwner RequestOwner(UE::DerivedData::EPriority::Blocking);
+	UE::DerivedData::FCacheKey CacheKey{ MaterialTranslationDDCBucket, DDCKeyHash };
 
-	if (bUsesAOMaterialMask)
 	{
-		OutEnvironment.SetDefine(TEXT("USES_AO_MATERIAL_MASK"),TEXT("1"));
-	}
+		UE::DerivedData::FCacheGetRequest Request;
+		Request.Name = TEXT("FHLSLMaterialTranslator::QueryDDCCachedTranslationResults");
+		Request.Key = CacheKey;
 
-	if (bUsesSpeedTree)
-	{
-		OutEnvironment.SetDefine(TEXT("USES_SPEEDTREE"),TEXT("1"));
+		UE::DerivedData::GetCache().Get(
+			{ Request },
+			RequestOwner,
+			[&](UE::DerivedData::FCacheGetResponse&& Response)
+			{
+				if (Response.Status == UE::DerivedData::EStatus::Ok)
+				{
+					MaterialCompilationOutputBuffer = Response.Record.GetValue(MaterialCompilationOutputId).GetData().Decompress();
+					TranslationResultsBuffer = Response.Record.GetValue(MaterialResultsOutputId).GetData().Decompress();
+					EnvironmentDefinesBuffer = Response.Record.GetValue(EnvironmentDefinesId).GetData().Decompress();
+				}
+			}
+		);
 	}
 
-	if (bNeedsWorldPositionExcludingShaderOffsets)
+	RequestOwner.Wait();
+
+	// Load the results if we hit the cache.
+	if (!MaterialCompilationOutputBuffer.IsNull() && !TranslationResultsBuffer.IsNull() && !EnvironmentDefinesBuffer.IsNull())
 	{
-		OutEnvironment.SetDefine(TEXT("NEEDS_WORLD_POSITION_EXCLUDING_SHADER_OFFSETS"), TEXT("1"));
+		// Read the material compilation output
+		FShaderMapPointerTable PointerTable;
+		FPlatformTypeLayoutParameters LayoutParams;
+		LayoutParams.InitializeForPlatform(GetTargetPlatform());
+		FMemoryReaderView MaterialCompilationOutputReader{ TArrayView<uint8>{ (uint8*)MaterialCompilationOutputBuffer.GetData(), (int)MaterialCompilationOutputBuffer.GetSize() } };
+		FMemoryImageObject LoadedContent = FMemoryImageResult::LoadFromArchive(MaterialCompilationOutputReader, StaticGetTypeLayoutDesc<FMaterialCompilationOutput>(), &PointerTable, LayoutParams);
+		MaterialCompilationOutput = *(FMaterialCompilationOutput*)LoadedContent.Object;
+
+		// Read the environment defines
+		FMemoryReaderView EnvironmentDefinesBufferReader{ TArrayView<uint8>{ (uint8*)EnvironmentDefinesBuffer.GetData(), (int)EnvironmentDefinesBuffer.GetSize() } };
+		FObjectAndNameAsStringProxyArchive EnvironmentDefinesBufferReaderProxy{ EnvironmentDefinesBufferReader, true };
+		EnvironmentDefines.StaticStruct()->SerializeItem(EnvironmentDefinesBufferReaderProxy, &EnvironmentDefines, nullptr);
+
+		// Read the array of material string parameters
+		FMemoryReaderView ResultsMemoryReader{ TArrayView<uint8>{ (uint8*)TranslationResultsBuffer.GetData(), (int)TranslationResultsBuffer.GetSize() } };
+		ResultsMemoryReader << MaterialSourceTemplateParams;
+
+		return true;
 	}
 
-	if (bNeedsParticleSize)
+	return false;
+}
+
+void FHLSLMaterialTranslator::PrepareMaterialSourceStringParameters()
+{
+	MaterialSourceTemplateParams.Empty();
+
+	// Assign slots to vertex interpolators
+	FString VertexInterpolatorsOffsetsDefinition;
+	TBitArray<> FinalAllocatedCoords = GetVertexInterpolatorsOffsets(VertexInterpolatorsOffsetsDefinition);
+
+	const uint32 NumUserVertexTexCoords = GetNumUserVertexTexCoords();
+	const uint32 NumUserTexCoords = GetNumUserTexCoords();
+	const uint32 NumCustomVectors = FMath::DivideAndRoundUp((uint32)CurrentCustomVertexInterpolatorOffset, 2u);
+	const uint32 NumTexCoordVectors = FinalAllocatedCoords.FindLast(true) + 1;
+
+	MaterialSourceTemplateParams.Add({ TEXT("num_material_texcoords_vertex"), FString::Printf(TEXT("%u"), NumUserVertexTexCoords) });
+	MaterialSourceTemplateParams.Add({ TEXT("num_material_texcoords"), FString::Printf(TEXT("%u"), NumUserTexCoords) });
+	MaterialSourceTemplateParams.Add({ TEXT("num_custom_vertex_interpolators"), FString::Printf(TEXT("%u"), NumCustomVectors) });
+	MaterialSourceTemplateParams.Add({ TEXT("num_tex_coord_interpolators"), FString::Printf(TEXT("%u"), NumTexCoordVectors) });
+	MaterialSourceTemplateParams.Add({ TEXT("vertex_interpolators_offsets_definition"), VertexInterpolatorsOffsetsDefinition });
+
+	FString MaterialAttributesDeclaration;
+	FString MaterialAttributesUtilities;
+
+	const TArray<FGuid>& OrderedVisibleAttributes = FMaterialAttributeDefinitionMap::GetOrderedVisibleAttributeList();
+
+	// Reserve enough space vased on an estimate of the length of a single attribute string
+	MaterialAttributesDeclaration.Reserve(30 + MaterialAttributesDeclaration.Len() * 128); 
+	MaterialAttributesUtilities.Reserve(MaterialAttributesDeclaration.Len() * 256);
+
+	MaterialAttributesDeclaration.Append(TEXT("struct FMaterialAttributes\n{\n"));
+
+	const EMaterialShadingModel DefaultShadingModel = Material->GetShadingModels().GetFirstShadingModel();
+	for (const FGuid& AttributeID : OrderedVisibleAttributes)
 	{
-		OutEnvironment.SetDefine(TEXT("NEEDS_PARTICLE_SIZE"), TEXT("1"));
+		const FString PropertyName = FMaterialAttributeDefinitionMap::GetAttributeName(AttributeID);
+		const EMaterialValueType PropertyType = FMaterialAttributeDefinitionMap::GetValueType(AttributeID);
+		const TCHAR* HLSLType = nullptr;
+
+		switch (PropertyType)
+		{
+		case MCT_Float1: case MCT_Float: HLSLType = TEXT("float"); break;
+		case MCT_Float2: HLSLType = TEXT("float2"); break;
+		case MCT_Float3: HLSLType = TEXT("float3"); break;
+		case MCT_Float4: HLSLType = TEXT("float4"); break;
+		case MCT_UInt: case MCT_UInt1: case MCT_ShadingModel: HLSLType = TEXT("uint"); break;
+		case MCT_UInt2: HLSLType = TEXT("uint2"); break;
+		case MCT_UInt3: HLSLType = TEXT("uint3"); break;
+		case MCT_UInt4: HLSLType = TEXT("uint4"); break;
+		case MCT_Strata: HLSLType = TEXT("FStrataData"); break;
+		default: break;
+		}
+
+		if (HLSLType)
+		{
+			const FVector4f DefaultValue = FMaterialAttributeDefinitionMap::GetDefaultValue(AttributeID);
+
+			MaterialAttributesDeclaration.Appendf(TEXT("\t%s %s;") LINE_TERMINATOR, HLSLType, *PropertyName);
+
+			// Chainable method to set the attribute
+			MaterialAttributesUtilities.Appendf(TEXT("FMaterialAttributes FMaterialAttributes_Set%s(FMaterialAttributes InAttributes, %s InValue) { InAttributes.%s = InValue; return InAttributes; }") LINE_TERMINATOR,
+				*PropertyName, HLSLType, *PropertyName);
 	}
+}
+
+	MaterialAttributesDeclaration.Append(TEXT("};\n"));
+
+	MaterialSourceTemplateParams.Add({ TEXT("material_declarations"), MaterialAttributesDeclaration });
+	MaterialSourceTemplateParams.Add({ TEXT("material_attributes_utilities"), MaterialAttributesUtilities });
+
+	// Stores the shared shader results member declarations
+	FString PixelMembersDeclaration[CompiledPDV_MAX];
 
-	if (bNeedsParticleSpriteRotation)
+	FString NormalAssignment[CompiledPDV_MAX];
+
+	// Stores the code to initialize all inputs after MP_Normal
+	FString PixelMembersSetupAndAssignments[CompiledPDV_MAX];
+
+	for (int32 Index = 0; Index < CompiledPDV_MAX; Index++)
 	{
-		OutEnvironment.SetDefine(TEXT("NEEDS_PARTICLE_SPRITE_ROTATION"), TEXT("1"));
+		GetSharedInputsMaterialCode(PixelMembersDeclaration[Index], NormalAssignment[Index], PixelMembersSetupAndAssignments[Index], (ECompiledPartialDerivativeVariation)Index);
 	}
 
-	if (MaterialCompilationOutput.bNeedsSceneTextures)
+	// PixelMembersDeclaration should be the same for all variations, but might change in the future. There are cases where work is shared
+	// between the pixel and vertex shader, but with Nanite all work has to be moved into the pixel shader, which means we will want
+	// different inputs. But for now, we are keeping them the same.
+	MaterialSourceTemplateParams.Add({ TEXT("pixel_material_inputs"), PixelMembersDeclaration[CompiledPDV_FiniteDifferences] });
+
 	{
-		OutEnvironment.SetDefine(TEXT("NEEDS_SCENE_TEXTURES"), TEXT("1"));
+		FString DerivativeHelpers = DerivativeAutogen.GenerateUsedFunctions(*this);
+		FString DerivativeHelpersAndResources = DerivativeHelpers + ResourcesString;
+		MaterialSourceTemplateParams.Add({ TEXT("uniform_material_expressions"), DerivativeHelpersAndResources });
 	}
-	if (MaterialCompilationOutput.bUsesEyeAdaptation)
+
+	// Anything used bye the GenerationFunctionCode() like WorldPositionOffset shouldn't be using texures, right?
+	// Let those use the standard finite differences textures, since they should be the same. If we actually want
+	// those to handle texture reads properly, we'll have to make extra versions.
+	ECompiledPartialDerivativeVariation BaseDerivativeVariation = CompiledPDV_FiniteDifferences;
+
+	if (bCompileForComputeShader)
 	{
-		OutEnvironment.SetDefine(TEXT("USES_EYE_ADAPTATION"), TEXT("1"));
+		MaterialSourceTemplateParams.Add({ TEXT("get_material_emissive_for_cs"), GenerateFunctionCode(CompiledMP_EmissiveColorCS, BaseDerivativeVariation) });
 	}
-
-	if (MaterialCompilationOutput.bHasRuntimeVirtualTextureOutputNode)
+	else
 	{
-		OutEnvironment.SetDefine(TEXT("VIRTUAL_TEXTURE_OUTPUT"), 1);
+		MaterialSourceTemplateParams.Add({ TEXT("get_material_emissive_for_cs"), TEXT("return 0") });
 	}
 
-	OutEnvironment.SetDefine(TEXT("USES_PER_INSTANCE_CUSTOM_DATA"), MaterialCompilationOutput.bUsesPerInstanceCustomData && Material->IsUsedWithInstancedStaticMeshes());
-	OutEnvironment.SetDefine(TEXT("USES_PER_INSTANCE_RANDOM"), MaterialCompilationOutput.bUsesPerInstanceRandom && Material->IsUsedWithInstancedStaticMeshes());
-	OutEnvironment.SetDefine(TEXT("USES_PER_INSTANCE_FADE_AMOUNT"), bUsesPerInstanceFadeAmount && Material->IsUsedWithInstancedStaticMeshes());
-	OutEnvironment.SetDefine(TEXT("USES_VERTEX_INTERPOLATOR"), MaterialCompilationOutput.bUsesVertexInterpolator);
-
-	OutEnvironment.SetDefine(TEXT("MATERIAL_SKY_ATMOSPHERE"), bUsesSkyAtmosphere);
-	OutEnvironment.SetDefine(TEXT("INTERPOLATE_VERTEX_COLOR"), bUsesVertexColor);
-	OutEnvironment.SetDefine(TEXT("NEEDS_PARTICLE_COLOR"), bUsesParticleColor); 
-	OutEnvironment.SetDefine(TEXT("NEEDS_PARTICLE_LOCAL_TO_WORLD"), bUsesParticleLocalToWorld);
-	OutEnvironment.SetDefine(TEXT("NEEDS_PARTICLE_WORLD_TO_LOCAL"), bUsesParticleWorldToLocal);
-	OutEnvironment.SetDefine(TEXT("NEEDS_INSTANCE_LOCAL_TO_WORLD_PS"), bUsesInstanceLocalToWorldPS);
-	OutEnvironment.SetDefine(TEXT("NEEDS_INSTANCE_WORLD_TO_LOCAL_PS"), bUsesInstanceWorldToLocalPS);
-	OutEnvironment.SetDefine(TEXT("USES_TRANSFORM_VECTOR"), bUsesTransformVector);
-	OutEnvironment.SetDefine(TEXT("WANT_PIXEL_DEPTH_OFFSET"), bUsesPixelDepthOffset);
-	
-	// we want USES_WORLD_POSITION_OFFSET to be readable as a bool compile argument, hence the != 0 comparison 
-	// (bUsesWorldPositionOffset is actually a 1-bit uint32 bitfield member)
-	OutEnvironment.SetDefineAndCompileArgument(TEXT("USES_WORLD_POSITION_OFFSET"), (bUsesWorldPositionOffset != 0));
-	OutEnvironment.SetDefineAndCompileArgument(TEXT("USES_DISPLACEMENT"), (bUsesDisplacement != 0));
-
-	OutEnvironment.SetDefine(TEXT("USES_EMISSIVE_COLOR"), bUsesEmissiveColor);
-	// Distortion uses tangent space transform 
-	OutEnvironment.SetDefine(TEXT("USES_DISTORTION"), Material->IsDistorted());
-	OutEnvironment.SetDefine(TEXT("DISTORTION_ACCOUNT_FOR_COVERAGE"), bStrataEnabled && Material->GetRefractionCoverageMode() == RCM_CoverageAccountedFor ? 1 : 0);
+	MaterialSourceTemplateParams.Add({ TEXT("get_material_translucency_directional_lighting_intensity"), FString::Printf(TEXT("return %.5f"), Material->GetTranslucencyDirectionalLightingIntensity()) });
+	MaterialSourceTemplateParams.Add({ TEXT("get_material_translucent_shadow_density_scale"), FString::Printf(TEXT("return %.5f"), Material->GetTranslucentShadowDensityScale()) });
+	MaterialSourceTemplateParams.Add({ TEXT("get_material_translucent_self_shadow_density_scale"), FString::Printf(TEXT("return %.5f"), Material->GetTranslucentSelfShadowDensityScale()) });
+	MaterialSourceTemplateParams.Add({ TEXT("get_material_translucent_self_shadow_second_density_scale"), FString::Printf(TEXT("return %.5f"), Material->GetTranslucentSelfShadowSecondDensityScale()) });
+	MaterialSourceTemplateParams.Add({ TEXT("get_material_translucent_self_shadow_second_opacity"), FString::Printf(TEXT("return %.5f"), Material->GetTranslucentSelfShadowSecondOpacity()) });
+	MaterialSourceTemplateParams.Add({ TEXT("get_material_translucent_backscattering_exponent"), FString::Printf(TEXT("return %.5f"), Material->GetTranslucentBackscatteringExponent()) });
 
-	OutEnvironment.SetDefine(TEXT("MATERIAL_ENABLE_TRANSLUCENCY_FOGGING"), Material->ShouldApplyFogging());
-	OutEnvironment.SetDefine(TEXT("MATERIAL_ENABLE_TRANSLUCENCY_CLOUD_FOGGING"), Material->ShouldApplyCloudFogging());
-	OutEnvironment.SetDefine(TEXT("MATERIAL_IS_SKY"), Material->IsSky());
-	OutEnvironment.SetDefine(TEXT("MATERIAL_COMPUTE_FOG_PER_PIXEL"), Material->ComputeFogPerPixel());
-	OutEnvironment.SetDefine(TEXT("MATERIAL_FULLY_ROUGH"), bIsFullyRough || Material->IsFullyRough());
-	OutEnvironment.SetDefine(TEXT("MATERIAL_USES_ANISOTROPY"), bUsesAnisotropy && FDataDrivenShaderPlatformInfo::GetSupportsAnisotropicMaterials(InPlatform));
+	{
+		FLinearColor Extinction = Material->GetTranslucentMultipleScatteringExtinction();
 
-	OutEnvironment.SetDefine(TEXT("MATERIAL_DECAL_READ_MASK"), MaterialCompilationOutput.UsedDBufferTextures);
-	OutEnvironment.SetDefine(TEXT("MATERIAL_USES_DECAL_LOOKUP"), MaterialCompilationOutput.bUsesDBufferTextureLookup);
-	OutEnvironment.SetDefine(TEXT("MATERIAL_PATH_TRACING_BUFFER_READ"), MaterialCompilationOutput.UsedPathTracingBufferTextures);
+		MaterialSourceTemplateParams.Add({ TEXT("get_material_translucent_multiple_scattering_extinction"), FString::Printf(TEXT("return MaterialFloat3(%.5f, %.5f, %.5f)"), Extinction.R, Extinction.G, Extinction.B) });
+	}
 
-	// Count the number of VTStacks (each stack will allocate a feedback slot)
-	OutEnvironment.SetDefine(TEXT("NUM_VIRTUALTEXTURE_SAMPLES"), VTStacks.Num());
+	MaterialSourceTemplateParams.Add({ TEXT("get_material_opacity_mask_clip_value"), FString::Printf(TEXT("return %.5f"), Material->GetOpacityMaskClipValue()) });
 
-	// Check if any feedback slots are in use. We can simplify shader and remove EARLYZ optimizations if none are.
-	bool bGenerateFeedback = false;
-	for (int i = 0; i < VTStacks.Num() && !bGenerateFeedback; ++i)
 	{
-		bGenerateFeedback |= VTStacks[i].bGenerateFeedback;
+		const FDisplacementScaling DisplacementScaling = Material->GetDisplacementScaling();
+		MaterialSourceTemplateParams.Add({ TEXT("get_material_displacement_magnitude"), FString::Printf(TEXT("return %.5f"), FMath::Max(0.0f, DisplacementScaling.Magnitude)) });
+		MaterialSourceTemplateParams.Add({ TEXT("get_material_displacement_center"), FString::Printf(TEXT("return %.5f"), FMath::Clamp(DisplacementScaling.Center, 0.0f, 1.0f)) });
 	}
-	OutEnvironment.SetDefine(TEXT("MATERIAL_VIRTUALTEXTURE_FEEDBACK"), bGenerateFeedback);
 
-	// Setup defines to map each VT stack to either 1 or 2 page table textures, depending on how many layers it uses
-	for (int i = 0; i < VTStacks.Num(); ++i)
+	MaterialSourceTemplateParams.Add({ TEXT("get_material_world_position_offset_raw"), !bEnableExecutionFlow ? *GenerateFunctionCode(MP_WorldPositionOffset, BaseDerivativeVariation) : TEXT("return Parameters.MaterialAttributes.WorldPositionOffset") });
+	MaterialSourceTemplateParams.Add({ TEXT("get_material_previous_world_position_offset_raw"), !bEnableExecutionFlow ? *GenerateFunctionCode(CompiledMP_PrevWorldPositionOffset, BaseDerivativeVariation) : TEXT("return 0.0f") });
+	MaterialSourceTemplateParams.Add({ TEXT("get_material_custom_data0"), !bEnableExecutionFlow ? *GenerateFunctionCode(MP_CustomData0, BaseDerivativeVariation) : TEXT("return 0.0f") });
+	MaterialSourceTemplateParams.Add({ TEXT("get_material_custom_data1"), !bEnableExecutionFlow ? *GenerateFunctionCode(MP_CustomData1, BaseDerivativeVariation) : TEXT("return 0.0f") });
+
+	// Print custom texture coordinate assignments, should be fine with regular derivatives
+	FString CustomUVAssignments;
+
+	int32 LastProperty = -1;
+	for (uint32 CustomUVIndex = 0; CustomUVIndex < NumUserTexCoords; CustomUVIndex++)
 	{
-		const FMaterialVirtualTextureStack& Stack = MaterialCompilationOutput.UniformExpressionSet.VTStacks[i];
-		FString PageTableValue = FString::Printf(TEXT("Material.VirtualTexturePageTable0_%d"), i);
-		if (Stack.GetNumLayers() > 4u)
-		{
-			PageTableValue += FString::Printf(TEXT(", Material.VirtualTexturePageTable1_%d"), i);
+		if (bEnableExecutionFlow)
+{
+			const FString AttributeName = FMaterialAttributeDefinitionMap::GetAttributeName((EMaterialProperty)(MP_CustomizedUVs0 + CustomUVIndex));
+			CustomUVAssignments += FString::Printf(TEXT("\tOutTexCoords[%u] = Parameters.MaterialAttributes.%s;") LINE_TERMINATOR, CustomUVIndex, *AttributeName);
 		}
-		if (VTStacks[i].bAdaptive)
+		else
 		{
-			PageTableValue += FString::Printf(TEXT(", Material.VirtualTexturePageTableIndirection_%d"), i);
+			if (CustomUVIndex == 0)
+			{
+				CustomUVAssignments += DerivativeVariations[BaseDerivativeVariation].TranslatedCodeChunkDefinitions[MP_CustomizedUVs0 + CustomUVIndex];
+			}
+
+			if (DerivativeVariations[BaseDerivativeVariation].TranslatedCodeChunkDefinitions[MP_CustomizedUVs0 + CustomUVIndex].Len() > 0)
+			{
+				LastProperty = MP_CustomizedUVs0 + CustomUVIndex;
+			}
+			CustomUVAssignments += FString::Printf(TEXT("\tOutTexCoords[%u] = %s;") LINE_TERMINATOR, CustomUVIndex, *DerivativeVariations[BaseDerivativeVariation].TranslatedCodeChunks[MP_CustomizedUVs0 + CustomUVIndex]);
 		}
-		OutEnvironment.SetDefine(*FString::Printf(TEXT("VIRTUALTEXTURE_PAGETABLE_%d"), i), *PageTableValue);
 	}
 
-	for (int32 CollectionIndex = 0; CollectionIndex < ParameterCollections.Num(); CollectionIndex++)
-	{
-		// Add uniform buffer declarations for any parameter collections referenced
-		const FString CollectionName = FString::Printf(TEXT("MaterialCollection%u"), CollectionIndex);
-		// This can potentially become an issue for MaterialCollection Uniform Buffers if they ever get non-numeric resources (eg Textures), as
-		// OutEnvironment.ResourceTableMap has a map by name, and the N ParameterCollection Uniform Buffers ALL are names "MaterialCollection"
-		// (and the hlsl cbuffers are named MaterialCollection0, etc, so the names don't match the layout)
-		FShaderUniformBufferParameter::ModifyCompilationEnvironment(*CollectionName, ParameterCollections[CollectionIndex]->GetUniformBufferStruct(), InPlatform, OutEnvironment);
-	}
-	OutEnvironment.SetDefine(TEXT("IS_MATERIAL_SHADER"), TEXT("1"));
+	MaterialSourceTemplateParams.Add({ TEXT("get_material_customized_u_vs"), CustomUVAssignments });
 
-	// Set all the shading models for this material here 
-	FMaterialShadingModelField ShadingModels = Material->GetShadingModels();
+	// Print custom vertex shader interpolator assignments
+	FString CustomInterpolatorAssignments;
 
-	// If the material gets its shading model from the material expressions, then we use the result from the compilation (assuming it's valid).
-	// This result will potentially be tighter than what GetShadingModels() returns, because it only picks up the shading models from the expressions that get compiled for a specific feature level and quality level
-	// For example, the material might have shading models behind static switches. GetShadingModels() will return both the true and the false paths from that switch, whereas the shading model field from the compilation will only contain the actual shading model selected 
-	if (Material->IsShadingModelFromMaterialExpression() && ShadingModelsFromCompilation.IsValid())
+	for (UMaterialExpressionVertexInterpolator* Interpolator : CustomVertexInterpolators)
 	{
-		// Shading models fetched from the compilation of the expression graph
-		ShadingModels = ShadingModelsFromCompilation;
-	}
+		if (Interpolator->InterpolatorOffset != INDEX_NONE)
+		{
+			check(Interpolator->InterpolatorIndex != INDEX_NONE);
+			check(Interpolator->InterpolatedType & MCT_Float);
 
-	ensure(ShadingModels.IsValid());
+			const EMaterialValueType Type = Interpolator->InterpolatedType == MCT_Float ? MCT_Float1 : Interpolator->InterpolatedType;
+			const TCHAR* Swizzle[2] = { TEXT("x"), TEXT("y") };
+			const int32 Offset = Interpolator->InterpolatorOffset;
+			const int32 Index = Interpolator->InterpolatorIndex;
 
-	if (ShadingModels.IsLit())
-	{	
-		int NumSetMaterials = 0;
-		if (ShadingModels.HasShadingModel(MSM_DefaultLit))
-		{
-			OutEnvironment.SetDefine(TEXT("MATERIAL_SHADINGMODEL_DEFAULT_LIT"), TEXT("1"));
-			NumSetMaterials++;
-		}
-		if (ShadingModels.HasShadingModel(MSM_Subsurface))
-		{
-			OutEnvironment.SetDefine(TEXT("MATERIAL_SHADINGMODEL_SUBSURFACE"), TEXT("1"));
-			NumSetMaterials++;
+			// Note: We reference the UV define directly to avoid having to pre-accumulate UV counts before property translation
+			CustomInterpolatorAssignments += FString::Printf(TEXT("\tOutTexCoords[VERTEX_INTERPOLATOR_%i_TEXCOORDS_X].%s = VertexInterpolator%i(Parameters).x;") LINE_TERMINATOR, Index, Swizzle[Offset % 2], Index);
+
+			if (Type >= MCT_Float2)
+			{
+				CustomInterpolatorAssignments += FString::Printf(TEXT("\tOutTexCoords[VERTEX_INTERPOLATOR_%i_TEXCOORDS_Y].%s = VertexInterpolator%i(Parameters).y;") LINE_TERMINATOR, Index, Swizzle[(Offset + 1) % 2], Index);
+
+				if (Type >= MCT_Float3)
+				{
+					CustomInterpolatorAssignments += FString::Printf(TEXT("\tOutTexCoords[VERTEX_INTERPOLATOR_%i_TEXCOORDS_Z].%s = VertexInterpolator%i(Parameters).z;") LINE_TERMINATOR, Index, Swizzle[(Offset + 2) % 2], Index);
+
+					if (Type == MCT_Float4)
+					{
+						CustomInterpolatorAssignments += FString::Printf(TEXT("\tOutTexCoords[VERTEX_INTERPOLATOR_%i_TEXCOORDS_W].%s = VertexInterpolator%i(Parameters).w;") LINE_TERMINATOR, Index, Swizzle[(Offset + 3) % 2], Index);
+					}
+				}
+			}
 		}
-		if (ShadingModels.HasShadingModel(MSM_PreintegratedSkin))
+	}
+
+	MaterialSourceTemplateParams.Add({ TEXT("get_custom_interpolators"), CustomInterpolatorAssignments });
+
+	if (bEnableExecutionFlow)
+	{
+		FString EvaluateVertexCode;
+
+		// Set default texcoords in the VS
+		for (uint32 TexCoordIndex = 0; TexCoordIndex < NumUserVertexTexCoords; ++TexCoordIndex)
 		{
-			OutEnvironment.SetDefine(TEXT("MATERIAL_SHADINGMODEL_PREINTEGRATED_SKIN"), TEXT("1"));
-			NumSetMaterials++;
+			const FString AttributeName = FMaterialAttributeDefinitionMap::GetAttributeName((EMaterialProperty)(MP_CustomizedUVs0 + TexCoordIndex));
+
+			EvaluateVertexCode += FString::Printf(TEXT("\tDefaultMaterialAttributes.%s = Parameters.TexCoords[%d];") LINE_TERMINATOR, *AttributeName, TexCoordIndex);
 		}
-		if (ShadingModels.HasShadingModel(MSM_SubsurfaceProfile))
+
+		EvaluateVertexCode += TranslatedAttributesCodeChunks[SF_Vertex];
+
+		MaterialSourceTemplateParams.Add({ TEXT("evaluate_material_attributes"), TranslatedAttributesCodeChunks[SF_Pixel] });
+
+		FString EvaluateMaterialAttributesCode = TEXT("    FMaterialAttributes MaterialAttributes = EvaluatePixelMaterialAttributes(Parameters);" LINE_TERMINATOR);
+
+		for (int32 PropertyIndex = 0; PropertyIndex < MP_MAX; ++PropertyIndex)
 		{
-			OutEnvironment.SetDefine(TEXT("MATERIAL_SHADINGMODEL_SUBSURFACE_PROFILE"), TEXT("1"));
-			NumSetMaterials++;
+			// Skip non-shared properties
+			if (!SharedPixelProperties[PropertyIndex])
+			{
+				continue;
+			}
 
-			if (bUsesCurvature)
+			const EMaterialProperty Property = (EMaterialProperty)PropertyIndex;
+			check(FMaterialAttributeDefinitionMap::GetShaderFrequency(Property) == SF_Pixel);
+			// Special case MP_SubsurfaceColor as the actual property is a combination of the color and the profile but we don't want to expose the profile
+			const FString PropertyName = FMaterialAttributeDefinitionMap::GetAttributeName(Property);
+
+			if (PropertyIndex == MP_SubsurfaceColor)
 			{
-				OutEnvironment.SetDefine(TEXT("MATERIAL_SUBSURFACE_PROFILE_USE_CURVATURE"), TEXT("1"));
+				// TODO - properly handle subsurface profile
+				EvaluateMaterialAttributesCode += FString::Printf("    PixelMaterialInputs.Subsurface = float4(MaterialAttributes.%s, 0.0f);" LINE_TERMINATOR, *PropertyName);
+			}
+			else
+			{
+				EvaluateMaterialAttributesCode += FString::Printf("    PixelMaterialInputs.%s = MaterialAttributes.%s;" LINE_TERMINATOR, *PropertyName, *PropertyName);
 			}
 		}
-		if (ShadingModels.HasShadingModel(MSM_ClearCoat))
+
+		MaterialSourceTemplateParams.Add({ TEXT("calc_pixel_material_inputs_initial_calculations"), EvaluateMaterialAttributesCode });
+		MaterialSourceTemplateParams.Add({ TEXT("calc_pixel_material_inputs_normal"), TEXT("") });
+		MaterialSourceTemplateParams.Add({ TEXT("calc_pixel_material_inputs_other_inputs"), TEXT("") });
+
+		MaterialSourceTemplateParams.Add({ TEXT("calc_pixel_material_inputs_analytic_derivatives_initial"), EvaluateMaterialAttributesCode });
+		MaterialSourceTemplateParams.Add({ TEXT("calc_pixel_material_inputs_analytic_derivatives_normal"), TEXT("") });
+		MaterialSourceTemplateParams.Add({ TEXT("calc_pixel_material_inputs_analytic_derivatives_other_inputs"), TEXT("") });
+	}
+	else
+	{
+		// skip material attributes code
+		MaterialSourceTemplateParams.Add({ TEXT("evaluate_material_attributes"), TEXT("") });
+
 		{
-			OutEnvironment.SetDefine(TEXT("MATERIAL_SHADINGMODEL_CLEAR_COAT"), TEXT("1"));
-			NumSetMaterials++;
+			// Initializers required for Normal
+			MaterialSourceTemplateParams.Add({ TEXT("calc_pixel_material_inputs_initial_calculations"), DerivativeVariations[CompiledPDV_FiniteDifferences].TranslatedCodeChunkDefinitions[MP_Normal] });
+			MaterialSourceTemplateParams.Add({ TEXT("calc_pixel_material_inputs_normal"), NormalAssignment[CompiledPDV_FiniteDifferences] });
+			// Finally the rest of common code followed by assignment into each input
+			MaterialSourceTemplateParams.Add({ TEXT("calc_pixel_material_inputs_other_inputs"), PixelMembersSetupAndAssignments[CompiledPDV_FiniteDifferences] });
 		}
-		if (ShadingModels.HasShadingModel(MSM_TwoSidedFoliage))
 		{
-			OutEnvironment.SetDefine(TEXT("MATERIAL_SHADINGMODEL_TWOSIDED_FOLIAGE"), TEXT("1"));
-			NumSetMaterials++;
+			// Initializers required for Normal
+			MaterialSourceTemplateParams.Add({ TEXT("calc_pixel_material_inputs_analytic_derivatives_initial"), DerivativeVariations[CompiledPDV_Analytic].TranslatedCodeChunkDefinitions[MP_Normal] });
+			MaterialSourceTemplateParams.Add({ TEXT("calc_pixel_material_inputs_analytic_derivatives_normal"), NormalAssignment[CompiledPDV_Analytic] });
+			// Finally the rest of common code followed by assignment into each input
+			MaterialSourceTemplateParams.Add({ TEXT("calc_pixel_material_inputs_analytic_derivatives_other_inputs"), PixelMembersSetupAndAssignments[CompiledPDV_Analytic] });
 		}
-		if (ShadingModels.HasShadingModel(MSM_Hair))
+	}
+}
+
+void FHLSLMaterialTranslator::PrepareEnvironmentDefines()
+{
+	EShaderPlatform InPlatform = GetShaderPlatform();
+
+	EnvironmentDefines = {};
+	const bool bStrataEnabled = Strata::IsStrataEnabled();
+	bool bMaterialRequestsDualSourceBlending = false;
+
+	EnvironmentDefines.bNeedsParticlePosition = bNeedsParticlePosition || Material->ShouldGenerateSphericalParticleNormals() || bUsesSphericalParticleOpacity;
+	EnvironmentDefines.bNeedsParticleVelocity = bNeedsParticleVelocity || Material->IsUsedWithNiagaraMeshParticles();
+	EnvironmentDefines.bUseDynamicParameters = bool(DynamicParticleParameterMask);
+	EnvironmentDefines.DynamicParametersMask = DynamicParticleParameterMask;
+	EnvironmentDefines.bNeedsParticleTime = bNeedsParticleTime;
+	EnvironmentDefines.bUsesParticleMotionBlur = bUsesParticleMotionBlur;
+	EnvironmentDefines.bNeedsParticleRandom = bNeedsParticleRandom;
+	EnvironmentDefines.bSphericalParticleOpacity = bUsesSphericalParticleOpacity;
+	EnvironmentDefines.bUseParticleSubUVs = bUsesParticleSubUVs;
+	EnvironmentDefines.bLightmapUVAccess = bUsesLightmapUVs;
+	EnvironmentDefines.bUsesAOMaterialMask = bUsesAOMaterialMask;
+	EnvironmentDefines.bUsesSpeedtree = bUsesSpeedTree;
+	EnvironmentDefines.bNeedsWorldPositionExcludingShaderOffsets = bNeedsWorldPositionExcludingShaderOffsets;
+	EnvironmentDefines.bNeedsParticleSize = bNeedsParticleSize;
+	EnvironmentDefines.bNeedsParticleSpriteRotation = bNeedsParticleSpriteRotation;
+	EnvironmentDefines.bNeedsSceneTextures = MaterialCompilationOutput.bNeedsSceneTextures;
+	EnvironmentDefines.bUsesEyeAdaptation = MaterialCompilationOutput.bUsesEyeAdaptation;
+	EnvironmentDefines.bVirtualTextureOutput = MaterialCompilationOutput.bHasRuntimeVirtualTextureOutputNode;
+	EnvironmentDefines.bUsesPerInstanceCustomData = MaterialCompilationOutput.bUsesPerInstanceCustomData && Material->IsUsedWithInstancedStaticMeshes();
+	EnvironmentDefines.bUsesPerInstanceRandom = MaterialCompilationOutput.bUsesPerInstanceRandom && Material->IsUsedWithInstancedStaticMeshes();
+	EnvironmentDefines.bUsesPerInstanceFadeAmount = bUsesPerInstanceFadeAmount&& Material->IsUsedWithInstancedStaticMeshes();
+	EnvironmentDefines.bUsesVertexInterpolator = MaterialCompilationOutput.bUsesVertexInterpolator;
+	EnvironmentDefines.bUsesSkyAtmosphere = bUsesSkyAtmosphere;
+	EnvironmentDefines.bUsesVertexColor = bUsesVertexColor;
+	EnvironmentDefines.bUsesParticleColor = bUsesParticleColor;
+	EnvironmentDefines.bUsesParticleLocalToWorld = bUsesParticleLocalToWorld;
+	EnvironmentDefines.bUsesParticleWorldToLocal = bUsesParticleWorldToLocal;
+	EnvironmentDefines.bUsesInstanceLocalToWorldPS = bUsesInstanceLocalToWorldPS;
+	EnvironmentDefines.bUsesInstanceWorldToLocalPS = bUsesInstanceWorldToLocalPS;
+	EnvironmentDefines.bUsesTransformVector = bUsesTransformVector;
+	EnvironmentDefines.bUsesPixelDepthOffset = bUsesPixelDepthOffset;
+	// we want USES_WORLD_POSITION_OFFSET to be readable as a bool compile argument, hence the != 0 comparison 
+	// (bUsesWorldPositionOffset is actually a 1-bit uint32 bitfield member)
+	EnvironmentDefines.bUsesWorldPositionOffset = (bUsesWorldPositionOffset != 0);
+	EnvironmentDefines.bUsesDisplacement = (bUsesDisplacement != 0);
+	EnvironmentDefines.bUsesEmissiveColor = bUsesEmissiveColor;
+	// Distortion uses tangent space transform 
+	EnvironmentDefines.bUsesDistortion = Material->IsDistorted();
+	EnvironmentDefines.bDistortionAccountForCoverage = bStrataEnabled && Material->GetRefractionCoverageMode() == RCM_CoverageAccountedFor ? 1 : 0;
+	EnvironmentDefines.bMaterialEnableTranslucencyFogging = Material->ShouldApplyFogging();
+	EnvironmentDefines.bMaterialEnableTranslucencyCloudFogging = Material->ShouldApplyCloudFogging();
+	EnvironmentDefines.bMaterialIsSky = Material->IsSky();
+	EnvironmentDefines.bMaterialComputeFogPerPixel = Material->ComputeFogPerPixel();
+	EnvironmentDefines.bMaterialFullyRough = bIsFullyRough || Material->IsFullyRough();
+	EnvironmentDefines.bMaterialUsesAnisotropy = bUsesAnisotropy && FDataDrivenShaderPlatformInfo::GetSupportsAnisotropicMaterials(InPlatform);
+	EnvironmentDefines.MaterialDecalReadMask = MaterialCompilationOutput.UsedDBufferTextures;
+	EnvironmentDefines.bMaterialUsesDecalLookup = MaterialCompilationOutput.bUsesDBufferTextureLookup;
+	EnvironmentDefines.MaterialPathTracingBufferRead = MaterialCompilationOutput.UsedPathTracingBufferTextures;
+	// Count the number of VTStacks (each stack will allocate a feedback slot)
+	EnvironmentDefines.NumVirtualTextureSamples = VTStacks.Num();
+
+	// Check if any feedback slots are in use. We can simplify shader and remove EARLYZ optimizations if none are.
+	bool bGenerateFeedback = false;
+	for (int i = 0; i < VTStacks.Num() && !bGenerateFeedback; ++i)
+	{
+		bGenerateFeedback |= VTStacks[i].bGenerateFeedback;
+	}
+	EnvironmentDefines.bMaterialVirtualTextureFeedback = bGenerateFeedback;
+
+	// Setup defines to map each VT stack to either 1 or 2 page table textures, depending on how many layers it uses
+	EnvironmentDefines.VirtualPageTypes.SetNumUninitialized(VTStacks.Num());
+	for (int i = 0; i < VTStacks.Num(); ++i)
+	{
+		const FMaterialVirtualTextureStack& Stack = MaterialCompilationOutput.UniformExpressionSet.VTStacks[i];
+		EnvironmentDefines.VirtualPageTypes[i] = FHLSLMaterialTranslatorEnvironmentDefines::TABLE0;
+		if (Stack.GetNumLayers() > 4u)
 		{
-			OutEnvironment.SetDefine(TEXT("MATERIAL_SHADINGMODEL_HAIR"), TEXT("1"));
-			NumSetMaterials++;
+			EnvironmentDefines.VirtualPageTypes[i] |= FHLSLMaterialTranslatorEnvironmentDefines::TABLE1;
 		}
-		if (ShadingModels.HasShadingModel(MSM_Cloth))
+		if (VTStacks[i].bAdaptive)
 		{
-			OutEnvironment.SetDefine(TEXT("MATERIAL_SHADINGMODEL_CLOTH"), TEXT("1"));
-			NumSetMaterials++;
+			EnvironmentDefines.VirtualPageTypes[i] |= FHLSLMaterialTranslatorEnvironmentDefines::TABLE_INDIRECTION;
 		}
-		if (ShadingModels.HasShadingModel(MSM_Eye))
-		{
-			OutEnvironment.SetDefine(TEXT("MATERIAL_SHADINGMODEL_EYE"), TEXT("1"));
-			NumSetMaterials++;
+	}
 
-			if (bUsesCurvature)
-			{
-				OutEnvironment.SetDefine(TEXT("MATERIAL_SHADINGMODEL_EYE_USE_CURVATURE"), TEXT("1"));
-			}
-		}
-		if (ShadingModels.HasShadingModel(MSM_SingleLayerWater))
-		{
-			OutEnvironment.SetDefine(TEXT("MATERIAL_SHADINGMODEL_SINGLELAYERWATER"), TEXT("1"));
-			NumSetMaterials++;
-		}
-		if (ShadingModels.HasShadingModel(MSM_ThinTranslucent))
-		{
-			OutEnvironment.SetDefine(TEXT("MATERIAL_SHADINGMODEL_THIN_TRANSLUCENT"), TEXT("1"));
-			NumSetMaterials++;
+	EnvironmentDefines.ParameterCollections.Reserve(ParameterCollections.Num());
+	for (int32 CollectionIndex = 0; CollectionIndex < ParameterCollections.Num(); CollectionIndex++)
+	{
+		EnvironmentDefines.ParameterCollections.Add(ParameterCollections[CollectionIndex]);
+	}
 
-			bMaterialRequestsDualSourceBlending = true;
-		}
 
-		if (ShadingModels.HasShadingModel(MSM_SingleLayerWater) && FDataDrivenShaderPlatformInfo::GetRequiresDisableForwardLocalLights(Platform))
-		{
-			OutEnvironment.SetDefine(TEXT("DISABLE_FORWARD_LOCAL_LIGHTS"), TEXT("1"));
-		}
+	// Set all the shading models for this material here 
+	FMaterialShadingModelField ShadingModels = Material->GetShadingModels();
+
+	// If the material gets its shading model from the material expressions, then we use the result from the compilation (assuming it's valid).
+	// This result will potentially be tighter than what GetShadingModels() returns, because it only picks up the shading models from the expressions that get compiled for a specific feature level and quality level
+	// For example, the material might have shading models behind static switches. GetShadingModels() will return both the true and the false paths from that switch, whereas the shading model field from the compilation will only contain the actual shading model selected 
+	if (Material->IsShadingModelFromMaterialExpression() && ShadingModelsFromCompilation.IsValid())
+	{
+		// Shading models fetched from the compilation of the expression graph
+		ShadingModels = ShadingModelsFromCompilation;
+	}
+
+	ensure(ShadingModels.IsValid());
 
-		// This is to have platforms use the simple single layer water shading similar to mobile: no dynamic lights, only sun and sky, no distortion, no colored transmittance on background, no custom depth read.
-		const bool bSingleLayerWaterUsesSimpleShading = FDataDrivenShaderPlatformInfo::GetWaterUsesSimpleForwardShading(InPlatform) && IsForwardShadingEnabled(InPlatform);
-		if (ShadingModels.HasShadingModel(MSM_SingleLayerWater) && bSingleLayerWaterUsesSimpleShading)
+	// This is to have platforms use the simple single layer water shading similar to mobile: no dynamic lights, only sun and sky, no distortion, no colored transmittance on background, no custom depth read.
+	const bool bSingleLayerWaterUsesSimpleShading = FDataDrivenShaderPlatformInfo::GetWaterUsesSimpleForwardShading(InPlatform) && IsForwardShadingEnabled(InPlatform);
+	// Value must match SINGLE_LAYER_WATER_SHADING_QUALITY_MOBILE_WITH_DEPTH_TEXTURE in SingleLayerWaterCommon.ush!
+	EnvironmentDefines.bSingleLayerWaterShadingQuality = (ShadingModels.HasShadingModel(MSM_SingleLayerWater) && bSingleLayerWaterUsesSimpleShading);
+
+	if (ShadingModels.IsLit())
+	{	
+		EnvironmentDefines.bShadingModelsIsLit = true;
+
+		int NumSetMaterials = 0;
+
+		for (int i = 0; i < MSM_NUM; ++i)
 		{
-			// Value must match SINGLE_LAYER_WATER_SHADING_QUALITY_MOBILE_WITH_DEPTH_TEXTURE in SingleLayerWaterCommon.ush!
-			OutEnvironment.SetDefine(TEXT("SINGLE_LAYER_WATER_SHADING_QUALITY"), TEXT("1"));
+			if (ShadingModels.HasShadingModel((EMaterialShadingModel)i))
+			{
+				EnvironmentDefines.MaterialShadingModelEnabled |= 1 << i;
+				++NumSetMaterials;
+			}
 		}
 
+		EnvironmentDefines.bMaterialSubsurfaceProfileUseCurvature = ShadingModels.HasShadingModel(MSM_SubsurfaceProfile) && bUsesCurvature;
+		EnvironmentDefines.bMaterialShadingModelEyeUseCurvature = ShadingModels.HasShadingModel(MSM_Eye) && bUsesCurvature;
+		EnvironmentDefines.bDisableForwardLocalLights = ShadingModels.HasShadingModel(MSM_SingleLayerWater) && FDataDrivenShaderPlatformInfo::GetRequiresDisableForwardLocalLights(Platform);
+
 		const bool bIsWaterDistanceFieldShadowEnabled = IsWaterDistanceFieldShadowEnabled(InPlatform);
 		const bool bIsWaterVSMFilteringEnabled = IsWaterVirtualShadowMapFilteringEnabled(InPlatform);
-		if (ShadingModels.HasShadingModel(MSM_SingleLayerWater) && (bIsWaterDistanceFieldShadowEnabled || bIsWaterVSMFilteringEnabled))
-		{
-			OutEnvironment.SetDefine(TEXT("SINGLE_LAYER_WATER_SEPARATED_MAIN_LIGHT"), TEXT("1"));
-		}
-
-		if (NumSetMaterials == 1)
-		{
-			OutEnvironment.SetDefine(TEXT("MATERIAL_SINGLE_SHADINGMODEL"), TEXT("1"));
-		}
+		EnvironmentDefines.bSingleLayerWaterSeparatedMainLight = (ShadingModels.HasShadingModel(MSM_SingleLayerWater) && (bIsWaterDistanceFieldShadowEnabled || bIsWaterVSMFilteringEnabled));
+		EnvironmentDefines.bMaterialSingleShadingModel = (NumSetMaterials == 1);
 
 		ensure(NumSetMaterials != 0);
 		if (NumSetMaterials == 0)
 		{
 			// Should not really end up here
 			UE_LOG(LogMaterial, Warning, TEXT("Unknown material shading model(s). Setting to MSM_DefaultLit"));
-			OutEnvironment.SetDefine(TEXT("MATERIAL_SHADINGMODEL_DEFAULT_LIT"),TEXT("1"));
+			EnvironmentDefines.MaterialShadingModelEnabled |= 1 << MSM_DefaultLit;
 		}
 	}
 	else
 	{
 		// Unlit shading model can only exist by itself
-		OutEnvironment.SetDefine(TEXT("MATERIAL_SINGLE_SHADINGMODEL"), TEXT("1"));
-		OutEnvironment.SetDefine(TEXT("MATERIAL_SHADINGMODEL_UNLIT"), TEXT("1"));
+		EnvironmentDefines.bMaterialSingleShadingModel = true;
+		EnvironmentDefines.MaterialShadingModelEnabled |= 1 << MSM_Unlit;
 	}
 
 	if (Material->GetMaterialDomain() == MD_Volume)
@@ -2269,32 +2465,18 @@ void FHLSLMaterialTranslator::GetMaterialEnvironment(EShaderPlatform InPlatform,
 				UE_LOG(LogMaterial, Fatal, TEXT("Only a single UMaterialExpressionVolumetricAdvancedMaterialOutput node is supported."));
 			}
 
-			OutEnvironment.SetDefine(TEXT("MATERIAL_VOLUMETRIC_ADVANCED"), TEXT("1"));
+			EnvironmentDefines.bMaterialVolumetricAdvanced = true;
 
 			const UMaterialExpressionVolumetricAdvancedMaterialOutput* VolumetricAdvancedNode = VolumetricAdvancedExpressions[0];
-			if (VolumetricAdvancedNode->GetEvaluatePhaseOncePerSample())
-			{
-				OutEnvironment.SetDefine(TEXT("MATERIAL_VOLUMETRIC_ADVANCED_PHASE_PERSAMPLE"), TEXT("1"));
-			}
-			else
-			{
-				OutEnvironment.SetDefine(TEXT("MATERIAL_VOLUMETRIC_ADVANCED_PHASE_PERPIXEL"), TEXT("1"));
-			}
-
-			OutEnvironment.SetDefine(TEXT("MATERIAL_VOLUMETRIC_ADVANCED_GRAYSCALE_MATERIAL"), VolumetricAdvancedNode->bGrayScaleMaterial ? TEXT("1") : TEXT("0"));
-			OutEnvironment.SetDefine(TEXT("MATERIAL_VOLUMETRIC_ADVANCED_RAYMARCH_VOLUME_SHADOW"), VolumetricAdvancedNode->bRayMarchVolumeShadow ? TEXT("1") : TEXT("0"));
-			OutEnvironment.SetDefine(TEXT("MATERIAL_VOLUMETRIC_ADVANCED_CLAMP_MULTISCATTERING_CONTRIBUTION"), VolumetricAdvancedNode->bClampMultiScatteringContribution ? TEXT("1") : TEXT("0"));
-
-			OutEnvironment.SetDefine(TEXT("MATERIAL_VOLUMETRIC_ADVANCED_MULTISCATTERING_OCTAVE_COUNT"), VolumetricAdvancedNode->GetMultiScatteringApproximationOctaveCount());
-
-			OutEnvironment.SetDefine(TEXT("MATERIAL_VOLUMETRIC_ADVANCED_CONSERVATIVE_DENSITY"),
-				VolumetricAdvancedNode->ConservativeDensity.IsConnected() ? TEXT("1") : TEXT("0"));
+			EnvironmentDefines.bMaterialVolumetricAdvancedPhasePerSample = VolumetricAdvancedNode->GetEvaluatePhaseOncePerSample();
 
-			OutEnvironment.SetDefine(TEXT("MATERIAL_VOLUMETRIC_ADVANCED_OVERRIDE_AMBIENT_OCCLUSION"),
-				Material->HasAmbientOcclusionConnected() ? TEXT("1") : TEXT("0"));
-
-			OutEnvironment.SetDefine(TEXT("MATERIAL_VOLUMETRIC_ADVANCED_GROUND_CONTRIBUTION"),
-				VolumetricAdvancedNode->bGroundContribution ? TEXT("1") : TEXT("0"));
+			EnvironmentDefines.bMaterialVolumetricAdvancedGreyscaleMaterial = VolumetricAdvancedNode->bGrayScaleMaterial;
+			EnvironmentDefines.bMaterialVolumetricAdvancedRaymarchVolumeShadow = VolumetricAdvancedNode->bRayMarchVolumeShadow;
+			EnvironmentDefines.bMaterialVolumetricAdvancedClampMultiscatteringContribution = VolumetricAdvancedNode->bClampMultiScatteringContribution;
+			EnvironmentDefines.MaterialVolumetricAdvancedMultiscatteringOctaveCount = VolumetricAdvancedNode->GetMultiScatteringApproximationOctaveCount();
+			EnvironmentDefines.bMaterialVolumetricAdvancedConservativeDensity = VolumetricAdvancedNode->ConservativeDensity.IsConnected();
+			EnvironmentDefines.bMaterialVolumetricAdvancedOverrideAmbientOcclusion = Material->HasAmbientOcclusionConnected();
+			EnvironmentDefines.bMaterialVolumetricAdvancedAdvancedGroundContribution = VolumetricAdvancedNode->bGroundContribution;
 		}
 
 		TArray<const UMaterialExpressionVolumetricCloudEmptySpaceSkippingOutput*> EmptySpaceSkippingOutputExpressions;
@@ -2305,40 +2487,39 @@ void FHLSLMaterialTranslator::GetMaterialEnvironment(EShaderPlatform InPlatform,
 			{
 				UE_LOG(LogMaterial, Fatal, TEXT("Only a single UMaterialExpressionEmptySpaceSkippingOutput node is supported."));
 			}
-			OutEnvironment.SetDefine(TEXT("MATERIAL_VOLUMETRIC_CLOUD_EMPTY_SPACE_SKIPPING_OUTPUT"), TEXT("1"));
+			EnvironmentDefines.bMaterialVolumetricCloudEmptySpaceSkippingOutput = true;
 		}
 	}
 
-	OutEnvironment.SetDefine(TEXT("MATERIAL_IS_STRATA"), bMaterialIsStrata ? TEXT("1") : TEXT("0"));
+	EnvironmentDefines.bMaterialIsStrata = bMaterialIsStrata;
 
 	// Strata requests dual source blending only for BLEND_TranslucentColoredTransmittance
 	bMaterialRequestsDualSourceBlending |= bMaterialIsStrata && Material->GetBlendMode() == EBlendMode::BLEND_TranslucentColoredTransmittance;
 
 	// if duals source blending (colored transmittance) is not supported on a platform, it will fall back to standard alpha blending (grey scale transmittance)
-	OutEnvironment.SetDefine(TEXT("DUAL_SOURCE_COLOR_BLENDING_ENABLED"), bMaterialRequestsDualSourceBlending&& Material->IsDualBlendingEnabled(Platform) ? TEXT("1") : TEXT("0"));
-
-	OutEnvironment.SetDefine(TEXT("STRATA_PREMULTIPLIED_ALPHA_OPACITY_OVERRIDEN"), bMaterialIsStrata && bOpacityPropertyIsUsed ? TEXT("1") : TEXT("0"));
+	EnvironmentDefines.bDualSourceColorBlendingEnabled = bMaterialRequestsDualSourceBlending && Material->IsDualBlendingEnabled(Platform);
+	EnvironmentDefines.bStrataPremultipliedAlphaOpacityOverridden = bMaterialIsStrata && bOpacityPropertyIsUsed;
 
 	if (bMaterialIsStrata)
 	{
-		OutEnvironment.SetDefine(TEXT("STRATA_USES_CONVERSION_FROM_LEGACY"), bStrataUsesConversionFromLegacy ? 1 : 0);
-		OutEnvironment.SetDefine(TEXT("STRATA_MATERIAL_OUTPUT_OPAQUE_ROUGH_REFRACTIONS"), bStrataOutputsOpaqueRoughRefractions ? 1 : 0);
+		EnvironmentDefines.bStrataUsesConversionFromLegacy = bStrataUsesConversionFromLegacy;
+		EnvironmentDefines.bStrataMaterialOutputOpaqueRoughRefractions = bStrataOutputsOpaqueRoughRefractions;
 
-		OutEnvironment.SetDefine(TEXT("STRATA_MATERIAL_EXPORT_TYPE"), (int32)GetStrataMaterialExportType());
-		OutEnvironment.SetDefine(TEXT("STRATA_MATERIAL_EXPORT_CONTEXT"), (int32)GetStrataMaterialExportContext());
-		OutEnvironment.SetDefine(TEXT("STRATA_MATERIAL_EXPORT_LEGACY_BLEND_MODE"), (int32)GetStrataMaterialExportLegacyBlendMode());
+		EnvironmentDefines.StrataMaterialExportType = (int32)GetStrataMaterialExportType();
+		EnvironmentDefines.StrataMaterialExportContext = (int32)GetStrataMaterialExportContext();
+		EnvironmentDefines.StrataMaterialExportLegacyBlendMode = (int32)GetStrataMaterialExportLegacyBlendMode();
 
 		// Unlit cannot be combined with other BSDF so we can simply pick the default strata context
-		OutEnvironment.SetDefine(TEXT("STRATA_OPTIMIZED_UNLIT"), StrataCompilationContext[EStrataCompilationContext::SCC_Default].bStrataMaterialIsUnlitNode ? 1 : 0);
+		EnvironmentDefines.bStrataOptimizedUnlit = StrataCompilationContext[EStrataCompilationContext::SCC_Default].bStrataMaterialIsUnlitNode;
 
 		{
 			// For now, the fully simplified mode is used for Lumen or anything else supported inlined evaluation. The export is only valid for the default case.
 			// STRATA_TODO: generate an export for the different context (need to generate two export functions: the default one and the FullSimplification one)
 			FStrataCompilationContext& StrataCtx = StrataCompilationContext[EStrataCompilationContext::SCC_Default];
-			OutEnvironment.SetDefine(TEXT("STRATA_SINGLEPATH"), StrataCtx.bStrataMaterialIsSingle ? TEXT("1") : TEXT("0"));
-			OutEnvironment.SetDefine(TEXT("STRATA_FASTPATH"), StrataCtx.bStrataMaterialIsSingle ? TEXT("0") : (StrataCtx.bStrataMaterialIsSimple ? TEXT("1") : TEXT("0")));
-			OutEnvironment.SetDefine(TEXT("STRATA_CLAMPED_BSDF_COUNT"), StrataCtx.StrataMaterialBSDFCount);
-			OutEnvironment.SetDefine(TEXT("STRATA_COMPLEXSPECIALPATH"), MaterialCompilationOutput.StrataMaterialCompilationOutput.bUsesComplexSpecialRenderPath ? TEXT("1") : TEXT("0"));
+			EnvironmentDefines.bStrataSinglePath = StrataCtx.bStrataMaterialIsSingle;
+			EnvironmentDefines.bStrataFastPath = StrataCtx.bStrataMaterialIsSingle ? false : StrataCtx.bStrataMaterialIsSimple ? true : false;
+			EnvironmentDefines.StrataClampedBSDFCount = StrataCtx.StrataMaterialBSDFCount;
+			EnvironmentDefines.bStrataComplexSpecialPath = MaterialCompilationOutput.StrataMaterialCompilationOutput.bUsesComplexSpecialRenderPath;
 		}
 
 
@@ -2355,7 +2536,7 @@ void FHLSLMaterialTranslator::GetMaterialEnvironment(EShaderPlatform InPlatform,
 			StrataCtx.StrataPixelNormalInitializerValues = FString::Printf(TEXT("\n\n\n\t// Strata normal and tangent\n"));
 			StrataCtx.FinalUsedSharedLocalBasesCount = 0;
 			uint8 RequestedSharedLocalBasesCount = 0;
-			StrataCtx.StrataEvaluateSharedLocalBases(this, RequestedSharedLocalBasesCount, &OutEnvironment);
+			StrataCtx.StrataEvaluateSharedLocalBases(this, RequestedSharedLocalBasesCount, &EnvironmentDefines);
 
 #if WITH_EDITOR
 			// Now write some feedback to the user, but only produce debug string if in editor
@@ -2409,46 +2590,427 @@ void FHLSLMaterialTranslator::GetMaterialEnvironment(EShaderPlatform InPlatform,
 					}
 					const FStrataSharedLocalBasesInfo& StrataSharedLocalBasesInfo = StrataCtx.StrataCompilationInfoGetMatchingSharedLocalBasisInfo(BSDFOperator.BSDFRegisteredSharedLocalBasis);
 
-					StrataMaterialContextDescription += FString::Printf(TEXT("     - %s - SharedLocalBasisIndexMacro = %s \r\n"), *GetStrataBSDFName(BSDFOperator.BSDFType), *GetStrataSharedLocalBasisIndexMacroInner(BSDFOperator.BSDFRegisteredSharedLocalBasis, StrataCompilationContextIndex));
-				}
+					StrataMaterialContextDescription += FString::Printf(TEXT("     - %s - SharedLocalBasisIndexMacro = %s \r\n"), *GetStrataBSDFName(BSDFOperator.BSDFType), *GetStrataSharedLocalBasisIndexMacroInner(BSDFOperator.BSDFRegisteredSharedLocalBasis, StrataCompilationContextIndex));
+				}
+
+				StrataMaterialContextDescription += FString::Printf(TEXT("----------- SUBSTRATE TREE - %s -----------\r\n"), *StrataCompilationContextName);
+				StrataMaterialContextDescription += FString::Printf(TEXT("Graph maximum distance to leaves %u\r\n"), RootMaximumDistanceToLeaves);
+				// Debug print operators according to depth from root.
+				{
+
+					for (int32 DistanceToLeaves = RootMaximumDistanceToLeaves; DistanceToLeaves >= 0; --DistanceToLeaves)
+					{
+						StrataMaterialContextDescription += FString::Printf(TEXT("----- DistanceFromLeaves = %d -----\r\n"), DistanceToLeaves);
+						for (auto& It : StrataCtx.StrataMaterialExpressionRegisteredOperators)
+						{
+							if (!It.IsDiscarded() && It.MaxDistanceFromLeaves == DistanceToLeaves)
+							{
+								StrataMaterialContextDescription += FString::Printf(TEXT("\tIdx=%d Op=%s ParentIdx=%d LeftIndex=%d RightIndex=%d BSDFIdx=%d LayerDepth=%d IsTop=%d IsBot=%d BSDFType=%s SSS=%d MFP=%d F90=%d Rough2=%d Fuzz=%d Aniso=%d Glint=%d SpecularProfile=%d\r\n"),
+									It.Index, GetStrataOperatorStr(It.OperatorType), It.ParentIndex, It.LeftIndex, It.RightIndex, It.BSDFIndex, It.LayerDepth, It.bIsTop, It.bIsBottom,
+									*GetStrataBSDFName(It.BSDFType), It.bBSDFHasSSS, It.bBSDFHasMFPPluggedIn, It.bBSDFHasEdgeColor, It.bBSDFHasSecondRoughnessOrSimpleClearCoat, It.bBSDFHasFuzz, It.bBSDFHasAnisotropy, It.bBSDFHasGlint, It.bBSDFHasSpecularProfile);
+							}
+						}
+					}
+				}
+
+				ResourcesString += TEXT("/*");
+				ResourcesString += StrataMaterialContextDescription;
+				ResourcesString += TEXT("*/");
+
+				StrataMaterialDescription += StrataMaterialContextDescription;
+
+				if (StrataCompilationContextIndex == EStrataCompilationContext::SCC_Default)
+				{
+					MaterialCompilationOutput.StrataMaterialCompilationOutput.SharedLocalBasesCount = StrataCtx.FinalUsedSharedLocalBasesCount;
+				}
+			}
+#endif // WITH_EDITOR
+		}
+		MaterialCompilationOutput.StrataMaterialCompilationOutput.StrataMaterialDescription = StrataMaterialDescription;
+	}
+
+	EnvironmentDefines.bTextureSampleDebug = IsDebugTextureSampleEnabled();
+}
+
+void FHLSLMaterialTranslator::PushResultsToDDCCache(const FIoHash& KeyHash)
+{
+	if (!CVarJobMaterialTranslateDDCEnable.GetValueOnAnyThread())
+	{
+		return;
+	}
+
+	UE::DerivedData::FCacheKey CacheKey{ MaterialTranslationDDCBucket, KeyHash };
+	UE::DerivedData::FCacheRecordBuilder RecordBuilder{ CacheKey };
+
+	// Push the material compilation output
+	FBufferArchive MaterialCompilationOutputBuffer;
+	FShaderMapPointerTable PointerTable;
+
+	FMemoryImage MemoryImage;
+	MemoryImage.PrevPointerTable = &PointerTable;
+	MemoryImage.PointerTable = &PointerTable;
+	MemoryImage.TargetLayoutParameters.InitializeForArchive(MaterialCompilationOutputBuffer);
+
+	FMemoryImageWriter Writer(MemoryImage);
+	Writer.WriteRootObject(&MaterialCompilationOutput, StaticGetTypeLayoutDesc<FMaterialCompilationOutput>());
+
+	FMemoryImageResult MemoryImageResult;
+	MemoryImage.Flatten(MemoryImageResult, true);
+	MemoryImageResult.SaveToArchive(MaterialCompilationOutputBuffer);
+	RecordBuilder.AddValue(MaterialCompilationOutputId, FSharedBuffer::MakeView(MaterialCompilationOutputBuffer.GetData(), MaterialCompilationOutputBuffer.Num()));
+
+	// Push the material source string parameters
+	FBufferArchive ResultsBuffer;
+	ResultsBuffer << MaterialSourceTemplateParams;
+	RecordBuilder.AddValue(MaterialResultsOutputId, FSharedBuffer::MakeView(ResultsBuffer.GetData(), ResultsBuffer.Num()));
+
+	// Push the material shader defines
+	FBufferArchive EnvironmentDefinesBuffer;
+	FObjectAndNameAsStringProxyArchive EnvironmentDefinesBufferProxy{ EnvironmentDefinesBuffer, true };
+	EnvironmentDefines.StaticStruct()->SerializeItem(EnvironmentDefinesBufferProxy, &EnvironmentDefines, nullptr);
+	RecordBuilder.AddValue(EnvironmentDefinesId, FSharedBuffer::MakeView(EnvironmentDefinesBuffer.GetData(), EnvironmentDefinesBuffer.Num()));
+
+	// Push the all the data to the DDC
+	UE::DerivedData::FRequestOwner RequestOwner{ UE::DerivedData::EPriority::Normal };
+	UE::DerivedData::FCachePutRequest Request{ {TEXT("FHLSLMaterialTranslator::PushResultsToDDCCache")}, RecordBuilder.Build(), UE::DerivedData::ECachePolicy::Default | UE::DerivedData::ECachePolicy::KeepAlive };
+	UE::DerivedData::GetCache().Put({ Request }, RequestOwner);
+}
+
+void FHLSLMaterialTranslator::ValidateShadingModelsForFeatureLevel(const FMaterialShadingModelField& ShadingModels)
+{
+	if (FeatureLevel <= ERHIFeatureLevel::ES3_1)
+	{
+		const TArray<EMaterialShadingModel>& InvalidShadingModels = {};
+		for (EMaterialShadingModel InvalidShadingModel : InvalidShadingModels)
+		{
+			if (ShadingModels.HasShadingModel(InvalidShadingModel))
+			{
+				FString FeatureLevelName;
+				GetFeatureLevelName(FeatureLevel, FeatureLevelName);
+
+				FString ShadingModelName;
+				const UEnum* EnumPtr = FindObject<UEnum>(nullptr, TEXT("/Script/Engine.EMaterialShadingModel"), true);
+				if (EnumPtr)
+				{
+					ShadingModelName = EnumPtr->GetNameStringByValue(InvalidShadingModel);
+				}
+
+				Errorf(TEXT("ShadingModel %s not supported in feature level %s"), *ShadingModelName, *FeatureLevelName);
+			}
+		}
+	}
+}
+
+void FHLSLMaterialTranslator::GetMaterialEnvironment(EShaderPlatform InPlatform, FShaderCompilerEnvironment& OutEnvironment)
+{
+	check(InPlatform == GetShaderPlatform());
+
+	const bool bStrataEnabled = Strata::IsStrataEnabled();
+	bool bMaterialRequestsDualSourceBlending = false;
+
+	if (EnvironmentDefines.bNeedsParticlePosition)
+	{
+		OutEnvironment.SetDefine(TEXT("NEEDS_PARTICLE_POSITION"), 1);
+	}
+
+	if (EnvironmentDefines.bNeedsParticleVelocity)
+	{
+		OutEnvironment.SetDefine(TEXT("NEEDS_PARTICLE_VELOCITY"), 1);
+	}
+
+	if (EnvironmentDefines.bUseDynamicParameters)
+	{
+		OutEnvironment.SetDefine(TEXT("USE_DYNAMIC_PARAMETERS"), 1);
+		OutEnvironment.SetDefine(TEXT("DYNAMIC_PARAMETERS_MASK"), EnvironmentDefines.DynamicParametersMask);
+	}
+
+	if (EnvironmentDefines.bNeedsParticleTime)
+	{
+		OutEnvironment.SetDefine(TEXT("NEEDS_PARTICLE_TIME"), 1);
+	}
+
+	if (EnvironmentDefines.bUsesParticleMotionBlur)
+	{
+		OutEnvironment.SetDefine(TEXT("USES_PARTICLE_MOTION_BLUR"), 1);
+	}
+
+	if (EnvironmentDefines.bNeedsParticleRandom)
+	{
+		OutEnvironment.SetDefine(TEXT("NEEDS_PARTICLE_RANDOM"), 1);
+	}
+
+	if (EnvironmentDefines.bSphericalParticleOpacity)
+	{
+		OutEnvironment.SetDefine(TEXT("SPHERICAL_PARTICLE_OPACITY"), TEXT("1"));
+	}
+
+	if (EnvironmentDefines.bUseParticleSubUVs)
+	{
+		OutEnvironment.SetDefine(TEXT("USE_PARTICLE_SUBUVS"), TEXT("1"));
+	}
+
+	if (EnvironmentDefines.bLightmapUVAccess)
+	{
+		OutEnvironment.SetDefine(TEXT("LIGHTMAP_UV_ACCESS"), TEXT("1"));
+	}
+
+	if (EnvironmentDefines.bUsesAOMaterialMask)
+	{
+		OutEnvironment.SetDefine(TEXT("USES_AO_MATERIAL_MASK"), TEXT("1"));
+	}
+
+	if (EnvironmentDefines.bUsesSpeedtree)
+	{
+		OutEnvironment.SetDefine(TEXT("USES_SPEEDTREE"), TEXT("1"));
+	}
+
+	if (EnvironmentDefines.bNeedsWorldPositionExcludingShaderOffsets)
+	{
+		OutEnvironment.SetDefine(TEXT("NEEDS_WORLD_POSITION_EXCLUDING_SHADER_OFFSETS"), TEXT("1"));
+	}
+
+	if (EnvironmentDefines.bNeedsParticleSize)
+	{
+		OutEnvironment.SetDefine(TEXT("NEEDS_PARTICLE_SIZE"), TEXT("1"));
+	}
+
+	if (EnvironmentDefines.bNeedsParticleSpriteRotation)
+	{
+		OutEnvironment.SetDefine(TEXT("NEEDS_PARTICLE_SPRITE_ROTATION"), TEXT("1"));
+	}
+
+	if (EnvironmentDefines.bNeedsSceneTextures)
+	{
+		OutEnvironment.SetDefine(TEXT("NEEDS_SCENE_TEXTURES"), TEXT("1"));
+	}
+
+	if (EnvironmentDefines.bUsesEyeAdaptation)
+	{
+		OutEnvironment.SetDefine(TEXT("USES_EYE_ADAPTATION"), TEXT("1"));
+	}
+
+	if (EnvironmentDefines.bVirtualTextureOutput)
+	{
+		OutEnvironment.SetDefine(TEXT("VIRTUAL_TEXTURE_OUTPUT"), 1);
+	}
+
+	OutEnvironment.SetDefine(TEXT("USES_PER_INSTANCE_CUSTOM_DATA"), EnvironmentDefines.bUsesPerInstanceCustomData);
+	OutEnvironment.SetDefine(TEXT("USES_PER_INSTANCE_RANDOM"), EnvironmentDefines.bUsesPerInstanceRandom);
+	OutEnvironment.SetDefine(TEXT("USES_PER_INSTANCE_FADE_AMOUNT"), EnvironmentDefines.bUsesPerInstanceFadeAmount);
+	OutEnvironment.SetDefine(TEXT("USES_VERTEX_INTERPOLATOR"), EnvironmentDefines.bUsesVertexInterpolator);
+
+	OutEnvironment.SetDefine(TEXT("MATERIAL_SKY_ATMOSPHERE"), EnvironmentDefines.bUsesSkyAtmosphere);
+	OutEnvironment.SetDefine(TEXT("INTERPOLATE_VERTEX_COLOR"), EnvironmentDefines.bUsesVertexColor);
+	OutEnvironment.SetDefine(TEXT("NEEDS_PARTICLE_COLOR"), EnvironmentDefines.bUsesParticleColor);
+	OutEnvironment.SetDefine(TEXT("NEEDS_PARTICLE_LOCAL_TO_WORLD"), EnvironmentDefines.bUsesParticleLocalToWorld);
+	OutEnvironment.SetDefine(TEXT("NEEDS_PARTICLE_WORLD_TO_LOCAL"), EnvironmentDefines.bUsesParticleWorldToLocal);
+	OutEnvironment.SetDefine(TEXT("NEEDS_INSTANCE_LOCAL_TO_WORLD_PS"), EnvironmentDefines.bUsesInstanceLocalToWorldPS);
+	OutEnvironment.SetDefine(TEXT("NEEDS_INSTANCE_WORLD_TO_LOCAL_PS"), EnvironmentDefines.bUsesInstanceWorldToLocalPS);
+	OutEnvironment.SetDefine(TEXT("USES_TRANSFORM_VECTOR"), EnvironmentDefines.bUsesTransformVector);
+	OutEnvironment.SetDefine(TEXT("WANT_PIXEL_DEPTH_OFFSET"), EnvironmentDefines.bUsesPixelDepthOffset);
+
+	// we want USES_WORLD_POSITION_OFFSET to be readable as a bool compile argument, hence the != 0 comparison 
+	// (bUsesWorldPositionOffset is actually a 1-bit uint32 bitfield member)
+	OutEnvironment.SetDefineAndCompileArgument(TEXT("USES_WORLD_POSITION_OFFSET"), EnvironmentDefines.bUsesWorldPositionOffset);
+	OutEnvironment.SetDefineAndCompileArgument(TEXT("USES_DISPLACEMENT"), EnvironmentDefines.bUsesDisplacement);
+
+	OutEnvironment.SetDefine(TEXT("USES_EMISSIVE_COLOR"), EnvironmentDefines.bUsesEmissiveColor);
+	// Distortion uses tangent space transform 
+	OutEnvironment.SetDefine(TEXT("USES_DISTORTION"), EnvironmentDefines.bUsesDistortion);
+	OutEnvironment.SetDefine(TEXT("DISTORTION_ACCOUNT_FOR_COVERAGE"), EnvironmentDefines.bDistortionAccountForCoverage);
+
+	OutEnvironment.SetDefine(TEXT("MATERIAL_ENABLE_TRANSLUCENCY_FOGGING"), EnvironmentDefines.bMaterialEnableTranslucencyFogging);
+	OutEnvironment.SetDefine(TEXT("MATERIAL_ENABLE_TRANSLUCENCY_CLOUD_FOGGING"), EnvironmentDefines.bMaterialEnableTranslucencyCloudFogging);
+	OutEnvironment.SetDefine(TEXT("MATERIAL_IS_SKY"), EnvironmentDefines.bMaterialIsSky);
+	OutEnvironment.SetDefine(TEXT("MATERIAL_COMPUTE_FOG_PER_PIXEL"), EnvironmentDefines.bMaterialComputeFogPerPixel);
+	OutEnvironment.SetDefine(TEXT("MATERIAL_FULLY_ROUGH"), EnvironmentDefines.bMaterialFullyRough);
+
+	OutEnvironment.SetDefine(TEXT("MATERIAL_USES_ANISOTROPY"), EnvironmentDefines.bMaterialUsesAnisotropy);
+
+	OutEnvironment.SetDefine(TEXT("MATERIAL_DECAL_READ_MASK"), EnvironmentDefines.MaterialDecalReadMask);
+	OutEnvironment.SetDefine(TEXT("MATERIAL_USES_DECAL_LOOKUP"), EnvironmentDefines.bMaterialUsesDecalLookup);
+	OutEnvironment.SetDefine(TEXT("MATERIAL_PATH_TRACING_BUFFER_READ"), EnvironmentDefines.MaterialPathTracingBufferRead);
+
+	OutEnvironment.SetDefine(TEXT("NUM_VIRTUALTEXTURE_SAMPLES"), EnvironmentDefines.NumVirtualTextureSamples);
+	OutEnvironment.SetDefine(TEXT("MATERIAL_VIRTUALTEXTURE_FEEDBACK"), EnvironmentDefines.bMaterialVirtualTextureFeedback);
+
+	// Setup defines to map each VT stack to either 1 or 2 page table textures, depending on how many layers it uses
+	for (int i = 0; i < EnvironmentDefines.VirtualPageTypes.Num(); ++i)
+	{
+		const FMaterialVirtualTextureStack& Stack = MaterialCompilationOutput.UniformExpressionSet.VTStacks[i];
+		FString PageTableValue = FString::Printf(TEXT("Material.VirtualTexturePageTable0_%d"), i);
+		if (EnvironmentDefines.VirtualPageTypes[i] & FHLSLMaterialTranslatorEnvironmentDefines::TABLE1)
+		{
+			PageTableValue += FString::Printf(TEXT(", Material.VirtualTexturePageTable1_%d"), i);
+		}
+		if (EnvironmentDefines.VirtualPageTypes[i] & FHLSLMaterialTranslatorEnvironmentDefines::TABLE_INDIRECTION)
+		{
+			PageTableValue += FString::Printf(TEXT(", Material.VirtualTexturePageTableIndirection_%d"), i);
+		}
+		OutEnvironment.SetDefine(*FString::Printf(TEXT("VIRTUALTEXTURE_PAGETABLE_%d"), i), *PageTableValue);
+	}
+
+	for (int32 CollectionIndex = 0; CollectionIndex < EnvironmentDefines.ParameterCollections.Num(); CollectionIndex++)
+	{
+		// Add uniform buffer declarations for any parameter collections referenced
+		const FString CollectionName = FString::Printf(TEXT("MaterialCollection%u"), CollectionIndex);
+		// This can potentially become an issue for MaterialCollection Uniform Buffers if they ever get non-numeric resources (eg Textures), as
+		// OutEnvironment.ResourceTableMap has a map by name, and the N ParameterCollection Uniform Buffers ALL are names "MaterialCollection"
+		// (and the hlsl cbuffers are named MaterialCollection0, etc, so the names don't match the layout)
+		FShaderUniformBufferParameter::ModifyCompilationEnvironment(*CollectionName, EnvironmentDefines.ParameterCollections[CollectionIndex]->GetUniformBufferStruct(), InPlatform, OutEnvironment);
+	}
+
+	OutEnvironment.SetDefine(TEXT("IS_MATERIAL_SHADER"), TEXT("1"));
+
+	if (EnvironmentDefines.bSingleLayerWaterShadingQuality)
+	{
+		// Value must match SINGLE_LAYER_WATER_SHADING_QUALITY_MOBILE_WITH_DEPTH_TEXTURE in SingleLayerWaterCommon.ush!
+		OutEnvironment.SetDefine(TEXT("SINGLE_LAYER_WATER_SHADING_QUALITY"), TEXT("1"));
+	}
+
+	if (EnvironmentDefines.bShadingModelsIsLit)
+	{
+		if (EnvironmentDefines.HasShadingModel(MSM_DefaultLit))
+		{
+			OutEnvironment.SetDefine(TEXT("MATERIAL_SHADINGMODEL_DEFAULT_LIT"), TEXT("1"));
+		}
+
+		if (EnvironmentDefines.HasShadingModel(MSM_Subsurface))
+		{
+			OutEnvironment.SetDefine(TEXT("MATERIAL_SHADINGMODEL_SUBSURFACE"), TEXT("1"));
+		}
+
+		if (EnvironmentDefines.HasShadingModel(MSM_PreintegratedSkin))
+		{
+			OutEnvironment.SetDefine(TEXT("MATERIAL_SHADINGMODEL_PREINTEGRATED_SKIN"), TEXT("1"));
+		}
+
+		if (EnvironmentDefines.HasShadingModel(MSM_SubsurfaceProfile))
+		{
+			OutEnvironment.SetDefine(TEXT("MATERIAL_SHADINGMODEL_SUBSURFACE_PROFILE"), TEXT("1"));
+		}
+
+		if (EnvironmentDefines.bMaterialSubsurfaceProfileUseCurvature)
+		{
+			OutEnvironment.SetDefine(TEXT("MATERIAL_SUBSURFACE_PROFILE_USE_CURVATURE"), TEXT("1"));
+		}
+
+		if (EnvironmentDefines.HasShadingModel(MSM_ClearCoat))
+		{
+			OutEnvironment.SetDefine(TEXT("MATERIAL_SHADINGMODEL_CLEAR_COAT"), TEXT("1"));
+		}
+
+		if (EnvironmentDefines.HasShadingModel(MSM_TwoSidedFoliage))
+		{
+			OutEnvironment.SetDefine(TEXT("MATERIAL_SHADINGMODEL_TWOSIDED_FOLIAGE"), TEXT("1"));
+		}
+
+		if (EnvironmentDefines.HasShadingModel(MSM_Hair))
+		{
+			OutEnvironment.SetDefine(TEXT("MATERIAL_SHADINGMODEL_HAIR"), TEXT("1"));
+		}
+
+		if (EnvironmentDefines.HasShadingModel(MSM_Cloth))
+		{
+			OutEnvironment.SetDefine(TEXT("MATERIAL_SHADINGMODEL_CLOTH"), TEXT("1"));
+		}
+
+		if (EnvironmentDefines.HasShadingModel(MSM_Eye))
+		{
+			OutEnvironment.SetDefine(TEXT("MATERIAL_SHADINGMODEL_EYE"), TEXT("1"));
+		}
+
+		if (EnvironmentDefines.bMaterialShadingModelEyeUseCurvature)
+		{
+			OutEnvironment.SetDefine(TEXT("MATERIAL_SHADINGMODEL_EYE_USE_CURVATURE"), TEXT("1"));
+		}
+
+		if (EnvironmentDefines.HasShadingModel(MSM_SingleLayerWater))
+		{
+			OutEnvironment.SetDefine(TEXT("MATERIAL_SHADINGMODEL_SINGLELAYERWATER"), TEXT("1"));
+		}
+
+		if (EnvironmentDefines.HasShadingModel(MSM_ThinTranslucent))
+		{
+			OutEnvironment.SetDefine(TEXT("MATERIAL_SHADINGMODEL_THIN_TRANSLUCENT"), TEXT("1"));
+			bMaterialRequestsDualSourceBlending = true;
+		}
+
+		if (EnvironmentDefines.bDisableForwardLocalLights)
+		{
+			OutEnvironment.SetDefine(TEXT("DISABLE_FORWARD_LOCAL_LIGHTS"), TEXT("1"));
+		}
+
+		if (EnvironmentDefines.bSingleLayerWaterSeparatedMainLight)
+		{
+			OutEnvironment.SetDefine(TEXT("SINGLE_LAYER_WATER_SEPARATED_MAIN_LIGHT"), TEXT("1"));
+		}
+
+		if (EnvironmentDefines.bMaterialSingleShadingModel)
+		{
+			OutEnvironment.SetDefine(TEXT("MATERIAL_SINGLE_SHADINGMODEL"), TEXT("1"));
+		}
+	}
+	else
+	{
+		OutEnvironment.SetDefine(TEXT("MATERIAL_SINGLE_SHADINGMODEL"), TEXT("1"));
+		OutEnvironment.SetDefine(TEXT("MATERIAL_SHADINGMODEL_UNLIT"), TEXT("1"));
+	}
+
+	if (EnvironmentDefines.bMaterialVolumetricAdvanced)
+	{
+		OutEnvironment.SetDefine(TEXT("MATERIAL_VOLUMETRIC_ADVANCED"), TEXT("1"));
+
+		if (EnvironmentDefines.bMaterialVolumetricAdvancedPhasePerSample)
+		{
+			OutEnvironment.SetDefine(TEXT("MATERIAL_VOLUMETRIC_ADVANCED_PHASE_PERSAMPLE"), TEXT("1"));
+		}
+		else
+		{
+			OutEnvironment.SetDefine(TEXT("MATERIAL_VOLUMETRIC_ADVANCED_PHASE_PERPIXEL"), TEXT("1"));
+		}
 
-				StrataMaterialContextDescription += FString::Printf(TEXT("----------- SUBSTRATE TREE - %s -----------\r\n"), *StrataCompilationContextName);
-				StrataMaterialContextDescription += FString::Printf(TEXT("Graph maximum distance to leaves %u\r\n"), RootMaximumDistanceToLeaves);
-				// Debug print operators according to depth from root.
-				{
+		OutEnvironment.SetDefine(TEXT("MATERIAL_VOLUMETRIC_ADVANCED_GRAYSCALE_MATERIAL"), EnvironmentDefines.bMaterialVolumetricAdvancedGreyscaleMaterial);
+		OutEnvironment.SetDefine(TEXT("MATERIAL_VOLUMETRIC_ADVANCED_RAYMARCH_VOLUME_SHADOW"), EnvironmentDefines.bMaterialVolumetricAdvancedRaymarchVolumeShadow);
+		OutEnvironment.SetDefine(TEXT("MATERIAL_VOLUMETRIC_ADVANCED_CLAMP_MULTISCATTERING_CONTRIBUTION"), EnvironmentDefines.bMaterialVolumetricAdvancedClampMultiscatteringContribution);
+		OutEnvironment.SetDefine(TEXT("MATERIAL_VOLUMETRIC_ADVANCED_MULTISCATTERING_OCTAVE_COUNT"), EnvironmentDefines.MaterialVolumetricAdvancedMultiscatteringOctaveCount);
+		OutEnvironment.SetDefine(TEXT("MATERIAL_VOLUMETRIC_ADVANCED_CONSERVATIVE_DENSITY"), EnvironmentDefines.bMaterialVolumetricAdvancedConservativeDensity);
+		OutEnvironment.SetDefine(TEXT("MATERIAL_VOLUMETRIC_ADVANCED_OVERRIDE_AMBIENT_OCCLUSION"), EnvironmentDefines.bMaterialVolumetricAdvancedOverrideAmbientOcclusion);
+		OutEnvironment.SetDefine(TEXT("MATERIAL_VOLUMETRIC_ADVANCED_GROUND_CONTRIBUTION"), EnvironmentDefines.bMaterialVolumetricAdvancedAdvancedGroundContribution);
+	}
 
-					for (int32 DistanceToLeaves = RootMaximumDistanceToLeaves; DistanceToLeaves >= 0; --DistanceToLeaves)
-					{
-						StrataMaterialContextDescription += FString::Printf(TEXT("----- DistanceFromLeaves = %d -----\r\n"), DistanceToLeaves);
-						for (auto& It : StrataCtx.StrataMaterialExpressionRegisteredOperators)
-						{
-							if (!It.IsDiscarded() && It.MaxDistanceFromLeaves == DistanceToLeaves)
-							{
-								StrataMaterialContextDescription += FString::Printf(TEXT("\tIdx=%d Op=%s ParentIdx=%d LeftIndex=%d RightIndex=%d BSDFIdx=%d LayerDepth=%d IsTop=%d IsBot=%d BSDFType=%s SSS=%d MFP=%d F90=%d Rough2=%d Fuzz=%d Aniso=%d Glint=%d SpecularProfile=%d\r\n"),
-									It.Index, GetStrataOperatorStr(It.OperatorType), It.ParentIndex, It.LeftIndex, It.RightIndex, It.BSDFIndex, It.LayerDepth, It.bIsTop, It.bIsBottom,
-									*GetStrataBSDFName(It.BSDFType), It.bBSDFHasSSS, It.bBSDFHasMFPPluggedIn, It.bBSDFHasEdgeColor, It.bBSDFHasSecondRoughnessOrSimpleClearCoat, It.bBSDFHasFuzz, It.bBSDFHasAnisotropy, It.bBSDFHasGlint, It.bBSDFHasSpecularProfile);
-							}
-						}
-					}
-				}
+	if (EnvironmentDefines.bMaterialVolumetricCloudEmptySpaceSkippingOutput)
+	{
+		OutEnvironment.SetDefine(TEXT("MATERIAL_VOLUMETRIC_CLOUD_EMPTY_SPACE_SKIPPING_OUTPUT"), TEXT("1"));
+	}
 
-				ResourcesString += TEXT("/*");
-				ResourcesString += StrataMaterialContextDescription;
-				ResourcesString += TEXT("*/");
+	OutEnvironment.SetDefine(TEXT("MATERIAL_IS_STRATA"), EnvironmentDefines.bMaterialIsStrata);
+	OutEnvironment.SetDefine(TEXT("DUAL_SOURCE_COLOR_BLENDING_ENABLED"), EnvironmentDefines.bDualSourceColorBlendingEnabled);
+	OutEnvironment.SetDefine(TEXT("STRATA_PREMULTIPLIED_ALPHA_OPACITY_OVERRIDEN"), EnvironmentDefines.bStrataPremultipliedAlphaOpacityOverridden);
 
-				StrataMaterialDescription += StrataMaterialContextDescription;
+	if (EnvironmentDefines.bMaterialIsStrata)
+	{
+		OutEnvironment.SetDefine(TEXT("STRATA_USES_CONVERSION_FROM_LEGACY"), EnvironmentDefines.bStrataUsesConversionFromLegacy);
+		OutEnvironment.SetDefine(TEXT("STRATA_MATERIAL_OUTPUT_OPAQUE_ROUGH_REFRACTIONS"), EnvironmentDefines.bStrataMaterialOutputOpaqueRoughRefractions);
+		OutEnvironment.SetDefine(TEXT("STRATA_MATERIAL_EXPORT_TYPE"), EnvironmentDefines.StrataMaterialExportType);
+		OutEnvironment.SetDefine(TEXT("STRATA_MATERIAL_EXPORT_CONTEXT"), EnvironmentDefines.StrataMaterialExportContext);
+		OutEnvironment.SetDefine(TEXT("STRATA_MATERIAL_EXPORT_LEGACY_BLEND_MODE"), EnvironmentDefines.StrataMaterialExportLegacyBlendMode);
+		OutEnvironment.SetDefine(TEXT("STRATA_OPTIMIZED_UNLIT"), EnvironmentDefines.bStrataOptimizedUnlit);
 
-				if (StrataCompilationContextIndex == EStrataCompilationContext::SCC_Default)
-				{
-					MaterialCompilationOutput.StrataMaterialCompilationOutput.SharedLocalBasesCount = StrataCtx.FinalUsedSharedLocalBasesCount;
-				}
-			}
-#endif // WITH_EDITOR
+		{
+			OutEnvironment.SetDefine(TEXT("STRATA_SINGLEPATH"), EnvironmentDefines.bStrataSinglePath);
+			OutEnvironment.SetDefine(TEXT("STRATA_FASTPATH"), EnvironmentDefines.bStrataFastPath);
+			OutEnvironment.SetDefine(TEXT("STRATA_CLAMPED_BSDF_COUNT"), EnvironmentDefines.StrataClampedBSDFCount);
+			OutEnvironment.SetDefine(TEXT("STRATA_COMPLEXSPECIALPATH"), EnvironmentDefines.bStrataComplexSpecialPath);
 		}
-		MaterialCompilationOutput.StrataMaterialCompilationOutput.StrataMaterialDescription = StrataMaterialDescription;
 	}
 
-	OutEnvironment.SetDefine(TEXT("TEXTURE_SAMPLE_DEBUG"), IsDebugTextureSampleEnabled() ? TEXT("1") : TEXT("0"));
+	for (int i = 0; i < EnvironmentDefines.StrataDefineValues.Num(); ++i)
+	{
+		OutEnvironment.SetDefine(*EnvironmentDefines.StrataDefineKeys[i], EnvironmentDefines.StrataDefineValues[i]);
+	}
+
+	OutEnvironment.SetDefine(TEXT("TEXTURE_SAMPLE_DEBUG"), EnvironmentDefines.bTextureSampleDebug);
 }
 
 // Assign custom interpolators to slots, packing them as much as possible in unused slots.
@@ -2597,271 +3159,14 @@ void FHLSLMaterialTranslator::GetSharedInputsMaterialCode(FString& PixelMembersD
 
 FString FHLSLMaterialTranslator::GetMaterialShaderCode()
 {	
-	// use "/Engine/Private/MaterialTemplate.ush" to create the functions to get data (e.g. material attributes) and code (e.g. material expressions to create specular color) from C++
-	FLazyPrintf LazyPrintf(*MaterialTemplate);
-
-	// Assign slots to vertex interpolators
-	FString VertexInterpolatorsOffsetsDefinition;
-	TBitArray<> FinalAllocatedCoords = GetVertexInterpolatorsOffsets(VertexInterpolatorsOffsetsDefinition);
-
-	const uint32 NumUserVertexTexCoords = GetNumUserVertexTexCoords();
-	const uint32 NumUserTexCoords = GetNumUserTexCoords();
-	const uint32 NumCustomVectors = FMath::DivideAndRoundUp((uint32)CurrentCustomVertexInterpolatorOffset, 2u);
-	const uint32 NumTexCoordVectors = FinalAllocatedCoords.FindLast(true) + 1;
-
-	LazyPrintf.PushParam(*FString::Printf(TEXT("%u"),NumUserVertexTexCoords));
-	LazyPrintf.PushParam(*FString::Printf(TEXT("%u"),NumUserTexCoords));
-	LazyPrintf.PushParam(*FString::Printf(TEXT("%u"),NumCustomVectors));
-	LazyPrintf.PushParam(*FString::Printf(TEXT("%u"),NumTexCoordVectors));
-
-	LazyPrintf.PushParam(*VertexInterpolatorsOffsetsDefinition);
-
-	FString MaterialAttributesDeclaration;
-	FString MaterialAttributesUtilities;
-
-	const EMaterialShadingModel DefaultShadingModel = Material->GetShadingModels().GetFirstShadingModel();
-
-	MaterialAttributesDeclaration += TEXT("struct FMaterialAttributes\n{\n");
-
-	const TArray<FGuid>& OrderedVisibleAttributes = FMaterialAttributeDefinitionMap::GetOrderedVisibleAttributeList();
-	for (const FGuid& AttributeID : OrderedVisibleAttributes)
-	{
-		const FString PropertyName = FMaterialAttributeDefinitionMap::GetAttributeName(AttributeID);
-		const EMaterialValueType PropertyType = FMaterialAttributeDefinitionMap::GetValueType(AttributeID);
-		const TCHAR* HLSLType = nullptr;
-
-		switch (PropertyType)
-		{
-		case MCT_Float1: case MCT_Float: HLSLType = TEXT("float"); break;
-		case MCT_Float2: HLSLType = TEXT("float2"); break;
-		case MCT_Float3: HLSLType = TEXT("float3"); break;
-		case MCT_Float4: HLSLType = TEXT("float4"); break;
-		case MCT_UInt: case MCT_UInt1: case MCT_ShadingModel: HLSLType = TEXT("uint"); break;
-		case MCT_UInt2: HLSLType = TEXT("uint2"); break;
-		case MCT_UInt3: HLSLType = TEXT("uint3"); break;
-		case MCT_UInt4: HLSLType = TEXT("uint4"); break;
-		case MCT_Strata: HLSLType = TEXT("FStrataData"); break;
-		default: break;
-		}
-
-		if (HLSLType)
-		{
-			const FVector4f DefaultValue = FMaterialAttributeDefinitionMap::GetDefaultValue(AttributeID);
-
-			MaterialAttributesDeclaration += FString::Printf(TEXT("\t%s %s;") LINE_TERMINATOR, HLSLType, *PropertyName);
-
-			// Chainable method to set the attribute
-			MaterialAttributesUtilities += FString::Printf(TEXT("FMaterialAttributes FMaterialAttributes_Set%s(FMaterialAttributes InAttributes, %s InValue) { InAttributes.%s = InValue; return InAttributes; }") LINE_TERMINATOR,
-				*PropertyName, HLSLType, *PropertyName);
-		}
-	}
-
-	MaterialAttributesDeclaration += TEXT("};\n");
-
-	LazyPrintf.PushParam(*MaterialAttributesDeclaration);
-	LazyPrintf.PushParam(*MaterialAttributesUtilities);
-
-	// Stores the shared shader results member declarations
-	FString PixelMembersDeclaration[CompiledPDV_MAX];
-
-	FString NormalAssignment[CompiledPDV_MAX];
-
-	// Stores the code to initialize all inputs after MP_Normal
-	FString PixelMembersSetupAndAssignments[CompiledPDV_MAX];
-
-	for (int32 Index = 0; Index < CompiledPDV_MAX; Index++)
-	{
-		GetSharedInputsMaterialCode(PixelMembersDeclaration[Index], NormalAssignment[Index], PixelMembersSetupAndAssignments[Index], (ECompiledPartialDerivativeVariation)Index);
-	}
-
-	// PixelMembersDeclaration should be the same for all variations, but might change in the future. There are cases where work is shared
-	// between the pixel and vertex shader, but with Nanite all work has to be moved into the pixel shader, which means we will want
-	// different inputs. But for now, we are keeping them the same.
-	LazyPrintf.PushParam(*PixelMembersDeclaration[CompiledPDV_FiniteDifferences]);
-
-	{
-		FString DerivativeHelpers = DerivativeAutogen.GenerateUsedFunctions(*this);
-		FString DerivativeHelpersAndResources = DerivativeHelpers + ResourcesString;
-		//LazyPrintf.PushParam(*ResourcesString);
-		LazyPrintf.PushParam(*DerivativeHelpersAndResources);
-	}
-
-	// Anything used bye the GenerationFunctionCode() like WorldPositionOffset shouldn't be using texures, right?
-	// Let those use the standard finite differences textures, since they should be the same. If we actually want
-	// those to handle texture reads properly, we'll have to make extra versions.
-	ECompiledPartialDerivativeVariation BaseDerivativeVariation = CompiledPDV_FiniteDifferences;
-
-	if (bCompileForComputeShader)
-	{
-		LazyPrintf.PushParam(*GenerateFunctionCode(CompiledMP_EmissiveColorCS, BaseDerivativeVariation));
-	}
-	else
-	{
-		LazyPrintf.PushParam(TEXT("return 0"));
-	}
-
-	LazyPrintf.PushParam(*FString::Printf(TEXT("return %.5f"), Material->GetTranslucencyDirectionalLightingIntensity()));
-
-	LazyPrintf.PushParam(*FString::Printf(TEXT("return %.5f"), Material->GetTranslucentShadowDensityScale()));
-	LazyPrintf.PushParam(*FString::Printf(TEXT("return %.5f"), Material->GetTranslucentSelfShadowDensityScale()));
-	LazyPrintf.PushParam(*FString::Printf(TEXT("return %.5f"), Material->GetTranslucentSelfShadowSecondDensityScale()));
-	LazyPrintf.PushParam(*FString::Printf(TEXT("return %.5f"), Material->GetTranslucentSelfShadowSecondOpacity()));
-	LazyPrintf.PushParam(*FString::Printf(TEXT("return %.5f"), Material->GetTranslucentBackscatteringExponent()));
-
-	{
-		FLinearColor Extinction = Material->GetTranslucentMultipleScatteringExtinction();
-
-		LazyPrintf.PushParam(*FString::Printf(TEXT("return MaterialFloat3(%.5f, %.5f, %.5f)"), Extinction.R, Extinction.G, Extinction.B));
-	}
-
-	LazyPrintf.PushParam(*FString::Printf(TEXT("return %.5f"), Material->GetOpacityMaskClipValue()));
-
-	{
-		const FDisplacementScaling DisplacementScaling = Material->GetDisplacementScaling();
-		LazyPrintf.PushParam(*FString::Printf(TEXT("return %.5f"), FMath::Max(0.0f, DisplacementScaling.Magnitude)));
-		LazyPrintf.PushParam(*FString::Printf(TEXT("return %.5f"), FMath::Clamp(DisplacementScaling.Center, 0.0f, 1.0f)));
-	}
-
-	LazyPrintf.PushParam(!bEnableExecutionFlow ? *GenerateFunctionCode(MP_WorldPositionOffset, BaseDerivativeVariation) : TEXT("return Parameters.MaterialAttributes.WorldPositionOffset"));
-	LazyPrintf.PushParam(!bEnableExecutionFlow ? *GenerateFunctionCode(CompiledMP_PrevWorldPositionOffset, BaseDerivativeVariation) : TEXT("return 0.0f"));
-	LazyPrintf.PushParam(!bEnableExecutionFlow ? *GenerateFunctionCode(MP_CustomData0, BaseDerivativeVariation) : TEXT("return 0.0f"));
-	LazyPrintf.PushParam(!bEnableExecutionFlow ? *GenerateFunctionCode(MP_CustomData1, BaseDerivativeVariation) : TEXT("return 0.0f"));
-
-	// Print custom texture coordinate assignments, should be fine with regular derivatives
-	FString CustomUVAssignments;
-
-	int32 LastProperty = -1;
-	for (uint32 CustomUVIndex = 0; CustomUVIndex < NumUserTexCoords; CustomUVIndex++)
-	{
-		if (bEnableExecutionFlow)
-		{
-			const FString AttributeName = FMaterialAttributeDefinitionMap::GetAttributeName((EMaterialProperty)(MP_CustomizedUVs0 + CustomUVIndex));
-			CustomUVAssignments += FString::Printf(TEXT("\tOutTexCoords[%u] = Parameters.MaterialAttributes.%s;") LINE_TERMINATOR, CustomUVIndex, *AttributeName);
-		}
-		else
-		{
-		if (CustomUVIndex == 0)
-		{
-				CustomUVAssignments += DerivativeVariations[BaseDerivativeVariation].TranslatedCodeChunkDefinitions[MP_CustomizedUVs0 + CustomUVIndex];
-		}
-
-			if (DerivativeVariations[BaseDerivativeVariation].TranslatedCodeChunkDefinitions[MP_CustomizedUVs0 + CustomUVIndex].Len() > 0)
-			{
-				LastProperty = MP_CustomizedUVs0 + CustomUVIndex;
-			}
-			CustomUVAssignments += FString::Printf(TEXT("\tOutTexCoords[%u] = %s;") LINE_TERMINATOR, CustomUVIndex, *DerivativeVariations[BaseDerivativeVariation].TranslatedCodeChunks[MP_CustomizedUVs0 + CustomUVIndex]);
-		}
-	}
-
-	LazyPrintf.PushParam(*CustomUVAssignments);
-
-	// Print custom vertex shader interpolator assignments
-	FString CustomInterpolatorAssignments;
-
-	for (UMaterialExpressionVertexInterpolator* Interpolator : CustomVertexInterpolators)
-	{
-		if (Interpolator->InterpolatorOffset != INDEX_NONE)
-		{
-			check(Interpolator->InterpolatorIndex != INDEX_NONE);
-			check(Interpolator->InterpolatedType & MCT_Float);
-
-			const EMaterialValueType Type = Interpolator->InterpolatedType == MCT_Float ? MCT_Float1 : Interpolator->InterpolatedType;
-			const TCHAR* Swizzle[2] = { TEXT("x"), TEXT("y") };
-			const int32 Offset = Interpolator->InterpolatorOffset;
-			const int32 Index = Interpolator->InterpolatorIndex;
-
-			// Note: We reference the UV define directly to avoid having to pre-accumulate UV counts before property translation
-			CustomInterpolatorAssignments += FString::Printf(TEXT("\tOutTexCoords[VERTEX_INTERPOLATOR_%i_TEXCOORDS_X].%s = VertexInterpolator%i(Parameters).x;") LINE_TERMINATOR, Index, Swizzle[Offset%2], Index);
-				
-			if (Type >= MCT_Float2)
-			{
-				CustomInterpolatorAssignments += FString::Printf(TEXT("\tOutTexCoords[VERTEX_INTERPOLATOR_%i_TEXCOORDS_Y].%s = VertexInterpolator%i(Parameters).y;") LINE_TERMINATOR, Index, Swizzle[(Offset+1)%2], Index);
-
-				if (Type >= MCT_Float3)
-				{
-					CustomInterpolatorAssignments += FString::Printf(TEXT("\tOutTexCoords[VERTEX_INTERPOLATOR_%i_TEXCOORDS_Z].%s = VertexInterpolator%i(Parameters).z;") LINE_TERMINATOR, Index, Swizzle[(Offset+2)%2], Index);
-
-					if (Type == MCT_Float4)
-					{
-						CustomInterpolatorAssignments += FString::Printf(TEXT("\tOutTexCoords[VERTEX_INTERPOLATOR_%i_TEXCOORDS_W].%s = VertexInterpolator%i(Parameters).w;") LINE_TERMINATOR, Index, Swizzle[(Offset+3)%2], Index);
-					}
-				}
-			}
-		}
-	}
-
-	LazyPrintf.PushParam(*CustomInterpolatorAssignments);
-
-	if (bEnableExecutionFlow)
-	{
-		FString EvaluateVertexCode;
-
-		// Set default texcoords in the VS
-		for (uint32 TexCoordIndex = 0; TexCoordIndex < NumUserVertexTexCoords; ++TexCoordIndex)
-		{
-			const FString AttributeName = FMaterialAttributeDefinitionMap::GetAttributeName((EMaterialProperty)(MP_CustomizedUVs0 + TexCoordIndex));
-
-			EvaluateVertexCode += FString::Printf(TEXT("\tDefaultMaterialAttributes.%s = Parameters.TexCoords[%d];") LINE_TERMINATOR, *AttributeName, TexCoordIndex);
-		}
-
-		EvaluateVertexCode += TranslatedAttributesCodeChunks[SF_Vertex];
-
-		LazyPrintf.PushParam(*TranslatedAttributesCodeChunks[SF_Pixel]);
-
-		FString EvaluateMaterialAttributesCode = TEXT("    FMaterialAttributes MaterialAttributes = EvaluatePixelMaterialAttributes(Parameters);" LINE_TERMINATOR);
-
-		for (int32 PropertyIndex = 0; PropertyIndex < MP_MAX; ++PropertyIndex)
-		{
-			// Skip non-shared properties
-			if (!SharedPixelProperties[PropertyIndex])
-			{
-				continue;
-			}
-
-			const EMaterialProperty Property = (EMaterialProperty)PropertyIndex;
-			check(FMaterialAttributeDefinitionMap::GetShaderFrequency(Property) == SF_Pixel);
-			// Special case MP_SubsurfaceColor as the actual property is a combination of the color and the profile but we don't want to expose the profile
-			const FString PropertyName = FMaterialAttributeDefinitionMap::GetAttributeName(Property);
-
-			if (PropertyIndex == MP_SubsurfaceColor)
-			{
-				// TODO - properly handle subsurface profile
-				EvaluateMaterialAttributesCode += FString::Printf("    PixelMaterialInputs.Subsurface = float4(MaterialAttributes.%s, 0.0f);" LINE_TERMINATOR, *PropertyName);
-			}
-			else
-			{
-				EvaluateMaterialAttributesCode += FString::Printf("    PixelMaterialInputs.%s = MaterialAttributes.%s;" LINE_TERMINATOR, *PropertyName, *PropertyName);
-			}
-		}
-
-		// TODO - deriv
-		for (int32 Iter = 0; Iter < CompiledPDV_MAX; Iter++)
-		{
-			LazyPrintf.PushParam(*EvaluateMaterialAttributesCode);
-			LazyPrintf.PushParam(TEXT(""));
-			LazyPrintf.PushParam(TEXT(""));
-		}
-	}
-	else
-	{
-		// skip material attributes code
-		LazyPrintf.PushParam(TEXT(""));
+	COOK_STAT(FScopedDurationTimer DurationTimer(MaterialTranslatorCookStats::MaterialTranslateAndSourceGenTimeSec));
 
-		for (int32 Iter = 0; Iter < CompiledPDV_MAX; Iter++)
-		{
-			ECompiledPartialDerivativeVariation Variation = (ECompiledPartialDerivativeVariation)Iter;
-
-	// Initializers required for Normal
-			LazyPrintf.PushParam(*DerivativeVariations[Variation].TranslatedCodeChunkDefinitions[MP_Normal]);
-			LazyPrintf.PushParam(*NormalAssignment[Variation]);
-	// Finally the rest of common code followed by assignment into each input
-			LazyPrintf.PushParam(*PixelMembersSetupAndAssignments[Variation]);
-		}
-	}
-
-	LazyPrintf.PushParam(*FString::Printf(TEXT("%u"),MaterialTemplateLineNumber));
-
-	return LazyPrintf.GetResultString();
+	// use "/Engine/Private/MaterialTemplate.ush" to create the functions to get data (e.g. material attributes) and code (e.g. material expressions to create specular color) from C++
+	int32 MaterialTemplateLineNumber;
+	FStringTemplateResolver Resolver = FMaterialSourceTemplate::Get().BeginResolve(GetShaderPlatform(), &MaterialTemplateLineNumber);
+	MaterialSourceTemplateParams.Add({ TEXT("line_number"), FString::Printf(TEXT("%u"), MaterialTemplateLineNumber) });
+	Resolver.SetParameterMap(&MaterialSourceTemplateParams);
+	return Resolver.Finalize();
 }
 
 
@@ -11369,7 +11674,7 @@ static FString GetParametersStrataTreeName(EStrataCompilationContext Compilation
 void FHLSLMaterialTranslator::FStrataCompilationContext::StrataEvaluateSharedLocalBases(
 	FHLSLMaterialTranslator* Compiler,
 	uint8& OutRequestedSharedLocalBasesCount,
-	FShaderCompilerEnvironment* OutEnvironment)
+	FHLSLMaterialTranslatorEnvironmentDefines* OutEnvironment)
 {
 	/*
 	* The final output code/workflow for shared tangent basis should look like
@@ -11469,6 +11774,8 @@ void FHLSLMaterialTranslator::FStrataCompilationContext::StrataEvaluateSharedLoc
 	if (OutEnvironment)
 	{
 		// Now write out all the macros, them mapping from the BSDF to the effective position/index in the shared local basis array they should write to.
+		OutEnvironment->StrataDefineKeys.Reserve(CodeChunkToStrataSharedLocalBasis.Num());
+		OutEnvironment->StrataDefineValues.Reserve(CodeChunkToStrataSharedLocalBasis.Num());
 		for (TMultiMap<uint64, FStrataSharedLocalBasesInfo>::TConstIterator It(CodeChunkToStrataSharedLocalBasis); It; ++It)
 		{
 			// The default linear output index will be 0 by default, and different if in fact the shared local basis points to one that is effectively in used in the array of shared local bases.
@@ -11483,7 +11790,9 @@ void FHLSLMaterialTranslator::FStrataCompilationContext::StrataEvaluateSharedLoc
 				}
 			}
 
-			OutEnvironment->SetDefine(*Compiler->GetStrataSharedLocalBasisIndexMacroInner(It->Value.SharedData, CompilationContextIndex), LinearIndex);
+			FString DefineName = Compiler->GetStrataSharedLocalBasisIndexMacroInner(It->Value.SharedData, CompilationContextIndex);
+			OutEnvironment->StrataDefineKeys.Add(MoveTemp(DefineName));
+			OutEnvironment->StrataDefineValues.Add(LinearIndex);
 		}
 	}
 }
@@ -12286,8 +12595,7 @@ bool FHLSLMaterialTranslator::FStrataCompilationContext::StrataGenerateDerivedMa
 #endif // EDITOR_ONLY
 			}
 		}
-	}
-	while (!StrataSimplificationStatus.bMaterialFitsInMemoryBudget);
+	} while (!StrataSimplificationStatus.bMaterialFitsInMemoryBudget);
 
 	return true; // Success
 }
diff --git a/Engine/Source/Runtime/Engine/Private/Materials/HLSLMaterialTranslator.h b/Engine/Source/Runtime/Engine/Private/Materials/HLSLMaterialTranslator.h
--- a/Engine/Source/Runtime/Engine/Private/Materials/HLSLMaterialTranslator.h
+++ b/Engine/Source/Runtime/Engine/Private/Materials/HLSLMaterialTranslator.h
@@ -33,6 +33,7 @@
 #include "Containers/Map.h"
 #include "Shader/ShaderTypes.h"
 #include "SparseVolumeTexture/SparseVolumeTexture.h"
+#include "Runtime/RenderCore/Internal/ShaderCompilerDefinitions.h"
 
 #if WITH_EDITORONLY_DATA
 #include "Materials/MaterialExpressionSceneTexture.h"
@@ -53,6 +54,8 @@
 #include "HLSLMaterialDerivativeAutogen.h"
 #endif
 
+#include "HLSLMaterialTranslator.generated.h"
+
 class Error;
 
 /**
@@ -232,6 +235,277 @@ enum EStrataCompilationContext : uint8
 	SCC_MAX = 2u
 };
 
+
+/** Data structure used to cache a part of material translation results. It contains all the generated
+    defines that will be declared during the compilation of the generated material shader. */
+USTRUCT()
+struct FHLSLMaterialTranslatorEnvironmentDefines
+{
+	GENERATED_BODY()
+
+	UPROPERTY()
+	bool bNeedsParticlePosition;
+
+	UPROPERTY()
+	bool bNeedsParticleVelocity;
+
+	UPROPERTY()
+	bool bUseDynamicParameters;
+
+	UPROPERTY()
+	uint32 DynamicParametersMask;
+
+	UPROPERTY()
+	bool bNeedsParticleTime;
+
+	UPROPERTY()
+	bool bUsesParticleMotionBlur;
+
+	UPROPERTY()
+	bool bNeedsParticleRandom;
+
+	UPROPERTY()
+	bool bSphericalParticleOpacity;
+
+	UPROPERTY()
+	bool bUseParticleSubUVs;
+
+	UPROPERTY()
+	bool bLightmapUVAccess;
+
+	UPROPERTY()
+	bool bUsesAOMaterialMask;
+
+	UPROPERTY()
+	bool bUsesSpeedtree;
+
+	UPROPERTY()
+	bool bNeedsWorldPositionExcludingShaderOffsets;
+
+	UPROPERTY()
+	bool bNeedsParticleSize;
+
+	UPROPERTY()
+	bool bNeedsParticleSpriteRotation;
+
+	UPROPERTY()
+	bool bNeedsSceneTextures;
+
+	UPROPERTY()
+	bool bUsesEyeAdaptation;
+
+	UPROPERTY()
+	bool bVirtualTextureOutput;
+
+	UPROPERTY()
+	bool bUsesPerInstanceCustomData;
+
+	UPROPERTY()
+	bool bUsesPerInstanceRandom;
+
+	UPROPERTY()
+	bool bUsesPerInstanceFadeAmount;
+
+	UPROPERTY()
+	bool bUsesVertexInterpolator;
+
+	UPROPERTY()
+	bool bUsesSkyAtmosphere;
+
+	UPROPERTY()
+	bool bUsesVertexColor;
+
+	UPROPERTY()
+	bool bUsesParticleColor;
+
+	UPROPERTY()
+	bool bUsesParticleLocalToWorld;
+
+	UPROPERTY()
+	bool bUsesParticleWorldToLocal;
+
+	UPROPERTY()
+	bool bUsesInstanceLocalToWorldPS;
+
+	UPROPERTY()
+	bool bUsesInstanceWorldToLocalPS;
+
+	UPROPERTY()
+	bool bUsesTransformVector;
+
+	UPROPERTY()
+	bool bUsesPixelDepthOffset;
+
+	UPROPERTY()
+	bool bUsesWorldPositionOffset;
+
+	UPROPERTY()
+	bool bUsesDisplacement;
+
+	UPROPERTY()
+	bool bUsesEmissiveColor;
+
+	UPROPERTY()
+	bool bUsesDistortion;
+
+	UPROPERTY()
+	bool bDistortionAccountForCoverage;
+
+	UPROPERTY()
+	bool bMaterialEnableTranslucencyFogging;
+
+	UPROPERTY()
+	bool bMaterialEnableTranslucencyCloudFogging;
+
+	UPROPERTY()
+	bool bMaterialIsSky;
+
+	UPROPERTY()
+	bool bMaterialComputeFogPerPixel;
+
+	UPROPERTY()
+	bool bMaterialFullyRough;
+
+	UPROPERTY()
+	bool bMaterialUsesAnisotropy;
+
+	UPROPERTY()
+	uint8 MaterialDecalReadMask;
+
+	UPROPERTY()
+	bool bMaterialUsesDecalLookup;
+
+	UPROPERTY()
+	uint8 MaterialPathTracingBufferRead;
+
+	UPROPERTY()
+	uint32 NumVirtualTextureSamples;
+
+	UPROPERTY()
+	bool bMaterialVirtualTextureFeedback;
+
+	enum EVirtualPageType
+	{
+		TABLE0 = 0,
+		TABLE1 = 1,
+		TABLE_INDIRECTION = 2
+	};
+
+	UPROPERTY()
+	TArray<uint8> VirtualPageTypes;
+
+	UPROPERTY()
+	TArray<TObjectPtr<UMaterialParameterCollection>> ParameterCollections;
+
+	UPROPERTY()
+	bool bSingleLayerWaterShadingQuality;
+
+	UPROPERTY()
+	bool bShadingModelsIsLit;
+
+	UPROPERTY()
+	uint32 MaterialShadingModelEnabled;
+
+	UPROPERTY()
+	bool bMaterialSubsurfaceProfileUseCurvature;
+
+	UPROPERTY()
+	bool bMaterialShadingModelEyeUseCurvature;
+
+	UPROPERTY()
+	bool bDisableForwardLocalLights;
+
+	UPROPERTY()
+	bool bSingleLayerWaterSeparatedMainLight;
+
+	UPROPERTY()
+	bool bMaterialSingleShadingModel;
+
+	UPROPERTY()
+	bool bMaterialVolumetricAdvanced;
+
+	UPROPERTY()
+	bool bMaterialVolumetricAdvancedPhasePerSample;
+
+	UPROPERTY()
+	bool bMaterialVolumetricAdvancedGreyscaleMaterial;
+
+	UPROPERTY()
+	bool bMaterialVolumetricAdvancedRaymarchVolumeShadow;
+
+	UPROPERTY()
+	bool bMaterialVolumetricAdvancedClampMultiscatteringContribution;
+
+	UPROPERTY()
+	uint32 MaterialVolumetricAdvancedMultiscatteringOctaveCount;
+
+	UPROPERTY()
+	bool bMaterialVolumetricAdvancedConservativeDensity;
+
+	UPROPERTY()
+	bool bMaterialVolumetricAdvancedOverrideAmbientOcclusion;
+
+	UPROPERTY()
+	bool bMaterialVolumetricAdvancedAdvancedGroundContribution;
+
+	UPROPERTY()
+	bool bMaterialVolumetricCloudEmptySpaceSkippingOutput;
+
+	UPROPERTY()
+	bool bMaterialIsStrata;
+
+	UPROPERTY()
+	bool bDualSourceColorBlendingEnabled;
+
+	UPROPERTY()
+	bool bStrataPremultipliedAlphaOpacityOverridden;
+
+	UPROPERTY()
+	bool bStrataUsesConversionFromLegacy;
+
+	UPROPERTY()
+	bool bStrataMaterialOutputOpaqueRoughRefractions;
+
+	UPROPERTY()
+	int32 StrataMaterialExportType;
+
+	UPROPERTY()
+	int32 StrataMaterialExportContext;
+
+	UPROPERTY()
+	int32 StrataMaterialExportLegacyBlendMode;
+
+	UPROPERTY()
+	bool bStrataOptimizedUnlit;
+
+	UPROPERTY()
+	bool bStrataSinglePath;
+
+	UPROPERTY()
+	bool bStrataFastPath;
+
+	UPROPERTY()
+	uint32 StrataClampedBSDFCount;
+
+	UPROPERTY()
+	bool bStrataComplexSpecialPath;
+
+	UPROPERTY()
+	bool bTextureSampleDebug;
+
+	UPROPERTY()
+	TArray<FString> StrataDefineKeys;
+
+	UPROPERTY()
+	TArray<int32> StrataDefineValues;
+	
+	bool HasShadingModel(EMaterialShadingModel model) const
+	{
+		return (MaterialShadingModelEnabled & (1 << model)) != 0;
+	}
+};
+
+
 class FHLSLMaterialTranslator : public FMaterialCompiler
 {
 	friend class FMaterialDerivativeAutogen;
@@ -278,12 +552,6 @@ protected:
 
 	uint64 MaterialAttributesReturned[SF_NumFrequencies];
 
-	/** Line number of the #line in MaterialTemplate.usf */
-	int32 MaterialTemplateLineNumber;
-
-	/** Contents of the MaterialTemplate.usf file */
-	FString MaterialTemplate;
-
 	TArray<int32> ScopeStack;
 
 	TArray<int32> ReferencedCodeChunks;
@@ -526,7 +794,7 @@ protected:
 
 		bool StrataGenerateDerivedMaterialOperatorData(FHLSLMaterialTranslator* Compiler);
 
-		void StrataEvaluateSharedLocalBases(FHLSLMaterialTranslator* Compiler, uint8& RequestedSharedLocalBasesCount, FShaderCompilerEnvironment* OutEnvironment);
+		void StrataEvaluateSharedLocalBases(FHLSLMaterialTranslator* Compiler, uint8& RequestedSharedLocalBasesCount, FHLSLMaterialTranslatorEnvironmentDefines* OutEnvironment);
 
 		FStrataSharedLocalBasesInfo StrataCompilationInfoGetMatchingSharedLocalBasisInfo(const FStrataRegisteredSharedLocalBasis& SearchedSharedLocalBasis);
 
@@ -1323,14 +1591,52 @@ protected:
 	/**Experimental access to the EyeAdaptation RT for applying an inverse. */
 	virtual int32 EyeAdaptationInverse(int32 LightValueArg, int32 AlphaArg) override;
 
-	// to only have one piece of code dealing with error handling if the Primitive constant buffer is not used.
-	// @param Name e.g. TEXT("ObjectWorldPositionAndRadius.w")
+	/**
+	 * To only have one piece of code dealing with error handling if the Primitive constant buffer is not used.
+	 * @param Name e.g. TEXT("ObjectWorldPositionAndRadius.w")
+	 */
 	int32 GetPrimitiveProperty(EMaterialValueType Type, const TCHAR* ExpressionName, const TCHAR* HLSLName);
 
-	// The compiler can run in a different state and this affects caching of sub expression, Expressions are different (e.g. View.PrevWorldViewOrigin) when using previous frame's values
+	/**
+	 * The compiler can run in a different state and this affects caching of sub expression, Expressions are different(e.g.View.PrevWorldViewOrigin) when using previous frame's values.
+	 */
 	virtual bool IsCurrentlyCompilingForPreviousFrame() const;
 
 	virtual bool IsDevelopmentFeatureEnabled(const FName& FeatureName) const override;
+
+	/**
+	 * Creates the key hash for this material translation request.
+	 */
+	FIoHash ComputeMaterialTranslationDDCKeyHash();
+
+	/**
+	 * Queries the DDC cache for a cached translation.
+	 * @param DDCKeyHash out the material key hash used to query the DDC.
+	 * @return Whether the speecified key is in the DDC.
+	 */
+	bool QueryDDCCachedTranslationResults(const FIoHash& DDCKeyHash);
+	
+	/**
+	 * Prepares the Environment Defines array based on compilation results.
+	 */
+	void PrepareEnvironmentDefines();
+	
+	/**
+	 * Prepares the material source generation parameters.
+	 */
+	void PrepareMaterialSourceStringParameters();
+	
+	/**
+	 * Pushes the final results to the DDC cache.
+	 * @param DDCKeyHash the hash of the material key use to push the results. Returned by QueryDDCCachedTranslationResults().
+	 */
+	void PushResultsToDDCCache(const FIoHash& DDCKeyHash);
+
+	/** The output material shader defines */
+	FHLSLMaterialTranslatorEnvironmentDefines EnvironmentDefines;
+	
+	/** The material shader source template parameters */
+	TMap<FString, FString> MaterialSourceTemplateParams;
 };
 
 #endif // WITH_EDITORONLY_DATA
diff --git a/Engine/Source/Runtime/Engine/Private/Materials/MaterialSourceTemplate.cpp b/Engine/Source/Runtime/Engine/Private/Materials/MaterialSourceTemplate.cpp
new file mode 100644
--- /dev/null
+++ b/Engine/Source/Runtime/Engine/Private/Materials/MaterialSourceTemplate.cpp
@@ -0,0 +1,108 @@
+// Copyright Epic Games, Inc. All Rights Reserved.
+
+#pragma once
+
+#if WITH_EDITOR
+
+#include "MaterialSourceTemplate.h"
+#include "Misc/ScopeRWLock.h"
+
+FMaterialSourceTemplate& FMaterialSourceTemplate::Get()
+{
+	static FMaterialSourceTemplate Instance;
+	return Instance;
+}
+
+FStringTemplateResolver FMaterialSourceTemplate::BeginResolve(EShaderPlatform ShaderPlatform, int32* MaterialTemplateLineNumber)
+{
+	Preload(ShaderPlatform);
+
+	if (MaterialTemplateLineNumber)
+	{
+		*MaterialTemplateLineNumber = MaterialTemplateLineNumbers[ShaderPlatform];
+	}
+
+	return { Templates[ShaderPlatform], 50 * 1024 };
+}
+
+const FSHA1& FMaterialSourceTemplate::GetTemplateHash(EShaderPlatform ShaderPlatform)
+{
+	Preload(ShaderPlatform);
+	return TemplateHash[ShaderPlatform];
+}
+
+bool FMaterialSourceTemplate::Preload(EShaderPlatform ShaderPlatform)
+{
+	// Is the material source template already loaded?
+	{
+		FReadScopeLock Lock{ RWLocks[ShaderPlatform] };
+		if (!Templates[ShaderPlatform].GetTemplateString().IsEmpty())
+		{
+			return true;
+		}
+	}
+
+	// Material source template not yet loaded. Acquire a write lock an try again.
+	FWriteScopeLock Lock{ RWLocks[ShaderPlatform] };
+	if (!Templates[ShaderPlatform].GetTemplateString().IsEmpty())
+	{
+		return true;
+	}
+
+	FString MaterialTemplateString;
+	LoadShaderSourceFileChecked(TEXT("/Engine/Private/MaterialTemplate.ush"), ShaderPlatform, MaterialTemplateString);
+
+	// Find the string index of the '#line' statement in MaterialTemplate.usf
+	const int32 LineIndex = MaterialTemplateString.Find(TEXT("#line"), ESearchCase::CaseSensitive);
+	check(LineIndex != INDEX_NONE);
+
+	// Count line endings before the '#line' statement
+	int32 TemplateLineNumber = INDEX_NONE;
+	int32 StartPosition = LineIndex + 1;
+	do
+	{
+		TemplateLineNumber++;
+		// Using \n instead of LINE_TERMINATOR as not all of the lines are terminated consistently
+		// Subtract one from the last found line ending index to make sure we skip over it
+		StartPosition = MaterialTemplateString.Find(TEXT("\n"), ESearchCase::CaseSensitive, ESearchDir::FromEnd, StartPosition - 1);
+	} while (StartPosition != INDEX_NONE);
+
+	check(TemplateLineNumber != INDEX_NONE);
+
+	// At this point MaterialTemplateLineNumber is one less than the line number of the '#line' statement
+	// For some reason we have to add 2 more to the #line value to get correct error line numbers from D3DXCompileShader
+	TemplateLineNumber += 3;
+
+	// Save the material template line numbers for this shader platform
+	MaterialTemplateLineNumbers[ShaderPlatform] = TemplateLineNumber;
+
+	// Load the material string template
+	FStringTemplate::FErrorInfo ErrorInfo;
+	if (!Templates[ShaderPlatform].Load(MoveTemp(MaterialTemplateString), ErrorInfo))
+	{
+		UE_LOG(LogMaterial, Error, TEXT("Error in MaterialTemplate.ush source template at line %d offset %d: %s"), ErrorInfo.Line, ErrorInfo.Offset, ErrorInfo.Message.GetData());
+		return false;
+	}
+
+	// Extract the material template string TemplateVersion parameter
+	const TStringView TemplateVersionKeyword = TEXT("$TemplateVersion{");
+	int Begin = MaterialTemplateString.Find(TemplateVersionKeyword.GetData());
+	int End = MaterialTemplateString.Find(TEXT("}"), ESearchCase::CaseSensitive, ESearchDir::FromStart, Begin + TemplateVersionKeyword.Len());
+	if (Begin > 0 && End > 0)
+	{
+		Begin += TemplateVersionKeyword.Len();
+		TemplateHash[ShaderPlatform].UpdateWithString(*MaterialTemplateString + Begin, End - Begin);
+	}
+
+	TArray<FStringView> Parameters;
+	Templates[ShaderPlatform].GetParameters(Parameters);
+	Parameters.Sort();
+	for (const FStringView& Param : Parameters)
+	{
+		TemplateHash[ShaderPlatform].UpdateWithString(Param.GetData(), Param.Len());
+	}
+
+	return true;
+}
+
+#endif // WITH_EDITOR
diff --git a/Engine/Source/Runtime/Engine/Private/Materials/MaterialSourceTemplate.h b/Engine/Source/Runtime/Engine/Private/Materials/MaterialSourceTemplate.h
new file mode 100644
--- /dev/null
+++ b/Engine/Source/Runtime/Engine/Private/Materials/MaterialSourceTemplate.h
@@ -0,0 +1,44 @@
+// Copyright Epic Games, Inc. All Rights Reserved.
+
+#pragma once
+
+#include "Materials/StringTemplate.h"
+#include "HAL/CriticalSection.h"
+#include "RHI.h"
+
+#if WITH_EDITOR
+
+/** Cache of material templates to be used for material source generation. */
+struct FMaterialSourceTemplate
+{
+	/** Returns the singleton instance */
+	static FMaterialSourceTemplate& Get();
+
+	/** Creates and returns a string template resolver for the specified shader platform.
+	 * The appropriate template is loaded for the first time if it wasn't loaded before.*/
+	FStringTemplateResolver BeginResolve(EShaderPlatform ShaderPlatform, int32* MaterialTemplateLineNumber = nullptr);
+
+	/** Return the hash of the material template for specified shader platform 
+	  * The information hashed is the string contained within "$TemplateVersion{...}" and the name of the parameters
+	  * contained in the template string (insensitive to ordering).
+	  */
+	const FSHA1& GetTemplateHash(EShaderPlatform ShaderPlatform);
+
+	/** Loads a material source template if still unloaded. */
+	bool Preload(EShaderPlatform ShaderPlatform);
+
+	/** Cached templates per shader platform */
+	FStringTemplate Templates[SP_NumPlatforms];
+
+	/** Hash of the template parameters and TemplateVersion parameter extracted from file */
+	FSHA1 TemplateHash[SP_NumPlatforms];
+
+	/** Cached material template line numbers per shader platform */
+	int32 MaterialTemplateLineNumbers[SP_NumPlatforms];
+
+	/** RW locks for thread safe access */
+	FRWLock RWLocks[SP_NumPlatforms];
+
+};
+
+#endif // WITH_EDITOR
diff --git a/Engine/Source/Runtime/Engine/Private/Materials/StringTemplate.cpp b/Engine/Source/Runtime/Engine/Private/Materials/StringTemplate.cpp
new file mode 100644
--- /dev/null
+++ b/Engine/Source/Runtime/Engine/Private/Materials/StringTemplate.cpp
@@ -0,0 +1,211 @@
+// Copyright Epic Games, Inc. All Rights Reserved.
+
+#include "StringTemplate.h"
+
+static FStringView Substring(const FString& String, int Begin, int Length)
+{
+	return { GetData(String) + Begin, Length};
+}
+
+FStringTemplate::FStringTemplate()
+{
+}
+
+bool FStringTemplate::Load(FString TemplateStringArg, FErrorInfo& OutErrorinfo)
+{
+	TArray<FChunk> NewChunks;
+	int32 ChunkBegin = 0;
+	int32 ChunkLength = 0;
+	int32 Offset = 0;
+	int32 Line = 1;
+	TCHAR Char = 0;
+	NumNamedParameters = 0;
+
+	OutErrorinfo = {};
+
+	// Local function that pushes the lookahead character to the current chunk
+	// and reads the next in the string. If at string end, Char is set to 0.
+	auto Next = [&]() -> TCHAR
+	{
+		++Offset;
+		++ChunkLength;
+		if (Offset < TemplateStringArg.Len())
+		{
+			Char = TemplateStringArg[Offset];
+		}
+		else
+		{
+			Char = 0;
+		}
+		return Char;
+	};
+
+	// Checks that the lookahead characters is `C` and if so, pushes it to the current
+	// chunk and reads the next character. Otherwise it returns false.
+	auto Accept = [&](TCHAR C) -> bool
+	{
+		if (Char == C)
+		{
+			Next();
+			return true;
+		}
+		return false;
+	};
+
+	// Pushes the current chunk substring to the list of chunks and resets the
+	// current chunk substring coordinates.
+	auto PushChunk = [&](bool bIsParameter)
+	{
+		NewChunks.Add({ Substring(TemplateStringArg, ChunkBegin, ChunkLength), bIsParameter });
+		ChunkBegin = Offset;
+		ChunkLength = 0;
+	};
+
+	// Read the first character
+	Char = TemplateStringArg.IsEmpty() ? 0 : TemplateStringArg[0];
+
+	// Scan the entire string...
+	while (Char != 0)
+	{
+		// Record token line and offset in the error info
+		OutErrorinfo.Line = Line;
+		OutErrorinfo.Offset = Offset;
+
+		// Check whether this is a parameter
+		if (Accept('%'))
+		{
+			if (Accept('s')) // e.g. %s
+			{
+				ChunkLength -= 2; // take out the "%s"
+
+				// Push the chunk of text read so far
+				if (ChunkLength > 0)
+				{
+					PushChunk(false);
+				}
+
+				// Push an nameless parameter chunk
+				NewChunks.Add({ TEXT(""), true});
+			}
+			else if (Accept('{')) // e.g. ${my_parameter}
+			{
+				ChunkLength -= 2; // take out the "%{"
+
+				// Push the chunk of text read so far
+				if (ChunkLength > 0)
+				{
+					PushChunk(false);
+				}
+
+				// Parse the parameter name
+				while (Char != 0 && Char != '}')
+				{
+					Next();
+				}
+
+				// Expect the closing parameter name bracket
+				if (!Accept('}'))
+				{
+					OutErrorinfo.Message = TEXT("missing '}' in named parameter");
+					return false;
+				}
+
+				ChunkLength -= 1; // take out the "}"
+
+				// Push the parameter name as the next chunk
+				PushChunk(true);
+				NumNamedParameters += 1;
+			}
+			else
+			{
+				Accept('%'); // e.g. %%
+				NewChunks.Add({ TEXT("%"), false });
+			}
+		}
+		else
+		{
+			Line += (Char == '\n');
+
+			// Not a parameter ahead, push the character to the current chunk
+			Next();
+		}
+	}
+
+	// If there's a trailing chunk of text, push the last chunk
+	if (ChunkLength > 0)
+	{
+		PushChunk(false);
+	}
+
+	// Parsing complete and successful
+	TemplateString = MoveTemp(TemplateStringArg);
+	Chunks = MoveTemp(NewChunks);
+
+	return true;
+
+}
+
+void FStringTemplate::GetParameters(TArray<FStringView>& OutParams) const
+{
+	OutParams.Reserve(NumNamedParameters);
+	for (const FChunk& Chunk : Chunks)
+	{
+		if (Chunk.bIsParameter && Chunk.Text.Len() > 0)
+		{
+			OutParams.Push(Chunk.Text);
+		}
+	}
+}
+
+FStringTemplateResolver::FStringTemplateResolver(const FStringTemplate& Template, uint32 ResolvedStringSizeHint)
+: Template{ Template }
+{
+	ResolvedString.Reserve(ResolvedStringSizeHint);
+}
+
+void FStringTemplateResolver::Advance(FStringView NextNamelessParameterValue)
+{
+	// Process all chunks until a nameless parameter is encountered (inclusively)
+	for (; ChunkIndex < Template.Chunks.Num(); ++ChunkIndex)
+	{
+		const FStringTemplate::FChunk& Chunk = Template.Chunks[ChunkIndex];
+		if (Chunk.bIsParameter)
+		{
+			if (Chunk.Text.Len() > 0)
+			{
+				// This is a named parameter, resolve it.
+				const FString* ParameterValue = NamedParameters ? NamedParameters->FindByHash(GetTypeHash(Chunk.Text), Chunk.Text) : nullptr;
+				if (ParameterValue)
+				{
+					ResolvedString += *ParameterValue;
+				}
+			}
+			else
+			{
+				// This is a nameless parameter
+				ResolvedString += NextNamelessParameterValue;
+				++ChunkIndex;
+
+				// Quit here, the rest of the chunks will be resolved by further Advance() calls and by Finalize()
+				return;
+			}
+		}
+		else
+		{
+			// Push the static chunk of text
+			ResolvedString += Chunk.Text;
+		}
+	}
+}
+
+FString FStringTemplateResolver::Finalize(uint32 NewResolvedStringSizeHint)
+{
+	// Finish processing the remaining chunks, using an empty string for the remaining nameless parameters.
+	while (ChunkIndex < Template.Chunks.Num())
+	{
+		Advance(TEXT(""));
+	}
+
+	ChunkIndex = 0;
+	return MoveTemp(ResolvedString);
+}
diff --git a/Engine/Source/Runtime/Engine/Private/Materials/StringTemplate.h b/Engine/Source/Runtime/Engine/Private/Materials/StringTemplate.h
new file mode 100644
--- /dev/null
+++ b/Engine/Source/Runtime/Engine/Private/Materials/StringTemplate.h
@@ -0,0 +1,155 @@
+// Copyright Epic Games, Inc. All Rights Reserved.
+
+#pragma once
+
+#include "CoreTypes.h"
+#include "Misc/AssertionMacros.h"
+#include "Containers/UnrealString.h"
+
+/**
+ * A string template is a utility class used to generate parameterized strings.
+ * To use it, initialize the template by loading a template string with the
+ * constructor or manually calling Load() after construction.
+ * The Load() function will parse the template string and split it into chunks,
+ * where each chunk is either a static span of text or a parameter. After a
+ * template has been loaded, resolve the full string using an instance of
+ * FStringTemplateResolver.
+ * 
+ * Parameters can be specifed *named* and *nameless* depending on your needs.
+ * A named parameter can be specified with the syntax `%{parameter_name}` while
+ * a nameless parameter with `%s` (without backquotes). A string can include both
+ * type of parameters. Keep in mind that nameless parameters will be resolved in
+ * the order they appear in the template string.
+ * 
+ * Here is an example of a valid template string:
+ *
+ *		The last president of the USA is %{name_of_president}.
+ *		He was born in %s.
+ * 
+ * Finally, use may force the emission of the '%' character using the "%%" escape sequence.
+ */
+class FStringTemplate
+{
+public:
+	/**
+	 * Wraps info about where an error in the template string is.
+	 */
+	struct FErrorInfo
+	{
+		/** A user friendly textual description of the error */
+		FStringView Message;
+
+		/** Which line contains the error */
+		int32 Line;
+
+		/** At what offset in characters from the start */
+		int32 Offset;
+	};
+
+public:
+	/**
+	 * Constructs an empty string template.
+	 */
+	FStringTemplate();
+
+	/**
+	 * Loads the specified template string for future string resolutions with Resolve().
+	 * Returns: whether parsing the string was completed successfully.
+	 * Note: `TemplateString` is passed by value as it will be stored within this instance.
+	 *        Consider using MoveTemp() to pass ownership of the template string if you
+	 *        no longer need it.
+	 */
+	bool Load(FString TemplateString, FErrorInfo& Errorinfo);
+
+	/**
+	 * Returns the template string.
+	 */
+	const FString& GetTemplateString() const { return TemplateString; }
+
+	/**
+	 * Returns the array of named parameters contained in this template.
+	 */
+	void GetParameters(TArray<FStringView>& OutParams) const;
+
+private:
+	/** Information about a template string chunk (a static text or a parameter) */
+	struct FChunk
+	{
+		/** Static text or name of the parameter */
+		FStringView Text;
+
+		/** Whether this is a parameter */
+		bool bIsParameter;
+	};
+
+	/** The loaded and parsed template string */
+	FString TemplateString;
+
+	/** Chunks making out the loaded template string */
+	TArray<FChunk> Chunks;
+
+	/** The number of named parameters in this template string */
+	int32 NumNamedParameters = 0;
+
+	friend class FStringTemplateResolver;
+};
+
+/**
+ * It provides the mechanism to resolve a FStringTemplate instance into a string.
+ * 
+ * To resolve a template into a string, instantiate this class then set all the
+ * named parameters with SetParameterMap() _before_ calling Advance() or Finalize().
+ * 
+ * Once all named parameters have been set (or if you wish to set no named parameters),
+ * call Advance() repeatedly, passing in the value of the next nameless parameter. Call
+ * Advance() for each nameless parameter you wish to replace.
+ * 
+ * When you have no more nameless parameters to use or if you have none at all, call
+ * Finalize() to complete resolving the string which will be returned by it.
+ */
+class FStringTemplateResolver
+{
+public:
+	/**
+	 * Constructs the resolver using specified `Template`.
+	 * You can optionally provide a value for `ResolvedStringSizeHint` to be used to reserve
+	 * space in the resolved string.
+	 */
+	FStringTemplateResolver(const FStringTemplate& Template, uint32 ResolvedStringSizeHint = 0);
+
+	/**
+	 * Sets the map of named parameters.
+	 * Note: make sure all named parameters are set before calling Advance() or Finalize() otherwise
+	 *       you may get an incorrect resolved string.
+	 */
+	void SetParameterMap(const TMap<FString, FString>* Parameters)
+	{
+		NamedParameters = Parameters;
+	}
+
+	/** 
+	 * Generates the next section of the string until a nameless parameter is encountered for which
+	 * `NextNamelessParameterValue` is used as the replaced value. Call Advance() for each nameless
+	 * parameter you need to push.
+	 */
+	void Advance(FStringView NextNamelessParameterValue);
+
+	/**
+	 * Completes generating the string and returns it.
+	 * Note: Any subsequent and unresolved nameless parameter will be replaced to the empty string.
+	 */
+	FString Finalize(uint32 ResolvedStringSizeHint = 0);
+	
+private:
+	/** The template to use */
+	const FStringTemplate& Template;
+
+	/** Index of the next chunk to process */
+	int32 ChunkIndex = 0;
+
+	/** Map of named parameters key-values */
+	const TMap<FString, FString>* NamedParameters = nullptr;
+
+	/** The resolved string being generated */
+	FString ResolvedString;
+};
diff --git a/Engine/Source/Runtime/RenderCore/Internal/ShaderCompilerDefinitions.h b/Engine/Source/Runtime/RenderCore/Internal/ShaderCompilerDefinitions.h
--- a/Engine/Source/Runtime/RenderCore/Internal/ShaderCompilerDefinitions.h
+++ b/Engine/Source/Runtime/RenderCore/Internal/ShaderCompilerDefinitions.h
@@ -6,6 +6,8 @@
 // as part of hiding FShaderCompilerDefinitions from public view, and will need to add this include back in.
 // #include "ShaderCore.h"
 
+#define SHADER_COMPILER_FLOAT32_FORMAT_STRING TEXT("%#.9gf")
+
 enum class EShaderCompilerDefineVariant : uint8
 {
 	None = 0,
@@ -280,7 +282,7 @@ public:
 			else if (Defines.ValueTypes[Index] == EShaderCompilerDefineVariant::Float)
 			{
 				// Make sure the printed value perfectly matches the given number
-				FCString::Sprintf(ValueStringBuffer, TEXT("%#.9gf"), Defines.Pairs[Index].ValueFloat);
+				FCString::Sprintf(ValueStringBuffer, SHADER_COMPILER_FLOAT32_FORMAT_STRING, Defines.Pairs[Index].ValueFloat);
 				Result = ValueStringBuffer;
 			}
 			else
diff --git a/Engine/Source/Runtime/RenderCore/Public/ShaderCore.h b/Engine/Source/Runtime/RenderCore/Public/ShaderCore.h
--- a/Engine/Source/Runtime/RenderCore/Public/ShaderCore.h
+++ b/Engine/Source/Runtime/RenderCore/Public/ShaderCore.h
@@ -75,7 +75,7 @@ DECLARE_FLOAT_ACCUMULATOR_STAT_EXTERN(TEXT("Total Global Shader Compiling Time")
 DECLARE_FLOAT_ACCUMULATOR_STAT_EXTERN(TEXT("RHI Compile Time"),STAT_ShaderCompiling_RHI,STATGROUP_ShaderCompiling, RENDERCORE_API);
 DECLARE_FLOAT_ACCUMULATOR_STAT_EXTERN(TEXT("Loading Shader Files"),STAT_ShaderCompiling_LoadingShaderFiles,STATGROUP_ShaderCompiling, RENDERCORE_API);
 DECLARE_FLOAT_ACCUMULATOR_STAT_EXTERN(TEXT("CRCing Shader Files"),STAT_ShaderCompiling_HashingShaderFiles,STATGROUP_ShaderCompiling, RENDERCORE_API);
-DECLARE_FLOAT_ACCUMULATOR_STAT_EXTERN(TEXT("HLSL Translation"),STAT_ShaderCompiling_HLSLTranslation,STATGROUP_ShaderCompiling, RENDERCORE_API);
+DECLARE_FLOAT_ACCUMULATOR_STAT_EXTERN(TEXT("HLSL Translation"), STAT_ShaderCompiling_HLSLTranslation, STATGROUP_ShaderCompiling, RENDERCORE_API);
 DECLARE_FLOAT_ACCUMULATOR_STAT_EXTERN(TEXT("DDC Loading"),STAT_ShaderCompiling_DDCLoading,STATGROUP_ShaderCompiling, RENDERCORE_API);
 DECLARE_FLOAT_ACCUMULATOR_STAT_EXTERN(TEXT("Material Loading"),STAT_ShaderCompiling_MaterialLoading,STATGROUP_ShaderCompiling, RENDERCORE_API);
 DECLARE_FLOAT_ACCUMULATOR_STAT_EXTERN(TEXT("Material Compiling"),STAT_ShaderCompiling_MaterialCompiling,STATGROUP_ShaderCompiling, RENDERCORE_API);
@@ -551,6 +551,12 @@ struct FShaderCompilerEnvironment
 	{
 		CompileArgs.Add(Name, TVariant<bool, float, int32, uint32, FString>(TInPlaceType<ValueType>(), Value));
 	}
+	
+	// Like above, but this overload takes in the define value variant explicitly.
+	void SetCompileArgument(const TCHAR* Name, TVariant<bool, float, int32, uint32, FString> Value)
+	{
+		CompileArgs.Add(Name, MoveTempIfPossible(Value));
+	}
 
 	// Helper to set both a define and a compile argument to the same value. Useful for various parameters which
 	// need to be consumed both by preprocessing and in the shader format backends to modify compilation behaviour.
