PLAY-11038 - Angular replication fixes and improvements

Replication jitter:
Change1: Average out receive rate so that fluctuating send-rate and network conditions doesn't affect replication.
Change2: Redesign the angular velocity replication to adjust the target angular velocity based on the angular diff, which is a more stable approach than adjusting the the current angular velocity.
Change3: Only perform snap reconciliation if velocity is low enough, else let velocity replication take its course.

diff --git a/Engine/Source/Runtime/Engine/Private/PhysicsEngine/PhysicsReplication.cpp b/Engine/Source/Runtime/Engine/Private/PhysicsEngine/PhysicsReplication.cpp
--- a/Engine/Source/Runtime/Engine/Private/PhysicsEngine/PhysicsReplication.cpp
+++ b/Engine/Source/Runtime/Engine/Private/PhysicsEngine/PhysicsReplication.cpp
@@ -97,18 +97,36 @@ namespace PhysicsReplicationCVars
 
 		static float InterpolationTimeMultiplier = 1.1f;
 		static FAutoConsoleVariableRef CVarInterpolationTimeMultiplier(TEXT("np2.PredictiveInterpolation.InterpolationTimeMultiplier"), InterpolationTimeMultiplier, TEXT("Multiplier to adjust the replication interpolation time which is based on the sendrate of replication data from the server."));
+		
+		static float AverageReceiveIntervalSmoothing = 3.0f;
+		static FAutoConsoleVariableRef CVarAverageReceiveIntervalSmoothing(TEXT("np2.PredictiveInterpolation.AverageReceiveIntervalSmoothing"), AverageReceiveIntervalSmoothing, TEXT("Recommended range: 1.0 - 5.0. Higher value makes the average receive interval adjust itself slower, reducing spikes in InterpolationTime."));
 
-		static float ExtrapolationTimeMultiplier = 1.5f;
+		static float ExtrapolationTimeMultiplier = 3.0f;
 		static FAutoConsoleVariableRef CVarExtrapolationTimeMultiplier(TEXT("np2.PredictiveInterpolation.ExtrapolationTimeMultiplier"), ExtrapolationTimeMultiplier, TEXT("Multiplier to adjust the time to extrapolate the target forward over, the time is based on current send-rate."));
 
-		static float MinExpectedDistanceCovered = 0.25f;
+		static float MinExpectedDistanceCovered = 0.5f;
 		static FAutoConsoleVariableRef CVarMinExpectedDistanceCovered(TEXT("np2.PredictiveInterpolation.MinExpectedDistanceCovered"), MinExpectedDistanceCovered, TEXT("Value between 0-1, in percentage where 0.25 = 25%. How much of the expected distance based on replication velocity should the object have covered in a simulation tick to Not be considered stuck."));
 
 		static float ErrorAccumulationDecreaseMultiplier = 0.5f;
 		static FAutoConsoleVariableRef CVarErrorAccumulationDecreaseMultiplier(TEXT("np2.PredictiveInterpolation.ErrorAccumulationDecreaseMultiplier"), ErrorAccumulationDecreaseMultiplier, TEXT("Multiplier to adjust how fast we decrease accumulated error time when we no longer accumulate error."));
 
-		static float ErrorAccumulationSeconds = 0.5f;
+		static float ErrorAccumulationSeconds = 3.0f;
 		static FAutoConsoleVariableRef CVarErrorAccumulationSeconds(TEXT("np2.PredictiveInterpolation.ErrorAccumulationSeconds"), ErrorAccumulationSeconds, TEXT("Perform a reposition if replication have not been able to cover the min expected distance towards the target for this amount of time."));
+		
+		static bool DisableErrorVelocityLimits = false;
+		static FAutoConsoleVariableRef CVarDisableErrorVelocityLimits(TEXT("np2.PredictiveInterpolation.DisableErrorVelocityLimits"), DisableErrorVelocityLimits, TEXT("Disable the velocity limit and allow error accumulation at any velocity."));
+		
+		static float ErrorAccLinVelMaxLimit = 50.0f;
+		static FAutoConsoleVariableRef CVarErrorAccLinVelMaxLimit(TEXT("np2.PredictiveInterpolation.ErrorAccLinVelMaxLimit"), ErrorAccLinVelMaxLimit, TEXT("If target velocity is below this limit we check for desync to trigger softsnap and accumulate time to build up to a hardsnap."));
+		
+		static float ErrorAccAngVelMaxLimit = 1.5f;
+		static FAutoConsoleVariableRef CVarErrorAccAngVelMaxLimit(TEXT("np2.PredictiveInterpolation.ErrorAccAngVelMaxLimit"), ErrorAccAngVelMaxLimit, TEXT("If target angular velocity (in degrees) is below this limit we check for desync to trigger softsnap and accumulate time to build up to a hardsnap."));
+		
+		static float SoftSnapPosStrength = 0.5f;
+		static FAutoConsoleVariableRef CVarSoftSnapPosStrength(TEXT("np2.PredictiveInterpolation.SoftSnapPosStrength"), SoftSnapPosStrength, TEXT("Value in percent between 0.0 - 1.0 representing how much to softsnap each tick of the remaining distance."));
+		
+		static float SoftSnapRotStrength = 0.5f;
+		static FAutoConsoleVariableRef CVarSoftSnapRotStrength(TEXT("np2.PredictiveInterpolation.SoftSnapRotStrength"), SoftSnapRotStrength, TEXT("Value in percent between 0.0 - 1.0 representing how much to softsnap each tick of the remaining distance."));
 
 		static float EarlyOutDistanceSqr = 2.f;
 		static FAutoConsoleVariableRef CVarEarlyOutDistanceSqr(TEXT("np2.PredictiveInterpolation.EarlyOutDistanceSqr"), EarlyOutDistanceSqr, TEXT("Squared value. If object is within this distance from the source target, early out from replication and apply sleep if replicated."));
@@ -788,6 +806,7 @@ void FPhysicsReplicationAsync::UpdateAsyncTarget(const FPhysicsRepAsyncInputData
 		Target->ServerFrame = Input.ServerFrame;
 		Target->PrevReceiveFrame = (Target->ReceiveFrame == INDEX_NONE) ? (RigidsSolver->GetCurrentFrame() - 1) : Target->ReceiveFrame;
 		Target->ReceiveFrame = RigidsSolver->GetCurrentFrame();
+		Target->ReceiveInterval = FMath::Clamp((Target->ReceiveFrame - Target->PrevReceiveFrame), 1, 255);
 		Target->TargetState = Input.TargetState;
 		Target->RepMode = Input.RepMode;
 		Target->FrameOffset = Input.FrameOffset;
@@ -1212,45 +1231,57 @@ bool FPhysicsReplicationAsync::PredictiveInterpolation(Chaos::FPBDRigidParticleH
 		}
 	}
 
+	// Calculate position correction time based on current Round Trip Time
 	const float RTT = LatencyOneWay * 2.f;
 	const float PosCorrectionTime = PhysicsReplicationCVars::PredictiveInterpolationCVars::PosCorrectionTimeBase + RTT * PhysicsReplicationCVars::PredictiveInterpolationCVars::PosCorrectionTimeMultiplier;
-	const float ReceiveInterval = Target.ReceiveFrame - Target.PrevReceiveFrame;
-	const float InterpolationTime = ReceiveInterval * DeltaSeconds * PhysicsReplicationCVars::PredictiveInterpolationCVars::InterpolationTimeMultiplier;
-	bool bSoftSnap = !PhysicsReplicationCVars::PredictiveInterpolationCVars::bVelocityBased;
+
+	// Calculate interpolation time based on current average receive rate of targets from the server (receive rate = send-rate from server with network conditions taken into account)
+	Target.AverageReceiveInterval = FMath::Lerp(Target.AverageReceiveInterval, Target.ReceiveInterval, FMath::Clamp((1.0f / (Target.ReceiveInterval * PhysicsReplicationCVars::PredictiveInterpolationCVars::AverageReceiveIntervalSmoothing)), 0.0f, 1.0f));
+	const float InterpolationTime = Target.AverageReceiveInterval * DeltaSeconds * PhysicsReplicationCVars::PredictiveInterpolationCVars::InterpolationTimeMultiplier;
 
 	// CurrentState
 	FRigidBodyState CurrentState;
 	CurrentState.Position = Handle->X();
 	CurrentState.Quaternion = Handle->R();
 	CurrentState.LinVel = Handle->V();
-	CurrentState.AngVel = Handle->W();
+	CurrentState.AngVel = Handle->W(); // Note: Current angular velocity is in Radians
 
 	// NewState
 	const FVector TargetPos = Target.TargetState.Position;
 	const FQuat TargetRot = Target.TargetState.Quaternion;
 	const FVector TargetLinVel = Target.TargetState.LinVel;
-	const FVector TargetAngVel = Target.TargetState.AngVel;
+	const FVector TargetAngVel = Target.TargetState.AngVel; // Note: Target angular velocity is in Degrees
 
 	/** --- Reconciliation ---
-	* Get the traveled direction and distance from previous frame and compare with replicated linear velocity.
-	* If the object isn't moving enough along the replicated velocity it's considered stuck and needs a hard reconciliation.
-	*/
-	const FVector PrevDiff = CurrentState.Position - Target.PrevPos;
-	const float	ExpectedDistance = (Target.PrevLinVel * DeltaSeconds).Size();
-	const float CoveredDistance = FVector::DotProduct(PrevDiff, Target.PrevLinVel.GetSafeNormal());
-	const float CoveredAplha = FMath::Clamp(CoveredDistance / ExpectedDistance, 0.f, 1.f);
-
-	// If the object is moving less than X% of the expected distance, accumulate error seconds
-	if (ExpectedDistance > UE_SMALL_NUMBER && CoveredAplha < PhysicsReplicationCVars::PredictiveInterpolationCVars::MinExpectedDistanceCovered)
+	* If target velocities are low enough, check the traveled direction and distance from previous frame and compare with replicated linear velocity.
+	* If the object isn't moving enough along the replicated velocity it's considered stuck and needs reconciliation.
+	* SoftSnap is performed each tick while there is a registered error, if enough time pass HardSnap forces the object into the correct state. */
+	bool bSoftSnap = !PhysicsReplicationCVars::PredictiveInterpolationCVars::bVelocityBased;
+
+	if ( PhysicsReplicationCVars::PredictiveInterpolationCVars::DisableErrorVelocityLimits ||
+		(TargetLinVel.Size() < PhysicsReplicationCVars::PredictiveInterpolationCVars::ErrorAccLinVelMaxLimit && TargetAngVel.Size() < PhysicsReplicationCVars::PredictiveInterpolationCVars::ErrorAccAngVelMaxLimit))
 	{
-		Target.AccumulatedErrorSeconds += DeltaSeconds;
-		bSoftSnap = true;
+		const FVector PrevDiff = CurrentState.Position - Target.PrevPos;
+		const float ExpectedDistance = (Target.PrevLinVel * DeltaSeconds).Size();
+		const float CoveredDistance = FVector::DotProduct(PrevDiff, Target.PrevLinVel.GetSafeNormal());
+		const float CoveredAplha = FMath::Clamp(CoveredDistance / ExpectedDistance, 0.0f, 1.0f);
+
+		// If the object is moving less than X% of the expected distance, accumulate error seconds
+		if (CoveredAplha < PhysicsReplicationCVars::PredictiveInterpolationCVars::MinExpectedDistanceCovered)
+		{
+			Target.AccumulatedErrorSeconds += DeltaSeconds;
+			bSoftSnap = true;
+		}
+		else if (Target.AccumulatedErrorSeconds > 0.f)
+		{
+			const float DecreaseTime = DeltaSeconds * PhysicsReplicationCVars::PredictiveInterpolationCVars::ErrorAccumulationDecreaseMultiplier;
+			Target.AccumulatedErrorSeconds = FMath::Max(Target.AccumulatedErrorSeconds - DecreaseTime, 0.0f);
+			bSoftSnap = true;
+		}
 	}
-	else if (Target.AccumulatedErrorSeconds > 0.f)
+	else
 	{
-		const float DecreaseTime = DeltaSeconds * PhysicsReplicationCVars::PredictiveInterpolationCVars::ErrorAccumulationDecreaseMultiplier;
-		Target.AccumulatedErrorSeconds = FMath::Max(Target.AccumulatedErrorSeconds - DecreaseTime, 0.0f);
-		bSoftSnap = true;
+		Target.AccumulatedErrorSeconds = 0;
 	}
 
 	if (PhysicsReplicationCVars::PredictiveInterpolationCVars::bDisableSoftSnap && PhysicsReplicationCVars::PredictiveInterpolationCVars::bVelocityBased)
@@ -1258,7 +1289,10 @@ bool FPhysicsReplicationAsync::PredictiveInterpolation(Chaos::FPBDRigidParticleH
 		bSoftSnap = false;
 	}
 
-	const bool bHardSnap = !bCanSimulate || Target.AccumulatedErrorSeconds > PhysicsReplicationCVars::PredictiveInterpolationCVars::ErrorAccumulationSeconds || PhysicsReplicationCVars::PredictiveInterpolationCVars::bAlwaysHardSnap;
+	const bool bHardSnap = !bCanSimulate ||
+		Target.AccumulatedErrorSeconds > PhysicsReplicationCVars::PredictiveInterpolationCVars::ErrorAccumulationSeconds ||
+		PhysicsReplicationCVars::PredictiveInterpolationCVars::bAlwaysHardSnap;
+
 	if (bHardSnap)
 	{
 		// Too much error so just snap state here and be done with it
@@ -1268,7 +1302,7 @@ bool FPhysicsReplicationAsync::PredictiveInterpolation(Chaos::FPBDRigidParticleH
 		Handle->SetR(Target.PrevRotTarget);
 		Handle->SetQ(Target.PrevRotTarget);
 		Handle->SetV(Target.TargetState.LinVel);
-		Handle->SetW(Target.TargetState.AngVel);
+		Handle->SetW(FMath::DegreesToRadians(Target.TargetState.AngVel));
 
 		// Cache data for next replication
 		Target.PrevLinVel = FVector(Target.TargetState.LinVel);
@@ -1298,7 +1332,6 @@ bool FPhysicsReplicationAsync::PredictiveInterpolation(Chaos::FPBDRigidParticleH
 		float CurAngVelSize;
 		FVector CurAngVelAxis;
 		CurrentState.AngVel.FVector::ToDirectionAndLength(CurAngVelAxis, CurAngVelSize);
-		CurAngVelSize = FMath::DegreesToRadians(CurAngVelSize);
 		const FQuat CurRotExtrapDelta = FQuat(CurAngVelAxis, CurAngVelSize * DeltaSeconds);
 		const FQuat CurRotExtrap = CurRotExtrapDelta * CurrentState.Quaternion;
 
@@ -1313,16 +1346,8 @@ bool FPhysicsReplicationAsync::PredictiveInterpolation(Chaos::FPBDRigidParticleH
 		float WAngle;
 		FVector WAxis;
 		TargetRotDelta.ToAxisAndAngle(WAxis, WAngle);
-
-		const FVector RepAngVel = WAxis * (WAngle / DeltaSeconds);
-
-
-		if (bSoftSnap)
-		{
-			const float AlphaToSoftSnap = 1.f - CoveredAplha; // If we covered 20% of the distance by velocity (1.f - 0.2f = 0.8f) cover 80% of the distance via softsnap.
-			Handle->SetX(CurrentState.Position + (RepLinVel * (AlphaToSoftSnap * DeltaSeconds)));
-			Handle->SetR(TargetRotBlended);
-		}
+		const FVector TargetRotDeltaBlend = FVector(WAxis * (WAngle / DeltaSeconds));
+		const FVector RepAngVel = FMath::DegreesToRadians(Target.TargetState.AngVel) + TargetRotDeltaBlend;
 
 		Handle->SetV(RepLinVel);
 		Handle->SetW(RepAngVel);
@@ -1332,7 +1357,7 @@ bool FPhysicsReplicationAsync::PredictiveInterpolation(Chaos::FPBDRigidParticleH
 		Target.PrevPos = FVector(CurrentState.Position);
 
 		// --- Target Extrapolation ---
-		if (Target.TickCount < FMath::CeilToInt(ReceiveInterval * PhysicsReplicationCVars::PredictiveInterpolationCVars::ExtrapolationTimeMultiplier))
+		if (Target.TickCount <= FMath::CeilToInt(Target.ReceiveInterval * PhysicsReplicationCVars::PredictiveInterpolationCVars::ExtrapolationTimeMultiplier))
 		{
 			// Extrapolate target position
 			Target.TargetState.Position = Target.TargetState.Position + Target.TargetState.LinVel * DeltaSeconds;
@@ -1345,8 +1370,18 @@ bool FPhysicsReplicationAsync::PredictiveInterpolation(Chaos::FPBDRigidParticleH
 			const FQuat TargetRotExtrapDelta = FQuat(TargetAngVelAxis, TargetAngVelSize * DeltaSeconds);
 			Target.TargetState.Quaternion = TargetRotExtrapDelta * Target.TargetState.Quaternion;
 		}
+
+		if (bSoftSnap)
+		{
+			const FVector SoftSnapPos = FMath::Lerp(FVector(CurrentState.Position), Target.PrevPosTarget, FMath::Clamp(PhysicsReplicationCVars::PredictiveInterpolationCVars::SoftSnapPosStrength, 0.0f, 1.0f));
+			const FQuat SoftSnapRot = FQuat::Slerp(CurrentState.Quaternion, Target.PrevRotTarget, FMath::Clamp(PhysicsReplicationCVars::PredictiveInterpolationCVars::SoftSnapRotStrength, 0.0f, 1.0f));
+			Handle->SetX(SoftSnapPos);
+			Handle->SetP(SoftSnapPos);
+			Handle->SetR(SoftSnapRot);
+			Handle->SetQ(SoftSnapRot);
+		}
 	}
-	
+
 	return !bCanSimulate && !PhysicsReplicationCVars::PredictiveInterpolationCVars::bDontClearTarget; // If the object can't simulate, clear target
 }
 
diff --git a/Engine/Source/Runtime/Engine/Public/PhysicsReplication.h b/Engine/Source/Runtime/Engine/Public/PhysicsReplication.h
--- a/Engine/Source/Runtime/Engine/Public/PhysicsReplication.h
+++ b/Engine/Source/Runtime/Engine/Public/PhysicsReplication.h
@@ -94,6 +94,8 @@ struct FReplicatedPhysicsTargetAsync
 		: AccumulatedErrorSeconds(0.0f)
 		, ServerFrame(INDEX_NONE)
 		, ReceiveFrame(INDEX_NONE)
+		, ReceiveInterval(5)
+		, AverageReceiveInterval(5.f)
 	{ }
 
 	/** The target state replicated by server */
@@ -114,6 +116,10 @@ struct FReplicatedPhysicsTargetAsync
 	/** The local client frame when receiving this target from the server */
 	int32 ReceiveFrame;
 
+	/** Local physics frames between received targets */
+	uint8 ReceiveInterval;
+	float AverageReceiveInterval;
+
 	/** The replication mode this PhysicsObject should use */
 	EPhysicsReplicationMode RepMode;
 
