Microoptimisations for FPhsyScene_Chaos::HandleGlobalCollisionEvent -  pre allocating large array, removing unnecessary casts, and reducing duplicate weak object ptr gets.

diff --git a/Engine/Source/Runtime/Engine/Private/PhysicsEngine/Experimental/PhysScene_Chaos.cpp b/Engine/Source/Runtime/Engine/Private/PhysicsEngine/Experimental/PhysScene_Chaos.cpp
--- a/Engine/Source/Runtime/Engine/Private/PhysicsEngine/Experimental/PhysScene_Chaos.cpp
+++ b/Engine/Source/Runtime/Engine/Private/PhysicsEngine/Experimental/PhysScene_Chaos.cpp
@@ -799,6 +799,16 @@ void FPhysScene_Chaos::AddReferencedObjects(FReferenceCollector& Collector)
 #endif
 }
 
+template<>
+UPrimitiveComponent* FPhysScene_Chaos::GetOwningComponent(const IPhysicsProxyBase* PhysicsProxy) const
+{
+	if (const TObjectPtr<UPrimitiveComponent>* FoundComp = PhysicsProxyToComponentMap.Find(PhysicsProxy))
+	{
+		return *FoundComp;
+	}
+	return nullptr;
+}
+
 FBodyInstance* FPhysScene_Chaos::GetBodyInstanceFromProxy(const IPhysicsProxyBase* PhysicsProxy) const
 {
 	FBodyInstance* BodyInstance = nullptr;
@@ -1031,43 +1041,57 @@ void FPhysScene_Chaos::HandleGlobalCollisionEvent(Chaos::FCollisionDataArray con
 {
 	QUICK_SCOPE_CYCLE_COUNTER(STAT_HandleGlobalCollisionEvents);
 	TArray<FCollisionChaosEvent> CollisionEvents;
-	CollisionEvents.Reserve(CollisionEvents.Num());
+	CollisionEvents.Reserve(CollisionData.Num());
 
 	// If iterating by proxy, then need to process the duplicate collision, 
 	// so iterating by collision data should be the faster when using only globals
 	for (const Chaos::FCollidingData& CollisionItem : CollisionData)
 	{
-		UPrimitiveComponent* Comp0 = GetOwningComponent<UPrimitiveComponent>(CollisionItem.Proxy1);
-		UPrimitiveComponent* Comp1 = GetOwningComponent<UPrimitiveComponent>(CollisionItem.Proxy2);
+		UPrimitiveComponent* BodyPrimitive1 = GetOwningComponent<UPrimitiveComponent>(CollisionItem.Proxy1);
+		UPrimitiveComponent* BodyPrimitive2 = GetOwningComponent<UPrimitiveComponent>(CollisionItem.Proxy2);
 
-		if (GlobalCollisionEventRegistrations.Contains(Comp0) || GlobalCollisionEventRegistrations.Contains(Comp1))
+		if (GlobalCollisionEventRegistrations.Contains(BodyPrimitive1) || GlobalCollisionEventRegistrations.Contains(BodyPrimitive2))
 		{
 			FCollisionChaosEvent& CollisionEvent = CollisionEvents.Emplace_GetRef(CollisionItem);
-			CollisionEvent.Body1.Component = GetOwningComponent<UPrimitiveComponent>(CollisionItem.Proxy1);
-			CollisionEvent.Body2.Component = GetOwningComponent<UPrimitiveComponent>(CollisionItem.Proxy2);
+			CollisionEvent.Body1.Component = BodyPrimitive1;
+			CollisionEvent.Body2.Component = BodyPrimitive2;
 
 			const Chaos::FChaosPhysicsMaterial* InternalMat1 = CollisionItem.Mat1.Get();
 			const Chaos::FChaosPhysicsMaterial* InternalMat2 = CollisionItem.Mat2.Get();
 			CollisionEvent.Body1.PhysMaterial = InternalMat1 ? FPhysicsUserData::Get<UPhysicalMaterial>(InternalMat1->UserData) : nullptr;
 			CollisionEvent.Body2.PhysMaterial = InternalMat2 ? FPhysicsUserData::Get<UPhysicalMaterial>(InternalMat2->UserData) : nullptr;
 			
-			if (CollisionEvent.Body1.Component.IsValid())
+			if (BodyPrimitive1)
 			{
 				const FBodyInstance* BodyInst1 = GetBodyInstanceFromProxyAndShape(CollisionItem.Proxy1, CollisionItem.ShapeIndex1);
 				if (BodyInst1 != nullptr)
 				{
 					CollisionEvent.Body1.BodyIndex = BodyInst1->InstanceBodyIndex;
-					CollisionEvent.Body1.BoneName = BodyInst1->BodySetup.IsValid() ? BodyInst1->BodySetup->BoneName : NAME_None;
+					if (UBodySetupCore* BodySetup1 = BodyInst1->BodySetup.Get())
+					{
+						CollisionEvent.Body1.BoneName = BodySetup1->BoneName;
+					}
+					else
+					{
+						CollisionEvent.Body1.BoneName = NAME_None;
+					}
 				}
 			}
 
-			if (CollisionEvent.Body2.Component.IsValid())
+			if (BodyPrimitive2)
 			{
 				const FBodyInstance* BodyInst2 = GetBodyInstanceFromProxyAndShape(CollisionItem.Proxy2, CollisionItem.ShapeIndex2);
 				if (BodyInst2 != nullptr)
 				{
 					CollisionEvent.Body2.BodyIndex = BodyInst2->InstanceBodyIndex;
-					CollisionEvent.Body2.BoneName = BodyInst2->BodySetup.IsValid() ? BodyInst2->BodySetup->BoneName : NAME_None;
+					if (UBodySetupCore* BodySetup2 = BodyInst2->BodySetup.Get())
+					{
+						CollisionEvent.Body2.BoneName = BodySetup2->BoneName;
+					}
+					else
+					{
+						CollisionEvent.Body2.BoneName = NAME_None;
+					}
 				}
 			}
 		}
diff --git a/Engine/Source/Runtime/Engine/Public/Physics/Experimental/PhysScene_Chaos.h b/Engine/Source/Runtime/Engine/Public/Physics/Experimental/PhysScene_Chaos.h
--- a/Engine/Source/Runtime/Engine/Public/Physics/Experimental/PhysScene_Chaos.h
+++ b/Engine/Source/Runtime/Engine/Public/Physics/Experimental/PhysScene_Chaos.h
@@ -168,6 +168,9 @@ public:
 		return CompPtr ? Cast<OwnerType>(*CompPtr) : nullptr;
 	}
 
+	template<>
+	ENGINE_API UPrimitiveComponent* GetOwningComponent(const IPhysicsProxyBase* PhysicsProxy) const;
+
 	/** Given a component, returns its associated solver objects. */
 	const TArray<IPhysicsProxyBase*>* GetOwnedPhysicsProxies(UPrimitiveComponent* Comp) const
 	{
