Build fix error in StringTemplate.h

diff --git a/Engine/Source/Runtime/Engine/Private/Materials/StringTemplate.h b/Engine/Source/Runtime/Engine/Private/Materials/StringTemplate.h
--- a/Engine/Source/Runtime/Engine/Private/Materials/StringTemplate.h
+++ b/Engine/Source/Runtime/Engine/Private/Materials/StringTemplate.h
@@ -5,6 +5,7 @@
 #include "CoreTypes.h"
 #include "Misc/AssertionMacros.h"
 #include "Containers/UnrealString.h"
+#include "Containers/StringView.h"
 
 /**
  * A string template is a utility class used to generate parameterized strings.
