Build fix in MaterialSourceTemplate.cpp

diff --git a/Engine/Source/Runtime/Engine/Private/Materials/MaterialSourceTemplate.cpp b/Engine/Source/Runtime/Engine/Private/Materials/MaterialSourceTemplate.cpp
--- a/Engine/Source/Runtime/Engine/Private/Materials/MaterialSourceTemplate.cpp
+++ b/Engine/Source/Runtime/Engine/Private/Materials/MaterialSourceTemplate.cpp
@@ -5,8 +5,7 @@
 #include "MaterialSourceTemplate.h"
 #include "Misc/ScopeRWLock.h"
 #include "ShaderCore.h"
-
-DECLARE_LOG_CATEGORY_EXTERN(LogMaterial, Log, All);
+#include "MaterialShared.h"
 
 FMaterialSourceTemplate& FMaterialSourceTemplate::Get()
 {
