World Partition - Use the pawn as the streaming source whenever possible
* If there is no pawn, the camera will be used (previous behavior)
* Added wp.Runtime.PlayerController.ForceUsingCameraAsStreamingSource cvar to allow disabling/hotfixing off this new behavior if ever needed

diff --git a/Engine/Source/Runtime/Engine/Private/PlayerController.cpp b/Engine/Source/Runtime/Engine/Private/PlayerController.cpp
--- a/Engine/Source/Runtime/Engine/Private/PlayerController.cpp
+++ b/Engine/Source/Runtime/Engine/Private/PlayerController.cpp
@@ -98,6 +98,13 @@ namespace PlayerControllerCVars
 		TEXT("Whether to reset server prediction data for the possessed Pawn when the pawn ack handshake completes.\n")
 		TEXT("0: Disable, 1: Enable"),
 		ECVF_Default);
+
+	static int32 ForceUsingCameraAsStreamingSource = 0;
+	FAutoConsoleVariableRef CVarForceUsingCameraAsStreamingSource(
+		TEXT("wp.Runtime.PlayerController.ForceUsingCameraAsStreamingSource"),
+		ForceUsingCameraAsStreamingSource,
+		TEXT("Whether to force the use of the camera as the streaming source for World Partition. By default the player pawn is used.\n")
+		TEXT("0: Use pawn as streaming source, 1: Use camera as streaming source"));
 }
 
 namespace NetworkPhysicsCvars
@@ -3611,7 +3618,14 @@ void APlayerController::OnRemovedFromPlayerControllerList()
 
 void APlayerController::GetStreamingSourceLocationAndRotation(FVector& OutLocation, FRotator& OutRotation) const
 {
-	GetPlayerViewPoint(OutLocation, OutRotation);
+	if (!PlayerControllerCVars::ForceUsingCameraAsStreamingSource)
+	{
+		GetActorEyesViewPoint(OutLocation, OutRotation);
+	}
+	else
+	{
+		GetPlayerViewPoint(OutLocation, OutRotation);
+	}
 }
 
 void APlayerController::GetStreamingSourceShapes(TArray<FStreamingSourceShape>& OutShapes) const
