Server crash with 'FPhysScene_Chaos::GetBodyInstanceFromProxy' when many players create and destroy assets.
Crash fix attempt: The secondary proxy might be null, and could be the cause of the crash. This crash is a hard to repro one, so the fix cannot be validated easily. Only Null pointer check have been added.

diff --git a/Engine/Source/Runtime/Engine/Private/PhysicsEngine/Experimental/PhysScene_Chaos.cpp b/Engine/Source/Runtime/Engine/Private/PhysicsEngine/Experimental/PhysScene_Chaos.cpp
--- a/Engine/Source/Runtime/Engine/Private/PhysicsEngine/Experimental/PhysScene_Chaos.cpp
+++ b/Engine/Source/Runtime/Engine/Private/PhysicsEngine/Experimental/PhysScene_Chaos.cpp
@@ -812,17 +812,20 @@ UPrimitiveComponent* FPhysScene_Chaos::GetOwningComponent(const IPhysicsProxyBas
 FBodyInstance* FPhysScene_Chaos::GetBodyInstanceFromProxy(const IPhysicsProxyBase* PhysicsProxy) const
 {
 	FBodyInstance* BodyInstance = nullptr;
-	if (PhysicsProxy && PhysicsProxy->GetType() == EPhysicsProxyType::SingleParticleProxy)
+	if (PhysicsProxy)
 	{
-		const Chaos::FRigidBodyHandle_External& RigidBodyHandle = static_cast<const Chaos::FSingleParticlePhysicsProxy*>(PhysicsProxy)->GetGameThreadAPI();
-		BodyInstance = FPhysicsUserData::Get<FBodyInstance>(RigidBodyHandle.UserData());
-	}
-	// found none, let's see if there's an owning component in the scene
-	if (BodyInstance == nullptr)
-	{
-		if (UPrimitiveComponent* const OwningComponent = GetOwningComponent<UPrimitiveComponent>(PhysicsProxy))
+		if (PhysicsProxy->GetType() == EPhysicsProxyType::SingleParticleProxy)
 		{
-			BodyInstance = OwningComponent->GetBodyInstance();
+			const Chaos::FRigidBodyHandle_External& RigidBodyHandle = static_cast<const Chaos::FSingleParticlePhysicsProxy*>(PhysicsProxy)->GetGameThreadAPI();
+			BodyInstance = FPhysicsUserData::Get<FBodyInstance>(RigidBodyHandle.UserData());
+		}
+		// found none, let's see if there's an owning component in the scene
+		if (BodyInstance == nullptr)
+		{
+			if (UPrimitiveComponent* const OwningComponent = GetOwningComponent<UPrimitiveComponent>(PhysicsProxy))
+			{
+				BodyInstance = OwningComponent->GetBodyInstance();
+			}
 		}
 	}
 	return BodyInstance;
@@ -877,6 +880,10 @@ FORCEINLINE void FPhysScene_Chaos::HandleEachCollisionEvent(const TArray<int32>&
 		Chaos::FCollidingData const& CollisionDataItem = CollisionData[CollisionIdx];
 		IPhysicsProxyBase* const PhysicsProxy1 = bSwapOrder ? CollisionDataItem.Proxy1 : CollisionDataItem.Proxy2;
 
+		if (PhysicsProxy1 == nullptr)
+		{
+			continue;
+		}
 		// Are the proxies pending destruction? If they are no longer tracked by the PhysScene, the proxy is deleted or pending deletion.
 		UPrimitiveComponent* const Comp1 = GetOwningComponent<UPrimitiveComponent>(PhysicsProxy1);
 		if (Comp1 == nullptr)
