Fix a bug where a new platform user Id would be incorrectly allocated via a call to RemapControllerId on the device mapper, resulting in the platform user id never being destroyed across PIE sessions

diff --git a/Engine/Source/Runtime/Engine/Private/GameInstance.cpp b/Engine/Source/Runtime/Engine/Private/GameInstance.cpp
--- a/Engine/Source/Runtime/Engine/Private/GameInstance.cpp
+++ b/Engine/Source/Runtime/Engine/Private/GameInstance.cpp
@@ -832,9 +832,7 @@ ULocalPlayer* UGameInstance::CreateLocalPlayer(FPlatformUserId UserId, FString&
 			for (int32 Id = 0; Id < MaxSplitscreenPlayers; ++Id)
 			{
 				// Iterate until we find a null player. We want the next available platform user ID
-				FInputDeviceId DummyDevice;
-				FPlatformUserId DummyId;
-				IPlatformInputDeviceMapper::Get().RemapControllerIdToPlatformUserAndDevice(Id, OUT DummyId, OUT DummyDevice);
+				FPlatformUserId DummyId = IPlatformInputDeviceMapper::Get().GetPlatformUserForUserIndex(Id);
 
 				if (DummyId.IsValid())
 				{
