[WP navmesh] Fix for ANavigationDataChunkActor having wrong actor location caused by missing RootComponent

diff --git a/Engine/Source/Runtime/Engine/Private/WorldPartition/NavigationData/NavigationDataChunkActor.cpp b/Engine/Source/Runtime/Engine/Private/WorldPartition/NavigationData/NavigationDataChunkActor.cpp
--- a/Engine/Source/Runtime/Engine/Private/WorldPartition/NavigationData/NavigationDataChunkActor.cpp
+++ b/Engine/Source/Runtime/Engine/Private/WorldPartition/NavigationData/NavigationDataChunkActor.cpp
@@ -13,13 +13,6 @@
 ANavigationDataChunkActor::ANavigationDataChunkActor(const FObjectInitializer& ObjectInitializer)
 	: Super(ObjectInitializer)
 {
-	if (!IsTemplate())
-	{
-		USceneComponent* Root = CreateDefaultSubobject<USceneComponent>(USceneComponent::GetDefaultSceneRootVariableName());
-		Root->SetMobility(EComponentMobility::Static);
-		SetRootComponent(Root);
-	}
-
 	SetCanBeDamaged(false);
 	SetActorEnableCollision(false);
 }
