WorldPartition: Added wp.Runtime.DumpWorldPartitions for debugging and rename wp.DumpstreamingSources to wp.Runtime.DumpStreamingSources
Added log when uninitializing a world partition.

diff --git a/Engine/Source/Runtime/Engine/Private/WorldPartition/WorldPartition.cpp b/Engine/Source/Runtime/Engine/Private/WorldPartition/WorldPartition.cpp
--- a/Engine/Source/Runtime/Engine/Private/WorldPartition/WorldPartition.cpp
+++ b/Engine/Source/Runtime/Engine/Private/WorldPartition/WorldPartition.cpp
@@ -541,7 +541,7 @@ void UWorldPartition::Initialize(UWorld* InWorld, const FTransform& InTransform)
 	const bool bIsDedicatedServer = IsRunningDedicatedServer();
 
 	UE_LOG(LogWorldPartition, Log, TEXT("UWorldPartition::Initialize : World = %s, World Type = %s, IsMainWorldPartition = %d, Location = %s, Rotation = %s, IsEditor = %d, IsGame = %d, IsPIEWorldTravel = %d, IsCooking = %d"),
-		*OuterWorld->GetName(),
+		*OuterWorld->GetPathName(),
 		LexToString(World->WorldType),
 		IsMainWorldPartition() ? 1 : 0,
 		*InTransform.GetLocation().ToCompactString(),
@@ -754,6 +754,8 @@ void UWorldPartition::Uninitialize()
 	{
 		check(World);
 
+		UE_LOG(LogWorldPartition, Log, TEXT("UWorldPartition::Uninitialize : World = %s"), *GetTypedOuter<UWorld>()->GetPathName());
+
 		InitState = EWorldPartitionInitState::Uninitializing;
 
 		if (IsMainWorldPartition())
diff --git a/Engine/Source/Runtime/Engine/Private/WorldPartition/WorldPartitionSubsystem.cpp b/Engine/Source/Runtime/Engine/Private/WorldPartition/WorldPartitionSubsystem.cpp
--- a/Engine/Source/Runtime/Engine/Private/WorldPartition/WorldPartitionSubsystem.cpp
+++ b/Engine/Source/Runtime/Engine/Private/WorldPartition/WorldPartitionSubsystem.cpp
@@ -111,7 +111,7 @@ static FAutoConsoleVariableRef CVarGLevelStreamingContinuouslyIncrementalGCWhile
 );
 
 static FAutoConsoleCommandWithOutputDevice GDumpStreamingSourcesCmd(
-	TEXT("wp.DumpstreamingSources"),
+	TEXT("wp.Runtime.DumpStreamingSources"),
 	TEXT("Dumps active streaming sources to the log"),
 	FConsoleCommandWithOutputDeviceDelegate::CreateStatic([](FOutputDevice& OutputDevice)
 	{
@@ -129,6 +129,25 @@ static FAutoConsoleCommandWithOutputDevice GDumpStreamingSourcesCmd(
 	})
 );
 
+static FAutoConsoleCommandWithOutputDevice GDumpWorldPartitionsCmd(
+	TEXT("wp.Runtime.DumpWorldPartitions"),
+	TEXT("Dumps active world partitions to the log"),
+	FConsoleCommandWithOutputDeviceDelegate::CreateStatic([](FOutputDevice& OutputDevice)
+	{
+		for (const FWorldContext& Context : GEngine->GetWorldContexts())
+		{
+			UWorld* World = Context.World();
+			if (World && World->IsGameWorld())
+			{
+				if (const UWorldPartitionSubsystem* WorldPartitionSubsystem = World->GetSubsystem<UWorldPartitionSubsystem>())
+				{				
+					WorldPartitionSubsystem->DumpWorldPartitions(OutputDevice);
+				}
+			}
+		}
+	})
+);
+
 static int32 GMaxLoadingStreamingCells = 4;
 static FAutoConsoleVariableRef CMaxLoadingStreamingCells(
 	TEXT("wp.Runtime.MaxLoadingStreamingCells"),
@@ -820,6 +839,18 @@ bool UWorldPartitionSubsystem::IsStreamingCompleted(EWorldPartitionRuntimeCellSt
 	return true;
 }
 
+void UWorldPartitionSubsystem::DumpWorldPartitions(FOutputDevice& OutputDevice) const
+{
+	if (RegisteredWorldPartitions.Num() > 0)
+	{
+		OutputDevice.Logf(TEXT("Registered World Partitions:"));
+		TArray<FString> WorldPartitions;
+		Algo::ForEach(RegisteredWorldPartitions, [&WorldPartitions](const UWorldPartition* WorldPartition) { WorldPartitions.Add(WorldPartition->GetPathName()); });
+		WorldPartitions.Sort();
+		Algo::ForEach(WorldPartitions, [&OutputDevice](const FString& WorldPartition) { OutputDevice.Logf(TEXT("  - %s"), *WorldPartition); });
+	}
+}
+
 void UWorldPartitionSubsystem::DumpStreamingSources(FOutputDevice& OutputDevice) const
 {
 	if (StreamingSources.Num() > 0)
diff --git a/Engine/Source/Runtime/Engine/Public/WorldPartition/WorldPartitionSubsystem.h b/Engine/Source/Runtime/Engine/Public/WorldPartition/WorldPartitionSubsystem.h
--- a/Engine/Source/Runtime/Engine/Public/WorldPartition/WorldPartitionSubsystem.h
+++ b/Engine/Source/Runtime/Engine/Public/WorldPartition/WorldPartitionSubsystem.h
@@ -91,6 +91,7 @@ public:
 	ENGINE_API bool IsStreamingCompleted(const IWorldPartitionStreamingSourceProvider* InStreamingSourceProvider = nullptr) const;
 
 	ENGINE_API void DumpStreamingSources(FOutputDevice& OutputDevice) const;
+	ENGINE_API void DumpWorldPartitions(FOutputDevice& OutputDevice) const;
 
 	ENGINE_API TSet<IWorldPartitionStreamingSourceProvider*> GetStreamingSourceProviders() const;
 	ENGINE_API void RegisterStreamingSourceProvider(IWorldPartitionStreamingSourceProvider* StreamingSource);
