#UE Do not add primary assets found in the transient package to the cook

diff --git a/Engine/Source/Runtime/Engine/Private/AssetManager.cpp b/Engine/Source/Runtime/Engine/Private/AssetManager.cpp
--- a/Engine/Source/Runtime/Engine/Private/AssetManager.cpp
+++ b/Engine/Source/Runtime/Engine/Private/AssetManager.cpp
@@ -4144,6 +4144,9 @@ void UAssetManager::ModifyCook(TConstArrayView<const ITargetPlatform*> TargetPla
 
 	bool bIncludeDevelopmentAssets = !bOnlyCookProductionAssets || bTargetPlatformsAllowDevelopmentObjects;
 
+	// Some primary assets exist in the transient package. No need to include them in the cook since they are transient.
+	FName TransientPackageName = GetTransientPackage()->GetFName();
+
 	// Uniquely append packages we need that are not already in PackagesToCook and PackagesToNeverCook
 	TSet<FName> PackagesToCookSet(PackagesToCook);
 	TSet<FName> PackagesToNeverCookSet(PackagesToNeverCook);
@@ -4160,7 +4163,7 @@ void UAssetManager::ModifyCook(TConstArrayView<const ITargetPlatform*> TargetPla
 		for (const FPrimaryAssetId& PrimaryAssetId : AssetIdList)
 		{
 			FAssetData AssetData;
-			if (GetPrimaryAssetData(PrimaryAssetId, AssetData))
+			if (GetPrimaryAssetData(PrimaryAssetId, AssetData) && AssetData.PackageName != TransientPackageName)
 			{
 				// If this has an asset data, add that package name
 				AssetPackages.Add(AssetData.PackageName);
