Fixed false positive verify in UHLODSubsystem::UnregisterHLODActor when unloading a previously injected/deinjected cell containing an HLOD Actor.
The HLOD actor of the injected cell was properly registered, so OutdatedHLODActors is not filled.
When the injected content is removed, UHLODSubsystem::OnExternalStreamingObjectRemoved does the cleanup.
Then, when UnregisterHLODActor is called, it won't find the cell data for this actor and will fail to remove the actor from OutdatedHLODActors.

diff --git a/Engine/Source/Runtime/Engine/Private/WorldPartition/HLOD/HLODSubsystem.cpp b/Engine/Source/Runtime/Engine/Private/WorldPartition/HLOD/HLODSubsystem.cpp
--- a/Engine/Source/Runtime/Engine/Private/WorldPartition/HLOD/HLODSubsystem.cpp
+++ b/Engine/Source/Runtime/Engine/Private/WorldPartition/HLOD/HLODSubsystem.cpp
@@ -479,7 +479,7 @@ void UHLODSubsystem::UnregisterHLODActor(AWorldPartitionHLOD* InWorldPartitionHL
 	else
 	{
 #if WITH_EDITOR
-		verify(OutdatedHLODActors.Remove(InWorldPartitionHLOD));
+		OutdatedHLODActors.Remove(InWorldPartitionHLOD);
 #endif
 	}
 
