Fix dormancy issue for physics objects - address a missing case in the Physics path.

diff --git a/Engine/Source/Runtime/Engine/Private/PhysicsEngine/Experimental/PhysScene_Chaos.cpp b/Engine/Source/Runtime/Engine/Private/PhysicsEngine/Experimental/PhysScene_Chaos.cpp
--- a/Engine/Source/Runtime/Engine/Private/PhysicsEngine/Experimental/PhysScene_Chaos.cpp
+++ b/Engine/Source/Runtime/Engine/Private/PhysicsEngine/Experimental/PhysScene_Chaos.cpp
@@ -2278,12 +2278,19 @@ void FPhysScene_Chaos::OnSyncBodies(Chaos::FPhysicsSolverBase* Solver)
 			{
 				const FRigidTransform3 NewTransform(DirtyParticle->X(), DirtyParticle->R());
 
+				bool bHasMoved = false;
 				if (!NewTransform.EqualsNoScale(ParentComponent->GetComponentTransform()))
 				{
+					bHasMoved = true;
 					const FVector MoveBy = NewTransform.GetLocation() - ParentComponent->GetComponentTransform().GetLocation();
 					PendingTransforms.Add(FPhysScenePendingComponentTransform_Chaos(ParentComponent, MoveBy, NewTransform.GetRotation(), DirtyParticle->GetWakeEvent()));
 				}
 
+				if (DirtyParticle->GetWakeEvent() != Chaos::EWakeEventEntry::None && !bHasMoved)
+				{
+					PendingTransforms.Add(FPhysScenePendingComponentTransform_Chaos(ParentComponent, DirtyParticle->GetWakeEvent()));
+				}
+
 				// make sure we have at least a child to be added to the acceleration structure 
 				// this avoid the invalid bounds to cause the particle to be added to the global acceleration structure array
 				if (ParentComponent->NumChildClusterComponents() > 0)
