GC debug: CollectGarbageEveryFrame can now collect every N frames (for values > 1)

diff --git a/Engine/Source/Runtime/Engine/Private/UnrealEngine.cpp b/Engine/Source/Runtime/Engine/Private/UnrealEngine.cpp
--- a/Engine/Source/Runtime/Engine/Private/UnrealEngine.cpp
+++ b/Engine/Source/Runtime/Engine/Private/UnrealEngine.cpp
@@ -1551,7 +1551,7 @@ static FAutoConsoleVariableRef CVarPerformGCWhileAsyncLoading(
 static TAutoConsoleVariable<int32> CVarCollectGarbageEveryFrame(
 	TEXT("gc.CollectGarbageEveryFrame"),
 	0,
-	TEXT("Used to debug garbage collection...Collects garbage every frame if the value is > 0."));
+	TEXT("Used to debug garbage collection...Collects garbage every N frames if the value is > 0."));
 
 
 static TAutoConsoleVariable<int32> CVarContinuousIncrementalGC(
@@ -1859,9 +1859,12 @@ void UEngine::ConditionalCollectGarbage()
 			}
 		}
 
-		if (CVarCollectGarbageEveryFrame.GetValueOnGameThread())
+		if (const int32 Interval = CVarCollectGarbageEveryFrame.GetValueOnGameThread())
 		{
-			ForceGarbageCollection(true);
+			if (0 == (GFrameCounter % Interval))
+			{
+				ForceGarbageCollection(true);
+			}
 		}
 		else if (CVarContinuousIncrementalGC.GetValueOnGameThread() > 0 && !IsIncrementalReachabilityAnalysisPending() && !IsIncrementalUnhashPending() && !IsIncrementalPurgePending())
 		{
