Add an ensure() to avoid a runtime crash and guard against using a CDO as a data source template when the BP component instancing fast path is enabled.

Previously, this scenario could lead to an assertion at runtime in a cooked build, but only when the owning Actor has enabled the fast path to dynamically instance its Blueprint-added components.

This is typically going to be due to an invalid/corrupt Blueprint asset, which will be identified in the output log. The ensure() can be fixed by removing the reported component from the Blueprint asset and then re-adding it, followed by re-saving/re-cooking.

#rnx

diff --git a/Engine/Source/Runtime/Engine/Private/BlueprintGeneratedClass.cpp b/Engine/Source/Runtime/Engine/Private/BlueprintGeneratedClass.cpp
--- a/Engine/Source/Runtime/Engine/Private/BlueprintGeneratedClass.cpp
+++ b/Engine/Source/Runtime/Engine/Private/BlueprintGeneratedClass.cpp
@@ -573,8 +573,10 @@ void UBlueprintGeneratedClass::SerializeDefaultObject(UObject* Object, FStructur
 			if (CookedData.bHasValidCookedData)
 			{
 				// This feature requires EDL at cook time, so ensure that the source template is also fully loaded at this point.
+				// Also ensure that the source template is not a class default object; it must always be a unique archetype object.
 				if (SourceTemplate != nullptr
-					&& ensure(!SourceTemplate->HasAnyFlags(RF_NeedLoad)))
+					&& ensure(!SourceTemplate->HasAnyFlags(RF_NeedLoad))
+					&& ensure(!SourceTemplate->HasAnyFlags(RF_ClassDefaultObject)))
 				{
 					CookedData.BuildCachedPropertyDataFromTemplate(SourceTemplate);
 				}
