Remove newly added ensure while we sort things out
#rnx

diff --git a/Engine/Source/Runtime/Engine/Private/WorldPartition/ActorPartition/PartitionActorDesc.cpp b/Engine/Source/Runtime/Engine/Private/WorldPartition/ActorPartition/PartitionActorDesc.cpp
--- a/Engine/Source/Runtime/Engine/Private/WorldPartition/ActorPartition/PartitionActorDesc.cpp
+++ b/Engine/Source/Runtime/Engine/Private/WorldPartition/ActorPartition/PartitionActorDesc.cpp
@@ -68,7 +68,6 @@ bool FPartitionActorDesc::Equals(const FWorldPartitionActorDesc* Other) const
 FBox FPartitionActorDesc::GetEditorBounds() const
 {
 	const UActorPartitionSubsystem::FCellCoord CellCoord(GridIndexX, GridIndexY, GridIndexZ, 0);
-	ensureMsgf(GridSize > 1, TEXT("Invalid GridSize. GridSize needs to be set."));
 	return UActorPartitionSubsystem::FCellCoord::GetCellBounds(CellCoord, GridSize);
 }
 
