Remove incorrect mobile specific assert related to PSO pre-caching of skin VF

diff --git a/Engine/Source/Runtime/Engine/Private/GPUSkinVertexFactory.cpp b/Engine/Source/Runtime/Engine/Private/GPUSkinVertexFactory.cpp
--- a/Engine/Source/Runtime/Engine/Private/GPUSkinVertexFactory.cpp
+++ b/Engine/Source/Runtime/Engine/Private/GPUSkinVertexFactory.cpp
@@ -750,10 +750,9 @@ void TGPUSkinVertexFactory<BoneInfluenceType>::GetVertexElements(ERHIFeatureLeve
 	int32 MorphDeltaStreamIndex;
 	GetVertexElements(FeatureLevel, InputStreamType, GPUSkinData, OutElements, VertexStreams, MorphDeltaStreamIndex);
 
-	if (UseGPUScene(GMaxRHIShaderPlatform, GMaxRHIFeatureLevel))
+	if (UseGPUScene(GMaxRHIShaderPlatform, GMaxRHIFeatureLevel) && 
+		FeatureLevel > ERHIFeatureLevel::ES3_1) // Skin VF does not use GPUScene on mobile
 	{
-		// For ES3.1 attribute ID needs to be done differently
-		check(FeatureLevel > ERHIFeatureLevel::ES3_1);
 		OutElements.Add(FVertexElement(VertexStreams.Num(), 0, VET_UInt, 16, 0, true));
 	}
 }
