[NetPhysics] Enable rewind capture in runtime if Project Settings -> Physics -> Physics Prediction is enabled.

Made it possible to enable rewind caching at runtime by tweaking the call flows and making sure that we have a RewindCallback available for components that want to register themselves before rewind caching is enabled.

Cleaned up and corrected a bunch of logs etc.
Removed the resimulation setting from Project Settings since it's activated automatically if using NetworkPhysicsComponent or can be activated manually by calling FPBDRigidsSolver::EnableRewindCapture()
Improved performance of async physics tick offset synchronization.
Made PlayerController only send redundant input history if using NetworkPhysicsComponent, or force enabling it via settings.

diff --git a/Engine/Plugins/Experimental/ChaosVehiclesPlugin/Source/ChaosVehicles/Private/ChaosVehicleManagerAsyncCallback.cpp b/Engine/Plugins/Experimental/ChaosVehiclesPlugin/Source/ChaosVehicles/Private/ChaosVehicleManagerAsyncCallback.cpp
--- a/Engine/Plugins/Experimental/ChaosVehiclesPlugin/Source/ChaosVehicles/Private/ChaosVehicleManagerAsyncCallback.cpp
+++ b/Engine/Plugins/Experimental/ChaosVehiclesPlugin/Source/ChaosVehicles/Private/ChaosVehicleManagerAsyncCallback.cpp
@@ -207,7 +207,7 @@ void FNetworkVehicleInputs::ApplyDatas(UActorComponent* NetworkComponent) const
 			}
 		}
 
-		if (NetworkComponent->GetWorld()->IsNetMode(NM_ListenServer))
+		if (NetworkComponent->GetWorld()->IsNetMode(NM_ListenServer) || NetworkComponent->GetWorld()->IsNetMode(NM_DedicatedServer))
 		{
 			UE_LOG(LogTemp, Log, TEXT("SERVER | PT | ApplyDatas | Report replicated inputs at frame %d %d: Throttle = %f Brake = %f Roll = %f Pitch = %f Yaw = %f Steering = %f Handbrake = %f Gear = %d | VehicleInputs size = %d | ControlInputs size = %d | NetworkInputs = %d"),
 				LocalFrame, SolverFrame, VehicleInputs.ThrottleInput, VehicleInputs.BrakeInput, VehicleInputs.RollInput, VehicleInputs.PitchInput,
@@ -242,7 +242,7 @@ void FNetworkVehicleInputs::BuildDatas(const UActorComponent* NetworkComponent)
 				}
 			}
 #if DEBUG_NETWORK_PHYSICS
-			if(NetworkComponent->GetWorld()->IsNetMode(NM_ListenServer))
+			if(NetworkComponent->GetWorld()->IsNetMode(NM_ListenServer) || NetworkComponent->GetWorld()->IsNetMode(NM_DedicatedServer))
 			{
 				UE_LOG(LogTemp, Log, TEXT("SERVER | PT | BuildDatas | Extract local inputs at frame %d : Throttle = %f Brake = %f Roll = %f Pitch = %f Yaw = %f Steering = %f Handbrake = %f Gear = %d"),
 					LocalFrame, VehicleInputs.ThrottleInput, VehicleInputs.BrakeInput, VehicleInputs.RollInput, VehicleInputs.PitchInput,
@@ -333,7 +333,7 @@ void FNetworkVehicleStates::ApplyDatas(UActorComponent* NetworkComponent) const
 			Vehicle->GetEngine().SetEngineOmega(EngineOmega);
 
 #if DEBUG_NETWORK_PHYSICS
-			if (NetworkComponent->GetWorld()->IsNetMode(NM_ListenServer))
+			if (NetworkComponent->GetWorld()->IsNetMode(NM_ListenServer) || NetworkComponent->GetWorld()->IsNetMode(NM_DedicatedServer))
 			{
 				UE_LOG(LogTemp, Log, TEXT("SERVER | PT | ApplyDatas | Report replicated states at frame %d %d : Omega = %f"),
 					LocalFrame, ServerFrame, EngineOmega);
@@ -378,7 +378,7 @@ void FNetworkVehicleStates::BuildDatas(const UActorComponent* NetworkComponent)
 			{
 				EngineOmega = Vehicle->GetEngine().GetEngineOmega();
 #if DEBUG_NETWORK_PHYSICS
-				if (NetworkComponent->GetWorld()->IsNetMode(NM_ListenServer))
+				if (NetworkComponent->GetWorld()->IsNetMode(NM_ListenServer) || NetworkComponent->GetWorld()->IsNetMode(NM_DedicatedServer))
 				{
 					UE_LOG(LogTemp, Log, TEXT("SERVER | PT | BuildDatas | Extract local states at frame %d %d : Omega = %f"),
 						LocalFrame, ServerFrame, EngineOmega);
diff --git a/Engine/Plugins/Experimental/ChaosVehiclesPlugin/Source/ChaosVehicles/Private/ChaosVehicleMovementComponent.cpp b/Engine/Plugins/Experimental/ChaosVehiclesPlugin/Source/ChaosVehicles/Private/ChaosVehicleMovementComponent.cpp
--- a/Engine/Plugins/Experimental/ChaosVehiclesPlugin/Source/ChaosVehicles/Private/ChaosVehicleMovementComponent.cpp
+++ b/Engine/Plugins/Experimental/ChaosVehiclesPlugin/Source/ChaosVehicles/Private/ChaosVehicleMovementComponent.cpp
@@ -146,7 +146,7 @@ void UChaosVehicleSimulation::TickVehicle(UWorld* WorldIn, float DeltaTime, cons
 	if (World && RigidHandle)
 	{
 #if DEBUG_NETWORK_PHYSICS
-		if (WorldIn->IsNetMode(NM_ListenServer))
+		if (WorldIn->IsNetMode(NM_ListenServer) || WorldIn->IsNetMode(NM_DedicatedServer))
 		{
 			UE_LOG(LogTemp, Log, TEXT("SERVER | PT | TickVehicle | Async tick vehicle with inputs at frame %d : Throttle = %f Brake = %f Roll = %f Pitch = %f Yaw = %f Steering = %f Handbrake = %f Sleeping = %d"),
 				InputData.PhysicsInputs.NetworkInputs.LocalFrame, VehicleInputs.ThrottleInput, VehicleInputs.BrakeInput, VehicleInputs.RollInput, VehicleInputs.PitchInput,
diff --git a/Engine/Source/Runtime/Engine/Classes/Engine/EngineTypes.h b/Engine/Source/Runtime/Engine/Classes/Engine/EngineTypes.h
--- a/Engine/Source/Runtime/Engine/Classes/Engine/EngineTypes.h
+++ b/Engine/Source/Runtime/Engine/Classes/Engine/EngineTypes.h
@@ -3011,18 +3011,15 @@ enum class EPhysicsReplicationMode : uint8
 	/** Default physics replication.*/
 	Default UMETA(DisplayName = "Default"),
 
-	/** Physics replication performing velocity interpolation.
+	/** Work In Progress. Physics replication performing velocity interpolation.
 	* Recommendation: Set on actors with a local role of ENetRole::ROLE_SimulatedProxy.
-	* Designed to handle local predictive interactions with other actors, especially actors of the local role ENetRole::ROLE_AutonomousProxy.
-	* Work in progress, available when Project Settings > Physics Prediction is enabled.
-	*/
+	* Designed to handle local predictive interactions with other actors, especially actors of the local role ENetRole::ROLE_AutonomousProxy. */
 	PredictiveInterpolation UMETA(DisplayName = "Predictive Interpolation (WIP)"),
 
-	/** Forward predicted replication by simulating physics and correcting errors through resimulating physics from a correct state in the past.
+	/** Work In Progress. Forward predicted replication by simulating physics and correcting errors through resimulating physics from a correct state in the past.
 	* Recommendation: Set on actors with a local role of ENetRole::ROLE_AutonomousProxy.
 	* Can be used for both ROLE_AutonomousProxy and ROLE_SimulatedProxy, though not recommended for ROLE_SimulatedProxy due to CPU performance.
-	* Work in progress, available when Project Settings > Physics Prediction > Resimulation is enabled.
-	*/
+	* Available when Project Settings > Physics Prediction is enabled. */
 	Resimulation UMETA(DisplayName = "Resimulation (WIP)"),
 };
 
diff --git a/Engine/Source/Runtime/Engine/Classes/GameFramework/PlayerController.h b/Engine/Source/Runtime/Engine/Classes/GameFramework/PlayerController.h
--- a/Engine/Source/Runtime/Engine/Classes/GameFramework/PlayerController.h
+++ b/Engine/Source/Runtime/Engine/Classes/GameFramework/PlayerController.h
@@ -2291,6 +2291,12 @@ public:
 	UFUNCTION(Server, unreliable)
 	ENGINE_API void ServerRecvClientInputFrame(int32 RecvClientInputFrame, const TArray<uint8>& Data);
 
+	/** When enabled the client will send the last N inputs to the server, N dictated by CVar: np2.NumRedundantCmds 3 */
+	ENGINE_API void EnableNetworkedPhysicsInputSync(bool EnablePrediction)
+	{
+		bSyncInputsForNetworkedPhysics = EnablePrediction;
+	}
+
 	const FClientFrameInfo& GetClientFrameInfo() const { return ClientFrameInfo; }
 	FClientFrameInfo& GetClientFrameInfo() { return ClientFrameInfo; }
 
@@ -2324,6 +2330,9 @@ private:
 	FClientFrameInfo ClientFrameInfo;
 	FServerFrameInfo ServerFrameInfo;
 
+	/** If true, inputs are synced from client to server in a way that's expected for NetworkPhysicsComponent */
+	bool bSyncInputsForNetworkedPhysics = false;
+
 	/** The estimated offset between the local async physics tick frame number and the server's
 	*	This is used to synchronize events that happen in the async physics tick */
 	int32 LocalToServerAsyncPhysicsTickOffset;
@@ -2347,8 +2356,8 @@ private:
 	/** The server records the latest timestamp it has to correct. This is used to update client (which may not happen on every physics step) */
 	FAsyncPhysicsTimestamp ServerLatestTimestampToCorrect;
 
-	/** The set of timestamps the client has sent to the server. Sorted by local frame number */
-	TArray<FAsyncPhysicsTimestamp> ServerPendingTimestamps;
+	/** The latest timestamp the client has sent to the server with prediction of which server frame it corresponds to. */
+	FAsyncPhysicsTimestamp ServerPendingTimestamp;
 
 	/** Update the tick ofsset in between the local client and the server */
 	ENGINE_API void UpdateServerAsyncPhysicsTickOffset();
diff --git a/Engine/Source/Runtime/Engine/Classes/PhysicsEngine/PhysicsSettings.h b/Engine/Source/Runtime/Engine/Classes/PhysicsEngine/PhysicsSettings.h
--- a/Engine/Source/Runtime/Engine/Classes/PhysicsEngine/PhysicsSettings.h
+++ b/Engine/Source/Runtime/Engine/Classes/PhysicsEngine/PhysicsSettings.h
@@ -91,11 +91,15 @@ struct FPhysicsPredictionSettings
 {
 	GENERATED_BODY();
 
-	/** Enable networked physics prediction */
+	/** Enable networked physics prediction (experimental)
+	* Note: If an AActor::PhysicsReplicationMode is set to use Resimulation this will allow physics to cache history which is required by resimulation replication.
+	* Note: This can also affect how physics is solved even when not using resimulation. */
 	UPROPERTY(EditAnywhere, Category = "Replication")
 	bool bEnablePhysicsPrediction;
 
-	/** Enable physics resimulation */
+	/** Forces the PlayerController to sync inputs as used in Physics Prediction.
+	* Only enable this if actively using a custom solution that needs this enabled for resimulation.
+	* This is automatically enabled when using the recommended NetworkPhysicsComponent on a pawn to handle Rewind / Resimulation. */
 	UPROPERTY(EditAnywhere, Category = "Replication", meta = (editcondition = "bEnablePhysicsPrediction"))
 	bool bEnablePhysicsResimulation;
 
diff --git a/Engine/Source/Runtime/Engine/Private/PhysicsEngine/Experimental/PhysScene_Chaos.cpp b/Engine/Source/Runtime/Engine/Private/PhysicsEngine/Experimental/PhysScene_Chaos.cpp
--- a/Engine/Source/Runtime/Engine/Private/PhysicsEngine/Experimental/PhysScene_Chaos.cpp
+++ b/Engine/Source/Runtime/Engine/Private/PhysicsEngine/Experimental/PhysScene_Chaos.cpp
@@ -98,7 +98,7 @@ public:
 		}
 
 		const UPhysicsSettings* PhysicsSettings = UPhysicsSettings::Get();
-		const bool bEnableResim = PhysicsSettings->PhysicsPrediction.bEnablePhysicsResimulation;
+		const bool bAllowResim = PhysicsSettings->PhysicsPrediction.bEnablePhysicsPrediction;
 		const int32 NumFrames = PhysicsSettings->GetPhysicsHistoryCount();
 
 		for (int32 Idx = 0; Idx < PendingCommands.Num(); ++Idx)
@@ -108,7 +108,7 @@ public:
 			FPendingAsyncPhysicsCommand* PendingCommand = &PendingCommands[Idx];
 			bool bRemove = PendingCommand->OwningObject.IsStale();
 
-			if (bEnableResim && PendingCommand->bEnableResim && PendingCommands[Idx].PhysicsStep > (CurrentFrame - NumFrames))
+			if (bAllowResim && PendingCommand->bEnableResim && PendingCommands[Idx].PhysicsStep > (CurrentFrame - NumFrames))
 			{
 				if (PendingCommands[Idx].PhysicsStep < CurrentFrame && PendingCommands[Idx].OwningObject.IsValid())
 				{
diff --git a/Engine/Source/Runtime/Engine/Private/PhysicsEngine/NetworkPhysicsComponent.cpp b/Engine/Source/Runtime/Engine/Private/PhysicsEngine/NetworkPhysicsComponent.cpp
--- a/Engine/Source/Runtime/Engine/Private/PhysicsEngine/NetworkPhysicsComponent.cpp
+++ b/Engine/Source/Runtime/Engine/Private/PhysicsEngine/NetworkPhysicsComponent.cpp
@@ -172,7 +172,7 @@ int32 FNetworkPhysicsCallback::TriggerRewindIfNeeded_Internal(int32 LatestStepCo
 	}
 
 #if DEBUG_NETWORK_PHYSICS
-	UE_LOG(LogTemp, Log, TEXT("CLIENT | PT | TriggerRewindIfNeeded_Internal | Callbacks Frame = %d"), ResimFrame);
+	UE_LOG(LogTemp, Log, TEXT("COMMON | PT | TriggerRewindIfNeeded_Internal | Callbacks Frame = %d"), ResimFrame);
 #endif
 
 	if (RewindData)
@@ -367,8 +367,7 @@ void UNetworkPhysicsSystem::OnWorldPostInit(UWorld* World, const UWorld::Initial
 			{ 
 				if (Solver->GetRewindCallback() == nullptr)
 				{
-					const int32 NumFrames = FMath::Max<int32>(1, UPhysicsSettings::Get()->GetPhysicsHistoryCount());
-					Solver->EnableRewindCapture(NumFrames, true, MakeUnique<FNetworkPhysicsCallback>(World));
+					Solver->SetRewindCallback(MakeUnique<FNetworkPhysicsCallback>(World));
 				}
 			}
 		}
@@ -416,10 +415,28 @@ void UNetworkPhysicsComponent::BeginPlay()
 		{
 			if (Chaos::FPhysicsSolver* Solver = PhysScene->GetSolver())
 			{
-				if(FNetworkPhysicsCallback* SolverCallback = static_cast<FNetworkPhysicsCallback*>(Solver->GetRewindCallback()))
+				if (UPhysicsSettings::Get()->PhysicsPrediction.bEnablePhysicsPrediction)
 				{
-					SolverCallback->PreProcessInputsInternal.AddUObject(this, &UNetworkPhysicsComponent::OnPreProcessInputsInternal);
-					SolverCallback->PostProcessInputsInternal.AddUObject(this, &UNetworkPhysicsComponent::OnPostProcessInputsInternal);
+					if (APlayerController* PlayerController = GetPlayerController())
+					{
+						PlayerController->EnableNetworkedPhysicsInputSync(true);
+					}
+
+					if (Solver->GetRewindData() == nullptr)
+					{
+						const int32 NumFrames = FMath::Max<int32>(1, UPhysicsSettings::Get()->GetPhysicsHistoryCount());
+						Solver->EnableRewindCapture(NumFrames, true);
+					}
+
+					if(FNetworkPhysicsCallback* SolverCallback = static_cast<FNetworkPhysicsCallback*>(Solver->GetRewindCallback()))
+					{
+						SolverCallback->PreProcessInputsInternal.AddUObject(this, &UNetworkPhysicsComponent::OnPreProcessInputsInternal);
+						SolverCallback->PostProcessInputsInternal.AddUObject(this, &UNetworkPhysicsComponent::OnPostProcessInputsInternal);
+					}
+				}
+				else
+				{
+					UE_LOG(LogPhysics, Warning, TEXT("A NetworkPhysicsComponent is trying to set up but 'Project Settings -> Physics -> Physics Prediction' is not enabled. The component might not work as intended."));
 				}
 			}
 		}
diff --git a/Engine/Source/Runtime/Engine/Private/PhysicsEngine/PhysLevel.cpp b/Engine/Source/Runtime/Engine/Private/PhysicsEngine/PhysLevel.cpp
--- a/Engine/Source/Runtime/Engine/Private/PhysicsEngine/PhysLevel.cpp
+++ b/Engine/Source/Runtime/Engine/Private/PhysicsEngine/PhysLevel.cpp
@@ -38,6 +38,7 @@ public:
 	{
 
 	}
+
 	virtual float GetMinDeltaVelocityForHitEvents() const override
 	{
 		return GetSettings()->MinDeltaVelocityForHitEvents;
@@ -47,14 +48,12 @@ public:
 	{
 		return GetSettings()->PhysicsPrediction.bEnablePhysicsPrediction;
 	}
-	virtual bool GetPhysicsResimulationEnabled() const override
-	{
-		return GetSettings()->PhysicsPrediction.bEnablePhysicsResimulation;
-	}
+
 	virtual float GetResimulationErrorThreshold() const override
 	{
 		return GetSettings()->PhysicsPrediction.ResimulationErrorThreshold;
 	}
+
 	virtual int32 GetPhysicsHistoryCount() const override
 	{
 		return GetSettings()->GetPhysicsHistoryCount();
diff --git a/Engine/Source/Runtime/Engine/Private/PhysicsEngine/PhysicsReplication.cpp b/Engine/Source/Runtime/Engine/Private/PhysicsEngine/PhysicsReplication.cpp
--- a/Engine/Source/Runtime/Engine/Private/PhysicsEngine/PhysicsReplication.cpp
+++ b/Engine/Source/Runtime/Engine/Private/PhysicsEngine/PhysicsReplication.cpp
@@ -711,7 +711,8 @@ void FPhysicsReplicationAsync::OnPreSimulate_Internal()
 
 		// Early out if this is a resim frame
 		Chaos::FRewindData* RewindData = RigidsSolver->GetRewindData();
-		if (RewindData && RewindData->IsResim())
+		const bool bRewindDataExist = RewindData != nullptr;
+		if (bRewindDataExist && RewindData->IsResim())
 		{
 			// TODO, Handle the transition from post-resim to interpolation better.
 			if (PhysicsReplicationCVars::PredictiveInterpolationCVars::PostResimWaitForUpdate && RewindData->IsFinalResim())
@@ -742,6 +743,16 @@ void FPhysicsReplicationAsync::OnPreSimulate_Internal()
 				continue;
 			}
 
+			if (!bRewindDataExist && Input.RepMode == EPhysicsReplicationMode::Resimulation)
+			{
+				// We don't have rewind data but an actor is set to replicate using resimulation; we need to enable rewind capture.
+				if (ensure(Chaos::FPBDRigidsSolver::IsNetworkPhysicsPredictionEnabled()))
+				{
+					const int32 NumFrames = FMath::Max<int32>(1, Chaos::FPBDRigidsSolver::GetPhysicsHistoryCount());
+					RigidsSolver->EnableRewindCapture(NumFrames, true);
+				}
+			}
+
 			UpdateRewindDataTarget(Input);
 			UpdateAsyncTarget(Input, RigidsSolver);
 		}
@@ -857,12 +868,6 @@ void FPhysicsReplicationAsync::ApplyTargetStatesAsync(const float DeltaSeconds,
 			{
 				EPhysicsReplicationMode RepMode = Target.RepMode;
 
-				// TODO, Remove the resim option from project settings, we only need the physics prediction one now
-				if (!Chaos::FPBDRigidsSolver::IsPhysicsResimulationEnabled() && RepMode == EPhysicsReplicationMode::Resimulation)
-				{
-					RepMode = EPhysicsReplicationMode::Default;
-				}
-
 				switch (RepMode)
 				{
 					case EPhysicsReplicationMode::Default:
diff --git a/Engine/Source/Runtime/Engine/Private/PlayerController.cpp b/Engine/Source/Runtime/Engine/Private/PlayerController.cpp
--- a/Engine/Source/Runtime/Engine/Private/PlayerController.cpp
+++ b/Engine/Source/Runtime/Engine/Private/PlayerController.cpp
@@ -121,6 +121,9 @@ namespace NetworkPhysicsCvars
 
 	int32 NetworkPhysicsPredictionFrameOffset = 4;
 	FAutoConsoleVariableRef CVarNetworkPhysicsPredictionFrameOffset(TEXT("np2.NetworkPhysicsPredictionFrameOffset"), NetworkPhysicsPredictionFrameOffset, TEXT("Additional frame offset to be added to the local to server offset used by network prediction"));
+
+	int32 TickOffsetUpdateInterval = 10;
+	FAutoConsoleVariableRef CVarTickOffsetUpdateInterval(TEXT("np2.TickOffsetUpdateInterval"), TickOffsetUpdateInterval, TEXT("How many physics ticks to wait between each tick offset update. Lowest viable value = 1, which means update each tick."));
 }
 
 const float RetryClientRestartThrottleTime = 0.5f;
@@ -1673,20 +1676,18 @@ void APlayerController::ClientSetCameraFade_Implementation(bool bEnableFading, F
 
 void APlayerController::SendClientAdjustment()
 {
-	if(!UPhysicsSettings::Get()->PhysicsPrediction.bEnablePhysicsPrediction)
+	if (ServerFrameInfo.LastProcessedInputFrame != INDEX_NONE && ServerFrameInfo.LastProcessedInputFrame != ServerFrameInfo.LastSentLocalFrame)
 	{
-		if (ServerFrameInfo.LastProcessedInputFrame != INDEX_NONE && ServerFrameInfo.LastProcessedInputFrame != ServerFrameInfo.LastSentLocalFrame)
-		{
-			ServerFrameInfo.LastSentLocalFrame = ServerFrameInfo.LastProcessedInputFrame;		
-			ClientRecvServerAckFrame(ServerFrameInfo.LastProcessedInputFrame, ServerFrameInfo.LastLocalFrame, ServerFrameInfo.QuantizedTimeDilation);
+		ServerFrameInfo.LastSentLocalFrame = ServerFrameInfo.LastProcessedInputFrame;		
+		ClientRecvServerAckFrame(ServerFrameInfo.LastProcessedInputFrame, ServerFrameInfo.LastLocalFrame, ServerFrameInfo.QuantizedTimeDilation);
 
-			if (NetworkPhysicsCvars::EnableDebugRPC)
-			{
-				ClientRecvServerAckFrameDebug(InputBuffer.HeadFrame() - ServerFrameInfo.LastProcessedInputFrame, ServerFrameInfo.TargetNumBufferedCmds);
-			}
+		if (NetworkPhysicsCvars::EnableDebugRPC)
+		{
+			ClientRecvServerAckFrameDebug(InputBuffer.HeadFrame() - ServerFrameInfo.LastProcessedInputFrame, ServerFrameInfo.TargetNumBufferedCmds);
 		}
 	}
-	else
+
+	if (UPhysicsSettings::Get()->PhysicsPrediction.bEnablePhysicsPrediction)
 	{
 		if (ServerLatestTimestampToCorrect.ServerFrame == INDEX_NONE)
 		{
@@ -1718,6 +1719,11 @@ void APlayerController::SendClientAdjustment()
 
 void APlayerController::PushClientInput(int32 InRecvClientInputFrame, TArray<uint8>& Data)
 {
+	if (!bSyncInputsForNetworkedPhysics && !UPhysicsSettings::Get()->PhysicsPrediction.bEnablePhysicsResimulation)
+	{
+		return;
+	}
+
 	InputBuffer.Write(InRecvClientInputFrame) = MoveTemp(Data);
 
 	// Do the RPC right here, including the redundant send. This should probably be time based and managed somewhere else like in Tick eventually
@@ -6137,8 +6143,9 @@ void APlayerController::UpdateServerAsyncPhysicsTickOffset()
 		}
 	}
 	FAsyncPhysicsTimestamp Timestamp = GetAsyncPhysicsTimestamp();
-	if(ClientLatestAsyncPhysicsStepSent == Timestamp.LocalFrame)
+	if(ClientLatestAsyncPhysicsStepSent + NetworkPhysicsCvars::TickOffsetUpdateInterval > Timestamp.LocalFrame)
 	{
+		//Only send a new timestamp if enough physics ticks have passed, based on CVar.
 		//If GT is running faster than physics sim the physics timestep will not have changed, so no need to send another update to server
 		//This ensures monotonic increase
 		return;
@@ -6151,17 +6158,11 @@ void APlayerController::UpdateServerAsyncPhysicsTickOffset()
 void APlayerController::ServerSendLatestAsyncPhysicsTimestamp_Implementation(FAsyncPhysicsTimestamp Timestamp)
 {
 	//This tells the server how the client thinks the async physics tick will line up.
-	//Timestamps could be out of order due to networking, so we make sure they are sorted using the client's local frame which is monotonically increasing
-	int32 Idx;
-	for(Idx = ServerPendingTimestamps.Num() - 1; Idx >= 0; --Idx)
+	//Only cache the most up to date timestamp based on LocalFrame
+	if (Timestamp.LocalFrame > ServerPendingTimestamp.LocalFrame)
 	{
-		ensureMsgf(ServerPendingTimestamps[Idx].LocalFrame != Timestamp.LocalFrame, TEXT("Client should never send duplicate timestamps, something is wrong"));
-		if(ServerPendingTimestamps[Idx].LocalFrame < Timestamp.LocalFrame)
-		{
-			break;
-		}
+		ServerPendingTimestamp = Timestamp;
 	}
-	ServerPendingTimestamps.Insert(Timestamp, Idx + 1);
 }
 
 void APlayerController::ClientCorrectionAsyncPhysicsTimestamp_Implementation(FAsyncPhysicsTimestamp Timestamp)
@@ -6220,19 +6221,21 @@ void APlayerController::AsyncPhysicsTickActor(float DeltaTime, float SimTime)
 	if(UPhysicsSettings::Get()->PhysicsPrediction.bEnablePhysicsPrediction)
 	{
 		//TODO: only kick this off if server and using this feature
-		if (IsLocalController()) { return; }
-
-		if (ServerPendingTimestamps.Num() == 0)
+		if (IsLocalController())
 		{
-			//TODO: starved for input user needs to speed up
 			return;
 		}
 
+		// If the client has not sent the correct estimate of which server frame its local frame would correspond to, send back which server frame that local frame actually corresponded to
 		const FAsyncPhysicsTimestamp ActualTimestamp = GetAsyncPhysicsTimestamp();
+		
+		if (ServerPendingTimestamp.ServerFrame != INDEX_NONE && ServerPendingTimestamp.ServerFrame != ActualTimestamp.ServerFrame)
+		{
+			ServerLatestTimestampToCorrect.ServerFrame = ActualTimestamp.ServerFrame;
+			ServerLatestTimestampToCorrect.LocalFrame = ServerPendingTimestamp.LocalFrame;
+			ServerPendingTimestamp.ServerFrame = INDEX_NONE;
+		}
 
-		ServerLatestTimestampToCorrect.ServerFrame = ActualTimestamp.ServerFrame;
-		ServerLatestTimestampToCorrect.LocalFrame = ServerPendingTimestamps[0].LocalFrame;
-		ServerPendingTimestamps.Reset();
 	}
 }
 
diff --git a/Engine/Source/Runtime/Experimental/Chaos/Private/Chaos/Collision/PBDCollisionConstraint.cpp b/Engine/Source/Runtime/Experimental/Chaos/Private/Chaos/Collision/PBDCollisionConstraint.cpp
--- a/Engine/Source/Runtime/Experimental/Chaos/Private/Chaos/Collision/PBDCollisionConstraint.cpp
+++ b/Engine/Source/Runtime/Experimental/Chaos/Private/Chaos/Collision/PBDCollisionConstraint.cpp
@@ -1055,7 +1055,7 @@ namespace Chaos
 	{
 		int32 MatchIndex = INDEX_NONE;
 
-		const bool bEnableNetworkPhysicsResim = FPhysicsSolverBase::IsNetworkPhysicsPredictionEnabled() && FPhysicsSolverBase::IsPhysicsResimulationEnabled();
+		const bool bEnableNetworkPhysicsResim = FPhysicsSolverBase::IsNetworkPhysicsPredictionEnabled();
 		if (bChaos_Manifold_EnableFrictionRestore && !bEnableNetworkPhysicsResim)
 		{
 			const FManifoldPoint& ManifoldPoint = ManifoldPoints[ManifoldPointIndex];
diff --git a/Engine/Source/Runtime/Experimental/Chaos/Private/Chaos/CollisionOneShotManifolds.cpp b/Engine/Source/Runtime/Experimental/Chaos/Private/Chaos/CollisionOneShotManifolds.cpp
--- a/Engine/Source/Runtime/Experimental/Chaos/Private/Chaos/CollisionOneShotManifolds.cpp
+++ b/Engine/Source/Runtime/Experimental/Chaos/Private/Chaos/CollisionOneShotManifolds.cpp
@@ -787,7 +787,7 @@ namespace Chaos
 			const FReal Margin2 = Constraint.GetCollisionMargin1();
 			const FRigidTransform3 Convex2ToConvex1Transform = Convex2Transform.GetRelativeTransformNoScale(Convex1Transform);
 
-			const bool bEnableNetworkPhysicsResim = FPhysicsSolverBase::IsNetworkPhysicsPredictionEnabled() && FPhysicsSolverBase::IsPhysicsResimulationEnabled();
+			const bool bEnableNetworkPhysicsResim = FPhysicsSolverBase::IsNetworkPhysicsPredictionEnabled();
 			
 			if (bEnableNetworkPhysicsResim)
 			{
diff --git a/Engine/Source/Runtime/Experimental/Chaos/Private/Chaos/Framework/PhysicsSolverBase.cpp b/Engine/Source/Runtime/Experimental/Chaos/Private/Chaos/Framework/PhysicsSolverBase.cpp
--- a/Engine/Source/Runtime/Experimental/Chaos/Private/Chaos/Framework/PhysicsSolverBase.cpp
+++ b/Engine/Source/Runtime/Experimental/Chaos/Private/Chaos/Framework/PhysicsSolverBase.cpp
@@ -368,8 +368,13 @@ namespace Chaos
 
 	void FPhysicsSolverBase::SetRewindCallback(TUniquePtr<IRewindCallback>&& RewindCallback)
 	{
-		ensure(!RewindCallback || MRewindData);
+		ensure(RewindCallback);
 		MRewindCallback = MoveTemp(RewindCallback);
+
+		if (MRewindData.IsValid())
+		{
+			MRewindCallback->RewindData = MRewindData.Get();
+		}
 	}
 
 	FGraphEventRef FPhysicsSolverBase::AdvanceAndDispatch_External(FReal InDt)
@@ -445,7 +450,7 @@ namespace Chaos
 
 		while(FPushPhysicsData* PushData = MarshallingManager.StepInternalTime_External())
 		{
-			if(MRewindCallback && !bIsShuttingDown)
+			if(ShouldApplyRewindCallbacks() && !bIsShuttingDown)
 			{
 				MRewindCallback->ProcessInputs_External(PushData->InternalStep, PushData->SimCallbackInputs);
 			}
diff --git a/Engine/Source/Runtime/Experimental/Chaos/Private/PBDRigidsSolver.cpp b/Engine/Source/Runtime/Experimental/Chaos/Private/PBDRigidsSolver.cpp
--- a/Engine/Source/Runtime/Experimental/Chaos/Private/PBDRigidsSolver.cpp
+++ b/Engine/Source/Runtime/Experimental/Chaos/Private/PBDRigidsSolver.cpp
@@ -664,11 +664,11 @@ namespace Chaos
 	{
 		Super::ApplyCallbacks_Internal();
 
-		if (IRewindCallback* RewindCallback = GetRewindCallback())
+		if (MRewindCallback) // Note: Don't use ShouldApplyRewindCallbacks() here since we want this called even if we don't have RewindData enabled.
 		{
-			if(bGameThreadFrozen)
+			if (bGameThreadFrozen)
 			{ 
-				RewindCallback->ApplyCallbacks_Internal(GetCurrentFrame(), SimCallbackObjects);
+				MRewindCallback->ApplyCallbacks_Internal(GetCurrentFrame(), SimCallbackObjects);
 			}
 		}
 	}
@@ -725,7 +725,7 @@ namespace Chaos
 				});
 
 			FGeometryParticleHandle* Handle = Proxy->GetHandle_LowLevel();
-			if(MRewindData)
+			if (MRewindData)
 			{
 				MRewindData->RemoveObject(Handle);
 			}
@@ -989,6 +989,14 @@ namespace Chaos
 
 	void FPBDRigidsSolver::EnableRewindCapture(int32 NumFrames, bool InUseCollisionResimCache, TUniquePtr<IRewindCallback>&& RewindCallback)
 	{
+		SetRewindCallback(MoveTemp(RewindCallback));
+		EnableRewindCapture(NumFrames, InUseCollisionResimCache);
+	}
+
+	void FPBDRigidsSolver::EnableRewindCapture(int32 NumFrames, bool InUseCollisionResimCache)
+	{
+		UE_LOG(LogChaos, Log, TEXT("PBDRigidsSolver::EnableRewindCapture - Starting physics data history caching for rewind / resimulation."));
+
 		//TODO: this function calls both internal and external - sort of assumed during initialization. Should decide what thread it's called on and mark it as either external or internal
 		if (MRewindData.IsValid())
 		{
@@ -999,11 +1007,11 @@ namespace Chaos
 			MRewindData = MakeUnique<FRewindData>(((FPBDRigidsSolver*)this), NumFrames, InUseCollisionResimCache, ((FPBDRigidsSolver*)this)->GetCurrentFrame()); // FIXME
 		}
 		bUseCollisionResimCache = InUseCollisionResimCache;
-		MRewindCallback = MoveTemp(RewindCallback);
 		const int32 NumFramesSet = GetRewindData() != nullptr ? GetRewindData()->Capacity() : NumFrames;
 		MarshallingManager.SetHistoryLength_Internal(NumFramesSet);
 		MEvolution->SetRewindData(GetRewindData());
-		if(MRewindCallback)
+		
+		if (MRewindCallback) 
 		{
 			MRewindCallback->RewindData = MRewindData.Get();
 		}
@@ -1028,12 +1036,11 @@ namespace Chaos
 		//todo: do we need this?
 		//MarshallingManager.Reset();
 
-		const bool PhysicsPredictionEnabled = FChaosSolversModule::GetModule()->GetSettingsProvider().GetPhysicsPredictionEnabled();
 		const int32 PhysicsHistoryLength = FChaosSolversModule::GetModule()->GetSettingsProvider().GetPhysicsHistoryCount();
 
-		if ((PhysicsPredictionEnabled || bUseCollisionResimCache) && PhysicsHistoryLength >= 0)
+		if (bUseCollisionResimCache && PhysicsHistoryLength >= 0)
 		{
-			EnableRewindCapture(PhysicsHistoryLength, bUseCollisionResimCache || PhysicsPredictionEnabled);
+			EnableRewindCapture(PhysicsHistoryLength, true);
 		}
 
 		MEvolution->SetCaptureRewindDataFunction([this](const TParticleView<TPBDRigidParticles<FReal,3>>& ActiveParticles)
@@ -1404,7 +1411,7 @@ namespace Chaos
 
 		GetEvolution()->GetBroadPhase().GetIgnoreCollisionManager().PushProducerStorageData_External(MarshallingManager.GetExternalTimestamp_External());
 
-		if (MRewindCallback && !IsShuttingDown())
+		if (ShouldApplyRewindCallbacks() && !IsShuttingDown())
 		{
 			MRewindCallback->InjectInputs_External(MarshallingManager.GetInternalStep_External(), NumSteps);
 		}
@@ -1668,10 +1675,15 @@ namespace Chaos
 
 			if (SimCallbackObject->HasOption(ESimCallbackOptions::Rewind))
 			{
-				if (MRewindCallback.IsValid())
+				if (MRewindCallback)
 				{
 					MRewindCallback->RegisterRewindableSimCallback_Internal(SimCallbackObject);
 				}
+				else
+				{
+					ensure(!IsNetworkPhysicsPredictionEnabled());
+					UE_LOG(LogChaos, Warning, TEXT("A SimCallbackObject with ESimCallbackOptions::Rewind defined is registering but there is no IRewindCallback set up. Make sure to enable 'Project Settings -> Physics -> Physics Prediction'"));
+				}
 			}
 
 			if (SimCallbackObject->HasOption(ESimCallbackOptions::PhysicsObjectUnregister))
@@ -1778,7 +1790,7 @@ namespace Chaos
 		}
 		PushData.SimCommands.Reset();
 
-		if(MRewindCallback && !IsShuttingDown())
+		if(MRewindCallback && MRewindData && !IsShuttingDown())
 		{
 			MRewindCallback->ProcessInputs_Internal(MRewindData->CurrentFrame(), PushData.SimCallbackInputs);
 		}
@@ -1787,7 +1799,7 @@ namespace Chaos
 	void FPBDRigidsSolver::ConditionalApplyRewind_Internal()
 	{
 		// Note: checking MRewindData->IsResim() can lead to recursion into this function on the last resim frame since the call to AdvanceSolver is what advances RewindData's internal frame
-		if(!IsShuttingDown() && MRewindCallback && !GetEvolution()->IsResimming())
+		if(!IsShuttingDown() && ShouldApplyRewindCallbacks() && MRewindData && !GetEvolution()->IsResimming())
 		{
 			const int32 LastStep = MRewindData->CurrentFrame() - 1;
 			const int32 ResimStep = MRewindCallback->TriggerRewindIfNeeded_Internal(LastStep);
diff --git a/Engine/Source/Runtime/Experimental/Chaos/Public/Chaos/Framework/PhysicsSolverBase.h b/Engine/Source/Runtime/Experimental/Chaos/Public/Chaos/Framework/PhysicsSolverBase.h
--- a/Engine/Source/Runtime/Experimental/Chaos/Public/Chaos/Framework/PhysicsSolverBase.h
+++ b/Engine/Source/Runtime/Experimental/Chaos/Public/Chaos/Framework/PhysicsSolverBase.h
@@ -398,6 +398,11 @@ namespace Chaos
 			return MRewindCallback.Get();
 		}
 
+		bool ShouldApplyRewindCallbacks()
+		{
+			return MRewindCallback.IsValid() && MRewindData.IsValid();
+		}
+
 		//Used as helper for GT to go from unique idx back to gt particle
 		//If GT deletes a particle, this function will return null (that's a good thing when consuming async outputs as GT may have already deleted the particle we care about)
 		//Note: if the physics solver has been advanced after the particle was freed on GT, the index may have been freed and reused.
@@ -662,11 +667,11 @@ namespace Chaos
 			return NetworkPhysicsEnabled;
 		}
 
-		/** Check if resim is enabled for network physics */
-		static bool IsPhysicsResimulationEnabled()
+		/** Get the number of physics history frames to cache */
+		static int32 GetPhysicsHistoryCount()
 		{
-			const bool PhysicsResimulationEnabled = FChaosSolversModule::GetModule()->GetSettingsProvider().GetPhysicsResimulationEnabled();
-			return PhysicsResimulationEnabled;
+			const int32 PhysicsHistoryCount = FChaosSolversModule::GetModule()->GetSettingsProvider().GetPhysicsHistoryCount();
+			return PhysicsHistoryCount;
 		}
 
 		static float ResimulationErrorThreshold()
diff --git a/Engine/Source/Runtime/Experimental/Chaos/Public/ChaosSolversModule.h b/Engine/Source/Runtime/Experimental/Chaos/Public/ChaosSolversModule.h
--- a/Engine/Source/Runtime/Experimental/Chaos/Public/ChaosSolversModule.h
+++ b/Engine/Source/Runtime/Experimental/Chaos/Public/ChaosSolversModule.h
@@ -40,19 +40,16 @@ public:
 	{
 		return false;
 	}
-	virtual bool GetPhysicsResimulationEnabled() const 
-	{
-		return false;
-	}
+
 	virtual float GetResimulationErrorThreshold() const
 	{
 		return 0;
 	}
+
 	virtual int32 GetPhysicsHistoryCount() const
 	{
 		return 0;
 	}
-	
 }; 
 
 namespace Chaos
diff --git a/Engine/Source/Runtime/Experimental/Chaos/Public/PBDRigidsSolver.h b/Engine/Source/Runtime/Experimental/Chaos/Public/PBDRigidsSolver.h
--- a/Engine/Source/Runtime/Experimental/Chaos/Public/PBDRigidsSolver.h
+++ b/Engine/Source/Runtime/Experimental/Chaos/Public/PBDRigidsSolver.h
@@ -221,7 +221,8 @@ namespace Chaos
 		FCharacterGroundConstraintContainer& GetCharacterGroundConstraints() { return MEvolution->GetCharacterGroundConstraints(); }
 		const FCharacterGroundConstraintContainer& GetCharacterGroundConstraints() const { return MEvolution->GetCharacterGroundConstraints(); }
 
-		CHAOS_API void EnableRewindCapture(int32 NumFrames, bool InUseCollisionResimCache, TUniquePtr<IRewindCallback>&& RewindCallback = TUniquePtr<IRewindCallback>());
+		CHAOS_API void EnableRewindCapture(int32 NumFrames, bool InUseCollisionResimCache, TUniquePtr<IRewindCallback>&& RewindCallback);
+		CHAOS_API void EnableRewindCapture(int32 NumFrames, bool InUseCollisionResimCache);
 
 		/**/
 		FPBDRigidsEvolution* GetEvolution() { return MEvolution.Get(); }
