Converted UWorldPartitionLevelStreamingDynamic's OuterWorldPartition to a weak object ptr to fix a GC reporting a referenced garbage object (the OuterWorldPartition).
An unloaded instanced partitioned world could be marked for garbage while an injected world partition streaming cell created outside of this world was still referencing the WorldPartition through this OuterWorldPartition pointer.

diff --git a/Engine/Source/Runtime/Engine/Private/WorldPartition/WorldPartitionLevelStreamingDynamic.cpp b/Engine/Source/Runtime/Engine/Private/WorldPartition/WorldPartitionLevelStreamingDynamic.cpp
--- a/Engine/Source/Runtime/Engine/Private/WorldPartition/WorldPartitionLevelStreamingDynamic.cpp
+++ b/Engine/Source/Runtime/Engine/Private/WorldPartition/WorldPartitionLevelStreamingDynamic.cpp
@@ -647,7 +647,7 @@ UWorld* UWorldPartitionLevelStreamingDynamic::GetStreamingWorld() const
 {
 	// For UWorldPartitionLevelStreamingDynamic the StreamingWorld is the world to which the OuterWorldPartition is outered.
 	// This World can be used to resolved SoftObjectPaths between cells.
-	check(OuterWorldPartition);
+	check(OuterWorldPartition.IsValid());
 	return OuterWorldPartition->GetTypedOuter<UWorld>();
 }
 
@@ -703,7 +703,7 @@ bool UWorldPartitionLevelStreamingDynamic::ShouldBlockOnUnload() const
 
 	// When world partition cannot stream (anymore), return true so that RemoveFromWorld of this level is not incremental. 
 	// This guarantees that unloaded instanced wp levels fully unload their cell levels.
-	if (OuterWorldPartition && !OuterWorldPartition->CanStream())
+	if (OuterWorldPartition.IsValid() && !OuterWorldPartition->CanStream())
 	{
 		return true;
 	}
@@ -714,7 +714,7 @@ bool UWorldPartitionLevelStreamingDynamic::ShouldBlockOnUnload() const
 void UWorldPartitionLevelStreamingDynamic::UpdateShouldSkipMakingVisibilityTransactionRequest()
 {
 	const UWorldPartitionRuntimeCell* Cell = GetWorldPartitionRuntimeCell();
-	if (ensure(Cell && OuterWorldPartition))
+	if (ensure(Cell && OuterWorldPartition.IsValid()))
 	{
 		// It is safe to skip client MakingVisibility transaction requests for cells without data layers when world partition server streaming is disabled
 		bSkipClientUseMakingVisibleTransactionRequest = !(OuterWorldPartition->IsServerStreamingEnabled() || Cell->HasDataLayers());
diff --git a/Engine/Source/Runtime/Engine/Public/WorldPartition/WorldPartitionLevelStreamingDynamic.h b/Engine/Source/Runtime/Engine/Public/WorldPartition/WorldPartitionLevelStreamingDynamic.h
--- a/Engine/Source/Runtime/Engine/Public/WorldPartition/WorldPartitionLevelStreamingDynamic.h
+++ b/Engine/Source/Runtime/Engine/Public/WorldPartition/WorldPartitionLevelStreamingDynamic.h
@@ -105,5 +105,5 @@ private:
 	TWeakObjectPtr<const UWorldPartitionRuntimeLevelStreamingCell> StreamingCell;
 
 	UPROPERTY()
-	TObjectPtr<UWorldPartition> OuterWorldPartition;
+	TWeakObjectPtr<UWorldPartition> OuterWorldPartition;
 };
