- Added virtual UWorldPartitionRuntimeCell::ShouldServerWaitForClientLevelVisibility which can be overriden to allow for a server cell to ignore client level visibility before unloading the cell.
- WorldPartitionStreamingPolicy keeps UPROPERTY references for activated/loaded cells to avoid any deinjected cells to be GC'ed.
- Deprecated IsAddedToWorld and CanAddToWorld in UWorldPartitionRuntimeCell as they are no longer used
- Removed unused SortedAddToWorldCells in UWorldPartitionStreamingPolicy

diff --git a/Engine/Source/Runtime/Engine/Private/WorldPartition/WorldPartitionRuntimeLevelStreamingCell.cpp b/Engine/Source/Runtime/Engine/Private/WorldPartition/WorldPartitionRuntimeLevelStreamingCell.cpp
--- a/Engine/Source/Runtime/Engine/Private/WorldPartition/WorldPartitionRuntimeLevelStreamingCell.cpp
+++ b/Engine/Source/Runtime/Engine/Private/WorldPartition/WorldPartitionRuntimeLevelStreamingCell.cpp
@@ -452,18 +452,6 @@ void UWorldPartitionRuntimeLevelStreamingCell::Activate() const
 	}
 }
 
-bool UWorldPartitionRuntimeLevelStreamingCell::IsAddedToWorld() const
-{
-	return LevelStreaming && LevelStreaming->GetLoadedLevel() && LevelStreaming->GetLoadedLevel()->bIsVisible;
-}
-
-bool UWorldPartitionRuntimeLevelStreamingCell::CanAddToWorld() const
-{
-	return LevelStreaming &&
-		   LevelStreaming->GetLoadedLevel() &&
-		   (LevelStreaming->GetLevelStreamingState() == ELevelStreamingState::MakingVisible);
-}
-
 void UWorldPartitionRuntimeLevelStreamingCell::SetStreamingPriority(int32 InStreamingPriority) const
 {
 	if (LevelStreaming)
diff --git a/Engine/Source/Runtime/Engine/Private/WorldPartition/WorldPartitionStreamingPolicy.cpp b/Engine/Source/Runtime/Engine/Private/WorldPartition/WorldPartitionStreamingPolicy.cpp
--- a/Engine/Source/Runtime/Engine/Private/WorldPartition/WorldPartitionStreamingPolicy.cpp
+++ b/Engine/Source/Runtime/Engine/Private/WorldPartition/WorldPartitionStreamingPolicy.cpp
@@ -455,13 +455,16 @@ void UWorldPartitionStreamingPolicy::UpdateStreamingState()
 	auto ShouldWaitForClientVisibility = [bIsServer, this, &bUpdateServerEpoch, &ServerClientsVisibleLevelNames](const UWorldPartitionRuntimeCell* Cell)
 	{
 		check(bIsServer);
-		if (ULevel* Level = Cell->GetLevel())
+		if (Cell->ShouldServerWaitForClientLevelVisibility())
 		{
-			if (ServerClientsVisibleLevelNames.Contains(Cell->GetLevel()->GetPackage()->GetFName()))
+			if (ULevel* Level = Cell->GetLevel())
 			{
-				UE_CLOG(bUpdateServerEpoch, LogWorldPartition, Verbose, TEXT("Server epoch update delayed by client visibility"));
-				bUpdateServerEpoch = false;
-				return true;
+				if (ServerClientsVisibleLevelNames.Contains(Level->GetPackage()->GetFName()))
+				{
+					UE_CLOG(bUpdateServerEpoch, LogWorldPartition, Verbose, TEXT("Server epoch update delayed by client visibility"));
+					bUpdateServerEpoch = false;
+					return true;
+				}
 			}
 		}
 		return false;
@@ -539,7 +542,7 @@ void UWorldPartitionStreamingPolicy::UpdateStreamingState()
 	TArray<const UWorldPartitionRuntimeCell*> ToUnloadCells;
 	{
 		TRACE_CPUPROFILER_EVENT_SCOPE(UWorldPartitionStreamingPolicy::UpdateStreamingState_ToUnloadCells);
-		auto BuildCellsToUnload = [this, &ToUnloadCells, bCanStream, bIsServer, ShouldWaitForClientVisibility](const TSet<const UWorldPartitionRuntimeCell*>& InCells)
+		auto BuildCellsToUnload = [this, &ToUnloadCells, bCanStream, bIsServer, ShouldWaitForClientVisibility](const TSet<TObjectPtr<const UWorldPartitionRuntimeCell>>& InCells)
 		{
 			for (const UWorldPartitionRuntimeCell* Cell : InCells)
 			{
@@ -984,7 +987,7 @@ void UWorldPartitionStreamingPolicy::OnCellHidden(const UWorldPartitionRuntimeCe
 	ActivatedCells.OnRemovedFromWorld(InCell);
 }
 
-void UWorldPartitionStreamingPolicy::FActivatedCells::Add(const UWorldPartitionRuntimeCell* InCell)
+void FActivatedCells::Add(const UWorldPartitionRuntimeCell* InCell)
 {
 	Cells.Add(InCell);
 	if (!InCell->IsAlwaysLoaded())
@@ -993,18 +996,18 @@ void UWorldPartitionStreamingPolicy::FActivatedCells::Add(const UWorldPartitionR
 	}
 }
 
-void UWorldPartitionStreamingPolicy::FActivatedCells::Remove(const UWorldPartitionRuntimeCell* InCell)
+void FActivatedCells::Remove(const UWorldPartitionRuntimeCell* InCell)
 {
 	Cells.Remove(InCell);
 	PendingAddToWorldCells.Remove(InCell);
 }
 
-void UWorldPartitionStreamingPolicy::FActivatedCells::OnAddedToWorld(const UWorldPartitionRuntimeCell* InCell)
+void FActivatedCells::OnAddedToWorld(const UWorldPartitionRuntimeCell* InCell)
 {
 	PendingAddToWorldCells.Remove(InCell);
 }
 
-void UWorldPartitionStreamingPolicy::FActivatedCells::OnRemovedFromWorld(const UWorldPartitionRuntimeCell* InCell)
+void FActivatedCells::OnRemovedFromWorld(const UWorldPartitionRuntimeCell* InCell)
 {
 	if (Cells.Contains(InCell))
 	{
diff --git a/Engine/Source/Runtime/Engine/Public/WorldPartition/WorldPartitionRuntimeCell.h b/Engine/Source/Runtime/Engine/Public/WorldPartition/WorldPartitionRuntimeCell.h
--- a/Engine/Source/Runtime/Engine/Public/WorldPartition/WorldPartitionRuntimeCell.h
+++ b/Engine/Source/Runtime/Engine/Public/WorldPartition/WorldPartitionRuntimeCell.h
@@ -178,8 +178,10 @@ class UWorldPartitionRuntimeCell : public UObject, public IWorldPartitionCell
 	ENGINE_API virtual bool CanUnload() const PURE_VIRTUAL(UWorldPartitionRuntimeCell::CanUnload, return true;);
 	ENGINE_API virtual void Activate() const PURE_VIRTUAL(UWorldPartitionRuntimeCell::Activate,);
 	ENGINE_API virtual void Deactivate() const PURE_VIRTUAL(UWorldPartitionRuntimeCell::Deactivate,);
-	ENGINE_API virtual bool IsAddedToWorld() const PURE_VIRTUAL(UWorldPartitionRuntimeCell::IsAddedToWorld, return false;);
-	ENGINE_API virtual bool CanAddToWorld() const PURE_VIRTUAL(UWorldPartitionRuntimeCell::CanAddToWorld, return false;);
+	UE_DEPRECATED(5.4, "IsAddedToWorld is deprecated.")
+	ENGINE_API virtual bool IsAddedToWorld() const { return false; }
+	UE_DEPRECATED(5.4, "CanAddToWorld is deprecated.")
+	ENGINE_API virtual bool CanAddToWorld() const { return false; }
 	ENGINE_API virtual ULevel* GetLevel() const PURE_VIRTUAL(UWorldPartitionRuntimeCell::GetLevel, return nullptr;);
 	ENGINE_API virtual EWorldPartitionRuntimeCellState GetCurrentState() const PURE_VIRTUAL(UWorldPartitionRuntimeCell::GetCurrentState, return EWorldPartitionRuntimeCellState::Unloaded;);
 	virtual FLinearColor GetDebugColor(EWorldPartitionRuntimeCellVisualizeMode VisualizeMode) const { static const FLinearColor DefaultColor = FLinearColor::Black.CopyWithNewOpacity(0.25f); return DefaultColor; }
@@ -261,6 +263,14 @@ class UWorldPartitionRuntimeCell : public UObject, public IWorldPartitionCell
 protected:
 	ENGINE_API FLinearColor GetDebugStreamingPriorityColor() const;
 
+	//@todo_ow: Implement ServerOnlyVisible and refactor ClientOnlyVisible.
+	//          Instead of this function, server would not not wait for client level visibility 
+	//          for server-only visible cells.
+	//          Also, server-only & client-only visible cells should be pre-filtered at 
+	//          world partition initialization to avoid visiting them at runtime.
+	//          The same should be done for injected external streaming objects.
+	ENGINE_API virtual bool ShouldServerWaitForClientLevelVisibility() const { return true; }
+
 	UPROPERTY()
 	bool bIsAlwaysLoaded;
 
@@ -315,4 +325,6 @@ public:
 
 	UPROPERTY()
 	TObjectPtr<UWorldPartitionRuntimeCellData> RuntimeCellData;
+
+	friend class UWorldPartitionStreamingPolicy;
 };
diff --git a/Engine/Source/Runtime/Engine/Public/WorldPartition/WorldPartitionRuntimeLevelStreamingCell.h b/Engine/Source/Runtime/Engine/Public/WorldPartition/WorldPartitionRuntimeLevelStreamingCell.h
--- a/Engine/Source/Runtime/Engine/Public/WorldPartition/WorldPartitionRuntimeLevelStreamingCell.h
+++ b/Engine/Source/Runtime/Engine/Public/WorldPartition/WorldPartitionRuntimeLevelStreamingCell.h
@@ -22,8 +22,6 @@ class UWorldPartitionRuntimeLevelStreamingCell : public UWorldPartitionRuntimeCe
 	ENGINE_API virtual bool CanUnload() const override;
 	ENGINE_API virtual void Activate() const override;
 	ENGINE_API virtual void Deactivate() const override;
-	ENGINE_API virtual bool IsAddedToWorld() const override;
-	ENGINE_API virtual bool CanAddToWorld() const override;
 	ENGINE_API virtual ULevel* GetLevel() const override;
 	ENGINE_API virtual EWorldPartitionRuntimeCellState GetCurrentState() const override;
 	ENGINE_API virtual FLinearColor GetDebugColor(EWorldPartitionRuntimeCellVisualizeMode VisualizeMode) const override;
diff --git a/Engine/Source/Runtime/Engine/Public/WorldPartition/WorldPartitionStreamingPolicy.h b/Engine/Source/Runtime/Engine/Public/WorldPartition/WorldPartitionStreamingPolicy.h
--- a/Engine/Source/Runtime/Engine/Public/WorldPartition/WorldPartitionStreamingPolicy.h
+++ b/Engine/Source/Runtime/Engine/Public/WorldPartition/WorldPartitionStreamingPolicy.h
@@ -19,6 +19,28 @@
 class UWorldPartition;
 class FWorldPartitionDraw2DContext;
 
+USTRUCT()
+struct FActivatedCells
+{
+	GENERATED_USTRUCT_BODY()
+
+	void Add(const UWorldPartitionRuntimeCell* InCell);
+	void Remove(const UWorldPartitionRuntimeCell* InCell);
+	bool Contains(const UWorldPartitionRuntimeCell* InCell) const { return Cells.Contains(InCell); }
+	void OnAddedToWorld(const UWorldPartitionRuntimeCell* InCell);
+	void OnRemovedFromWorld(const UWorldPartitionRuntimeCell* InCell);
+
+	const TSet<TObjectPtr<const UWorldPartitionRuntimeCell>>& GetCells() const { return Cells; }
+	const TSet<const UWorldPartitionRuntimeCell*>& GetPendingAddToWorldCells() const { return PendingAddToWorldCells; }
+
+private:
+
+	UPROPERTY(Transient)
+	TSet<TObjectPtr<const UWorldPartitionRuntimeCell>> Cells;
+
+	TSet<const UWorldPartitionRuntimeCell*> PendingAddToWorldCells;
+};
+
 UCLASS(Abstract, Within = WorldPartition)
 class UWorldPartitionStreamingPolicy : public UObject
 {
@@ -80,27 +102,11 @@ protected:
 	bool IsInBlockTillLevelStreamingCompleted(bool bIsCausedByBadStreamingPerformance = false) const;
 
 	const UWorldPartition* WorldPartition;
-	TSet<const UWorldPartitionRuntimeCell*> LoadedCells;
-
-	struct FActivatedCells
-	{
-		void Add(const UWorldPartitionRuntimeCell* InCell);
-		void Remove(const UWorldPartitionRuntimeCell* InCell);
-		bool Contains(const UWorldPartitionRuntimeCell* InCell) const { return Cells.Contains(InCell); }
-		void OnAddedToWorld(const UWorldPartitionRuntimeCell* InCell);
-		void OnRemovedFromWorld(const UWorldPartitionRuntimeCell* InCell);
-
-		const TSet<const UWorldPartitionRuntimeCell*>& GetCells() const { return Cells; }
-		const TSet<const UWorldPartitionRuntimeCell*>& GetPendingAddToWorldCells() const { return PendingAddToWorldCells; }
-
-	private:
-
-		TSet<const UWorldPartitionRuntimeCell*> Cells;
-		TSet<const UWorldPartitionRuntimeCell*> PendingAddToWorldCells;
-	};
+	UPROPERTY(Transient)
+	TSet<TObjectPtr<const UWorldPartitionRuntimeCell>> LoadedCells;
 
+	UPROPERTY(Transient)
 	FActivatedCells ActivatedCells;
-	mutable TArray<const UWorldPartitionRuntimeCell*> SortedAddToWorldCells;
 
 	// Streaming Sources
 	TArray<FWorldPartitionStreamingSource> StreamingSources;
@@ -109,8 +115,10 @@ protected:
 	TSet<const UWorldPartitionRuntimeCell*> FrameLoadCells;
 
 	// Used by UWorldPartitionSubsystem
-	TArray<const UWorldPartitionRuntimeCell*> ToActivateCells;
-	TArray<const UWorldPartitionRuntimeCell*> ToLoadCells;
+	UPROPERTY(Transient)
+	TArray<TObjectPtr<const UWorldPartitionRuntimeCell>> ToActivateCells;
+	UPROPERTY(Transient)
+	TArray<TObjectPtr<const UWorldPartitionRuntimeCell>> ToLoadCells;
 	int32 ProcessedToActivateCells;
 	int32 ProcessedToLoadCells;
 
