Inline first chunk to work for prime on demand as well

diff --git a/Engine/Source/Runtime/Engine/Classes/Sound/SoundClass.h b/Engine/Source/Runtime/Engine/Classes/Sound/SoundClass.h
--- a/Engine/Source/Runtime/Engine/Classes/Sound/SoundClass.h
+++ b/Engine/Source/Runtime/Engine/Classes/Sound/SoundClass.h
@@ -135,8 +135,8 @@ struct FSoundClassProperties
 	ESoundWaveLoadingBehavior LoadingBehavior;
 
 #if WITH_EDITORONLY_DATA
-   	/** How much audio to add to First Audio Chunk (in seconds), (only for RetainOnLoad) */
-   	UPROPERTY(EditAnywhere, Category = Loading, meta = (UIMin = 0, UIMax = 10, EditCondition = "LoadingBehavior == ESoundWaveLoadingBehavior::RetainOnLoad"), DisplayName="Size of First Audio Chunk (seconds)")
+   	/** How much audio to add to First Audio Chunk (in seconds) */
+	UPROPERTY(EditAnywhere, Category = Loading, meta = (UIMin = 0, UIMax = 10, EditCondition = "LoadingBehavior == ESoundWaveLoadingBehavior::RetainOnLoad || LoadingBehavior == ESoundWaveLoadingBehavior::PrimeOnLoad"), DisplayName="Size of First Audio Chunk (seconds)")
    	FPerPlatformFloat SizeOfFirstAudioChunkInSeconds = 0.0f;
 #endif //WITH_EDITORONLY_DATA
 
diff --git a/Engine/Source/Runtime/Engine/Classes/Sound/SoundWave.h b/Engine/Source/Runtime/Engine/Classes/Sound/SoundWave.h
--- a/Engine/Source/Runtime/Engine/Classes/Sound/SoundWave.h
+++ b/Engine/Source/Runtime/Engine/Classes/Sound/SoundWave.h
@@ -721,8 +721,8 @@ public:
 
 #if WITH_EDITORONLY_DATA
 public:
-   	/** How much audio to add to First Audio Chunk (in seconds), (only for RetainOnLoad) */
-   	UPROPERTY(EditAnywhere, Category = Loading, meta = (UIMin = 0, UIMax = 10, EditCondition = "LoadingBehavior == ESoundWaveLoadingBehavior::RetainOnLoad"), DisplayName="Size of First Audio Chunk (seconds)")
+   	/** How much audio to add to First Audio Chunk (in seconds) */
+	UPROPERTY(EditAnywhere, Category = Loading, meta = (UIMin = 0, UIMax = 10, EditCondition = "LoadingBehavior == ESoundWaveLoadingBehavior::RetainOnLoad || LoadingBehavior == ESoundWaveLoadingBehavior::PrimeOnLoad"), DisplayName="Size of First Audio Chunk (seconds)")
    	FPerPlatformFloat SizeOfFirstAudioChunkInSeconds = 0.0f;
 #endif //WITH_EDITOR_ONLY_DATA
 
diff --git a/Engine/Source/Runtime/Engine/Classes/Sound/SoundWaveLoadingBehavior.cpp b/Engine/Source/Runtime/Engine/Classes/Sound/SoundWaveLoadingBehavior.cpp
--- a/Engine/Source/Runtime/Engine/Classes/Sound/SoundWaveLoadingBehavior.cpp
+++ b/Engine/Source/Runtime/Engine/Classes/Sound/SoundWaveLoadingBehavior.cpp
@@ -13,25 +13,16 @@
 #include "SoundWave.h"
 #include "UObject/LinkerLoad.h"
 
+#ifndef CASE_ENUM_TO_TEXT
+#define CASE_ENUM_TO_TEXT(TXT) case TXT: return TEXT(#TXT);
+#endif
 
-const TCHAR* EnumToString(ESoundWaveLoadingBehavior InCurrentState)
+const TCHAR* EnumToString(const ESoundWaveLoadingBehavior InCurrentState)
 {
-	switch (InCurrentState)
+	switch(InCurrentState)
 	{
-	case ESoundWaveLoadingBehavior::Inherited:
-		return TEXT("Inherited");
-	case ESoundWaveLoadingBehavior::RetainOnLoad:
-		return TEXT("RetainOnLoad");
-	case ESoundWaveLoadingBehavior::PrimeOnLoad:
-		return TEXT("PrimeOnLoad");
-	case ESoundWaveLoadingBehavior::LoadOnDemand:
-		return TEXT("LoadOnDemand");
-	case ESoundWaveLoadingBehavior::ForceInline:
-		return TEXT("ForceInline");
-	case ESoundWaveLoadingBehavior::Uninitialized:
-		return TEXT("Uninitialized");
-	}
-	ensure(false);
+		FOREACH_ENUM_ESOUNDWAVELOADINGBEHAVIOR(CASE_ENUM_TO_TEXT)
+	} 
 	return TEXT("Unknown");
 }
 
@@ -70,8 +61,6 @@ public:
 private:
 	void CacheAllClassLoadingBehaviors()
 	{
-		TRACE_CPUPROFILER_EVENT_SCOPE(LoadAndCacheAllSoundClassLoadingBehaviors);
-
 		ensureMsgf(!AssetRegistry.IsSearchAsync() || !AssetRegistry.IsLoadingAssets(), 
 		           TEXT("Function must not be called until after cook has started and waited on the AssetRegistry already."));
 		
@@ -124,7 +113,7 @@ private:
 		return {};
 	}
 	
-	virtual FClassData FindOwningLoadingBehavior(const USoundWave* InWave, const ITargetPlatform* InTargetPlatform) const 
+	virtual FClassData FindOwningLoadingBehavior(const USoundWave* InWave, const ITargetPlatform* InTargetPlatform) const override
 	{
 		TRACE_CPUPROFILER_EVENT_SCOPE(FindOwningLoadingBehavior);
 
@@ -266,8 +255,9 @@ bool ISoundWaveLoadingBehaviorUtil::FClassData::CompareGreater(const FClassData&
 	}
 	else 
 	{
-		// If we are using Retain, use one with the higher Length.
-		if (LoadingBehavior == ESoundWaveLoadingBehavior::RetainOnLoad)
+		// If we are using Prime/Retain, use one with the higher Length.
+		if (LoadingBehavior == ESoundWaveLoadingBehavior::RetainOnLoad ||
+			LoadingBehavior == ESoundWaveLoadingBehavior::PrimeOnLoad )
 		{
 			const float Length = LengthOfFirstChunkInSeconds.GetValueForPlatform(*InPlatform->PlatformName());
 			const float OtherLength = InOther.LengthOfFirstChunkInSeconds.GetValueForPlatform(*InPlatform->PlatformName());
diff --git a/Engine/Source/Runtime/Engine/Private/SoundWave.cpp b/Engine/Source/Runtime/Engine/Private/SoundWave.cpp
--- a/Engine/Source/Runtime/Engine/Private/SoundWave.cpp
+++ b/Engine/Source/Runtime/Engine/Private/SoundWave.cpp
@@ -2303,25 +2303,39 @@ float USoundWave::GetSizeOfFirstAudioChunkInSeconds(const ITargetPlatform* InPla
 	
 	ensure(InPlatform);
 
+	// This is the most important function for determining the size of chunk 1
+
 	// We need to walk up the hierarchy of sound classes to first determine our FirstAudioChunkSize.
 	if (SoundWaveDataPtr->LoadingBehavior == ESoundWaveLoadingBehavior::Uninitialized)
 	{
 		CacheInheritedLoadingBehavior();
 	}
 
-	// Use owner loading behavior if we have it.
-	ISoundWaveLoadingBehaviorUtil::FClassData Behavior = GetOwnerLoadingBehavior(InPlatform);
-	if (Behavior.LoadingBehavior == ESoundWaveLoadingBehavior::Uninitialized)
+	// Order of importance for where we get this setting.
+	// Wave -> Owner -> Class
+
+	// If Wave has a non-inherited Behavior set, use that.
+	ISoundWaveLoadingBehaviorUtil::FClassData Behavior;
+	if (LoadingBehavior != ESoundWaveLoadingBehavior::Inherited)
+	{
+		Behavior.LoadingBehavior = LoadingBehavior;
+		Behavior.LengthOfFirstChunkInSeconds = SizeOfFirstAudioChunkInSeconds;
+	}
+	else // ... if we are inherited, try the owner first.
+	{
+		Behavior = GetOwnerLoadingBehavior(InPlatform);
+	}
+
+	// ... otherwise (or if we failed above) use this soundwave resolved behavior,
+	if (Behavior.LoadingBehavior == ESoundWaveLoadingBehavior::Uninitialized)		
 	{
-		// Otherwise use what we've been tagged
 		Behavior.LoadingBehavior = SoundWaveDataPtr->LoadingBehavior;
 		Behavior.LengthOfFirstChunkInSeconds = SoundWaveDataPtr->SizeOfFirstAudioChunkInSeconds;
 	}
 
-	// We only allow setting the first chunk size if we're planning on retaining it.
-	// This could be relaxed later once we have a better format that wouldn't cause a huge delta
-	// on tuning the size (all assets etc). By using just retained, we expect a smaller set.
-	if (Behavior.LoadingBehavior == ESoundWaveLoadingBehavior::RetainOnLoad)
+	// Only give values if the behavior is Retain/Prime.
+	if (Behavior.LoadingBehavior == ESoundWaveLoadingBehavior::RetainOnLoad ||
+		Behavior.LoadingBehavior == ESoundWaveLoadingBehavior::PrimeOnLoad)
 	{
 		// See if we have an override for this platform on the SoundClass hierarchy.
 		const FName PlatformName = InPlatform ? InPlatform->GetPlatformInfo().IniPlatformName : FName();
