Readd FullySimplifiedStrataFrontMaterial code chunk to Strata.

This fixes compilation error on simplified material like LumenCards.

diff --git a/Engine/Source/Runtime/Engine/Private/Materials/HLSLMaterialTranslator.cpp b/Engine/Source/Runtime/Engine/Private/Materials/HLSLMaterialTranslator.cpp
--- a/Engine/Source/Runtime/Engine/Private/Materials/HLSLMaterialTranslator.cpp
+++ b/Engine/Source/Runtime/Engine/Private/Materials/HLSLMaterialTranslator.cpp
@@ -1833,6 +1833,40 @@ bool FHLSLMaterialTranslator::Translate()
 			}
 		}
 
+		// The code chunk corresponding to FullySimplifiedStrataFrontMaterialCodeChunk have already been written as part of MP_FrontMaterial.
+		// Here we get the FullySimplifiedStrataFrontMaterialTranslatedCodeChunks representing the variable storing the final fully simplified StrataData.
+		if(bStrataEnabled)
+		{
+			uint32 PropertyId = MP_FrontMaterial;
+		
+			if (PropertyId == MP_MaterialAttributes || PropertyId == MP_Normal || PropertyId == MP_CustomOutput)
+			{
+				continue;
+			}
+		
+			const EShaderFrequency PropertyShaderFrequency = FMaterialAttributeDefinitionMap::GetShaderFrequency((EMaterialProperty)PropertyId);
+		
+			int32 StartChunk = 0;
+			if (PropertyShaderFrequency == NormalShaderFrequency && SharedPixelProperties[PropertyId])
+			{
+				// When processing shared properties, do not generate the code before the Normal was generated as those are already handled
+				StartChunk = NormalCodeChunkEnd;
+			}
+		
+			// Reduce definition statements that don't contribute to the function's return value.
+			// @todo-lh: This should be expanded to a general reduction, but is currently only intended to fix an FXC internal compiler error reported in UE-117831
+			const bool bReduceAfterReturnValue = (PropertyId == MP_WorldPositionOffset || PropertyId == CompiledMP_PrevWorldPositionOffset || PropertyId == MP_Displacement);
+		
+			GetFixedParameterCode(
+				StartChunk,
+				SharedPropertyCodeChunks[PropertyShaderFrequency].Num(),
+				FullySimplifiedStrataFrontMaterialCodeChunk,								//Chunk[PropertyId],
+				SharedPropertyCodeChunks[PropertyShaderFrequency],
+				FullySimplifiedStrataFrontMaterialTranslatedCodeChunkDefinitions,
+				FullySimplifiedStrataFrontMaterialTranslatedCodeChunks,
+				Variation,
+				bReduceAfterReturnValue); 
+		}
 
 		for (uint32 PropertyId = MP_MAX; PropertyId < CompiledMP_MAX; ++PropertyId)
 		{
