Fix missing InstancingContext mapping: Content Bundle edge case where spatial actor references a non-spatial actor

#rnx

diff --git a/Engine/Source/Runtime/Engine/Private/WorldPartition/WorldPartitionLevelStreamingDynamic.cpp b/Engine/Source/Runtime/Engine/Private/WorldPartition/WorldPartitionLevelStreamingDynamic.cpp
--- a/Engine/Source/Runtime/Engine/Private/WorldPartition/WorldPartitionLevelStreamingDynamic.cpp
+++ b/Engine/Source/Runtime/Engine/Private/WorldPartition/WorldPartitionLevelStreamingDynamic.cpp
@@ -336,7 +336,7 @@ bool UWorldPartitionLevelStreamingDynamic::IssueLoadRequests()
 	bLoadSucceeded = false;
 	bLoadRequestInProgress = true;
 
-	auto BuildInstancingContext = [this](FLinkerInstancingContext& OutLinkInstancingContext)
+	auto BuildInstancingContext = [this](FLinkerInstancingContext& OutLinkInstancingContext, const TArray<FName>& InChildPackagesToDuplicate)
 	{
 		// Don't do SoftObjectPath remapping for PersistentLevel actors because references can end up in different cells
 		OutLinkInstancingContext.SetSoftObjectPathRemappingEnabled(false);
@@ -353,16 +353,19 @@ bool UWorldPartitionLevelStreamingDynamic::IssueLoadRequests()
 				if (!bIsContainerPackageAlreadyRemapped)
 				{
 					OutLinkInstancingContext.AddPackageMapping(CellObjectMapping.ContainerPackage, RuntimePackage->GetFName());
-					
 				}
 			}
 		}
+
+		for (FName ChildPackageToDuplicate : InChildPackagesToDuplicate)
+		{
+			// Add mapping in case we have a ChildPackageToLoad that references this package (Spatial actor references Non-spatial actor in a ContentBundle Alwaysloaded Cell)
+			OutLinkInstancingContext.AddPackageMapping(ChildPackageToDuplicate, RuntimePackage->GetFName());
+		}
 	};
 	
-	FLinkerInstancingContext InstancingContext;
-	BuildInstancingContext(InstancingContext);
-
 	ChildPackagesToLoad.Reset(ChildPackages.Num());
+	TArray<FName> ChildPackagesToDuplicate;
 
 	UWorld* World = GetWorld();
 	for (FWorldPartitionRuntimeCellObjectMapping& ChildPackage : ChildPackages)
@@ -378,6 +381,7 @@ bool UWorldPartitionLevelStreamingDynamic::IssueLoadRequests()
 				{
 					if (AActor* ActorModifiedForPIE = UnsavedActorsContainer->Actors.FindRef(*SubObjectName))
 					{
+						ChildPackagesToDuplicate.Add(ChildPackage.Package);
 						bNeedDup = true;
 					}
 				}
@@ -389,6 +393,9 @@ bool UWorldPartitionLevelStreamingDynamic::IssueLoadRequests()
 			ChildPackagesToLoad.Add(ChildPackage);
 		}
 	}
+		
+	FLinkerInstancingContext InstancingContext;
+	BuildInstancingContext(InstancingContext, ChildPackagesToDuplicate);
 
 	// Duplicate unsaved actors
 	if (UnsavedActorsContainer)
