Prevent crash during FSkeletalMeshAttributeVertexBuffer::Serialize when an unintended pixel format is loaded. Also attempt to detect corrupt data and bail if serialization error is detected.

diff --git a/Engine/Source/Runtime/Engine/Private/Animation/SkeletalMeshVertexAttribute.cpp b/Engine/Source/Runtime/Engine/Private/Animation/SkeletalMeshVertexAttribute.cpp
--- a/Engine/Source/Runtime/Engine/Private/Animation/SkeletalMeshVertexAttribute.cpp
+++ b/Engine/Source/Runtime/Engine/Private/Animation/SkeletalMeshVertexAttribute.cpp
@@ -130,7 +130,7 @@ void FSkeletalMeshVertexAttributeRenderData::Serialize(FArchive& Ar)
 		int32 Count;
 		Ar << Count;
 		Buffers.Reset();
-		for (int32 Index = 0; Index < Count; Index++)
+		for (int32 Index = 0; Index < Count && !Ar.IsError(); Index++)
 		{
 			FName Name;
 			Ar << Name;
diff --git a/Engine/Source/Runtime/Engine/Private/Rendering/SkeletalMeshAttributeVertexBuffer.cpp b/Engine/Source/Runtime/Engine/Private/Rendering/SkeletalMeshAttributeVertexBuffer.cpp
--- a/Engine/Source/Runtime/Engine/Private/Rendering/SkeletalMeshAttributeVertexBuffer.cpp
+++ b/Engine/Source/Runtime/Engine/Private/Rendering/SkeletalMeshAttributeVertexBuffer.cpp
@@ -194,8 +194,34 @@ void FSkeletalMeshAttributeVertexBuffer::Serialize(FArchive& Ar)
 
 	if (Ar.IsLoading())
 	{
+		bool bValid;
+		switch (PixelFormat)
+		{
+		case PF_R8:
+			bValid = ComponentStride == 1;
+			break;
+		case PF_R16F:
+			bValid = ComponentStride == 2;
+			break;
+		case PF_R32_FLOAT:
+			bValid = ComponentStride == 4;
+			break;
+		default:
+			bValid = false;
+			break;
+		}
+
+		if (!bValid)
+		{
+			Ar.SetError();
+			return;
+		}
+
 		Allocate();
 	}
 
-	ValueData->Serialize(Ar);
+	if (ValueData)
+	{
+		ValueData->Serialize(Ar);
+	}
 }
