[NO NON-PII COMMIT MESSAGE]

diff --git a/Engine/Source/Runtime/Engine/Private/DataDrivenCVars/DataDrivenCVars.cpp b/Engine/Source/Runtime/Engine/Private/DataDrivenCVars/DataDrivenCVars.cpp
--- a/Engine/Source/Runtime/Engine/Private/DataDrivenCVars/DataDrivenCVars.cpp
+++ b/Engine/Source/Runtime/Engine/Private/DataDrivenCVars/DataDrivenCVars.cpp
@@ -24,15 +24,15 @@ void FDataDrivenConsoleVariable::Register()
 		{
 			if (Type == FDataDrivenCVarType::CVarInt)
 			{
-				CVarToAdd = IConsoleManager::Get().RegisterConsoleVariable(*Name, DefaultValueInt, TEXT("RuntimeConsoleVariables"), ECVF_Default | ECVF_Scalability);
+				CVarToAdd = IConsoleManager::Get().RegisterConsoleVariable(*Name, DefaultValueInt, *ToolTip, ECVF_Default | ECVF_Scalability);
 			}
 			else if (Type == FDataDrivenCVarType::CVarBool)
 			{
-				CVarToAdd = IConsoleManager::Get().RegisterConsoleVariable(*Name, DefaultValueBool, TEXT("RuntimeConsoleVariables"), ECVF_Default | ECVF_Scalability);
+				CVarToAdd = IConsoleManager::Get().RegisterConsoleVariable(*Name, DefaultValueBool, *ToolTip, ECVF_Default | ECVF_Scalability);
 			}
 			else
 			{
-				CVarToAdd = IConsoleManager::Get().RegisterConsoleVariable(*Name, DefaultValueFloat, TEXT("RuntimeConsoleVariables"), ECVF_Default | ECVF_Scalability);
+				CVarToAdd = IConsoleManager::Get().RegisterConsoleVariable(*Name, DefaultValueFloat, *ToolTip, ECVF_Default | ECVF_Scalability);
 			}
 		}
 		CVarToAdd->SetOnChangedCallback(FConsoleVariableDelegate::CreateStatic(UDataDrivenConsoleVariableSettings::OnDataDrivenChange));
