Fix 128bit RT canvas tile rendering

diff --git a/Engine/Source/Runtime/Engine/Private/TileRendering.cpp b/Engine/Source/Runtime/Engine/Private/TileRendering.cpp
--- a/Engine/Source/Runtime/Engine/Private/TileRendering.cpp
+++ b/Engine/Source/Runtime/Engine/Private/TileRendering.cpp
@@ -187,7 +187,7 @@ void FCanvasTileRendererItem::FRenderData::RenderTiles(
 		if (CurrentMeshBatch->BatchHitProxyId != Tile.HitProxyId)
 		{
 			// Flush the current batch before allocating a new one: 
-			GetRendererModule().DrawTileMesh(RenderContext, DrawRenderState, View, *CurrentMeshBatch, bIsHitTesting, CurrentMeshBatch->BatchHitProxyId);
+			GetRendererModule().DrawTileMesh(RenderContext, DrawRenderState, View, *CurrentMeshBatch, bIsHitTesting, CurrentMeshBatch->BatchHitProxyId, bUse128bitRT);
 
 			CurrentMeshBatch = AllocTileMeshBatch(RenderContext, Tile.HitProxyId);
 			CurrentMeshBatch->Elements[0].FirstIndex = CanvasTileIndexCount * TileIdx;
@@ -199,7 +199,7 @@ void FCanvasTileRendererItem::FRenderData::RenderTiles(
 
 	// Flush the final batch: 
 	check(CurrentMeshBatch != nullptr);
-	GetRendererModule().DrawTileMesh(RenderContext, DrawRenderState, View, *CurrentMeshBatch, bIsHitTesting, CurrentMeshBatch->BatchHitProxyId);
+	GetRendererModule().DrawTileMesh(RenderContext, DrawRenderState, View, *CurrentMeshBatch, bIsHitTesting, CurrentMeshBatch->BatchHitProxyId, bUse128bitRT);
 
 	AddPass(RenderContext.GraphBuilder, RDG_EVENT_NAME("ReleaseTileMesh"), [this](FRHICommandListImmediate&)
 	{
