Removed unused CVar that has since been replaced by [REDACTED_PERSON].

diff --git a/Engine/Source/Runtime/Engine/Private/UnrealEngine.cpp b/Engine/Source/Runtime/Engine/Private/UnrealEngine.cpp
--- a/Engine/Source/Runtime/Engine/Private/UnrealEngine.cpp
+++ b/Engine/Source/Runtime/Engine/Private/UnrealEngine.cpp
@@ -414,14 +414,6 @@ static FAutoConsoleVariableRef GDelayTrimMemoryDuringMapLoadModeCVar(
 	ECVF_Default
 );
 
-static bool GShouldLogReferencesToLeakedWorldObjects = !UE_BUILD_SHIPPING;
-static FAutoConsoleVariableRef GShouldLogReferencesToLeakedWorldObjectsCVar(
-	TEXT("Engine.ShouldLogReferencesToLeakedWorldObjects"),
-	GShouldLogReferencesToLeakedWorldObjects,
-	TEXT("Enables logging references preventing the previous world objects from being cleaned up during map load"),
-	ECVF_Default
-);
-
 static int32 GVerifyLoadMapWorldCleanup_Severity = 2;
 static FAutoConsoleVariableRef CVarVerifyLoadMapWorldCleanup_Severity(
 #if UE_BUILD_SHIPPING
