Fix invalid Cast in material
Replace by call on this, but keep code commented as it wasn't doing anything for the past 5 years.
#rnx

diff --git a/Engine/Source/Runtime/Engine/Private/Materials/Material.cpp b/Engine/Source/Runtime/Engine/Private/Materials/Material.cpp
--- a/Engine/Source/Runtime/Engine/Private/Materials/Material.cpp
+++ b/Engine/Source/Runtime/Engine/Private/Materials/Material.cpp
@@ -1085,8 +1085,6 @@ void UMaterial::GetUsedTextures(TArray<UTexture*>& OutTextures, EMaterialQuality
 
 	if (!FPlatformProperties::IsServerOnly())
 	{
-		const UMaterialInstance* MaterialInstance = Cast<UMaterialInstance>(this);
-
 		FInt32Range QualityLevelRange(0, EMaterialQualityLevel::Num - 1);
 		if (!bAllQualityLevels)
 		{
@@ -1138,33 +1136,32 @@ void UMaterial::GetUsedTextures(TArray<UTexture*>& OutTextures, EMaterialQuality
 			}
 
 #if WITH_EDITOR
-			if (MaterialInstance)
+/*
+			// Also look for any scalar parameters that are acting as lookups for an atlas texture, and store the atlas texture
+			const TArrayView<const FMaterialNumericParameterInfo> AtlasExpressions[1] =
 			{
-				// Also look for any scalar parameters that are acting as lookups for an atlas texture, and store the atlas texture
-				const TArrayView<const FMaterialNumericParameterInfo> AtlasExpressions[1] =
-				{
-					CurrentResource->GetUniformNumericParameterExpressions()
-				};
+				CurrentResource->GetUniformNumericParameterExpressions()
+			};
 
-				for (int32 TypeIndex = 0; TypeIndex < UE_ARRAY_COUNT(AtlasExpressions); TypeIndex++)
+			for (int32 TypeIndex = 0; TypeIndex < UE_ARRAY_COUNT(AtlasExpressions); TypeIndex++)
+			{
+				// Iterate over each of the material's texture expressions.
+				for (const FMaterialNumericParameterInfo& Parameter : AtlasExpressions[TypeIndex])
 				{
-					// Iterate over each of the material's texture expressions.
-					for (const FMaterialNumericParameterInfo& Parameter : AtlasExpressions[TypeIndex])
+					if (Parameter.ParameterType == EMaterialParameterType::Scalar)
 					{
-						if (Parameter.ParameterType == EMaterialParameterType::Scalar)
+						bool bIsUsedAsAtlasPosition;
+						TSoftObjectPtr<UCurveLinearColor> Curve;
+						TSoftObjectPtr<UCurveLinearColorAtlas> Atlas;
+						IsScalarParameterUsedAsAtlasPosition(Parameter.ParameterInfo, bIsUsedAsAtlasPosition, Curve, Atlas);
+						if (Atlas)
 						{
-							bool bIsUsedAsAtlasPosition;
-							TSoftObjectPtr<class UCurveLinearColor> Curve;
-							TSoftObjectPtr<class UCurveLinearColorAtlas> Atlas;
-							MaterialInstance->IsScalarParameterUsedAsAtlasPosition(Parameter.ParameterInfo, bIsUsedAsAtlasPosition, Curve, Atlas);
-							if (Atlas)
-							{
-								OutTextures.AddUnique(Atlas.Get());
-							}
+							OutTextures.AddUnique(Atlas.Get());
 						}
 					}
 				}
 			}
+*/
 #endif // WITH_EDITOR
 		}
 	}
