Fix FTextRenderSceneProxy not properly releasing RayTracingGeometry when toggling raytracing at runtime.

diff --git a/Engine/Source/Runtime/Engine/Private/Components/TextRenderComponent.cpp b/Engine/Source/Runtime/Engine/Private/Components/TextRenderComponent.cpp
--- a/Engine/Source/Runtime/Engine/Private/Components/TextRenderComponent.cpp
+++ b/Engine/Source/Runtime/Engine/Private/Components/TextRenderComponent.cpp
@@ -760,8 +760,9 @@ void FTextRenderSceneProxy::ReleaseRenderThreadResources()
 	VertexFactory.ReleaseResource();
 
 #if RHI_RAYTRACING
-	if (IsRayTracingEnabled())
+	if (RayTracingGeometry.IsValid())
 	{
+		check(IsRayTracingAllowed());
 		RayTracingGeometry.ReleaseResource();
 	}
 #endif
@@ -1144,7 +1145,6 @@ void FTextRenderSceneProxy::UpdateRayTracingGeometry_RenderingThread(FRHICommand
 	Initializer.bFastBuild = true;
 	Initializer.bAllowUpdate = false;
 
-	TArray<FRayTracingGeometrySegment> GeometrySections;
 	FRayTracingGeometrySegment Segment;
 	Segment.VertexBuffer = VertexBuffers.PositionVertexBuffer.VertexBufferRHI;
 	Segment.VertexBufferElementType = VET_Float3;
@@ -1155,8 +1155,8 @@ void FTextRenderSceneProxy::UpdateRayTracingGeometry_RenderingThread(FRHICommand
 	Segment.NumPrimitives = IndexBuffer.Indices.Num() / 3;
 	Segment.bEnabled = true;
 	Segment.bForceOpaque = false;
-	GeometrySections.Add(Segment);
-	Initializer.Segments = GeometrySections;
+
+	Initializer.Segments.Add(Segment);
 
 	RayTracingGeometry.SetInitializer(Initializer);
 	RayTracingGeometry.InitResource(RHICmdList);
