Mark HeterogeneousVolumeComponent as being RayTracingRelevant so edits cause path tracing invalidation events.

diff --git a/Engine/Source/Runtime/Engine/Private/Components/HeterogeneousVolumeComponent.cpp b/Engine/Source/Runtime/Engine/Private/Components/HeterogeneousVolumeComponent.cpp
--- a/Engine/Source/Runtime/Engine/Private/Components/HeterogeneousVolumeComponent.cpp
+++ b/Engine/Source/Runtime/Engine/Private/Components/HeterogeneousVolumeComponent.cpp
@@ -43,7 +43,9 @@ public:
 		static size_t UniquePointer;
 		return reinterpret_cast<size_t>(&UniquePointer);
 	}
-
+#if RHI_RAYTRACING
+	virtual bool IsRayTracingRelevant() const override  { return true; }
+#endif // RHI_RAYTRACING
 	virtual uint32 GetMemoryFootprint(void) const override { return sizeof(*this) + GetAllocatedSize(); }
 	uint32 GetAllocatedSize(void) const { return FPrimitiveSceneProxy::GetAllocatedSize(); }
 	//~ End FPrimitiveSceneProxy Interface.
