USparseVolumeTextureFrame only removes the OwnerSVT from the streaming manager if the frame is in OwnerSVT->Frames array
Fixes issues with GC incorrectly removing the SVT from the streaming manager - noticed during implementation of SVT caching from the Niagara Baker.

diff --git a/Engine/Source/Runtime/Engine/Private/SparseVolumeTexture/SparseVolumeTexture.cpp b/Engine/Source/Runtime/Engine/Private/SparseVolumeTexture/SparseVolumeTexture.cpp
--- a/Engine/Source/Runtime/Engine/Private/SparseVolumeTexture/SparseVolumeTexture.cpp
+++ b/Engine/Source/Runtime/Engine/Private/SparseVolumeTexture/SparseVolumeTexture.cpp
@@ -936,9 +936,20 @@ void USparseVolumeTextureFrame::FinishDestroy()
 void USparseVolumeTextureFrame::BeginDestroy()
 {
 	// Ensure that the streamable SVT has been removed from the streaming manager
+	
 	if (IsValid(Owner))
-	{
-		UE::SVT::GetStreamingManager().Remove_GameThread(CastChecked<UStreamableSparseVolumeTexture>(Owner));
+	{		
+		UStreamableSparseVolumeTexture* SVTOwner = CastChecked<UStreamableSparseVolumeTexture>(Owner);
+		for (int i = 0; i < SVTOwner->GetNumFrames(); ++i)
+		{
+			// if the owner contains the current frame being deleted, remove the owner from the streaming manager
+			// SVT_TODO: This is a temporary fix for a GC problem.  In the future this will be replaced with a more robust solution.
+			if (SVTOwner->GetFrame(i) == this)
+			{
+				UE::SVT::GetStreamingManager().Remove_GameThread(SVTOwner);
+				break;
+			}
+		}		
 	}
 	
 	if (TextureRenderResources)
