Avoid unecessary copies of FStreamableDelegate within FStreamableManager.

+ Added an overload of FStreamableHandle::ExecuteDelegate that takes the FStreamableDelegates by rvalue. This code path will continue to pass the delegates along by rvalue to avoid copies.
+ Fixed up multiple call sites where copies were being done and replaced with moves.

#tests 64e80bfd2840fe78aaaec440, 64e8c1882840fe78aadeb95a, standard build, local cooked and editor tests

diff --git a/Engine/Source/Runtime/Engine/Classes/Engine/StreamableManager.h b/Engine/Source/Runtime/Engine/Classes/Engine/StreamableManager.h
--- a/Engine/Source/Runtime/Engine/Classes/Engine/StreamableManager.h
+++ b/Engine/Source/Runtime/Engine/Classes/Engine/StreamableManager.h
@@ -245,6 +245,8 @@ struct FStreamableHandle : public TSharedFromThis<FStreamableHandle>
 	 * @param CancelDelegate	If handle gets cancelled before primary delegate executes, this delegate will be called instead
 	 */
 	static ENGINE_API void ExecuteDelegate(const FStreamableDelegate& Delegate, TSharedPtr<FStreamableHandle> AssociatedHandle = nullptr, const FStreamableDelegate& CancelDelegate = FStreamableDelegate());
+	/** Overload taking the Delegate by rvalue. Using this will prevent any copies of FStreamableDelegate from being made if the callback is deferred. */
+	static ENGINE_API void ExecuteDelegate(FStreamableDelegate&& Delegate, TSharedPtr<FStreamableHandle> AssociatedHandle = nullptr, FStreamableDelegate&& CancelDelegate = FStreamableDelegate());
 
 	/**
 	 * Return a TSharedPtr of the first handle among this and descendants which satisfies the predicate.
diff --git a/Engine/Source/Runtime/Engine/Private/AssetManager.cpp b/Engine/Source/Runtime/Engine/Private/AssetManager.cpp
--- a/Engine/Source/Runtime/Engine/Private/AssetManager.cpp
+++ b/Engine/Source/Runtime/Engine/Private/AssetManager.cpp
@@ -1844,12 +1844,12 @@ TSharedPtr<FStreamableHandle> UAssetManager::ChangeBundleStateForPrimaryAssets(c
 		// Call delegate or bind to meta handle
 		if (ReturnHandle->HasLoadCompleted())
 		{
-			FStreamableHandle::ExecuteDelegate(DelegateToCall);
+			FStreamableHandle::ExecuteDelegate(MoveTemp(DelegateToCall));
 		}
 		else
 		{
 			// Call external callback when completed
-			ReturnHandle->BindCompleteDelegate(DelegateToCall);
+			ReturnHandle->BindCompleteDelegate(MoveTemp(DelegateToCall));
 		}
 	}
 	else if (NewHandles.Num() == 1)
@@ -1860,18 +1860,18 @@ TSharedPtr<FStreamableHandle> UAssetManager::ChangeBundleStateForPrimaryAssets(c
 		// If only one handle, return it and add callback
 		if (ReturnHandle->HasLoadCompleted())
 		{
-			FStreamableHandle::ExecuteDelegate(DelegateToCall);
+			FStreamableHandle::ExecuteDelegate(MoveTemp(DelegateToCall));
 		}
 		else
 		{
 			// Call internal callback and external callback when it finishes
-			ReturnHandle->BindCompleteDelegate(FStreamableDelegate::CreateUObject(this, &UAssetManager::OnAssetStateChangeCompleted, NewAssets[0], ReturnHandle, DelegateToCall));
+			ReturnHandle->BindCompleteDelegate(FStreamableDelegate::CreateUObject(this, &UAssetManager::OnAssetStateChangeCompleted, NewAssets[0], ReturnHandle, MoveTemp(DelegateToCall)));
 		}
 	}
 	else
 	{
 		// Call completion callback, nothing to do
-		FStreamableHandle::ExecuteDelegate(DelegateToCall);
+		FStreamableHandle::ExecuteDelegate(MoveTemp(DelegateToCall));
 	}
 
 	return ReturnHandle;
@@ -2206,7 +2206,7 @@ TSharedPtr<FStreamableHandle> UAssetManager::LoadAssetList(const TArray<FSoftObj
 	if (bShouldUseSynchronousLoad && MissingChunks.Num() == 0)
 	{
 		NewHandle = StreamableManager.RequestSyncLoad(AssetList, false, DebugName);
-		FStreamableHandle::ExecuteDelegate(DelegateToCall);
+		FStreamableHandle::ExecuteDelegate(MoveTemp(DelegateToCall));
 	}
 	else
 	{
@@ -2386,13 +2386,13 @@ void UAssetManager::AcquireResourcesForAssetList(const TArray<FSoftObjectPath>&
 	{
 		// At least one chunk doesn't exist, fail
 		FStreamableDelegate TempDelegate = FStreamableDelegate::CreateLambda([CompleteDelegate = MoveTemp(CompleteDelegate), MissingChunks]() { CompleteDelegate.ExecuteIfBound(false, MissingChunks); });
-		FStreamableHandle::ExecuteDelegate(TempDelegate);
+		FStreamableHandle::ExecuteDelegate(MoveTemp(TempDelegate));
 	}
 	else if (MissingChunks.Num() == 0)
 	{
 		// All here, schedule the callback
 		FStreamableDelegate TempDelegate = FStreamableDelegate::CreateLambda([CompleteDelegate = MoveTemp(CompleteDelegate)]() { CompleteDelegate.ExecuteIfBound(true, TArray<int32>()); });
-		FStreamableHandle::ExecuteDelegate(TempDelegate);
+		FStreamableHandle::ExecuteDelegate(MoveTemp(TempDelegate));
 	}
 	else
 	{
diff --git a/Engine/Source/Runtime/Engine/Private/StreamableManager.cpp b/Engine/Source/Runtime/Engine/Private/StreamableManager.cpp
--- a/Engine/Source/Runtime/Engine/Private/StreamableManager.cpp
+++ b/Engine/Source/Runtime/Engine/Private/StreamableManager.cpp
@@ -13,6 +13,7 @@
 #include "ProfilingDebugging/CpuProfilerTrace.h"
 #include "Trace/Trace.h"
 #include "Trace/Trace.inl"
+#include <type_traits>
 
 DEFINE_LOG_CATEGORY_STATIC(LogStreamableManager, Log, All);
 
@@ -45,9 +46,10 @@ class FStreamableDelegateDelayHelper : public FTickableGameObject
 public:
 
 	/** Adds a delegate to deferred list */
-	void AddDelegate(const FStreamableDelegate& Delegate, const FStreamableDelegate& CancelDelegate, TSharedPtr<FStreamableHandle> AssociatedHandle)
+	template<typename InStreamableDelegate>
+	void AddDelegate(InStreamableDelegate&& Delegate, InStreamableDelegate&& CancelDelegate, TSharedPtr<FStreamableHandle> AssociatedHandle)
 	{
-		FPendingDelegate* PendingDelegate = new FPendingDelegate(Delegate, CancelDelegate, AssociatedHandle);
+		FPendingDelegate* PendingDelegate = new FPendingDelegate(Forward<InStreamableDelegate>(Delegate), Forward<InStreamableDelegate>(CancelDelegate), AssociatedHandle);
 
 		{
 			FScopeLock Lock(&DataLock);
@@ -215,10 +217,11 @@ private:
 		/** Frames left to delay */
 		int32 DelayFrames;
 
-		FPendingDelegate(const FStreamableDelegate& InDelegate, const FStreamableDelegate& InCancelDelegate, TSharedPtr<FStreamableHandle> InHandle)
-			: Delegate(InDelegate)
-			, CancelDelegate(InCancelDelegate)
-			, RelatedHandle(InHandle)
+		template<typename InStreamableDelegate>
+		FPendingDelegate(InStreamableDelegate&& InDelegate, InStreamableDelegate&& InCancelDelegate, TSharedPtr<FStreamableHandle> InHandle)
+			: Delegate(Forward<InStreamableDelegate>(InDelegate))
+			, CancelDelegate(Forward<InStreamableDelegate>(InCancelDelegate))
+			, RelatedHandle(MoveTemp(InHandle))
 			, DelayFrames(GStreamableDelegateDelayFrames)
 		{}
 	};
@@ -358,7 +361,7 @@ bool FStreamableHandle::BindCompleteDelegate(FStreamableDelegate NewDelegate)
 		return false;
 	}
 
-	CompleteDelegate = NewDelegate;
+	CompleteDelegate = MoveTemp(NewDelegate);
 	return true;
 }
 
@@ -370,7 +373,7 @@ bool FStreamableHandle::BindCancelDelegate(FStreamableDelegate NewDelegate)
 		return false;
 	}
 
-	CancelDelegate = NewDelegate;
+	CancelDelegate = MoveTemp(NewDelegate);
 	return true;
 }
 
@@ -382,7 +385,7 @@ bool FStreamableHandle::BindUpdateDelegate(FStreamableUpdateDelegate NewDelegate
 		return false;
 	}
 
-	UpdateDelegate = NewDelegate;
+	UpdateDelegate = MoveTemp(NewDelegate);
 	return true;
 }
 
@@ -612,7 +615,7 @@ void FStreamableHandle::CancelHandle()
 	bCanceled = true;
 	NotifyParentsOfCancellation();
 
-	ExecuteDelegate(CancelDelegate, SharedThis);
+	ExecuteDelegate(MoveTemp(CancelDelegate), SharedThis);
 	UnbindDelegates();
 
 	// Remove from referenced list. If it is stalled then it won't have been registered with
@@ -764,7 +767,7 @@ void FStreamableHandle::CompleteLoad()
 	{
 		bLoadCompleted = true;
 
-		ExecuteDelegate(CompleteDelegate, AsShared(), CancelDelegate);
+		ExecuteDelegate(MoveTemp(CompleteDelegate), AsShared(), MoveTemp(CancelDelegate));
 		UnbindDelegates();
 
 		NotifyParentsOfCompletion();
@@ -922,28 +925,46 @@ void FStreamableHandle::AsyncLoadCallbackWrapper(const FName& PackageName, UPack
 	}
 }
 
-void FStreamableHandle::ExecuteDelegate(const FStreamableDelegate& Delegate, TSharedPtr<FStreamableHandle> AssociatedHandle, const FStreamableDelegate& CancelDelegate)
+namespace Private
 {
-	if (Delegate.IsBound())
+	// Internal helper for executing delegates. This avoids code duplication by allowing the delegate to be moved or copied when deferred depending on which overload of ExecuteDelegate is called.
+	template<typename InStreamableDelegate>
+	void ExecuteDelegateInternal(InStreamableDelegate&& Delegate, TSharedPtr<FStreamableHandle> AssociatedHandle, InStreamableDelegate&& CancelDelegate)
 	{
-		if (GStreamableDelegateDelayFrames == 0)
-		{
-			// Execute it immediately
-			Delegate.Execute();
-		}
-		else
+		using UnRefCvStreamableDelegate = std::remove_cv_t<std::remove_reference_t<InStreamableDelegate>>;
+		static_assert(std::is_same_v<UnRefCvStreamableDelegate, FStreamableDelegate>, "ExecuteDelegateInternal is only valid to be called with FStreamableDelegate.");
+
+		if (Delegate.IsBound())
 		{
-			// Add to execution queue for next tick
-			if (!StreamableDelegateDelayHelper)
+			if (GStreamableDelegateDelayFrames == 0)
 			{
-				StreamableDelegateDelayHelper = new FStreamableDelegateDelayHelper;
+				// Execute it immediately
+				Delegate.Execute();
 			}
+			else
+			{
+				// Add to execution queue for next tick
+				if (!StreamableDelegateDelayHelper)
+				{
+					StreamableDelegateDelayHelper = new FStreamableDelegateDelayHelper;
+				}
 
-			StreamableDelegateDelayHelper->AddDelegate(Delegate, CancelDelegate, AssociatedHandle);
+				StreamableDelegateDelayHelper->AddDelegate(Forward<InStreamableDelegate>(Delegate), Forward<InStreamableDelegate>(CancelDelegate), MoveTemp(AssociatedHandle));
+			}
 		}
 	}
 }
 
+void FStreamableHandle::ExecuteDelegate(const FStreamableDelegate& Delegate, TSharedPtr<FStreamableHandle> AssociatedHandle, const FStreamableDelegate& CancelDelegate)
+{
+	Private::ExecuteDelegateInternal(Delegate, MoveTemp(AssociatedHandle), CancelDelegate);
+}
+
+void FStreamableHandle::ExecuteDelegate(FStreamableDelegate&& Delegate, TSharedPtr<FStreamableHandle> AssociatedHandle, FStreamableDelegate&& CancelDelegate)
+{
+	Private::ExecuteDelegateInternal(MoveTemp(Delegate), MoveTemp(AssociatedHandle), MoveTemp(CancelDelegate));
+}
+
 TSharedPtr<FStreamableHandle> FStreamableHandle::FindMatchingHandle(TFunction<bool(const FStreamableHandle&)> Predicate) const
 {
 	if (Predicate(*this))
@@ -1079,7 +1100,7 @@ struct FStreamable
 					// Keep handle alive to stop the cancel callback from dropping the last reference
 					TSharedPtr<FStreamableHandle> SharedHandle = ActiveHandle->AsShared();
 
-					FStreamableHandle::ExecuteDelegate(ActiveHandle->CancelDelegate, SharedHandle);
+					FStreamableHandle::ExecuteDelegate(MoveTemp(ActiveHandle->CancelDelegate), SharedHandle);
 					ActiveHandle->UnbindDelegates();
 					ActiveHandle->NotifyParentsOfCancellation();
 				}
@@ -1368,7 +1389,7 @@ TSharedPtr<FStreamableHandle> FStreamableManager::RequestAsyncLoad(TArray<FSoftO
 
 	// Schedule a new callback, this will get called when all related async loads are completed
 	TSharedRef<FStreamableHandle> NewRequest = MakeShareable(new FStreamableHandle());
-	NewRequest->CompleteDelegate = DelegateToCall;
+	NewRequest->CompleteDelegate = MoveTemp(DelegateToCall);
 	NewRequest->OwningManager = this;
 	NewRequest->RequestedAssets = MoveTemp(TargetsToStream);
 #if (!PLATFORM_IOS && !PLATFORM_ANDROID)
