Fix ambiguous namespace within StreamableManager.

#rnx

diff --git a/Engine/Source/Runtime/Engine/Private/StreamableManager.cpp b/Engine/Source/Runtime/Engine/Private/StreamableManager.cpp
--- a/Engine/Source/Runtime/Engine/Private/StreamableManager.cpp
+++ b/Engine/Source/Runtime/Engine/Private/StreamableManager.cpp
@@ -925,7 +925,7 @@ void FStreamableHandle::AsyncLoadCallbackWrapper(const FName& PackageName, UPack
 	}
 }
 
-namespace Private
+namespace UE::StreamableManager::Private
 {
 	// Internal helper for executing delegates. This avoids code duplication by allowing the delegate to be moved or copied when deferred depending on which overload of ExecuteDelegate is called.
 	template<typename InStreamableDelegate>
@@ -957,12 +957,12 @@ namespace Private
 
 void FStreamableHandle::ExecuteDelegate(const FStreamableDelegate& Delegate, TSharedPtr<FStreamableHandle> AssociatedHandle, const FStreamableDelegate& CancelDelegate)
 {
-	Private::ExecuteDelegateInternal(Delegate, MoveTemp(AssociatedHandle), CancelDelegate);
+	UE::StreamableManager::Private::ExecuteDelegateInternal(Delegate, MoveTemp(AssociatedHandle), CancelDelegate);
 }
 
 void FStreamableHandle::ExecuteDelegate(FStreamableDelegate&& Delegate, TSharedPtr<FStreamableHandle> AssociatedHandle, FStreamableDelegate&& CancelDelegate)
 {
-	Private::ExecuteDelegateInternal(MoveTemp(Delegate), MoveTemp(AssociatedHandle), MoveTemp(CancelDelegate));
+	UE::StreamableManager::Private::ExecuteDelegateInternal(MoveTemp(Delegate), MoveTemp(AssociatedHandle), MoveTemp(CancelDelegate));
 }
 
 TSharedPtr<FStreamableHandle> FStreamableHandle::FindMatchingHandle(TFunction<bool(const FStreamableHandle&)> Predicate) const
