Optimize ULevel::AddLoadedActors / RemoveLoadedActors.
Saves over 5 seconds of loading/unloading time in a 32k actors map.

diff --git a/Engine/Source/Runtime/Engine/Private/Level.cpp b/Engine/Source/Runtime/Engine/Private/Level.cpp
--- a/Engine/Source/Runtime/Engine/Private/Level.cpp
+++ b/Engine/Source/Runtime/Engine/Private/Level.cpp
@@ -7,6 +7,9 @@ Level.cpp: Level-related functions
 #include "Engine/Level.h"
 
 #include "EngineLogs.h"
+#include "Algo/Copy.h"
+#include "Algo/ForEach.h"
+#include "Algo/Transform.h"
 #include "Misc/ScopedSlowTask.h"
 #include "HAL/PlatformFile.h"
 #include "UObject/UE5MainStreamObjectVersion.h"
@@ -804,17 +807,21 @@ void ULevel::AddLoadedActor(AActor* Actor, const FTransform* TransformToApply)
 
 void ULevel::AddLoadedActors(const TArray<AActor*>& ActorList, const FTransform* TransformToApply)
 {
+	// Actors set used to accelerate lookups for actors already in the level's actor list. This can happen for newly added
+	// actors, as spawning them will add them to the actors array, and grabbing an actor descriptor reference on them will trigger
+	// loading, even if the actor already exists.
+	TSet<AActor*> ActorsSet(ObjectPtrDecay(Actors));
+
+	// Actors queue that will be added to the level, filtered for actors already in the level's actor list.
 	TArray<AActor*> ActorsQueue;
 	ActorsQueue.Reserve(ActorList.Num());
 
-	TFunction<void(AActor* Actor)> QueueActor = [this, &ActorsQueue, &QueueActor](AActor* Actor)
+	TFunction<void(AActor* Actor)> QueueActor = [this, &ActorsSet, &ActorsQueue, &QueueActor](AActor* Actor)
 	{
-		check(Actor);
-		check(Actor->GetLevel() == this);
 		check(IsValidChecked(Actor));
+		check(Actor->GetLevel() == this);		
 
-		int32 ActorIndex;
-		if (!Actors.Find(Actor, ActorIndex))
+		if (!ActorsSet.Contains(Actor))
 		{
 			Actors.Add(Actor);
 			ActorsForGC.Add(Actor);
@@ -845,22 +852,20 @@ void ULevel::AddLoadedActors(const TArray<AActor*>& ActorList, const FTransform*
 	for (AActor* Actor : ActorsQueue)
 	{
 		// RegisterAllComponents can destroy child actors
-		if (!IsValid(Actor))
-		{
-			continue;
-		}
-
-		if (TransformToApply)
+		if (IsValid(Actor))
 		{
-			FLevelUtils::FApplyLevelTransformParams TransformParams(this, *TransformToApply);
-			TransformParams.Actor = Actor;
-			TransformParams.bDoPostEditMove = true;
-			FLevelUtils::ApplyLevelTransform(TransformParams);
-		}
+			if (TransformToApply)
+			{
+				FLevelUtils::FApplyLevelTransformParams TransformParams(this, *TransformToApply);
+				TransformParams.Actor = Actor;
+				TransformParams.bDoPostEditMove = true;
+				FLevelUtils::ApplyLevelTransform(TransformParams);
+			}
 
-		if (bAreComponentsCurrentlyRegistered)
-		{
-			Actor->RegisterAllComponents();
+			if (bAreComponentsCurrentlyRegistered)
+			{
+				Actor->RegisterAllComponents();
+			}
 		}
 
 		SlowTask.EnterProgressFrame(1);
@@ -870,14 +875,12 @@ void ULevel::AddLoadedActors(const TArray<AActor*>& ActorList, const FTransform*
 	for (AActor* Actor : ActorsQueue)
 	{
 		// RegisterAllComponents/RerunConstructionScripts can destroy child actors
-		if (!IsValid(Actor))
-		{
-			continue;
-		}
-
-		if (bAreComponentsCurrentlyRegistered)
+		if (IsValid(Actor))
 		{
-			Actor->RerunConstructionScripts();
+			if (bAreComponentsCurrentlyRegistered)
+			{
+				Actor->RerunConstructionScripts();
+			}
 		}
 
 		SlowTask.EnterProgressFrame(1);
@@ -887,24 +890,22 @@ void ULevel::AddLoadedActors(const TArray<AActor*>& ActorList, const FTransform*
 	for (AActor* Actor : ActorsQueue)
 	{
 		// RegisterAllComponents/RerunConstructionScripts can destroy child actors
-		if (!IsValid(Actor))
+		if (IsValid(Actor))
 		{
-			continue;
-		}
+			if (bAreComponentsCurrentlyRegistered)
+			{
+				GetWorld()->UpdateCullDistanceVolumes(Actor);
+				Actor->MarkComponentsRenderStateDirty();
+			}
 
-		if (bAreComponentsCurrentlyRegistered)
-		{
-			GetWorld()->UpdateCullDistanceVolumes(Actor);
-			Actor->MarkComponentsRenderStateDirty();
-		}
+			if (IsUsingActorFolders() && bFixupActorFoldersAtLoad)
+			{
+				Actor->FixupActorFolder();
+			}
 
-		if (IsUsingActorFolders() && bFixupActorFoldersAtLoad)
-		{
-			Actor->FixupActorFolder();
+			OnLoadedActorAddedToLevelEvent.Broadcast(*Actor);
 		}
 
-		OnLoadedActorAddedToLevelEvent.Broadcast(*Actor);
-
 		SlowTask.EnterProgressFrame(1);
 	}
 
@@ -922,14 +923,41 @@ void ULevel::RemoveLoadedActor(AActor* Actor, const FTransform* TransformToRemov
 
 void ULevel::RemoveLoadedActors(const TArray<AActor*>& ActorList, const FTransform* TransformToRemove)
 {
+	// Build an actor pointer to pair of indices into Actors and ActorsForGC to accelerate removal lookups. This doesn't make much sense
+	// when removing a single actor, but most of the time we are removing actors in large batches and it makes a huge difference in removal time.
+	TMap<AActor*, TPair<int32, int32>> ActorsIndices;
+	{
+		ActorsIndices.Reserve(Actors.Num());
+
+		int32 ActorIndex = 0;
+		Algo::ForEach(Actors, [&ActorsIndices, &ActorIndex](AActor* Actor)
+		{
+			if (Actor)
+			{
+				ActorsIndices.Add(Actor, TPair<int32, int32>(ActorIndex, INDEX_NONE));
+			}
+			ActorIndex++;
+		});
+	
+		int32 ActorForGCIndex = 0;
+		Algo::ForEach(ActorsForGC, [&ActorsIndices, &ActorForGCIndex](AActor* Actor)
+		{
+			if (Actor)
+			{
+				ActorsIndices[Actor].Value = ActorForGCIndex;
+			}
+			ActorForGCIndex++;
+		});
+	}
+
+	// Build the actual actor list that needs removal.
 	TArray<AActor*> ActorsQueue;
 	ActorsQueue.Reserve(ActorList.Num());
 
-	TFunction<void(AActor* Actor)> QueueActor = [this, &ActorsQueue, &QueueActor](AActor* Actor)
+	TFunction<void(AActor* Actor)> QueueActor = [this, &ActorsIndices, &ActorsQueue, &QueueActor](AActor* Actor)
 	{
-		check(Actor);
-		check(Actor->GetLevel() == this);
 		check(IsValidChecked(Actor));
+		check(Actor->GetLevel() == this);		
 
 		// Handle child actors
 		Actor->ForEachComponent<UChildActorComponent>(false, [this, &QueueActor](UChildActorComponent* ChildActorComponent)
@@ -940,12 +968,16 @@ void ULevel::RemoveLoadedActors(const TArray<AActor*>& ActorList, const FTransfo
 			}
 		});
 
-		int32 ActorIndex;
-		// temporarily downgraded to ensure while an issue is fixed
-		if (ensure(Actors.Find(Actor, ActorIndex))) 
+		if (const TPair<int32, int32>* ActorIndices = ActorsIndices.Find(Actor))
 		{
-			Actors[ActorIndex] = nullptr;
-			ActorsForGC.Remove(Actor);
+			check(Actors[ActorIndices->Key] == Actor);
+			Actors[ActorIndices->Key] = nullptr;
+
+			if (ActorIndices->Value != INDEX_NONE)
+			{
+				check(ActorsForGC[ActorIndices->Value] == Actor);
+				ActorsForGC[ActorIndices->Value] = nullptr;
+			}
 
 			ActorsQueue.Add(Actor);
 		}
@@ -964,27 +996,28 @@ void ULevel::RemoveLoadedActors(const TArray<AActor*>& ActorList, const FTransfo
 	for (AActor* Actor : ActorsQueue)
 	{
 		// UnregisterAllComponents can destroy child actors
-		if (!IsValid(Actor))
+		if (IsValid(Actor))
 		{
-			continue;
-		}
-
-		Actor->UnregisterAllComponents();
-		Actor->RegisterAllActorTickFunctions(false, true);
+			Actor->UnregisterAllComponents();
+			Actor->RegisterAllActorTickFunctions(false, true);
 
-		OnLoadedActorRemovedFromLevelEvent.Broadcast(*Actor);
+			OnLoadedActorRemovedFromLevelEvent.Broadcast(*Actor);
 
-		if (TransformToRemove)
-		{
-			FLevelUtils::FApplyLevelTransformParams TransformParams(this, TransformToRemove->Inverse());
-			TransformParams.Actor = Actor;
-			TransformParams.bDoPostEditMove = true;
-			FLevelUtils::ApplyLevelTransform(TransformParams);
+			if (TransformToRemove)
+			{
+				FLevelUtils::FApplyLevelTransformParams TransformParams(this, TransformToRemove->Inverse());
+				TransformParams.Actor = Actor;
+				TransformParams.bDoPostEditMove = true;
+				FLevelUtils::ApplyLevelTransform(TransformParams);
+			}
 		}
 
 		SlowTask.EnterProgressFrame(1);
 	}
 
+	Actors.Remove(nullptr);
+	ActorsForGC.Remove(nullptr);
+
 	TArray<AActor*> ValidActorsQueue;
 	ValidActorsQueue.Reserve(ActorsQueue.Num());
 	Algo::CopyIf(ActorsQueue, ValidActorsQueue, [&](AActor* Actor) { return IsValid(Actor); });
@@ -4552,4 +4585,4 @@ void ULevel::OnMultipleLSAsPopupDismissed()
 
 #endif // WITH_EDITOR
 
-#undef LOCTEXT_NAMESPACE
+#undef LOCTEXT_NAMESPACE
\ No newline at end of file
