[Core] Added GetTypeHash to FBaseComponentReference

#rnx

diff --git a/Engine/Source/Runtime/Engine/Classes/Engine/EngineTypes.h b/Engine/Source/Runtime/Engine/Classes/Engine/EngineTypes.h
--- a/Engine/Source/Runtime/Engine/Classes/Engine/EngineTypes.h
+++ b/Engine/Source/Runtime/Engine/Classes/Engine/EngineTypes.h
@@ -3317,6 +3317,15 @@ struct FBaseComponentReference
 	}
 };
 
+inline uint32 GetTypeHash(const FBaseComponentReference& Reference)
+{
+	return HashCombineFast(
+		HashCombineFast(
+			GetTypeHash(Reference.ComponentProperty),
+			GetTypeHash(Reference.PathToComponent)),
+			GetTypeHash(Reference.OverrideComponent));
+}
+
 /** 
  *	Struct that allows for different ways to reference a component using TObjectPtr. 
  *	If just an Actor is specified, will return RootComponent of that Actor.
