Fixed ULevelStreamingLevelInstance::LoadInstance to properly handle a streaming level load while the underlying streaming level already exists but has not executed/completed.
Added bAllowReuseExitingLevelStreaming to FLoadLevelInstanceParams to allow reusing an existing streaming level (defaulted to false).

diff --git a/Engine/Source/Runtime/Engine/Classes/Engine/LevelStreamingDynamic.h b/Engine/Source/Runtime/Engine/Classes/Engine/LevelStreamingDynamic.h
--- a/Engine/Source/Runtime/Engine/Classes/Engine/LevelStreamingDynamic.h
+++ b/Engine/Source/Runtime/Engine/Classes/Engine/LevelStreamingDynamic.h
@@ -57,6 +57,9 @@ class ULevelStreamingDynamic : public ULevelStreaming
 
 		/** Set whether the level will be made visible initially. */
 		bool bInitiallyVisible = true;
+
+		/** Set whether we allow to reuse an existing level streaming. */
+		bool bAllowReuseExitingLevelStreaming = false;
 	};
 
 	/**  
diff --git a/Engine/Source/Runtime/Engine/Private/LevelInstance/LevelInstanceLevelStreaming.cpp b/Engine/Source/Runtime/Engine/Private/LevelInstance/LevelInstanceLevelStreaming.cpp
--- a/Engine/Source/Runtime/Engine/Private/LevelInstance/LevelInstanceLevelStreaming.cpp
+++ b/Engine/Source/Runtime/Engine/Private/LevelInstance/LevelInstanceLevelStreaming.cpp
@@ -247,6 +247,7 @@ ULevelStreamingLevelInstance* ULevelStreamingLevelInstance::LoadInstance(ILevelI
 	if (World->IsGameWorld())
 	{
 		Params.bInitiallyVisible = LevelInstance->IsInitiallyVisible();
+		Params.bAllowReuseExitingLevelStreaming = true;
 	}
 
 	const FString LevelInstancePackageName = ULevelStreamingDynamic::GetLevelInstancePackageName(Params);
diff --git a/Engine/Source/Runtime/Engine/Private/LevelInstance/LevelInstanceSubsystem.cpp b/Engine/Source/Runtime/Engine/Private/LevelInstance/LevelInstanceSubsystem.cpp
--- a/Engine/Source/Runtime/Engine/Private/LevelInstance/LevelInstanceSubsystem.cpp
+++ b/Engine/Source/Runtime/Engine/Private/LevelInstance/LevelInstanceSubsystem.cpp
@@ -451,6 +451,14 @@ void ULevelInstanceSubsystem::LoadLevelInstance(ILevelInstanceInterface* LevelIn
 #if WITH_EDITOR
 		check(LevelInstanceActor->GetWorld()->IsGameWorld() || LoadedLevelInstances.Contains(LevelInstanceID));
 #endif
+		// If still considered loading but level streaming was reused and its level is loaded
+		if (GetWorld()->IsGameWorld() && IsLoading(LevelInstance) && LevelStreaming->GetLoadedLevel())
+		{
+			// Register the loaded level instance
+			RegisterLoadedLevelStreamingLevelInstance(LevelStreaming);
+			check(!IsLoading(LevelInstance));
+			check(IsLoaded(LevelInstance));
+		}
 	}
 	else
 	{
@@ -576,7 +584,11 @@ ILevelInstanceInterface* ULevelInstanceSubsystem::GetOwningLevelInstance(const U
 		}
 		else if (UWorldPartitionLevelStreamingDynamic* WorldPartitionLevelStreaming = Cast<UWorldPartitionLevelStreamingDynamic>(BaseLevelStreaming))
 		{
-			return GetOwningLevelInstance(WorldPartitionLevelStreaming->GetStreamingWorld()->PersistentLevel);
+			// Instanced world partition might be uninitialized here, can't resolved if it's the case.
+			if (UWorld* StreamingWorld = WorldPartitionLevelStreaming->GetStreamingWorld())
+			{
+				return GetOwningLevelInstance(StreamingWorld->PersistentLevel);
+			}
 		}
 	}
 
diff --git a/Engine/Source/Runtime/Engine/Private/LevelStreaming.cpp b/Engine/Source/Runtime/Engine/Private/LevelStreaming.cpp
--- a/Engine/Source/Runtime/Engine/Private/LevelStreaming.cpp
+++ b/Engine/Source/Runtime/Engine/Private/LevelStreaming.cpp
@@ -2548,11 +2548,31 @@ ULevelStreamingDynamic* ULevelStreamingDynamic::LoadLevelInstance_Internal(const
 			ModifiedLevelPackageName = FName(*UWorld::ConvertToPIEPackageName(LevelPackageNameStr, PIEInstance));
 		}
 #endif
-		if (Params.World->GetStreamingLevels().ContainsByPredicate([&ModifiedLevelPackageName](ULevelStreaming* LS) { return LS && LS->GetWorldAssetPackageFName() == ModifiedLevelPackageName; }))
-		{
-			// The streaming level already exists, error and return.
-			UE_LOG(LogLevelStreaming, Error, TEXT("LoadLevelInstance called with a name that already exists, returning nullptr. LevelPackageName:%s"), *ModifiedLevelPackageName.ToString());
-			return nullptr;
+		// Test if the streaming level already exists
+		if (ULevelStreaming* const* ExistingLevelStreaming = Params.World->GetStreamingLevels().FindByPredicate([&ModifiedLevelPackageName](ULevelStreaming* LS) { return LS && LS->GetWorldAssetPackageFName() == ModifiedLevelPackageName; }))
+		{
+			// Allow reusing a streaming level only if :
+			// - Params.bAllowReuseExitingLevelStreaming is true
+			// - Params.World is a game world
+			// - Existing LevelStreaming has the same Class
+			// - Existing LevelStreaming has the same Level Transform
+			ULevelStreamingDynamic* StreamingLevel = Params.bAllowReuseExitingLevelStreaming && Params.World->IsGameWorld() ? Cast<ULevelStreamingDynamic>(*ExistingLevelStreaming) : nullptr;
+			if (StreamingLevel && 
+				StreamingLevel->GetClass() == LevelStreamingClass &&
+				StreamingLevel->LevelTransform.Equals(Params.LevelTransform))
+			{
+				bOutSuccess = true;
+				StreamingLevel->SetShouldBeLoaded(true);
+				StreamingLevel->SetShouldBeVisible(Params.bInitiallyVisible);
+				StreamingLevel->SetIsRequestingUnloadAndRemoval(false);
+				UE_LOG(LogLevelStreaming, Verbose, TEXT("LoadLevelInstance found existing StreamingLevel for LevelPackageName:%s"), *ModifiedLevelPackageName.ToString());
+				return StreamingLevel;
+			}
+			else
+			{
+				UE_LOG(LogLevelStreaming, Error, TEXT("LoadLevelInstance called with a name that already exists, returning nullptr. LevelPackageName:%s"), *ModifiedLevelPackageName.ToString());
+				return nullptr;
+			}
 		}
 	}
     
