ShapeComponent: Move LineThickness to ShapeComponent instead of only using it in some derived classes.

Provides settable LineThickness for SphereComponents, and CapsuleComponents (and any derived classes that don't override GetSceneProxy, in which case LineThickness would still be exposed but would be unused unless the new scene proxy made use of it).

diff --git a/Engine/Source/Runtime/Engine/Classes/Components/BoxComponent.h b/Engine/Source/Runtime/Engine/Classes/Components/BoxComponent.h
--- a/Engine/Source/Runtime/Engine/Classes/Components/BoxComponent.h
+++ b/Engine/Source/Runtime/Engine/Classes/Components/BoxComponent.h
@@ -23,10 +23,6 @@ protected:
 	UPROPERTY(EditAnywhere, BlueprintReadOnly, export, Category=Shape)
 	FVector BoxExtent;
 
-	/** Used to control the line thickness when rendering */
-	UPROPERTY(EditAnywhere, BlueprintReadOnly, export, Category=Shape)
-	float LineThickness;
-
 public:
 	/** 
 	 * Change the box extent size. This is the unscaled size, before component scale is applied.
@@ -36,10 +32,6 @@ public:
 	UFUNCTION(BlueprintCallable, Category="Components|Box")
 	ENGINE_API void SetBoxExtent(FVector InBoxExtent, bool bUpdateOverlaps=true);
 
-	// Set the LineThickness
-	UFUNCTION(BlueprintCallable, Category="Components|Box")
-	ENGINE_API void SetLineThickness(float Thickness);
-
 	// @return the box extent, scaled by the component scale.
 	UFUNCTION(BlueprintCallable, Category="Components|Box")
 	ENGINE_API FVector GetScaledBoxExtent() const;
diff --git a/Engine/Source/Runtime/Engine/Classes/Components/ShapeComponent.h b/Engine/Source/Runtime/Engine/Classes/Components/ShapeComponent.h
--- a/Engine/Source/Runtime/Engine/Classes/Components/ShapeComponent.h
+++ b/Engine/Source/Runtime/Engine/Classes/Components/ShapeComponent.h
@@ -43,6 +43,10 @@ class UShapeComponent : public UPrimitiveComponent
 	UPROPERTY(EditAnywhere, AdvancedDisplay, BlueprintReadOnly, Category=Shape)
 	FColor ShapeColor;
 
+	/** Used to control the line thickness when rendering */
+	UPROPERTY(EditAnywhere, BlueprintReadOnly, export, Category=Shape)
+	float LineThickness;
+	
 	/** Only show this component if the actor is selected */
 	UPROPERTY()
 	uint8 bDrawOnlyIfSelected:1;
@@ -126,6 +130,10 @@ public:
 	ENGINE_API TSubclassOf<class UNavAreaBase> GetDesiredAreaClass() const;
 	ENGINE_API void SetAreaClassOverride(TSubclassOf<class UNavAreaBase> InAreaClassOverride);
 	ENGINE_API void SetUseSystemDefaultObstacleAreaClass();
+
+	// Set the LineThickness
+    UFUNCTION(BlueprintCallable, Category="Components|Shape")
+    ENGINE_API void SetLineThickness(float Thickness);
 };
 
 enum class EShapeBodySetupHelper
diff --git a/Engine/Source/Runtime/Engine/Private/Components/BoxComponent.cpp b/Engine/Source/Runtime/Engine/Private/Components/BoxComponent.cpp
--- a/Engine/Source/Runtime/Engine/Private/Components/BoxComponent.cpp
+++ b/Engine/Source/Runtime/Engine/Private/Components/BoxComponent.cpp
@@ -42,12 +42,6 @@ void UBoxComponent::SetBoxExtent(FVector NewBoxExtent, bool bUpdateOverlaps)
 	}
 }
 
-void UBoxComponent::SetLineThickness(float Thickness)
-{
-	LineThickness = Thickness;
-	MarkRenderStateDirty();
-}
-
 template <EShapeBodySetupHelper UpdateBodySetupAction, typename BodySetupType>
 bool InvalidateOrUpdateBoxBodySetup(BodySetupType& ShapeBodySetup, bool bUseArchetypeBodySetup, FVector BoxExtent)
 {
@@ -183,7 +177,7 @@ FPrimitiveSceneProxy* UBoxComponent::CreateSceneProxy()
 		const uint32	bDrawOnlyIfSelected:1;
 		const FVector	BoxExtents;
 		const FColor	BoxColor;
-		const float LineThickness;
+		const float		LineThickness;
 	};
 
 	return new FBoxSceneProxy( this );
diff --git a/Engine/Source/Runtime/Engine/Private/Components/CapsuleComponent.cpp b/Engine/Source/Runtime/Engine/Private/Components/CapsuleComponent.cpp
--- a/Engine/Source/Runtime/Engine/Private/Components/CapsuleComponent.cpp
+++ b/Engine/Source/Runtime/Engine/Private/Components/CapsuleComponent.cpp
@@ -44,6 +44,7 @@ FPrimitiveSceneProxy* UCapsuleComponent::CreateSceneProxy()
 			,	CapsuleRadius( InComponent->GetScaledCapsuleRadius() )
 			,	CapsuleHalfHeight( InComponent->GetScaledCapsuleHalfHeight() )
 			,	ShapeColor( InComponent->ShapeColor )
+			,	LineThickness( InComponent->LineThickness )
 		{
 			bWillEverBeLit = false;
 		}
@@ -65,7 +66,7 @@ FPrimitiveSceneProxy* UCapsuleComponent::CreateSceneProxy()
 					const FLinearColor DrawCapsuleColor = GetViewSelectionColor(ShapeColor, *View, IsSelected(), IsHovered(), false, IsIndividuallySelected() );
 
 					FPrimitiveDrawInterface* PDI = Collector.GetPDI(ViewIndex);
-					DrawWireCapsule( PDI, LocalToWorld.GetOrigin(), LocalToWorld.GetUnitAxis(EAxis::X), LocalToWorld.GetUnitAxis(EAxis::Y), LocalToWorld.GetUnitAxis(EAxis::Z), DrawCapsuleColor, CapsuleRadius, CapsuleHalfHeight, CapsuleSides, SDPG_World );
+					DrawWireCapsule( PDI, LocalToWorld.GetOrigin(), LocalToWorld.GetUnitAxis(EAxis::X), LocalToWorld.GetUnitAxis(EAxis::Y), LocalToWorld.GetUnitAxis(EAxis::Z), DrawCapsuleColor, CapsuleRadius, CapsuleHalfHeight, CapsuleSides, SDPG_World, LineThickness );
 				}
 			}
 		}
@@ -92,6 +93,7 @@ FPrimitiveSceneProxy* UCapsuleComponent::CreateSceneProxy()
 		const float		CapsuleRadius;
 		const float		CapsuleHalfHeight;
 		const FColor	ShapeColor;
+		const float		LineThickness;
 	};
 
 	return new FDrawCylinderSceneProxy( this );
diff --git a/Engine/Source/Runtime/Engine/Private/Components/ShapeComponent.cpp b/Engine/Source/Runtime/Engine/Private/Components/ShapeComponent.cpp
--- a/Engine/Source/Runtime/Engine/Private/Components/ShapeComponent.cpp
+++ b/Engine/Source/Runtime/Engine/Private/Components/ShapeComponent.cpp
@@ -254,6 +254,12 @@ void UShapeComponent::CreateShapeBodySetupIfNeeded()
 	}
 }
 
+void UShapeComponent::SetLineThickness(float Thickness)
+{
+	LineThickness = Thickness;
+	MarkRenderStateDirty();
+}
+
 //Explicit instantiation of the different shape components
 template ENGINE_API void UShapeComponent::CreateShapeBodySetupIfNeeded<FKSphylElem>();
 template ENGINE_API void UShapeComponent::CreateShapeBodySetupIfNeeded<FKBoxElem>();
diff --git a/Engine/Source/Runtime/Engine/Private/Components/SphereComponent.cpp b/Engine/Source/Runtime/Engine/Private/Components/SphereComponent.cpp
--- a/Engine/Source/Runtime/Engine/Private/Components/SphereComponent.cpp
+++ b/Engine/Source/Runtime/Engine/Private/Components/SphereComponent.cpp
@@ -122,6 +122,7 @@ FPrimitiveSceneProxy* USphereComponent::CreateSceneProxy()
 			,	bDrawOnlyIfSelected( InComponent->bDrawOnlyIfSelected )
 			,	SphereColor(InComponent->ShapeColor)
 			,	SphereRadius(InComponent->SphereRadius)
+			,	LineThickness(InComponent->LineThickness)
 		{
 			bWillEverBeLit = false;
 		}
@@ -159,9 +160,9 @@ FPrimitiveSceneProxy* USphereComponent::CreateSceneProxy()
 					FVector ScaledZ = LocalToWorld.GetUnitAxis(EAxis::Z) * MinAbsScale;
 
 					const int32 SphereSides = FMath::Clamp<int32>(SphereRadius / 4.f, 16, 64);
-					DrawCircle(PDI, LocalToWorld.GetOrigin(), ScaledX, ScaledY, DrawSphereColor, SphereRadius, SphereSides, SDPG_World);
-					DrawCircle(PDI, LocalToWorld.GetOrigin(), ScaledX, ScaledZ, DrawSphereColor, SphereRadius, SphereSides, SDPG_World);
-					DrawCircle(PDI, LocalToWorld.GetOrigin(), ScaledY, ScaledZ, DrawSphereColor, SphereRadius, SphereSides, SDPG_World);
+					DrawCircle(PDI, LocalToWorld.GetOrigin(), ScaledX, ScaledY, DrawSphereColor, SphereRadius, SphereSides, SDPG_World, LineThickness);
+					DrawCircle(PDI, LocalToWorld.GetOrigin(), ScaledX, ScaledZ, DrawSphereColor, SphereRadius, SphereSides, SDPG_World, LineThickness);
+					DrawCircle(PDI, LocalToWorld.GetOrigin(), ScaledY, ScaledZ, DrawSphereColor, SphereRadius, SphereSides, SDPG_World, LineThickness);
 				}
 			}
 		}
@@ -189,6 +190,7 @@ FPrimitiveSceneProxy* USphereComponent::CreateSceneProxy()
 		const uint32				bDrawOnlyIfSelected:1;
 		const FColor				SphereColor;
 		const float					SphereRadius;
+		const float					LineThickness;
 	};
 
 	return new FSphereSceneProxy( this );
