Changed UWorldPartitionLevelStreamingDynamic OuterWorldPartition from WeakObjectPtr to SoftobjectPtr.
Replaced check by condition on valid OuterWorldPartition as UWorldPartitionLevelStreamingDynamic::GetStreamingWorld() can be called on an streaming level part of the world.
The use case is when unregistering a world partition, we might visit a levelstreaming of a cell that has not yet been processed and removed from the list, but it's outer world partition has been uninitialized.

diff --git a/Engine/Source/Runtime/Engine/Private/WorldPartition/WorldPartitionLevelStreamingDynamic.cpp b/Engine/Source/Runtime/Engine/Private/WorldPartition/WorldPartitionLevelStreamingDynamic.cpp
--- a/Engine/Source/Runtime/Engine/Private/WorldPartition/WorldPartitionLevelStreamingDynamic.cpp
+++ b/Engine/Source/Runtime/Engine/Private/WorldPartition/WorldPartitionLevelStreamingDynamic.cpp
@@ -654,8 +654,7 @@ UWorld* UWorldPartitionLevelStreamingDynamic::GetStreamingWorld() const
 {
 	// For UWorldPartitionLevelStreamingDynamic the StreamingWorld is the world to which the OuterWorldPartition is outered.
 	// This World can be used to resolved SoftObjectPaths between cells.
-	check(OuterWorldPartition.IsValid());
-	return OuterWorldPartition->GetTypedOuter<UWorld>();
+	return OuterWorldPartition.IsValid() ? OuterWorldPartition->GetTypedOuter<UWorld>() : nullptr;
 }
 
 bool UWorldPartitionLevelStreamingDynamic::CanChangeVisibility(bool bMakeVisible) const
diff --git a/Engine/Source/Runtime/Engine/Private/WorldPartition/WorldPartitionSubsystem.cpp b/Engine/Source/Runtime/Engine/Private/WorldPartition/WorldPartitionSubsystem.cpp
--- a/Engine/Source/Runtime/Engine/Private/WorldPartition/WorldPartitionSubsystem.cpp
+++ b/Engine/Source/Runtime/Engine/Private/WorldPartition/WorldPartitionSubsystem.cpp
@@ -547,8 +547,8 @@ void UWorldPartitionSubsystem::OnWorldPartitionUninitialized(UWorldPartition* In
 	const UWorld* OwningWorld = GetWorld();
 	if (OwningWorld->IsGameWorld())
 	{
-		TOptional<TSet<TWeakObjectPtr<ULevelStreaming>>*> PendingStreamingLevels;
-		auto GetPendingStreamingLevels = [this, InWorldPartition, &PendingStreamingLevels]() -> TSet<TWeakObjectPtr<ULevelStreaming>>&
+		TOptional<TSet<TWeakObjectPtr<UWorldPartitionLevelStreamingDynamic>>*> PendingStreamingLevels;
+		auto GetPendingStreamingLevels = [this, InWorldPartition, &PendingStreamingLevels]() -> TSet<TWeakObjectPtr<UWorldPartitionLevelStreamingDynamic>>&
 		{
 			if (!PendingStreamingLevels.IsSet())
 			{
@@ -557,14 +557,15 @@ void UWorldPartitionSubsystem::OnWorldPartitionUninitialized(UWorldPartition* In
 			return *PendingStreamingLevels.GetValue();
 		};
 
-		const UWorld* WorldPartitionOuterWorld = InWorldPartition->GetTypedOuter<UWorld>();
-		if (WorldPartitionOuterWorld != OwningWorld)
+		if (InWorldPartition->GetTypedOuter<UWorld>() != OwningWorld)
 		{
+			const FSoftObjectPath WorldPartition(InWorldPartition);
 			for (ULevelStreaming* StreamingLevel : OwningWorld->GetStreamingLevels())
 			{
-				if (StreamingLevel->GetStreamingWorld() == WorldPartitionOuterWorld)
+				UWorldPartitionLevelStreamingDynamic* WorldPartitionStreamingLevel = Cast<UWorldPartitionLevelStreamingDynamic>(StreamingLevel);
+				if (WorldPartitionStreamingLevel && (WorldPartitionStreamingLevel->GetOuterWorldPartition() == WorldPartition))
 				{
-					GetPendingStreamingLevels().Add(StreamingLevel);
+					GetPendingStreamingLevels().Add(WorldPartitionStreamingLevel);
 				}
 			}
 		}
@@ -590,7 +591,7 @@ void UWorldPartitionSubsystem::OnWorldPartitionUninitialized(UWorldPartition* In
 
 bool UWorldPartitionSubsystem::HasUninitializationPendingStreamingLevels(const UWorldPartition* InWorldPartition) const
 {
-	if (const TSet<TWeakObjectPtr<ULevelStreaming>>* PendingStreamingLevels = InWorldPartition ? WorldPartitionUninitializationPendingStreamingLevels.Find(FSoftObjectPath(InWorldPartition)) : nullptr)
+	if (const TSet<TWeakObjectPtr<UWorldPartitionLevelStreamingDynamic>>* PendingStreamingLevels = InWorldPartition ? WorldPartitionUninitializationPendingStreamingLevels.Find(FSoftObjectPath(InWorldPartition)) : nullptr)
 	{
 		if (ensure(!PendingStreamingLevels->IsEmpty()))
 		{
@@ -663,16 +664,17 @@ void UWorldPartitionSubsystem::OnLevelStreamingStateChanged(UWorld* InWorld, con
 
 	if (NewState == ELevelStreamingState::Removed)
 	{
-		const UWorld* OuterWorld = InStreamingLevel->GetStreamingWorld();
-		if (const UWorldPartition* OuterWorldPartition = OuterWorld && OuterWorld->IsGameWorld() ? OuterWorld->GetWorldPartition() : nullptr)
+		if (const UWorldPartitionLevelStreamingDynamic* WorldPartitionStreamingLevel = Cast<const UWorldPartitionLevelStreamingDynamic>(InStreamingLevel))
 		{
-			if (TSet<TWeakObjectPtr<ULevelStreaming>>* PendingStreamingLevels = WorldPartitionUninitializationPendingStreamingLevels.Find(FSoftObjectPath(OuterWorldPartition)))
+			check(InWorld->IsGameWorld());
+			const FSoftObjectPath& WorldPartition = WorldPartitionStreamingLevel->GetOuterWorldPartition();
+			if (TSet<TWeakObjectPtr<UWorldPartitionLevelStreamingDynamic>>* PendingStreamingLevels = WorldPartitionUninitializationPendingStreamingLevels.Find(WorldPartition))
 			{
-				if (PendingStreamingLevels->Remove(InStreamingLevel))
+				if (PendingStreamingLevels->Remove(WorldPartitionStreamingLevel))
 				{
 					if (PendingStreamingLevels->IsEmpty())
 					{
-						WorldPartitionUninitializationPendingStreamingLevels.Remove(OuterWorldPartition);
+						WorldPartitionUninitializationPendingStreamingLevels.Remove(WorldPartition);
 					}
 				}
 			}
@@ -843,7 +845,7 @@ void UWorldPartitionSubsystem::DumpWorldPartitions(FOutputDevice& OutputDevice)
 {
 	if (RegisteredWorldPartitions.Num() > 0)
 	{
-		OutputDevice.Logf(TEXT("Registered World Partitions:"));
+		OutputDevice.Logf(TEXT("Registered World Partitions for %s:"), *GetWorld()->GetPathName());
 		TArray<FString> WorldPartitions;
 		Algo::ForEach(RegisteredWorldPartitions, [&WorldPartitions](const UWorldPartition* WorldPartition) { WorldPartitions.Add(WorldPartition->GetPathName()); });
 		WorldPartitions.Sort();
@@ -855,7 +857,7 @@ void UWorldPartitionSubsystem::DumpStreamingSources(FOutputDevice& OutputDevice)
 {
 	if (StreamingSources.Num() > 0)
 	{
-		OutputDevice.Logf(TEXT("Streaming Sources:"));
+		OutputDevice.Logf(TEXT("Streaming Sources for %s:"), *GetWorld()->GetPathName());
 		for (const FWorldPartitionStreamingSource& StreamingSource : StreamingSources)
 		{
 			OutputDevice.Logf(TEXT("  - %s: %s"), *StreamingSource.Name.ToString(), *StreamingSource.ToString());
diff --git a/Engine/Source/Runtime/Engine/Public/WorldPartition/WorldPartitionLevelStreamingDynamic.h b/Engine/Source/Runtime/Engine/Public/WorldPartition/WorldPartitionLevelStreamingDynamic.h
--- a/Engine/Source/Runtime/Engine/Public/WorldPartition/WorldPartitionLevelStreamingDynamic.h
+++ b/Engine/Source/Runtime/Engine/Public/WorldPartition/WorldPartitionLevelStreamingDynamic.h
@@ -11,6 +11,7 @@
 
 #include "CoreMinimal.h"
 #include "UObject/ObjectMacros.h"
+#include "UObject/SoftObjectPtr.h"
 #include "Engine/LevelStreaming.h"
 #include "Engine/LevelStreamingDynamic.h"
 #include "WorldPartition/WorldPartition.h"
@@ -45,6 +46,7 @@ class UWorldPartitionLevelStreamingDynamic : public ULevelStreamingDynamic
 
 	ENGINE_API void Initialize(const UWorldPartitionRuntimeLevelStreamingCell& InCell);
 	ENGINE_API void SetLevelTransform(const FTransform& InLevelTransform);
+	ENGINE_API const FSoftObjectPath& GetOuterWorldPartition() const { return OuterWorldPartition.ToSoftObjectPath(); }
 
 	virtual bool CanReplicateStreamingStatus() const override { return false; }
 
@@ -105,5 +107,5 @@ private:
 	TWeakObjectPtr<const UWorldPartitionRuntimeLevelStreamingCell> StreamingCell;
 
 	UPROPERTY()
-	TWeakObjectPtr<UWorldPartition> OuterWorldPartition;
+	TSoftObjectPtr<UWorldPartition> OuterWorldPartition;
 };
diff --git a/Engine/Source/Runtime/Engine/Public/WorldPartition/WorldPartitionSubsystem.h b/Engine/Source/Runtime/Engine/Public/WorldPartition/WorldPartitionSubsystem.h
--- a/Engine/Source/Runtime/Engine/Public/WorldPartition/WorldPartitionSubsystem.h
+++ b/Engine/Source/Runtime/Engine/Public/WorldPartition/WorldPartitionSubsystem.h
@@ -15,6 +15,7 @@ class ULevel;
 class ULevelStreaming;
 class UWorldPartition;
 class UActorDescContainer;
+class UWorldPartitionLevelStreamingDynamic;
 class FWorldPartitionActorDesc;
 class FWorldPartitionDraw2DContext;
 
@@ -218,7 +219,7 @@ private:
 	int32 LevelStreamingForceGCAfterLevelStreamedOut;
 
 	// Tracks streaming levels of uninitialized world partition used to delay next initialization until they're done being removed
-	TMap<FSoftObjectPath, TSet<TWeakObjectPtr<ULevelStreaming>>> WorldPartitionUninitializationPendingStreamingLevels;
+	TMap<FSoftObjectPath, TSet<TWeakObjectPtr<UWorldPartitionLevelStreamingDynamic>>> WorldPartitionUninitializationPendingStreamingLevels;
 
 	// Tracks world partition loading and pending loads
 	TSet<TWeakObjectPtr<const ULevelStreaming>> WorldPartitionLoadingAndPendingLoadStreamingLevels;
