Added ensure to prevent archetypes from being added to the replicated subobject list

#rnx

diff --git a/Engine/Source/Runtime/Engine/Private/ActorReplication.cpp b/Engine/Source/Runtime/Engine/Private/ActorReplication.cpp
--- a/Engine/Source/Runtime/Engine/Private/ActorReplication.cpp
+++ b/Engine/Source/Runtime/Engine/Private/ActorReplication.cpp
@@ -644,6 +644,12 @@ void AActor::AddReplicatedSubObject(UObject* SubObject, ELifetimeCondition NetCo
 		return;
 	}
 
+	if (SubObject->HasAnyFlags(RF_ArchetypeObject | RF_ClassDefaultObject))
+	{
+		ensureMsgf(false, TEXT("AddReplicatedSubObject cannot replicate %s for %s. Archetypes or CDO's cannot be replicated subobject's."), *GetNameSafe(SubObject), *GetNameSafe(this));
+		return;
+	}
+
 	ensureMsgf(IsUsingRegisteredSubObjectList(), TEXT("%s is registering subobjects but bReplicateUsingRegisteredSubObjectList is false. Without the flag set to true the registered subobjects will not be replicated."), *GetName());
 
 	FSubObjectRegistry::EResult Result = ReplicatedSubObjects.AddSubObjectUnique(SubObject, NetCondition);
@@ -809,6 +815,12 @@ void AActor::AddActorComponentReplicatedSubObject(UActorComponent* OwnerComponen
 		ensureMsgf(false, TEXT("Ignoring AddReplicatedSubObject for %s::%s. Invalid pointer received."), *GetNameSafe(this), *GetNameSafe(OwnerComponent));
 		return;
 	}
+
+	if (SubObject->HasAnyFlags(RF_ArchetypeObject | RF_ClassDefaultObject))
+	{
+		ensureMsgf(false, TEXT("AddReplicatedSubObject cannot replicate %s for %s::%s. Archetypes or CDO's cannot be replicated subobject's."), *GetNameSafe(SubObject), *GetNameSafe(this), *GetNameSafe(OwnerComponent));
+		return;
+	}
 	
 	ensureMsgf(OwnerComponent->GetIsReplicated(), TEXT("Only components with replication enabled can register subobjects. %s::%s has replication disabled."), *GetName(), *SubObject->GetName());
 	ensureMsgf(NetCondition != COND_Custom, TEXT("Custom netconditions do not work with SubObjects. %s::%s will not be replicated."), *GetName(), *SubObject->GetName());
