Added ResetEditorProperty to reset a given property back to its archetype value

diff --git a/Engine/Source/Runtime/Engine/Classes/Kismet/KismetSystemLibrary.h b/Engine/Source/Runtime/Engine/Classes/Kismet/KismetSystemLibrary.h
--- a/Engine/Source/Runtime/Engine/Classes/Kismet/KismetSystemLibrary.h
+++ b/Engine/Source/Runtime/Engine/Classes/Kismet/KismetSystemLibrary.h
@@ -1952,6 +1952,18 @@ class UKismetSystemLibrary : public UBlueprintFunctionLibrary
     static ENGINE_API bool SetEditorProperty(UObject* Object, const FName PropertyName, const int32& PropertyValue, const EPropertyAccessChangeNotifyMode ChangeNotifyMode);
 	static ENGINE_API bool Generic_SetEditorProperty(UObject* Object, const FName PropertyName, const void* ValuePtr, const FProperty* ValueProp, const EPropertyAccessChangeNotifyMode ChangeNotifyMode);
 	DECLARE_FUNCTION(execSetEditorProperty);
+
+	/**
+	 * Attempts to reset the value of a named property on the given object so that it matches the value of the archetype.
+	 *
+	 * @param Object The object you want to reset a property value on.
+	 * @param PropertyName The name of the object property to reset the value of.
+	 * @param ChangeNotifyMode When to emit property change notifications.
+	 *
+	 * @return Whether the property value was found and correctly reset.
+	 */
+	UFUNCTION(BlueprintCallable, Category = "Utilities", meta=(ScriptMethod))
+	static ENGINE_API bool ResetEditorProperty(UObject* Object, const FName PropertyName, const EPropertyAccessChangeNotifyMode ChangeNotifyMode = EPropertyAccessChangeNotifyMode::Default);
 #endif
 
 	// --- Transactions ------------------------------
diff --git a/Engine/Source/Runtime/Engine/Private/KismetSystemLibrary.cpp b/Engine/Source/Runtime/Engine/Private/KismetSystemLibrary.cpp
--- a/Engine/Source/Runtime/Engine/Private/KismetSystemLibrary.cpp
+++ b/Engine/Source/Runtime/Engine/Private/KismetSystemLibrary.cpp
@@ -3328,6 +3328,60 @@ DEFINE_FUNCTION(UKismetSystemLibrary::execSetEditorProperty)
 	*(bool*)RESULT_PARAM = bResult;
 }
 
+bool UKismetSystemLibrary::ResetEditorProperty(UObject* Object, const FName PropertyName, const EPropertyAccessChangeNotifyMode ChangeNotifyMode)
+{
+	if (!Object)
+	{
+		LogRuntimeError(NSLOCTEXT("KismetSystemLibrary", "ResetEditorProperty_AccessNone", "Accessed None attempting to call ResetEditorProperty."));
+		return false;
+	}
+
+	auto FindArchetypeValue = [Object, PropertyName](const FProperty*& OutArchetypeProperty, const void*& OutArchetypeValuePtr)
+	{
+		OutArchetypeProperty = nullptr;
+		OutArchetypeValuePtr = nullptr;
+
+		const FProperty* ObjectProp = PropertyAccessUtil::FindPropertyByName(PropertyName, Object->GetClass());
+		if ((!ObjectProp || ObjectProp->HasAnyPropertyFlags(CPF_Deprecated)) && Object->HasAllFlags(RF_ClassDefaultObject))
+		{
+			// look for a sparse member of the same name - the sparse data is treated as an extension
+			// of the class default object by the details panel:
+			if (const UScriptStruct* SparseDataStruct = Object->GetClass()->GetSparseClassDataStruct())
+			{
+				if (const FProperty* SparseProp = PropertyAccessUtil::FindPropertyByName(PropertyName, SparseDataStruct))
+				{
+					if (UScriptStruct* SparseDataArchetypeStruct = Object->GetClass()->GetSparseClassDataArchetypeStruct())
+					{
+						OutArchetypeProperty = SparseProp;
+						OutArchetypeValuePtr = SparseProp->ContainerPtrToValuePtrForDefaults<const void>(SparseDataArchetypeStruct, Object->GetClass()->GetArchetypeForSparseClassData());
+					}
+					return;
+				}
+			}
+		}
+
+		if (ObjectProp)
+		{
+			if (const UObject* ObjectArchetype = Object->GetArchetype())
+			{
+				OutArchetypeProperty = ObjectProp;
+				OutArchetypeValuePtr = ObjectProp->ContainerPtrToValuePtrForDefaults<const void>(ObjectArchetype->GetClass(), ObjectArchetype);
+			}
+		}
+	};
+
+	const FProperty* ArchetypeProperty = nullptr;
+	const void* ArchetypeValuePtr = nullptr;
+	FindArchetypeValue(ArchetypeProperty, ArchetypeValuePtr);
+	if (!ArchetypeValuePtr)
+	{
+		FFrame::KismetExecutionMessage(*FString::Printf(TEXT("Property '%s' on '%s' (%s) had no archetype value to reset to"), *PropertyName.ToString(), *Object->GetPathName(), *Object->GetClass()->GetName()), ELogVerbosity::Warning, UE::Blueprint::Private::PropertySetFailedWarning);
+		return false;
+	}
+
+	return Generic_SetEditorProperty(Object, PropertyName, ArchetypeValuePtr, ArchetypeProperty, ChangeNotifyMode);
+}
+
 #endif	// WITH_EDITOR
 
 int32 UKismetSystemLibrary::BeginTransaction(const FString& Context, FText Description, UObject* PrimaryObject)
