-Fix for crash during UHLODSubsystem::OnExternalStreamingObjectRemoved
-This is called during CleanupWorld in during play session teardown in PIE, and the PIE World is not resolvable from a SoftObjectPtr at this point.

diff --git a/Engine/Source/Runtime/Engine/Private/WorldPartition/HLOD/HLODSubsystem.cpp b/Engine/Source/Runtime/Engine/Private/WorldPartition/HLOD/HLODSubsystem.cpp
--- a/Engine/Source/Runtime/Engine/Private/WorldPartition/HLOD/HLODSubsystem.cpp
+++ b/Engine/Source/Runtime/Engine/Private/WorldPartition/HLOD/HLODSubsystem.cpp
@@ -378,7 +378,7 @@ void UHLODSubsystem::OnExternalStreamingObjectInjected(URuntimeHashExternalStrea
 
 void UHLODSubsystem::OnExternalStreamingObjectRemoved(URuntimeHashExternalStreamingObjectBase* ExternalStreamingObject)
 {
-	UWorldPartition* OwnerPartition = ExternalStreamingObject->GetOuterWorld()->GetWorldPartition();
+	UWorldPartition* OwnerPartition = ExternalStreamingObject->GetOuterWorld() ? ExternalStreamingObject->GetOuterWorld()->GetWorldPartition() : ExternalStreamingObject->GetTypedOuter<UWorldPartition>();
 	FWorldPartitionHLODRuntimeData* WorldPartitionHLODRuntimeData = WorldPartitionsHLODRuntimeData.Find(OwnerPartition);
 	if (WorldPartitionHLODRuntimeData)
 	{
