Validate Audio Device Manager in Modulation Destination Update (spec fix for crash on shutdown)

diff --git a/Engine/Source/Runtime/Engine/Private/SoundModulationDestination.cpp b/Engine/Source/Runtime/Engine/Private/SoundModulationDestination.cpp
--- a/Engine/Source/Runtime/Engine/Private/SoundModulationDestination.cpp
+++ b/Engine/Source/Runtime/Engine/Private/SoundModulationDestination.cpp
@@ -288,23 +288,26 @@ namespace Audio
 	{
 		auto UpdateHandleLambda = [this, ModSettings = MoveTemp(ProxySettings)]() mutable
 		{
-			if (FAudioDevice* AudioDevice = FAudioDeviceManager::Get()->GetAudioDeviceRaw(DeviceId))
+			if (FAudioDeviceManager* AudioDeviceManager = FAudioDeviceManager::Get())
 			{
-				if (AudioDevice->IsModulationPluginEnabled() && AudioDevice->ModulationInterface.IsValid())
+				if (FAudioDevice* AudioDevice = AudioDeviceManager->GetAudioDeviceRaw(DeviceId))
 				{
-					if (IAudioModulationManager* Modulation = AudioDevice->ModulationInterface.Get())
+					if (AudioDevice->IsModulationPluginEnabled() && AudioDevice->ModulationInterface.IsValid())
 					{
-						TSet<FModulatorHandle> NewHandles;
-						for (TUniquePtr<Audio::IModulatorSettings>& ModSetting : ModSettings)
+						if (IAudioModulationManager* Modulation = AudioDevice->ModulationInterface.Get())
 						{
-							Audio::FModulationParameter HandleParam = Parameter;
-							NewHandles.Add(FModulatorHandle{ *Modulation, *ModSetting.Get(), MoveTemp(HandleParam) });
+							TSet<FModulatorHandle> NewHandles;
+							for (TUniquePtr<Audio::IModulatorSettings>& ModSetting : ModSettings)
+							{
+								Audio::FModulationParameter HandleParam = Parameter;
+								NewHandles.Add(FModulatorHandle{ *Modulation, *ModSetting.Get(), MoveTemp(HandleParam) });
+							}
+
+							FScopeLock Lock(&HandleCritSection);
+							Handles = MoveTemp(NewHandles);
 						}
-
-						FScopeLock Lock(&HandleCritSection);
-						Handles = MoveTemp(NewHandles);
+						return;
 					}
-					return;
 				}
 			}
 
