Fixed the function DownsizeImageUsingTextureSettings comment as it didn't reflect the  implementation details of the function.

diff --git a/Engine/Source/Runtime/Engine/Classes/Engine/Texture.h b/Engine/Source/Runtime/Engine/Classes/Engine/Texture.h
--- a/Engine/Source/Runtime/Engine/Classes/Engine/Texture.h
+++ b/Engine/Source/Runtime/Engine/Classes/Engine/Texture.h
@@ -1916,7 +1916,7 @@ public:
 
 	/*
 	 * Downsize the 2D Image with the build setting of the texture.
-	 * Try to get as close as it can to the target resolution but it will stay above it if it can reach it
+	 * Try to get as close as it can to the target resolution but it will stay above it if it can't reach it
 	 */
 	ENGINE_API bool DownsizeImageUsingTextureSettings(const ITargetPlatform* TargetPlatform, FImage& InOutImage, int32 TargetSize, int32 LayerIndex);
 #endif
