[Backout] - CL27440582
Original CL Desc
-----------------------------------------------------------------
Optimize ULevel::AddLoadedActors / RemoveLoadedActors.
Saves over 5 seconds of loading/unloading time in a 32k actors map.

diff --git a/Engine/Source/Runtime/Engine/Private/Level.cpp b/Engine/Source/Runtime/Engine/Private/Level.cpp
--- a/Engine/Source/Runtime/Engine/Private/Level.cpp
+++ b/Engine/Source/Runtime/Engine/Private/Level.cpp
@@ -7,9 +7,6 @@ Level.cpp: Level-related functions
 #include "Engine/Level.h"
 
 #include "EngineLogs.h"
-#include "Algo/Copy.h"
-#include "Algo/ForEach.h"
-#include "Algo/Transform.h"
 #include "Misc/ScopedSlowTask.h"
 #include "HAL/PlatformFile.h"
 #include "UObject/UE5MainStreamObjectVersion.h"
@@ -807,21 +804,17 @@ void ULevel::AddLoadedActor(AActor* Actor, const FTransform* TransformToApply)
 
 void ULevel::AddLoadedActors(const TArray<AActor*>& ActorList, const FTransform* TransformToApply)
 {
-	// Actors set used to accelerate lookups for actors already in the level's actor list. This can happen for newly added
-	// actors, as spawning them will add them to the actors array, and grabbing an actor descriptor reference on them will trigger
-	// loading, even if the actor already exists.
-	TSet<AActor*> ActorsSet(ObjectPtrDecay(Actors));
-
-	// Actors queue that will be added to the level, filtered for actors already in the level's actor list.
 	TArray<AActor*> ActorsQueue;
 	ActorsQueue.Reserve(ActorList.Num());
 
-	TFunction<void(AActor* Actor)> QueueActor = [this, &ActorsSet, &ActorsQueue, &QueueActor](AActor* Actor)
+	TFunction<void(AActor* Actor)> QueueActor = [this, &ActorsQueue, &QueueActor](AActor* Actor)
 	{
+		check(Actor);
+		check(Actor->GetLevel() == this);
 		check(IsValidChecked(Actor));
-		check(Actor->GetLevel() == this);		
 
-		if (!ActorsSet.Contains(Actor))
+		int32 ActorIndex;
+		if (!Actors.Find(Actor, ActorIndex))
 		{
 			Actors.Add(Actor);
 			ActorsForGC.Add(Actor);
@@ -852,20 +845,22 @@ void ULevel::AddLoadedActors(const TArray<AActor*>& ActorList, const FTransform*
 	for (AActor* Actor : ActorsQueue)
 	{
 		// RegisterAllComponents can destroy child actors
-		if (IsValid(Actor))
+		if (!IsValid(Actor))
 		{
-			if (TransformToApply)
-			{
-				FLevelUtils::FApplyLevelTransformParams TransformParams(this, *TransformToApply);
-				TransformParams.Actor = Actor;
-				TransformParams.bDoPostEditMove = true;
-				FLevelUtils::ApplyLevelTransform(TransformParams);
-			}
+			continue;
+		}
 
-			if (bAreComponentsCurrentlyRegistered)
-			{
-				Actor->RegisterAllComponents();
-			}
+		if (TransformToApply)
+		{
+			FLevelUtils::FApplyLevelTransformParams TransformParams(this, *TransformToApply);
+			TransformParams.Actor = Actor;
+			TransformParams.bDoPostEditMove = true;
+			FLevelUtils::ApplyLevelTransform(TransformParams);
+		}
+
+		if (bAreComponentsCurrentlyRegistered)
+		{
+			Actor->RegisterAllComponents();
 		}
 
 		SlowTask.EnterProgressFrame(1);
@@ -875,12 +870,14 @@ void ULevel::AddLoadedActors(const TArray<AActor*>& ActorList, const FTransform*
 	for (AActor* Actor : ActorsQueue)
 	{
 		// RegisterAllComponents/RerunConstructionScripts can destroy child actors
-		if (IsValid(Actor))
+		if (!IsValid(Actor))
 		{
-			if (bAreComponentsCurrentlyRegistered)
-			{
-				Actor->RerunConstructionScripts();
-			}
+			continue;
+		}
+
+		if (bAreComponentsCurrentlyRegistered)
+		{
+			Actor->RerunConstructionScripts();
 		}
 
 		SlowTask.EnterProgressFrame(1);
@@ -890,22 +887,24 @@ void ULevel::AddLoadedActors(const TArray<AActor*>& ActorList, const FTransform*
 	for (AActor* Actor : ActorsQueue)
 	{
 		// RegisterAllComponents/RerunConstructionScripts can destroy child actors
-		if (IsValid(Actor))
+		if (!IsValid(Actor))
 		{
-			if (bAreComponentsCurrentlyRegistered)
-			{
-				GetWorld()->UpdateCullDistanceVolumes(Actor);
-				Actor->MarkComponentsRenderStateDirty();
-			}
+			continue;
+		}
 
-			if (IsUsingActorFolders() && bFixupActorFoldersAtLoad)
-			{
-				Actor->FixupActorFolder();
-			}
+		if (bAreComponentsCurrentlyRegistered)
+		{
+			GetWorld()->UpdateCullDistanceVolumes(Actor);
+			Actor->MarkComponentsRenderStateDirty();
+		}
 
-			OnLoadedActorAddedToLevelEvent.Broadcast(*Actor);
+		if (IsUsingActorFolders() && bFixupActorFoldersAtLoad)
+		{
+			Actor->FixupActorFolder();
 		}
 
+		OnLoadedActorAddedToLevelEvent.Broadcast(*Actor);
+
 		SlowTask.EnterProgressFrame(1);
 	}
 
@@ -923,41 +922,14 @@ void ULevel::RemoveLoadedActor(AActor* Actor, const FTransform* TransformToRemov
 
 void ULevel::RemoveLoadedActors(const TArray<AActor*>& ActorList, const FTransform* TransformToRemove)
 {
-	// Build an actor pointer to pair of indices into Actors and ActorsForGC to accelerate removal lookups. This doesn't make much sense
-	// when removing a single actor, but most of the time we are removing actors in large batches and it makes a huge difference in removal time.
-	TMap<AActor*, TPair<int32, int32>> ActorsIndices;
-	{
-		ActorsIndices.Reserve(Actors.Num());
-
-		int32 ActorIndex = 0;
-		Algo::ForEach(Actors, [&ActorsIndices, &ActorIndex](AActor* Actor)
-		{
-			if (Actor)
-			{
-				ActorsIndices.Add(Actor, TPair<int32, int32>(ActorIndex, INDEX_NONE));
-			}
-			ActorIndex++;
-		});
-	
-		int32 ActorForGCIndex = 0;
-		Algo::ForEach(ActorsForGC, [&ActorsIndices, &ActorForGCIndex](AActor* Actor)
-		{
-			if (Actor)
-			{
-				ActorsIndices[Actor].Value = ActorForGCIndex;
-			}
-			ActorForGCIndex++;
-		});
-	}
-
-	// Build the actual actor list that needs removal.
 	TArray<AActor*> ActorsQueue;
 	ActorsQueue.Reserve(ActorList.Num());
 
-	TFunction<void(AActor* Actor)> QueueActor = [this, &ActorsIndices, &ActorsQueue, &QueueActor](AActor* Actor)
+	TFunction<void(AActor* Actor)> QueueActor = [this, &ActorsQueue, &QueueActor](AActor* Actor)
 	{
+		check(Actor);
+		check(Actor->GetLevel() == this);
 		check(IsValidChecked(Actor));
-		check(Actor->GetLevel() == this);		
 
 		// Handle child actors
 		Actor->ForEachComponent<UChildActorComponent>(false, [this, &QueueActor](UChildActorComponent* ChildActorComponent)
@@ -968,16 +940,12 @@ void ULevel::RemoveLoadedActors(const TArray<AActor*>& ActorList, const FTransfo
 			}
 		});
 
-		if (const TPair<int32, int32>* ActorIndices = ActorsIndices.Find(Actor))
+		int32 ActorIndex;
+		// temporarily downgraded to ensure while an issue is fixed
+		if (ensure(Actors.Find(Actor, ActorIndex))) 
 		{
-			check(Actors[ActorIndices->Key] == Actor);
-			Actors[ActorIndices->Key] = nullptr;
-
-			if (ActorIndices->Value != INDEX_NONE)
-			{
-				check(ActorsForGC[ActorIndices->Value] == Actor);
-				ActorsForGC[ActorIndices->Value] = nullptr;
-			}
+			Actors[ActorIndex] = nullptr;
+			ActorsForGC.Remove(Actor);
 
 			ActorsQueue.Add(Actor);
 		}
@@ -996,28 +964,27 @@ void ULevel::RemoveLoadedActors(const TArray<AActor*>& ActorList, const FTransfo
 	for (AActor* Actor : ActorsQueue)
 	{
 		// UnregisterAllComponents can destroy child actors
-		if (IsValid(Actor))
+		if (!IsValid(Actor))
 		{
-			Actor->UnregisterAllComponents();
-			Actor->RegisterAllActorTickFunctions(false, true);
+			continue;
+		}
 
-			OnLoadedActorRemovedFromLevelEvent.Broadcast(*Actor);
+		Actor->UnregisterAllComponents();
+		Actor->RegisterAllActorTickFunctions(false, true);
 
-			if (TransformToRemove)
-			{
-				FLevelUtils::FApplyLevelTransformParams TransformParams(this, TransformToRemove->Inverse());
-				TransformParams.Actor = Actor;
-				TransformParams.bDoPostEditMove = true;
-				FLevelUtils::ApplyLevelTransform(TransformParams);
-			}
+		OnLoadedActorRemovedFromLevelEvent.Broadcast(*Actor);
+
+		if (TransformToRemove)
+		{
+			FLevelUtils::FApplyLevelTransformParams TransformParams(this, TransformToRemove->Inverse());
+			TransformParams.Actor = Actor;
+			TransformParams.bDoPostEditMove = true;
+			FLevelUtils::ApplyLevelTransform(TransformParams);
 		}
 
 		SlowTask.EnterProgressFrame(1);
 	}
 
-	Actors.Remove(nullptr);
-	ActorsForGC.Remove(nullptr);
-
 	TArray<AActor*> ValidActorsQueue;
 	ValidActorsQueue.Reserve(ActorsQueue.Num());
 	Algo::CopyIf(ActorsQueue, ValidActorsQueue, [&](AActor* Actor) { return IsValid(Actor); });
@@ -4585,4 +4552,4 @@ void ULevel::OnMultipleLSAsPopupDismissed()
 
 #endif // WITH_EDITOR
 
-#undef LOCTEXT_NAMESPACE
\ No newline at end of file
+#undef LOCTEXT_NAMESPACE
