Substrate
- Legacy shading model node can now be used as a regular node.
- We can have multiple legacy shading model node in a isngle material now.
- There is now a ShadinModel drop down on the node that is used when the ShadingModel input is not plugged in
- Material instance shading model override is now accounted for by the legacy shading model node. Overriden from the HLSLTranslator when detected.

The data flow now for legacy mateiral conversion is as such:

--- Conversion time - UMaterial::ConvertMaterialToStrataMaterial()
Shading model => legacy node setup (shading model from expression: relink input, otherwise: set shading model on the conversion node)

--- GatherStrataMaterialInfo
Send back what is the shadingmodel used (or if it is "from expression")

--- UMaterial::RebuildShadingModelField()
Set material Domain, ShadingModel and ShadingModels according ot the returned value (FromExpression, etc.) from GatherStrataMaterialInfo.
Those values are not authorable but do enable shading models via defines from the HLSLTranslator. Seeing them being correct also helps to understand and debug.

--- Material instance shading model override
Overriden from the HLSLTranslator when detected.

diff --git a/Engine/Shaders/Shared/StrataDefinitions.h b/Engine/Shaders/Shared/StrataDefinitions.h
--- a/Engine/Shaders/Shared/StrataDefinitions.h
+++ b/Engine/Shaders/Shared/StrataDefinitions.h
@@ -8,7 +8,7 @@
 #pragma once
 
 // Change this to force recompilation of all strata dependent shaders (use https://www.random.org/cgi-bin/randbyte?nbytes=4&format=h)
-#define STRATA_SHADER_VERSION 0x56F06029 
+#define STRATA_SHADER_VERSION 0x72C8A053
 
 // BSDF offsets are packed into 32bits, each entry using STRATA_BSDF_OFFSET_BIT_COUNT bits
 #define STRATA_MAX_BSDF_COUNT_FOR_BDSFOFFSET	8u
diff --git a/Engine/Source/Editor/UnrealEd/Private/MaterialGraphNode.cpp b/Engine/Source/Editor/UnrealEd/Private/MaterialGraphNode.cpp
--- a/Engine/Source/Editor/UnrealEd/Private/MaterialGraphNode.cpp
+++ b/Engine/Source/Editor/UnrealEd/Private/MaterialGraphNode.cpp
@@ -142,22 +142,6 @@ bool UMaterialGraphNode::CanPasteHere(const UEdGraph* TargetGraph) const
 {
 	if (Super::CanPasteHere(TargetGraph) && MaterialExpression)
 	{
-		if (MaterialExpression && MaterialExpression->IsA(UMaterialExpressionStrataLegacyConversion::StaticClass()))
-		{
-			// We could have used CanDuplicateNode() returning false to prevent the copy but it is nicer to have a notification about why the copy is not happening.
-			FText NotificationText = NSLOCTEXT("UMaterialGraphNode", "SkippedSubstrateLegacyConversion", "SubstrateLegacyConversion node cannot be copied! It is only used to convert legacy material to Substrate.");
-			FNotificationInfo Info(NotificationText);
-			Info.ExpireDuration = 5.0f;
-			Info.bUseLargeFont = false;
-			Info.Image = FCoreStyle::Get().GetBrush(TEXT("MessageLog.Warning"));
-			TSharedPtr<SNotificationItem> Notification = FSlateNotificationManager::Get().AddNotification(Info);
-			if (Notification.IsValid())
-			{
-				Notification->SetCompletionState(SNotificationItem::CS_None);
-			}
-			return false; // We do not allow the copy of the StrataLegacyConversion which should only be used to convert pre-strata materials.
-		}
-
 		const UMaterialGraph* MaterialGraph = Cast<const UMaterialGraph>(TargetGraph);
 		if (MaterialGraph)
 		{
diff --git a/Engine/Source/Runtime/Engine/Classes/Engine/EngineTypes.h b/Engine/Source/Runtime/Engine/Classes/Engine/EngineTypes.h
--- a/Engine/Source/Runtime/Engine/Classes/Engine/EngineTypes.h
+++ b/Engine/Source/Runtime/Engine/Classes/Engine/EngineTypes.h
@@ -683,44 +683,43 @@ private:
 
 /**
  * Specifies the Substrate runtime shading model summarized from the material graph
+ * Not exposed in UI, only used in code.
+ * Those states are deducted from the material graph and map to a specific domain/shading state.
  */
-UENUM()
 enum EStrataShadingModel : int
 {
-	SSM_Unlit					UMETA(DisplayName = "Unlit"),
-	SSM_DefaultLit				UMETA(DisplayName = "DefaultLit"),
-	SSM_SubsurfaceLit			UMETA(DisplayName = "SubsurfaceLit"),
-	SSM_VolumetricFogCloud		UMETA(DisplayName = "VolumetricFogCloud"),
-	SSM_Hair					UMETA(DisplayName = "Hair"),
-	SSM_Eye						UMETA(DisplayName = "Eye"),
-	SSM_Cloth					UMETA(DisplayName = "Cloth"),
-	SSM_ClearCoat				UMETA(DisplayName = "ClearCoat"),
-	SSM_SingleLayerWater		UMETA(DisplayName = "SingleLayerWater"),
-	SSM_LightFunction			UMETA(DisplayName = "LightFunction"),
-	SSM_PostProcess				UMETA(DisplayName = "PostProcess"),
-	SSM_Decal					UMETA(DisplayName = "Decal"),
-	SSM_UI						UMETA(DisplayName = "UI"),
+	SSM_Unlit,
+	SSM_DefaultLit,
+	SSM_ThinTranslucent,
+	SSM_SubsurfaceMFP,
+	SSM_SubsurfaceProfile,
+	SSM_SubsurfaceWrap,
+	SSM_SubsurfaceThinTwoSided,
+	SSM_VolumetricFogCloud,
+	SSM_Hair,
+	SSM_Eye,
+	SSM_Cloth,
+	SSM_ClearCoat,
+	SSM_SingleLayerWater,
+	SSM_LightFunction,
+	SSM_PostProcess,
+	SSM_Decal,
+	SSM_UI,
 	/** Number of unique shading models. */
-	SSM_NUM						UMETA(Hidden),
+	SSM_NUM,
 };
-static_assert(SSM_NUM <= 16, "Do not exceed 16 shading models without expanding FStrataMaterialShadingModelField to support uint32 instead of uint16!");
+static_assert(SSM_NUM <= 32, "Do not exceed 32 shading models without expanding FStrataMaterialInfo::ShadingModelField to support more bits!");
 
 // This used to track cyclic graph which we do not support. We only support acyclic graph and a depth of 128 is already too high for a realistic use case.
 #define STRATA_TREE_MAX_DEPTH 48
 
 /** Gather information from the Substrate material graph to setup material for runtime. */
-USTRUCT()
 struct FStrataMaterialInfo
 {
-	GENERATED_USTRUCT_BODY()
-
 public:
 	FStrataMaterialInfo() {}
 	FStrataMaterialInfo(EStrataShadingModel InShadingModel) { AddShadingModel(InShadingModel); }
 
-	ENGINE_API bool Serialize(FArchive& Ar);
-	ENGINE_API void PostSerialize(const FArchive& Ar);
-
 	// Shading model
 	void AddShadingModel(EStrataShadingModel InShadingModel) { check(InShadingModel < SSM_NUM); ShadingModelField |= (uint16)(1 << (uint16)InShadingModel); }
 	void SetSingleShadingModel(EStrataShadingModel InShadingModel) { check(InShadingModel < SSM_NUM); ShadingModelField = (uint16)(1 << (uint16)InShadingModel); }
@@ -772,24 +771,16 @@ public:
 #endif
 
 private:
-	UPROPERTY()
-	uint16 ShadingModelField = 0;
+	uint32 ShadingModelField = 0;
 
 	/* Indicates if the shading model is constant or data-driven from the shader graph */
-	UPROPERTY()
 	uint8 bHasShadingModelFromExpression = 0;
 
-	UPROPERTY()
-	uint32 ConnectedProperties_DEPRECATED = 0;
-
 	/* Indicates which (legacy) inputs are connected */
-	UPROPERTY()
 	uint64 ConnectedPropertyMask = 0;
 
-	UPROPERTY()
 	TArray<TObjectPtr<USubsurfaceProfile>> SubsurfaceProfiles;
 
-	UPROPERTY()
 	TArray<TObjectPtr<USpecularProfile>> SpecularProfiles;
 
 #if WITH_EDITOR
diff --git a/Engine/Source/Runtime/Engine/Classes/Materials/MaterialExpressionStrata.h b/Engine/Source/Runtime/Engine/Classes/Materials/MaterialExpressionStrata.h
--- a/Engine/Source/Runtime/Engine/Classes/Materials/MaterialExpressionStrata.h
+++ b/Engine/Source/Runtime/Engine/Classes/Materials/MaterialExpressionStrata.h
@@ -170,9 +170,8 @@ class UMaterialExpressionStrataLegacyConversion : public UMaterialExpressionStra
 	UPROPERTY(EditAnywhere, BlueprintReadOnly, Category = Material, meta = (DisplayName = "Subsurface Profile"))
 	TObjectPtr<class USubsurfaceProfile> SubsurfaceProfile;
 
-	/** Store converted material models. */
-	UPROPERTY()
-	FStrataMaterialInfo ConvertedStrataMaterialInfo;
+	UPROPERTY(EditAnywhere, Category = ShadingModel, meta = (ShowAsInputPin = "Primary", DisplayName = "Single Shading Model"))
+	TEnumAsByte<enum EMaterialShadingModel> ShadingModelOverride = MSM_DefaultLit;
 
 	//~ Begin UMaterialExpression Interface
 #if WITH_EDITOR
diff --git a/Engine/Source/Runtime/Engine/Private/Materials/HLSLMaterialTranslator.cpp b/Engine/Source/Runtime/Engine/Private/Materials/HLSLMaterialTranslator.cpp
--- a/Engine/Source/Runtime/Engine/Private/Materials/HLSLMaterialTranslator.cpp
+++ b/Engine/Source/Runtime/Engine/Private/Materials/HLSLMaterialTranslator.cpp
@@ -12534,7 +12534,7 @@ int32 FHLSLMaterialTranslator::StrataConversionFromLegacy(
 	int32 Opacity,
 	int32 TransmittanceColor,
 	int32 WaterScatteringCoefficients, int32 WaterAbsorptionCoefficients, int32 WaterPhaseG, int32 ColorScaleBehindWater,
-	int32 ShadingModel,
+	int32 InShadingModel,
 	int32 Normal, int32 Tangent, const FString& SharedLocalBasisIndexMacro,
 	int32 ClearCoat_Normal, int32 ClearCoat_Tangent, const FString& ClearCoat_SharedLocalBasisIndexMacro,
 	int32 CustomTangent_Tangent,
@@ -12547,6 +12547,24 @@ int32 FHLSLMaterialTranslator::StrataConversionFromLegacy(
 	const FString ClearCoat_NormalCode = GetParameterCode(ClearCoat_Normal);
 	const FString ClearCoat_TangentCode = Tangent != INDEX_NONE ? *GetParameterCode(ClearCoat_Tangent) : TEXT("NONE");
 
+	UMaterial* BaseMaterial = Material->GetMaterialInterface()->GetBaseMaterial();
+	// Material is probably an instanced material, so we check that this is the case before potentially overriding the ShadingModel from the material instance.
+	UMaterialInterface* BaseMaterialInterface = static_cast<UMaterialInterface*>(BaseMaterial);
+	UMaterialInterface*     MaterialInterface = Material->GetMaterialInterface();
+	if (BaseMaterial && (BaseMaterialInterface != MaterialInterface))
+	{
+		FMaterialShadingModelField BaseMaterialShadingModels	= BaseMaterial->GetShadingModels();
+		FMaterialShadingModelField MaterialShadingModels		= Material->GetShadingModels();
+
+		// If the potentially instanced material does not have the same shading model as the instance material, this means it would have overridden the shading model.
+		// From the UI, only a single shading model is selectable, so we simply apply the one coming form the material instance.
+		if (MaterialShadingModels.IsValid() && MaterialShadingModels.CountShadingModels()==1 && MaterialShadingModels != BaseMaterialShadingModels)
+		{
+			bHasDynamicShadingModels = false;	// No need to go dynamic when there is only a single shading model selected.
+			InShadingModel = ShadingModel(MaterialShadingModels.GetFirstShadingModel());
+		}
+	}
+
 	if (PromoteToOperator)
 	{
 		if (PromoteToOperator->Index == INDEX_NONE || PromoteToOperator->BSDFIndex == INDEX_NONE)
@@ -12575,7 +12593,7 @@ int32 FHLSLMaterialTranslator::StrataConversionFromLegacy(
 			*StrataGetCastParameterCode(WaterAbsorptionCoefficients,	MCT_Float3),
 			*StrataGetCastParameterCode(WaterPhaseG,					MCT_Float),
 			*StrataGetCastParameterCode(ColorScaleBehindWater,			MCT_Float3),
-			*GetParameterCode(ShadingModel),
+			*GetParameterCode(InShadingModel),
 			// Raw access to Normal/Tangent/ClearCoatNormal/CustomTangent for conversion purpose
 			*StrataGetCastParameterCode(Normal,							MCT_Float3),
 			*StrataGetCastParameterCode(RawTangent,						MCT_Float3),
@@ -12616,7 +12634,7 @@ int32 FHLSLMaterialTranslator::StrataConversionFromLegacy(
 		*StrataGetCastParameterCode(WaterAbsorptionCoefficients,	MCT_Float3),
 		*StrataGetCastParameterCode(WaterPhaseG,					MCT_Float),
 		*StrataGetCastParameterCode(ColorScaleBehindWater,			MCT_Float3),
-		*GetParameterCode(ShadingModel),
+		*GetParameterCode(InShadingModel),
 		// Raw access to Normal/Tangent/ClearCoatNormal/CustomTangent for conversion purpose
 		*StrataGetCastParameterCode(Normal,							MCT_Float3),
 		*StrataGetCastParameterCode(RawTangent,						MCT_Float3),
diff --git a/Engine/Source/Runtime/Engine/Private/Materials/Material.cpp b/Engine/Source/Runtime/Engine/Private/Materials/Material.cpp
--- a/Engine/Source/Runtime/Engine/Private/Materials/Material.cpp
+++ b/Engine/Source/Runtime/Engine/Private/Materials/Material.cpp
@@ -3110,20 +3110,30 @@ void UMaterial::BackwardsCompatibilityDecalConversion()
 #endif // WITH_EDITOR
 }
 
-static void AddStrataShadingModelFromMaterialShadingModel(FStrataMaterialInfo& OutInfo, const FMaterialShadingModelField& InShadingModels)
+static void AddSurfaceStrataShadingModelFromMaterialShadingModels(FStrataMaterialInfo& OutInfo, const FMaterialShadingModelField& InShadingModels)
 {
 	if (InShadingModels.HasShadingModel(MSM_Unlit))				{ OutInfo.AddShadingModel(EStrataShadingModel::SSM_Unlit); }
 	if (InShadingModels.HasShadingModel(MSM_DefaultLit))		{ OutInfo.AddShadingModel(EStrataShadingModel::SSM_DefaultLit); }
-	if (InShadingModels.HasShadingModel(MSM_Subsurface))		{ OutInfo.AddShadingModel(EStrataShadingModel::SSM_SubsurfaceLit); }
-	if (InShadingModels.HasShadingModel(MSM_PreintegratedSkin))	{ OutInfo.AddShadingModel(EStrataShadingModel::SSM_SubsurfaceLit); }
+	if (InShadingModels.HasShadingModel(MSM_Subsurface))		{ OutInfo.AddShadingModel(EStrataShadingModel::SSM_SubsurfaceWrap); }
+	if (InShadingModels.HasShadingModel(MSM_PreintegratedSkin))	{ OutInfo.AddShadingModel(EStrataShadingModel::SSM_SubsurfaceMFP); }
 	if (InShadingModels.HasShadingModel(MSM_ClearCoat))			{ OutInfo.AddShadingModel(EStrataShadingModel::SSM_ClearCoat); }
-	if (InShadingModels.HasShadingModel(MSM_SubsurfaceProfile))	{ OutInfo.AddShadingModel(EStrataShadingModel::SSM_SubsurfaceLit); }
-	if (InShadingModels.HasShadingModel(MSM_TwoSidedFoliage))	{ OutInfo.AddShadingModel(EStrataShadingModel::SSM_SubsurfaceLit); }
+	if (InShadingModels.HasShadingModel(MSM_SubsurfaceProfile))	{ OutInfo.AddShadingModel(EStrataShadingModel::SSM_SubsurfaceProfile); }
+	if (InShadingModels.HasShadingModel(MSM_TwoSidedFoliage))	{ OutInfo.AddShadingModel(EStrataShadingModel::SSM_SubsurfaceThinTwoSided); }
 	if (InShadingModels.HasShadingModel(MSM_Hair))				{ OutInfo.AddShadingModel(EStrataShadingModel::SSM_Hair); }
 	if (InShadingModels.HasShadingModel(MSM_Cloth))				{ OutInfo.AddShadingModel(EStrataShadingModel::SSM_Cloth); }
 	if (InShadingModels.HasShadingModel(MSM_Eye))				{ OutInfo.AddShadingModel(EStrataShadingModel::SSM_Eye); }
 	if (InShadingModels.HasShadingModel(MSM_SingleLayerWater))	{ OutInfo.AddShadingModel(EStrataShadingModel::SSM_SingleLayerWater); }
-	if (InShadingModels.HasShadingModel(MSM_ThinTranslucent))	{ OutInfo.AddShadingModel(EStrataShadingModel::SSM_DefaultLit); }
+	if (InShadingModels.HasShadingModel(MSM_ThinTranslucent))	{ OutInfo.AddShadingModel(EStrataShadingModel::SSM_ThinTranslucent); }
+}
+
+static void AddSurfaceStrataShadingModelFromMaterialShadingModel(FStrataMaterialInfo& OutInfo, const EMaterialShadingModel& InShadingModel)
+{
+	if (InShadingModel < MSM_NUM)
+	{
+		FMaterialShadingModelField ShadingModel;
+		ShadingModel.AddShadingModel(InShadingModel);
+		AddSurfaceStrataShadingModelFromMaterialShadingModels(OutInfo, ShadingModel);
+	}
 }
 
 EBlendMode ConvertLegacyBlendMode(EBlendMode InBlendMode, FMaterialShadingModelField InShadingModels)
@@ -3141,6 +3151,22 @@ EBlendMode ConvertLegacyBlendMode(EBlendMode InBlendMode, FMaterialShadingModelF
 
 void UMaterial::ConvertMaterialToStrataMaterial()
 {
+	/*
+	* The data flow for legacy material conversion node that can be used in isolation is as such:
+	*
+	* --- Conversion time - UMaterial::ConvertMaterialToStrataMaterial()
+	*     Shading model => legacy node setup (shading model from expression? relink input : otherwise set the shading model on the conversion node itself)
+	*
+	* --- GatherStrataMaterialInfo
+	*     Send back what is the shading model used, or if it is "from expression"
+	*
+	* --- UMaterial::RebuildShadingModelField()
+	*     Set material Domain, ShadingModel and ShadingModels according to the returned value from GatherStrataMaterialInfo gathered from the graph.
+	*     Those values are not authorable but do enable shading models via defines from the HLSLTranslator. Seeing them being correct also helps to understand and debug.
+	*
+	* --- Material instance shading model override
+	*     Overridden from the HLSLTranslator when detected by comparing base and instanced materials.
+	*/
 #if WITH_EDITOR
 	if (!Strata::IsStrataEnabled())
 	{
@@ -3303,7 +3329,6 @@ void UMaterial::ConvertMaterialToStrataMaterial()
 		ConvertNode->ClearCoat.Connect(12, BreakMatAtt);
 		ConvertNode->ClearCoatRoughness.Connect(13, BreakMatAtt);
 		ConvertNode->Opacity.Connect(6, BreakMatAtt);
-		ConvertNode->ShadingModel.Connect(25, BreakMatAtt);
 		ConvertNode->SubsurfaceProfile = bRequireNoSubsurfaceProfile ? nullptr : SubsurfaceProfile;
 		bRelinkCustomOutputNodes = true;
 
@@ -3322,23 +3347,14 @@ void UMaterial::ConvertMaterialToStrataMaterial()
 
 		if (ShadingModel == MSM_FromMaterialExpression)
 		{
-			ConvertNode->ConvertedStrataMaterialInfo.SetShadingModelFromExpression(true);
-			AddStrataShadingModelFromMaterialShadingModel(ConvertNode->ConvertedStrataMaterialInfo, ShadingModels);
-			check(ConvertNode->ConvertedStrataMaterialInfo.CountShadingModels() >= 1);
+			check(ShadingModels.CountShadingModels() >= 1);
+			ConvertNode->ShadingModel.Connect(25, BreakMatAtt);
 		}
 		else
 		{
-			check(!bHasShadingModelMixture);
-
-			// Add shading model node
-			UMaterialExpressionShadingModel* ShadingModelNode = NewObject<UMaterialExpressionShadingModel>(this);
-			ReplaceNodeAndMoveToTheRight(ConvertNode, ShadingModelNode);
-			ShadingModelNode->ShadingModel = ShadingModel;
-			ConvertNode->ShadingModel.Connect(0, ShadingModelNode);
-
 			// Store strata shading model of the converted material. 
-			AddStrataShadingModelFromMaterialShadingModel(ConvertNode->ConvertedStrataMaterialInfo, ShadingModels);
-			check(ConvertNode->ConvertedStrataMaterialInfo.CountShadingModels() == 1);
+			check(ShadingModels.CountShadingModels() == 1);
+			ConvertNode->ShadingModelOverride = ShadingModel;
 		}
 
 		if (MaterialDomain == MD_DeferredDecal)
@@ -3495,17 +3511,11 @@ void UMaterial::ConvertMaterialToStrataMaterial()
 				// Shading Model
 				// * either use the shader graph expression 
 				// * or add a constant shading model
-				// 
-				// Note: store this conversion type(s) into ConvertedStrataMaterialInfo for having more context when 
-				// rebuilding the final Shading model (see RebuildShadingModelField())
 				if (ShadingModel == MSM_FromMaterialExpression)
 				{
 					if (!EditorOnly->ShadingModelFromMaterialExpression.IsConnected())
 					{
-						UMaterialExpressionShadingModel* ShadingModelNode = NewObject<UMaterialExpressionShadingModel>(this);
-						ReplaceNodeAndMoveToTheRight(ConvertNode, ShadingModelNode);
-						ShadingModelNode->ShadingModel = MSM_DefaultLit;
-						ConvertNode->ShadingModel.Connect(0, ShadingModelNode);
+						ConvertNode->ShadingModelOverride = MSM_DefaultLit;
 					}
 					else
 					{
@@ -3518,24 +3528,15 @@ void UMaterial::ConvertMaterialToStrataMaterial()
 					GatherCustomNodes();
 					if (SingleLayerWaterOutput)
 					{
-						ConvertNode->ConvertedStrataMaterialInfo.AddShadingModel(SSM_SingleLayerWater);
+						ShadingModels.AddShadingModel(MSM_SingleLayerWater);
 					}
 
-					ConvertNode->ConvertedStrataMaterialInfo.SetShadingModelFromExpression(true);
-					AddStrataShadingModelFromMaterialShadingModel(ConvertNode->ConvertedStrataMaterialInfo, ShadingModels);
-					check(ConvertNode->ConvertedStrataMaterialInfo.CountShadingModels() >= 1);
+					check(ShadingModels.CountShadingModels() >= 1);
 				}
 				else
 				{
-					check(!bHasShadingModelMixture);
-
-					UMaterialExpressionShadingModel* ShadingModelNode = NewObject<UMaterialExpressionShadingModel>(this);
-					ReplaceNodeAndMoveToTheRight(ConvertNode, ShadingModelNode);
-					ShadingModelNode->ShadingModel = ShadingModel;
-					ConvertNode->ShadingModel.Connect(0, ShadingModelNode);
-
-					AddStrataShadingModelFromMaterialShadingModel(ConvertNode->ConvertedStrataMaterialInfo, ShadingModels);
-					check(ConvertNode->ConvertedStrataMaterialInfo.CountShadingModels() == 1);
+					ConvertNode->ShadingModelOverride = ShadingModel;
+					check(ShadingModels.CountShadingModels() == 1);
 				}
 
 				EditorOnly->FrontMaterial.Connect(0, ConvertNode);
@@ -3615,13 +3616,8 @@ void UMaterial::ConvertMaterialToStrataMaterial()
 			CopyConnectionTo(EditorOnly->Opacity, ConvertNode, 11);	// We only copy, to keep Opacity on the root node in case BLEND_AlphaComposite is selected.
 
 			// Add constant for the Unlit shading model
-			UMaterialExpressionShadingModel* ShadingModelNode = NewObject<UMaterialExpressionShadingModel>(this);
-			ReplaceNodeAndMoveToTheRight(ConvertNode, ShadingModelNode);
-			ShadingModelNode->ShadingModel = ShadingModel;
-			ConvertNode->ShadingModel.Connect(0, ShadingModelNode);
-
-			AddStrataShadingModelFromMaterialShadingModel(ConvertNode->ConvertedStrataMaterialInfo, ShadingModels);
-			check(ConvertNode->ConvertedStrataMaterialInfo.CountShadingModels() == 1);
+			ConvertNode->ShadingModelOverride = ShadingModel;
+			check(ShadingModels.CountShadingModels() == 1);
 
 			// Now pass through the convert to decal node, which flag the material as SSM_Decal, which will set the domain to Decal.
 			UMaterialExpressionStrataConvertToDecal* ConvertToDecalNode= NewObject<UMaterialExpressionStrataConvertToDecal>(this);
@@ -4878,11 +4874,12 @@ void UMaterial::RebuildShadingModelField()
 		}
 
 		bool bSanitizeMaterial = false;
-		if (!StrataMaterialInfo.IsValid())
+		if (!StrataMaterialInfo.IsValid() && !StrataMaterialInfo.HasShadingModelFromExpression() )
 		{
 			bSanitizeMaterial = true;
 			UE_LOG(LogMaterial, Error, TEXT("%s: Material information is invalid."), *GetName());
 		}
+
 		if (StrataMaterialInfo.CountShadingModels() > 1)
 		{
 			if (StrataMaterialInfo.HasShadingModel(SSM_Decal))
@@ -4903,29 +4900,53 @@ void UMaterial::RebuildShadingModelField()
 					bSanitizeMaterial = true;
 				}
 			}
-			else if (StrataMaterialInfo.CountShadingModels() == 2 && StrataMaterialInfo.HasShadingModel(EStrataShadingModel::SSM_DefaultLit) && StrataMaterialInfo.HasShadingModel(EStrataShadingModel::SSM_SubsurfaceLit))
+			else if (StrataMaterialInfo.CountShadingModels() == 2 && StrataMaterialInfo.HasShadingModel(EStrataShadingModel::SSM_DefaultLit) && StrataMaterialInfo.HasShadingModel(EStrataShadingModel::SSM_SubsurfaceMFP))
 			{
-				if (BlendMode == EBlendMode::BLEND_Opaque || BlendMode == EBlendMode::BLEND_Masked) 
+				if (BlendMode == EBlendMode::BLEND_Opaque || BlendMode == EBlendMode::BLEND_Masked)
 				{
-					// We only consider SSS subsurface post process for opaque materials.
-					// STRATA_TODO what does this case cover?
-					StrataMaterialInfo.SetSingleShadingModel(SSM_SubsurfaceLit);
+					StrataMaterialInfo.SetSingleShadingModel(SSM_SubsurfaceMFP);	// We only consider SSS subsurface post process for opaque materials.
 				}
 				else
 				{
-					// For transparent, we will fall back to use DefaultLit with simple volumetric
+					bSanitizeMaterial = true;										// For transparent, we will fall back to use DefaultLit with simple volumetric
+				}
+			}
+			else if (StrataMaterialInfo.CountShadingModels() == 2 && StrataMaterialInfo.HasShadingModel(EStrataShadingModel::SSM_DefaultLit) && StrataMaterialInfo.HasShadingModel(EStrataShadingModel::SSM_SubsurfaceProfile))
+			{
+				if (BlendMode == EBlendMode::BLEND_Opaque || BlendMode == EBlendMode::BLEND_Masked)
+				{
+					StrataMaterialInfo.SetSingleShadingModel(SSM_SubsurfaceProfile);// We only consider SSS subsurface post process for opaque materials.
+				}
+				else
+				{
+					bSanitizeMaterial = true;										// For transparent, we will fall back to use DefaultLit with simple volumetric
+				}
+			}
+			else if (StrataMaterialInfo.CountShadingModels() > 1 && StrataMaterialInfo.HasShadingModel(EStrataShadingModel::SSM_SubsurfaceMFP))
+			{
+				// When we gather the worst case slab material, we can hit multiple shading models when gathering from multiple BSDF possibilities of switch nodes.
+				// We do not want to sanitize the material in this case, but only keep the worst case: subsurface with all potential features on a slab.
+				if (BlendMode == EBlendMode::BLEND_Opaque || BlendMode == EBlendMode::BLEND_Masked)
+				{
+					// We only consider SSS subsurface post process for opaque materials. This can also be triggered by switch node accumulating multiple slab node.
+					// We still keep all the other material info to consider the worst case when encountering multiple shading models. And we also keep around the encountered subsurface profiles.
+					StrataMaterialInfo.SetSingleShadingModel(SSM_SubsurfaceMFP);
+				}
+				else
+				{
+					// For transparent, we will fall back to use DefaultLit worst case with simple volumetric.
 					bSanitizeMaterial = true;
 				}
 			}
-			else if (StrataMaterialInfo.CountShadingModels() > 1 && StrataMaterialInfo.HasShadingModel(EStrataShadingModel::SSM_SubsurfaceLit))
+			else if (StrataMaterialInfo.CountShadingModels() > 1 && StrataMaterialInfo.HasShadingModel(EStrataShadingModel::SSM_SubsurfaceProfile))
 			{
 				// When we gather the worst case slab material, we can hit multiple shading models when gathering from multiple BSDF possibilities of switch nodes.
 				// We do not want to sanitize the material in this case, but only keep the worst case: subsurface with all potential features on a slab.
-				if (BlendMode == EBlendMode::BLEND_Opaque || BlendMode == EBlendMode::BLEND_Masked) 
+				if (BlendMode == EBlendMode::BLEND_Opaque || BlendMode == EBlendMode::BLEND_Masked)
 				{
 					// We only consider SSS subsurface post process for opaque materials. This can also be triggered by switch node accumulating multiple slab node.
 					// We still keep all the other material info to consider the worst case when encountering multiple shading models. And we also keep around the encountered subsurface profiles.
-					StrataMaterialInfo.SetSingleShadingModel(SSM_SubsurfaceLit);
+					StrataMaterialInfo.SetSingleShadingModel(SSM_SubsurfaceProfile);
 				}
 				else
 				{
@@ -4940,20 +4961,55 @@ void UMaterial::RebuildShadingModelField()
 				UE_LOG(LogMaterial, Error, TEXT("%s: Material has more than a single material represented."), *GetName());
 			}
 		}
-		if (StrataMaterialInfo.HasOnlyShadingModel(SSM_SubsurfaceLit) && (BlendMode != EBlendMode::BLEND_Opaque && BlendMode != EBlendMode::BLEND_Masked))
+
+		if ((StrataMaterialInfo.HasOnlyShadingModel(SSM_SubsurfaceMFP) || StrataMaterialInfo.HasOnlyShadingModel(SSM_SubsurfaceProfile)) 
+			&& (BlendMode != EBlendMode::BLEND_Opaque && BlendMode != EBlendMode::BLEND_Masked))
 		{
 			// For transparent, we will fall back to use DefaultLit with simple volumetric
 			bSanitizeMaterial = true;
 		}
+
 		if (bSanitizeMaterial)
 		{
 			StrataMaterialInfo = FStrataMaterialInfo();
 			StrataMaterialInfo.AddShadingModel(SSM_DefaultLit);
+
+			if (MaterialDomain == MD_Surface)
+			{
+				// Nothing to do, the node should have added its own type. And if not type but from expression, we are going to generate that below.
+				//AddSurfaceStrataShadingModelFromMaterialShadingModel(StrataMaterialInfo, ShadingModel);
+			}
+			else if (MaterialDomain == MD_DeferredDecal)
+			{
+				StrataMaterialInfo.AddShadingModel(SSM_Decal);
+			}
+			else if (MaterialDomain == MD_LightFunction)
+			{
+				StrataMaterialInfo.AddShadingModel(SSM_VolumetricFogCloud);
+			}
+			else if (MaterialDomain == MD_Volume)
+			{
+				StrataMaterialInfo.AddShadingModel(SSM_LightFunction);
+			}
+			else if (MaterialDomain == MD_PostProcess)
+			{
+				StrataMaterialInfo.AddShadingModel(SSM_PostProcess);
+			}
+			else if (MaterialDomain == MD_UI)
+			{
+				StrataMaterialInfo.AddShadingModel(SSM_UI);
+			}
+			else if (MaterialDomain == MD_RuntimeVirtualTexture)
+			{
+				// TODO
+			}
 		}
 		
 		if (StrataMaterialInfo.HasShadingModelFromExpression())
 		{
 			MaterialDomain = EMaterialDomain::MD_Surface;
+			ShadingModel = MSM_FromMaterialExpression;
+
 			{
 				TArray<UMaterialExpressionShadingModel*> ShadingModelExpressions;
 				GetAllExpressionsInMaterialAndFunctionsOfType(ShadingModelExpressions);
@@ -4964,6 +5020,7 @@ void UMaterial::RebuildShadingModelField()
 					if (MatExpr->ShadingModel < MSM_NUM)
 					{
 						ShadingModels.AddShadingModel(MatExpr->ShadingModel);
+						AddSurfaceStrataShadingModelFromMaterialShadingModel(StrataMaterialInfo, MatExpr->ShadingModel);
 					}
 				}
 
@@ -4971,18 +5028,24 @@ void UMaterial::RebuildShadingModelField()
 				if (!ShadingModels.IsValid())
 				{
 					ShadingModels.AddShadingModel(MSM_DefaultLit);
+					AddSurfaceStrataShadingModelFromMaterialShadingModel(StrataMaterialInfo, MSM_DefaultLit);
 				}
 			}
 		}
 		else
 		{
-			// Now derive some properties from the material
+			// Now derive some properties from the material into the legacy fields
 			if (StrataMaterialInfo.HasOnlyShadingModel(SSM_Unlit))
 			{
 				MaterialDomain = EMaterialDomain::MD_Surface;
 				ShadingModel = MSM_Unlit;
 			}
-			else if (StrataMaterialInfo.HasOnlyShadingModel(SSM_SubsurfaceLit))
+			else if (StrataMaterialInfo.HasOnlyShadingModel(SSM_DefaultLit))
+			{
+				MaterialDomain = EMaterialDomain::MD_Surface;
+				ShadingModel = MSM_DefaultLit;
+			}
+			else if (StrataMaterialInfo.HasOnlyShadingModel(SSM_SubsurfaceProfile))
 			{
 				MaterialDomain = EMaterialDomain::MD_Surface;
 				ShadingModel = MSM_SubsurfaceProfile;
@@ -4992,10 +5055,25 @@ void UMaterial::RebuildShadingModelField()
 					BlendMode = EBlendMode::BLEND_Opaque;
 				}
 			}
-			else if (StrataMaterialInfo.HasOnlyShadingModel(SSM_DefaultLit))
+			else if (StrataMaterialInfo.HasOnlyShadingModel(SSM_SubsurfaceWrap))
 			{
 				MaterialDomain = EMaterialDomain::MD_Surface;
-				ShadingModel = MSM_DefaultLit;
+				ShadingModel = MSM_Subsurface;
+				if (BlendMode != EBlendMode::BLEND_Opaque && BlendMode != EBlendMode::BLEND_Masked)
+				{
+					UE_LOG(LogMaterial, Error, TEXT("%s: Material has subsurface data, and its blending mode is not set to Opaque or Masked. Forcing blend mode to Opaque."), *GetName());
+					BlendMode = EBlendMode::BLEND_Opaque;
+				}
+			}
+			else if (StrataMaterialInfo.HasOnlyShadingModel(SSM_SubsurfaceThinTwoSided))
+			{
+				MaterialDomain = EMaterialDomain::MD_Surface;
+				ShadingModel = MSM_TwoSidedFoliage;
+			}
+			else if (StrataMaterialInfo.HasOnlyShadingModel(SSM_ThinTranslucent))
+			{
+				MaterialDomain = EMaterialDomain::MD_Surface;
+				ShadingModel = MSM_ThinTranslucent;
 			}
 			else if (StrataMaterialInfo.HasOnlyShadingModel(SSM_VolumetricFogCloud))
 			{
@@ -5067,12 +5145,8 @@ void UMaterial::RebuildShadingModelField()
 
 		// Now, reset the subsurface profile (in case it has been removed from any slab before) and set it only if needed.
 		SubsurfaceProfile = nullptr;
-		if ((StrataMaterialInfo.HasOnlyShadingModel(SSM_Eye) || StrataMaterialInfo.HasOnlyShadingModel(SSM_SubsurfaceLit)) && StrataMaterialInfo.CountSubsurfaceProfiles() > 0)
+		if ((StrataMaterialInfo.HasOnlyShadingModel(SSM_Eye) || StrataMaterialInfo.HasOnlyShadingModel(SSM_SubsurfaceProfile)) && StrataMaterialInfo.CountSubsurfaceProfiles() > 0)
 		{
-			if (StrataMaterialInfo.CountSubsurfaceProfiles() > 1)
-			{
-				UE_LOG(LogMaterial, Error, TEXT("%s: Material has more than a single sub-surface profile used."), *GetName());
-			}
 			SubsurfaceProfile = StrataMaterialInfo.GetSubsurfaceProfile();
 		}
 
@@ -5090,6 +5164,8 @@ void UMaterial::RebuildShadingModelField()
 		TArray<UMaterialExpressionShadingModel*> ShadingModelExpressions;
 		GetAllExpressionsInMaterialAndFunctionsOfType(ShadingModelExpressions);
 
+		ShadingModel = MSM_FromMaterialExpression;
+
 		for (UMaterialExpressionShadingModel* MatExpr : ShadingModelExpressions)
 		{
 			// Ensure the Shading model is valid
diff --git a/Engine/Source/Runtime/Engine/Private/Materials/MaterialCachedData.cpp b/Engine/Source/Runtime/Engine/Private/Materials/MaterialCachedData.cpp
--- a/Engine/Source/Runtime/Engine/Private/Materials/MaterialCachedData.cpp
+++ b/Engine/Source/Runtime/Engine/Private/Materials/MaterialCachedData.cpp
@@ -1100,29 +1100,12 @@ void FMaterialCachedExpressionData::GetParameterValueByIndex(EMaterialParameterT
 	}
 }
 
-bool FStrataMaterialInfo::Serialize(FArchive& Ar)
-{
-	Ar.UsingCustomVersion(FUE5MainStreamObjectVersion::GUID);
-	return false;
-}
-
 bool FMaterialCachedExpressionData::Serialize(FArchive& Ar)
 {
 	Ar.UsingCustomVersion(FUE5MainStreamObjectVersion::GUID);
 	return false;
 }
 
-void FStrataMaterialInfo::PostSerialize(const FArchive& Ar)
-{
-	if (Ar.IsLoading())
-	{
-		if (Ar.CustomVer(FUE5MainStreamObjectVersion::GUID) < FUE5MainStreamObjectVersion::IncreaseMaterialAttributesInputMask)
-		{
-			ConnectedPropertyMask = uint64(ConnectedProperties_DEPRECATED);
-		}
-	}
-}
-
 void FMaterialCachedExpressionData::PostSerialize(const FArchive& Ar)
 {
 	if (Ar.IsLoading())
diff --git a/Engine/Source/Runtime/Engine/Private/Materials/MaterialExpressions.cpp b/Engine/Source/Runtime/Engine/Private/Materials/MaterialExpressions.cpp
--- a/Engine/Source/Runtime/Engine/Private/Materials/MaterialExpressions.cpp
+++ b/Engine/Source/Runtime/Engine/Private/Materials/MaterialExpressions.cpp
@@ -329,8 +329,7 @@ bool IsAllowedExpressionType(const UClass* const Class, const bool bMaterialFunc
 	const bool bSharedAllowed = Class != UMaterialExpressionComment::StaticClass() 
 		&& Class != UMaterialExpressionPinBase::StaticClass()
 		&& Class != UMaterialExpressionParameter::StaticClass()
-		&& (Class != UMaterialExpressionTextureSampleParameter2DArray::StaticClass() || AllowTextureArrayAssetCreationVar->GetValueOnGameThread() != 0)
-		&& Class != UMaterialExpressionStrataLegacyConversion::StaticClass();
+		&& (Class != UMaterialExpressionTextureSampleParameter2DArray::StaticClass() || AllowTextureArrayAssetCreationVar->GetValueOnGameThread() != 0);
 
 	if (bMaterialFunction)
 	{
@@ -23290,7 +23289,6 @@ int32 UMaterialExpressionStrataLegacyConversion::Compile(class FMaterialCompiler
 	FStrataOperator& StrataOperator = Compiler->StrataCompilationGetOperator(Compiler->StrataTreeStackGetPathUniqueId());
 	StrataOperator.BSDFRegisteredSharedLocalBasis = NewRegisteredSharedLocalBasis;
 
-	int32 ShadingModelCount = ConvertedStrataMaterialInfo.CountShadingModels();
 	int32 OpacityCodeChunk = INDEX_NONE;
 	if (!Compiler->StrataSkipsOpacityEvaluation())
 	{
@@ -23303,6 +23301,8 @@ int32 UMaterialExpressionStrataLegacyConversion::Compile(class FMaterialCompiler
 		OpacityCodeChunk = Compiler->Constant(1.0f);
 	}
 
+	int32 ShadingModelCodeChunk = ShadingModel.IsConnected() ? CompileWithDefaultFloat1(Compiler, ShadingModel, float(MSM_DefaultLit)) : Compiler->Constant(float(ShadingModelOverride));
+	int32 ShadingModelCount = Compiler->GetMaterialShadingModels().CountShadingModels();
 	const bool bHasDynamicShadingModels = ShadingModelCount > 1;
 	// We probably need to do something along these line as well :::
 	int32 OutputCodeChunk = Compiler->StrataConversionFromLegacy(
@@ -23330,7 +23330,7 @@ int32 UMaterialExpressionStrataLegacyConversion::Compile(class FMaterialCompiler
 		CompileWithDefaultFloat1(Compiler, WaterPhaseG, 0.0f),
 		CompileWithDefaultFloat3(Compiler, ColorScaleBehindWater, 1.0f, 1.0f, 1.0f),
 		// Shading model
-		CompileWithDefaultFloat1(Compiler, ShadingModel, 0.0f),
+		ShadingModelCodeChunk,
 		NormalCodeChunk,
 		TangentCodeChunk,
 		BasisIndexMacro,
@@ -23411,7 +23411,7 @@ FName UMaterialExpressionStrataLegacyConversion::GetInputName(int32 InputIndex)
 	else if (InputIndex == 16)	return TEXT("Color Scale BehindWater");
 	else if (InputIndex == 17)	return TEXT("Clear Coat Normal");
 	else if (InputIndex == 18)	return TEXT("Custom Tangent");
-	else if (InputIndex == 19)	return TEXT("Shading Model");
+	else if (InputIndex == 19)	return TEXT("Shading Model From Expression");
 	return TEXT("Unknown");
 }
 
@@ -23445,19 +23445,30 @@ void UMaterialExpressionStrataLegacyConversion::GatherStrataMaterialInfo(FStrata
 	if (ClearCoat.IsConnected())			{ StrataMaterialInfo.AddPropertyConnected(MP_CustomData0); }
 	if (ClearCoatRoughness.IsConnected())	{ StrataMaterialInfo.AddPropertyConnected(MP_CustomData1); }
 	if (Opacity.IsConnected())				{ StrataMaterialInfo.AddPropertyConnected(MP_Opacity); }
-	if (ShadingModel.IsConnected())			{ StrataMaterialInfo.AddPropertyConnected(MP_ShadingModel); }
 
-	if (ConvertedStrataMaterialInfo.HasShadingModel(SSM_Unlit))					{ StrataMaterialInfo.AddShadingModel(SSM_Unlit); }
-	if (ConvertedStrataMaterialInfo.HasShadingModel(SSM_DefaultLit))			{ StrataMaterialInfo.AddShadingModel(SSM_DefaultLit); }
-	if (ConvertedStrataMaterialInfo.HasShadingModel(SSM_SubsurfaceLit))			{ StrataMaterialInfo.AddShadingModel(SSM_SubsurfaceLit); }
-	if (ConvertedStrataMaterialInfo.HasShadingModel(SSM_VolumetricFogCloud))	{ StrataMaterialInfo.AddShadingModel(SSM_VolumetricFogCloud); }
-	if (ConvertedStrataMaterialInfo.HasShadingModel(SSM_Hair))					{ StrataMaterialInfo.AddShadingModel(SSM_Hair); }
-	if (ConvertedStrataMaterialInfo.HasShadingModel(SSM_Eye))					{ StrataMaterialInfo.AddShadingModel(SSM_Eye); }
-	if (ConvertedStrataMaterialInfo.HasShadingModel(SSM_Cloth))					{ StrataMaterialInfo.AddShadingModel(SSM_Cloth); }
-	if (ConvertedStrataMaterialInfo.HasShadingModel(SSM_ClearCoat))				{ StrataMaterialInfo.AddShadingModel(SSM_ClearCoat); }
-	if (ConvertedStrataMaterialInfo.HasShadingModel(SSM_SingleLayerWater))		{ StrataMaterialInfo.AddShadingModel(SSM_SingleLayerWater); }
-	if (SubsurfaceProfile)														{ StrataMaterialInfo.AddSubsurfaceProfile(SubsurfaceProfile); }
-	if (ConvertedStrataMaterialInfo.HasShadingModelFromExpression())			{ StrataMaterialInfo.SetShadingModelFromExpression(true); }
+	if (ShadingModel.IsConnected())
+	{
+		StrataMaterialInfo.AddPropertyConnected(MP_ShadingModel);
+
+		// If the ShadingModel pin is plugged in, we must use a shading model from expression path.
+		StrataMaterialInfo.SetShadingModelFromExpression(true);
+	}
+	else
+	{
+		// If the ShadingModel pin is NOT plugged in, we simply use the shading model selected on the root node drop box.
+		if (ShadingModelOverride == MSM_Unlit)				{ StrataMaterialInfo.AddShadingModel(EStrataShadingModel::SSM_Unlit); }
+		if (ShadingModelOverride == MSM_DefaultLit)			{ StrataMaterialInfo.AddShadingModel(EStrataShadingModel::SSM_DefaultLit); }
+		if (ShadingModelOverride == MSM_Subsurface)			{ StrataMaterialInfo.AddShadingModel(EStrataShadingModel::SSM_SubsurfaceWrap); }
+		if (ShadingModelOverride == MSM_PreintegratedSkin)	{ StrataMaterialInfo.AddShadingModel(EStrataShadingModel::SSM_SubsurfaceWrap); }
+		if (ShadingModelOverride == MSM_ClearCoat)			{ StrataMaterialInfo.AddShadingModel(EStrataShadingModel::SSM_ClearCoat); }
+		if (ShadingModelOverride == MSM_SubsurfaceProfile)	{ StrataMaterialInfo.AddShadingModel(EStrataShadingModel::SSM_SubsurfaceProfile); }
+		if (ShadingModelOverride == MSM_TwoSidedFoliage)	{ StrataMaterialInfo.AddShadingModel(EStrataShadingModel::SSM_SubsurfaceThinTwoSided); }
+		if (ShadingModelOverride == MSM_Hair)				{ StrataMaterialInfo.AddShadingModel(EStrataShadingModel::SSM_Hair); }
+		if (ShadingModelOverride == MSM_Cloth)				{ StrataMaterialInfo.AddShadingModel(EStrataShadingModel::SSM_Cloth); }
+		if (ShadingModelOverride == MSM_Eye)				{ StrataMaterialInfo.AddShadingModel(EStrataShadingModel::SSM_Eye); }
+		if (ShadingModelOverride == MSM_SingleLayerWater)	{ StrataMaterialInfo.AddShadingModel(EStrataShadingModel::SSM_SingleLayerWater); }
+		if (ShadingModelOverride == MSM_ThinTranslucent)	{ StrataMaterialInfo.AddShadingModel(EStrataShadingModel::SSM_ThinTranslucent); }
+	}
 }
 
 FStrataOperator* UMaterialExpressionStrataLegacyConversion::StrataGenerateMaterialTopologyTree(class FMaterialCompiler* Compiler, class UMaterialExpression* Parent, int32 OutputIndex)
@@ -23478,35 +23489,55 @@ FStrataOperator* UMaterialExpressionStrataLegacyConversion::StrataGenerateMateri
 		return &SlabOperator;
 	};
 
+	// Get the shading models resulting from the UMaterial::RebuildShadingModelField().
+	FMaterialShadingModelField ShadingModels = Compiler->GetMaterialShadingModels();
+
 	// Logic about shading models and complexity should match UMaterialExpressionStrataLegacyConversion::Compile.
-	if (ConvertedStrataMaterialInfo.CountShadingModels() > 1 || ConvertedStrataMaterialInfo.HasShadingModelFromExpression())
+	const bool bHasShadingModelFromExpression = ShadingModel.IsConnected(); // We keep HasShadingModelFromExpression in case all shading models cannot be safely recovered from material functions.
+	if ((ShadingModels.CountShadingModels() > 1) || bHasShadingModelFromExpression) 
 	{
 		return AddDefaultWorstCase(true, true);
 	}
 	// else
 	{
-		check(ConvertedStrataMaterialInfo.CountShadingModels() == 1);
+		check(ShadingModels.CountShadingModels() == 1);
 
-		if (ConvertedStrataMaterialInfo.HasShadingModel(SSM_Unlit))
+		if (ShadingModels.HasShadingModel(MSM_Unlit))
 		{
 			FStrataOperator& Operator = Compiler->StrataCompilationRegisterOperator(STRATA_OPERATOR_BSDF_LEGACY, Compiler->StrataTreeStackGetPathUniqueId(), this, Parent, Compiler->StrataTreeStackGetParentPathUniqueId());
 			Operator.BSDFType = STRATA_BSDF_TYPE_UNLIT;
 			Operator.ThicknessIndex = ThicknessIndex;
 			return &Operator;
 		}
-		else if (ConvertedStrataMaterialInfo.HasShadingModel(SSM_DefaultLit))
+		else if (ShadingModels.HasShadingModel(MSM_DefaultLit))
 		{
 			return AddDefaultWorstCase(false, false);
 		}
-		else if (ConvertedStrataMaterialInfo.HasShadingModel(SSM_SubsurfaceLit))
+		else if (ShadingModels.HasShadingModel(MSM_ThinTranslucent))
+		{
+			return AddDefaultWorstCase(false, false);
+		}
+		else if (ShadingModels.HasShadingModel(MSM_SubsurfaceProfile))
+		{
+			return AddDefaultWorstCase(true, false);
+		}
+		else if (ShadingModels.HasShadingModel(MSM_Subsurface))
+		{
+			return AddDefaultWorstCase(true, false);
+		}
+		else if (ShadingModels.HasShadingModel(MSM_TwoSidedFoliage))
 		{
 			return AddDefaultWorstCase(true, false);
 		}
-		else if (ConvertedStrataMaterialInfo.HasShadingModel(SSM_Cloth))
+		else if (ShadingModels.HasShadingModel(MSM_PreintegratedSkin))
+		{
+			return AddDefaultWorstCase(true, false);
+		}
+		else if (ShadingModels.HasShadingModel(MSM_Cloth))
 		{
 			return AddDefaultWorstCase(false, true);
 		}
-		else if (ConvertedStrataMaterialInfo.HasShadingModel(SSM_ClearCoat))
+		else if (ShadingModels.HasShadingModel(MSM_ClearCoat))
 		{
 			FStrataOperator& Operator = Compiler->StrataCompilationRegisterOperator(STRATA_OPERATOR_BSDF_LEGACY, Compiler->StrataTreeStackGetPathUniqueId(), this, Parent, Compiler->StrataTreeStackGetParentPathUniqueId());
 			Operator.BSDFType = STRATA_BSDF_TYPE_SLAB;
@@ -23515,21 +23546,21 @@ FStrataOperator* UMaterialExpressionStrataLegacyConversion::StrataGenerateMateri
 			Operator.bBSDFHasAnisotropy = Anisotropy.IsConnected();
 			return &Operator;
 		}
-		else if (ConvertedStrataMaterialInfo.HasShadingModel(SSM_Hair))
+		else if (ShadingModels.HasShadingModel(MSM_Hair))
 		{
 			FStrataOperator& Operator = Compiler->StrataCompilationRegisterOperator(STRATA_OPERATOR_BSDF_LEGACY, Compiler->StrataTreeStackGetPathUniqueId(), this, Parent, Compiler->StrataTreeStackGetParentPathUniqueId());
 			Operator.BSDFType = STRATA_BSDF_TYPE_HAIR;
 			Operator.ThicknessIndex = ThicknessIndex;
 			return &Operator;
 		}
-		else if (ConvertedStrataMaterialInfo.HasShadingModel(SSM_Eye))
+		else if (ShadingModels.HasShadingModel(MSM_Eye))
 		{
 			FStrataOperator& Operator = Compiler->StrataCompilationRegisterOperator(STRATA_OPERATOR_BSDF_LEGACY, Compiler->StrataTreeStackGetPathUniqueId(), this, Parent, Compiler->StrataTreeStackGetParentPathUniqueId());
 			Operator.BSDFType = STRATA_BSDF_TYPE_EYE;
 			Operator.ThicknessIndex = ThicknessIndex;
 			return &Operator;
 		}
-		else if (ConvertedStrataMaterialInfo.HasShadingModel(SSM_SingleLayerWater))
+		else if (ShadingModels.HasShadingModel(MSM_SingleLayerWater))
 		{
 			FStrataOperator& Operator = Compiler->StrataCompilationRegisterOperator(STRATA_OPERATOR_BSDF_LEGACY, Compiler->StrataTreeStackGetPathUniqueId(), this, Parent, Compiler->StrataTreeStackGetParentPathUniqueId());
 			Operator.BSDFType = STRATA_BSDF_TYPE_SINGLELAYERWATER;
@@ -23952,11 +23983,15 @@ void UMaterialExpressionStrataSlabBSDF::GatherStrataMaterialInfo(FStrataMaterial
 	if (HasSSS())
 	{
 		// We still do not know if this is going to be a real SSS node because it is only possible for BSDF at the bottom of the stack. Nevertheless, we take the worst case into account.
-		StrataMaterialInfo.AddShadingModel(SSM_SubsurfaceLit);
 		if (SubsurfaceProfile)
 		{
+			StrataMaterialInfo.AddShadingModel(SSM_SubsurfaceProfile);
 			StrataMaterialInfo.AddSubsurfaceProfile(SubsurfaceProfile);
 		}
+		else
+		{
+			StrataMaterialInfo.AddShadingModel(SSM_SubsurfaceMFP);
+		}
 	}
 	else
 	{
