Substrate - LegacyConversionNode now update the input name according to shading model selected.

diff --git a/Engine/Source/Runtime/Engine/Private/Materials/MaterialExpressions.cpp b/Engine/Source/Runtime/Engine/Private/Materials/MaterialExpressions.cpp
--- a/Engine/Source/Runtime/Engine/Private/Materials/MaterialExpressions.cpp
+++ b/Engine/Source/Runtime/Engine/Private/Materials/MaterialExpressions.cpp
@@ -23392,6 +23392,8 @@ uint32 UMaterialExpressionStrataLegacyConversion::GetInputType(int32 InputIndex)
 
 FName UMaterialExpressionStrataLegacyConversion::GetInputName(int32 InputIndex) const
 {
+	const bool bShadingModelFromExpression = ShadingModel.IsConnected();
+
 	if (InputIndex == 0)		return TEXT("BaseColor");
 	else if (InputIndex == 1)	return TEXT("Metallic");
 	else if (InputIndex == 2)	return TEXT("Specular");
@@ -23400,9 +23402,50 @@ FName UMaterialExpressionStrataLegacyConversion::GetInputName(int32 InputIndex)
 	else if (InputIndex == 5)	return TEXT("Emissive Color");
 	else if (InputIndex == 6)	return TEXT("Normal");
 	else if (InputIndex == 7)	return TEXT("Tangent");
-	else if (InputIndex == 8)	return TEXT("Sub-Surface Color");
-	else if (InputIndex == 9)	return TEXT("Clear Coat");
-	else if (InputIndex == 10)	return TEXT("Clear Coat Roughness");
+	else if (InputIndex == 8)
+	{
+		if (!bShadingModelFromExpression && ShadingModelOverride == MSM_Cloth)
+		{
+			return TEXT("Fuzz Color");
+		}
+		return TEXT("Subsurface Color");
+	}
+	else if (InputIndex == 9)
+	{
+		if (!bShadingModelFromExpression)
+		{
+			if (ShadingModelOverride == MSM_Cloth)
+			{
+				return TEXT("Fuzz Amount");
+			}
+			else if (ShadingModelOverride == MSM_Eye)
+			{
+				return TEXT("Iris Mask");
+			}
+			else if (ShadingModelOverride == MSM_Hair)
+			{
+				return TEXT("Backlit");
+			}
+			else if (ShadingModelOverride == MSM_ClearCoat)
+			{
+				return TEXT("Clear Coat");
+			}
+			return TEXT("Unused");
+		}
+		return TEXT("Custom0");
+	}
+	else if (InputIndex == 10)
+	{
+		if (!bShadingModelFromExpression)
+		{
+			if (ShadingModelOverride == MSM_Eye)
+			{
+				return TEXT("Iris Distance");
+			}
+			return TEXT("Unused");
+		}
+		return TEXT("Custom1");
+	}
 	else if (InputIndex == 11)	return TEXT("Opacity");
 	else if (InputIndex == 12)	return TEXT("TransmittanceColor");
 	else if (InputIndex == 13)	return TEXT("Water Scattering Coefficients");
