Substrate - renamed the legacing conversion node to "Substrate Shading Models".

diff --git a/Engine/Source/Runtime/Engine/Classes/Materials/MaterialExpressionStrata.h b/Engine/Source/Runtime/Engine/Classes/Materials/MaterialExpressionStrata.h
--- a/Engine/Source/Runtime/Engine/Classes/Materials/MaterialExpressionStrata.h
+++ b/Engine/Source/Runtime/Engine/Classes/Materials/MaterialExpressionStrata.h
@@ -41,7 +41,7 @@ class UMaterialExpressionStrataBSDF : public UMaterialExpression
 };
 
 
-UCLASS(MinimalAPI, collapsecategories, hidecategories = Object, DisplayName = "Substrate Legacy Conversion")
+UCLASS(MinimalAPI, collapsecategories, hidecategories = Object, DisplayName = "Substrate Shading Models")
 class UMaterialExpressionStrataLegacyConversion : public UMaterialExpressionStrataBSDF
 {
 	GENERATED_UCLASS_BODY()
diff --git a/Engine/Source/Runtime/Engine/Private/Materials/MaterialExpressions.cpp b/Engine/Source/Runtime/Engine/Private/Materials/MaterialExpressions.cpp
--- a/Engine/Source/Runtime/Engine/Private/Materials/MaterialExpressions.cpp
+++ b/Engine/Source/Runtime/Engine/Private/Materials/MaterialExpressions.cpp
@@ -23355,7 +23355,7 @@ void UMaterialExpressionStrataLegacyConversion::PostEditChangeProperty(FProperty
 
 void UMaterialExpressionStrataLegacyConversion::GetCaption(TArray<FString>& OutCaptions) const
 {
-	OutCaptions.Add(TEXT("Substrate Legacy Conversion"));
+	OutCaptions.Add(TEXT("Substrate Shading Models"));
 }
 
 uint32 UMaterialExpressionStrataLegacyConversion::GetOutputType(int32 OutputIndex)
