MemReport now can write command's output into separate csv files for commands specified in Engine.ini in [MemReportCsvCommands] section.

#tests editor

diff --git a/Engine/Source/Runtime/Engine/Private/UnrealEngine.cpp b/Engine/Source/Runtime/Engine/Private/UnrealEngine.cpp
--- a/Engine/Source/Runtime/Engine/Private/UnrealEngine.cpp
+++ b/Engine/Source/Runtime/Engine/Private/UnrealEngine.cpp
@@ -7419,7 +7419,6 @@ bool UEngine::HandleMemReportDeferredCommand( const TCHAR* Cmd, FOutputDevice& A
 	FSlowHeartBeatScope HeartBeatSuspend;
 
 	const bool bPerformSlowCommands = FParse::Param( Cmd, TEXT("FULL") );
-	const bool bPerformGPUCommands = FParse::Param( Cmd, TEXT("GPU") );
 	const bool bLogOutputToFile = !FParse::Param( Cmd, TEXT("LOG") );
 	const bool bCSV = FParse::Param(Cmd, TEXT("CSV"));
 	FString InFileName;
@@ -7434,8 +7433,8 @@ bool UEngine::HandleMemReportDeferredCommand( const TCHAR* Cmd, FOutputDevice& A
 	FlushRenderingCommands();
 
 	FOutputDevice* ReportAr = &Ar;
-	FArchive* FileAr = nullptr;
-	FOutputDeviceArchiveWrapper* FileArWrapper = nullptr;
+	TUniquePtr<FArchive> FileAr;
+	TOptional<FOutputDeviceArchiveWrapper> FileArWrapper;
 	FString FilenameFull;
 
 	if (bLogOutputToFile)
@@ -7446,9 +7445,9 @@ bool UEngine::HandleMemReportDeferredCommand( const TCHAR* Cmd, FOutputDevice& A
 		const FString Filename = CreateProfileFilename(InFileName, TEXT(".memreport"), true );
 		FilenameFull = PathName + Filename;
 
-		FileAr = IFileManager::Get().CreateDebugFileWriter(*FilenameFull);
-		FileArWrapper = new FOutputDeviceArchiveWrapper(FileAr);
-		ReportAr = FileArWrapper;
+		FileAr = TUniquePtr<FArchive>(IFileManager::Get().CreateDebugFileWriter(*FilenameFull));
+		FileArWrapper.Emplace(FileAr.Get());
+		ReportAr = FileArWrapper.GetPtrOrNull();
 
 		UE_LOG(LogEngine, Log, TEXT("MemReportDeferred: saving to %s"), *FilenameFull);		
 	}
@@ -7499,13 +7498,12 @@ bool UEngine::HandleMemReportDeferredCommand( const TCHAR* Cmd, FOutputDevice& A
 
 
 	// Run commands from the ini
-	FConfigSection* CommandsToRun = GConfig->GetSectionPrivate(TEXT("MemReportCommands"), 0, 1, GEngineIni);
+	TArray<FString> CommandsToRun;
 
-	if (CommandsToRun)
+	if (GConfig->GetArray(TEXT("MemReportCommands"), TEXT("Cmd"), CommandsToRun, GEngineIni))
 	{
-		for (FConfigSectionMap::TIterator It(*CommandsToRun); It; ++It)
+		for (FString& Command : CommandsToRun)
 		{
-			FString Command = It.Value().GetValue();
 			if (bCSV)
 			{
 				Command += " -csv";
@@ -7519,13 +7517,12 @@ bool UEngine::HandleMemReportDeferredCommand( const TCHAR* Cmd, FOutputDevice& A
 
 	if (bPerformSlowCommands)
 	{
-		CommandsToRun = GConfig->GetSectionPrivate(TEXT("MemReportFullCommands"), 0, 1, GEngineIni);
+		CommandsToRun.Reset();
 
-		if (CommandsToRun)
+		if (GConfig->GetArray(TEXT("MemReportFullCommands"), TEXT("Cmd"), CommandsToRun, GEngineIni))
 		{
-			for (FConfigSectionMap::TIterator It(*CommandsToRun); It; ++It)
+			for (FString& Command : CommandsToRun)
 			{
-				FString Command = It.Value().GetValue();
 				if (bCSV)
 				{
 					Command += " -csv";
@@ -7538,11 +7535,26 @@ bool UEngine::HandleMemReportDeferredCommand( const TCHAR* Cmd, FOutputDevice& A
 		}
 	}
 
-	if (FileArWrapper != nullptr)
+	if (bLogOutputToFile && bCSV)
 	{
-		FileArWrapper->TearDown();
-		delete FileArWrapper;
-		delete FileAr;
+		CommandsToRun.Reset();
+
+		if (GConfig->GetArray(TEXT("MemReportCsvCommands"), TEXT("Cmd"), CommandsToRun, GEngineIni))
+		{
+			const FString DirectoryName = FPaths::GetBaseFilename(FilenameFull, false) + TEXT("/");
+			IFileManager::Get().MakeDirectory( *DirectoryName );
+
+			for (FString& Command : CommandsToRun)
+			{
+				const FString ValidFilePath = DirectoryName + FPaths::MakeValidFileName(Command) + TEXT(".csv");
+
+				TUniquePtr<FArchive> CommandFile(IFileManager::Get().CreateDebugFileWriter(*ValidFilePath));
+				FOutputDeviceArchiveWrapper CommandFileAr(CommandFile.Get());
+				Command += TEXT(" -csv");
+				Exec(InWorld, *Command, CommandFileAr);
+				UE_LOG(LogEngine, Log, TEXT("MemReportDeferred: Written \"%s\" output to \"%s\""), *Command, *ValidFilePath);
+			}
+		}
 	}
 #endif
 
