Fix for crash during UHLODSubsystem::OnExternalStreamingObjectRemoved since we disabled ShouldReuseUnloadedButStillAroundLevels: because streaming objects keep a soft object path to their owning world (since they can be reoutered to be saved), we need to unregister the containing world partition object before renaming the owning world.

diff --git a/Engine/Source/Runtime/Engine/Private/WorldPartition/HLOD/HLODSubsystem.cpp b/Engine/Source/Runtime/Engine/Private/WorldPartition/HLOD/HLODSubsystem.cpp
--- a/Engine/Source/Runtime/Engine/Private/WorldPartition/HLOD/HLODSubsystem.cpp
+++ b/Engine/Source/Runtime/Engine/Private/WorldPartition/HLOD/HLODSubsystem.cpp
@@ -378,7 +378,7 @@ void UHLODSubsystem::OnExternalStreamingObjectInjected(URuntimeHashExternalStrea
 
 void UHLODSubsystem::OnExternalStreamingObjectRemoved(URuntimeHashExternalStreamingObjectBase* ExternalStreamingObject)
 {
-	UWorldPartition* OwnerPartition = ExternalStreamingObject->GetOuterWorld() ? ExternalStreamingObject->GetOuterWorld()->GetWorldPartition() : ExternalStreamingObject->GetTypedOuter<UWorldPartition>();
+	UWorldPartition* OwnerPartition = ExternalStreamingObject->GetOuterWorld()->GetWorldPartition();
 	FWorldPartitionHLODRuntimeData* WorldPartitionHLODRuntimeData = WorldPartitionsHLODRuntimeData.Find(OwnerPartition);
 	if (WorldPartitionHLODRuntimeData)
 	{
