Fixed compilation issue from CL 27408620.

diff --git a/Engine/Source/Runtime/Engine/Classes/Engine/EngineTypes.h b/Engine/Source/Runtime/Engine/Classes/Engine/EngineTypes.h
--- a/Engine/Source/Runtime/Engine/Classes/Engine/EngineTypes.h
+++ b/Engine/Source/Runtime/Engine/Classes/Engine/EngineTypes.h
@@ -725,7 +725,7 @@ public:
 	void SetSingleShadingModel(EStrataShadingModel InShadingModel) { check(InShadingModel < SSM_NUM); ShadingModelField = (uint16)(1 << (uint16)InShadingModel); }
 	bool HasShadingModel(EStrataShadingModel InShadingModel) const { return (ShadingModelField & (1 << (uint16)InShadingModel)) != 0; }
 	bool HasOnlyShadingModel(EStrataShadingModel InShadingModel) const { return ShadingModelField == (1 << (uint16)InShadingModel); }
-	uint16 GetShadingModelField() const { return ShadingModelField; }
+	uint32 GetShadingModelField() const { return ShadingModelField; }
 	int32 CountShadingModels() const { return FMath::CountBits(ShadingModelField); }
 
 	// Subsurface profiles
