[mutable] Expose function to get the platform-specific build settings of a texture in the editor.
#rnx

diff --git a/Engine/Source/Runtime/Engine/Classes/Engine/Texture.h b/Engine/Source/Runtime/Engine/Classes/Engine/Texture.h
--- a/Engine/Source/Runtime/Engine/Classes/Engine/Texture.h
+++ b/Engine/Source/Runtime/Engine/Classes/Engine/Texture.h
@@ -1919,6 +1919,11 @@ public:
 	 * Try to get as close as it can to the target resolution but it will stay above it if it can't reach it
 	 */
 	ENGINE_API bool DownsizeImageUsingTextureSettings(const ITargetPlatform* TargetPlatform, FImage& InOutImage, int32 TargetSize, int32 LayerIndex);
+
+	/*
+	*/
+	ENGINE_API void GetTargetPlatformBuildSettings(const ITargetPlatform* TargetPlatform, TArray<FTextureBuildSettings>& OutSettings);
+
 #endif
 
 protected:
diff --git a/Engine/Source/Runtime/Engine/Private/TextureDerivedData.cpp b/Engine/Source/Runtime/Engine/Private/TextureDerivedData.cpp
--- a/Engine/Source/Runtime/Engine/Private/TextureDerivedData.cpp
+++ b/Engine/Source/Runtime/Engine/Private/TextureDerivedData.cpp
@@ -3959,4 +3959,12 @@ bool UTexture::DownsizeImageUsingTextureSettings(const ITargetPlatform* TargetPl
 
 	return true;
 }
+
+
+void UTexture::GetTargetPlatformBuildSettings(const ITargetPlatform* TargetPlatform, TArray<FTextureBuildSettings>& OutSettings )
+{
+	GetBuildSettingsForTargetPlatform(*this, TargetPlatform, ETextureEncodeSpeed::Final, OutSettings, nullptr);
+}
+
+
 #endif
