Fix for ThinTranslucent shading model not working correctly after 27349515.

diff --git a/Engine/Source/Runtime/Engine/Private/Materials/HLSLMaterialTranslator.cpp b/Engine/Source/Runtime/Engine/Private/Materials/HLSLMaterialTranslator.cpp
--- a/Engine/Source/Runtime/Engine/Private/Materials/HLSLMaterialTranslator.cpp
+++ b/Engine/Source/Runtime/Engine/Private/Materials/HLSLMaterialTranslator.cpp
@@ -1972,7 +1972,6 @@ void FHLSLMaterialTranslator::ValidateShadingModelsForFeatureLevel(const FMateri
 void FHLSLMaterialTranslator::GetMaterialEnvironment(EShaderPlatform InPlatform, FShaderCompilerEnvironment& OutEnvironment)
 {
 	const bool bStrataEnabled = Strata::IsStrataEnabled();
-	bool bMaterialRequestsDualSourceBlending = false;
 
 	if (bNeedsParticlePosition || Material->ShouldGenerateSphericalParticleNormals() || bUsesSphericalParticleOpacity)
 	{
@@ -2214,8 +2213,6 @@ void FHLSLMaterialTranslator::GetMaterialEnvironment(EShaderPlatform InPlatform,
 		{
 			OutEnvironment.SetDefine(TEXT("MATERIAL_SHADINGMODEL_THIN_TRANSLUCENT"), TEXT("1"));
 			NumSetMaterials++;
-
-			bMaterialRequestsDualSourceBlending = true;
 		}
 
 		if (ShadingModels.HasShadingModel(MSM_SingleLayerWater) && FDataDrivenShaderPlatformInfo::GetRequiresDisableForwardLocalLights(Platform))
@@ -2311,11 +2308,8 @@ void FHLSLMaterialTranslator::GetMaterialEnvironment(EShaderPlatform InPlatform,
 
 	OutEnvironment.SetDefine(TEXT("MATERIAL_IS_STRATA"), bMaterialIsStrata ? TEXT("1") : TEXT("0"));
 
-	// Strata requests dual source blending only for BLEND_TranslucentColoredTransmittance
-	bMaterialRequestsDualSourceBlending |= bMaterialIsStrata && Material->GetBlendMode() == EBlendMode::BLEND_TranslucentColoredTransmittance;
-
 	// if duals source blending (colored transmittance) is not supported on a platform, it will fall back to standard alpha blending (grey scale transmittance)
-	OutEnvironment.SetDefine(TEXT("DUAL_SOURCE_COLOR_BLENDING_ENABLED"), bMaterialRequestsDualSourceBlending&& Material->IsDualBlendingEnabled(Platform) ? TEXT("1") : TEXT("0"));
+	OutEnvironment.SetDefine(TEXT("DUAL_SOURCE_COLOR_BLENDING_ENABLED"), Material->IsDualBlendingEnabled(Platform) ? TEXT("1") : TEXT("0"));
 
 	OutEnvironment.SetDefine(TEXT("STRATA_PREMULTIPLIED_ALPHA_OPACITY_OVERRIDEN"), bMaterialIsStrata && bOpacityPropertyIsUsed ? TEXT("1") : TEXT("0"));
 
