Substrate - Rename/move file(s)

diff --git a/Engine/Plugins/FX/Niagara/Shaders/Private/NiagaraDataInterfaceCollisionQuery.ush b/Engine/Plugins/FX/Niagara/Shaders/Private/NiagaraDataInterfaceCollisionQuery.ush
--- a/Engine/Plugins/FX/Niagara/Shaders/Private/NiagaraDataInterfaceCollisionQuery.ush
+++ b/Engine/Plugins/FX/Niagara/Shaders/Private/NiagaraDataInterfaceCollisionQuery.ush
@@ -3,7 +3,7 @@
 #pragma once
 
 #include "/Engine/Private/Common.ush"
-#include "/Engine/Public/StrataPublic.ush"
+#include "/Engine/Public/SubstratePublic.ush"
 
 #define NDICOLLISIONQUERY_USE_GBUFFER_NORMAL    (FEATURE_LEVEL >= FEATURE_LEVEL_SM5)
 
diff --git a/Engine/Plugins/FX/Niagara/Shaders/Private/NiagaraDataInterfaceGBufferTemplate.ush b/Engine/Plugins/FX/Niagara/Shaders/Private/NiagaraDataInterfaceGBufferTemplate.ush
--- a/Engine/Plugins/FX/Niagara/Shaders/Private/NiagaraDataInterfaceGBufferTemplate.ush
+++ b/Engine/Plugins/FX/Niagara/Shaders/Private/NiagaraDataInterfaceGBufferTemplate.ush
@@ -1,7 +1,7 @@
 // Copyright Epic Games, Inc. All Rights Reserved.
 
 #include "/Engine/Private/DeferredShadingCommon.ush"
-#include "/Engine/Public/StrataPublic.ush"
+#include "/Engine/Public/SubstratePublic.ush"
 
 #define NDIGBUFFER_HAS_GBUFFER				(FEATURE_LEVEL >= FEATURE_LEVEL_SM5)
 #define NDIGBUFFER_USE_GBUFFER_VELOCITY		(false && NDIGBUFFER_HAS_GBUFFER)
diff --git a/Engine/Plugins/FX/Niagara/Shaders/Private/NiagaraEmitterInstanceShader.usf b/Engine/Plugins/FX/Niagara/Shaders/Private/NiagaraEmitterInstanceShader.usf
--- a/Engine/Plugins/FX/Niagara/Shaders/Private/NiagaraEmitterInstanceShader.usf
+++ b/Engine/Plugins/FX/Niagara/Shaders/Private/NiagaraEmitterInstanceShader.usf
@@ -18,7 +18,7 @@ NiagaraSimulationShader.usf:
 #if GPU_SIMULATION
 	#include "/Engine/Private/Common.ush"
 	#include "/Engine/Private/DistanceField/GlobalDistanceFieldShared.ush"
-	#include "/Engine/Generated/UniformBuffers/StrataPublic.ush"
+	#include "/Engine/Public/SubstratePublic.ush"
 #else
 	const static float PI = 3.1415926535897932f;
 #endif
diff --git a/Engine/Shaders/Private/Substrate/Substrate.ush b/Engine/Shaders/Private/Substrate/Substrate.ush
--- a/Engine/Shaders/Private/Substrate/Substrate.ush
+++ b/Engine/Shaders/Private/Substrate/Substrate.ush
@@ -22,7 +22,7 @@
 #include "../SubsurfaceProfileCommon.ush"
 #include "../BurleyNormalizedSSSCommon.ush"
 #include "../GammaCorrectionCommon.ush"
-#include "/Engine/Shared/StrataDefinitions.h"
+#include "/Engine/Shared/SubstrateDefinitions.h"
 #include "/Engine/Private/Substrate/SubstrateStatisticalOperators.ush"
 #include "/Engine/Private/Substrate/SubstrateSubsurface.ush"
 
@@ -292,7 +292,7 @@ struct FStrataBSDF
 #define HEADER_MASK_HASSUBSURFACE					(1u << 7)
 #define HEADER_MASK_BIT_COUNT						8
 
-// For HEADER_BIT_BSDF_COUNT and  HEADER_BIT_SHAREDLOCALBASES_X, please refer to values in StrataDefinitions.h
+// For HEADER_BIT_BSDF_COUNT and  HEADER_BIT_SHAREDLOCALBASES_X, please refer to values in SubstrateDefinitions.h
 #define HEADER_BIT_COUNT_STATE						HEADER_MASK_BIT_COUNT
 #define HEADER_BIT_COUNT_IRRADIANCE_AO				8
 // The remaining header bits are only available for the complex path.
diff --git a/Engine/Shaders/Private/Substrate/SubstrateDBuffer.usf b/Engine/Shaders/Private/Substrate/SubstrateDBuffer.usf
--- a/Engine/Shaders/Private/Substrate/SubstrateDBuffer.usf
+++ b/Engine/Shaders/Private/Substrate/SubstrateDBuffer.usf
@@ -1,7 +1,7 @@
 // Copyright Epic Games, Inc. All Rights Reserved.
 
 #include "/Engine/Private/Common.ush"
-#include "/Engine/Shared/StrataDefinitions.h"   
+#include "/Engine/Shared/SubstrateDefinitions.h"   
 
 #if PERMUTATION_TILETYPE == 0
 #define STRATA_FASTPATH 1
diff --git a/Engine/Shaders/Private/Substrate/SubstrateTile.usf b/Engine/Shaders/Private/Substrate/SubstrateTile.usf
--- a/Engine/Shaders/Private/Substrate/SubstrateTile.usf
+++ b/Engine/Shaders/Private/Substrate/SubstrateTile.usf
@@ -1,7 +1,7 @@
 // Copyright Epic Games, Inc. All Rights Reserved.
 
 #include "/Engine/Private/Common.ush"
-#include "/Engine/Shared/StrataDefinitions.h"   
+#include "/Engine/Shared/SubstrateDefinitions.h"   
 #include "SubstrateTile.ush"
 
 
diff --git a/Engine/Shaders/Private/Substrate/SubstrateTile.ush b/Engine/Shaders/Private/Substrate/SubstrateTile.ush
--- a/Engine/Shaders/Private/Substrate/SubstrateTile.ush
+++ b/Engine/Shaders/Private/Substrate/SubstrateTile.ush
@@ -3,7 +3,7 @@
 #pragma once
 
 #include "/Engine/Private/Common.ush"
-#include "/Engine/Shared/StrataDefinitions.h"   
+#include "/Engine/Shared/SubstrateDefinitions.h"   
 #include "Substrate.ush"
 
 #if STRATA_ENABLED
diff --git a/Engine/Shaders/Public/StrataPublic.ush b/Engine/Shaders/Public/StrataPublic.ush
deleted file mode 100644
diff --git a/Engine/Shaders/Public/SubstratePublic.ush b/Engine/Shaders/Public/SubstratePublic.ush
new file mode 100644
--- /dev/null
+++ b/Engine/Shaders/Public/SubstratePublic.ush
@@ -0,0 +1,21 @@
+// Copyright Epic Games, Inc. All Rights Reserved.
+
+#pragma once
+
+/*=============================================================================
+	SubstratePublic.ush: Public facing API for external users of Strata.
+=============================================================================*/
+
+#include "/Engine/Private/Substrate/Substrate.ush"
+
+#if STRATA_ENABLED
+
+/** Public facing functions allowing external users access to some Strata data. */
+float3 StrataPublic_GetWorldNormal(uint2 PixelCoord)
+{
+  return StrataUnpackTopLayerData(StrataPublic.TopLayerTexture.Load(uint3(PixelCoord, 0))).WorldNormal;
+}
+
+//TODO: Implement other accessors that were previously available via SceneTextures.
+
+#endif
\ No newline at end of file
diff --git a/Engine/Shaders/Shared/StrataDefinitions.h b/Engine/Shaders/Shared/StrataDefinitions.h
deleted file mode 100644
diff --git a/Engine/Shaders/Shared/SubstrateDefinitions.h b/Engine/Shaders/Shared/SubstrateDefinitions.h
new file mode 100644
--- /dev/null
+++ b/Engine/Shaders/Shared/SubstrateDefinitions.h
@@ -0,0 +1,83 @@
+// Copyright Epic Games, Inc. All Rights Reserved.
+
+/*================================================================================================
+	StrataDefinitions.ush: used in ray tracing shaders and C++ code to define common constants
+	!!! Changing this file requires recompilation of the engine !!!
+=================================================================================================*/
+
+#pragma once
+
+// Change this to force recompilation of all strata dependent shaders (use https://www.random.org/cgi-bin/randbyte?nbytes=4&format=h)
+#define STRATA_SHADER_VERSION 0x72C8A053
+
+// BSDF offsets are packed into 32bits, each entry using STRATA_BSDF_OFFSET_BIT_COUNT bits
+#define STRATA_MAX_BSDF_COUNT_FOR_BDSFOFFSET	8u
+#define STRATA_BSDF_OFFSET_BIT_COUNT			4u
+#define STRATA_BSDF_OFFSET_BIT_MASK				0xF
+
+// We can only ever use STRATA_MAX_BSDF_COUNT_FOR_BDSFOFFSET BSDFs for Lumen, so we use that as a global BSDF count limit today.
+#define STRATA_MAX_BSDF_COUNT				STRATA_MAX_BSDF_COUNT_FOR_BDSFOFFSET
+#define STRATA_MAX_OPERATOR_COUNT			15
+
+// It this is changed, STATE_BIT_COUNT_SHAREDLOCALBASESID and HEADER_BIT_COUNT_SHAREDLOCALBASES_COUNT also needs to be updated
+#define STRATA_MAX_SHAREDLOCALBASES_REGISTERS		4
+
+#define STRATA_PACKED_SHAREDLOCALBASIS_STRIDE_BYTES	4
+
+// As of today, a fully simplified material is a slab with all features allowed. It can thus be complex if anisotropy is enabled and in this case eats up to 32bytes.
+// STRATA_TODO: fully simplified should remove all features but fuzz maybe. 
+#define STRATA_FULLY_SIMPLIFIED_NUM_UINTS	(32/4)
+
+#define STRATA_BSDF_TYPE_SLAB				0
+#define STRATA_BSDF_TYPE_VOLUMETRICFOGCLOUD	1
+#define STRATA_BSDF_TYPE_UNLIT				2
+#define STRATA_BSDF_TYPE_HAIR				3
+#define STRATA_BSDF_TYPE_SINGLELAYERWATER	4
+#define STRATA_BSDF_TYPE_EYE				5
+// When more than 7 BSDF must exists, please update STATE_BIT_COUNT_BSDF and FStrataClassification.ShadingModels packing in Strata.ush
+
+// The size of strata material classification tiles on screen
+#define STRATA_TILE_SIZE					8
+#define STRATA_TILE_SIZE_DIV_AS_SHIFT	    3
+
+// The default thickness of a layer is considered to be 0.01 centimeter = 0.1 millimeter
+#define STRATA_LAYER_DEFAULT_THICKNESS_CM	0.01f
+
+// Min Fuzz Roughness to avoid numerical issue
+#define STRATA_MIN_FUZZ_ROUGHNESS 0.02f
+
+#define STRATA_BASE_PASS_MRT_OUTPUT_COUNT	3
+
+#define STRATA_SSS_DATA_UINT_COUNT		2
+
+#define STRATA_OPERATOR_WEIGHT			0
+#define STRATA_OPERATOR_VERTICAL		1
+#define STRATA_OPERATOR_HORIZONTAL		2
+#define STRATA_OPERATOR_ADD				3
+#define STRATA_OPERATOR_BSDF			4
+#define STRATA_OPERATOR_BSDF_LEGACY		5
+
+// This must map directly to EStrataTileMaterialType
+#define STRATA_TILE_TYPE_SIMPLE						0
+#define STRATA_TILE_TYPE_SINGLE						1
+#define STRATA_TILE_TYPE_COMPLEX					2
+#define STRATA_TILE_TYPE_COMPLEX_SPECIAL			3
+#define STRATA_TILE_TYPE_ROUGH_REFRACT				4
+#define STRATA_TILE_TYPE_ROUGH_REFRACT_SSS_WITHOUT	5
+#define STRATA_TILE_TYPE_DECAL_SIMPLE				6
+#define STRATA_TILE_TYPE_DECAL_SINGLE				7
+#define STRATA_TILE_TYPE_DECAL_COMPLEX				8
+#define STRATA_TILE_TYPE_COUNT						9
+
+// sizeof(FRHIDrawIndirectParameters) = 4 uints = 16 bytes
+#define GetStrataTileTypeDrawIndirectArgOffset_Byte(x)  (x * 16)
+#define GetStrataTileTypeDrawIndirectArgOffset_DWord(x) (x * 4)
+
+// sizeof(FRHIDispatchIndirectParameters) can vary per-platform
+#ifdef __cplusplus
+	#define GetStrataTileTypeDispatchIndirectArgOffset_Byte(x)  (x * sizeof(FRHIDispatchIndirectParameters))
+	#define GetStrataTileTypeDispatchIndirectArgOffset_DWord(x) (x * sizeof(FRHIDispatchIndirectParameters) / sizeof(uint32))
+#else
+	#define GetStrataTileTypeDispatchIndirectArgOffset_Byte(x)  (x * DISPATCH_INDIRECT_UINT_COUNT * 4)
+	#define GetStrataTileTypeDispatchIndirectArgOffset_DWord(x) (x * DISPATCH_INDIRECT_UINT_COUNT)
+#endif
\ No newline at end of file
diff --git a/Engine/Source/Editor/CommonMenuExtensions/Private/StrataVisualizationMenuCommands.cpp b/Engine/Source/Editor/CommonMenuExtensions/Private/StrataVisualizationMenuCommands.cpp
--- a/Engine/Source/Editor/CommonMenuExtensions/Private/StrataVisualizationMenuCommands.cpp
+++ b/Engine/Source/Editor/CommonMenuExtensions/Private/StrataVisualizationMenuCommands.cpp
@@ -11,7 +11,7 @@
 #include "Framework/MultiBox/MultiBoxBuilder.h"
 #include "Internationalization/Internationalization.h"
 #include "Internationalization/Text.h"
-#include "StrataVisualizationData.h"
+#include "SubstrateVisualizationData.h"
 #include "Misc/AssertionMacros.h"
 #include "Styling/AppStyle.h"
 #include "UObject/UnrealNames.h"
diff --git a/Engine/Source/Editor/CommonMenuExtensions/Public/StrataVisualizationMenuCommands.h b/Engine/Source/Editor/CommonMenuExtensions/Public/StrataVisualizationMenuCommands.h
--- a/Engine/Source/Editor/CommonMenuExtensions/Public/StrataVisualizationMenuCommands.h
+++ b/Engine/Source/Editor/CommonMenuExtensions/Public/StrataVisualizationMenuCommands.h
@@ -9,7 +9,7 @@
 #include "HAL/Platform.h"
 #include "Templates/SharedPointer.h"
 #include "UObject/NameTypes.h"
-#include "StrataVisualizationData.h"
+#include "SubstrateVisualizationData.h"
 
 class FEditorViewportClient;
 class FMenuBuilder;
diff --git a/Engine/Source/Editor/GraphEditor/Private/MaterialNodes/SGraphNodeMaterialBase.cpp b/Engine/Source/Editor/GraphEditor/Private/MaterialNodes/SGraphNodeMaterialBase.cpp
--- a/Engine/Source/Editor/GraphEditor/Private/MaterialNodes/SGraphNodeMaterialBase.cpp
+++ b/Engine/Source/Editor/GraphEditor/Private/MaterialNodes/SGraphNodeMaterialBase.cpp
@@ -25,7 +25,7 @@
 #include "Materials/MaterialExpression.h"
 #include "Materials/MaterialExpressionBreakMaterialAttributes.h"
 #include "Materials/MaterialExpressionMakeMaterialAttributes.h"
-#include "Materials/MaterialExpressionStrata.h"
+#include "Materials/MaterialExpressionSubstrate.h"
 #include "Materials/MaterialFunction.h"
 #include "Math/Color.h"
 #include "Math/IntPoint.h"
diff --git a/Engine/Source/Editor/GraphEditor/Private/MaterialNodes/SGraphSubstrateMaterial.cpp b/Engine/Source/Editor/GraphEditor/Private/MaterialNodes/SGraphSubstrateMaterial.cpp
--- a/Engine/Source/Editor/GraphEditor/Private/MaterialNodes/SGraphSubstrateMaterial.cpp
+++ b/Engine/Source/Editor/GraphEditor/Private/MaterialNodes/SGraphSubstrateMaterial.cpp
@@ -7,11 +7,11 @@
 #include "MaterialGraph/MaterialGraphNode.h"
 #include "MaterialGraph/MaterialGraphNode_Root.h"
 #include "Materials/MaterialExpression.h"
-#include "Materials/MaterialExpressionStrata.h"
+#include "Materials/MaterialExpressionSubstrate.h"
 #include "SGraphPin.h"
 #include "Math/Color.h"
 #include "Styling/StyleColors.h"
-#include "StrataDefinitions.h"
+#include "SubstrateDefinitions.h"
 #include "Widgets/SBoxPanel.h"
 
 enum class ESubstrateWidgetOutputType : uint8
diff --git a/Engine/Source/Editor/LevelEditor/Private/LevelViewportActions.cpp b/Engine/Source/Editor/LevelEditor/Private/LevelViewportActions.cpp
--- a/Engine/Source/Editor/LevelEditor/Private/LevelViewportActions.cpp
+++ b/Engine/Source/Editor/LevelEditor/Private/LevelViewportActions.cpp
@@ -11,7 +11,7 @@
 #include "BufferVisualizationData.h"
 #include "NaniteVisualizationData.h"
 #include "LumenVisualizationData.h"
-#include "StrataVisualizationData.h"
+#include "SubstrateVisualizationData.h"
 #include "GroomVisualizationData.h"
 #include "VirtualShadowMapVisualizationData.h"
 #include "Bookmarks/BookmarkUI.h"
diff --git a/Engine/Source/Editor/LevelEditor/Private/SLevelViewport.cpp b/Engine/Source/Editor/LevelEditor/Private/SLevelViewport.cpp
--- a/Engine/Source/Editor/LevelEditor/Private/SLevelViewport.cpp
+++ b/Engine/Source/Editor/LevelEditor/Private/SLevelViewport.cpp
@@ -64,7 +64,7 @@
 #include "BufferVisualizationData.h"
 #include "NaniteVisualizationData.h"
 #include "LumenVisualizationData.h"
-#include "StrataVisualizationData.h"
+#include "SubstrateVisualizationData.h"
 #include "GroomVisualizationData.h"
 #include "VirtualShadowMapVisualizationData.h"
 #include "Framework/Notifications/NotificationManager.h"
diff --git a/Engine/Source/Editor/LevelEditor/Private/SLevelViewportToolBar.cpp b/Engine/Source/Editor/LevelEditor/Private/SLevelViewportToolBar.cpp
--- a/Engine/Source/Editor/LevelEditor/Private/SLevelViewportToolBar.cpp
+++ b/Engine/Source/Editor/LevelEditor/Private/SLevelViewportToolBar.cpp
@@ -35,7 +35,7 @@
 #include "BufferVisualizationData.h"
 #include "NaniteVisualizationData.h"
 #include "LumenVisualizationData.h"
-#include "StrataVisualizationData.h"
+#include "SubstrateVisualizationData.h"
 #include "GroomVisualizationData.h"
 #include "VirtualShadowMapVisualizationData.h"
 #include "FoliageType.h"
diff --git a/Engine/Source/Editor/MaterialEditor/Private/MaterialEditor.cpp b/Engine/Source/Editor/MaterialEditor/Private/MaterialEditor.cpp
--- a/Engine/Source/Editor/MaterialEditor/Private/MaterialEditor.cpp
+++ b/Engine/Source/Editor/MaterialEditor/Private/MaterialEditor.cpp
@@ -148,7 +148,7 @@
 #include "Materials/MaterialExpressionMaterialLayerOutput.h"
 #include "Materials/MaterialExpressionNamedReroute.h"
 #include "Materials/MaterialExpressionReroute.h"
-#include "Materials/MaterialExpressionStrata.h"
+#include "Materials/MaterialExpressionSubstrate.h"
 
 #include "MaterialStats.h"
 #include "MaterialEditorTabs.h"
diff --git a/Engine/Source/Editor/MaterialEditor/Private/MaterialExpressionClasses.cpp b/Engine/Source/Editor/MaterialEditor/Private/MaterialExpressionClasses.cpp
--- a/Engine/Source/Editor/MaterialEditor/Private/MaterialExpressionClasses.cpp
+++ b/Engine/Source/Editor/MaterialEditor/Private/MaterialExpressionClasses.cpp
@@ -25,7 +25,7 @@
 #include "Materials/MaterialExpressionForLoop.h"
 #include "Materials/MaterialExpressionGetLocal.h"
 #include "Materials/MaterialExpressionSetLocal.h"
-#include "Materials/MaterialExpressionStrata.h"
+#include "Materials/MaterialExpressionSubstrate.h"
 
 ///////////////////////////////////////////////////////////////////////////////////////////////////////////////////////
 //
diff --git a/Engine/Source/Editor/MaterialEditor/Private/SMaterialEditorStrataWidget.cpp b/Engine/Source/Editor/MaterialEditor/Private/SMaterialEditorStrataWidget.cpp
--- a/Engine/Source/Editor/MaterialEditor/Private/SMaterialEditorStrataWidget.cpp
+++ b/Engine/Source/Editor/MaterialEditor/Private/SMaterialEditorStrataWidget.cpp
@@ -5,7 +5,7 @@
 #include "EditorWidgetsModule.h"
 #include "AssetRegistry/AssetRegistryModule.h"
 #include "MaterialEditor.h"
-#include "StrataDefinitions.h"
+#include "SubstrateDefinitions.h"
 #include <functional>
 
 #include "Widgets/Input/SButton.h"
diff --git a/Engine/Source/Editor/UnrealEd/Private/EditorViewportClient.cpp b/Engine/Source/Editor/UnrealEd/Private/EditorViewportClient.cpp
--- a/Engine/Source/Editor/UnrealEd/Private/EditorViewportClient.cpp
+++ b/Engine/Source/Editor/UnrealEd/Private/EditorViewportClient.cpp
@@ -56,7 +56,7 @@
 #include "BufferVisualizationData.h"
 #include "NaniteVisualizationData.h"
 #include "LumenVisualizationData.h"
-#include "StrataVisualizationData.h"
+#include "SubstrateVisualizationData.h"
 #include "GroomVisualizationData.h"
 #include "VirtualShadowMapVisualizationData.h"
 #include "GPUSkinCacheVisualizationData.h"
diff --git a/Engine/Source/Editor/UnrealEd/Private/MaterialGraphNode.cpp b/Engine/Source/Editor/UnrealEd/Private/MaterialGraphNode.cpp
--- a/Engine/Source/Editor/UnrealEd/Private/MaterialGraphNode.cpp
+++ b/Engine/Source/Editor/UnrealEd/Private/MaterialGraphNode.cpp
@@ -44,7 +44,7 @@
 #include "Materials/MaterialExpressionReroute.h"
 #include "Materials/MaterialExpressionCurveAtlasRowParameter.h"
 #include "Materials/MaterialExpressionExecBegin.h"
-#include "Materials/MaterialExpressionStrata.h"
+#include "Materials/MaterialExpressionSubstrate.h"
 #include "Materials/MaterialFunction.h"
 #include "MaterialEditorUtilities.h"
 #include "MaterialEditorActions.h"
diff --git a/Engine/Source/Runtime/Engine/Classes/Materials/MaterialExpressionStrata.h b/Engine/Source/Runtime/Engine/Classes/Materials/MaterialExpressionStrata.h
deleted file mode 100644
diff --git a/Engine/Source/Runtime/Engine/Classes/Materials/MaterialExpressionSubstrate.h b/Engine/Source/Runtime/Engine/Classes/Materials/MaterialExpressionSubstrate.h
new file mode 100644
--- /dev/null
+++ b/Engine/Source/Runtime/Engine/Classes/Materials/MaterialExpressionSubstrate.h
@@ -0,0 +1,1130 @@
+// Copyright Epic Games, Inc. All Rights Reserved.
+
+#pragma once
+
+#include "CoreMinimal.h"
+#include "Engine/EngineTypes.h"
+#include "Materials/MaterialExpression.h"
+#include "MaterialExpressionSubstrate.generated.h"
+
+class FMaterialCompiler;
+
+/**
+ * Compile a special blend function for Substrate when blending material attribute
+ *
+ * @param Compiler				The compiler to add code to
+ * @param Foreground			Entry A, has a bigger impact when Alpha is close to 0
+ * @param Background			Entry B, has a bigger impact when Alpha is close to 1
+ * @param Alpha					Blend factor [0..1], when 0
+ * @return						Index to a new code chunk
+ */
+extern int32 CompileStrataBlendFunction(FMaterialCompiler* Compiler, const int32 A, const int32 B, const int32 Alpha);
+
+
+///////////////////////////////////////////////////////////////////////////////
+// BSDF nodes
+
+// UMaterialExpressionStrataBSDF can only be used for Substrate nodes ouputing StrataData that would need a preview,
+UCLASS(MinimalAPI, collapsecategories, hidecategories = Object, Abstract, DisplayName = "Substrate Expression")
+class UMaterialExpressionStrataBSDF : public UMaterialExpression
+{
+	GENERATED_UCLASS_BODY()
+
+#if WITH_EDITOR
+	virtual int32 CompilePreview(class FMaterialCompiler* Compiler, int32 OutputIndex) override;
+#endif
+
+	float DielectricSpecularToF0(float SpecularIn)
+	{
+		return 0.08f * SpecularIn;
+	};
+};
+
+
+UCLASS(MinimalAPI, collapsecategories, hidecategories = Object, DisplayName = "Substrate Shading Models")
+class UMaterialExpressionStrataLegacyConversion : public UMaterialExpressionStrataBSDF
+{
+	GENERATED_UCLASS_BODY()
+
+	/**
+	 * Defines the overall color of the Material. (type = float3, unit = unitless, defaults to 0.18)
+	 */
+	UPROPERTY()
+	FExpressionInput BaseColor;
+
+	/**
+	 * Controls how \"metal-like\" your surface looks like. 0 means dielectric, 1 means conductor (type = float, unit = unitless, defaults to 0)
+	 */
+	UPROPERTY()
+	FExpressionInput Metallic;
+
+	/**
+	 * Used to scale the current amount of specularity on non-metallic surfaces and is a value between 0 and 1 (type = float, unit = unitless, defaults to plastic 0.5)
+	 */
+	UPROPERTY()
+	FExpressionInput Specular;
+	
+	/**
+	 * Controls how rough the Material is. Roughness of 0 (smooth) is a mirror reflection and 1 (rough) is completely matte or diffuse. When using anisotropy, it is the roughness used along the Tangent axis. (type = float, unit = unitless, defaults to 0.5)
+	 */
+	UPROPERTY()
+	FExpressionInput Roughness;
+		
+	/**
+	 * Controls the anisotropy factor of the roughness. Positive value elongates the specular lobe along the Tangent vector, Negative value elongates the specular lobe along the perpendicular of the Tangent. (type = float, unit = unitless).
+	 */
+	UPROPERTY()
+	FExpressionInput Anisotropy;
+
+	/**
+	 * Emissive color on top of the surface (type = float3, unit = luminance, default = 0)
+	 */
+	UPROPERTY()
+	FExpressionInput EmissiveColor;
+
+	/**
+	 * Take the surface normal as input. The normal is considered tangent or world space according to the space properties on the main material node. (type = float3, unit = unitless, defaults to vertex normal)
+	 */
+	UPROPERTY()
+	FExpressionInput Normal;
+
+	/**
+	* Take a surface tangent as input. The tangent is considered tangent or world space according to the space properties on the main material node. (type = float3, unit = unitless, defaults to vertex tangent)
+	*/
+	UPROPERTY()
+	FExpressionInput Tangent;
+
+	/**
+	 * Scale the mean free path radius of the SSS profile according to a value between 0 and 1. Always used, when a subsurface profile is provided or not. (type = float, unitless, defaults to 1)
+	 */
+	UPROPERTY()
+	FExpressionInput SubSurfaceColor;
+
+	/**
+	 * Coverage of the clear coat layer. (type = float, unit = unitless, defaults to 0.0)
+	 */
+	UPROPERTY()
+	FExpressionInput ClearCoat;
+
+	/**
+	 * Roughness of the top clear coat layer. (type = float, unit = unitless, defaults to 0.0)
+	 */
+	UPROPERTY()
+	FExpressionInput ClearCoatRoughness;
+
+	/**
+	 * Opacity of the material
+	 */
+	UPROPERTY()
+	FExpressionInput Opacity;
+
+	/**
+	 * The amount of transmitted light from the back side of the surface to the front side of the surface (type = float3, unit = unitless, defaults to 1)
+	 */
+	UPROPERTY()
+	FExpressionInput TransmittanceColor;
+
+		/**
+	* The single scattering Albedo defining the overall color of the Material (type = float3, unit = unitless, default = 0)
+	 */
+	UPROPERTY()
+	FExpressionInput WaterScatteringCoefficients;
+
+	/**
+	 * The rate at which light is absorbed or out-scattered by the medium. Mean Free Path = 1 / Extinction. (type = float3, unit = 1/cm, default = 0)
+	 */
+	UPROPERTY()
+	FExpressionInput WaterAbsorptionCoefficients;
+
+	/**
+	 * Anisotropy of the volume with values lower than 0 representing back-scattering, equal 0 representing isotropic scattering and greater than 0 representing forward scattering. (type = float, unit = unitless, defaults to 0)
+	 */
+	UPROPERTY()
+	FExpressionInput WaterPhaseG;
+
+	/**
+	 * A scale to apply on the scene color behind the water surface. It can be used to approximate caustics for instance. (type = float3, unit = unitless, defaults to 1)
+	 */
+	UPROPERTY()
+	FExpressionInput ColorScaleBehindWater;
+
+	/**
+	 * Take the bottom clear coat surface normal as input. The normal is considered tangent or world space according to the space properties on the main material node. (type = float3, unit = unitless, defaults to vertex normal)
+	 */
+	UPROPERTY()
+	FExpressionInput ClearCoatNormal;
+
+	/**
+	 * Take the tangent output node as input. The tangent is considered tangent or world space according to the space properties on the main material node. (type = float3, unit = unitless, defaults to vertex tangent)
+	 */
+	UPROPERTY()
+	FExpressionInput CustomTangent;
+
+	/**
+	 * Shading models
+	 */
+	UPROPERTY()
+	FShadingModelMaterialInput ShadingModel;
+	
+	/** SubsurfaceProfile, for Screen Space Subsurface Scattering. The profile needs to be set up on both the Substrate diffuse node, and the material node at the moment. */
+	UPROPERTY(EditAnywhere, BlueprintReadOnly, Category = Material, meta = (DisplayName = "Subsurface Profile"))
+	TObjectPtr<class USubsurfaceProfile> SubsurfaceProfile;
+
+	UPROPERTY(EditAnywhere, Category = ShadingModel, meta = (ShowAsInputPin = "Primary", DisplayName = "Single Shading Model"))
+	TEnumAsByte<enum EMaterialShadingModel> ShadingModelOverride = MSM_DefaultLit;
+
+	//~ Begin UMaterialExpression Interface
+#if WITH_EDITOR
+	virtual int32 Compile(class FMaterialCompiler* Compiler, int32 OutputIndex) override;
+	virtual void GetCaption(TArray<FString>& OutCaptions) const override;
+	virtual uint32 GetOutputType(int32 OutputIndex) override;
+	virtual uint32 GetInputType(int32 InputIndex) override;
+	virtual bool IsResultStrataMaterial(int32 OutputIndex) override;
+	virtual void GatherStrataMaterialInfo(FStrataMaterialInfo& StrataMaterialInfo, int32 OutputIndex) override;
+	virtual FStrataOperator* StrataGenerateMaterialTopologyTree(class FMaterialCompiler* Compiler, class UMaterialExpression* Parent, int32 OutputIndex) override;
+	virtual FName GetInputName(int32 InputIndex) const override;
+	virtual void GetConnectorToolTip(int32 InputIndex, int32 OutputIndex, TArray<FString>& OutToolTip) override;
+	virtual void PostEditChangeProperty(FPropertyChangedEvent& PropertyChangedEvent) override;
+
+	bool HasSSS() const;
+	bool HasAnisotropy() const;
+
+#endif
+	//~ End UMaterialExpression Interface
+};
+
+UCLASS(MinimalAPI, collapsecategories, hidecategories = Object, DisplayName = "Substrate Slab")
+class UMaterialExpressionStrataSlabBSDF : public UMaterialExpressionStrataBSDF
+{
+	GENERATED_UCLASS_BODY()
+
+	/**
+	 * Defines the diffused albedo, the percentage of light reflected as diffuse from the surface. (type = float3, unit = unitless, defaults to 0.18)
+	 */
+	UPROPERTY()
+	FExpressionInput DiffuseAlbedo;
+
+	/**
+	 * Defines the color and brightness of the specular highlight where the surface is facing the camera. Each specular contribution will fade to black as F0 drops below 0.02. (type = float3, unit = unitless, defaults to plastic 0.04)
+	 */
+	UPROPERTY()
+	FExpressionInput F0;
+
+	/**
+	 * Defines the color of the specular highlight where the surface normal is 90 degrees from the view direction. Only the hue and saturation are preserved, the brightness is fixed at 1.0. Fades to black as F0 drops below 0.02. (type = float3, unit = unitless, defaults to 1.0f)
+	 */
+	UPROPERTY()
+	FExpressionInput F90;
+
+	/**
+	 * Controls how rough the Material is. Roughness of 0 (smooth) is a mirror reflection and 1 (rough) is completely matte or diffuse. When using anisotropy, it is the roughness used along the Tangent axis. (type = float, unit = unitless, defaults to 0.5)
+	 */
+	UPROPERTY()
+	FExpressionInput Roughness;
+		
+	/**
+	 * Controls the anisotropy factor of the roughness. Positive value elongates the specular lobe along the Tangent vector, Negative value elongates the specular lobe along the perpendicular of the Tangent. (type = float, unit = unitless).
+	 */
+	UPROPERTY()
+	FExpressionInput Anisotropy;
+
+	/**
+	 * Take the surface normal as input. The normal is considered tangent or world space according to the space properties on the main material node. (type = float3, unit = unitless, defaults to vertex normal)
+	 */
+	UPROPERTY()
+	FExpressionInput Normal;
+
+	/**
+	* Take a surface tangent as input. The tangent is considered tangent or world space according to the space properties on the main material node. (type = float3, unit = unitless, defaults to vertex tangent)
+	*/
+	UPROPERTY()
+	FExpressionInput Tangent;
+
+	/**
+	 * Chromatic mean free path . Only used when there is not any sub-surface profile provided. (type = float3, unit = centimeters, default = 0)
+	 */
+	UPROPERTY()
+	FExpressionInput SSSMFP;
+
+	/**
+	 * Scale the mean free path radius of the SSS profile according to a value between 0 and 1. Always used, when a subsurface profile is provided or not. (type = float, unitless, defaults to 1)
+	 */
+	UPROPERTY()
+	FExpressionInput SSSMFPScale;
+
+	/**
+	 * Phase function anisotropy. Positive value elongates the phase function along the lignt direction, causing forward scattering. Negative value elongates the phase function backward to the light direction, causing backward scattering.  (type = float, unitless, defaults to 1, valid value -1..1)
+	 */
+	UPROPERTY()
+	FExpressionInput SSSPhaseAnisotropy;
+
+	/**
+	 * Emissive color on top of the surface (type = float3, unit = luminance, default = 0)
+	 */
+	UPROPERTY()
+	FExpressionInput EmissiveColor;
+
+	/**
+	 * Controls the roughness of a secondary specular lobe. Roughness of 0 (smooth) is a mirror reflection and 1 (rough) is completely matte or diffuse. Does not influence diffuse roughness. (type = float, unit = unitless, defaults to 0.5)
+	 */
+	UPROPERTY()
+	FExpressionInput SecondRoughness;
+
+	/**
+	 * The weight of the second specular lobe using SecondRoughness. The first specular using Roughness will have a weight of (1 - SecondRoughnessWeight). (type = float, unitless, default = 0)
+	 */
+	UPROPERTY()
+	FExpressionInput SecondRoughnessWeight;
+
+	/**
+	 * Controls how rough the Fuzz layer is. Roughness of 0 is smooth and 1 is rough. If FuzzRoughness Is not connected, the Roughness input will be used instead. (type = float, unit = unitless, defaults to 0.5)
+	 */
+	UPROPERTY()
+	FExpressionInput FuzzRoughness;
+
+	/**
+	 * The amount of fuzz on top of the surface used to simulate cloth-like appearance.
+	 */
+	UPROPERTY()
+	FExpressionInput FuzzAmount;
+
+	/**
+	 * The base color of the fuzz.
+	 */
+	UPROPERTY()
+	FExpressionInput FuzzColor;
+
+	/**
+	 * This represent the logarithm of the micro facet density. Only used when `r.Substrate.Glints=1`. Defaults to 0.
+	 */
+	UPROPERTY()
+	FExpressionInput GlintValue;
+
+	/**
+	 * The parameterization of the surface required to position glints on a surface. Only used when `r.Substrate.Glints=1`. Defaults to (0,0).
+	 */
+	UPROPERTY()
+	FExpressionInput GlintUV;
+
+	/** SubsurfaceProfile, for Screen Space Subsurface Scattering. The profile needs to be set up on both the Substrate diffuse node, and the material node at the moment. */
+	UPROPERTY(EditAnywhere, BlueprintReadOnly, Category = Material, meta = (DisplayName = "Subsurface Profile"))
+	TObjectPtr<class USubsurfaceProfile> SubsurfaceProfile;
+
+	/** SpecularProfile, for modulating specular appearance and simulating more complex visuals such as iridescence. */
+	UPROPERTY(EditAnywhere, BlueprintReadOnly, Category = Material, meta = (DisplayName = "Specular Profile"))
+	TObjectPtr<class USpecularProfile> SpecularProfile;
+
+	/** Whether to use light diffusion (i.e., SSS diffusion) or wrap-approximation for material with scattering behavior. This option trades quality over performance and will result into visual differences. */
+	UPROPERTY(EditAnywhere, Category = Mode, meta = (DisplayName = "Use Subsurface Diffusion"))
+	uint32 bUseSSSDiffusion : 1;
+
+	//~ Begin UMaterialExpression Interface
+#if WITH_EDITOR
+	virtual int32 Compile(class FMaterialCompiler* Compiler, int32 OutputIndex) override;
+	virtual void GetCaption(TArray<FString>& OutCaptions) const override;
+	virtual uint32 GetOutputType(int32 OutputIndex) override;
+	virtual uint32 GetInputType(int32 InputIndex) override;
+	virtual bool IsResultStrataMaterial(int32 OutputIndex) override;
+	virtual void GatherStrataMaterialInfo(FStrataMaterialInfo& StrataMaterialInfo, int32 OutputIndex) override;
+	virtual FStrataOperator* StrataGenerateMaterialTopologyTree(class FMaterialCompiler* Compiler, class UMaterialExpression* Parent, int32 OutputIndex) override;
+	virtual FName GetInputName(int32 InputIndex) const override;
+	virtual void GetConnectorToolTip(int32 InputIndex, int32 OutputIndex, TArray<FString>& OutToolTip) override;
+	virtual void GetExpressionToolTip(TArray<FString>& OutToolTip) override;
+
+	bool HasEdgeColor() const;
+	bool HasFuzz() const;
+	bool HasFuzzRoughness() const;
+	bool HasSecondRoughness() const;
+	bool HasSSS() const;
+	bool HasSSSProfile() const;
+	bool HasMFPPluggedIn() const;
+	bool HasAnisotropy() const;
+	bool HasGlint() const;
+	bool HasSpecularProfile() const;
+
+	struct FComplexity
+	{
+		bool bStrataMaterialIsComplexSpecial;
+		bool bStrataMaterialIsComplex;
+		bool bStrataMaterialIsSingle;
+		// If all the above are false, complexity is Simple
+	};
+	FComplexity GetComplexity() const;
+#endif
+	//~ End UMaterialExpression Interface
+};
+
+UCLASS(MinimalAPI, collapsecategories, hidecategories = Object, DisplayName = "Substrate Simple Clear Coat")
+class UMaterialExpressionStrataSimpleClearCoatBSDF : public UMaterialExpressionStrataBSDF
+{
+	GENERATED_UCLASS_BODY()
+
+	/**
+	 * Defines the diffused albedo, the percentage of light reflected as diffuse from the surface. (type = float3, unit = unitless, defaults to 0.18)
+	 */
+	UPROPERTY()
+	FExpressionInput DiffuseAlbedo;
+
+	/**
+	 * Defines the color and brightness of the specular highlight where the surface is facing the camera. Each specular contribution will fade to black as F0 drops below 0.02. (type = float3, unit = unitless, defaults to plastic 0.04)
+	 */
+	UPROPERTY()
+	FExpressionInput F0;
+
+	/**
+	 * Controls how rough the bottom layer of the material is. Roughness of 0 (smooth) is a mirror reflection and 1 (rough) is completely matte or diffuse. (type = float, unit = unitless, defaults to 0.5)
+	 */
+	UPROPERTY()
+	FExpressionInput Roughness;
+
+	/**
+	 * Controls the coverage of the clear coat layer: 0 means no clear coat, 1 means coat is fully visible. (type = float, unit = unitless, defaults to 0.5)
+	 */
+	UPROPERTY()
+	FExpressionInput ClearCoatCoverage;
+
+	/**
+	 * Controls how rough the top layer of the material is. Roughness of 0 (smooth) is a mirror reflection and 1 (rough) is completely matte or diffuse. (type = float, unit = unitless, defaults to 0.5)
+	 */
+	UPROPERTY()
+	FExpressionInput ClearCoatRoughness;
+
+	/**
+	 * Take the surface normal as input. The normal is considered tangent or world space according to the space properties on the main material node. (type = float3, unit = unitless, defaults to vertex normal)
+	 */
+	UPROPERTY()
+	FExpressionInput Normal;
+
+	/**
+	 * Emissive color of the medium (type = float3, unit = luminance, default = 0)
+	 */
+	UPROPERTY()
+	FExpressionInput EmissiveColor;
+
+	//~ Begin UMaterialExpression Interface
+#if WITH_EDITOR
+	virtual int32 Compile(class FMaterialCompiler* Compiler, int32 OutputIndex) override;
+	virtual void GetCaption(TArray<FString>& OutCaptions) const override;
+	virtual uint32 GetOutputType(int32 OutputIndex) override;
+	virtual uint32 GetInputType(int32 InputIndex) override;
+	virtual bool IsResultStrataMaterial(int32 OutputIndex) override;
+	virtual void GatherStrataMaterialInfo(FStrataMaterialInfo& StrataMaterialInfo, int32 OutputIndex) override;
+	virtual FStrataOperator* StrataGenerateMaterialTopologyTree(class FMaterialCompiler* Compiler, class UMaterialExpression* Parent, int32 OutputIndex) override;
+	virtual FName GetInputName(int32 InputIndex) const override;
+#endif
+	//~ End UMaterialExpression Interface
+};
+
+UCLASS(MinimalAPI, collapsecategories, hidecategories = Object, DisplayName = "Substrate Volumetric-Fog-Cloud BSDF")
+class UMaterialExpressionStrataVolumetricFogCloudBSDF : public UMaterialExpressionStrataBSDF
+{
+	GENERATED_UCLASS_BODY()
+
+	/**
+	* The single scattering Albedo defining the overall color of the Material (type = float3, unit = unitless, default = 0)
+	*/
+	UPROPERTY()
+	FExpressionInput Albedo;
+
+	/**
+	 * The rate at which light is absorbed or scattered by the medium. Mean Free Path = 1 / Extinction. (type = float3, unit = 1/m, default = 0)
+	 */
+	UPROPERTY()
+	FExpressionInput Extinction;
+
+	/**
+	 * Emissive color of the medium (type = float3, unit = luminance, default = 0)
+	 */
+	UPROPERTY()
+	FExpressionInput EmissiveColor;
+
+	/**
+	 * Ambient occlusion: 1 means no occlusion while 0 means fully occluded. (type = float, unit = unitless, default = 1)
+	 */
+	UPROPERTY()
+	FExpressionInput AmbientOcclusion;
+
+	//~ Begin UMaterialExpression Interface
+#if WITH_EDITOR
+	virtual int32 Compile(class FMaterialCompiler* Compiler, int32 OutputIndex) override;
+	virtual void GetCaption(TArray<FString>& OutCaptions) const override;
+	virtual uint32 GetOutputType(int32 OutputIndex) override;
+	virtual uint32 GetInputType(int32 InputIndex) override;
+	virtual bool IsResultStrataMaterial(int32 OutputIndex) override;
+	virtual void GatherStrataMaterialInfo(FStrataMaterialInfo& StrataMaterialInfo, int32 OutputIndex) override;
+	virtual FStrataOperator* StrataGenerateMaterialTopologyTree(class FMaterialCompiler* Compiler, class UMaterialExpression* Parent, int32 OutputIndex) override;
+#endif
+	//~ End UMaterialExpression Interface
+};
+
+UCLASS(MinimalAPI, collapsecategories, hidecategories = Object, DisplayName = "Substrate Unlit BSDF")
+class UMaterialExpressionStrataUnlitBSDF : public UMaterialExpressionStrataBSDF
+{
+	GENERATED_UCLASS_BODY()
+
+	/**
+	* Emissive color on top of the surface (type = float3, unit = Luminance, default = 0)
+	*/
+	UPROPERTY()
+	FExpressionInput EmissiveColor;
+
+	/**
+	 * The amount of transmitted light from the back side of the surface to the front side of the surface (type = float3, unit = unitless, defaults to 1)
+	 */
+	UPROPERTY()
+	FExpressionInput TransmittanceColor;
+
+	/**
+	 * The surface normal. Only used for refraction effects when `IOR` or `pixel normal offset` modes are selected.
+	 */
+	UPROPERTY()
+	FExpressionInput Normal;
+
+	//~ Begin UMaterialExpression Interface
+#if WITH_EDITOR
+	virtual int32 Compile(class FMaterialCompiler* Compiler, int32 OutputIndex) override;
+	virtual void GetCaption(TArray<FString>& OutCaptions) const override;
+	virtual uint32 GetOutputType(int32 OutputIndex) override;
+	virtual uint32 GetInputType(int32 InputIndex) override;
+	virtual bool IsResultStrataMaterial(int32 OutputIndex) override;
+	virtual void GatherStrataMaterialInfo(FStrataMaterialInfo& StrataMaterialInfo, int32 OutputIndex) override;
+	virtual FStrataOperator* StrataGenerateMaterialTopologyTree(class FMaterialCompiler* Compiler, class UMaterialExpression* Parent, int32 OutputIndex) override;
+#endif
+	//~ End UMaterialExpression Interface
+};
+
+UCLASS(MinimalAPI, collapsecategories, hidecategories = Object, DisplayName = "Substrate Hair BSDF")
+class UMaterialExpressionStrataHairBSDF : public UMaterialExpressionStrataBSDF
+{
+	GENERATED_UCLASS_BODY()
+		
+	/**
+	 * Hair fiber base color resulting from single and multiple scattering combined. (type = float3, unit = unitless, defaults to black)
+	 */
+	UPROPERTY()
+	FExpressionInput BaseColor;
+	
+	/**
+	 * Amount of light scattering, only available for non-HairStrand rendering (type = float, unit = unitless, defaults to 0.0)
+	 */
+	UPROPERTY()
+	FExpressionInput Scatter;
+		
+	/**
+	 * Specular (type = float, unit = unitless, defaults to 0.5)
+	 */
+	UPROPERTY()
+	FExpressionInput Specular;
+		
+	/**
+	 * Controls how rough the Material is. Roughness of 0 (smooth) is a mirror reflection and 1 (rough) is completely matte or diffuse (type = float, unit = unitless, defaults to 0.5)
+	 */
+	UPROPERTY()
+	FExpressionInput Roughness;
+
+	/**
+	 * How much light contributs when lighting hairs from the back side opposite from the view, only available for HairStrand rendering (type = float3, unit = unitless, defaults to 0.0)
+	 */
+	UPROPERTY()
+	FExpressionInput Backlit;
+
+	/**
+	 * Tangent (type = float3, unit = unitless, defaults to +X vector)
+	 */
+	UPROPERTY()
+	FExpressionInput Tangent;
+
+	/**
+	 * Emissive color on top of the surface (type = float3, unit = luminance, defaults to 0.0)
+	 */
+	UPROPERTY()
+	FExpressionInput EmissiveColor;
+
+	//~ Begin UMaterialExpression Interface
+#if WITH_EDITOR
+	virtual int32 Compile(class FMaterialCompiler* Compiler, int32 OutputIndex) override;
+	virtual void GetCaption(TArray<FString>& OutCaptions) const override;
+	virtual uint32 GetOutputType(int32 OutputIndex) override;
+	virtual uint32 GetInputType(int32 InputIndex) override;
+	virtual bool IsResultStrataMaterial(int32 OutputIndex) override;
+	virtual void GatherStrataMaterialInfo(FStrataMaterialInfo& StrataMaterialInfo, int32 OutputIndex) override;
+	virtual FStrataOperator* StrataGenerateMaterialTopologyTree(class FMaterialCompiler* Compiler, class UMaterialExpression* Parent, int32 OutputIndex) override;
+#endif
+	//~ End UMaterialExpression Interface
+};
+
+UCLASS(MinimalAPI, collapsecategories, hidecategories = Object, DisplayName = "Substrate Eye BSDF")
+class UMaterialExpressionStrataEyeBSDF : public UMaterialExpressionStrataBSDF
+{
+	GENERATED_UCLASS_BODY()
+		
+	/**
+	 * Hair fiber base color resulting from single and multiple scattering combined. (type = float3, unit = unitless, defaults to black)
+	 */
+	UPROPERTY()
+	FExpressionInput DiffuseColor;
+		
+	/**
+	 * Controls how rough the Material is. Roughness of 0 (smooth) is a mirror reflection and 1 (rough) is completely matte or diffuse (type = float, unit = unitless, defaults to 0.5)
+	 */
+	UPROPERTY()
+	FExpressionInput Roughness;
+
+	/**
+	 * Normal of the sclera and cornea (type = float3, unit = unitless, defaults to +X vector)
+	 */
+	UPROPERTY()
+	FExpressionInput CorneaNormal;
+
+	/**
+	 * Normal of the iris (type = float3, unit = unitless, defaults to +X vector)
+	 */
+	UPROPERTY()
+	FExpressionInput IrisNormal;
+
+	/**
+	 * Normal of the iris plane (type = float3, unit = unitless, defaults to +X vector)
+	 */
+	UPROPERTY()
+	FExpressionInput IrisPlaneNormal;
+
+	/**
+	 * Mask defining the iris surface (type = float, unit = unitless, defaults to 0.0)
+	 */
+	UPROPERTY()
+	FExpressionInput IrisMask;
+
+	/**
+	 * Distance from the center of the iris (type = float, unit = unitless, defaults to 0.0)
+	 */
+	UPROPERTY()
+	FExpressionInput IrisDistance;
+
+	/**
+	 * Emissive color on top of the surface (type = float3, unit = luminance, defaults to 0.0)
+	 */
+	UPROPERTY()
+	FExpressionInput EmissiveColor;
+
+	/** SubsurfaceProfile, for Subsurface Scattering diffusion. */
+	UPROPERTY(EditAnywhere, BlueprintReadOnly, Category = Material, meta = (DisplayName = "Subsurface Profile"))
+	TObjectPtr<class USubsurfaceProfile> SubsurfaceProfile;
+
+	//~ Begin UMaterialExpression Interface
+#if WITH_EDITOR
+	virtual int32 Compile(class FMaterialCompiler* Compiler, int32 OutputIndex) override;
+	virtual void GetCaption(TArray<FString>& OutCaptions) const override;
+	virtual uint32 GetOutputType(int32 OutputIndex) override;
+	virtual uint32 GetInputType(int32 InputIndex) override;
+	virtual bool IsResultStrataMaterial(int32 OutputIndex) override;
+	virtual void GatherStrataMaterialInfo(FStrataMaterialInfo& StrataMaterialInfo, int32 OutputIndex) override;
+	virtual FStrataOperator* StrataGenerateMaterialTopologyTree(class FMaterialCompiler* Compiler, class UMaterialExpression* Parent, int32 OutputIndex) override;
+#endif
+	//~ End UMaterialExpression Interface
+};
+
+UCLASS(MinimalAPI, collapsecategories, hidecategories = Object, DisplayName = "Substrate Single Layer Water BSDF")
+class UMaterialExpressionStrataSingleLayerWaterBSDF : public UMaterialExpressionStrataBSDF
+{
+	GENERATED_UCLASS_BODY()
+
+	/**
+	 * Surface base color. (type = float3, unit = unitless, defaults to black)
+	 */
+	UPROPERTY()
+	FExpressionInput BaseColor;
+
+	/**
+	 * Whether the surface represents a dielectric (such as plastic) or a conductor (such as metal). (type = float, unit = unitless, defaults to 0 = dielectric)
+	 */
+	UPROPERTY()
+	FExpressionInput Metallic;
+
+	/**
+	 * Specular amount (type = float, unit = unitless, defaults to 0.5)
+	 */
+	UPROPERTY()
+	FExpressionInput Specular;
+
+	/**
+	 * Controls how rough the Material is. Roughness of 0 (smooth) is a mirror reflection and 1 (rough) is completely matte or diffuse (type = float, unit = unitless, defaults to 0.5)
+	 */
+	UPROPERTY()
+	FExpressionInput Roughness;
+
+	/**
+	 * The normal of the surface (type = float3, unit = unitless, defaults to +Z vector)
+	 */
+	UPROPERTY()
+	FExpressionInput Normal;
+
+	/**
+	 * Emissive color on top of the surface (type = float3, unit = luminance, defaults to 0.0)
+	 */
+	UPROPERTY()
+	FExpressionInput EmissiveColor;
+
+	/**
+	 * Opacity of the material layered on top of the water (type = float3, unit = unitless, defaults to 0.0)
+	 */
+	UPROPERTY()
+	FExpressionInput TopMaterialOpacity;
+
+	/**
+	* The single scattering Albedo defining the overall color of the Material (type = float3, unit = unitless, default = 0)
+	 */
+	UPROPERTY()
+	FExpressionInput WaterAlbedo;
+
+	/**
+	 * The rate at which light is absorbed or out-scattered by the medium. Mean Free Path = 1 / Extinction. (type = float3, unit = 1/cm, default = 0)
+	 */
+	UPROPERTY()
+	FExpressionInput WaterExtinction;
+
+	/**
+	 * Anisotropy of the volume with values lower than 0 representing back-scattering, equal 0 representing isotropic scattering and greater than 0 representing forward scattering. (type = float, unit = unitless, defaults to 0)
+	 */
+	UPROPERTY()
+	FExpressionInput WaterPhaseG;
+
+	/**
+	 * A scale to apply on the scene color behind the water surface. It can be used to approximate caustics for instance. (type = float3, unit = unitless, defaults to 1)
+	 */
+	UPROPERTY()
+	FExpressionInput ColorScaleBehindWater;
+
+	//~ Begin UMaterialExpression Interface
+#if WITH_EDITOR
+	virtual int32 Compile(class FMaterialCompiler* Compiler, int32 OutputIndex) override;
+	virtual void GetCaption(TArray<FString>& OutCaptions) const override;
+	virtual uint32 GetOutputType(int32 OutputIndex) override;
+	virtual uint32 GetInputType(int32 InputIndex) override;
+	virtual bool IsResultStrataMaterial(int32 OutputIndex) override;
+	virtual void GatherStrataMaterialInfo(FStrataMaterialInfo& StrataMaterialInfo, int32 OutputIndex) override;
+	virtual FStrataOperator* StrataGenerateMaterialTopologyTree(class FMaterialCompiler* Compiler, class UMaterialExpression* Parent, int32 OutputIndex) override;
+#endif
+	//~ End UMaterialExpression Interface
+};
+
+UCLASS(MinimalAPI, collapsecategories, hidecategories = Object, DisplayName = "Substrate Light Function")
+class UMaterialExpressionStrataLightFunction : public UMaterialExpressionStrataBSDF
+{
+	GENERATED_UCLASS_BODY()
+
+	/**
+	 * The output color of the light function
+	 */
+	UPROPERTY()
+	FExpressionInput Color;
+
+	//~ Begin UMaterialExpression Interface
+#if WITH_EDITOR
+	virtual int32 Compile(class FMaterialCompiler* Compiler, int32 OutputIndex) override;
+	virtual void GetCaption(TArray<FString>& OutCaptions) const override;
+	virtual uint32 GetOutputType(int32 OutputIndex) override;
+	virtual uint32 GetInputType(int32 InputIndex) override;
+	virtual bool IsResultStrataMaterial(int32 OutputIndex) override;
+	virtual void GatherStrataMaterialInfo(FStrataMaterialInfo& StrataMaterialInfo, int32 OutputIndex) override;
+	virtual FStrataOperator* StrataGenerateMaterialTopologyTree(class FMaterialCompiler* Compiler, class UMaterialExpression* Parent, int32 OutputIndex) override;
+#endif
+	//~ End UMaterialExpression Interface
+};
+
+UCLASS(MinimalAPI, collapsecategories, hidecategories = Object, DisplayName = "Substrate Post Process")
+class UMaterialExpressionStrataPostProcess : public UMaterialExpressionStrataBSDF
+{
+	GENERATED_UCLASS_BODY()
+
+	/**
+	 * The output color of the post process: it represents a color added over the back buffer, or a color multiplied if the Substrate blend mode is TransmittanceOnly.
+	 */
+	UPROPERTY()
+	FExpressionInput Color;
+
+	/**
+	 * The coverage of the post process: the more the value is high, the less the back buffer will be visible. Only used if "Output Alpha" is enabled on the root node.
+	 */
+	UPROPERTY()
+	FExpressionInput Opacity;
+
+	//~ Begin UMaterialExpression Interface
+#if WITH_EDITOR
+	virtual int32 Compile(class FMaterialCompiler* Compiler, int32 OutputIndex) override;
+	virtual void GetCaption(TArray<FString>& OutCaptions) const override;
+	virtual uint32 GetOutputType(int32 OutputIndex) override;
+	virtual uint32 GetInputType(int32 InputIndex) override;
+	virtual bool IsResultStrataMaterial(int32 OutputIndex) override;
+	virtual void GatherStrataMaterialInfo(FStrataMaterialInfo& StrataMaterialInfo, int32 OutputIndex) override;
+	virtual FStrataOperator* StrataGenerateMaterialTopologyTree(class FMaterialCompiler* Compiler, class UMaterialExpression* Parent, int32 OutputIndex) override;
+#endif
+	//~ End UMaterialExpression Interface
+};
+
+UCLASS(MinimalAPI, collapsecategories, hidecategories = Object, DisplayName = "Substrate UI")
+class UMaterialExpressionStrataUI : public UMaterialExpressionStrataBSDF
+{
+	GENERATED_UCLASS_BODY()
+
+	/**
+	 * The output color of the UI element.
+	 */
+	UPROPERTY()
+	FExpressionInput Color;
+
+	/**
+	 * The coverage of the UI element: the more the value is high, the less the back buffer will be visible.
+	 */
+	UPROPERTY()
+	FExpressionInput Opacity;
+
+	//~ Begin UMaterialExpression Interface
+#if WITH_EDITOR
+	virtual int32 Compile(class FMaterialCompiler* Compiler, int32 OutputIndex) override;
+	virtual void GetCaption(TArray<FString>& OutCaptions) const override;
+	virtual uint32 GetOutputType(int32 OutputIndex) override;
+	virtual uint32 GetInputType(int32 InputIndex) override;
+	virtual bool IsResultStrataMaterial(int32 OutputIndex) override;
+	virtual void GatherStrataMaterialInfo(FStrataMaterialInfo& StrataMaterialInfo, int32 OutputIndex) override;
+	virtual FStrataOperator* StrataGenerateMaterialTopologyTree(class FMaterialCompiler* Compiler, class UMaterialExpression* Parent, int32 OutputIndex) override;
+#endif
+	//~ End UMaterialExpression Interface
+};
+
+UCLASS(MinimalAPI, collapsecategories, hidecategories = Object, DisplayName = "Substrate Convert To Decal")
+class UMaterialExpressionStrataConvertToDecal : public UMaterialExpressionStrataBSDF
+{
+	GENERATED_UCLASS_BODY()
+
+	/**
+	 * The Substrate material to convert to a decal.
+	 */
+	UPROPERTY()
+	FExpressionInput DecalMaterial;
+
+	/**
+	 * The coverage of the decal (default 1)
+	 */
+	UPROPERTY()
+	FExpressionInput Coverage;
+
+	//~ Begin UMaterialExpression Interface
+#if WITH_EDITOR
+	virtual int32 Compile(class FMaterialCompiler* Compiler, int32 OutputIndex) override;
+	virtual void GetCaption(TArray<FString>& OutCaptions) const override;
+	virtual uint32 GetOutputType(int32 OutputIndex) override;
+	virtual uint32 GetInputType(int32 InputIndex) override;
+	virtual bool IsResultStrataMaterial(int32 OutputIndex) override;
+	virtual void GatherStrataMaterialInfo(FStrataMaterialInfo& StrataMaterialInfo, int32 OutputIndex) override;
+	virtual FStrataOperator* StrataGenerateMaterialTopologyTree(class FMaterialCompiler* Compiler, class UMaterialExpression* Parent, int32 OutputIndex) override;
+#endif
+	//~ End UMaterialExpression Interface
+};
+
+
+
+///////////////////////////////////////////////////////////////////////////////
+// Operator nodes
+
+UCLASS(MinimalAPI, collapsecategories, hidecategories = Object, DisplayName = "Substrate Horizontal Blend")
+class UMaterialExpressionStrataHorizontalMixing : public UMaterialExpressionStrataBSDF
+{
+	GENERATED_UCLASS_BODY()
+		
+	/**
+	 * Substrate material
+	 */
+	UPROPERTY()
+	FExpressionInput Background;
+
+	/**
+	 * Substrate material
+	 */
+	UPROPERTY()
+	FExpressionInput Foreground;
+
+	/**
+	 * Lerp factor between Background (Mix == 0) and Foreground (Mix == 1).
+	 */
+	UPROPERTY()
+	FExpressionInput Mix;
+
+	/**
+	 * Merge Background and Foreground into a single material by mixing their inputs rather than their evaluation. This makes lighting evaluation cheaper (Default: off)
+	 */
+	UPROPERTY(EditAnywhere, Category = Mode)
+	uint32 bUseParameterBlending : 1;
+
+	//~ Begin UMaterialExpression Interface
+#if WITH_EDITOR
+	virtual int32 Compile(class FMaterialCompiler* Compiler, int32 OutputIndex) override;
+	virtual void GetCaption(TArray<FString>& OutCaptions) const override;
+	virtual uint32 GetOutputType(int32 OutputIndex) override;
+	virtual uint32 GetInputType(int32 InputIndex) override;
+	virtual bool IsResultStrataMaterial(int32 OutputIndex) override;
+	virtual void GatherStrataMaterialInfo(FStrataMaterialInfo& StrataMaterialInfo, int32 OutputIndex) override;
+	virtual FStrataOperator* StrataGenerateMaterialTopologyTree(class FMaterialCompiler* Compiler, class UMaterialExpression* Parent, int32 OutputIndex) override;
+#endif
+	//~ End UMaterialExpression Interface
+};
+
+UCLASS(MinimalAPI, collapsecategories, hidecategories = Object, DisplayName = "Substrate Vertical Layer")
+class UMaterialExpressionStrataVerticalLayering : public UMaterialExpressionStrataBSDF
+{
+	GENERATED_UCLASS_BODY()
+
+	/**
+	 * Substrate material layer on top of the Base material layer
+	 */
+	UPROPERTY()
+	FExpressionInput Top;
+	
+	/**
+	 * Substrate material layer below the Top material layer
+	 */
+	UPROPERTY()
+	FExpressionInput Base;
+
+	/**
+	 * Thickness of the Top material layer
+	 */
+	UPROPERTY()
+	FExpressionInput Thickness;
+
+	/**
+	 * Merge Top and Base into a single material by mixing their inputs rather than their evaluation. This makes lighting evaluation cheaper (Default: off)
+	 */
+	UPROPERTY(EditAnywhere, Category = Mode)
+	uint32 bUseParameterBlending : 1;
+
+	//~ Begin UMaterialExpression Interface
+#if WITH_EDITOR
+	virtual int32 Compile(class FMaterialCompiler* Compiler, int32 OutputIndex) override;
+	virtual void GetCaption(TArray<FString>& OutCaptions) const override;
+	virtual uint32 GetOutputType(int32 OutputIndex) override;
+	virtual uint32 GetInputType(int32 InputIndex) override;
+	virtual bool IsResultStrataMaterial(int32 OutputIndex) override;
+	virtual void GatherStrataMaterialInfo(FStrataMaterialInfo& StrataMaterialInfo, int32 OutputIndex) override;
+	virtual FStrataOperator* StrataGenerateMaterialTopologyTree(class FMaterialCompiler* Compiler, UMaterialExpression* Parent, int32 OutputIndex) override;
+	virtual FName GetInputName(int32 InputIndex) const override;
+#endif
+	//~ End UMaterialExpression Interface
+};
+
+UCLASS(MinimalAPI, collapsecategories, hidecategories = Object, DisplayName = "Substrate Add")
+class UMaterialExpressionStrataAdd : public UMaterialExpressionStrataBSDF
+{
+	GENERATED_UCLASS_BODY()
+
+	/**
+	 * Substrate material
+	 */
+	UPROPERTY()
+	FExpressionInput A;
+	
+	/**
+	 * Substrate material
+	 */
+	UPROPERTY()
+	FExpressionInput B;
+
+	/**
+	 * Merge A and B into a single material by mixing their inputs rather than their evaluation. This makes lighting evaluation cheaper (Default: off)
+	 */
+	UPROPERTY(EditAnywhere, Category = Mode)
+	uint32 bUseParameterBlending : 1;
+
+	//~ Begin UMaterialExpression Interface
+#if WITH_EDITOR
+	virtual int32 Compile(class FMaterialCompiler* Compiler, int32 OutputIndex) override;
+	virtual void GetCaption(TArray<FString>& OutCaptions) const override;
+	virtual uint32 GetOutputType(int32 OutputIndex) override;
+	virtual uint32 GetInputType(int32 InputIndex) override;
+	virtual bool IsResultStrataMaterial(int32 OutputIndex) override;
+	virtual void GatherStrataMaterialInfo(FStrataMaterialInfo& StrataMaterialInfo, int32 OutputIndex) override;
+	virtual FStrataOperator* StrataGenerateMaterialTopologyTree(class FMaterialCompiler* Compiler, class UMaterialExpression* Parent, int32 OutputIndex) override;
+#endif
+	//~ End UMaterialExpression Interface
+};
+
+UCLASS(MinimalAPI, collapsecategories, hidecategories = Object, DisplayName = "Substrate Coverage Weight")
+class UMaterialExpressionStrataWeight : public UMaterialExpressionStrataBSDF
+{
+	GENERATED_UCLASS_BODY()
+
+	/**
+	 * Substrate material
+	 */
+	UPROPERTY()
+	FExpressionInput A;
+	
+	/**
+	 * Weight to apply to the Substrate material BSDFs
+	 */
+	UPROPERTY()
+	FExpressionInput Weight;
+
+	//~ Begin UMaterialExpression Interface
+#if WITH_EDITOR
+	virtual int32 Compile(class FMaterialCompiler* Compiler, int32 OutputIndex) override;
+	virtual void GetCaption(TArray<FString>& OutCaptions) const override;
+	virtual uint32 GetOutputType(int32 OutputIndex) override;
+	virtual uint32 GetInputType(int32 InputIndex) override;
+	virtual bool IsResultStrataMaterial(int32 OutputIndex) override;
+	virtual void GatherStrataMaterialInfo(FStrataMaterialInfo& StrataMaterialInfo, int32 OutputIndex) override;
+	virtual FStrataOperator* StrataGenerateMaterialTopologyTree(class FMaterialCompiler* Compiler, class UMaterialExpression* Parent, int32 OutputIndex) override;
+#endif
+	//~ End UMaterialExpression Interface
+};
+
+
+///////////////////////////////////////////////////////////////////////////////
+// Utilities
+
+UCLASS(MinimalAPI, collapsecategories, hidecategories = Object, Abstract, DisplayName = "Substrate Utility Base Class")
+class UMaterialExpressionStrataUtilityBase : public UMaterialExpression
+{
+	GENERATED_UCLASS_BODY()
+};
+
+UCLASS(MinimalAPI, collapsecategories, hidecategories = Object, DisplayName = "Substrate Transmittance-To-MeanFreePath")
+class UMaterialExpressionStrataTransmittanceToMFP : public UMaterialExpressionStrataUtilityBase
+{
+	GENERATED_UCLASS_BODY()
+
+	/**
+	* The colored transmittance for a view perpendicular to the surface. The transmittance for other view orientations will automatically be deduced according to surface thickness.
+	*/
+	UPROPERTY()
+	FExpressionInput TransmittanceColor;
+
+	/**
+	* The desired thickness in centimeter.
+	* Another use case example: this node output called thickness can be modulated before it is plugged in a slab node, this can be used to achieve simple scattering/transmittance variation of the same material.
+	*/
+	UPROPERTY()
+	FExpressionInput Thickness;
+
+	//~ Begin UMaterialExpression Interface
+#if WITH_EDITOR
+	virtual int32 Compile(class FMaterialCompiler* Compiler, int32 OutputIndex) override;
+	virtual void GetCaption(TArray<FString>& OutCaptions) const override;
+	virtual uint32 GetOutputType(int32 OutputIndex) override;
+	virtual uint32 GetInputType(int32 InputIndex) override;
+	virtual void GetConnectorToolTip(int32 InputIndex, int32 OutputIndex, TArray<FString>& OutToolTip) override;
+	virtual void GetExpressionToolTip(TArray<FString>& OutToolTip) override;
+#endif
+	//~ End UMaterialExpression Interface
+};
+
+UCLASS(MinimalAPI, collapsecategories, hidecategories = Object, DisplayName = "Substrate Metalness-To-DiffuseColorF0")
+class UMaterialExpressionStrataMetalnessToDiffuseAlbedoF0 : public UMaterialExpressionStrataUtilityBase
+{
+	GENERATED_UCLASS_BODY()
+
+	/**
+	 * Defines the overall color of the Material. (type = float3, unit = unitless, defaults to 0.18)
+	 */
+	UPROPERTY()
+	FExpressionInput BaseColor;
+
+	/**
+	 * Controls how \"metal-like\" your surface looks like. 0 means dielectric, 1 means conductor (type = float, unit = unitless, defaults to 0)
+	 */
+	UPROPERTY()
+	FExpressionInput Metallic;
+
+	/**
+	 * Used to scale the current amount of specularity on non-metallic surfaces and is a value between 0 and 1 (type = float, unit = unitless, defaults to plastic 0.5)
+	 */
+	UPROPERTY()
+	FExpressionInput Specular;
+
+	//~ Begin UMaterialExpression Interface
+#if WITH_EDITOR
+	virtual int32 Compile(class FMaterialCompiler* Compiler, int32 OutputIndex) override;
+	virtual void GetCaption(TArray<FString>& OutCaptions) const override;
+	virtual uint32 GetOutputType(int32 OutputIndex) override;
+	virtual uint32 GetInputType(int32 InputIndex) override;
+	virtual void GetConnectorToolTip(int32 InputIndex, int32 OutputIndex, TArray<FString>& OutToolTip) override;
+	virtual void GetExpressionToolTip(TArray<FString>& OutToolTip) override;
+#endif
+	//~ End UMaterialExpression Interface
+};
+
+UCLASS(MinimalAPI, collapsecategories, hidecategories = Object, DisplayName = "Substrate Haziness-To-Secondary-Roughness")
+class UMaterialExpressionStrataHazinessToSecondaryRoughness : public UMaterialExpressionStrataUtilityBase
+{
+	GENERATED_UCLASS_BODY()
+
+	/**
+	* The base roughness of the surface. It represented the smoothest part of the reflection.
+	*/
+	UPROPERTY()
+	FExpressionInput BaseRoughness;
+
+	/**
+	* Haziness represent the amount of irregularity of the surface. A high value will lead to a second rough specular lobe causing the surface too look `milky`.
+	*/
+	UPROPERTY()
+	FExpressionInput Haziness;
+
+	//~ Begin UMaterialExpression Interface
+#if WITH_EDITOR
+	virtual int32 Compile(class FMaterialCompiler* Compiler, int32 OutputIndex) override;
+	virtual void GetCaption(TArray<FString>& OutCaptions) const override;
+	virtual uint32 GetOutputType(int32 OutputIndex) override;
+	virtual uint32 GetInputType(int32 InputIndex) override;
+	virtual void GetConnectorToolTip(int32 InputIndex, int32 OutputIndex, TArray<FString>& OutToolTip) override;
+	virtual void GetExpressionToolTip(TArray<FString>& OutToolTip) override;
+#endif
+	//~ End UMaterialExpression Interface
+};
+
+UCLASS(MinimalAPI, collapsecategories, hidecategories = Object, DisplayName = "Substrate Thin-Film")
+class UMaterialExpressionStrataThinFilm : public UMaterialExpressionStrataUtilityBase
+{
+	GENERATED_UCLASS_BODY()
+
+	/**
+	 * The normal of the surface to consider. This input respects the normal space setup on the root node (tangent or world)
+	 */
+	UPROPERTY()
+	FExpressionInput Normal;
+
+	/**
+	 * Defines the color and brightness of the specular highlight where the surface is facing the camera. Each specular contribution will fade to black as F0 drops below 0.02. (type = float3, unit = unitless, defaults to plastic 0.04)
+	 */
+	UPROPERTY()
+	FExpressionInput F0;
+
+	/**
+	 * Defines the color of the specular highlight where the surface normal is 90 degrees from the view direction. Only the hue and saturation are preserved, the brightness is fixed at 1.0. Fades to black as F0 drops below 0.02. (type = float3, unit = unitless, defaults to 1.0f)
+	 */
+	UPROPERTY()
+	FExpressionInput F90;
+	
+	/**
+	 * Thin film controls the thin film layer coating the current slab. 0 means disabled and 1 means a coating layer of 10 micrometer. (type = float, unitless, default = 0)
+	 */
+	UPROPERTY()
+	FExpressionInput Thickness;
+
+	/**
+	 * Thin film IOR
+	 */
+	UPROPERTY()
+	FExpressionInput IOR;
+
+	//~ Begin UMaterialExpression Interface
+#if WITH_EDITOR
+	virtual int32 Compile(class FMaterialCompiler* Compiler, int32 OutputIndex) override;
+	virtual void GetCaption(TArray<FString>& OutCaptions) const override;
+	virtual uint32 GetOutputType(int32 OutputIndex) override;
+	virtual uint32 GetInputType(int32 InputIndex) override;
+	virtual void GetConnectorToolTip(int32 InputIndex, int32 OutputIndex, TArray<FString>& OutToolTip) override;
+	virtual void GetExpressionToolTip(TArray<FString>& OutToolTip) override;
+#endif
+	//~ End UMaterialExpression Interface
+};
+
diff --git a/Engine/Source/Runtime/Engine/Private/Materials/HLSLMaterialTranslator.cpp b/Engine/Source/Runtime/Engine/Private/Materials/HLSLMaterialTranslator.cpp
--- a/Engine/Source/Runtime/Engine/Private/Materials/HLSLMaterialTranslator.cpp
+++ b/Engine/Source/Runtime/Engine/Private/Materials/HLSLMaterialTranslator.cpp
@@ -36,18 +36,18 @@
 #include "Materials/MaterialExpressionViewProperty.h"
 #include "Materials/MaterialExpressionVolumetricAdvancedMaterialOutput.h"
 #include "Materials/MaterialExpressionWorldPosition.h"
-#include "Materials/StrataMaterial.h"
+#include "Materials/SubstrateMaterial.h"
 #include "ParameterCollection.h"
 #include "RenderUtils.h"
 #include "Stats/StatsMisc.h"
 #include "Stats/StatsTrace.h"
-#include "StrataDefinitions.h"
+#include "SubstrateDefinitions.h"
 #include "VT/RuntimeVirtualTexture.h"
 #include <memory>
 #include <tuple>
 
 #if WITH_EDITORONLY_DATA
-#include "Materials/MaterialExpressionStrata.h"
+#include "Materials/MaterialExpressionSubstrate.h"
 #include "ShaderPlatformCachedIniValue.h"
 #endif
 
diff --git a/Engine/Source/Runtime/Engine/Private/Materials/HLSLMaterialTranslator.h b/Engine/Source/Runtime/Engine/Private/Materials/HLSLMaterialTranslator.h
--- a/Engine/Source/Runtime/Engine/Private/Materials/HLSLMaterialTranslator.h
+++ b/Engine/Source/Runtime/Engine/Private/Materials/HLSLMaterialTranslator.h
@@ -49,7 +49,7 @@
 #include "Containers/LazyPrintf.h"
 #include "Containers/HashTable.h"
 #include "Engine/Texture2D.h"
-#include "StrataMaterial.h"
+#include "SubstrateMaterial.h"
 #include "HLSLMaterialDerivativeAutogen.h"
 #endif
 
diff --git a/Engine/Source/Runtime/Engine/Private/Materials/Material.cpp b/Engine/Source/Runtime/Engine/Private/Materials/Material.cpp
--- a/Engine/Source/Runtime/Engine/Private/Materials/Material.cpp
+++ b/Engine/Source/Runtime/Engine/Private/Materials/Material.cpp
@@ -75,7 +75,7 @@
 #include "Materials/MaterialExpressionVolumetricAdvancedMaterialInput.h"
 #include "Materials/MaterialExpressionVolumetricAdvancedMaterialOutput.h"
 #include "Materials/MaterialExpressionCloudLayer.h"
-#include "Materials/MaterialExpressionStrata.h"
+#include "Materials/MaterialExpressionSubstrate.h"
 #include "Materials/MaterialExpressionThinTranslucentMaterialOutput.h"
 #include "Materials/MaterialExpressionClearCoatNormalCustomOutput.h"
 #include "Materials/MaterialExpressionTangentOutput.h"
diff --git a/Engine/Source/Runtime/Engine/Private/Materials/MaterialAttributeDefinitionMap.cpp b/Engine/Source/Runtime/Engine/Private/Materials/MaterialAttributeDefinitionMap.cpp
--- a/Engine/Source/Runtime/Engine/Private/Materials/MaterialAttributeDefinitionMap.cpp
+++ b/Engine/Source/Runtime/Engine/Private/Materials/MaterialAttributeDefinitionMap.cpp
@@ -6,8 +6,8 @@
 #include "MaterialDomain.h"
 #include "Materials/Material.h"
 #include "Materials/MaterialExpressionShadingModel.h"
-#include "Materials/MaterialExpressionStrata.h"
-#include "StrataDefinitions.h"
+#include "Materials/MaterialExpressionSubstrate.h"
+#include "SubstrateDefinitions.h"
 
 #define LOCTEXT_NAMESPACE "MaterialShared"
 
diff --git a/Engine/Source/Runtime/Engine/Private/Materials/MaterialExpressions.cpp b/Engine/Source/Runtime/Engine/Private/Materials/MaterialExpressions.cpp
--- a/Engine/Source/Runtime/Engine/Private/Materials/MaterialExpressions.cpp
+++ b/Engine/Source/Runtime/Engine/Private/Materials/MaterialExpressions.cpp
@@ -14,7 +14,7 @@
 #include "UObject/UObjectAnnotation.h"
 #include "RenderUtils.h"
 #include "UObject/UObjectIterator.h"
-#include "StrataDefinitions.h"
+#include "SubstrateDefinitions.h"
 #include "UObject/ConstructorHelpers.h"
 #include "Engine/Engine.h"
 #include "Engine/Font.h"
@@ -268,7 +268,7 @@
 #include "Materials/MaterialExpressionShaderStageSwitch.h"
 #include "Materials/MaterialExpressionReflectionCapturePassSwitch.h"
 #include "Materials/MaterialExpressionReflectionCapturePassSwitch.h"
-#include "Materials/MaterialExpressionStrata.h"
+#include "Materials/MaterialExpressionSubstrate.h"
 #include "Materials/MaterialExpressionSamplePhysicsField.h"
 #include "Materials/MaterialExpressionExecBegin.h"
 #include "Materials/MaterialExpressionExecEnd.h"
@@ -289,7 +289,7 @@
 #include "Framework/Notifications/NotificationManager.h"
 #include "Widgets/Notifications/SNotificationList.h"
 #include "MaterialGraph/MaterialGraphSchema.h"
-#include "StrataMaterial.h"
+#include "SubstrateMaterial.h"
 #else
 #include "Materials/MaterialExpressionVertexInterpolator.h"
 #include "Materials/MaterialParameterCollection.h"
diff --git a/Engine/Source/Runtime/Engine/Private/Materials/MaterialShared.cpp b/Engine/Source/Runtime/Engine/Private/Materials/MaterialShared.cpp
--- a/Engine/Source/Runtime/Engine/Private/Materials/MaterialShared.cpp
+++ b/Engine/Source/Runtime/Engine/Private/Materials/MaterialShared.cpp
@@ -52,7 +52,7 @@
 #endif
 #include "ProfilingDebugging/CountersTrace.h"
 #include "RenderCore.h"
-#include "StrataDefinitions.h"
+#include "SubstrateDefinitions.h"
 #include "DataDrivenShaderPlatformInfo.h"
 #include "ProfilingDebugging/CookStats.h"
 
diff --git a/Engine/Source/Runtime/Engine/Private/Materials/StrataMaterial.cpp b/Engine/Source/Runtime/Engine/Private/Materials/StrataMaterial.cpp
deleted file mode 100644
diff --git a/Engine/Source/Runtime/Engine/Private/Materials/StrataMaterial.h b/Engine/Source/Runtime/Engine/Private/Materials/StrataMaterial.h
deleted file mode 100644
diff --git a/Engine/Source/Runtime/Engine/Private/Materials/SubstrateMaterial.cpp b/Engine/Source/Runtime/Engine/Private/Materials/SubstrateMaterial.cpp
new file mode 100644
--- /dev/null
+++ b/Engine/Source/Runtime/Engine/Private/Materials/SubstrateMaterial.cpp
@@ -0,0 +1,76 @@
+// Copyright Epic Games, Inc. All Rights Reserved.
+
+#include "SubstrateMaterial.h"
+#include "MaterialCompiler.h"
+#include "Materials/MaterialInterface.h"
+#include "SubstrateDefinitions.h"
+#include "RenderUtils.h"
+
+FString GetStrataBSDFName(uint8 BSDFType)
+{
+	switch (BSDFType)
+	{
+	case STRATA_BSDF_TYPE_SLAB:
+		return TEXT("SLAB");
+		break;
+	case STRATA_BSDF_TYPE_VOLUMETRICFOGCLOUD:
+		return TEXT("VOLUMETRICFOGCLOUD");
+		break;
+	case STRATA_BSDF_TYPE_UNLIT:
+		return TEXT("UNLIT");
+		break;
+	case STRATA_BSDF_TYPE_HAIR:
+		return TEXT("HAIR");
+		break;
+	case STRATA_BSDF_TYPE_EYE:
+		return TEXT("EYE");
+		break;
+	case STRATA_BSDF_TYPE_SINGLELAYERWATER:
+		return TEXT("SINGLELAYERWATER");
+		break;
+	}
+	check(false);
+	return "";
+}
+
+FStrataRegisteredSharedLocalBasis StrataCompilationInfoCreateNullSharedLocalBasis()
+{
+	return FStrataRegisteredSharedLocalBasis();
+}
+
+FStrataRegisteredSharedLocalBasis StrataCompilationInfoCreateSharedLocalBasis(FMaterialCompiler* Compiler, int32 NormalCodeChunk, int32 TangentCodeChunk)
+{
+	if (TangentCodeChunk == INDEX_NONE)
+	{
+		return Compiler->StrataCompilationInfoRegisterSharedLocalBasis(NormalCodeChunk);
+	}
+	return Compiler->StrataCompilationInfoRegisterSharedLocalBasis(NormalCodeChunk, TangentCodeChunk);
+}
+
+inline bool IsStrataEnabled()
+{
+	static bool bStrataEnabled = Strata::IsStrataEnabled();
+	return bStrataEnabled;
+}
+
+#define IsGenericBlendMode_Conversion0(Name, LegacyCondition) \
+	bool Is##Name##BlendMode(EBlendMode BlendMode)											{ return LegacyCondition; } \
+	bool Is##Name##BlendMode(const FMaterial& In)											{ return Is##Name##BlendMode(In.GetBlendMode()); } \
+	bool Is##Name##BlendMode(const UMaterialInterface& In)									{ return Is##Name##BlendMode(In.GetBlendMode()); } \
+	bool Is##Name##BlendMode(const FMaterialShaderParameters& In)							{ return Is##Name##BlendMode(In.BlendMode); }
+
+#define IsGenericBlendMode_Conversion1(Name, LegacyCondition, StrataCondition) \
+	bool Is##Name##BlendMode(EBlendMode BlendMode)											{ return IsStrataEnabled() ? (StrataCondition) : (LegacyCondition); } \
+	bool Is##Name##BlendMode(const FMaterial& In)											{ return Is##Name##BlendMode(In.GetBlendMode()); } \
+	bool Is##Name##BlendMode(const UMaterialInterface& In)									{ return Is##Name##BlendMode(In.GetBlendMode()); } \
+	bool Is##Name##BlendMode(const FMaterialShaderParameters& In)							{ return Is##Name##BlendMode(In.BlendMode); }
+
+IsGenericBlendMode_Conversion0(Opaque,			BlendMode == BLEND_Opaque)																												// Opaque blend mode
+IsGenericBlendMode_Conversion0(Masked,			BlendMode == BLEND_Masked)																												// Masked blend mode
+IsGenericBlendMode_Conversion0(OpaqueOrMasked,	BlendMode == BLEND_Opaque || BlendMode == BLEND_Masked)																					// Opaque or Masked blend mode
+IsGenericBlendMode_Conversion0(Translucent,		BlendMode != BLEND_Opaque && BlendMode != BLEND_Masked)																					// General translucency (i.e., blend mode is something else than Opaque/Masked)
+IsGenericBlendMode_Conversion1(TranslucentOnly, BlendMode == BLEND_Translucent, BlendMode == BLEND_TranslucentColoredTransmittance || BlendMode == BLEND_TranslucentGreyTransmittance) 	// Explicit translucency blend mode
+IsGenericBlendMode_Conversion0(AlphaHoldout,	BlendMode == BLEND_AlphaHoldout)																										// AlphaHoldout blend mode
+IsGenericBlendMode_Conversion1(Modulate,		BlendMode == BLEND_Modulate, BlendMode == BLEND_ColoredTransmittanceOnly)																// Modulate blend mode
+IsGenericBlendMode_Conversion0(Additive,		BlendMode == BLEND_Additive)																											// Additive blend mode
+IsGenericBlendMode_Conversion0(AlphaComposite,	BlendMode == BLEND_AlphaComposite)																										// AlphaComposite blend mode
\ No newline at end of file
diff --git a/Engine/Source/Runtime/Engine/Private/Materials/SubstrateMaterial.h b/Engine/Source/Runtime/Engine/Private/Materials/SubstrateMaterial.h
new file mode 100644
--- /dev/null
+++ b/Engine/Source/Runtime/Engine/Private/Materials/SubstrateMaterial.h
@@ -0,0 +1,16 @@
+// Copyright Epic Games, Inc. All Rights Reserved.
+
+#pragma once
+
+#include "CoreMinimal.h"
+#include "Containers/Map.h"
+#include "MaterialCompiler.h"
+
+
+class FMaterialCompiler;
+
+FString GetStrataBSDFName(uint8 BSDFType);
+
+FStrataRegisteredSharedLocalBasis StrataCompilationInfoCreateNullSharedLocalBasis();
+FStrataRegisteredSharedLocalBasis StrataCompilationInfoCreateSharedLocalBasis(FMaterialCompiler* Compiler, int32 NormalCodeChunk, int32 TangentCodeChunk = INDEX_NONE);
+
diff --git a/Engine/Source/Runtime/Engine/Private/StrataVisualizationData.cpp b/Engine/Source/Runtime/Engine/Private/StrataVisualizationData.cpp
--- a/Engine/Source/Runtime/Engine/Private/StrataVisualizationData.cpp
+++ b/Engine/Source/Runtime/Engine/Private/StrataVisualizationData.cpp
@@ -1,6 +1,6 @@
 // Copyright Epic Games, Inc. All Rights Reserved.
 
-#include "StrataVisualizationData.h"
+#include "SubstrateVisualizationData.h"
 #include "HAL/IConsoleManager.h"
 #include "Materials/Material.h"
 #include "Misc/ConfigCacheIni.h"
diff --git a/Engine/Source/Runtime/Engine/Public/Rendering/StrataMaterialShared.h b/Engine/Source/Runtime/Engine/Public/Rendering/StrataMaterialShared.h
--- a/Engine/Source/Runtime/Engine/Public/Rendering/StrataMaterialShared.h
+++ b/Engine/Source/Runtime/Engine/Public/Rendering/StrataMaterialShared.h
@@ -3,7 +3,7 @@
 #pragma once
 
 #include "CoreMinimal.h"
-#include "StrataDefinitions.h"
+#include "SubstrateDefinitions.h"
 #include "Serialization/MemoryImage.h"
 
 
diff --git a/Engine/Source/Runtime/Engine/Public/StrataVisualizationData.h b/Engine/Source/Runtime/Engine/Public/StrataVisualizationData.h
deleted file mode 100644
diff --git a/Engine/Source/Runtime/Engine/Public/SubstrateVisualizationData.h b/Engine/Source/Runtime/Engine/Public/SubstrateVisualizationData.h
new file mode 100644
--- /dev/null
+++ b/Engine/Source/Runtime/Engine/Public/SubstrateVisualizationData.h
@@ -0,0 +1,85 @@
+// Copyright Epic Games, Inc. All Rights Reserved.
+
+#pragma once
+
+#include "CoreMinimal.h"
+
+enum class FStrataViewMode : uint8
+{
+	None,
+	MaterialProperties,
+	MaterialCount,
+	MaterialByteCount,
+	StrataInfo,
+	AdvancedMaterialProperties,
+	MaterialClassification,
+	RoughRefractionClassification,
+	DecalClassification,
+};
+
+class FStrataVisualizationData
+{
+public:
+
+	/** Describes a single available visualization mode. */
+	struct FModeRecord
+	{
+		FString   ModeString;
+		FName     ModeName;
+		FText     ModeText;
+		FText     ModeDesc;
+		FStrataViewMode ViewMode;
+
+		// Whether or not this mode (by default) composites with regular scene depth.
+		bool      bDefaultComposited;
+
+		bool      bAvailableCommand;
+		FText	  UnavailableReason;
+	};
+
+	/** Mapping of FName to a visualization mode record. */
+	typedef TMultiMap<FName, FModeRecord> TModeMap;
+
+public:
+	FStrataVisualizationData()
+	: bIsInitialized(false)
+	{
+	}
+
+	/** Initialize the system. */
+	void Initialize();
+
+	/** Check if system was initialized. */
+	inline bool IsInitialized() const { return bIsInitialized; }
+
+	/** Get the display name of a named mode from the available mode map. **/
+	ENGINE_API FText GetModeDisplayName(const FName& InModeName) const;
+
+	ENGINE_API FStrataViewMode GetViewMode(const FName& InModeName) const;
+
+	ENGINE_API bool GetModeDefaultComposited(const FName& InModeName) const;
+
+	inline const TModeMap& GetModeMap() const
+	{
+		return ModeMap;
+	}
+
+	/** Return the console command name for enabling single mode visualization. */
+	static const TCHAR* GetVisualizeConsoleCommandName()
+	{
+		return TEXT("r.Substrate.ViewMode");
+	}
+
+private:
+
+	/** The name->mode mapping table */
+	TModeMap ModeMap;
+
+	/** Storage for console variable documentation strings. **/
+	FString ConsoleDocumentationVisualizationMode;
+
+	/** Flag indicating if system is initialized. **/
+	bool bIsInitialized;
+};
+
+ENGINE_API FStrataVisualizationData& GetStrataVisualizationData();
diff --git a/Engine/Source/Runtime/RenderCore/Private/RenderUtils.cpp b/Engine/Source/Runtime/RenderCore/Private/RenderUtils.cpp
--- a/Engine/Source/Runtime/RenderCore/Private/RenderUtils.cpp
+++ b/Engine/Source/Runtime/RenderCore/Private/RenderUtils.cpp
@@ -15,7 +15,7 @@
 #include "ShaderPlatformCachedIniValue.h"
 #include "DataDrivenShaderPlatformInfo.h"
 #include "RenderCore.h"
-#include "StrataDefinitions.h"
+#include "SubstrateDefinitions.h"
 
 #if WITH_EDITOR
 #include "Interfaces/ITargetPlatform.h"
diff --git a/Engine/Source/Runtime/Renderer/Private/Strata/Strata.h b/Engine/Source/Runtime/Renderer/Private/Strata/Strata.h
--- a/Engine/Source/Runtime/Renderer/Private/Strata/Strata.h
+++ b/Engine/Source/Runtime/Renderer/Private/Strata/Strata.h
@@ -7,7 +7,7 @@
 #include "MeshPassProcessor.h"
 #include "UnifiedBuffer.h"
 #include "RHIUtilities.h"
-#include "StrataDefinitions.h"
+#include "SubstrateDefinitions.h"
 #include "GBufferInfo.h"
 #include "RendererUtils.h"
 #include "ShaderParameterStruct.h"
diff --git a/Engine/Source/Runtime/Renderer/Private/Strata/StrataVisualize.cpp b/Engine/Source/Runtime/Renderer/Private/Strata/StrataVisualize.cpp
--- a/Engine/Source/Runtime/Renderer/Private/Strata/StrataVisualize.cpp
+++ b/Engine/Source/Runtime/Renderer/Private/Strata/StrataVisualize.cpp
@@ -9,7 +9,7 @@
 #include "PixelShaderUtils.h"
 #include "BasePassRendering.h"
 #include "IndirectLightRendering.h"
-#include "StrataVisualizationData.h"
+#include "SubstrateVisualizationData.h"
 
 namespace Strata
 {
@@ -368,7 +368,7 @@ static FStrataViewMode GetStrataVisualizeMode(const FViewInfo & View)
 	FStrataViewMode Out = FStrataViewMode::None;
 	if (IsStrataEnabled() && StrataDebugVisualizationCanRunOnPlatform(View.GetShaderPlatform()))
 	{
-		// Variable defined in StrataVisualizationData.h/.cpp
+		// Variable defined in SubstrateVisualizationData.h/.cpp
 		static const auto CVar = IConsoleManager::Get().FindConsoleVariable(FStrataVisualizationData::GetVisualizeConsoleCommandName());
 		const uint32 ViewMode = CVar && CVar->AsVariableInt() ? CVar->AsVariableInt()->GetValueOnRenderThread() : 0;
 		switch (ViewMode)
