Fix crash when opening material with SubstrateShadingModels node due to missing pin type.

diff --git a/Engine/Source/Runtime/Engine/Private/Materials/MaterialExpressions.cpp b/Engine/Source/Runtime/Engine/Private/Materials/MaterialExpressions.cpp
--- a/Engine/Source/Runtime/Engine/Private/Materials/MaterialExpressions.cpp
+++ b/Engine/Source/Runtime/Engine/Private/Materials/MaterialExpressions.cpp
@@ -23385,6 +23385,7 @@ uint32 UMaterialExpressionStrataLegacyConversion::GetInputType(int32 InputIndex)
 	else if (InputIndex == 17) return MCT_Float3; // ClearCoatNormal
 	else if (InputIndex == 18) return MCT_Float3; // CustomTangent
 	else if (InputIndex == 19) return MCT_ShadingModel; // ShadingModel
+	else if (InputIndex == 20) return MCT_ShadingModel; // ShadingModelOverride (as it uses 'ShowAsInputPin' metadata)
 
 	check(false);
 	return MCT_Float1;
