Substrate - Renaming - Strata class redirection to Substrate class.

diff --git a/Engine/Config/BaseEngine.ini b/Engine/Config/BaseEngine.ini
--- a/Engine/Config/BaseEngine.ini
+++ b/Engine/Config/BaseEngine.ini
@@ -1176,6 +1176,35 @@ DefaultWorldPartitionHLODLayer=/Engine/Maps/Templates/HLODs/HLODLayer_Instanced.
 +ClassRedirects=(OldName="/Script/InterchangeImport.MaterialExpressionLogarithm", NewName="/Script/Engine.MaterialExpressionLogarithm")
 +ClassRedirects=(OldName="/Script/InterchangeImport.MaterialExpressionRgbToHsv", NewName="/Script/Engine.MaterialExpressionRgbToHsv")
 
+;5.4
+
+; Strata to Substrate renaming
++ClassRedirects=(OldName="/Script/Engine.MaterialExpressionStrataBSDF",							NewName="/Script/Engine.MaterialExpressionSubstrateBSDF")
++ClassRedirects=(OldName="/Script/Engine.MaterialExpressionStrataLegacyConversion",				NewName="/Script/Engine.MaterialExpressionSubstrateShadingModels")
++ClassRedirects=(OldName="/Script/Engine.MaterialExpressionStrataSlabBSDF",						NewName="/Script/Engine.MaterialExpressionSubstrateSlabBSDF")
++ClassRedirects=(OldName="/Script/Engine.MaterialExpressionStrataSimpleClearCoatBSDF",			NewName="/Script/Engine.MaterialExpressionSubstrateSimpleClearCoatBSDF")
++ClassRedirects=(OldName="/Script/Engine.MaterialExpressionStrataVolumetricFogCloudBSDF",		NewName="/Script/Engine.MaterialExpressionSubstrateVolumetricFogCloudBSDF")
++ClassRedirects=(OldName="/Script/Engine.MaterialExpressionStrataUnlitBSDF",					NewName="/Script/Engine.MaterialExpressionSubstrateUnlitBSDF")
++ClassRedirects=(OldName="/Script/Engine.MaterialExpressionStrataHairBSDF",						NewName="/Script/Engine.MaterialExpressionSubstrateHairBSDF")
++ClassRedirects=(OldName="/Script/Engine.MaterialExpressionStrataEyeBSDF",						NewName="/Script/Engine.MaterialExpressionSubstrateEyeBSDF")
++ClassRedirects=(OldName="/Script/Engine.MaterialExpressionStrataSingleLayerWaterBSDF",			NewName="/Script/Engine.MaterialExpressionSubstrateSingleLayerWaterBSDF")
++ClassRedirects=(OldName="/Script/Engine.MaterialExpressionStrataLightFunction",				NewName="/Script/Engine.MaterialExpressionSubstrateLightFunction")
++ClassRedirects=(OldName="/Script/Engine.MaterialExpressionStrataPostProcess",					NewName="/Script/Engine.MaterialExpressionSubstratePostProcess")
++ClassRedirects=(OldName="/Script/Engine.MaterialExpressionStrataUI",							NewName="/Script/Engine.MaterialExpressionSubstrateUI")
++ClassRedirects=(OldName="/Script/Engine.MaterialExpressionStrataConvertToDecal",				NewName="/Script/Engine.MaterialExpressionSubstrateConvertToDecal")
++ClassRedirects=(OldName="/Script/Engine.MaterialExpressionStrataHorizontalMixing",				NewName="/Script/Engine.MaterialExpressionSubstrateHorizontalMixing")
++ClassRedirects=(OldName="/Script/Engine.MaterialExpressionStrataVerticalLayering",				NewName="/Script/Engine.MaterialExpressionSubstrateVerticalLayering")
++ClassRedirects=(OldName="/Script/Engine.MaterialExpressionStrataAdd",							NewName="/Script/Engine.MaterialExpressionSubstrateAdd")
++ClassRedirects=(OldName="/Script/Engine.MaterialExpressionStrataWeight",						NewName="/Script/Engine.MaterialExpressionSubstrateWeight")
++ClassRedirects=(OldName="/Script/Engine.MaterialExpressionStrataUtilityBase",					NewName="/Script/Engine.MaterialExpressionSubstrateUtilityBase")
++ClassRedirects=(OldName="/Script/Engine.MaterialExpressionStrataTransmittanceToMFP",			NewName="/Script/Engine.MaterialExpressionSubstrateTransmittanceToMFP")
++ClassRedirects=(OldName="/Script/Engine.MaterialExpressionStrataMetalnessToDiffuseAlbedoF0",	NewName="/Script/Engine.MaterialExpressionSubstrateMetalnessToDiffuseAlbedoF0")
++ClassRedirects=(OldName="/Script/Engine.MaterialExpressionStrataHazinessToSecondaryRoughness",	NewName="/Script/Engine.MaterialExpressionSubstrateHazinessToSecondaryRoughness")
++ClassRedirects=(OldName="/Script/Engine.MaterialExpressionStrataThinFilm",						NewName="/Script/Engine.MaterialExpressionSubstrateThinFilm")
+
+; End of releases redirects
+
+
 [CoreUObject.Metadata]
 ; Note: UnrealHeaderTool should be rerun after making changes to MetadataRedirects to catch any keys specified in class headers
 +MetadataRedirects=(OldKey="K2Protected", NewKey="BlueprintProtected")
diff --git a/Engine/Source/Editor/GraphEditor/Private/MaterialNodes/SGraphNodeMaterialBase.cpp b/Engine/Source/Editor/GraphEditor/Private/MaterialNodes/SGraphNodeMaterialBase.cpp
--- a/Engine/Source/Editor/GraphEditor/Private/MaterialNodes/SGraphNodeMaterialBase.cpp
+++ b/Engine/Source/Editor/GraphEditor/Private/MaterialNodes/SGraphNodeMaterialBase.cpp
@@ -483,9 +483,9 @@ void SGraphNodeMaterialBase::CreateBelowPinControls(TSharedPtr<SVerticalBox> Mai
 	}
 
 	// Preview of Substrate nodes topology
-	if (Strata::IsStrataEnabled() && MaterialNode && MaterialNode->MaterialExpression->IsA(UMaterialExpressionStrataBSDF::StaticClass()))
+	if (Strata::IsStrataEnabled() && MaterialNode && MaterialNode->MaterialExpression->IsA(UMaterialExpressionSubstrateBSDF::StaticClass()))
 	{
-		if (const UMaterialExpressionStrataBSDF* StrataExpression = (const UMaterialExpressionStrataBSDF*)MaterialNode->MaterialExpression)
+		if (const UMaterialExpressionSubstrateBSDF* StrataExpression = (const UMaterialExpressionSubstrateBSDF*)MaterialNode->MaterialExpression)
 		{		
 			if (UMaterial* MaterialForStats = StrataExpression->Material)
 			{
diff --git a/Engine/Source/Editor/GraphEditor/Private/MaterialNodes/SGraphSubstrateMaterial.cpp b/Engine/Source/Editor/GraphEditor/Private/MaterialNodes/SGraphSubstrateMaterial.cpp
--- a/Engine/Source/Editor/GraphEditor/Private/MaterialNodes/SGraphSubstrateMaterial.cpp
+++ b/Engine/Source/Editor/GraphEditor/Private/MaterialNodes/SGraphSubstrateMaterial.cpp
@@ -227,7 +227,7 @@ void FSubstrateWidget::GetPinColor(TSharedPtr<SGraphPin>& Out, const UMaterialGr
 	// Strata operator override pin color to ease material topology visualization
 	const UEdGraphPin* Pin = Out->SGraphPin::GetPinObj();
 	const FName PinName = Pin->PinName;
-	if (InNode->MaterialExpression->IsA(UMaterialExpressionStrataVerticalLayering::StaticClass()))
+	if (InNode->MaterialExpression->IsA(UMaterialExpressionSubstrateVerticalLayering::StaticClass()))
 	{			
 		if (PinName == InNode->MaterialExpression->GetInputName(0)) // Top
 		{
@@ -240,7 +240,7 @@ void FSubstrateWidget::GetPinColor(TSharedPtr<SGraphPin>& Out, const UMaterialGr
 			ColorModifier = Color1;
 		}
 	}
-	else if (InNode->MaterialExpression->IsA(UMaterialExpressionStrataHorizontalMixing::StaticClass()))
+	else if (InNode->MaterialExpression->IsA(UMaterialExpressionSubstrateHorizontalMixing::StaticClass()))
 	{
 		if (PinName == InNode->MaterialExpression->GetInputName(1)) // Foreground
 		{
@@ -253,7 +253,7 @@ void FSubstrateWidget::GetPinColor(TSharedPtr<SGraphPin>& Out, const UMaterialGr
 			ColorModifier = Color0;
 		}
 	}
-	else if (InNode->MaterialExpression->IsA(UMaterialExpressionStrataAdd::StaticClass()))
+	else if (InNode->MaterialExpression->IsA(UMaterialExpressionSubstrateAdd::StaticClass()))
 	{
 		if (PinName == InNode->MaterialExpression->GetInputName(0)) // A
 		{
@@ -267,7 +267,7 @@ void FSubstrateWidget::GetPinColor(TSharedPtr<SGraphPin>& Out, const UMaterialGr
 		}
 	}
 
-	if (InNode->MaterialExpression->IsA(UMaterialExpressionStrataBSDF::StaticClass()) && Out->GetDirection() == EGPD_Output)
+	if (InNode->MaterialExpression->IsA(UMaterialExpressionSubstrateBSDF::StaticClass()) && Out->GetDirection() == EGPD_Output)
 	{
 		bHasColorModifier = true;
 		ColorModifier = FSubstrateWidget::GetConnectionColor();
diff --git a/Engine/Source/Editor/MaterialEditor/Private/MaterialEditor.cpp b/Engine/Source/Editor/MaterialEditor/Private/MaterialEditor.cpp
--- a/Engine/Source/Editor/MaterialEditor/Private/MaterialEditor.cpp
+++ b/Engine/Source/Editor/MaterialEditor/Private/MaterialEditor.cpp
@@ -4835,18 +4835,18 @@ void FMaterialEditor::OnCreateStrataNodeForPin(const FToolMenuContext& InMenuCon
 	const bool bTargetPinIsInput = TargetPin->Direction == EEdGraphPinDirection::EGPD_Input;
 
 	FMaterialGraphSchemaAction_NewNode Action;
-	Action.MaterialExpressionClass = UMaterialExpressionStrataSlabBSDF::StaticClass();
+	Action.MaterialExpressionClass = UMaterialExpressionSubstrateSlabBSDF::StaticClass();
 	if (NodeForPin == EStrataNodeForPin::HorizontalMix)
 	{
-		Action.MaterialExpressionClass = UMaterialExpressionStrataHorizontalMixing::StaticClass();
+		Action.MaterialExpressionClass = UMaterialExpressionSubstrateHorizontalMixing::StaticClass();
 	}
 	else if (NodeForPin == EStrataNodeForPin::VerticalLayer)
 	{
-		Action.MaterialExpressionClass = UMaterialExpressionStrataVerticalLayering::StaticClass();
+		Action.MaterialExpressionClass = UMaterialExpressionSubstrateVerticalLayering::StaticClass();
 	}
 	else if (NodeForPin == EStrataNodeForPin::Weight)
 	{
-		Action.MaterialExpressionClass = UMaterialExpressionStrataWeight::StaticClass();
+		Action.MaterialExpressionClass = UMaterialExpressionSubstrateWeight::StaticClass();
 	}
 
 	check(PinNode);
@@ -6097,7 +6097,7 @@ void FMaterialEditor::UpdateStrataTopologyPreview()
 		for (int32 Index = 0; Index < Material->MaterialGraph->Nodes.Num(); ++Index)
 		{
 			UMaterialGraphNode* MaterialNode = Cast<UMaterialGraphNode>(Material->MaterialGraph->Nodes[Index]);
-			if (MaterialNode && MaterialNode->MaterialExpression && MaterialNode->MaterialExpression->IsA(UMaterialExpressionStrataBSDF::StaticClass()))
+			if (MaterialNode && MaterialNode->MaterialExpression && MaterialNode->MaterialExpression->IsA(UMaterialExpressionSubstrateBSDF::StaticClass()))
 			{
 				UEdGraph* Graph = MaterialNode->GetGraph();
 				if (Graph)
diff --git a/Engine/Source/Editor/MaterialEditor/Private/MaterialExpressionClasses.cpp b/Engine/Source/Editor/MaterialEditor/Private/MaterialExpressionClasses.cpp
--- a/Engine/Source/Editor/MaterialEditor/Private/MaterialExpressionClasses.cpp
+++ b/Engine/Source/Editor/MaterialEditor/Private/MaterialExpressionClasses.cpp
@@ -182,7 +182,7 @@ void MaterialExpressionClasses::InitMaterialExpressionClasses()
 
 							// Skip adding Strata node to the context menu if Strata is disabled
 							// STRATA_TODO: remove this when Strata becomes the only shading path
-							if ((Class->IsChildOf(UMaterialExpressionStrataBSDF::StaticClass()) || Class->IsChildOf(UMaterialExpressionStrataUtilityBase::StaticClass())) && !Strata::IsStrataEnabled())
+							if ((Class->IsChildOf(UMaterialExpressionSubstrateBSDF::StaticClass()) || Class->IsChildOf(UMaterialExpressionSubstrateUtilityBase::StaticClass())) && !Strata::IsStrataEnabled())
 							{
 								continue;
 							}
diff --git a/Engine/Source/Editor/UnrealEd/Private/MaterialGraphNode.cpp b/Engine/Source/Editor/UnrealEd/Private/MaterialGraphNode.cpp
--- a/Engine/Source/Editor/UnrealEd/Private/MaterialGraphNode.cpp
+++ b/Engine/Source/Editor/UnrealEd/Private/MaterialGraphNode.cpp
@@ -324,7 +324,7 @@ FLinearColor UMaterialGraphNode::GetNodeTitleColor() const
 			return FColor( 0, 128, 128 );
 		}
 	}
-	else if (const UMaterialExpressionStrataBSDF* StrataBSDF = Cast<UMaterialExpressionStrataBSDF>(MaterialExpression))
+	else if (const UMaterialExpressionSubstrateBSDF* StrataBSDF = Cast<UMaterialExpressionSubstrateBSDF>(MaterialExpression))
 	{
 		return FColor(181, 29, 230);
 	}
diff --git a/Engine/Source/Runtime/Engine/Classes/Materials/MaterialExpressionSubstrate.h b/Engine/Source/Runtime/Engine/Classes/Materials/MaterialExpressionSubstrate.h
--- a/Engine/Source/Runtime/Engine/Classes/Materials/MaterialExpressionSubstrate.h
+++ b/Engine/Source/Runtime/Engine/Classes/Materials/MaterialExpressionSubstrate.h
@@ -24,9 +24,9 @@ extern int32 CompileStrataBlendFunction(FMaterialCompiler* Compiler, const int32
 ///////////////////////////////////////////////////////////////////////////////
 // BSDF nodes
 
-// UMaterialExpressionStrataBSDF can only be used for Substrate nodes ouputing StrataData that would need a preview,
+// UMaterialExpressionSubstrateBSDF can only be used for Substrate nodes ouputing StrataData that would need a preview,
 UCLASS(MinimalAPI, collapsecategories, hidecategories = Object, Abstract, DisplayName = "Substrate Expression")
-class UMaterialExpressionStrataBSDF : public UMaterialExpression
+class UMaterialExpressionSubstrateBSDF : public UMaterialExpression
 {
 	GENERATED_UCLASS_BODY()
 
@@ -42,7 +42,7 @@ class UMaterialExpressionStrataBSDF : public UMaterialExpression
 
 
 UCLASS(MinimalAPI, collapsecategories, hidecategories = Object, DisplayName = "Substrate Shading Models")
-class UMaterialExpressionStrataLegacyConversion : public UMaterialExpressionStrataBSDF
+class UMaterialExpressionSubstrateShadingModels : public UMaterialExpressionSubstrateBSDF
 {
 	GENERATED_UCLASS_BODY()
 
@@ -194,7 +194,7 @@ class UMaterialExpressionStrataLegacyConversion : public UMaterialExpressionStra
 };
 
 UCLASS(MinimalAPI, collapsecategories, hidecategories = Object, DisplayName = "Substrate Slab")
-class UMaterialExpressionStrataSlabBSDF : public UMaterialExpressionStrataBSDF
+class UMaterialExpressionSubstrateSlabBSDF : public UMaterialExpressionSubstrateBSDF
 {
 	GENERATED_UCLASS_BODY()
 
@@ -355,7 +355,7 @@ class UMaterialExpressionStrataSlabBSDF : public UMaterialExpressionStrataBSDF
 };
 
 UCLASS(MinimalAPI, collapsecategories, hidecategories = Object, DisplayName = "Substrate Simple Clear Coat")
-class UMaterialExpressionStrataSimpleClearCoatBSDF : public UMaterialExpressionStrataBSDF
+class UMaterialExpressionSubstrateSimpleClearCoatBSDF : public UMaterialExpressionSubstrateBSDF
 {
 	GENERATED_UCLASS_BODY()
 
@@ -416,7 +416,7 @@ class UMaterialExpressionStrataSimpleClearCoatBSDF : public UMaterialExpressionS
 };
 
 UCLASS(MinimalAPI, collapsecategories, hidecategories = Object, DisplayName = "Substrate Volumetric-Fog-Cloud BSDF")
-class UMaterialExpressionStrataVolumetricFogCloudBSDF : public UMaterialExpressionStrataBSDF
+class UMaterialExpressionSubstrateVolumetricFogCloudBSDF : public UMaterialExpressionSubstrateBSDF
 {
 	GENERATED_UCLASS_BODY()
 
@@ -458,7 +458,7 @@ class UMaterialExpressionStrataVolumetricFogCloudBSDF : public UMaterialExpressi
 };
 
 UCLASS(MinimalAPI, collapsecategories, hidecategories = Object, DisplayName = "Substrate Unlit BSDF")
-class UMaterialExpressionStrataUnlitBSDF : public UMaterialExpressionStrataBSDF
+class UMaterialExpressionSubstrateUnlitBSDF : public UMaterialExpressionSubstrateBSDF
 {
 	GENERATED_UCLASS_BODY()
 
@@ -494,7 +494,7 @@ class UMaterialExpressionStrataUnlitBSDF : public UMaterialExpressionStrataBSDF
 };
 
 UCLASS(MinimalAPI, collapsecategories, hidecategories = Object, DisplayName = "Substrate Hair BSDF")
-class UMaterialExpressionStrataHairBSDF : public UMaterialExpressionStrataBSDF
+class UMaterialExpressionSubstrateHairBSDF : public UMaterialExpressionSubstrateBSDF
 {
 	GENERATED_UCLASS_BODY()
 		
@@ -554,7 +554,7 @@ class UMaterialExpressionStrataHairBSDF : public UMaterialExpressionStrataBSDF
 };
 
 UCLASS(MinimalAPI, collapsecategories, hidecategories = Object, DisplayName = "Substrate Eye BSDF")
-class UMaterialExpressionStrataEyeBSDF : public UMaterialExpressionStrataBSDF
+class UMaterialExpressionSubstrateEyeBSDF : public UMaterialExpressionSubstrateBSDF
 {
 	GENERATED_UCLASS_BODY()
 		
@@ -624,7 +624,7 @@ class UMaterialExpressionStrataEyeBSDF : public UMaterialExpressionStrataBSDF
 };
 
 UCLASS(MinimalAPI, collapsecategories, hidecategories = Object, DisplayName = "Substrate Single Layer Water BSDF")
-class UMaterialExpressionStrataSingleLayerWaterBSDF : public UMaterialExpressionStrataBSDF
+class UMaterialExpressionSubstrateSingleLayerWaterBSDF : public UMaterialExpressionSubstrateBSDF
 {
 	GENERATED_UCLASS_BODY()
 
@@ -708,7 +708,7 @@ class UMaterialExpressionStrataSingleLayerWaterBSDF : public UMaterialExpression
 };
 
 UCLASS(MinimalAPI, collapsecategories, hidecategories = Object, DisplayName = "Substrate Light Function")
-class UMaterialExpressionStrataLightFunction : public UMaterialExpressionStrataBSDF
+class UMaterialExpressionSubstrateLightFunction : public UMaterialExpressionSubstrateBSDF
 {
 	GENERATED_UCLASS_BODY()
 
@@ -732,7 +732,7 @@ class UMaterialExpressionStrataLightFunction : public UMaterialExpressionStrataB
 };
 
 UCLASS(MinimalAPI, collapsecategories, hidecategories = Object, DisplayName = "Substrate Post Process")
-class UMaterialExpressionStrataPostProcess : public UMaterialExpressionStrataBSDF
+class UMaterialExpressionSubstratePostProcess : public UMaterialExpressionSubstrateBSDF
 {
 	GENERATED_UCLASS_BODY()
 
@@ -762,7 +762,7 @@ class UMaterialExpressionStrataPostProcess : public UMaterialExpressionStrataBSD
 };
 
 UCLASS(MinimalAPI, collapsecategories, hidecategories = Object, DisplayName = "Substrate UI")
-class UMaterialExpressionStrataUI : public UMaterialExpressionStrataBSDF
+class UMaterialExpressionSubstrateUI : public UMaterialExpressionSubstrateBSDF
 {
 	GENERATED_UCLASS_BODY()
 
@@ -792,7 +792,7 @@ class UMaterialExpressionStrataUI : public UMaterialExpressionStrataBSDF
 };
 
 UCLASS(MinimalAPI, collapsecategories, hidecategories = Object, DisplayName = "Substrate Convert To Decal")
-class UMaterialExpressionStrataConvertToDecal : public UMaterialExpressionStrataBSDF
+class UMaterialExpressionSubstrateConvertToDecal : public UMaterialExpressionSubstrateBSDF
 {
 	GENERATED_UCLASS_BODY()
 
@@ -827,7 +827,7 @@ class UMaterialExpressionStrataConvertToDecal : public UMaterialExpressionStrata
 // Operator nodes
 
 UCLASS(MinimalAPI, collapsecategories, hidecategories = Object, DisplayName = "Substrate Horizontal Blend")
-class UMaterialExpressionStrataHorizontalMixing : public UMaterialExpressionStrataBSDF
+class UMaterialExpressionSubstrateHorizontalMixing : public UMaterialExpressionSubstrateBSDF
 {
 	GENERATED_UCLASS_BODY()
 		
@@ -869,7 +869,7 @@ class UMaterialExpressionStrataHorizontalMixing : public UMaterialExpressionStra
 };
 
 UCLASS(MinimalAPI, collapsecategories, hidecategories = Object, DisplayName = "Substrate Vertical Layer")
-class UMaterialExpressionStrataVerticalLayering : public UMaterialExpressionStrataBSDF
+class UMaterialExpressionSubstrateVerticalLayering : public UMaterialExpressionSubstrateBSDF
 {
 	GENERATED_UCLASS_BODY()
 
@@ -912,7 +912,7 @@ class UMaterialExpressionStrataVerticalLayering : public UMaterialExpressionStra
 };
 
 UCLASS(MinimalAPI, collapsecategories, hidecategories = Object, DisplayName = "Substrate Add")
-class UMaterialExpressionStrataAdd : public UMaterialExpressionStrataBSDF
+class UMaterialExpressionSubstrateAdd : public UMaterialExpressionSubstrateBSDF
 {
 	GENERATED_UCLASS_BODY()
 
@@ -948,7 +948,7 @@ class UMaterialExpressionStrataAdd : public UMaterialExpressionStrataBSDF
 };
 
 UCLASS(MinimalAPI, collapsecategories, hidecategories = Object, DisplayName = "Substrate Coverage Weight")
-class UMaterialExpressionStrataWeight : public UMaterialExpressionStrataBSDF
+class UMaterialExpressionSubstrateWeight : public UMaterialExpressionSubstrateBSDF
 {
 	GENERATED_UCLASS_BODY()
 
@@ -982,13 +982,13 @@ class UMaterialExpressionStrataWeight : public UMaterialExpressionStrataBSDF
 // Utilities
 
 UCLASS(MinimalAPI, collapsecategories, hidecategories = Object, Abstract, DisplayName = "Substrate Utility Base Class")
-class UMaterialExpressionStrataUtilityBase : public UMaterialExpression
+class UMaterialExpressionSubstrateUtilityBase : public UMaterialExpression
 {
 	GENERATED_UCLASS_BODY()
 };
 
 UCLASS(MinimalAPI, collapsecategories, hidecategories = Object, DisplayName = "Substrate Transmittance-To-MeanFreePath")
-class UMaterialExpressionStrataTransmittanceToMFP : public UMaterialExpressionStrataUtilityBase
+class UMaterialExpressionSubstrateTransmittanceToMFP : public UMaterialExpressionSubstrateUtilityBase
 {
 	GENERATED_UCLASS_BODY()
 
@@ -1018,7 +1018,7 @@ class UMaterialExpressionStrataTransmittanceToMFP : public UMaterialExpressionSt
 };
 
 UCLASS(MinimalAPI, collapsecategories, hidecategories = Object, DisplayName = "Substrate Metalness-To-DiffuseColorF0")
-class UMaterialExpressionStrataMetalnessToDiffuseAlbedoF0 : public UMaterialExpressionStrataUtilityBase
+class UMaterialExpressionSubstrateMetalnessToDiffuseAlbedoF0 : public UMaterialExpressionSubstrateUtilityBase
 {
 	GENERATED_UCLASS_BODY()
 
@@ -1053,7 +1053,7 @@ class UMaterialExpressionStrataMetalnessToDiffuseAlbedoF0 : public UMaterialExpr
 };
 
 UCLASS(MinimalAPI, collapsecategories, hidecategories = Object, DisplayName = "Substrate Haziness-To-Secondary-Roughness")
-class UMaterialExpressionStrataHazinessToSecondaryRoughness : public UMaterialExpressionStrataUtilityBase
+class UMaterialExpressionSubstrateHazinessToSecondaryRoughness : public UMaterialExpressionSubstrateUtilityBase
 {
 	GENERATED_UCLASS_BODY()
 
@@ -1082,7 +1082,7 @@ class UMaterialExpressionStrataHazinessToSecondaryRoughness : public UMaterialEx
 };
 
 UCLASS(MinimalAPI, collapsecategories, hidecategories = Object, DisplayName = "Substrate Thin-Film")
-class UMaterialExpressionStrataThinFilm : public UMaterialExpressionStrataUtilityBase
+class UMaterialExpressionSubstrateThinFilm : public UMaterialExpressionSubstrateUtilityBase
 {
 	GENERATED_UCLASS_BODY()
 
diff --git a/Engine/Source/Runtime/Engine/Private/Materials/HLSLMaterialTranslator.cpp b/Engine/Source/Runtime/Engine/Private/Materials/HLSLMaterialTranslator.cpp
--- a/Engine/Source/Runtime/Engine/Private/Materials/HLSLMaterialTranslator.cpp
+++ b/Engine/Source/Runtime/Engine/Private/Materials/HLSLMaterialTranslator.cpp
@@ -4230,7 +4230,7 @@ int32 FHLSLMaterialTranslator::CallExpression(FMaterialExpressionKey ExpressionK
 	// This is automatically handled via the compiler StrataTreeStack.
 	// It means that a node can be blended at multiple point of the graph (allowing acyclic graph instead of tree, e.g. a Slab can be used into multiple input).
 	// It is worth noting that only strata BSDF can be duplicated today according to StrataTreeStack.
-	const bool bExpressionIsStrataBSDF = ExpressionKey.Expression && ExpressionKey.Expression->IsA<UMaterialExpressionStrataSlabBSDF>();
+	const bool bExpressionIsStrataBSDF = ExpressionKey.Expression && ExpressionKey.Expression->IsA<UMaterialExpressionSubstrateSlabBSDF>();
 	const bool bExpressionIsStrataFullySimplified = CurrentStrataCompilationContext == EStrataCompilationContext::SCC_FullySimplified && ExpressionKey.Expression && ExpressionKey.Expression->IsResultStrataMaterial(ExpressionKey.OutputIndex);
 
 	// Check if this expression has already been translated.
diff --git a/Engine/Source/Runtime/Engine/Private/Materials/Material.cpp b/Engine/Source/Runtime/Engine/Private/Materials/Material.cpp
--- a/Engine/Source/Runtime/Engine/Private/Materials/Material.cpp
+++ b/Engine/Source/Runtime/Engine/Private/Materials/Material.cpp
@@ -3260,7 +3260,7 @@ void UMaterial::ConvertMaterialToStrataMaterial()
 
 	bool bInvalidateShader = false;
 	bool bRelinkCustomOutputNodes = false;
-	UMaterialExpressionStrataLegacyConversion* ConvertNode = nullptr;
+	UMaterialExpressionSubstrateShadingModels* ConvertNode = nullptr;
 	// Connect all the legacy pin into the conversion node
 	if (bUseMaterialAttributes && EditorOnly->MaterialAttributes.Expression && !EditorOnly->MaterialAttributes.Expression->IsResultStrataMaterial(EditorOnly->MaterialAttributes.OutputIndex)) // M_Rifle cause issues there
 	{
@@ -3312,7 +3312,7 @@ void UMaterial::ConvertMaterialToStrataMaterial()
 			}
 		}
 
-		ConvertNode = NewObject<UMaterialExpressionStrataLegacyConversion>(this);
+		ConvertNode = NewObject<UMaterialExpressionSubstrateShadingModels>(this);
 		SetPosXAndMoveReferenceToTheRight(ConvertNode);
 		ConvertNode->BaseColor.Connect(0, BreakMatAtt);
 		ConvertNode->Metallic.Connect(1, BreakMatAtt);
@@ -3363,7 +3363,7 @@ void UMaterial::ConvertMaterialToStrataMaterial()
 			// control on this, but only DefaultLit should be supported.
 
 			// Now pass through the convert to decal node, which flag the material as SSM_Decal, which will set the domain to Decal.
-			UMaterialExpressionStrataConvertToDecal* ConvertToDecalNode = NewObject<UMaterialExpressionStrataConvertToDecal>(this);
+			UMaterialExpressionSubstrateConvertToDecal* ConvertToDecalNode = NewObject<UMaterialExpressionSubstrateConvertToDecal>(this);
 			ReplaceNodeAndMoveToTheRight(ConvertNode, ConvertToDecalNode);
 			ConvertToDecalNode->DecalMaterial.Connect(0, ConvertNode);
 
@@ -3401,7 +3401,7 @@ void UMaterial::ConvertMaterialToStrataMaterial()
 		else
 		{
 			// Or if it cannot be found, a slab node
-			UMaterialExpressionStrataSlabBSDF* SlabNode = NewObject<UMaterialExpressionStrataSlabBSDF>(this);
+			UMaterialExpressionSubstrateSlabBSDF* SlabNode = NewObject<UMaterialExpressionSubstrateSlabBSDF>(this);
 			SetPosXAndMoveReferenceToTheRight(SlabNode);
 			EditorOnly->FrontMaterial.Connect(0, SlabNode);
 		}
@@ -3421,14 +3421,14 @@ void UMaterial::ConvertMaterialToStrataMaterial()
 					// For this special case, using two slabs to create a clear coat material with separated top and bottom normal. 
 
 					// Create metalness to Slab parameterisation conveersion node
-					UMaterialExpressionStrataMetalnessToDiffuseAlbedoF0* StrataMetalnessToDiffuseAlbedoF0 = NewObject<UMaterialExpressionStrataMetalnessToDiffuseAlbedoF0>(this);
+					UMaterialExpressionSubstrateMetalnessToDiffuseAlbedoF0* StrataMetalnessToDiffuseAlbedoF0 = NewObject<UMaterialExpressionSubstrateMetalnessToDiffuseAlbedoF0>(this);
 					SetPosXAndMoveReferenceToTheRight(StrataMetalnessToDiffuseAlbedoF0);
 					MoveConnectionTo(EditorOnly->BaseColor, StrataMetalnessToDiffuseAlbedoF0, 0);					// BaseColor
 					MoveConnectionTo(EditorOnly->Metallic, StrataMetalnessToDiffuseAlbedoF0, 1);					// Metallic
 					MoveConnectionTo(EditorOnly->Specular, StrataMetalnessToDiffuseAlbedoF0, 2);					// Specular
 					
 					// Top slab BSDF as a simple Disney material
-					UMaterialExpressionStrataSlabBSDF* BottomSlabBSDF = NewObject<UMaterialExpressionStrataSlabBSDF>(this);
+					UMaterialExpressionSubstrateSlabBSDF* BottomSlabBSDF = NewObject<UMaterialExpressionSubstrateSlabBSDF>(this);
 					SetPosXAndMoveReferenceToTheRight(BottomSlabBSDF);
 					BottomSlabBSDF->GetInput(0)->Connect(0, StrataMetalnessToDiffuseAlbedoF0);
 					BottomSlabBSDF->GetInput(1)->Connect(1, StrataMetalnessToDiffuseAlbedoF0);
@@ -3441,7 +3441,7 @@ void UMaterial::ConvertMaterialToStrataMaterial()
 					CopyConnectionTo(*ClearCoatBottomNormalOutput->GetInput(0), BottomSlabBSDF, 5);// ClearColorBottomNormal -> BottomSlabBSDF.Normal
 
 					// Now weight the top base material by opacity.
-					UMaterialExpressionStrataSlabBSDF* TopSlabBSDF = NewObject<UMaterialExpressionStrataSlabBSDF>(this);
+					UMaterialExpressionSubstrateSlabBSDF* TopSlabBSDF = NewObject<UMaterialExpressionSubstrateSlabBSDF>(this);
 					TopSlabBSDF->MaterialExpressionEditorX = BottomSlabBSDF->MaterialExpressionEditorX;
 					TopSlabBSDF->MaterialExpressionEditorY = BottomSlabBSDF->MaterialExpressionEditorY + 650;
 					MoveConnectionTo(EditorOnly->EmissiveColor, TopSlabBSDF, 10);				// Emissive
@@ -3466,19 +3466,19 @@ void UMaterial::ConvertMaterialToStrataMaterial()
 					UMaterialExpressionConstant* Constant075 = NewObject<UMaterialExpressionConstant>(this);
 					ReplaceNodeAndMoveToTheRight(TopSlabBSDF, Constant075);
 					Constant075->R = 0.75f;
-					UMaterialExpressionStrataTransmittanceToMFP* TransToMDFP = NewObject<UMaterialExpressionStrataTransmittanceToMFP>(this);
+					UMaterialExpressionSubstrateTransmittanceToMFP* TransToMDFP = NewObject<UMaterialExpressionSubstrateTransmittanceToMFP>(this);
 					ReplaceNodeAndMoveToTheRight(TopSlabBSDF, TransToMDFP);
 					TransToMDFP->GetInput(0)->Connect(0, Constant075);
 					TopSlabBSDF->GetInput(7)->Connect(0, TransToMDFP);							// MFP -> MFP
 					TopSlabBSDF->GetInput(13)->Connect(1, TransToMDFP);							// Thickness -> Thickness
 
 					// Now weight the top base material by ClearCoat
-					UMaterialExpressionStrataWeight* TopSlabBSDFWithCoverage = NewObject<UMaterialExpressionStrataWeight>(this);
+					UMaterialExpressionSubstrateWeight* TopSlabBSDFWithCoverage = NewObject<UMaterialExpressionSubstrateWeight>(this);
 					SetPosXAndMoveReferenceToTheRight(TopSlabBSDFWithCoverage);
 					TopSlabBSDFWithCoverage->GetInput(0)->Connect(0, TopSlabBSDF);				// TopSlabBSDF -> A
 					MoveConnectionTo(EditorOnly->ClearCoat, TopSlabBSDFWithCoverage, 1);		// ClearCoat -> Weight
 
-					UMaterialExpressionStrataVerticalLayering* VerticalLayering = NewObject<UMaterialExpressionStrataVerticalLayering>(this);
+					UMaterialExpressionSubstrateVerticalLayering* VerticalLayering = NewObject<UMaterialExpressionSubstrateVerticalLayering>(this);
 					SetPosXAndMoveReferenceToTheRight(VerticalLayering);
 					VerticalLayering->GetInput(0)->Connect(0, TopSlabBSDFWithCoverage);			// Top -> Top
 					VerticalLayering->GetInput(1)->Connect(0, BottomSlabBSDF);					// Bottom -> Base
@@ -3491,7 +3491,7 @@ void UMaterial::ConvertMaterialToStrataMaterial()
 			
 			if (!bClearCoatConversionDone)
 			{
-				ConvertNode = NewObject<UMaterialExpressionStrataLegacyConversion>(this);
+				ConvertNode = NewObject<UMaterialExpressionSubstrateShadingModels>(this);
 				SetPosXAndMoveReferenceToTheRight(ConvertNode);
 				ConvertNode->SubsurfaceProfile = bRequireNoSubsurfaceProfile ? nullptr : SubsurfaceProfile;
 				MoveConnectionTo(EditorOnly->BaseColor, ConvertNode, 0);
@@ -3546,7 +3546,7 @@ void UMaterial::ConvertMaterialToStrataMaterial()
 		}
 		else if (MaterialDomain == MD_Volume)
 		{
-			UMaterialExpressionStrataVolumetricFogCloudBSDF* VolBSDF = NewObject<UMaterialExpressionStrataVolumetricFogCloudBSDF>(this);
+			UMaterialExpressionSubstrateVolumetricFogCloudBSDF* VolBSDF = NewObject<UMaterialExpressionSubstrateVolumetricFogCloudBSDF>(this);
 			SetPosXAndMoveReferenceToTheRight(VolBSDF);
 			MoveConnectionTo(EditorOnly->BaseColor, VolBSDF, 0);		// Albedo
 			MoveConnectionTo(EditorOnly->SubsurfaceColor, VolBSDF, 1);	// Extinction
@@ -3565,7 +3565,7 @@ void UMaterial::ConvertMaterialToStrataMaterial()
 			ShadingModels.AddShadingModel(MSM_Unlit);
 
 			// Only Emissive & Opacity are valid input for PostProcess material
-			UMaterialExpressionStrataLightFunction* LightFunctionNode = NewObject<UMaterialExpressionStrataLightFunction>(this);
+			UMaterialExpressionSubstrateLightFunction* LightFunctionNode = NewObject<UMaterialExpressionSubstrateLightFunction>(this);
 			SetPosXAndMoveReferenceToTheRight(LightFunctionNode);
 			MoveConnectionTo(EditorOnly->EmissiveColor, LightFunctionNode, 0);
 
@@ -3584,7 +3584,7 @@ void UMaterial::ConvertMaterialToStrataMaterial()
 				BlendMode = BLEND_Opaque;
 			}
 
-			UMaterialExpressionStrataPostProcess* PostProcNode = NewObject<UMaterialExpressionStrataPostProcess>(this);
+			UMaterialExpressionSubstratePostProcess* PostProcNode = NewObject<UMaterialExpressionSubstratePostProcess>(this);
 			SetPosXAndMoveReferenceToTheRight(PostProcNode);
 
 			MoveConnectionTo(EditorOnly->EmissiveColor, PostProcNode, 0);
@@ -3600,7 +3600,7 @@ void UMaterial::ConvertMaterialToStrataMaterial()
 			ShadingModels.ClearShadingModels();
 			ShadingModels.AddShadingModel(MSM_DefaultLit);
 
-			ConvertNode = NewObject<UMaterialExpressionStrataLegacyConversion>(this);
+			ConvertNode = NewObject<UMaterialExpressionSubstrateShadingModels>(this);
 			SetPosXAndMoveReferenceToTheRight(ConvertNode);
 			MoveConnectionTo(EditorOnly->BaseColor, ConvertNode, 0);
 			MoveConnectionTo(EditorOnly->Metallic, ConvertNode, 1);
@@ -3620,7 +3620,7 @@ void UMaterial::ConvertMaterialToStrataMaterial()
 			check(ShadingModels.CountShadingModels() == 1);
 
 			// Now pass through the convert to decal node, which flag the material as SSM_Decal, which will set the domain to Decal.
-			UMaterialExpressionStrataConvertToDecal* ConvertToDecalNode= NewObject<UMaterialExpressionStrataConvertToDecal>(this);
+			UMaterialExpressionSubstrateConvertToDecal* ConvertToDecalNode= NewObject<UMaterialExpressionSubstrateConvertToDecal>(this);
 			ReplaceNodeAndMoveToTheRight(ConvertNode, ConvertToDecalNode);
 			ConvertToDecalNode->DecalMaterial.Connect(0, ConvertNode);
 
@@ -3634,7 +3634,7 @@ void UMaterial::ConvertMaterialToStrataMaterial()
 			ShadingModels.ClearShadingModels();
 			ShadingModels.AddShadingModel(MSM_Unlit);
 
-			UMaterialExpressionStrataUI* UINode = NewObject<UMaterialExpressionStrataUI>(this);
+			UMaterialExpressionSubstrateUI* UINode = NewObject<UMaterialExpressionSubstrateUI>(this);
 			SetPosXAndMoveReferenceToTheRight(UINode);
 			MoveConnectionTo(EditorOnly->EmissiveColor, UINode, 0);
 			CopyConnectionTo(EditorOnly->Opacity, UINode, 1);	// We only copy, to keep Opacity on the root node in case BLEND_AlphaComposite is selected.
diff --git a/Engine/Source/Runtime/Engine/Private/Materials/MaterialExpressions.cpp b/Engine/Source/Runtime/Engine/Private/Materials/MaterialExpressions.cpp
--- a/Engine/Source/Runtime/Engine/Private/Materials/MaterialExpressions.cpp
+++ b/Engine/Source/Runtime/Engine/Private/Materials/MaterialExpressions.cpp
@@ -2084,7 +2084,7 @@ void UMaterialExpression::ConnectToPreviewMaterial(UMaterial* InMaterial, int32
 				BreakMatAtt->MaterialAttributes.Connect(OutputIndex, this);
 
 				// Get BaseColor as unlit color preview from material attribute
-				UMaterialExpressionStrataUnlitBSDF* UnlitBSDF = NewObject<UMaterialExpressionStrataUnlitBSDF>(this);
+				UMaterialExpressionSubstrateUnlitBSDF* UnlitBSDF = NewObject<UMaterialExpressionSubstrateUnlitBSDF>(this);
 				UnlitBSDF->EmissiveColor.Connect(0, BreakMatAtt);
 
 				FExpressionInput* MaterialInput = InMaterial->GetExpressionInputForProperty(MP_FrontMaterial);
@@ -2094,7 +2094,7 @@ void UMaterialExpression::ConnectToPreviewMaterial(UMaterial* InMaterial, int32
 			else
 			{
 				InMaterial->SetShadingModel(MSM_Unlit);
-				UMaterialExpressionStrataUnlitBSDF* UnlitBSDF = NewObject<UMaterialExpressionStrataUnlitBSDF>(this);
+				UMaterialExpressionSubstrateUnlitBSDF* UnlitBSDF = NewObject<UMaterialExpressionSubstrateUnlitBSDF>(this);
 				UnlitBSDF->EmissiveColor.Connect(OutputIndex, this);
 
 				FExpressionInput* MaterialInput = InMaterial->GetExpressionInputForProperty(MP_FrontMaterial);
@@ -23205,7 +23205,7 @@ static int32 CompileWithDefaultTangentWS(class FMaterialCompiler* Compiler, FExp
 
 #endif // WITH_EDITOR
 
-UMaterialExpressionStrataLegacyConversion::UMaterialExpressionStrataLegacyConversion(const FObjectInitializer& ObjectInitializer)
+UMaterialExpressionSubstrateShadingModels::UMaterialExpressionSubstrateShadingModels(const FObjectInitializer& ObjectInitializer)
 	: Super(ObjectInitializer)
 {
 	struct FConstructorStatics
@@ -23223,7 +23223,7 @@ UMaterialExpressionStrataLegacyConversion::UMaterialExpressionStrataLegacyConver
 }
 
 #if WITH_EDITOR
-int32 UMaterialExpressionStrataLegacyConversion::Compile(class FMaterialCompiler* Compiler, int32 OutputIndex)
+int32 UMaterialExpressionSubstrateShadingModels::Compile(class FMaterialCompiler* Compiler, int32 OutputIndex)
 {
 	int32 RoughnessCodeChunk = CompileWithDefaultFloat1(Compiler, Roughness, 0.5f);
 	int32 AnisotropyCodeChunk = CompileWithDefaultFloat1(Compiler, Anisotropy, 0.0f);
@@ -23344,7 +23344,7 @@ int32 UMaterialExpressionStrataLegacyConversion::Compile(class FMaterialCompiler
 }
 
 
-void UMaterialExpressionStrataLegacyConversion::PostEditChangeProperty(FPropertyChangedEvent& PropertyChangedEvent)
+void UMaterialExpressionSubstrateShadingModels::PostEditChangeProperty(FPropertyChangedEvent& PropertyChangedEvent)
 {
 	if (GraphNode)
 	{
@@ -23353,17 +23353,17 @@ void UMaterialExpressionStrataLegacyConversion::PostEditChangeProperty(FProperty
 	Super::PostEditChangeProperty(PropertyChangedEvent);
 }
 
-void UMaterialExpressionStrataLegacyConversion::GetCaption(TArray<FString>& OutCaptions) const
+void UMaterialExpressionSubstrateShadingModels::GetCaption(TArray<FString>& OutCaptions) const
 {
 	OutCaptions.Add(TEXT("Substrate Shading Models"));
 }
 
-uint32 UMaterialExpressionStrataLegacyConversion::GetOutputType(int32 OutputIndex)
+uint32 UMaterialExpressionSubstrateShadingModels::GetOutputType(int32 OutputIndex)
 {
 	return MCT_Strata;
 }
 
-uint32 UMaterialExpressionStrataLegacyConversion::GetInputType(int32 InputIndex)
+uint32 UMaterialExpressionSubstrateShadingModels::GetInputType(int32 InputIndex)
 {
 	if (InputIndex == 0)	   return MCT_Float3; // BaseColor
 	else if (InputIndex == 1)  return MCT_Float1; // Metallic
@@ -23391,7 +23391,7 @@ uint32 UMaterialExpressionStrataLegacyConversion::GetInputType(int32 InputIndex)
 	return MCT_Float1;
 }
 
-FName UMaterialExpressionStrataLegacyConversion::GetInputName(int32 InputIndex) const
+FName UMaterialExpressionSubstrateShadingModels::GetInputName(int32 InputIndex) const
 {
 	const bool bShadingModelFromExpression = ShadingModel.IsConnected();
 
@@ -23459,7 +23459,7 @@ FName UMaterialExpressionStrataLegacyConversion::GetInputName(int32 InputIndex)
 	return TEXT("Unknown");
 }
 
-void UMaterialExpressionStrataLegacyConversion::GetConnectorToolTip(int32 InputIndex, int32 OutputIndex, TArray<FString>& OutToolTip)
+void UMaterialExpressionSubstrateShadingModels::GetConnectorToolTip(int32 InputIndex, int32 OutputIndex, TArray<FString>& OutToolTip)
 {
 	if (OutputIndex == 0)
 	{
@@ -23469,12 +23469,12 @@ void UMaterialExpressionStrataLegacyConversion::GetConnectorToolTip(int32 InputI
 	Super::GetConnectorToolTip(InputIndex, INDEX_NONE, OutToolTip);
 }
 
-bool UMaterialExpressionStrataLegacyConversion::IsResultStrataMaterial(int32 OutputIndex)
+bool UMaterialExpressionSubstrateShadingModels::IsResultStrataMaterial(int32 OutputIndex)
 {
 	return true;
 }
 
-void UMaterialExpressionStrataLegacyConversion::GatherStrataMaterialInfo(FStrataMaterialInfo& StrataMaterialInfo, int32 OutputIndex)
+void UMaterialExpressionSubstrateShadingModels::GatherStrataMaterialInfo(FStrataMaterialInfo& StrataMaterialInfo, int32 OutputIndex)
 {	
 	// Track connected input
 	if (BaseColor.IsConnected())			{ StrataMaterialInfo.AddPropertyConnected(MP_BaseColor); }
@@ -23515,7 +23515,7 @@ void UMaterialExpressionStrataLegacyConversion::GatherStrataMaterialInfo(FStrata
 	}
 }
 
-FStrataOperator* UMaterialExpressionStrataLegacyConversion::StrataGenerateMaterialTopologyTree(class FMaterialCompiler* Compiler, class UMaterialExpression* Parent, int32 OutputIndex)
+FStrataOperator* UMaterialExpressionSubstrateShadingModels::StrataGenerateMaterialTopologyTree(class FMaterialCompiler* Compiler, class UMaterialExpression* Parent, int32 OutputIndex)
 {
 	// Note Thickness has no meaning/usage in the context of StrataLegacyConversionNode
 	int32 ThicknessIndex = Compiler->StrataThicknessStackGetThicknessIndex();
@@ -23536,7 +23536,7 @@ FStrataOperator* UMaterialExpressionStrataLegacyConversion::StrataGenerateMateri
 	// Get the shading models resulting from the UMaterial::RebuildShadingModelField().
 	FMaterialShadingModelField ShadingModels = Compiler->GetMaterialShadingModels();
 
-	// Logic about shading models and complexity should match UMaterialExpressionStrataLegacyConversion::Compile.
+	// Logic about shading models and complexity should match UMaterialExpressionSubstrateShadingModels::Compile.
 	const bool bHasShadingModelFromExpression = ShadingModel.IsConnected(); // We keep HasShadingModelFromExpression in case all shading models cannot be safely recovered from material functions.
 	if ((ShadingModels.CountShadingModels() > 1) || bHasShadingModelFromExpression) 
 	{
@@ -23618,12 +23618,12 @@ FStrataOperator* UMaterialExpressionStrataLegacyConversion::StrataGenerateMateri
 	}
 }
 
-bool UMaterialExpressionStrataLegacyConversion::HasSSS() const
+bool UMaterialExpressionSubstrateShadingModels::HasSSS() const
 {
 	return SubsurfaceProfile != nullptr;
 }
 
-bool UMaterialExpressionStrataLegacyConversion::HasAnisotropy() const
+bool UMaterialExpressionSubstrateShadingModels::HasAnisotropy() const
 {
 	return Anisotropy.IsConnected();
 }
@@ -23631,13 +23631,13 @@ bool UMaterialExpressionStrataLegacyConversion::HasAnisotropy() const
 #endif // WITH_EDITOR
 
 
-UMaterialExpressionStrataBSDF::UMaterialExpressionStrataBSDF(const FObjectInitializer& ObjectInitializer)
+UMaterialExpressionSubstrateBSDF::UMaterialExpressionSubstrateBSDF(const FObjectInitializer& ObjectInitializer)
 	: Super(ObjectInitializer)
 {
 }
 
 #if WITH_EDITOR
-int32 UMaterialExpressionStrataBSDF::CompilePreview(class FMaterialCompiler* Compiler, int32 OutputIndex)
+int32 UMaterialExpressionSubstrateBSDF::CompilePreview(class FMaterialCompiler* Compiler, int32 OutputIndex)
 {
 	// Compile the StrataData output.
 	int32 StrataDataCodeChunk = Compile(Compiler, OutputIndex);
@@ -23647,7 +23647,7 @@ int32 UMaterialExpressionStrataBSDF::CompilePreview(class FMaterialCompiler* Com
 }
 #endif
 
-UMaterialExpressionStrataSlabBSDF::UMaterialExpressionStrataSlabBSDF(const FObjectInitializer& ObjectInitializer)
+UMaterialExpressionSubstrateSlabBSDF::UMaterialExpressionSubstrateSlabBSDF(const FObjectInitializer& ObjectInitializer)
 	: Super(ObjectInitializer)
 	, bUseSSSDiffusion(true)
 {
@@ -23664,7 +23664,7 @@ UMaterialExpressionStrataSlabBSDF::UMaterialExpressionStrataSlabBSDF(const FObje
 
 #if WITH_EDITOR
 FName CreateSpecularProfileParameterName(USpecularProfile* InProfile);
-int32 UMaterialExpressionStrataSlabBSDF::Compile(class FMaterialCompiler* Compiler, int32 OutputIndex)
+int32 UMaterialExpressionSubstrateSlabBSDF::Compile(class FMaterialCompiler* Compiler, int32 OutputIndex)
 {
 	FStrataOperator& StrataOperator = Compiler->StrataCompilationGetOperator(Compiler->StrataTreeStackGetPathUniqueId());
 
@@ -23785,9 +23785,9 @@ int32 UMaterialExpressionStrataSlabBSDF::Compile(class FMaterialCompiler* Compil
 	return OutputCodeChunk;
 }
 
-UMaterialExpressionStrataSlabBSDF::FComplexity UMaterialExpressionStrataSlabBSDF::GetComplexity() const
+UMaterialExpressionSubstrateSlabBSDF::FComplexity UMaterialExpressionSubstrateSlabBSDF::GetComplexity() const
 {
-	UMaterialExpressionStrataSlabBSDF::FComplexity Complexity;
+	UMaterialExpressionSubstrateSlabBSDF::FComplexity Complexity;
 	Complexity.bStrataMaterialIsComplexSpecial = HasGlint() || HasSpecularProfile();
 	Complexity.bStrataMaterialIsComplex = HasAnisotropy();
 	Complexity.bStrataMaterialIsSingle = HasEdgeColor() || HasFuzz() || HasSecondRoughness() || HasMFPPluggedIn() || HasSSS();
@@ -23799,7 +23799,7 @@ UMaterialExpressionStrataSlabBSDF::FComplexity UMaterialExpressionStrataSlabBSDF
 	return Complexity;
 }
 
-static FString GetSlabComplexityString(UMaterialExpressionStrataSlabBSDF::FComplexity Complexity)
+static FString GetSlabComplexityString(UMaterialExpressionSubstrateSlabBSDF::FComplexity Complexity)
 {
 	FString ComplexityString = TEXT("Simple");
 	if (Complexity.bStrataMaterialIsComplexSpecial)
@@ -23817,13 +23817,13 @@ static FString GetSlabComplexityString(UMaterialExpressionStrataSlabBSDF::FCompl
 	return ComplexityString;
 }
 
-void UMaterialExpressionStrataSlabBSDF::GetCaption(TArray<FString>& OutCaptions) const
+void UMaterialExpressionSubstrateSlabBSDF::GetCaption(TArray<FString>& OutCaptions) const
 {
 	// The node complexity is manually maintained to match FStrataCompilationContext::StrataGenerateDerivedMaterialOperatorData and shaders.
 	OutCaptions.Add(TEXT("Substrate Slab BSDF - ") + GetSlabComplexityString(GetComplexity()));
 }
 
-void UMaterialExpressionStrataSlabBSDF::GetExpressionToolTip(TArray<FString>& OutToolTip)
+void UMaterialExpressionSubstrateSlabBSDF::GetExpressionToolTip(TArray<FString>& OutToolTip)
 {
 	OutToolTip.Add(TEXT("Substrate Slab BSDF"));
 	OutToolTip.Add(TEXT("Complexity = ") + GetSlabComplexityString(GetComplexity()));
@@ -23834,12 +23834,12 @@ void UMaterialExpressionStrataSlabBSDF::GetExpressionToolTip(TArray<FString>& Ou
 	OutToolTip.Add(TEXT(" - ComplexSpecial means the Slab is using more advanced features such as glints or specular LUT. This is the most expenssive shading path."));
 }
 
-uint32 UMaterialExpressionStrataSlabBSDF::GetOutputType(int32 OutputIndex)
+uint32 UMaterialExpressionSubstrateSlabBSDF::GetOutputType(int32 OutputIndex)
 {
 	return MCT_Strata;
 }
 
-uint32 UMaterialExpressionStrataSlabBSDF::GetInputType(int32 InputIndex)
+uint32 UMaterialExpressionSubstrateSlabBSDF::GetInputType(int32 InputIndex)
 {
 	if (InputIndex == 0)
 	{
@@ -23918,7 +23918,7 @@ uint32 UMaterialExpressionStrataSlabBSDF::GetInputType(int32 InputIndex)
 	return MCT_Float1;
 }
 
-FName UMaterialExpressionStrataSlabBSDF::GetInputName(int32 InputIndex) const
+FName UMaterialExpressionSubstrateSlabBSDF::GetInputName(int32 InputIndex) const
 {
 	if (InputIndex == 0)
 	{
@@ -23996,7 +23996,7 @@ FName UMaterialExpressionStrataSlabBSDF::GetInputName(int32 InputIndex) const
 	return TEXT("Unknown");
 }
 
-void UMaterialExpressionStrataSlabBSDF::GetConnectorToolTip(int32 InputIndex, int32 OutputIndex, TArray<FString>& OutToolTip)
+void UMaterialExpressionSubstrateSlabBSDF::GetConnectorToolTip(int32 InputIndex, int32 OutputIndex, TArray<FString>& OutToolTip)
 {
 	if (OutputIndex == 0)
 	{
@@ -24007,12 +24007,12 @@ void UMaterialExpressionStrataSlabBSDF::GetConnectorToolTip(int32 InputIndex, in
 	Super::GetConnectorToolTip(InputIndex, INDEX_NONE, OutToolTip);
 }
 
-bool UMaterialExpressionStrataSlabBSDF::IsResultStrataMaterial(int32 OutputIndex)
+bool UMaterialExpressionSubstrateSlabBSDF::IsResultStrataMaterial(int32 OutputIndex)
 {
 	return true;
 }
 
-void UMaterialExpressionStrataSlabBSDF::GatherStrataMaterialInfo(FStrataMaterialInfo& StrataMaterialInfo, int32 OutputIndex)
+void UMaterialExpressionSubstrateSlabBSDF::GatherStrataMaterialInfo(FStrataMaterialInfo& StrataMaterialInfo, int32 OutputIndex)
 {
 	// Track connected inputs
 	if (DiffuseAlbedo.IsConnected())		{ StrataMaterialInfo.AddPropertyConnected(MP_DiffuseColor); }
@@ -24048,7 +24048,7 @@ void UMaterialExpressionStrataSlabBSDF::GatherStrataMaterialInfo(FStrataMaterial
 	}
 }
 
-FStrataOperator* UMaterialExpressionStrataSlabBSDF::StrataGenerateMaterialTopologyTree(class FMaterialCompiler* Compiler, class UMaterialExpression* Parent, int32 OutputIndex)
+FStrataOperator* UMaterialExpressionSubstrateSlabBSDF::StrataGenerateMaterialTopologyTree(class FMaterialCompiler* Compiler, class UMaterialExpression* Parent, int32 OutputIndex)
 {
 	FStrataOperator& StrataOperator = Compiler->StrataCompilationRegisterOperator(STRATA_OPERATOR_BSDF, Compiler->StrataTreeStackGetPathUniqueId(), this, Parent, Compiler->StrataTreeStackGetParentPathUniqueId());
 	StrataOperator.BSDFType = STRATA_BSDF_TYPE_SLAB;
@@ -24064,53 +24064,53 @@ FStrataOperator* UMaterialExpressionStrataSlabBSDF::StrataGenerateMaterialTopolo
 	return &StrataOperator;
 }
 
-bool UMaterialExpressionStrataSlabBSDF::HasSSS() const
+bool UMaterialExpressionSubstrateSlabBSDF::HasSSS() const
 {
 	return SubsurfaceProfile != nullptr || SSSMFP.IsConnected();
 }
 
-bool UMaterialExpressionStrataSlabBSDF::HasSSSProfile() const
+bool UMaterialExpressionSubstrateSlabBSDF::HasSSSProfile() const
 {
 	return SubsurfaceProfile != nullptr;
 }
 
-bool UMaterialExpressionStrataSlabBSDF::HasMFPPluggedIn() const
+bool UMaterialExpressionSubstrateSlabBSDF::HasMFPPluggedIn() const
 {
 	return SSSMFP.IsConnected();
 }
 
-bool UMaterialExpressionStrataSlabBSDF::HasEdgeColor() const
+bool UMaterialExpressionSubstrateSlabBSDF::HasEdgeColor() const
 {
 	return F90.IsConnected();
 }
 
-bool UMaterialExpressionStrataSlabBSDF::HasFuzz() const
+bool UMaterialExpressionSubstrateSlabBSDF::HasFuzz() const
 {
 	return FuzzAmount.IsConnected();
 }
 
-bool UMaterialExpressionStrataSlabBSDF::HasFuzzRoughness() const
+bool UMaterialExpressionSubstrateSlabBSDF::HasFuzzRoughness() const
 {
 	return FuzzRoughness.IsConnected();
 }
 
-bool UMaterialExpressionStrataSlabBSDF::HasSecondRoughness() const
+bool UMaterialExpressionSubstrateSlabBSDF::HasSecondRoughness() const
 {
 	return SecondRoughnessWeight.IsConnected();
 }
 
-bool UMaterialExpressionStrataSlabBSDF::HasAnisotropy() const
+bool UMaterialExpressionSubstrateSlabBSDF::HasAnisotropy() const
 {
 	return Anisotropy.IsConnected();
 }
 
-bool UMaterialExpressionStrataSlabBSDF::HasGlint() const
+bool UMaterialExpressionSubstrateSlabBSDF::HasGlint() const
 {
 	// Use IsGlintEnabled to not promote to ComplexSpecial render path if glint rendering is not enabled and avoid registering such possibility at runtime.
 	return GlintValue.IsConnected() && Strata::IsGlintEnabled();
 }
 
-bool UMaterialExpressionStrataSlabBSDF::HasSpecularProfile() const
+bool UMaterialExpressionSubstrateSlabBSDF::HasSpecularProfile() const
 {
 	// Use IsSpecularLUTEnabled to not promote to ComplexSpecial render path if glint rendering is not enabled and avoid registering such possibility at runtime.
 	return SpecularProfile != nullptr && Strata::IsSpecularProfileEnabled();
@@ -24120,7 +24120,7 @@ bool UMaterialExpressionStrataSlabBSDF::HasSpecularProfile() const
 
 
 
-UMaterialExpressionStrataSimpleClearCoatBSDF::UMaterialExpressionStrataSimpleClearCoatBSDF(const FObjectInitializer& ObjectInitializer)
+UMaterialExpressionSubstrateSimpleClearCoatBSDF::UMaterialExpressionSubstrateSimpleClearCoatBSDF(const FObjectInitializer& ObjectInitializer)
 	: Super(ObjectInitializer)
 {
 	struct FConstructorStatics
@@ -24135,7 +24135,7 @@ UMaterialExpressionStrataSimpleClearCoatBSDF::UMaterialExpressionStrataSimpleCle
 }
 
 #if WITH_EDITOR
-int32 UMaterialExpressionStrataSimpleClearCoatBSDF::Compile(class FMaterialCompiler* Compiler, int32 OutputIndex)
+int32 UMaterialExpressionSubstrateSimpleClearCoatBSDF::Compile(class FMaterialCompiler* Compiler, int32 OutputIndex)
 {
 	const float DefaultSpecular = 0.5f;
 	const float DefaultF0 = DielectricSpecularToF0(DefaultSpecular);
@@ -24166,7 +24166,7 @@ int32 UMaterialExpressionStrataSimpleClearCoatBSDF::Compile(class FMaterialCompi
 		CompileWithDefaultFloat3(Compiler, EmissiveColor, 0.0f, 0.0f, 0.0f),
 		CompileWithDefaultFloat1(Compiler, ClearCoatRoughness, 0.1f),
 		CompileWithDefaultFloat1(Compiler, ClearCoatCoverage, 1.0f),
-		Compiler->Constant(1.0f),								// SecondRoughnessAsSimpleClearCoat == true for UMaterialExpressionStrataSimpleClearCoatBSDF
+		Compiler->Constant(1.0f),								// SecondRoughnessAsSimpleClearCoat == true for UMaterialExpressionSubstrateSimpleClearCoatBSDF
 		Compiler->Constant(0.0f),								// FuzzAmount
 		Compiler->Constant3(0.0f, 0.0f, 0.0f),					// FuzzColor
 		RoughnessCodeChunk,										// FuzzRoughness
@@ -24183,17 +24183,17 @@ int32 UMaterialExpressionStrataSimpleClearCoatBSDF::Compile(class FMaterialCompi
 	return OutputCodeChunk;
 }
 
-void UMaterialExpressionStrataSimpleClearCoatBSDF::GetCaption(TArray<FString>& OutCaptions) const
+void UMaterialExpressionSubstrateSimpleClearCoatBSDF::GetCaption(TArray<FString>& OutCaptions) const
 {
 	OutCaptions.Add(TEXT("Substrate Simple Clear Coat"));
 }
 
-uint32 UMaterialExpressionStrataSimpleClearCoatBSDF::GetOutputType(int32 OutputIndex)
+uint32 UMaterialExpressionSubstrateSimpleClearCoatBSDF::GetOutputType(int32 OutputIndex)
 {
 	return MCT_Strata;
 }
 
-uint32 UMaterialExpressionStrataSimpleClearCoatBSDF::GetInputType(int32 InputIndex)
+uint32 UMaterialExpressionSubstrateSimpleClearCoatBSDF::GetInputType(int32 InputIndex)
 {
 	if (InputIndex == 0)
 	{
@@ -24228,7 +24228,7 @@ uint32 UMaterialExpressionStrataSimpleClearCoatBSDF::GetInputType(int32 InputInd
 	return MCT_Float1;
 }
 
-FName UMaterialExpressionStrataSimpleClearCoatBSDF::GetInputName(int32 InputIndex) const
+FName UMaterialExpressionSubstrateSimpleClearCoatBSDF::GetInputName(int32 InputIndex) const
 {
 	if (InputIndex == 0)
 	{
@@ -24262,12 +24262,12 @@ FName UMaterialExpressionStrataSimpleClearCoatBSDF::GetInputName(int32 InputInde
 	return TEXT("Unknown");
 }
 
-bool UMaterialExpressionStrataSimpleClearCoatBSDF::IsResultStrataMaterial(int32 OutputIndex)
+bool UMaterialExpressionSubstrateSimpleClearCoatBSDF::IsResultStrataMaterial(int32 OutputIndex)
 {
 	return true;
 }
 
-void UMaterialExpressionStrataSimpleClearCoatBSDF::GatherStrataMaterialInfo(FStrataMaterialInfo& StrataMaterialInfo, int32 OutputIndex)
+void UMaterialExpressionSubstrateSimpleClearCoatBSDF::GatherStrataMaterialInfo(FStrataMaterialInfo& StrataMaterialInfo, int32 OutputIndex)
 {
 	// Track connected inputs
 	if (DiffuseAlbedo.IsConnected()) { StrataMaterialInfo.AddPropertyConnected(MP_BaseColor); }
@@ -24280,7 +24280,7 @@ void UMaterialExpressionStrataSimpleClearCoatBSDF::GatherStrataMaterialInfo(FStr
 	StrataMaterialInfo.AddShadingModel(SSM_DefaultLit);
 }
 
-FStrataOperator* UMaterialExpressionStrataSimpleClearCoatBSDF::StrataGenerateMaterialTopologyTree(class FMaterialCompiler* Compiler, class UMaterialExpression* Parent, int32 OutputIndex)
+FStrataOperator* UMaterialExpressionSubstrateSimpleClearCoatBSDF::StrataGenerateMaterialTopologyTree(class FMaterialCompiler* Compiler, class UMaterialExpression* Parent, int32 OutputIndex)
 {
 	FStrataOperator& StrataOperator = Compiler->StrataCompilationRegisterOperator(STRATA_OPERATOR_BSDF, Compiler->StrataTreeStackGetPathUniqueId(), this, Parent, Compiler->StrataTreeStackGetParentPathUniqueId());
 	StrataOperator.BSDFType = STRATA_BSDF_TYPE_SLAB;
@@ -24299,7 +24299,7 @@ FStrataOperator* UMaterialExpressionStrataSimpleClearCoatBSDF::StrataGenerateMat
 
 
 
-UMaterialExpressionStrataVolumetricFogCloudBSDF::UMaterialExpressionStrataVolumetricFogCloudBSDF(const FObjectInitializer& ObjectInitializer)
+UMaterialExpressionSubstrateVolumetricFogCloudBSDF::UMaterialExpressionSubstrateVolumetricFogCloudBSDF(const FObjectInitializer& ObjectInitializer)
 	: Super(ObjectInitializer)
 {
 	struct FConstructorStatics
@@ -24314,7 +24314,7 @@ UMaterialExpressionStrataVolumetricFogCloudBSDF::UMaterialExpressionStrataVolume
 }
 
 #if WITH_EDITOR
-int32 UMaterialExpressionStrataVolumetricFogCloudBSDF::Compile(class FMaterialCompiler* Compiler, int32 OutputIndex)
+int32 UMaterialExpressionSubstrateVolumetricFogCloudBSDF::Compile(class FMaterialCompiler* Compiler, int32 OutputIndex)
 {
 	int32 OutputCodeChunk = Compiler->StrataVolumetricFogCloudBSDF(
 		CompileWithDefaultFloat3(Compiler, Albedo, 0.0f, 0.0f, 0.0f),
@@ -24325,17 +24325,17 @@ int32 UMaterialExpressionStrataVolumetricFogCloudBSDF::Compile(class FMaterialCo
 	return OutputCodeChunk;
 }
 
-void UMaterialExpressionStrataVolumetricFogCloudBSDF::GetCaption(TArray<FString>& OutCaptions) const
+void UMaterialExpressionSubstrateVolumetricFogCloudBSDF::GetCaption(TArray<FString>& OutCaptions) const
 {
 	OutCaptions.Add(TEXT("Substrate Volumetric-Fog-Cloud BSDF"));
 }
 
-uint32 UMaterialExpressionStrataVolumetricFogCloudBSDF::GetOutputType(int32 OutputIndex)
+uint32 UMaterialExpressionSubstrateVolumetricFogCloudBSDF::GetOutputType(int32 OutputIndex)
 {
 	return MCT_Strata;
 }
 
-uint32 UMaterialExpressionStrataVolumetricFogCloudBSDF::GetInputType(int32 InputIndex)
+uint32 UMaterialExpressionSubstrateVolumetricFogCloudBSDF::GetInputType(int32 InputIndex)
 {
 	switch (InputIndex)
 	{
@@ -24357,17 +24357,17 @@ uint32 UMaterialExpressionStrataVolumetricFogCloudBSDF::GetInputType(int32 Input
 	return MCT_Float1;
 }
 
-bool UMaterialExpressionStrataVolumetricFogCloudBSDF::IsResultStrataMaterial(int32 OutputIndex)
+bool UMaterialExpressionSubstrateVolumetricFogCloudBSDF::IsResultStrataMaterial(int32 OutputIndex)
 {
 	return true;
 }
 
-void UMaterialExpressionStrataVolumetricFogCloudBSDF::GatherStrataMaterialInfo(FStrataMaterialInfo& StrataMaterialInfo, int32 OutputIndex)
+void UMaterialExpressionSubstrateVolumetricFogCloudBSDF::GatherStrataMaterialInfo(FStrataMaterialInfo& StrataMaterialInfo, int32 OutputIndex)
 {
 	StrataMaterialInfo.AddShadingModel(SSM_VolumetricFogCloud);
 }
 
-FStrataOperator* UMaterialExpressionStrataVolumetricFogCloudBSDF::StrataGenerateMaterialTopologyTree(class FMaterialCompiler* Compiler, class UMaterialExpression* Parent, int32 OutputIndex)
+FStrataOperator* UMaterialExpressionSubstrateVolumetricFogCloudBSDF::StrataGenerateMaterialTopologyTree(class FMaterialCompiler* Compiler, class UMaterialExpression* Parent, int32 OutputIndex)
 {
 	FStrataOperator& StrataOperator = Compiler->StrataCompilationRegisterOperator(STRATA_OPERATOR_BSDF, Compiler->StrataTreeStackGetPathUniqueId(), this, Parent, Compiler->StrataTreeStackGetParentPathUniqueId());
 	StrataOperator.BSDFType = STRATA_BSDF_TYPE_VOLUMETRICFOGCLOUD;
@@ -24378,7 +24378,7 @@ FStrataOperator* UMaterialExpressionStrataVolumetricFogCloudBSDF::StrataGenerate
 
 
 
-UMaterialExpressionStrataLightFunction::UMaterialExpressionStrataLightFunction(const FObjectInitializer& ObjectInitializer)
+UMaterialExpressionSubstrateLightFunction::UMaterialExpressionSubstrateLightFunction(const FObjectInitializer& ObjectInitializer)
 	: Super(ObjectInitializer)
 {
 	struct FConstructorStatics
@@ -24393,7 +24393,7 @@ UMaterialExpressionStrataLightFunction::UMaterialExpressionStrataLightFunction(c
 }
 
 #if WITH_EDITOR
-int32 UMaterialExpressionStrataLightFunction::Compile(class FMaterialCompiler* Compiler, int32 OutputIndex)
+int32 UMaterialExpressionSubstrateLightFunction::Compile(class FMaterialCompiler* Compiler, int32 OutputIndex)
 {
 	FStrataOperator& StrataOperator = Compiler->StrataCompilationGetOperator(Compiler->StrataTreeStackGetPathUniqueId());
 	int32 OutputCodeChunk = Compiler->StrataUnlitBSDF(
@@ -24404,17 +24404,17 @@ int32 UMaterialExpressionStrataLightFunction::Compile(class FMaterialCompiler* C
 	return OutputCodeChunk;
 }
 
-void UMaterialExpressionStrataLightFunction::GetCaption(TArray<FString>& OutCaptions) const
+void UMaterialExpressionSubstrateLightFunction::GetCaption(TArray<FString>& OutCaptions) const
 {
 	OutCaptions.Add(TEXT("Substrate Light Function"));
 }
 
-uint32 UMaterialExpressionStrataLightFunction::GetOutputType(int32 OutputIndex)
+uint32 UMaterialExpressionSubstrateLightFunction::GetOutputType(int32 OutputIndex)
 {
 	return MCT_Strata;
 }
 
-uint32 UMaterialExpressionStrataLightFunction::GetInputType(int32 InputIndex)
+uint32 UMaterialExpressionSubstrateLightFunction::GetInputType(int32 InputIndex)
 {
 	switch (InputIndex)
 	{
@@ -24427,17 +24427,17 @@ uint32 UMaterialExpressionStrataLightFunction::GetInputType(int32 InputIndex)
 	return MCT_Float1;
 }
 
-bool UMaterialExpressionStrataLightFunction::IsResultStrataMaterial(int32 OutputIndex)
+bool UMaterialExpressionSubstrateLightFunction::IsResultStrataMaterial(int32 OutputIndex)
 {
 	return true;
 }
 
-void UMaterialExpressionStrataLightFunction::GatherStrataMaterialInfo(FStrataMaterialInfo& StrataMaterialInfo, int32 OutputIndex)
+void UMaterialExpressionSubstrateLightFunction::GatherStrataMaterialInfo(FStrataMaterialInfo& StrataMaterialInfo, int32 OutputIndex)
 {
 	StrataMaterialInfo.AddShadingModel(SSM_LightFunction);
 }
 
-FStrataOperator* UMaterialExpressionStrataLightFunction::StrataGenerateMaterialTopologyTree(class FMaterialCompiler* Compiler, class UMaterialExpression* Parent, int32 OutputIndex)
+FStrataOperator* UMaterialExpressionSubstrateLightFunction::StrataGenerateMaterialTopologyTree(class FMaterialCompiler* Compiler, class UMaterialExpression* Parent, int32 OutputIndex)
 {
 	FStrataOperator& StrataOperator = Compiler->StrataCompilationRegisterOperator(STRATA_OPERATOR_BSDF, Compiler->StrataTreeStackGetPathUniqueId(), this, Parent, Compiler->StrataTreeStackGetParentPathUniqueId());
 	StrataOperator.BSDFType = STRATA_BSDF_TYPE_UNLIT;
@@ -24448,7 +24448,7 @@ FStrataOperator* UMaterialExpressionStrataLightFunction::StrataGenerateMaterialT
 
 
 
-UMaterialExpressionStrataPostProcess::UMaterialExpressionStrataPostProcess(const FObjectInitializer& ObjectInitializer)
+UMaterialExpressionSubstratePostProcess::UMaterialExpressionSubstratePostProcess(const FObjectInitializer& ObjectInitializer)
 	: Super(ObjectInitializer)
 {
 	struct FConstructorStatics
@@ -24463,7 +24463,7 @@ UMaterialExpressionStrataPostProcess::UMaterialExpressionStrataPostProcess(const
 }
 
 #if WITH_EDITOR
-int32 UMaterialExpressionStrataPostProcess::Compile(class FMaterialCompiler* Compiler, int32 OutputIndex)
+int32 UMaterialExpressionSubstratePostProcess::Compile(class FMaterialCompiler* Compiler, int32 OutputIndex)
 {
 	int OpacityCodeChunk = CompileWithDefaultFloat1(Compiler, Opacity, 0.0f);
 	int TransmittanceCodeChunk = Compiler->Saturate(Compiler->Sub(Compiler->Constant(1.0f), OpacityCodeChunk));
@@ -24477,17 +24477,17 @@ int32 UMaterialExpressionStrataPostProcess::Compile(class FMaterialCompiler* Com
 	return OutputCodeChunk;
 }
 
-void UMaterialExpressionStrataPostProcess::GetCaption(TArray<FString>& OutCaptions) const
+void UMaterialExpressionSubstratePostProcess::GetCaption(TArray<FString>& OutCaptions) const
 {
 	OutCaptions.Add(TEXT("Substrate Post Process"));
 }
 
-uint32 UMaterialExpressionStrataPostProcess::GetOutputType(int32 OutputIndex)
+uint32 UMaterialExpressionSubstratePostProcess::GetOutputType(int32 OutputIndex)
 {
 	return MCT_Strata;
 }
 
-uint32 UMaterialExpressionStrataPostProcess::GetInputType(int32 InputIndex)
+uint32 UMaterialExpressionSubstratePostProcess::GetInputType(int32 InputIndex)
 {
 	switch (InputIndex)
 	{
@@ -24503,17 +24503,17 @@ uint32 UMaterialExpressionStrataPostProcess::GetInputType(int32 InputIndex)
 	return MCT_Float1;
 }
 
-bool UMaterialExpressionStrataPostProcess::IsResultStrataMaterial(int32 OutputIndex)
+bool UMaterialExpressionSubstratePostProcess::IsResultStrataMaterial(int32 OutputIndex)
 {
 	return true;
 }
 
-void UMaterialExpressionStrataPostProcess::GatherStrataMaterialInfo(FStrataMaterialInfo& StrataMaterialInfo, int32 OutputIndex)
+void UMaterialExpressionSubstratePostProcess::GatherStrataMaterialInfo(FStrataMaterialInfo& StrataMaterialInfo, int32 OutputIndex)
 {
 	StrataMaterialInfo.AddShadingModel(SSM_PostProcess);
 }
 
-FStrataOperator* UMaterialExpressionStrataPostProcess::StrataGenerateMaterialTopologyTree(class FMaterialCompiler* Compiler, class UMaterialExpression* Parent, int32 OutputIndex)
+FStrataOperator* UMaterialExpressionSubstratePostProcess::StrataGenerateMaterialTopologyTree(class FMaterialCompiler* Compiler, class UMaterialExpression* Parent, int32 OutputIndex)
 {
 	FStrataOperator& StrataOperator = Compiler->StrataCompilationRegisterOperator(STRATA_OPERATOR_BSDF, Compiler->StrataTreeStackGetPathUniqueId(), this, Parent, Compiler->StrataTreeStackGetParentPathUniqueId());
 	StrataOperator.BSDFType = STRATA_BSDF_TYPE_UNLIT;
@@ -24524,7 +24524,7 @@ FStrataOperator* UMaterialExpressionStrataPostProcess::StrataGenerateMaterialTop
 
 
 
-UMaterialExpressionStrataUI::UMaterialExpressionStrataUI(const FObjectInitializer& ObjectInitializer)
+UMaterialExpressionSubstrateUI::UMaterialExpressionSubstrateUI(const FObjectInitializer& ObjectInitializer)
 	: Super(ObjectInitializer)
 {
 	struct FConstructorStatics
@@ -24539,7 +24539,7 @@ UMaterialExpressionStrataUI::UMaterialExpressionStrataUI(const FObjectInitialize
 }
 
 #if WITH_EDITOR
-int32 UMaterialExpressionStrataUI::Compile(class FMaterialCompiler* Compiler, int32 OutputIndex)
+int32 UMaterialExpressionSubstrateUI::Compile(class FMaterialCompiler* Compiler, int32 OutputIndex)
 {
 	int OpacityCodeChunk = CompileWithDefaultFloat1(Compiler, Opacity, 0.0f);
 	int TransmittanceCodeChunk = Compiler->Saturate(Compiler->Sub(Compiler->Constant(1.0f), OpacityCodeChunk));
@@ -24553,17 +24553,17 @@ int32 UMaterialExpressionStrataUI::Compile(class FMaterialCompiler* Compiler, in
 	return OutputCodeChunk;
 }
 
-void UMaterialExpressionStrataUI::GetCaption(TArray<FString>& OutCaptions) const
+void UMaterialExpressionSubstrateUI::GetCaption(TArray<FString>& OutCaptions) const
 {
 	OutCaptions.Add(TEXT("Substrate UI"));
 }
 
-uint32 UMaterialExpressionStrataUI::GetOutputType(int32 OutputIndex)
+uint32 UMaterialExpressionSubstrateUI::GetOutputType(int32 OutputIndex)
 {
 	return MCT_Strata;
 }
 
-uint32 UMaterialExpressionStrataUI::GetInputType(int32 InputIndex)
+uint32 UMaterialExpressionSubstrateUI::GetInputType(int32 InputIndex)
 {
 	switch (InputIndex)
 	{
@@ -24579,17 +24579,17 @@ uint32 UMaterialExpressionStrataUI::GetInputType(int32 InputIndex)
 	return MCT_Float1;
 }
 
-bool UMaterialExpressionStrataUI::IsResultStrataMaterial(int32 OutputIndex)
+bool UMaterialExpressionSubstrateUI::IsResultStrataMaterial(int32 OutputIndex)
 {
 	return true;
 }
 
-void UMaterialExpressionStrataUI::GatherStrataMaterialInfo(FStrataMaterialInfo& StrataMaterialInfo, int32 OutputIndex)
+void UMaterialExpressionSubstrateUI::GatherStrataMaterialInfo(FStrataMaterialInfo& StrataMaterialInfo, int32 OutputIndex)
 {
 	StrataMaterialInfo.AddShadingModel(SSM_UI);
 }
 
-FStrataOperator* UMaterialExpressionStrataUI::StrataGenerateMaterialTopologyTree(class FMaterialCompiler* Compiler, class UMaterialExpression* Parent, int32 OutputIndex)
+FStrataOperator* UMaterialExpressionSubstrateUI::StrataGenerateMaterialTopologyTree(class FMaterialCompiler* Compiler, class UMaterialExpression* Parent, int32 OutputIndex)
 {
 	FStrataOperator& StrataOperator = Compiler->StrataCompilationRegisterOperator(STRATA_OPERATOR_BSDF, Compiler->StrataTreeStackGetPathUniqueId(), this, Parent, Compiler->StrataTreeStackGetParentPathUniqueId());
 	StrataOperator.BSDFType = STRATA_BSDF_TYPE_UNLIT;
@@ -24600,7 +24600,7 @@ FStrataOperator* UMaterialExpressionStrataUI::StrataGenerateMaterialTopologyTree
 
 
 
-UMaterialExpressionStrataConvertToDecal::UMaterialExpressionStrataConvertToDecal(const FObjectInitializer& ObjectInitializer)
+UMaterialExpressionSubstrateConvertToDecal::UMaterialExpressionSubstrateConvertToDecal(const FObjectInitializer& ObjectInitializer)
 	: Super(ObjectInitializer)
 {
 	struct FConstructorStatics
@@ -24615,7 +24615,7 @@ UMaterialExpressionStrataConvertToDecal::UMaterialExpressionStrataConvertToDecal
 }
 
 #if WITH_EDITOR
-int32 UMaterialExpressionStrataConvertToDecal::Compile(class FMaterialCompiler* Compiler, int32 OutputIndex)
+int32 UMaterialExpressionSubstrateConvertToDecal::Compile(class FMaterialCompiler* Compiler, int32 OutputIndex)
 {
 	if (!DecalMaterial.GetTracedInput().Expression)
 	{
@@ -24649,17 +24649,17 @@ int32 UMaterialExpressionStrataConvertToDecal::Compile(class FMaterialCompiler*
 	return OutputCodeChunk;
 }
 
-void UMaterialExpressionStrataConvertToDecal::GetCaption(TArray<FString>& OutCaptions) const
+void UMaterialExpressionSubstrateConvertToDecal::GetCaption(TArray<FString>& OutCaptions) const
 {
 	OutCaptions.Add(TEXT("Substrate Convert To Decal"));
 }
 
-uint32 UMaterialExpressionStrataConvertToDecal::GetOutputType(int32 OutputIndex)
+uint32 UMaterialExpressionSubstrateConvertToDecal::GetOutputType(int32 OutputIndex)
 {
 	return MCT_Strata;
 }
 
-uint32 UMaterialExpressionStrataConvertToDecal::GetInputType(int32 InputIndex)
+uint32 UMaterialExpressionSubstrateConvertToDecal::GetInputType(int32 InputIndex)
 {
 	switch (InputIndex)
 	{
@@ -24675,12 +24675,12 @@ uint32 UMaterialExpressionStrataConvertToDecal::GetInputType(int32 InputIndex)
 	return MCT_Float1;
 }
 
-bool UMaterialExpressionStrataConvertToDecal::IsResultStrataMaterial(int32 OutputIndex)
+bool UMaterialExpressionSubstrateConvertToDecal::IsResultStrataMaterial(int32 OutputIndex)
 {
 	return true;
 }
 
-void UMaterialExpressionStrataConvertToDecal::GatherStrataMaterialInfo(FStrataMaterialInfo& StrataMaterialInfo, int32 OutputIndex)
+void UMaterialExpressionSubstrateConvertToDecal::GatherStrataMaterialInfo(FStrataMaterialInfo& StrataMaterialInfo, int32 OutputIndex)
 {
 	if (!StrataMaterialInfo.PushStrataTreeStack())
 	{
@@ -24697,7 +24697,7 @@ void UMaterialExpressionStrataConvertToDecal::GatherStrataMaterialInfo(FStrataMa
 	StrataMaterialInfo.PopStrataTreeStack();
 }
 
-FStrataOperator* UMaterialExpressionStrataConvertToDecal::StrataGenerateMaterialTopologyTree(class FMaterialCompiler* Compiler, class UMaterialExpression* Parent, int32 OutputIndex)
+FStrataOperator* UMaterialExpressionSubstrateConvertToDecal::StrataGenerateMaterialTopologyTree(class FMaterialCompiler* Compiler, class UMaterialExpression* Parent, int32 OutputIndex)
 {
 	const bool bUseParameterBlending = true;
 	FStrataOperator& StrataOperator = Compiler->StrataCompilationRegisterOperator(STRATA_OPERATOR_WEIGHT, Compiler->StrataTreeStackGetPathUniqueId(), this, Parent, Compiler->StrataTreeStackGetParentPathUniqueId(), bUseParameterBlending);
@@ -24724,7 +24724,7 @@ FStrataOperator* UMaterialExpressionStrataConvertToDecal::StrataGenerateMaterial
 
 
 
-UMaterialExpressionStrataUnlitBSDF::UMaterialExpressionStrataUnlitBSDF(const FObjectInitializer& ObjectInitializer)
+UMaterialExpressionSubstrateUnlitBSDF::UMaterialExpressionSubstrateUnlitBSDF(const FObjectInitializer& ObjectInitializer)
 	: Super(ObjectInitializer)
 {
 	struct FConstructorStatics
@@ -24739,7 +24739,7 @@ UMaterialExpressionStrataUnlitBSDF::UMaterialExpressionStrataUnlitBSDF(const FOb
 }
 
 #if WITH_EDITOR
-int32 UMaterialExpressionStrataUnlitBSDF::Compile(class FMaterialCompiler* Compiler, int32 OutputIndex)
+int32 UMaterialExpressionSubstrateUnlitBSDF::Compile(class FMaterialCompiler* Compiler, int32 OutputIndex)
 {
 	FStrataOperator& StrataOperator = Compiler->StrataCompilationGetOperator(Compiler->StrataTreeStackGetPathUniqueId());
 	int32 OutputCodeChunk = Compiler->StrataUnlitBSDF(
@@ -24750,17 +24750,17 @@ int32 UMaterialExpressionStrataUnlitBSDF::Compile(class FMaterialCompiler* Compi
 	return OutputCodeChunk;
 }
 
-void UMaterialExpressionStrataUnlitBSDF::GetCaption(TArray<FString>& OutCaptions) const
+void UMaterialExpressionSubstrateUnlitBSDF::GetCaption(TArray<FString>& OutCaptions) const
 {
 	OutCaptions.Add(TEXT("Substrate Unlit BSDF"));
 }
 
-uint32 UMaterialExpressionStrataUnlitBSDF::GetOutputType(int32 OutputIndex)
+uint32 UMaterialExpressionSubstrateUnlitBSDF::GetOutputType(int32 OutputIndex)
 {
 	return MCT_Strata;
 }
 
-uint32 UMaterialExpressionStrataUnlitBSDF::GetInputType(int32 InputIndex)
+uint32 UMaterialExpressionSubstrateUnlitBSDF::GetInputType(int32 InputIndex)
 {
 	switch (InputIndex)
 	{
@@ -24779,18 +24779,18 @@ uint32 UMaterialExpressionStrataUnlitBSDF::GetInputType(int32 InputIndex)
 	return MCT_Float1;
 }
 
-bool UMaterialExpressionStrataUnlitBSDF::IsResultStrataMaterial(int32 OutputIndex)
+bool UMaterialExpressionSubstrateUnlitBSDF::IsResultStrataMaterial(int32 OutputIndex)
 {
 	return true;
 }
 
-void UMaterialExpressionStrataUnlitBSDF::GatherStrataMaterialInfo(FStrataMaterialInfo& StrataMaterialInfo, int32 OutputIndex)
+void UMaterialExpressionSubstrateUnlitBSDF::GatherStrataMaterialInfo(FStrataMaterialInfo& StrataMaterialInfo, int32 OutputIndex)
 {
 	if (EmissiveColor.IsConnected()) { StrataMaterialInfo.AddPropertyConnected(MP_EmissiveColor); }
 	StrataMaterialInfo.AddShadingModel(SSM_Unlit);
 }
 
-FStrataOperator* UMaterialExpressionStrataUnlitBSDF::StrataGenerateMaterialTopologyTree(class FMaterialCompiler* Compiler, class UMaterialExpression* Parent, int32 OutputIndex)
+FStrataOperator* UMaterialExpressionSubstrateUnlitBSDF::StrataGenerateMaterialTopologyTree(class FMaterialCompiler* Compiler, class UMaterialExpression* Parent, int32 OutputIndex)
 {
 	FStrataOperator& StrataOperator = Compiler->StrataCompilationRegisterOperator(STRATA_OPERATOR_BSDF, Compiler->StrataTreeStackGetPathUniqueId(), this, Parent, Compiler->StrataTreeStackGetParentPathUniqueId());
 	StrataOperator.BSDFType = STRATA_BSDF_TYPE_UNLIT;
@@ -24801,7 +24801,7 @@ FStrataOperator* UMaterialExpressionStrataUnlitBSDF::StrataGenerateMaterialTopol
 
 
 
-UMaterialExpressionStrataHairBSDF::UMaterialExpressionStrataHairBSDF(const FObjectInitializer& ObjectInitializer)
+UMaterialExpressionSubstrateHairBSDF::UMaterialExpressionSubstrateHairBSDF(const FObjectInitializer& ObjectInitializer)
 	: Super(ObjectInitializer)
 {
 	struct FConstructorStatics
@@ -24816,7 +24816,7 @@ UMaterialExpressionStrataHairBSDF::UMaterialExpressionStrataHairBSDF(const FObje
 }
 
 #if WITH_EDITOR
-int32 UMaterialExpressionStrataHairBSDF::Compile(class FMaterialCompiler* Compiler, int32 OutputIndex)
+int32 UMaterialExpressionSubstrateHairBSDF::Compile(class FMaterialCompiler* Compiler, int32 OutputIndex)
 {
 	// For hair, the shared local basis normal in fact represent the tangent
 	int32 TangentCodeChunk = CompileWithDefaultTangentWS(Compiler, Tangent);
@@ -24848,17 +24848,17 @@ int32 UMaterialExpressionStrataHairBSDF::Compile(class FMaterialCompiler* Compil
 	return OutputCodeChunk;
 }
 
-void UMaterialExpressionStrataHairBSDF::GetCaption(TArray<FString>& OutCaptions) const
+void UMaterialExpressionSubstrateHairBSDF::GetCaption(TArray<FString>& OutCaptions) const
 {
 	OutCaptions.Add(TEXT("Substrate Hair BSDF"));
 }
 
-uint32 UMaterialExpressionStrataHairBSDF::GetOutputType(int32 OutputIndex)
+uint32 UMaterialExpressionSubstrateHairBSDF::GetOutputType(int32 OutputIndex)
 {
 	return MCT_Strata;
 }
 
-uint32 UMaterialExpressionStrataHairBSDF::GetInputType(int32 InputIndex)
+uint32 UMaterialExpressionSubstrateHairBSDF::GetInputType(int32 InputIndex)
 {
 	switch (InputIndex)
 	{
@@ -24889,12 +24889,12 @@ uint32 UMaterialExpressionStrataHairBSDF::GetInputType(int32 InputIndex)
 	return MCT_Float1;
 }
 
-bool UMaterialExpressionStrataHairBSDF::IsResultStrataMaterial(int32 OutputIndex)
+bool UMaterialExpressionSubstrateHairBSDF::IsResultStrataMaterial(int32 OutputIndex)
 {
 	return true;
 }
 
-void UMaterialExpressionStrataHairBSDF::GatherStrataMaterialInfo(FStrataMaterialInfo& StrataMaterialInfo, int32 OutputIndex)
+void UMaterialExpressionSubstrateHairBSDF::GatherStrataMaterialInfo(FStrataMaterialInfo& StrataMaterialInfo, int32 OutputIndex)
 {
 	// Track connected inputs
 	if (BaseColor.IsConnected())		{ StrataMaterialInfo.AddPropertyConnected(MP_BaseColor); }
@@ -24906,7 +24906,7 @@ void UMaterialExpressionStrataHairBSDF::GatherStrataMaterialInfo(FStrataMaterial
 	StrataMaterialInfo.AddShadingModel(SSM_Hair);
 }
 
-FStrataOperator* UMaterialExpressionStrataHairBSDF::StrataGenerateMaterialTopologyTree(class FMaterialCompiler* Compiler, class UMaterialExpression* Parent, int32 OutputIndex)
+FStrataOperator* UMaterialExpressionSubstrateHairBSDF::StrataGenerateMaterialTopologyTree(class FMaterialCompiler* Compiler, class UMaterialExpression* Parent, int32 OutputIndex)
 {
 	FStrataOperator& StrataOperator = Compiler->StrataCompilationRegisterOperator(STRATA_OPERATOR_BSDF, Compiler->StrataTreeStackGetPathUniqueId(), this, Parent, Compiler->StrataTreeStackGetParentPathUniqueId());
 	StrataOperator.BSDFType = STRATA_BSDF_TYPE_HAIR;
@@ -24915,7 +24915,7 @@ FStrataOperator* UMaterialExpressionStrataHairBSDF::StrataGenerateMaterialTopolo
 }
 #endif // WITH_EDITOR
 
-UMaterialExpressionStrataEyeBSDF::UMaterialExpressionStrataEyeBSDF(const FObjectInitializer& ObjectInitializer)
+UMaterialExpressionSubstrateEyeBSDF::UMaterialExpressionSubstrateEyeBSDF(const FObjectInitializer& ObjectInitializer)
 	: Super(ObjectInitializer)
 {
 	struct FConstructorStatics
@@ -24930,7 +24930,7 @@ UMaterialExpressionStrataEyeBSDF::UMaterialExpressionStrataEyeBSDF(const FObject
 }
 
 #if WITH_EDITOR
-int32 UMaterialExpressionStrataEyeBSDF::Compile(class FMaterialCompiler* Compiler, int32 OutputIndex)
+int32 UMaterialExpressionSubstrateEyeBSDF::Compile(class FMaterialCompiler* Compiler, int32 OutputIndex)
 {
 	int32 CorneaNormalCodeChunk = CompileWithDefaultTangentWS(Compiler, CorneaNormal);
 	const FStrataRegisteredSharedLocalBasis NewRegisteredSharedLocalBasis = StrataCompilationInfoCreateSharedLocalBasis(Compiler, CorneaNormalCodeChunk);
@@ -24969,17 +24969,17 @@ int32 UMaterialExpressionStrataEyeBSDF::Compile(class FMaterialCompiler* Compile
 	return OutputCodeChunk;
 }
 
-void UMaterialExpressionStrataEyeBSDF::GetCaption(TArray<FString>& OutCaptions) const
+void UMaterialExpressionSubstrateEyeBSDF::GetCaption(TArray<FString>& OutCaptions) const
 {
 	OutCaptions.Add(TEXT("Substrate Eye BSDF"));
 }
 
-uint32 UMaterialExpressionStrataEyeBSDF::GetOutputType(int32 OutputIndex)
+uint32 UMaterialExpressionSubstrateEyeBSDF::GetOutputType(int32 OutputIndex)
 {
 	return MCT_Strata;
 }
 
-uint32 UMaterialExpressionStrataEyeBSDF::GetInputType(int32 InputIndex)
+uint32 UMaterialExpressionSubstrateEyeBSDF::GetInputType(int32 InputIndex)
 {
 	switch (InputIndex)
 	{
@@ -24997,12 +24997,12 @@ uint32 UMaterialExpressionStrataEyeBSDF::GetInputType(int32 InputIndex)
 	return MCT_Float1;
 }
 
-bool UMaterialExpressionStrataEyeBSDF::IsResultStrataMaterial(int32 OutputIndex)
+bool UMaterialExpressionSubstrateEyeBSDF::IsResultStrataMaterial(int32 OutputIndex)
 {
 	return true;
 }
 
-void UMaterialExpressionStrataEyeBSDF::GatherStrataMaterialInfo(FStrataMaterialInfo& StrataMaterialInfo, int32 OutputIndex)
+void UMaterialExpressionSubstrateEyeBSDF::GatherStrataMaterialInfo(FStrataMaterialInfo& StrataMaterialInfo, int32 OutputIndex)
 {
 	// Track connected inputs
 	if (DiffuseColor.IsConnected())		{ StrataMaterialInfo.AddPropertyConnected(MP_BaseColor); }
@@ -25020,7 +25020,7 @@ void UMaterialExpressionStrataEyeBSDF::GatherStrataMaterialInfo(FStrataMaterialI
 	StrataMaterialInfo.AddShadingModel(SSM_Eye);
 }
 
-FStrataOperator* UMaterialExpressionStrataEyeBSDF::StrataGenerateMaterialTopologyTree(class FMaterialCompiler* Compiler, class UMaterialExpression* Parent, int32 OutputIndex)
+FStrataOperator* UMaterialExpressionSubstrateEyeBSDF::StrataGenerateMaterialTopologyTree(class FMaterialCompiler* Compiler, class UMaterialExpression* Parent, int32 OutputIndex)
 {
 	FStrataOperator& StrataOperator = Compiler->StrataCompilationRegisterOperator(STRATA_OPERATOR_BSDF, Compiler->StrataTreeStackGetPathUniqueId(), this, Parent, Compiler->StrataTreeStackGetParentPathUniqueId());
 	StrataOperator.BSDFType = STRATA_BSDF_TYPE_EYE;
@@ -25030,7 +25030,7 @@ FStrataOperator* UMaterialExpressionStrataEyeBSDF::StrataGenerateMaterialTopolog
 #endif // WITH_EDITOR
 
 
-UMaterialExpressionStrataSingleLayerWaterBSDF::UMaterialExpressionStrataSingleLayerWaterBSDF(const FObjectInitializer& ObjectInitializer)
+UMaterialExpressionSubstrateSingleLayerWaterBSDF::UMaterialExpressionSubstrateSingleLayerWaterBSDF(const FObjectInitializer& ObjectInitializer)
 	: Super(ObjectInitializer)
 {
 	struct FConstructorStatics
@@ -25045,7 +25045,7 @@ UMaterialExpressionStrataSingleLayerWaterBSDF::UMaterialExpressionStrataSingleLa
 }
 
 #if WITH_EDITOR
-int32 UMaterialExpressionStrataSingleLayerWaterBSDF::Compile(class FMaterialCompiler* Compiler, int32 OutputIndex)
+int32 UMaterialExpressionSubstrateSingleLayerWaterBSDF::Compile(class FMaterialCompiler* Compiler, int32 OutputIndex)
 {
 	int32 NormalCodeChunk = CompileWithDefaultNormalWS(Compiler, Normal);
 	const FStrataRegisteredSharedLocalBasis NewRegisteredSharedLocalBasis = StrataCompilationInfoCreateSharedLocalBasis(Compiler, NormalCodeChunk);
@@ -25080,17 +25080,17 @@ int32 UMaterialExpressionStrataSingleLayerWaterBSDF::Compile(class FMaterialComp
 	return OutputCodeChunk;
 }
 
-void UMaterialExpressionStrataSingleLayerWaterBSDF::GetCaption(TArray<FString>& OutCaptions) const
+void UMaterialExpressionSubstrateSingleLayerWaterBSDF::GetCaption(TArray<FString>& OutCaptions) const
 {
 	OutCaptions.Add(TEXT("Substrate Single Layer Water BSDF"));
 }
 
-uint32 UMaterialExpressionStrataSingleLayerWaterBSDF::GetOutputType(int32 OutputIndex)
+uint32 UMaterialExpressionSubstrateSingleLayerWaterBSDF::GetOutputType(int32 OutputIndex)
 {
 	return MCT_Strata;
 }
 
-uint32 UMaterialExpressionStrataSingleLayerWaterBSDF::GetInputType(int32 InputIndex)
+uint32 UMaterialExpressionSubstrateSingleLayerWaterBSDF::GetInputType(int32 InputIndex)
 {
 	switch (InputIndex)
 	{
@@ -25133,12 +25133,12 @@ uint32 UMaterialExpressionStrataSingleLayerWaterBSDF::GetInputType(int32 InputIn
 	return MCT_Float1;
 }
 
-bool UMaterialExpressionStrataSingleLayerWaterBSDF::IsResultStrataMaterial(int32 OutputIndex)
+bool UMaterialExpressionSubstrateSingleLayerWaterBSDF::IsResultStrataMaterial(int32 OutputIndex)
 {
 	return true;
 }
 
-void UMaterialExpressionStrataSingleLayerWaterBSDF::GatherStrataMaterialInfo(FStrataMaterialInfo& StrataMaterialInfo, int32 OutputIndex)
+void UMaterialExpressionSubstrateSingleLayerWaterBSDF::GatherStrataMaterialInfo(FStrataMaterialInfo& StrataMaterialInfo, int32 OutputIndex)
 {
 	// Track connected inputs
 	if (BaseColor.IsConnected())		{ StrataMaterialInfo.AddPropertyConnected(MP_BaseColor); }
@@ -25151,7 +25151,7 @@ void UMaterialExpressionStrataSingleLayerWaterBSDF::GatherStrataMaterialInfo(FSt
 	StrataMaterialInfo.AddShadingModel(SSM_SingleLayerWater);
 }
 
-FStrataOperator* UMaterialExpressionStrataSingleLayerWaterBSDF::StrataGenerateMaterialTopologyTree(class FMaterialCompiler* Compiler, class UMaterialExpression* Parent, int32 OutputIndex)
+FStrataOperator* UMaterialExpressionSubstrateSingleLayerWaterBSDF::StrataGenerateMaterialTopologyTree(class FMaterialCompiler* Compiler, class UMaterialExpression* Parent, int32 OutputIndex)
 {
 	FStrataOperator& StrataOperator = Compiler->StrataCompilationRegisterOperator(STRATA_OPERATOR_BSDF, Compiler->StrataTreeStackGetPathUniqueId(), this, Parent, Compiler->StrataTreeStackGetParentPathUniqueId());
 	StrataOperator.BSDFType = STRATA_BSDF_TYPE_SINGLELAYERWATER;
@@ -25162,7 +25162,7 @@ FStrataOperator* UMaterialExpressionStrataSingleLayerWaterBSDF::StrataGenerateMa
 
 
 
-UMaterialExpressionStrataHorizontalMixing::UMaterialExpressionStrataHorizontalMixing(const FObjectInitializer& ObjectInitializer)
+UMaterialExpressionSubstrateHorizontalMixing::UMaterialExpressionSubstrateHorizontalMixing(const FObjectInitializer& ObjectInitializer)
 	: Super(ObjectInitializer)
 	, bUseParameterBlending(false)
 {
@@ -25178,7 +25178,7 @@ UMaterialExpressionStrataHorizontalMixing::UMaterialExpressionStrataHorizontalMi
 }
 
 #if WITH_EDITOR
-int32 UMaterialExpressionStrataHorizontalMixing::Compile(class FMaterialCompiler* Compiler, int32 OutputIndex)
+int32 UMaterialExpressionSubstrateHorizontalMixing::Compile(class FMaterialCompiler* Compiler, int32 OutputIndex)
 {
 	if (!Foreground.GetTracedInput().Expression)
 	{
@@ -25257,7 +25257,7 @@ int32 UMaterialExpressionStrataHorizontalMixing::Compile(class FMaterialCompiler
 	return OutputCodeChunk;
 }
 
-void UMaterialExpressionStrataHorizontalMixing::GetCaption(TArray<FString>& OutCaptions) const
+void UMaterialExpressionSubstrateHorizontalMixing::GetCaption(TArray<FString>& OutCaptions) const
 {
 	if (bUseParameterBlending)
 	{
@@ -25269,22 +25269,22 @@ void UMaterialExpressionStrataHorizontalMixing::GetCaption(TArray<FString>& OutC
 	}
 }
 
-uint32 UMaterialExpressionStrataHorizontalMixing::GetOutputType(int32 OutputIndex)
+uint32 UMaterialExpressionSubstrateHorizontalMixing::GetOutputType(int32 OutputIndex)
 {
 	return MCT_Strata;
 }
 
-uint32 UMaterialExpressionStrataHorizontalMixing::GetInputType(int32 InputIndex)
+uint32 UMaterialExpressionSubstrateHorizontalMixing::GetInputType(int32 InputIndex)
 {
 	return InputIndex == 2 ? MCT_Float1 : MCT_Strata;
 }
 
-bool UMaterialExpressionStrataHorizontalMixing::IsResultStrataMaterial(int32 OutputIndex)
+bool UMaterialExpressionSubstrateHorizontalMixing::IsResultStrataMaterial(int32 OutputIndex)
 {
 	return true;
 }
 
-void UMaterialExpressionStrataHorizontalMixing::GatherStrataMaterialInfo(FStrataMaterialInfo& StrataMaterialInfo, int32 OutputIndex)
+void UMaterialExpressionSubstrateHorizontalMixing::GatherStrataMaterialInfo(FStrataMaterialInfo& StrataMaterialInfo, int32 OutputIndex)
 {
 	if (!StrataMaterialInfo.PushStrataTreeStack())
 	{
@@ -25305,7 +25305,7 @@ void UMaterialExpressionStrataHorizontalMixing::GatherStrataMaterialInfo(FStrata
 	StrataMaterialInfo.PopStrataTreeStack();
 }
 
-FStrataOperator* UMaterialExpressionStrataHorizontalMixing::StrataGenerateMaterialTopologyTree(class FMaterialCompiler* Compiler, class UMaterialExpression* Parent, int32 OutputIndex)
+FStrataOperator* UMaterialExpressionSubstrateHorizontalMixing::StrataGenerateMaterialTopologyTree(class FMaterialCompiler* Compiler, class UMaterialExpression* Parent, int32 OutputIndex)
 {
 	FStrataOperator& StrataOperator = Compiler->StrataCompilationRegisterOperator(STRATA_OPERATOR_HORIZONTAL, Compiler->StrataTreeStackGetPathUniqueId(), this, Parent, Compiler->StrataTreeStackGetParentPathUniqueId(), bUseParameterBlending);
 	if (Compiler->GetStrataTreeOutOfStackDepthOccurred())
@@ -25341,7 +25341,7 @@ FStrataOperator* UMaterialExpressionStrataHorizontalMixing::StrataGenerateMateri
 
 
 
-UMaterialExpressionStrataVerticalLayering::UMaterialExpressionStrataVerticalLayering(const FObjectInitializer& ObjectInitializer)
+UMaterialExpressionSubstrateVerticalLayering::UMaterialExpressionSubstrateVerticalLayering(const FObjectInitializer& ObjectInitializer)
 	: Super(ObjectInitializer)
 	, bUseParameterBlending(false)
 {
@@ -25357,7 +25357,7 @@ UMaterialExpressionStrataVerticalLayering::UMaterialExpressionStrataVerticalLaye
 }
 
 #if WITH_EDITOR
-int32 UMaterialExpressionStrataVerticalLayering::Compile(class FMaterialCompiler* Compiler, int32 OutputIndex)
+int32 UMaterialExpressionSubstrateVerticalLayering::Compile(class FMaterialCompiler* Compiler, int32 OutputIndex)
 {
 	if (!Top.GetTracedInput().Expression)
 	{
@@ -25433,7 +25433,7 @@ int32 UMaterialExpressionStrataVerticalLayering::Compile(class FMaterialCompiler
 	return OutputCodeChunk;
 }
 
-void UMaterialExpressionStrataVerticalLayering::GetCaption(TArray<FString>& OutCaptions) const
+void UMaterialExpressionSubstrateVerticalLayering::GetCaption(TArray<FString>& OutCaptions) const
 {
 	if (bUseParameterBlending)
 	{
@@ -25445,7 +25445,7 @@ void UMaterialExpressionStrataVerticalLayering::GetCaption(TArray<FString>& OutC
 	}
 }
 
-FName UMaterialExpressionStrataVerticalLayering::GetInputName(int32 InputIndex) const
+FName UMaterialExpressionSubstrateVerticalLayering::GetInputName(int32 InputIndex) const
 {
 	if (InputIndex == 0)
 	{
@@ -25463,12 +25463,12 @@ FName UMaterialExpressionStrataVerticalLayering::GetInputName(int32 InputIndex)
 	return TEXT("Unknown");
 }
 
-uint32 UMaterialExpressionStrataVerticalLayering::GetOutputType(int32 OutputIndex)
+uint32 UMaterialExpressionSubstrateVerticalLayering::GetOutputType(int32 OutputIndex)
 {
 	return MCT_Strata;
 }
 
-uint32 UMaterialExpressionStrataVerticalLayering::GetInputType(int32 InputIndex)
+uint32 UMaterialExpressionSubstrateVerticalLayering::GetInputType(int32 InputIndex)
 {
 	if (InputIndex == 2)
 	{
@@ -25477,12 +25477,12 @@ uint32 UMaterialExpressionStrataVerticalLayering::GetInputType(int32 InputIndex)
 	return MCT_Strata;
 }
 
-bool UMaterialExpressionStrataVerticalLayering::IsResultStrataMaterial(int32 OutputIndex)
+bool UMaterialExpressionSubstrateVerticalLayering::IsResultStrataMaterial(int32 OutputIndex)
 {
 	return true;
 }
 
-void UMaterialExpressionStrataVerticalLayering::GatherStrataMaterialInfo(FStrataMaterialInfo& StrataMaterialInfo, int32 OutputIndex)
+void UMaterialExpressionSubstrateVerticalLayering::GatherStrataMaterialInfo(FStrataMaterialInfo& StrataMaterialInfo, int32 OutputIndex)
 {
 	if (!StrataMaterialInfo.PushStrataTreeStack())
 	{
@@ -25503,7 +25503,7 @@ void UMaterialExpressionStrataVerticalLayering::GatherStrataMaterialInfo(FStrata
 	StrataMaterialInfo.PopStrataTreeStack();
 }
 
-FStrataOperator* UMaterialExpressionStrataVerticalLayering::StrataGenerateMaterialTopologyTree(class FMaterialCompiler* Compiler, class UMaterialExpression* Parent, int32 OutputIndex)
+FStrataOperator* UMaterialExpressionSubstrateVerticalLayering::StrataGenerateMaterialTopologyTree(class FMaterialCompiler* Compiler, class UMaterialExpression* Parent, int32 OutputIndex)
 {
 	FStrataOperator& StrataOperator = Compiler->StrataCompilationRegisterOperator(STRATA_OPERATOR_VERTICAL, Compiler->StrataTreeStackGetPathUniqueId(), this, Parent, Compiler->StrataTreeStackGetParentPathUniqueId(), bUseParameterBlending);
 	if (Compiler->GetStrataTreeOutOfStackDepthOccurred())
@@ -25544,7 +25544,7 @@ FStrataOperator* UMaterialExpressionStrataVerticalLayering::StrataGenerateMateri
 
 
 
-UMaterialExpressionStrataAdd::UMaterialExpressionStrataAdd(const FObjectInitializer& ObjectInitializer)
+UMaterialExpressionSubstrateAdd::UMaterialExpressionSubstrateAdd(const FObjectInitializer& ObjectInitializer)
 	: Super(ObjectInitializer)
 	, bUseParameterBlending(false)
 {
@@ -25560,7 +25560,7 @@ UMaterialExpressionStrataAdd::UMaterialExpressionStrataAdd(const FObjectInitiali
 }
 
 #if WITH_EDITOR
-int32 UMaterialExpressionStrataAdd::Compile(class FMaterialCompiler* Compiler, int32 OutputIndex)
+int32 UMaterialExpressionSubstrateAdd::Compile(class FMaterialCompiler* Compiler, int32 OutputIndex)
 {
 	if (!A.GetTracedInput().Expression)
 	{
@@ -25632,7 +25632,7 @@ int32 UMaterialExpressionStrataAdd::Compile(class FMaterialCompiler* Compiler, i
 	return OutputCodeChunk;
 }
 
-void UMaterialExpressionStrataAdd::GetCaption(TArray<FString>& OutCaptions) const
+void UMaterialExpressionSubstrateAdd::GetCaption(TArray<FString>& OutCaptions) const
 {
 	if (bUseParameterBlending)
 	{
@@ -25644,22 +25644,22 @@ void UMaterialExpressionStrataAdd::GetCaption(TArray<FString>& OutCaptions) cons
 	}
 }
 
-uint32 UMaterialExpressionStrataAdd::GetOutputType(int32 OutputIndex)
+uint32 UMaterialExpressionSubstrateAdd::GetOutputType(int32 OutputIndex)
 {
 	return MCT_Strata;
 }
 
-uint32 UMaterialExpressionStrataAdd::GetInputType(int32 InputIndex)
+uint32 UMaterialExpressionSubstrateAdd::GetInputType(int32 InputIndex)
 {
 	return MCT_Strata;
 }
 
-bool UMaterialExpressionStrataAdd::IsResultStrataMaterial(int32 OutputIndex)
+bool UMaterialExpressionSubstrateAdd::IsResultStrataMaterial(int32 OutputIndex)
 {
 	return true;
 }
 
-void UMaterialExpressionStrataAdd::GatherStrataMaterialInfo(FStrataMaterialInfo& StrataMaterialInfo, int32 OutputIndex)
+void UMaterialExpressionSubstrateAdd::GatherStrataMaterialInfo(FStrataMaterialInfo& StrataMaterialInfo, int32 OutputIndex)
 {
 	if (!StrataMaterialInfo.PushStrataTreeStack())
 	{
@@ -25680,7 +25680,7 @@ void UMaterialExpressionStrataAdd::GatherStrataMaterialInfo(FStrataMaterialInfo&
 	StrataMaterialInfo.PopStrataTreeStack();
 }
 
-FStrataOperator* UMaterialExpressionStrataAdd::StrataGenerateMaterialTopologyTree(class FMaterialCompiler* Compiler, class UMaterialExpression* Parent, int32 OutputIndex)
+FStrataOperator* UMaterialExpressionSubstrateAdd::StrataGenerateMaterialTopologyTree(class FMaterialCompiler* Compiler, class UMaterialExpression* Parent, int32 OutputIndex)
 {
 	FStrataOperator& StrataOperator = Compiler->StrataCompilationRegisterOperator(STRATA_OPERATOR_ADD, Compiler->StrataTreeStackGetPathUniqueId(), this, Parent, Compiler->StrataTreeStackGetParentPathUniqueId(), bUseParameterBlending);
 	if (Compiler->GetStrataTreeOutOfStackDepthOccurred())
@@ -25716,7 +25716,7 @@ FStrataOperator* UMaterialExpressionStrataAdd::StrataGenerateMaterialTopologyTre
 
 
 
-UMaterialExpressionStrataWeight::UMaterialExpressionStrataWeight(const FObjectInitializer& ObjectInitializer)
+UMaterialExpressionSubstrateWeight::UMaterialExpressionSubstrateWeight(const FObjectInitializer& ObjectInitializer)
 	: Super(ObjectInitializer)
 {
 	struct FConstructorStatics
@@ -25731,7 +25731,7 @@ UMaterialExpressionStrataWeight::UMaterialExpressionStrataWeight(const FObjectIn
 }
 
 #if WITH_EDITOR
-int32 UMaterialExpressionStrataWeight::Compile(class FMaterialCompiler* Compiler, int32 OutputIndex)
+int32 UMaterialExpressionSubstrateWeight::Compile(class FMaterialCompiler* Compiler, int32 OutputIndex)
 {
 	if (!A.GetTracedInput().Expression)
 	{
@@ -25777,17 +25777,17 @@ int32 UMaterialExpressionStrataWeight::Compile(class FMaterialCompiler* Compiler
 	return OutputCodeChunk;
 }
 
-void UMaterialExpressionStrataWeight::GetCaption(TArray<FString>& OutCaptions) const
+void UMaterialExpressionSubstrateWeight::GetCaption(TArray<FString>& OutCaptions) const
 {
 	OutCaptions.Add(TEXT("Substrate Coverage Weight"));
 }
 
-uint32 UMaterialExpressionStrataWeight::GetOutputType(int32 OutputIndex)
+uint32 UMaterialExpressionSubstrateWeight::GetOutputType(int32 OutputIndex)
 {
 	return MCT_Strata;
 }
 
-uint32 UMaterialExpressionStrataWeight::GetInputType(int32 InputIndex)
+uint32 UMaterialExpressionSubstrateWeight::GetInputType(int32 InputIndex)
 {
 	if (InputIndex == 0)
 	{
@@ -25796,12 +25796,12 @@ uint32 UMaterialExpressionStrataWeight::GetInputType(int32 InputIndex)
 	return MCT_Float1;
 }
 
-bool UMaterialExpressionStrataWeight::IsResultStrataMaterial(int32 OutputIndex)
+bool UMaterialExpressionSubstrateWeight::IsResultStrataMaterial(int32 OutputIndex)
 {
 	return true;
 }
 
-void UMaterialExpressionStrataWeight::GatherStrataMaterialInfo(FStrataMaterialInfo& StrataMaterialInfo, int32 OutputIndex)
+void UMaterialExpressionSubstrateWeight::GatherStrataMaterialInfo(FStrataMaterialInfo& StrataMaterialInfo, int32 OutputIndex)
 {
 	if (!StrataMaterialInfo.PushStrataTreeStack())
 	{
@@ -25817,7 +25817,7 @@ void UMaterialExpressionStrataWeight::GatherStrataMaterialInfo(FStrataMaterialIn
 	StrataMaterialInfo.PopStrataTreeStack();
 }
 
-FStrataOperator* UMaterialExpressionStrataWeight::StrataGenerateMaterialTopologyTree(class FMaterialCompiler* Compiler, class UMaterialExpression* Parent, int32 OutputIndex)
+FStrataOperator* UMaterialExpressionSubstrateWeight::StrataGenerateMaterialTopologyTree(class FMaterialCompiler* Compiler, class UMaterialExpression* Parent, int32 OutputIndex)
 {
 	FStrataOperator& StrataOperator = Compiler->StrataCompilationRegisterOperator(STRATA_OPERATOR_WEIGHT, Compiler->StrataTreeStackGetPathUniqueId(), this, Parent, Compiler->StrataTreeStackGetParentPathUniqueId());
 	if (Compiler->GetStrataTreeOutOfStackDepthOccurred())
@@ -25841,12 +25841,12 @@ FStrataOperator* UMaterialExpressionStrataWeight::StrataGenerateMaterialTopology
 }
 #endif // WITH_EDITOR
 
-UMaterialExpressionStrataUtilityBase::UMaterialExpressionStrataUtilityBase(const FObjectInitializer& ObjectInitializer)
+UMaterialExpressionSubstrateUtilityBase::UMaterialExpressionSubstrateUtilityBase(const FObjectInitializer& ObjectInitializer)
 	: Super(ObjectInitializer)
 {
 }
 
-UMaterialExpressionStrataTransmittanceToMFP::UMaterialExpressionStrataTransmittanceToMFP(const FObjectInitializer& ObjectInitializer)
+UMaterialExpressionSubstrateTransmittanceToMFP::UMaterialExpressionSubstrateTransmittanceToMFP(const FObjectInitializer& ObjectInitializer)
 	: Super(ObjectInitializer)
 {
 	struct FConstructorStatics
@@ -25867,7 +25867,7 @@ UMaterialExpressionStrataTransmittanceToMFP::UMaterialExpressionStrataTransmitta
 }
 
 #if WITH_EDITOR
-int32 UMaterialExpressionStrataTransmittanceToMFP::Compile(class FMaterialCompiler* Compiler, int32 OutputIndex)
+int32 UMaterialExpressionSubstrateTransmittanceToMFP::Compile(class FMaterialCompiler* Compiler, int32 OutputIndex)
 {
 	int32 TransmittanceColorCodeChunk = TransmittanceColor.GetTracedInput().Expression ? TransmittanceColor.Compile(Compiler) : Compiler->Constant(0.5f);
 	int32 ThicknessCodeChunk = Thickness.GetTracedInput().Expression ? Thickness.Compile(Compiler) : Compiler->Constant(STRATA_LAYER_DEFAULT_THICKNESS_CM);
@@ -25885,12 +25885,12 @@ int32 UMaterialExpressionStrataTransmittanceToMFP::Compile(class FMaterialCompil
 		OutputIndex);
 }
 
-void UMaterialExpressionStrataTransmittanceToMFP::GetCaption(TArray<FString>& OutCaptions) const
+void UMaterialExpressionSubstrateTransmittanceToMFP::GetCaption(TArray<FString>& OutCaptions) const
 {
 	OutCaptions.Add(TEXT("Substrate Transmittance-To-MeanFreePath"));
 }
 
-uint32 UMaterialExpressionStrataTransmittanceToMFP::GetOutputType(int32 OutputIndex)
+uint32 UMaterialExpressionSubstrateTransmittanceToMFP::GetOutputType(int32 OutputIndex)
 {
 	switch (OutputIndex)
 	{
@@ -25906,7 +25906,7 @@ uint32 UMaterialExpressionStrataTransmittanceToMFP::GetOutputType(int32 OutputIn
 	return MCT_Float1;
 }
 
-uint32 UMaterialExpressionStrataTransmittanceToMFP::GetInputType(int32 InputIndex)
+uint32 UMaterialExpressionSubstrateTransmittanceToMFP::GetInputType(int32 InputIndex)
 {
 	switch (InputIndex)
 	{
@@ -25921,7 +25921,7 @@ uint32 UMaterialExpressionStrataTransmittanceToMFP::GetInputType(int32 InputInde
 	check(false);
 	return MCT_Float1;
 }
-void UMaterialExpressionStrataTransmittanceToMFP::GetConnectorToolTip(int32 InputIndex, int32 OutputIndex, TArray<FString>& OutToolTip)
+void UMaterialExpressionSubstrateTransmittanceToMFP::GetConnectorToolTip(int32 InputIndex, int32 OutputIndex, TArray<FString>& OutToolTip)
 {
 	if (InputIndex != INDEX_NONE)
 	{
@@ -25949,14 +25949,14 @@ void UMaterialExpressionStrataTransmittanceToMFP::GetConnectorToolTip(int32 Inpu
 	}
 }
 
-void UMaterialExpressionStrataTransmittanceToMFP::GetExpressionToolTip(TArray<FString>& OutToolTip)
+void UMaterialExpressionSubstrateTransmittanceToMFP::GetExpressionToolTip(TArray<FString>& OutToolTip)
 {
 	ConvertToMultilineToolTip(TEXT("Convert a transmittance color corresponding to a slab of participating media viewed perpendicularly to its surface.\nThis node directly maps to the Slab BSDF input."), 80, OutToolTip);
 
 }
 #endif // WITH_EDITOR
 
-UMaterialExpressionStrataMetalnessToDiffuseAlbedoF0::UMaterialExpressionStrataMetalnessToDiffuseAlbedoF0(const FObjectInitializer& ObjectInitializer)
+UMaterialExpressionSubstrateMetalnessToDiffuseAlbedoF0::UMaterialExpressionSubstrateMetalnessToDiffuseAlbedoF0(const FObjectInitializer& ObjectInitializer)
 	: Super(ObjectInitializer)
 {
 	struct FConstructorStatics
@@ -25977,7 +25977,7 @@ UMaterialExpressionStrataMetalnessToDiffuseAlbedoF0::UMaterialExpressionStrataMe
 }
 
 #if WITH_EDITOR
-int32 UMaterialExpressionStrataMetalnessToDiffuseAlbedoF0::Compile(class FMaterialCompiler* Compiler, int32 OutputIndex)
+int32 UMaterialExpressionSubstrateMetalnessToDiffuseAlbedoF0::Compile(class FMaterialCompiler* Compiler, int32 OutputIndex)
 {
 	int32 BaseColorCodeChunk = BaseColor.GetTracedInput().Expression ? BaseColor.Compile(Compiler) : Compiler->Constant(0.18f);
 	int32 SpecularCodeChunk = Specular.GetTracedInput().Expression	 ? Specular.Compile(Compiler) : Compiler->Constant(0.5f);
@@ -26001,12 +26001,12 @@ int32 UMaterialExpressionStrataMetalnessToDiffuseAlbedoF0::Compile(class FMateri
 		OutputIndex);
 }
 
-void UMaterialExpressionStrataMetalnessToDiffuseAlbedoF0::GetCaption(TArray<FString>& OutCaptions) const
+void UMaterialExpressionSubstrateMetalnessToDiffuseAlbedoF0::GetCaption(TArray<FString>& OutCaptions) const
 {
 	OutCaptions.Add(TEXT("Substrate Metalness-To-DiffuseAlbedo-F0"));
 }
 
-uint32 UMaterialExpressionStrataMetalnessToDiffuseAlbedoF0::GetOutputType(int32 OutputIndex)
+uint32 UMaterialExpressionSubstrateMetalnessToDiffuseAlbedoF0::GetOutputType(int32 OutputIndex)
 {
 	switch (OutputIndex)
 	{
@@ -26022,7 +26022,7 @@ uint32 UMaterialExpressionStrataMetalnessToDiffuseAlbedoF0::GetOutputType(int32
 	return MCT_Float1;
 }
 
-uint32 UMaterialExpressionStrataMetalnessToDiffuseAlbedoF0::GetInputType(int32 InputIndex)
+uint32 UMaterialExpressionSubstrateMetalnessToDiffuseAlbedoF0::GetInputType(int32 InputIndex)
 {
 	if (InputIndex == 0) { return MCT_Float3; }
 	if (InputIndex == 1) { return MCT_Float1; }
@@ -26030,7 +26030,7 @@ uint32 UMaterialExpressionStrataMetalnessToDiffuseAlbedoF0::GetInputType(int32 I
 	return MCT_Float1;
 }
 
-void UMaterialExpressionStrataMetalnessToDiffuseAlbedoF0::GetConnectorToolTip(int32 InputIndex, int32 OutputIndex, TArray<FString>& OutToolTip)
+void UMaterialExpressionSubstrateMetalnessToDiffuseAlbedoF0::GetConnectorToolTip(int32 InputIndex, int32 OutputIndex, TArray<FString>& OutToolTip)
 {
 	switch (OutputIndex)
 	{
@@ -26040,7 +26040,7 @@ void UMaterialExpressionStrataMetalnessToDiffuseAlbedoF0::GetConnectorToolTip(in
 	}
 }
 
-void UMaterialExpressionStrataMetalnessToDiffuseAlbedoF0::GetExpressionToolTip(TArray<FString>& OutToolTip)
+void UMaterialExpressionSubstrateMetalnessToDiffuseAlbedoF0::GetExpressionToolTip(TArray<FString>& OutToolTip)
 {
 	ConvertToMultilineToolTip(TEXT("Convert a metalness parameterization (BaseColor/Specular/Metallic) into DiffuseAlbedo/F0 parameterization."), 80, OutToolTip);
 
@@ -26049,7 +26049,7 @@ void UMaterialExpressionStrataMetalnessToDiffuseAlbedoF0::GetExpressionToolTip(T
 
 
 
-UMaterialExpressionStrataHazinessToSecondaryRoughness::UMaterialExpressionStrataHazinessToSecondaryRoughness(const FObjectInitializer& ObjectInitializer)
+UMaterialExpressionSubstrateHazinessToSecondaryRoughness::UMaterialExpressionSubstrateHazinessToSecondaryRoughness(const FObjectInitializer& ObjectInitializer)
 	: Super(ObjectInitializer)
 {
 	struct FConstructorStatics
@@ -26070,7 +26070,7 @@ UMaterialExpressionStrataHazinessToSecondaryRoughness::UMaterialExpressionStrata
 }
 
 #if WITH_EDITOR
-int32 UMaterialExpressionStrataHazinessToSecondaryRoughness::Compile(class FMaterialCompiler* Compiler, int32 OutputIndex)
+int32 UMaterialExpressionSubstrateHazinessToSecondaryRoughness::Compile(class FMaterialCompiler* Compiler, int32 OutputIndex)
 {
 	int32 BaseRoughnessCodeChunk= BaseRoughness.GetTracedInput().Expression ? BaseRoughness.Compile(Compiler) : Compiler->Constant(0.1f);
 	int32 HazinessCodeChunk		= Haziness.GetTracedInput().Expression      ? Haziness.Compile(Compiler)      : Compiler->Constant(0.5f);
@@ -26088,12 +26088,12 @@ int32 UMaterialExpressionStrataHazinessToSecondaryRoughness::Compile(class FMate
 		OutputIndex);
 }
 
-void UMaterialExpressionStrataHazinessToSecondaryRoughness::GetCaption(TArray<FString>& OutCaptions) const
+void UMaterialExpressionSubstrateHazinessToSecondaryRoughness::GetCaption(TArray<FString>& OutCaptions) const
 {
 	OutCaptions.Add(TEXT("Substrate Haziness-To-Secondary-Roughness"));
 }
 
-uint32 UMaterialExpressionStrataHazinessToSecondaryRoughness::GetOutputType(int32 OutputIndex)
+uint32 UMaterialExpressionSubstrateHazinessToSecondaryRoughness::GetOutputType(int32 OutputIndex)
 {
 	switch (OutputIndex)
 	{
@@ -26109,7 +26109,7 @@ uint32 UMaterialExpressionStrataHazinessToSecondaryRoughness::GetOutputType(int3
 	return MCT_Float1;
 }
 
-uint32 UMaterialExpressionStrataHazinessToSecondaryRoughness::GetInputType(int32 InputIndex)
+uint32 UMaterialExpressionSubstrateHazinessToSecondaryRoughness::GetInputType(int32 InputIndex)
 {
 	switch (InputIndex)
 	{
@@ -26124,7 +26124,7 @@ uint32 UMaterialExpressionStrataHazinessToSecondaryRoughness::GetInputType(int32
 	check(false);
 	return MCT_Float1;
 }
-void UMaterialExpressionStrataHazinessToSecondaryRoughness::GetConnectorToolTip(int32 InputIndex, int32 OutputIndex, TArray<FString>& OutToolTip)
+void UMaterialExpressionSubstrateHazinessToSecondaryRoughness::GetConnectorToolTip(int32 InputIndex, int32 OutputIndex, TArray<FString>& OutToolTip)
 {
 	if (InputIndex != INDEX_NONE)
 	{
@@ -26152,7 +26152,7 @@ void UMaterialExpressionStrataHazinessToSecondaryRoughness::GetConnectorToolTip(
 	}
 }
 
-void UMaterialExpressionStrataHazinessToSecondaryRoughness::GetExpressionToolTip(TArray<FString>& OutToolTip)
+void UMaterialExpressionSubstrateHazinessToSecondaryRoughness::GetExpressionToolTip(TArray<FString>& OutToolTip)
 {
 	ConvertToMultilineToolTip(TEXT("Compute a second specular lobe roughness from a base surface roughness and haziness. This parameterisation ensure that the haziness makes physically and is perceptually easy to author."), 80, OutToolTip);
 
@@ -26161,7 +26161,7 @@ void UMaterialExpressionStrataHazinessToSecondaryRoughness::GetExpressionToolTip
 
 
 
-UMaterialExpressionStrataThinFilm::UMaterialExpressionStrataThinFilm(const FObjectInitializer& ObjectInitializer)
+UMaterialExpressionSubstrateThinFilm::UMaterialExpressionSubstrateThinFilm(const FObjectInitializer& ObjectInitializer)
 	: Super(ObjectInitializer)
 {
 	struct FConstructorStatics
@@ -26182,7 +26182,7 @@ UMaterialExpressionStrataThinFilm::UMaterialExpressionStrataThinFilm(const FObje
 }
 
 #if WITH_EDITOR
-int32 UMaterialExpressionStrataThinFilm::Compile(class FMaterialCompiler* Compiler, int32 OutputIndex)
+int32 UMaterialExpressionSubstrateThinFilm::Compile(class FMaterialCompiler* Compiler, int32 OutputIndex)
 {
 	int32 NormalCodeChunk		= CompileWithDefaultNormalWS(Compiler, Normal);
 
@@ -26216,12 +26216,12 @@ int32 UMaterialExpressionStrataThinFilm::Compile(class FMaterialCompiler* Compil
 	return Compiler->StrataThinFilm(NormalCodeChunk, F0CodeChunk, F90CodeChunk, ThicknessCodeChunk, IORCodeChunk, OutputIndex);
 }
 
-void UMaterialExpressionStrataThinFilm::GetCaption(TArray<FString>& OutCaptions) const
+void UMaterialExpressionSubstrateThinFilm::GetCaption(TArray<FString>& OutCaptions) const
 {
 	OutCaptions.Add(TEXT("Substrate Thin-Film"));
 }
 
-uint32 UMaterialExpressionStrataThinFilm::GetOutputType(int32 OutputIndex)
+uint32 UMaterialExpressionSubstrateThinFilm::GetOutputType(int32 OutputIndex)
 {
 	switch (OutputIndex)
 	{
@@ -26237,7 +26237,7 @@ uint32 UMaterialExpressionStrataThinFilm::GetOutputType(int32 OutputIndex)
 	return MCT_Float1;
 }
 
-uint32 UMaterialExpressionStrataThinFilm::GetInputType(int32 InputIndex)
+uint32 UMaterialExpressionSubstrateThinFilm::GetInputType(int32 InputIndex)
 {
 	if (InputIndex == 0) { return MCT_Float3; } // Normal
 	if (InputIndex == 1) { return MCT_Float3; } // F0
@@ -26248,7 +26248,7 @@ uint32 UMaterialExpressionStrataThinFilm::GetInputType(int32 InputIndex)
 	check(false);
 	return MCT_Float1;
 }
-void UMaterialExpressionStrataThinFilm::GetConnectorToolTip(int32 InputIndex, int32 OutputIndex, TArray<FString>& OutToolTip)
+void UMaterialExpressionSubstrateThinFilm::GetConnectorToolTip(int32 InputIndex, int32 OutputIndex, TArray<FString>& OutToolTip)
 {
 	if (OutputIndex != INDEX_NONE)
 	{
@@ -26268,7 +26268,7 @@ void UMaterialExpressionStrataThinFilm::GetConnectorToolTip(int32 InputIndex, in
 	Super::GetConnectorToolTip(InputIndex, OutputIndex, OutToolTip);
 }
 
-void UMaterialExpressionStrataThinFilm::GetExpressionToolTip(TArray<FString>& OutToolTip)
+void UMaterialExpressionSubstrateThinFilm::GetExpressionToolTip(TArray<FString>& OutToolTip)
 {
 	ConvertToMultilineToolTip(TEXT("Compute the resulting material specular parameter F0 and F90 according to input surface properties as well as the thin film parameters."), 80, OutToolTip);
 }
