[Resubmit with Fastcook fixes] Inline first chunk to work for prime on demand as well

diff --git a/Engine/Source/Runtime/Engine/Classes/Sound/SoundClass.h b/Engine/Source/Runtime/Engine/Classes/Sound/SoundClass.h
--- a/Engine/Source/Runtime/Engine/Classes/Sound/SoundClass.h
+++ b/Engine/Source/Runtime/Engine/Classes/Sound/SoundClass.h
@@ -135,8 +135,8 @@ struct FSoundClassProperties
 	ESoundWaveLoadingBehavior LoadingBehavior;
 
 #if WITH_EDITORONLY_DATA
-   	/** How much audio to add to First Audio Chunk (in seconds), (only for RetainOnLoad) */
-   	UPROPERTY(EditAnywhere, Category = Loading, meta = (UIMin = 0, UIMax = 10, EditCondition = "LoadingBehavior == ESoundWaveLoadingBehavior::RetainOnLoad"), DisplayName="Size of First Audio Chunk (seconds)")
+   	/** How much audio to add to First Audio Chunk (in seconds) */
+	UPROPERTY(EditAnywhere, Category = Loading, meta = (UIMin = 0, UIMax = 10, EditCondition = "LoadingBehavior == ESoundWaveLoadingBehavior::RetainOnLoad || LoadingBehavior == ESoundWaveLoadingBehavior::PrimeOnLoad"), DisplayName="Size of First Audio Chunk (seconds)")
    	FPerPlatformFloat SizeOfFirstAudioChunkInSeconds = 0.0f;
 #endif //WITH_EDITORONLY_DATA
 
diff --git a/Engine/Source/Runtime/Engine/Classes/Sound/SoundWave.h b/Engine/Source/Runtime/Engine/Classes/Sound/SoundWave.h
--- a/Engine/Source/Runtime/Engine/Classes/Sound/SoundWave.h
+++ b/Engine/Source/Runtime/Engine/Classes/Sound/SoundWave.h
@@ -721,8 +721,8 @@ public:
 
 #if WITH_EDITORONLY_DATA
 public:
-   	/** How much audio to add to First Audio Chunk (in seconds), (only for RetainOnLoad) */
-   	UPROPERTY(EditAnywhere, Category = Loading, meta = (UIMin = 0, UIMax = 10, EditCondition = "LoadingBehavior == ESoundWaveLoadingBehavior::RetainOnLoad"), DisplayName="Size of First Audio Chunk (seconds)")
+   	/** How much audio to add to First Audio Chunk (in seconds) */
+	UPROPERTY(EditAnywhere, Category = Loading, meta = (UIMin = 0, UIMax = 10, EditCondition = "LoadingBehavior == ESoundWaveLoadingBehavior::RetainOnLoad || LoadingBehavior == ESoundWaveLoadingBehavior::PrimeOnLoad"), DisplayName="Size of First Audio Chunk (seconds)")
    	FPerPlatformFloat SizeOfFirstAudioChunkInSeconds = 0.0f;
 #endif //WITH_EDITOR_ONLY_DATA
 
diff --git a/Engine/Source/Runtime/Engine/Classes/Sound/SoundWaveLoadingBehavior.cpp b/Engine/Source/Runtime/Engine/Classes/Sound/SoundWaveLoadingBehavior.cpp
--- a/Engine/Source/Runtime/Engine/Classes/Sound/SoundWaveLoadingBehavior.cpp
+++ b/Engine/Source/Runtime/Engine/Classes/Sound/SoundWaveLoadingBehavior.cpp
@@ -13,25 +13,16 @@
 #include "SoundWave.h"
 #include "UObject/LinkerLoad.h"
 
+#ifndef CASE_ENUM_TO_TEXT
+#define CASE_ENUM_TO_TEXT(TXT) case TXT: return TEXT(#TXT);
+#endif
 
-const TCHAR* EnumToString(ESoundWaveLoadingBehavior InCurrentState)
+const TCHAR* EnumToString(const ESoundWaveLoadingBehavior InCurrentState)
 {
-	switch (InCurrentState)
+	switch(InCurrentState)
 	{
-	case ESoundWaveLoadingBehavior::Inherited:
-		return TEXT("Inherited");
-	case ESoundWaveLoadingBehavior::RetainOnLoad:
-		return TEXT("RetainOnLoad");
-	case ESoundWaveLoadingBehavior::PrimeOnLoad:
-		return TEXT("PrimeOnLoad");
-	case ESoundWaveLoadingBehavior::LoadOnDemand:
-		return TEXT("LoadOnDemand");
-	case ESoundWaveLoadingBehavior::ForceInline:
-		return TEXT("ForceInline");
-	case ESoundWaveLoadingBehavior::Uninitialized:
-		return TEXT("Uninitialized");
-	}
-	ensure(false);
+		FOREACH_ENUM_ESOUNDWAVELOADINGBEHAVIOR(CASE_ENUM_TO_TEXT)
+	} 
 	return TEXT("Unknown");
 }
 
@@ -70,13 +61,12 @@ public:
 private:
 	void CacheAllClassLoadingBehaviors()
 	{
-		TRACE_CPUPROFILER_EVENT_SCOPE(LoadAndCacheAllSoundClassLoadingBehaviors);
-
-		ensureMsgf(!AssetRegistry.IsSearchAsync() || !AssetRegistry.IsLoadingAssets(), 
-		           TEXT("Function must not be called until after cook has started and waited on the AssetRegistry already."));
+		const bool bAssetRegistryInStartup = AssetRegistry.IsSearchAsync() && AssetRegistry.IsLoadingAssets();
+		if (!ensureMsgf(!bAssetRegistryInStartup,TEXT("Function must not be called until after cook has started and waited on the AssetRegistry already.")))
+		{
+			AssetRegistry.WaitForCompletion();
+		}
 		
-		AssetRegistry.SearchAllAssets(true);
-
 		TArray<FAssetData> SoundClasses;
 		AssetRegistry.GetAssetsByClass(USoundClass::StaticClass()->GetClassPathName(), SoundClasses, true);
 		
@@ -124,7 +114,7 @@ private:
 		return {};
 	}
 	
-	virtual FClassData FindOwningLoadingBehavior(const USoundWave* InWave, const ITargetPlatform* InTargetPlatform) const 
+	virtual FClassData FindOwningLoadingBehavior(const USoundWave* InWave, const ITargetPlatform* InTargetPlatform) const override
 	{
 		TRACE_CPUPROFILER_EVENT_SCOPE(FindOwningLoadingBehavior);
 
@@ -133,6 +123,15 @@ private:
 		// Then stack ranks the most important behavior. (RetainOnLoad (Highest), PrimeOnLoad (Medium), LoadOnDemand (Lowest))
 		// Which ever wins, we also capture the "SizeOfFirstChunk" to use for that wave.
 
+		const bool bIsAssetRegistryStartup = AssetRegistry.IsSearchAsync() && AssetRegistry.IsLoadingAssets();
+
+		// Disallow during startup of registry (cookers will have already done this)
+		if (bIsAssetRegistryStartup)
+		{
+			UE_LOG(LogAudio, Warning, TEXT("FindOwningLoadingBehavior called before AssetRegistry is ready. SoundWave=%s"), *InWave->GetName());
+			return {};
+		}
+
 		if (!InWave)
 		{
 			return {};
@@ -238,15 +237,13 @@ ISoundWaveLoadingBehaviorUtil* ISoundWaveLoadingBehaviorUtil::Get()
 	{
 		return nullptr;
 	}
-	
-	static bool bAllowOutsideOfCookCommandlet = FParse::Param(FCommandLine::Get(), TEXT("AllowSoundWaveOwnerLoadingBehaviorInEditor"));
-	
+		
 	// Only run while the cooker is active.
-	if (!IsRunningCookCommandlet() && !bAllowOutsideOfCookCommandlet)
-	{
+	if (!IsRunningCookCommandlet())
+	{		
 		return nullptr;
 	}
-	
+		
 	static FSoundWaveLoadingBehaviorUtil Instance;
 	return &Instance;
 }
@@ -266,8 +263,9 @@ bool ISoundWaveLoadingBehaviorUtil::FClassData::CompareGreater(const FClassData&
 	}
 	else 
 	{
-		// If we are using Retain, use one with the higher Length.
-		if (LoadingBehavior == ESoundWaveLoadingBehavior::RetainOnLoad)
+		// If we are using Prime/Retain, use one with the higher Length.
+		if (LoadingBehavior == ESoundWaveLoadingBehavior::RetainOnLoad ||
+			LoadingBehavior == ESoundWaveLoadingBehavior::PrimeOnLoad )
 		{
 			const float Length = LengthOfFirstChunkInSeconds.GetValueForPlatform(*InPlatform->PlatformName());
 			const float OtherLength = InOther.LengthOfFirstChunkInSeconds.GetValueForPlatform(*InPlatform->PlatformName());
diff --git a/Engine/Source/Runtime/Engine/Private/SoundWave.cpp b/Engine/Source/Runtime/Engine/Private/SoundWave.cpp
--- a/Engine/Source/Runtime/Engine/Private/SoundWave.cpp
+++ b/Engine/Source/Runtime/Engine/Private/SoundWave.cpp
@@ -1048,7 +1048,7 @@ void USoundWave::Serialize( FArchive& Ar )
 			const ISoundWaveLoadingBehaviorUtil::FClassData OwnerBehavior = GetOwnerLoadingBehavior(Ar.CookingTarget());
 			if (OwnerBehavior.LoadingBehavior != ESoundWaveLoadingBehavior::Uninitialized)
 			{
-				UE_LOG(LogAudio, Display, TEXT("Saving Owners Loading Behavior: %s, %s"), 
+				UE_LOG(LogAudio, VeryVerbose, TEXT("Saving Owners Loading Behavior: %s, %s"), 
 				       EnumToString(OwnerBehavior.LoadingBehavior), *GetFullName());
 				
 				Flags |= HasOwnerLoadingBehaviorFlag;
@@ -2303,25 +2303,39 @@ float USoundWave::GetSizeOfFirstAudioChunkInSeconds(const ITargetPlatform* InPla
 	
 	ensure(InPlatform);
 
+	// This is the most important function for determining the size of chunk 1
+
 	// We need to walk up the hierarchy of sound classes to first determine our FirstAudioChunkSize.
 	if (SoundWaveDataPtr->LoadingBehavior == ESoundWaveLoadingBehavior::Uninitialized)
 	{
 		CacheInheritedLoadingBehavior();
 	}
 
-	// Use owner loading behavior if we have it.
-	ISoundWaveLoadingBehaviorUtil::FClassData Behavior = GetOwnerLoadingBehavior(InPlatform);
-	if (Behavior.LoadingBehavior == ESoundWaveLoadingBehavior::Uninitialized)
+	// Order of importance for where we get this setting.
+	// Wave -> Owner -> Class
+
+	// If Wave has a non-inherited Behavior set, use that.
+	ISoundWaveLoadingBehaviorUtil::FClassData Behavior;
+	if (LoadingBehavior != ESoundWaveLoadingBehavior::Inherited)
+	{
+		Behavior.LoadingBehavior = LoadingBehavior;
+		Behavior.LengthOfFirstChunkInSeconds = SizeOfFirstAudioChunkInSeconds;
+	}
+	else // ... if we are inherited, try the owner first.
+	{
+		Behavior = GetOwnerLoadingBehavior(InPlatform);
+	}
+
+	// ... otherwise (or if we failed above) use this soundwave resolved behavior,
+	if (Behavior.LoadingBehavior == ESoundWaveLoadingBehavior::Uninitialized)		
 	{
-		// Otherwise use what we've been tagged
 		Behavior.LoadingBehavior = SoundWaveDataPtr->LoadingBehavior;
 		Behavior.LengthOfFirstChunkInSeconds = SoundWaveDataPtr->SizeOfFirstAudioChunkInSeconds;
 	}
 
-	// We only allow setting the first chunk size if we're planning on retaining it.
-	// This could be relaxed later once we have a better format that wouldn't cause a huge delta
-	// on tuning the size (all assets etc). By using just retained, we expect a smaller set.
-	if (Behavior.LoadingBehavior == ESoundWaveLoadingBehavior::RetainOnLoad)
+	// Only give values if the behavior is Retain/Prime.
+	if (Behavior.LoadingBehavior == ESoundWaveLoadingBehavior::RetainOnLoad ||
+		Behavior.LoadingBehavior == ESoundWaveLoadingBehavior::PrimeOnLoad)
 	{
 		// See if we have an override for this platform on the SoundClass hierarchy.
 		const FName PlatformName = InPlatform ? InPlatform->GetPlatformInfo().IniPlatformName : FName();
@@ -4108,16 +4122,21 @@ void USoundWave::OverrideLoadingBehavior(ESoundWaveLoadingBehavior InLoadingBeha
 {
 	TRACE_CPUPROFILER_EVENT_SCOPE(USoundWave::OverrideLoadingBehavior);
 
+	// Don't allow when running the cook commandlet.
+	if (IsRunningCookCommandlet())
+	{
+		return;
+	}
+
 	const bool bCanOverrideLoadingBehavior = (LoadingBehavior != ESoundWaveLoadingBehavior::ForceInline);
 
 	if (bCanOverrideLoadingBehavior)
 	{
 		const ESoundWaveLoadingBehavior OldBehavior = GetLoadingBehavior(false);
-		const bool bAlreadySetToRetained = (OldBehavior == ESoundWaveLoadingBehavior::RetainOnLoad);
 		const bool bAlreadyLoaded = !HasAnyFlags(RF_NeedLoad);
 
-		// already set to the most aggressive (non-inline) option
-		if (bAlreadySetToRetained)
+		// Already set to a higher setting (lower enum values are more important)
+		if (InLoadingBehavior > OldBehavior)
 		{
 			return;
 		}
