fix comments, parameter checking.

diff --git a/Engine/Source/Runtime/Engine/Classes/Engine/Texture.h b/Engine/Source/Runtime/Engine/Classes/Engine/Texture.h
--- a/Engine/Source/Runtime/Engine/Classes/Engine/Texture.h
+++ b/Engine/Source/Runtime/Engine/Classes/Engine/Texture.h
@@ -1915,8 +1915,9 @@ public:
 	ENGINE_API static int32 GetMaximumDimensionOfNonVT();
 
 	/*
-	 * Downsize the 2D Image with the build setting of the texture.
-	 * Try to get as close as it can to the target resolution but it will stay above it if it can't reach it
+	 * Downsize the 2D Image with the build settings for the texture until all dimensions are <= TargetSize.
+	 * This downsizes using the mip generation system and so will only cut sizes in half. Return false
+	 * if a parameter is invalid. Returns true if the output is <= TargetSize, whether or not downsizing occurred.
 	 */
 	ENGINE_API bool DownsizeImageUsingTextureSettings(const ITargetPlatform* TargetPlatform, FImage& InOutImage, int32 TargetSize, int32 LayerIndex);
 
diff --git a/Engine/Source/Runtime/Engine/Private/TextureDerivedData.cpp b/Engine/Source/Runtime/Engine/Private/TextureDerivedData.cpp
--- a/Engine/Source/Runtime/Engine/Private/TextureDerivedData.cpp
+++ b/Engine/Source/Runtime/Engine/Private/TextureDerivedData.cpp
@@ -3902,16 +3902,29 @@ bool UTexture::DownsizeImageUsingTextureSettings(const ITargetPlatform* TargetPl
 {
 	// resize so that the largest dimension is <= TargetSize
 
+	if (TargetSize < 1 || LayerIndex < 0 || InOutImage.IsImageInfoValid() == false)
+	{
+		UE_LOG(LogTexture, Error, TEXT("Invalid parameter supplied to DownsizeImageUsingTextureSettings target size = %d layer index = %d image valid: %s"),
+			TargetSize, LayerIndex, InOutImage.IsImageInfoValid() ? TEXT("true") : TEXT("false"));
+		return false;
+	}
+
 	if (TargetSize >= InOutImage.SizeX && TargetSize >= InOutImage.SizeY)
 	{
 		// both dimensions already small enough, early out
-		return false;
+		return true;
 	}
 
 	// Ideally this code wouldn't live here but at the moment of writing this code the coupling between the texture and the texture compressor make it hard to move that logic elsewhere
 	TArray<FTextureBuildSettings> SettingPerLayer;
 	GetBuildSettingsForTargetPlatform(*this, TargetPlatform, ETextureEncodeSpeed::Final, SettingPerLayer, nullptr);
 
+	if (LayerIndex >= SettingPerLayer.Num())
+	{
+		UE_LOG(LogTexture, Error, TEXT("Invalid layer supplied to DownsizeImageUsingTextureSettings, layer index = %d"), LayerIndex);
+		return false;
+	}
+
 	// Teak the build setting to generate a mip for our image
 	FTextureBuildSettings& BuildSettings = SettingPerLayer[LayerIndex];
 	BuildSettings.bCubemap = false;
