FWorldPartitionStremaingSource::ToString() now also reports more details on shapes/targetgrids

diff --git a/Engine/Source/Runtime/Engine/Private/WorldPartition/WorldPartitionStreamingSource.cpp b/Engine/Source/Runtime/Engine/Private/WorldPartition/WorldPartitionStreamingSource.cpp
--- a/Engine/Source/Runtime/Engine/Private/WorldPartition/WorldPartitionStreamingSource.cpp
+++ b/Engine/Source/Runtime/Engine/Private/WorldPartition/WorldPartitionStreamingSource.cpp
@@ -20,6 +20,67 @@ FAutoConsoleVariableRef FWorldPartitionStreamingSource::CVarRotationQuantization
 	TEXT("Angle (in degrees) used to quantize the streaming sources rotation to determine if a world partition streaming update is necessary."),
 	ECVF_Default);
 
+FString FStreamingSourceShape::ToString() const
+{
+	FStringBuilderBase StringBuilder;
+	if (bIsSector)
+	{
+		StringBuilder += TEXT("IsSector ");
+	}
+	if (bUseGridLoadingRange)
+	{
+		StringBuilder += TEXT("UsesGridLoadingRange ");
+		if (!FMath::IsNearlyEqual(LoadingRangeScale, 1.f))
+		{
+			StringBuilder.Appendf(TEXT("Scale: 3.2f"), LoadingRangeScale);
+		}
+	}
+	else
+	{
+		StringBuilder.Appendf(TEXT("Radius: %d"), (int32)Radius);
+	}
+	return StringBuilder.ToString();
+}
+
+FString FWorldPartitionStreamingSource::ToString() const
+{
+	FStringBuilderBase StringBuilder;
+	StringBuilder.Appendf(
+		TEXT("Priority: %d | %s | %s | %s | Pos: X=%lld,Y=%lld,Z=%lld | Rot: %s | Vel: %3.2f m/s (%d mph)"),
+		Priority,
+		bRemote ? TEXT("Remote") : TEXT("Local"),
+		GetStreamingSourceTargetStateName(TargetState),
+		bBlockOnSlowLoading ? TEXT("Blocking") : TEXT("NonBlocking"),
+		(int64)Location.X, (int64)Location.Y, (int64)Location.Z,
+		*Rotation.ToCompactString(),
+		Velocity,
+		(int32)(Velocity * 2.23694f)
+	);
+
+	if (Shapes.Num())
+	{
+		StringBuilder += TEXT(" | ");
+		int32 ShapeIndex = 0;
+		for (const FStreamingSourceShape& Shape : Shapes)
+		{
+			StringBuilder.Appendf(TEXT("Shape[%d]: %s "), ShapeIndex++, *Shape.ToString());
+		}
+		StringBuilder.RemoveSuffix(1);
+	}
+
+	if (TargetGrids.Num())
+	{
+		StringBuilder.Appendf(TEXT(" | %s TargetGrids: "), (TargetBehavior == EStreamingSourceTargetBehavior::Include) ? TEXT("Included") : TEXT("Excluded"));
+		for (const FName& TargetGrid : TargetGrids)
+		{
+			StringBuilder.Appendf(TEXT("%s, "), *TargetGrid.ToString());
+		}
+		StringBuilder.RemoveSuffix(2);
+	}
+
+	return StringBuilder.ToString();
+}
+
 void FWorldPartitionStreamingSource::UpdateHash()
 {
 	// Update old values when they are changing enough, to avoid the case where we are on the edge of a quantization unit.
@@ -172,4 +233,4 @@ bool FSphericalSector::IntersectsBox(const FBox2D& InBox) const
 	}
 
 	return false;
-}
+}
\ No newline at end of file
diff --git a/Engine/Source/Runtime/Engine/Public/WorldPartition/WorldPartitionStreamingSource.h b/Engine/Source/Runtime/Engine/Public/WorldPartition/WorldPartitionStreamingSource.h
--- a/Engine/Source/Runtime/Engine/Public/WorldPartition/WorldPartitionStreamingSource.h
+++ b/Engine/Source/Runtime/Engine/Public/WorldPartition/WorldPartitionStreamingSource.h
@@ -125,6 +125,8 @@ struct FStreamingSourceShape
 	UPROPERTY(BlueprintReadWrite, EditAnywhere, Category = Streaming)
 	FRotator Rotation;
 
+	FString ToString() const;
+
 	friend uint32 GetTypeHash(const FStreamingSourceShape& InShape)
 	{
 		uint32 Hash = GetTypeHash(InShape.bUseGridLoadingRange);
@@ -492,20 +494,7 @@ struct FWorldPartitionStreamingSource
 		}
 	}
 
-	FString ToString() const
-	{
-		return FString::Printf(
-			TEXT("Priority: %d | %s | %s | %s | Pos: X=%lld,Y=%lld,Z=%lld | Rot: %s | Vel: %3.2f m/s (%d mph)"), 
-			Priority, 
-			bRemote ? TEXT("Remote") : TEXT("Local"),
-			GetStreamingSourceTargetStateName(TargetState),
-			bBlockOnSlowLoading ? TEXT("Blocking") : TEXT("NonBlocking"),
-			(int64)Location.X, (int64)Location.Y, (int64)Location.Z, 
-			*Rotation.ToCompactString(),
-			Velocity, 
-			(int32)(Velocity*2.23694f)
-		);
-	}
+	FString ToString() const;
 
 	static int32 GetLocationQuantization() { return LocationQuantization; }
 	static int32 GetRotationQuantization() { return RotationQuantization; }
