Camera: Tooltips for CameraShake and AutoStart

diff --git a/Engine/Source/Runtime/Engine/Classes/Camera/CameraShakeSourceComponent.h b/Engine/Source/Runtime/Engine/Classes/Camera/CameraShakeSourceComponent.h
--- a/Engine/Source/Runtime/Engine/Classes/Camera/CameraShakeSourceComponent.h
+++ b/Engine/Source/Runtime/Engine/Classes/Camera/CameraShakeSourceComponent.h
@@ -56,9 +56,11 @@ public:
 	UPROPERTY(EditAnywhere, BlueprintReadWrite, Category = Attenuation)
 	float OuterAttenuationRadius;
 
+	/** The camera shake class to use for this camera shake source actor */
 	UPROPERTY(EditAnywhere, BlueprintReadWrite, Category = CameraShake)
 	TSubclassOf<UCameraShakeBase> CameraShake;
 
+	/** Whether to auto start when created */
 	UPROPERTY(EditAnywhere, BlueprintReadWrite, Category = CameraShake)
 	bool bAutoStart;
 
