Added GetTextId util to KismetTextLibrary

diff --git a/Engine/Source/Runtime/Engine/Classes/Kismet/KismetTextLibrary.h b/Engine/Source/Runtime/Engine/Classes/Kismet/KismetTextLibrary.h
--- a/Engine/Source/Runtime/Engine/Classes/Kismet/KismetTextLibrary.h
+++ b/Engine/Source/Runtime/Engine/Classes/Kismet/KismetTextLibrary.h
@@ -287,12 +287,19 @@ class UKismetTextLibrary : public UBlueprintFunctionLibrary
 	static ENGINE_API FText TextFromStringTable(const FName TableId, const FString& Key);
 
 	/**
-	 * Attempts to find the String Table ID and key used by the given text.
+	 * Attempts to get the String Table ID and key used by the given text.
 	 * @return True if the String Table ID and key were found, false otherwise.
 	 */
-	UFUNCTION(BlueprintPure, Category="Utilities|Text", meta=(DisplayName="Find String Table ID and Key from Text"))
+	UFUNCTION(BlueprintPure, Category="Utilities|Text", meta=(DisplayName="Get String Table ID and Key from Text"))
 	static ENGINE_API bool StringTableIdAndKeyFromText(FText Text, FName& OutTableId, FString& OutKey);
 
+	/**
+	 * Attempts to get the ID (namespace and key) used by the given text.
+	 * @return True if the namespace (which may be empty) and key were found, false otherwise.
+	 */
+	UFUNCTION(BlueprintPure, Category="Utilities|Text", meta=(DisplayName="Get ID from Text"))
+	static ENGINE_API bool GetTextId(FText Text, FString& OutNamespace, FString& OutKey);
+
 	/**
 	 * Check whether the given polyglot data is valid.
 	 * @return True if the polyglot data is valid, false otherwise. ErrorMessage will be filled in if the the data is invalid.
diff --git a/Engine/Source/Runtime/Engine/Private/KismetTextLibrary.cpp b/Engine/Source/Runtime/Engine/Private/KismetTextLibrary.cpp
--- a/Engine/Source/Runtime/Engine/Private/KismetTextLibrary.cpp
+++ b/Engine/Source/Runtime/Engine/Private/KismetTextLibrary.cpp
@@ -317,6 +317,18 @@ bool UKismetTextLibrary::StringTableIdAndKeyFromText(FText Text, FName& OutTable
 	return FTextInspector::GetTableIdAndKey(Text, OutTableId, OutKey);
 }
 
+bool UKismetTextLibrary::GetTextId(FText Text, FString& OutNamespace, FString& OutKey)
+{
+	const FTextId TextId = FTextInspector::GetTextId(Text);
+	if (!TextId.IsEmpty())
+	{
+		OutNamespace = TextId.GetNamespace().GetChars();
+		OutKey = TextId.GetKey().GetChars();
+		return true;
+	}
+	return false;
+}
+
 void UKismetTextLibrary::IsPolyglotDataValid(const FPolyglotTextData& PolyglotData, bool& IsValid, FText& ErrorMessage)
 {
 	IsValid = PolyglotData.IsValid(&ErrorMessage);
