add UClusterUnionReplicatedProxyComponent::GetChildClusteredComponent

diff --git a/Engine/Source/Runtime/Engine/Classes/PhysicsEngine/ClusterUnionReplicatedProxyComponent.h b/Engine/Source/Runtime/Engine/Classes/PhysicsEngine/ClusterUnionReplicatedProxyComponent.h
--- a/Engine/Source/Runtime/Engine/Classes/PhysicsEngine/ClusterUnionReplicatedProxyComponent.h
+++ b/Engine/Source/Runtime/Engine/Classes/PhysicsEngine/ClusterUnionReplicatedProxyComponent.h
@@ -34,6 +34,9 @@ public:
 	UFUNCTION()
 	ENGINE_API void SetParentClusterUnion(UClusterUnionComponent* InComponent);
 
+	UFUNCTION()
+	ENGINE_API UPrimitiveComponent* GetChildClusteredComponent() const { return ChildClusteredComponent.Get(); }
+
 	UFUNCTION()
 	ENGINE_API void SetChildClusteredComponent(UPrimitiveComponent* InComponent);
 
