Use FIntermediateRayTracingGeometry instead of partially initialized FRayTracingGeometry to store intermediate data during mesh streaming.

diff --git a/Engine/Source/Runtime/Engine/Private/Streaming/StaticMeshUpdate.cpp b/Engine/Source/Runtime/Engine/Private/Streaming/StaticMeshUpdate.cpp
--- a/Engine/Source/Runtime/Engine/Private/Streaming/StaticMeshUpdate.cpp
+++ b/Engine/Source/Runtime/Engine/Private/Streaming/StaticMeshUpdate.cpp
@@ -140,6 +140,50 @@ void FStaticMeshStreamIn::FIntermediateBuffers::CheckIsNull() const
 		&& !WireframeIndexBuffer);
 }
 
+#if RHI_RAYTRACING
+
+void FStaticMeshStreamIn::FIntermediateRayTracingGeometry::CreateFromCPUData(FRHICommandList& RHICmdList, FRayTracingGeometryInitializer InInitializer, TResourceArray<uint8>& OfflineData)
+{
+	Initializer = MoveTemp(InInitializer);
+
+	if (OfflineData.Num())
+	{
+		check(Initializer.OfflineData == nullptr);
+		Initializer.OfflineData = &OfflineData;
+	}
+
+	static const auto CVarDebugForceRuntimeBLAS = IConsoleManager::Get().FindTConsoleVariableDataInt(TEXT("r.Raytracing.DebugForceRuntimeBLAS"));
+	const bool bDebugForceRuntimeBLAS = (!CVarDebugForceRuntimeBLAS) || (CVarDebugForceRuntimeBLAS->GetValueOnAnyThread() != 0);
+
+	if (bDebugForceRuntimeBLAS && Initializer.OfflineData != nullptr)
+	{
+		Initializer.OfflineData->Discard();
+		Initializer.OfflineData = nullptr;
+	}
+
+	RayTracingGeometryRHI = RHICmdList.CreateRayTracingGeometry(Initializer);
+	bRequiresBuild = Initializer.OfflineData == nullptr || RayTracingGeometryRHI->IsCompressed();
+}
+
+void FStaticMeshStreamIn::FIntermediateRayTracingGeometry::SafeRelease()
+{
+	Initializer = {};
+	RayTracingGeometryRHI.SafeRelease();
+}
+
+void FStaticMeshStreamIn::FIntermediateRayTracingGeometry::TransferRayTracingGeometry(FRayTracingGeometry& RayTracingGeometry, FRHIResourceUpdateBatcher& Batcher)
+{
+	check(RayTracingGeometryRHI != nullptr);
+
+	// Should also set initializer?? ie: RayTracingGeometry.SetInitializer(Initializer);
+	RayTracingGeometry.InitRHIForStreaming(RayTracingGeometryRHI, Batcher);
+	RayTracingGeometry.SetRequiresBuild(bRequiresBuild);
+
+	SafeRelease();
+}
+
+#endif
+
 FStaticMeshStreamIn::FStaticMeshStreamIn(const UStaticMesh* InMesh)
 	: FStaticMeshUpdate(InMesh)
 {}
@@ -191,12 +235,11 @@ void FStaticMeshStreamIn::CreateBuffers_Internal(const FContext& Context)
 					Context.LODResourcesView[LODIdx]->SetupRayTracingGeometryInitializer(Initializer, Context.Mesh->GetFName(), OwnerName);
 				}
 				Initializer.Type = ERayTracingGeometryInitializerType::StreamingSource;
-				IntermediateRayTracingGeometry[LODIdx].SetInitializer(Initializer);
-
 
 				FRHIAsyncCommandList AsyncCommandList;
 				FRHICommandList& RHICmdList = bRenderThread ? FRHICommandListImmediate::Get() : *AsyncCommandList;
-				IntermediateRayTracingGeometry[LODIdx].CreateRayTracingGeometryFromCPUData(RHICmdList, LODResource.RayTracingGeometry.RawData);
+
+				IntermediateRayTracingGeometry[LODIdx].CreateFromCPUData(RHICmdList, MoveTemp(Initializer), LODResource.RayTracingGeometry.RawData);
 			}
 #endif
 		}
@@ -245,16 +288,11 @@ void FStaticMeshStreamIn::DoFinishUpdate(const FContext& Context)
 				IntermediateBuffersArray[LODIdx].TransferBuffers(LODResource, Batcher);
 
 #if RHI_RAYTRACING
-				if (IsRayTracingEnabled() && Context.Mesh->bSupportRayTracing &&
-					LODResource.VertexBuffers.StaticMeshVertexBuffer.GetNumVertices() > 0)
-				{					
-					check(IntermediateRayTracingGeometry[LODIdx].RayTracingGeometryRHI != nullptr);
-					LODResource.RayTracingGeometry.InitRHIForStreaming(IntermediateRayTracingGeometry[LODIdx].RayTracingGeometryRHI, Batcher);
-
-					LODResource.RayTracingGeometry.SetRequiresBuild(IntermediateRayTracingGeometry[LODIdx].GetRequiresBuild());
-
-					IntermediateRayTracingGeometry[LODIdx].Initializer = {};
-					IntermediateRayTracingGeometry[LODIdx].RayTracingGeometryRHI.SafeRelease();				
+				if (IsRayTracingEnabled()
+					&& Context.Mesh->bSupportRayTracing
+					&& LODResource.VertexBuffers.StaticMeshVertexBuffer.GetNumVertices() > 0)
+				{
+					IntermediateRayTracingGeometry[LODIdx].TransferRayTracingGeometry(LODResource.RayTracingGeometry, Batcher);
 				}
 #endif
 			}
@@ -308,8 +346,7 @@ void FStaticMeshStreamIn::DoFinishUpdate(const FContext& Context)
 			IntermediateBuffersArray[LODIdx].SafeRelease();
 
 #if RHI_RAYTRACING
-			IntermediateRayTracingGeometry[LODIdx].Initializer = {};
-			IntermediateRayTracingGeometry[LODIdx].RayTracingGeometryRHI.SafeRelease();
+			IntermediateRayTracingGeometry[LODIdx].SafeRelease();
 #endif
 		}
 	}
diff --git a/Engine/Source/Runtime/Engine/Private/Streaming/StaticMeshUpdate.h b/Engine/Source/Runtime/Engine/Private/Streaming/StaticMeshUpdate.h
--- a/Engine/Source/Runtime/Engine/Private/Streaming/StaticMeshUpdate.h
+++ b/Engine/Source/Runtime/Engine/Private/Streaming/StaticMeshUpdate.h
@@ -97,6 +97,23 @@ protected:
 		void CheckIsNull() const;
 	};
 
+#if RHI_RAYTRACING
+	struct FIntermediateRayTracingGeometry
+	{
+	private:
+		FRayTracingGeometryInitializer Initializer;
+		FRayTracingGeometryRHIRef RayTracingGeometryRHI;
+		bool bRequiresBuild = false;
+
+	public:
+		void CreateFromCPUData(FRHICommandList& RHICmdList, FRayTracingGeometryInitializer Initializer, TResourceArray<uint8>& OfflineData);
+
+		void SafeRelease();
+
+		void TransferRayTracingGeometry(FRayTracingGeometry& RayTracingGeometry, FRHIResourceUpdateBatcher& Batcher);
+	};
+#endif
+
 	/** Create buffers with new LOD data on render or pooled thread */
 	void CreateBuffers_RenderThread(const FContext& Context);
 	void CreateBuffers_Async(const FContext& Context);
@@ -114,7 +131,7 @@ protected:
 	FIntermediateBuffers IntermediateBuffersArray[MAX_MESH_LOD_COUNT];
 	
 #if RHI_RAYTRACING
-	FRayTracingGeometry IntermediateRayTracingGeometry[MAX_MESH_LOD_COUNT];
+	FIntermediateRayTracingGeometry IntermediateRayTracingGeometry[MAX_MESH_LOD_COUNT];
 #endif
 
 private:
diff --git a/Engine/Source/Runtime/RenderCore/Private/RayTracingGeometry.cpp b/Engine/Source/Runtime/RenderCore/Private/RayTracingGeometry.cpp
--- a/Engine/Source/Runtime/RenderCore/Private/RayTracingGeometry.cpp
+++ b/Engine/Source/Runtime/RenderCore/Private/RayTracingGeometry.cpp
@@ -8,10 +8,9 @@
 #include "RHIResourceUpdates.h"
 #include "RHITextureReference.h" // IWYU pragma: keep
 
-int32 GVarDebugForceRuntimeBLAS = 0;
-FAutoConsoleVariableRef CVarDebugForceRuntimeBLAS(
+static TAutoConsoleVariable<int32> CVarDebugForceRuntimeBLAS(
 	TEXT("r.Raytracing.DebugForceRuntimeBLAS"),
-	GVarDebugForceRuntimeBLAS,
+	0,
 	TEXT("Force building BLAS at runtime."),
 	ECVF_ReadOnly);
 
@@ -47,7 +46,7 @@ void FRayTracingGeometry::ReleaseRHIForStreaming(FRHIResourceUpdateBatcher& Batc
 	}
 }
 
-void FRayTracingGeometry::CreateRayTracingGeometryFromCPUData(FRHICommandList& RHICmdList, TResourceArray<uint8>& OfflineData)
+void FRayTracingGeometry::CreateRayTracingGeometryFromCPUData(TResourceArray<uint8>& OfflineData)
 {
 	check(OfflineData.Num() == 0 || Initializer.OfflineData == nullptr);
 	if (OfflineData.Num())
@@ -55,21 +54,17 @@ void FRayTracingGeometry::CreateRayTracingGeometryFromCPUData(FRHICommandList& R
 		Initializer.OfflineData = &OfflineData;
 	}
 
-	if (GVarDebugForceRuntimeBLAS && Initializer.OfflineData != nullptr)
+	if (CVarDebugForceRuntimeBLAS.GetValueOnAnyThread() && Initializer.OfflineData != nullptr)
 	{
 		Initializer.OfflineData->Discard();
 		Initializer.OfflineData = nullptr;
 	}
 	
+	FRHICommandList& RHICmdList = FRHICommandListImmediate::Get();
 	RayTracingGeometryRHI = RHICmdList.CreateRayTracingGeometry(Initializer);
 	SetRequiresBuild(Initializer.OfflineData == nullptr || RayTracingGeometryRHI->IsCompressed());
 }
 
-void FRayTracingGeometry::CreateRayTracingGeometryFromCPUData(TResourceArray<uint8>& OfflineData)
-{
-	CreateRayTracingGeometryFromCPUData(FRHICommandListImmediate::Get(), OfflineData);
-}
-
 void FRayTracingGeometry::RequestBuildIfNeeded(FRHICommandList& RHICmdList, ERTAccelerationStructureBuildPriority InBuildPriority)
 {
 	RayTracingGeometryRHI->SetInitializer(Initializer);
@@ -137,13 +132,13 @@ void FRayTracingGeometry::CreateRayTracingGeometry(FRHICommandList& RHICmdList,
 	// Release previous RHI object if any
 	ReleaseRHI();
 
-	check(RawData.Num() == 0 || Initializer.OfflineData == nullptr);
 	if (RawData.Num())
 	{
+		check(Initializer.OfflineData == nullptr);
 		Initializer.OfflineData = &RawData;
 	}
 
-	if (GVarDebugForceRuntimeBLAS && Initializer.OfflineData != nullptr)
+	if (CVarDebugForceRuntimeBLAS.GetValueOnAnyThread() && Initializer.OfflineData != nullptr)
 	{
 		Initializer.OfflineData->Discard();
 		Initializer.OfflineData = nullptr;
diff --git a/Engine/Source/Runtime/RenderCore/Public/RayTracingGeometry.h b/Engine/Source/Runtime/RenderCore/Public/RayTracingGeometry.h
--- a/Engine/Source/Runtime/RenderCore/Public/RayTracingGeometry.h
+++ b/Engine/Source/Runtime/RenderCore/Public/RayTracingGeometry.h
@@ -65,9 +65,9 @@ public:
 	/** LOD of the mesh associated with this ray tracing geometry object (-1 if unknown) */
 	int8 LODIndex = -1;
 
-	void SetInitializer(const FRayTracingGeometryInitializer& InInitializer)
+	void SetInitializer(FRayTracingGeometryInitializer InInitializer)
 	{
-		Initializer = InInitializer;
+		Initializer = MoveTemp(InInitializer);
 	}
 
 	RENDERCORE_API bool IsValid() const;
@@ -102,9 +102,7 @@ public:
 	RENDERCORE_API void InitRHIForStreaming(FRHIRayTracingGeometry* IntermediateGeometry, FRHIResourceUpdateBatcher& Batcher);
 	RENDERCORE_API void ReleaseRHIForStreaming(FRHIResourceUpdateBatcher& Batcher);
 
-	RENDERCORE_API void CreateRayTracingGeometryFromCPUData(FRHICommandList& RHICmdList, TResourceArray<uint8>& OfflineData);
-
-	UE_DEPRECATED(5.4, "CreateRayTracingGeometryFromCPUData now requires a command list.")
+	UE_DEPRECATED(5.4, "Use FStaticMeshStreamIn::FIntermediateRayTracingGeometry instead.")
 	RENDERCORE_API void CreateRayTracingGeometryFromCPUData(TResourceArray<uint8>& OfflineData);
 
 	RENDERCORE_API void RequestBuildIfNeeded(FRHICommandList& RHICmdList, ERTAccelerationStructureBuildPriority InBuildPriority);
