- Add option to deny individual Cascade systems for running

diff --git a/Engine/Source/Runtime/Engine/Private/Particles/ParticleComponents.cpp b/Engine/Source/Runtime/Engine/Private/Particles/ParticleComponents.cpp
--- a/Engine/Source/Runtime/Engine/Private/Particles/ParticleComponents.cpp
+++ b/Engine/Source/Runtime/Engine/Private/Particles/ParticleComponents.cpp
@@ -251,6 +251,40 @@ ENGINE_API bool GIsAllowingParticles = true;
 /** Whether to calculate LOD on the GameThread in-game. */
 bool GbEnableGameThreadLODCalculation = true;
 
+namespace CascadeLocal
+{
+	bool			bUseTemplateDenyList = false;
+	TSet<FName>		TemplateDenyList;
+	FString			TemplateDenyListString;
+
+	static void UpdateTemplateDenyList(IConsoleVariable*)
+	{
+		TArray<FString> Names;
+		TemplateDenyListString.ParseIntoArray(Names, TEXT(","));
+
+		TemplateDenyList.Empty();
+		for (const FString& Name : Names)
+		{
+			TemplateDenyList.Emplace(Name);
+		}
+
+		bUseTemplateDenyList = TemplateDenyList.Num() > 0;
+	}
+
+	bool AllowTemplate(UParticleSystem* Template)
+	{
+		return !bUseTemplateDenyList || (Template && !TemplateDenyList.Contains(Template->GetFName()));
+	}
+
+	static FAutoConsoleVariableRef CVarCascadeSetTemplateDenyList(
+		TEXT("fx.Cascade.SetTemplateDenyList"),
+		TemplateDenyListString,
+		TEXT("Set the template deny List to use. (i.e. P_SystemA,P_SystemB)"),
+		FConsoleVariableDelegate::CreateStatic(UpdateTemplateDenyList),
+		ECVF_Scalability | ECVF_Default
+	);
+}
+
 // Comment this in to debug empty emitter instance templates...
 //#define _PSYSCOMP_DEBUG_INVALID_EMITTER_INSTANCE_TEMPLATES_
 
@@ -6191,6 +6225,11 @@ void UParticleSystemComponent::ActivateSystem(bool bFlagAsJustAttached)
 		return;
 	}
 
+	if (!CascadeLocal::AllowTemplate(Template))
+	{
+		Template = nullptr;
+	}
+
 	bOldPositionValid = false;
 	OldPosition = FVector::ZeroVector;
 	PartSysVelocity = FVector::ZeroVector;
