Iris - Mark NotReplicated properties as such in HitResult and update HitResultNetSerializer to reflect this change.

#rnx

diff --git a/Engine/Source/Runtime/Engine/Classes/Engine/HitResult.h b/Engine/Source/Runtime/Engine/Classes/Engine/HitResult.h
--- a/Engine/Source/Runtime/Engine/Classes/Engine/HitResult.h
+++ b/Engine/Source/Runtime/Engine/Classes/Engine/HitResult.h
@@ -90,7 +90,7 @@ struct FHitResult
 	float PenetrationDepth;
 
 	/** If the hit result is from a collision this will have extra info about the item that hit the second item. */
-	UPROPERTY()
+	UPROPERTY(NotReplicated)
 	int32 MyItem;
 
 	/** Extra data about item that was hit (hit primitive specific). */
@@ -134,7 +134,7 @@ struct FHitResult
 	FName BoneName;
 
 	/** Name of the _my_ bone which took part in hit event (in case of two skeletal meshes colliding). */
-	UPROPERTY()
+	UPROPERTY(NotReplicated)
 	FName MyBoneName;
 
 
diff --git a/Engine/Source/Runtime/Engine/Private/Engine/HitResultNetSerializer.cpp b/Engine/Source/Runtime/Engine/Private/Engine/HitResultNetSerializer.cpp
--- a/Engine/Source/Runtime/Engine/Private/Engine/HitResultNetSerializer.cpp
+++ b/Engine/Source/Runtime/Engine/Private/Engine/HitResultNetSerializer.cpp
@@ -38,10 +38,9 @@ struct FHitResultNetSerializer
 		InvalidFaceIndex = InvalidItem << 1U,
 		NoPenetrationDepth = InvalidFaceIndex << 1U,
 		InvalidElementIndex = NoPenetrationDepth << 1U,
-		InvalidMyItem = InvalidElementIndex << 1U,
 	};
 
-	static constexpr uint32 ReplicatedFlagCount = 9U;
+	static constexpr uint32 ReplicatedFlagCount = 8U;
 
 	struct FQuantizedType
 	{
@@ -91,7 +90,6 @@ private:
 		PropertyName_ImpactPoint,
 		PropertyName_ImpactNormal,
 		PropertyName_PenetrationDepth,
-		PropertyName_MyItem,
 		PropertyName_Item,
 		PropertyName_ElementIndex,
 		PropertyName_bBlockingHit,
@@ -128,7 +126,6 @@ FNetBitArrayView FHitResultNetSerializer::GetMemberChangeMask(uint32* MemberMask
 	MemberMask.SetBitValue(PropertyToMemberIndex[EPropertyName::PropertyName_ImpactPoint], (ReplicationFlags & EReplicationFlags::ImpactPointEqualsLocation) == 0);
 	MemberMask.SetBitValue(PropertyToMemberIndex[EPropertyName::PropertyName_ImpactNormal], (ReplicationFlags & EReplicationFlags::ImpactNormalEqualsNormal) == 0);
 	MemberMask.SetBitValue(PropertyToMemberIndex[EPropertyName::PropertyName_PenetrationDepth], (ReplicationFlags & EReplicationFlags::NoPenetrationDepth) == 0);
-	MemberMask.SetBitValue(PropertyToMemberIndex[EPropertyName::PropertyName_MyItem], (ReplicationFlags & EReplicationFlags::InvalidMyItem) == 0);
 	MemberMask.SetBitValue(PropertyToMemberIndex[EPropertyName::PropertyName_Item], (ReplicationFlags & EReplicationFlags::InvalidItem) == 0);
 	MemberMask.SetBitValue(PropertyToMemberIndex[EPropertyName::PropertyName_FaceIndex], (ReplicationFlags & EReplicationFlags::InvalidFaceIndex) == 0);
 	MemberMask.SetBitValue(PropertyToMemberIndex[EPropertyName::PropertyName_ElementIndex], (ReplicationFlags & EReplicationFlags::InvalidElementIndex) == 0);
@@ -228,15 +225,14 @@ void FHitResultNetSerializer::Deserialize(FNetSerializationContext& Context, con
 	}
 }
 
-// For now just use normal serialization for delta as this struct typically is a one off.
 void FHitResultNetSerializer::SerializeDelta(FNetSerializationContext& Context, const FNetSerializeDeltaArgs& Args)
 {
-	NetSerializeDeltaDefault<Serialize>(Context, Args);
+	Serialize(Context, Args);
 }
 
 void FHitResultNetSerializer::DeserializeDelta(FNetSerializationContext& Context, const FNetDeserializeDeltaArgs& Args)
 {
-	NetDeserializeDeltaDefault<Deserialize>(Context, Args);
+	Deserialize(Context, Args);
 }
 
 void FHitResultNetSerializer::Quantize(FNetSerializationContext& Context, const FNetQuantizeArgs& Args)
@@ -249,7 +245,6 @@ void FHitResultNetSerializer::Quantize(FNetSerializationContext& Context, const
 	// Update flags based on SourceValue
 	ReplicationFlags |= SourceValue.ImpactPoint == SourceValue.Location ? EReplicationFlags::ImpactPointEqualsLocation : 0U;
 	ReplicationFlags |= SourceValue.ImpactNormal == SourceValue.Normal ? EReplicationFlags::ImpactNormalEqualsNormal : 0U;
-	ReplicationFlags |= SourceValue.MyItem == INDEX_NONE ? EReplicationFlags::InvalidMyItem : 0U;
 	ReplicationFlags |= SourceValue.Item == INDEX_NONE ? EReplicationFlags::InvalidItem : 0U;
 	ReplicationFlags |= SourceValue.FaceIndex == INDEX_NONE ? EReplicationFlags::InvalidFaceIndex : 0U;
 	ReplicationFlags |= (SourceValue.PenetrationDepth == 0.0f) ? EReplicationFlags::NoPenetrationDepth : 0U;
@@ -291,11 +286,6 @@ void FHitResultNetSerializer::Dequantize(FNetSerializationContext& Context, cons
 		TargetValue.ImpactNormal = TargetValue.Normal;
 	}
 
-	if (ReplicatonFlags & EReplicationFlags::InvalidMyItem)
-	{
-		TargetValue.MyItem = INDEX_NONE;
-	}
-
 	if (ReplicatonFlags & EReplicationFlags::InvalidItem)
 	{
 		TargetValue.Item = INDEX_NONE;
@@ -490,7 +480,6 @@ void FHitResultNetSerializer::FNetSerializerRegistryDelegates::OnPostFreezeNetSe
 		PropertyNames[PropertyName_ImpactPoint] = FName("ImpactPoint");
 		PropertyNames[PropertyName_ImpactNormal] = FName("ImpactNormal");
 		PropertyNames[PropertyName_PenetrationDepth] = FName("PenetrationDepth");
-		PropertyNames[PropertyName_MyItem] = FName("MyItem");
 		PropertyNames[PropertyName_Item] = FName("Item");
 		PropertyNames[PropertyName_ElementIndex] = FName("ElementIndex");
 		PropertyNames[PropertyName_bBlockingHit] = FName("bBlockingHit");
