Implemented dormancy for AWorldDataLayers which reduces actor replication to the minimum (i.e. replicates only when necessary).

diff --git a/Engine/Source/Runtime/Engine/Private/WorldPartition/DataLayer/WorldDataLayers.cpp b/Engine/Source/Runtime/Engine/Private/WorldPartition/DataLayer/WorldDataLayers.cpp
--- a/Engine/Source/Runtime/Engine/Private/WorldPartition/DataLayer/WorldDataLayers.cpp
+++ b/Engine/Source/Runtime/Engine/Private/WorldPartition/DataLayer/WorldDataLayers.cpp
@@ -16,6 +16,7 @@
 #include "Engine/Level.h"
 #include "ProfilingDebugging/CsvProfiler.h"
 #include "Net/UnrealNetwork.h"
+#include "Net/Core/PushModel/PushModel.h"
 
 #include UE_INLINE_GENERATED_CPP_BY_NAME(WorldDataLayers)
 
@@ -30,7 +31,7 @@
 
 int32 AWorldDataLayers::DataLayersStateEpoch = 0;
 
-FString JoinDataLayerShortNamesFromInstanceNames(AWorldDataLayers* InWorldDataLayers, const TArray<FName>& InDataLayerInstanceNames)
+static FString JoinDataLayerShortNamesFromInstanceNames(AWorldDataLayers* InWorldDataLayers, const TArray<FName>& InDataLayerInstanceNames)
 {
 	check(InWorldDataLayers);
 	TArray<FString> DataLayerShortNames;
@@ -53,6 +54,7 @@ AWorldDataLayers::AWorldDataLayers(const FObjectInitializer& ObjectInitializer)
 {
 	bAlwaysRelevant = true;
 	bReplicates = true;
+	SetNetDormancy(DORM_Initial);
 
 	// Avoid actor from being Destroyed/Recreated when scrubbing a replay
 	// instead AWorldDataLayers::RewindForReplay() gets called to reset this actors state
@@ -63,9 +65,8 @@ void AWorldDataLayers::GetLifetimeReplicatedProps(TArray< FLifetimeProperty >& O
 {
 	Super::GetLifetimeReplicatedProps(OutLifetimeProps);
 
-	// Use RepNotifyCondition = REPNOTIFY_Always because of current issue with dynamic arrays: https://jira.it.epicgames.com/browse/UE-155774
 	FDoRepLifetimeParams Params;
-	Params.RepNotifyCondition = REPNOTIFY_Always;
+	Params.bIsPushBased = true;
 	DOREPLIFETIME_WITH_PARAMS_FAST(AWorldDataLayers, RepLoadedDataLayerNames, Params);
 	DOREPLIFETIME_WITH_PARAMS_FAST(AWorldDataLayers, RepActiveDataLayerNames, Params);
 	DOREPLIFETIME_WITH_PARAMS_FAST(AWorldDataLayers, RepEffectiveLoadedDataLayerNames, Params);
@@ -94,6 +95,10 @@ void AWorldDataLayers::RewindForReplay()
 	InitializeDataLayerRuntimeStates();
 }
 
+#define AWORLDDATALAYERS_UPDATE_REPLICATED_DATALAYERS(ReplicatedArray, SourceArray) \
+	MARK_PROPERTY_DIRTY_FROM_NAME(AWorldDataLayers, ReplicatedArray, this); \
+	ReplicatedArray = SourceArray;
+
 void AWorldDataLayers::InitializeDataLayerRuntimeStates()
 {
 	check(ActiveDataLayerNames.IsEmpty() && LoadedDataLayerNames.IsEmpty());
@@ -130,8 +135,9 @@ void AWorldDataLayers::InitializeDataLayerRuntimeStates()
 			return true;
 		});
 
-		RepActiveDataLayerNames = ActiveDataLayerNames.Array();
-		RepLoadedDataLayerNames = LoadedDataLayerNames.Array();
+		FlushNetDormancy();
+		AWORLDDATALAYERS_UPDATE_REPLICATED_DATALAYERS(RepActiveDataLayerNames, ActiveDataLayerNames.Array());
+		AWORLDDATALAYERS_UPDATE_REPLICATED_DATALAYERS(RepLoadedDataLayerNames, LoadedDataLayerNames.Array());
 
 		const bool bNotifyChange = false;
 		for (UDataLayerInstance* DataLayerInstance : RuntimeDataLayerInstances)
@@ -139,8 +145,8 @@ void AWorldDataLayers::InitializeDataLayerRuntimeStates()
 			ResolveEffectiveRuntimeState(DataLayerInstance, bNotifyChange);
 		}
 
-		RepEffectiveActiveDataLayerNames = EffectiveActiveDataLayerNames.Array();
-		RepEffectiveLoadedDataLayerNames = EffectiveLoadedDataLayerNames.Array();
+		AWORLDDATALAYERS_UPDATE_REPLICATED_DATALAYERS(RepEffectiveActiveDataLayerNames, EffectiveActiveDataLayerNames.Array());
+		AWORLDDATALAYERS_UPDATE_REPLICATED_DATALAYERS(RepEffectiveLoadedDataLayerNames, EffectiveLoadedDataLayerNames.Array());
 
 		UE_CLOG(RepEffectiveActiveDataLayerNames.Num() || RepEffectiveLoadedDataLayerNames.Num(), LogWorldPartition, Log, TEXT("Initial Data Layer Effective States Activated(%s) Loaded(%s)"), *JoinDataLayerShortNamesFromInstanceNames(this, RepEffectiveActiveDataLayerNames), *JoinDataLayerShortNamesFromInstanceNames(this, RepEffectiveLoadedDataLayerNames));
 	}
@@ -150,17 +156,20 @@ void AWorldDataLayers::ResetDataLayerRuntimeStates()
 {
 	ActiveDataLayerNames.Reset();
 	LoadedDataLayerNames.Reset();
-	RepActiveDataLayerNames.Reset();
-	RepLoadedDataLayerNames.Reset();
 	LocalActiveDataLayerNames.Reset();
 	LocalLoadedDataLayerNames.Reset();
 
 	EffectiveActiveDataLayerNames.Reset();
 	EffectiveLoadedDataLayerNames.Reset();
-	RepEffectiveActiveDataLayerNames.Reset();
-	RepEffectiveLoadedDataLayerNames.Reset();
 	LocalEffectiveActiveDataLayerNames.Reset();
 	LocalEffectiveLoadedDataLayerNames.Reset();
+
+	static const TArray<FName> Empty;
+	FlushNetDormancy();
+	AWORLDDATALAYERS_UPDATE_REPLICATED_DATALAYERS(RepActiveDataLayerNames, Empty);
+	AWORLDDATALAYERS_UPDATE_REPLICATED_DATALAYERS(RepLoadedDataLayerNames, Empty);
+	AWORLDDATALAYERS_UPDATE_REPLICATED_DATALAYERS(RepEffectiveActiveDataLayerNames, Empty);
+	AWORLDDATALAYERS_UPDATE_REPLICATED_DATALAYERS(RepEffectiveLoadedDataLayerNames, Empty);
 }
 
 void AWorldDataLayers::SetDataLayerRuntimeState(const UDataLayerInstance* InDataLayerInstance, EDataLayerRuntimeState InState, bool bInIsRecursive)
@@ -256,8 +265,9 @@ void AWorldDataLayers::SetDataLayerRuntimeState(const UDataLayerInstance* InData
 		// Update replicated properties
 		if (NetMode == NM_DedicatedServer || NetMode == NM_ListenServer)
 		{
-			RepActiveDataLayerNames = ActiveDataLayerNames.Array();
-			RepLoadedDataLayerNames = LoadedDataLayerNames.Array();
+			FlushNetDormancy();
+			AWORLDDATALAYERS_UPDATE_REPLICATED_DATALAYERS(RepActiveDataLayerNames, ActiveDataLayerNames.Array());
+			AWORLDDATALAYERS_UPDATE_REPLICATED_DATALAYERS(RepLoadedDataLayerNames, LoadedDataLayerNames.Array());
 		}
 
 		++DataLayersStateEpoch;
@@ -438,8 +448,9 @@ void AWorldDataLayers::ResolveEffectiveRuntimeState(const UDataLayerInstance* In
 		// Update Replicated Properties
 		if (NetMode == NM_DedicatedServer || NetMode == NM_ListenServer)
 		{
-			RepEffectiveActiveDataLayerNames = EffectiveActiveDataLayerNames.Array();
-			RepEffectiveLoadedDataLayerNames = EffectiveLoadedDataLayerNames.Array();
+			FlushNetDormancy();
+			AWORLDDATALAYERS_UPDATE_REPLICATED_DATALAYERS(RepEffectiveActiveDataLayerNames, EffectiveActiveDataLayerNames.Array());
+			AWORLDDATALAYERS_UPDATE_REPLICATED_DATALAYERS(RepEffectiveLoadedDataLayerNames, EffectiveLoadedDataLayerNames.Array());
 		}
 
 		++DataLayersStateEpoch;
diff --git a/Engine/Source/Runtime/Engine/Public/WorldPartition/DataLayer/WorldDataLayers.h b/Engine/Source/Runtime/Engine/Public/WorldPartition/DataLayer/WorldDataLayers.h
--- a/Engine/Source/Runtime/Engine/Public/WorldPartition/DataLayer/WorldDataLayers.h
+++ b/Engine/Source/Runtime/Engine/Public/WorldPartition/DataLayer/WorldDataLayers.h
@@ -8,6 +8,7 @@
 #include "WorldPartition/DataLayer/ActorDataLayer.h"
 #include "WorldPartition/DataLayer/DataLayerInstance.h"
 #include "WorldPartition/DataLayer/DataLayer.h"
+#include "Net/Core/PushModel/PushModel.h"
 #include "Engine/World.h"
 #include "WorldDataLayers.generated.h"
 
@@ -296,11 +297,14 @@ DataLayerInstanceType* AWorldDataLayers::CreateDataLayer(CreationsArgs... InCrea
 	return NewDataLayer;
 }
 
+//todo_ow: Remove this
 template<class IdentifierType>
 void AWorldDataLayers::OverwriteDataLayerRuntimeStates(const TArray<IdentifierType>* InActiveDataLayers, const TArray<IdentifierType>* InLoadedDataLayers)
 {
 	if (GetLocalRole() == ROLE_Authority)
 	{
+		FlushNetDormancy();
+
 		// This should get called before game starts. It doesn't send out events
 		check(!GetWorld()->bMatchStarted);
 
@@ -317,6 +321,7 @@ void AWorldDataLayers::OverwriteDataLayerRuntimeStates(const TArray<IdentifierTy
 					ActiveDataLayerNames.Add(DataLayerInstance->GetDataLayerFName());
 				}
 			}
+			MARK_PROPERTY_DIRTY_FROM_NAME(AWorldDataLayers, RepActiveDataLayerNames, this);
 			RepActiveDataLayerNames = ActiveDataLayerNames.Array();
 		}
 		
@@ -334,6 +339,7 @@ void AWorldDataLayers::OverwriteDataLayerRuntimeStates(const TArray<IdentifierTy
 					LoadedDataLayerNames.Add(DataLayerInstance->GetDataLayerFName());
 				}
 			}
+			MARK_PROPERTY_DIRTY_FROM_NAME(AWorldDataLayers, RepLoadedDataLayerNames, this);
 			RepLoadedDataLayerNames = LoadedDataLayerNames.Array();
 		}
 
@@ -346,7 +352,9 @@ void AWorldDataLayers::OverwriteDataLayerRuntimeStates(const TArray<IdentifierTy
 			return true;
 		});
 
+		MARK_PROPERTY_DIRTY_FROM_NAME(AWorldDataLayers, RepEffectiveActiveDataLayerNames, this);
 		RepEffectiveActiveDataLayerNames = EffectiveActiveDataLayerNames.Array();
+		MARK_PROPERTY_DIRTY_FROM_NAME(AWorldDataLayers, RepEffectiveLoadedDataLayerNames, this);
 		RepEffectiveLoadedDataLayerNames = EffectiveLoadedDataLayerNames.Array();
 	}
 }
