Packed FStreamingRenderAsset fileds to reduce sizeof from 216 to 168 and save ~1MiB on default size of the streaming render asset buffer

diff --git a/Engine/Source/Runtime/Engine/Private/Streaming/StreamingManagerTexture.cpp b/Engine/Source/Runtime/Engine/Private/Streaming/StreamingManagerTexture.cpp
--- a/Engine/Source/Runtime/Engine/Private/Streaming/StreamingManagerTexture.cpp
+++ b/Engine/Source/Runtime/Engine/Private/Streaming/StreamingManagerTexture.cpp
@@ -34,7 +34,7 @@ CSV_DECLARE_CATEGORY_MODULE_EXTERN(CORE_API, Basic);
 CSV_DEFINE_CATEGORY(TextureStreaming, true);
 
 #ifndef UE_STREAMINGRENDERASSETS_ARRAY_DEFAULT_RESERVED_SIZE
-// The default size will reserve ~4MB, the element size is ~208 bytes.
+// The default size will reserve ~3MB, the element size is 168 bytes.
 #define UE_STREAMINGRENDERASSETS_ARRAY_DEFAULT_RESERVED_SIZE 20000
 #endif
 
diff --git a/Engine/Source/Runtime/Engine/Private/Streaming/StreamingTexture.h b/Engine/Source/Runtime/Engine/Private/Streaming/StreamingTexture.h
--- a/Engine/Source/Runtime/Engine/Private/Streaming/StreamingTexture.h
+++ b/Engine/Source/Runtime/Engine/Private/Streaming/StreamingTexture.h
@@ -209,49 +209,48 @@ struct FStreamingRenderAsset
 		};
 	};
 
-	/** (2) Cached number of mip-maps in memory (including the base mip) */
-	int32			ResidentMips;
-	/** (2) Min number of mip-maps requested by the streaming system. */
-	int32			RequestedMips;
-	/** (2) Min mip to be requested by the streaming  */
-	int32			MinAllowedMips;
-	/** (2) Max mip to be requested by the streaming  */
-	int32			MaxAllowedMips;
-	/** (2) LOD bias applied to no-ref meshes */
-	int32			NoRefLODBias;
 	/** (2) How much game time has elapsed since the texture was bound for rendering. Based on FApp::GetCurrentTime(). */
 	float			LastRenderTime;
-
 	/** (3) If non-zero, the most recent time an instance location was removed for this texture. */
 	double			InstanceRemovedTimestamp;
 	/** (3) Extra gameplay boost factor. Reset after every update. */
 	float			DynamicBoostFactor;
 
-	/** (4) How many mips are missing to satisfy ideal quality because of max size limitation. Used to prevent sacrificing mips that are already visually sacrificed. */
-	int32			NumMissingMips;
-	/** (4) Max wanted mips for visible instances.*/
-	int32			VisibleWantedMips;
-	/** (4) Wanted mips for non visible instances.*/
-	int32			HiddenWantedMips;
 	/** (4) Retention priority used to sacrifice mips when out of budget. */
 	int32			RetentionPriority;
 	/** (4) Normalized screen size. Only used by meshes. */
 	float			NormalizedScreenSize;
-	/** (4) The max allowed mips (based on Visible and Hidden wanted mips) in order to fit in budget. */
-	int32			BudgetedMips;
 	/** (4) The load request priority. */
 	int32			LoadOrderPriority;	
-	/** (4) The mip that will be requested. Note, that some texture are loaded using split requests, so that the first request can be smaller than budgeted mip. */
-	int32			WantedMips;
 	/** (4) A persistent bias applied to this texture. Increase whenever the streamer needs to make sacrifices to fit in budget */
 	int32			BudgetMipBias;
+	/** (2) LOD bias applied to no-ref meshes */
+	int32			NoRefLODBias;
+	/** (4) The mip that will be requested. Note, that some texture are loaded using split requests, so that the first request can be smaller than budgeted mip. */
+	int8			WantedMips;
 	/** (4) Number of LODs forced resident. Only used by meshes currently */
-	int32			NumForcedMips;
+	int8			NumForcedMips;
+	/** (2) Cached number of mip-maps in memory (including the base mip) */
+	int8			ResidentMips;
+	/** (2) Min number of mip-maps requested by the streaming system. */
+	int8			RequestedMips;
+	/** (2) Min mip to be requested by the streaming  */
+	int8			MinAllowedMips;
+	/** (2) Max mip to be requested by the streaming  */
+	int8			MaxAllowedMips;
+	/** (4) How many mips are missing to satisfy ideal quality because of max size limitation. Used to prevent sacrificing mips that are already visually sacrificed. */
+	int8			NumMissingMips;
+	/** (4) Max wanted mips for visible instances.*/
+	int8			VisibleWantedMips;
+	/** (4) Wanted mips for non visible instances.*/
+	int8			HiddenWantedMips;
+	/** (4) The max allowed mips (based on Visible and Hidden wanted mips) in order to fit in budget. */
+	int8			BudgetedMips;
 
 	/** (5) */
-	int32 CachedWantedMips;
+	int8 CachedWantedMips;
 	/** (5) */
-	int32 CachedVisibleWantedMips;
+	int8 CachedVisibleWantedMips;
 
 	/** (1) */
 	EStreamableRenderAssetType RenderAssetType;
@@ -259,29 +258,29 @@ struct FStreamingRenderAsset
 	EOptionalMipsState	OptionalMipsState;
 
 	/** (1) Whether the texture is loaded with higher priority. */
-	uint32			bLoadWithHigherPriority : 1;
+	uint8			bLoadWithHigherPriority : 1;
 	/** (1) Whether the texture should be forcibly fully loaded. */
-	uint32			bIsTerrainTexture : 1;
+	uint8			bIsTerrainTexture : 1;
 
 	/** (2) Whether the texture should be forcibly fully loaded. */
-	uint32			bForceFullyLoad : 1;
+	uint8			bForceFullyLoad : 1;
 	/** (2) Whether the texture resolution should be affected by the memory budget. Only effective when forced fully resident. */
-	uint32			bIgnoreStreamingMipBias : 1;
+	uint8			bIgnoreStreamingMipBias : 1;
 
 	/** (3) Wheter the streamer has streaming plans for this texture. */
-	uint32			bHasUpdatePending : 1;
+	uint8			bHasUpdatePending : 1;
 
 	/** (4) Same as force fully load, but takes into account component settings. */
-	uint32			bForceFullyLoadHeuristic : 1;
+	uint8			bForceFullyLoadHeuristic : 1;
 	/** (4) Whether this has not component referencing it. */
-	uint32			bUseUnkownRefHeuristic : 1;
+	uint8			bUseUnkownRefHeuristic : 1;
 	/** (4) Same as force fully load, but takes into account component settings. */
-	uint32			bLooksLowRes : 1;
+	uint8			bLooksLowRes : 1;
 	/** (4) Whether too many desired MIPs or LODs are missing. */
-	uint32			bMissingTooManyMips : 1;
+	uint8			bMissingTooManyMips : 1;
 
 	/** (5) */
-	uint32			bCachedForceFullyLoadHeuristic : 1;
+	uint8			bCachedForceFullyLoadHeuristic : 1;
 
 	/** (5) */
 	TBitArray<>		LevelIndexUsage;
