[shaders] fixing job cache validation & making it slightly more useful (will output job debug info paths when debug info is enabled, can be used to diff debug artifacts). if debug info is disabled it will just print the debug group name and a blurb suggesting turning on debug info will give more useful details.

diff --git a/Engine/Source/Runtime/Engine/Private/ShaderCompiler/ShaderCompiler.cpp b/Engine/Source/Runtime/Engine/Private/ShaderCompiler/ShaderCompiler.cpp
--- a/Engine/Source/Runtime/Engine/Private/ShaderCompiler/ShaderCompiler.cpp
+++ b/Engine/Source/Runtime/Engine/Private/ShaderCompiler/ShaderCompiler.cpp
@@ -835,6 +835,8 @@ private:
 	/* A lot of outputs can be duplicated, so they are deduplicated before storing */
 	TMap<FJobOutputHash, FStoredOutput*> Outputs;
 
+	TMap<FJobOutputHash, FString> CachedJobNames;
+
 	/** Map of input hashes to job data (in flight jobs and output) */
 	FShaderJobDataMap InputHashToJobData;
 
@@ -984,6 +986,7 @@ void FShaderJobCache::CullOutputsToMemoryBudget(uint64 TargetBudgetBytes)
 					if (StoredOutput->Release() == 0)
 					{
 						Outputs.Remove(JobData.OutputHash);
+						CachedJobNames.Remove(JobData.OutputHash);
 						CurrentlyAllocatedMemory -= OutputSize;
 					}
 				}
@@ -1417,12 +1420,12 @@ void FShaderJobCache::SubmitJob(FShaderCommonCompileJobPtr Job)
 
 	const int32 PriorityIndex = (int32)Job->Priority;
 	bool bNewJob = true;
-	bool bCacheEnabled = false;
+	bool bJobCacheLocked = false;
 
 	// check caches unless we're running in validation mode (which runs _all_ jobs and compares hashes of outputs)
 	if (ShaderCompiler::IsJobCacheEnabled() && !ShaderCompiler::IsJobCacheDebugValidateEnabled())
 	{
-		bCacheEnabled = true;
+		bJobCacheLocked = true;
 
 		const FShaderCommonCompileJob::FInputHash& InputHash = Job->GetInputHash();
 
@@ -1485,6 +1488,15 @@ void FShaderJobCache::SubmitJob(FShaderCommonCompileJobPtr Job)
 			}
 		}
 	}
+	else if (ShaderCompiler::IsJobCacheDebugValidateEnabled())
+	{
+		FSharedBuffer* ExistingOutput;
+		const FShaderCommonCompileJob::FInputHash& InputHash = Job->GetInputHash();
+		const bool bCheckDDC = !(Job->bIsDefaultMaterial || Job->bIsGlobalShader);
+		JobLock.WriteLock();
+		Job->JobCacheRef = FindOrAdd(InputHash, bCheckDDC, ExistingOutput);
+		bJobCacheLocked = true;
+	}
 
 	// new job
 	if (bNewJob)
@@ -1493,13 +1505,17 @@ void FShaderJobCache::SubmitJob(FShaderCommonCompileJobPtr Job)
 		ensure(!ShaderCompiler::IsJobCacheEnabled() || Job->bInputHashSet);
 
 		// If cache is disabled, we skipped the code that grabs the write lock above, so we need to do it here, before modifying the pending queue
-		if (bCacheEnabled == false)
+		if (bJobCacheLocked == false)
 		{
+			bJobCacheLocked = true;
 			JobLock.WriteLock();
 		}
 
 		LinkJobWithPriority(*Job);
+	}
 
+	if (bJobCacheLocked)
+	{
 		JobLock.WriteUnlock();
 	}
 }
@@ -10121,31 +10137,31 @@ void FShaderJobCache::AddJobOutput(FShaderJobData& JobData, const FShaderCommonC
 	{
 		if (OutputHash != JobData.OutputHash)
 		{
-			FString OriginalFilename = 
-				GShaderCompilingManager->GetAbsoluteShaderDebugInfoDirectory() / 
-				TEXT("CacheMismatches") / 
-				FString::Printf(TEXT("shaderoutput-%s.orig"), *LexToString(Hash));
-			FString NewFilename = 
-				GShaderCompilingManager->GetAbsoluteShaderDebugInfoDirectory() / 
-				TEXT("CacheMismatches") / 
-				FString::Printf(TEXT("shaderoutput-%s.mismatch"), *LexToString(Hash));
+			TStringBuilder<1024> FinishedJobName;
+			FinishedJob->AppendDebugName(FinishedJobName);
 
+			const FString* CachedJobName = CachedJobNames.Find(JobData.OutputHash);
+			check(CachedJobName);
 			UE_LOG(
 				LogShaderCompilers,
 				Warning,
-				TEXT("Job cache validation found output mismatch. Mismatching outputs dumped to %s and %s"),
-				*OriginalFilename, *NewFilename);
-
-			FFileHelper::SaveArrayToFile(
-				TArrayView<const uint8>((const uint8*)Contents.GetData(), Contents.GetSize()), 
-				*NewFilename);
+				TEXT("Job cache validation found output mismatch!\n")
+				TEXT("Cached job: %s\n")
+				TEXT("Original job: %s\n"),
+				**CachedJobName, FinishedJobName.ToString());
 
-			FStoredOutput** StoredOutput = Outputs.Find(JobData.OutputHash);
-			check(StoredOutput);
-
-			FFileHelper::SaveArrayToFile(
-				TArrayView<const uint8>((const uint8*)(*StoredOutput)->JobOutput.GetData(), (*StoredOutput)->JobOutput.GetSize()),
-				*OriginalFilename);
+			if (GDumpShaderDebugInfo != FShaderCompilingManager::EDumpShaderDebugInfo::Always)
+			{
+				static bool bOnce = false;
+				if (!bOnce)
+				{
+					UE_LOG(
+						LogShaderCompilers,
+						Warning,
+						TEXT("Enable r.DumpShaderDebugInfo=1 to get debug info paths for the mismatching jobs instead of group names (to allow diffing debug artifacts)"));
+					bOnce = true;
+				}
+			}
 		}
 		return;
 	}
@@ -10215,6 +10231,14 @@ void FShaderJobCache::AddJobOutput(FShaderJobData& JobData, const FShaderCommonC
 			NewStoredOutput->AddRef();
 			Outputs.Add(OutputHash, NewStoredOutput);
 
+			if (ShaderCompiler::IsJobCacheDebugValidateEnabled())
+			{
+				TStringBuilder<1024> NameBuilder;
+				FinishedJob->AppendDebugName(NameBuilder);
+
+				CachedJobNames.Add(OutputHash, NameBuilder.ToString());
+			}
+
 			CurrentlyAllocatedMemory += NewStoredOutput->GetAllocatedSize() + Outputs.GetAllocatedSize() - OutputsOriginalSize;
 		}
 
diff --git a/Engine/Source/Runtime/RenderCore/Private/ShaderCore.cpp b/Engine/Source/Runtime/RenderCore/Private/ShaderCore.cpp
--- a/Engine/Source/Runtime/RenderCore/Private/ShaderCore.cpp
+++ b/Engine/Source/Runtime/RenderCore/Private/ShaderCore.cpp
@@ -3366,6 +3366,11 @@ void FShaderCompileJob::OnComplete()
 	}
 }
 
+void FShaderCompileJob::AppendDebugName(FStringBuilderBase& OutName) const
+{
+	OutName << (Input.DumpDebugInfoPath.IsEmpty() ? Input.DebugGroupName : Input.DumpDebugInfoPath);
+}
+
 void FShaderCompileJob::SerializeWorkerOutput(FArchive& Ar)
 {
 	// Only serialize the modified source/entry point if either:
@@ -3508,3 +3513,12 @@ void FShaderPipelineCompileJob::OnComplete()
 		StageJobs[Index]->OnComplete();
 	}
 }
+
+void FShaderPipelineCompileJob::AppendDebugName(FStringBuilderBase& OutName) const
+{
+	for (int32 Index = 0, Num = StageJobs.Num(); Index < Num; ++Index)
+	{
+		StageJobs[Index]->AppendDebugName(OutName);
+		OutName << TEXT("\n");
+	}
+}
diff --git a/Engine/Source/Runtime/RenderCore/Public/ShaderCompilerJobTypes.h b/Engine/Source/Runtime/RenderCore/Public/ShaderCompilerJobTypes.h
--- a/Engine/Source/Runtime/RenderCore/Public/ShaderCompilerJobTypes.h
+++ b/Engine/Source/Runtime/RenderCore/Public/ShaderCompilerJobTypes.h
@@ -156,6 +156,8 @@ public:
 	// Executed for all jobs (including those read from cache) on completion.
 	virtual void OnComplete() = 0;
 
+	virtual void AppendDebugName(FStringBuilderBase& OutName) const = 0;
+	
 	bool Equals(const FShaderCommonCompileJob& Rhs) const;
 
 	/** Calls the specified predicate for each single compile job, i.e. FShaderCompileJob and each stage of FShaderPipelineCompileJob. */
@@ -249,6 +251,8 @@ public:
 	virtual RENDERCORE_API void SerializeOutput(FArchive& Ar) override;
 
 	virtual RENDERCORE_API void OnComplete() override;
+	
+	virtual RENDERCORE_API void AppendDebugName(FStringBuilderBase& OutName) const override;
 
 	// Serializes only the subset of data written by SCW/read back from ShaderCompiler when using worker processes.
 	RENDERCORE_API void SerializeWorkerOutput(FArchive& Ar);
@@ -302,6 +306,7 @@ public:
 	virtual RENDERCORE_API FInputHash GetInputHash() override;
 	virtual RENDERCORE_API void SerializeOutput(FArchive& Ar) override;
 	virtual RENDERCORE_API void OnComplete() override;
+	virtual RENDERCORE_API void AppendDebugName(FStringBuilderBase& OutName) const override;
 
 	RENDERCORE_API FShaderPipelineCompileJob(int32 NumStages);
 	RENDERCORE_API FShaderPipelineCompileJob(uint32 InHash, uint32 InId, EShaderCompileJobPriority InPriroity, const FShaderPipelineCompileJobKey& InKey);
