WorldPartition: Made console commands "wp.runtime wp.Runtime.OverrideRuntimeSpatialHashLoadingRange", "wp.Runtime.ToggleDataLayerActivation" and "wp.Runtime.SetDataLayerRuntimeState" also execute on the server when ran from a client (non shipping builds only).

diff --git a/Engine/Source/Runtime/Engine/Private/WorldPartition/DataLayer/DataLayerManager.cpp b/Engine/Source/Runtime/Engine/Private/WorldPartition/DataLayer/DataLayerManager.cpp
--- a/Engine/Source/Runtime/Engine/Private/WorldPartition/DataLayer/DataLayerManager.cpp
+++ b/Engine/Source/Runtime/Engine/Private/WorldPartition/DataLayer/DataLayerManager.cpp
@@ -62,8 +62,9 @@ FAutoConsoleCommandWithOutputDevice UDataLayerManager::DumpDataLayersCommand(
 	})
 );
 
+static const FString GToggleDataLayerActivationCommandName(TEXT("wp.Runtime.ToggleDataLayerActivation"));
 FAutoConsoleCommand UDataLayerManager::ToggleDataLayerActivation(
-	TEXT("wp.Runtime.ToggleDataLayerActivation"),
+	*GToggleDataLayerActivationCommandName,
 	TEXT("Toggles DataLayers active state. Args [DataLayerNames]"),
 	FConsoleCommandWithArgsDelegate::CreateLambda([](const TArray<FString>& InArgs)
 	{
@@ -72,6 +73,7 @@ FAutoConsoleCommand UDataLayerManager::ToggleDataLayerActivation(
 			UWorld* World = Context.World();
 			if (World && World->IsGameWorld())
 			{
+				FWorldPartitionHelpers::ServerExecConsoleCommand(World, GToggleDataLayerActivationCommandName, InArgs);
 				if (UWorldPartitionSubsystem* WorldPartitionSubsystem = World->GetSubsystem<UWorldPartitionSubsystem>())
 				{
 					WorldPartitionSubsystem->ForEachWorldPartition([&InArgs](UWorldPartition* WorldPartition)
@@ -92,24 +94,25 @@ FAutoConsoleCommand UDataLayerManager::ToggleDataLayerActivation(
 	})
 );
 
+static const FString GSetDataLayerRuntimeStateCommandCommandName(TEXT("wp.Runtime.SetDataLayerRuntimeState"));
 FAutoConsoleCommand UDataLayerManager::SetDataLayerRuntimeStateCommand(
-	TEXT("wp.Runtime.SetDataLayerRuntimeState"),
+	*GSetDataLayerRuntimeStateCommandCommandName,
 	TEXT("Sets Runtime DataLayers state. Args [State = Unloaded, Loaded, Activated] [DataLayerNames]"),
 	FConsoleCommandWithArgsDelegate::CreateLambda([](const TArray<FString>& InArgs)
 	{
 		if (InArgs.Num() < 2)
 		{
-			UE_LOG(LogWorldPartition, Warning, TEXT("wp.Runtime.SetDataLayerRuntimeState : Requires at least 2 arguments. First argument should be the target state and the next ones should be the list of DataLayers."));
+			UE_LOG(LogWorldPartition, Warning, TEXT("%s : Requires at least 2 arguments. First argument should be the target state and the next ones should be the list of DataLayers."), *GSetDataLayerRuntimeStateCommandCommandName);
 			return;
 		}
 
-		TArray<FString> Args = InArgs;
+		TArray<FString> Args(InArgs);
 		FString StatetStr;
 		Args.HeapPop(StatetStr);
 		EDataLayerRuntimeState State;
 		if (!GetDataLayerRuntimeStateFromName(StatetStr, State))
 		{
-			UE_LOG(LogWorldPartition, Warning, TEXT("wp.Runtime.SetDataLayerRuntimeState : Invalid first argument, expected one of these values : Unloaded, Loaded, Activated."));
+			UE_LOG(LogWorldPartition, Warning, TEXT("%s : Invalid first argument, expected one of these values : Unloaded, Loaded, Activated."), *GSetDataLayerRuntimeStateCommandCommandName);
 			return;
 		}
 
@@ -118,6 +121,8 @@ FAutoConsoleCommand UDataLayerManager::SetDataLayerRuntimeStateCommand(
 			UWorld* World = Context.World();
 			if (World && World->IsGameWorld())
 			{
+				FWorldPartitionHelpers::ServerExecConsoleCommand(World, GSetDataLayerRuntimeStateCommandCommandName, InArgs);
+
 				if (UWorldPartitionSubsystem* WorldPartitionSubsystem = World->GetSubsystem<UWorldPartitionSubsystem>())
 				{
 					WorldPartitionSubsystem->ForEachWorldPartition([&Args, State](UWorldPartition* WorldPartition)
diff --git a/Engine/Source/Runtime/Engine/Private/WorldPartition/WorldPartitionHelpers.cpp b/Engine/Source/Runtime/Engine/Private/WorldPartition/WorldPartitionHelpers.cpp
--- a/Engine/Source/Runtime/Engine/Private/WorldPartition/WorldPartitionHelpers.cpp
+++ b/Engine/Source/Runtime/Engine/Private/WorldPartition/WorldPartitionHelpers.cpp
@@ -6,6 +6,8 @@
 #include "WorldPartition/WorldPartitionLog.h"
 #include "WorldPartition/WorldPartitionEditorHash.h"
 #include "WorldPartition/WorldPartitionRuntimeCell.h"
+#include "GameFramework/PlayerController.h"
+#include "Engine/Engine.h"
 #include "Algo/AnyOf.h"
 
 #include "Commandlets/Commandlet.h"
@@ -45,6 +47,24 @@ namespace FWorldPartitionHelpersPrivate
 	}
 }
 
+void FWorldPartitionHelpers::ServerExecConsoleCommand(UWorld* InWorld, const FString& InConsoleCommandName, const TArray<FString>& InArgs)
+{
+#if !UE_BUILD_SHIPPING && !WITH_EDITOR
+	if (InWorld && InWorld->IsGameWorld() && InWorld->IsNetMode(NM_Client))
+	{
+		if (APlayerController* PC = GEngine->GetFirstLocalPlayerController(InWorld))
+		{
+			TArray<FString> CmdList;
+			CmdList.Add(InConsoleCommandName);
+			CmdList.Append(InArgs);
+			FString Cmd = FString::Join(CmdList, TEXT(" "));
+			// Use ServerExecRPC instead of ServerExec to avoid any truncation
+			PC->ServerExecRPC(Cmd);
+		}
+	}
+#endif
+}
+
 #if WITH_EDITOR
 
 bool FWorldPartitionHelpers::IsActorDescClassCompatibleWith(const FWorldPartitionActorDesc* ActorDesc, const UClass* Class)
diff --git a/Engine/Source/Runtime/Engine/Private/WorldPartition/WorldPartitionRuntimeSpatialHash.cpp b/Engine/Source/Runtime/Engine/Private/WorldPartition/WorldPartitionRuntimeSpatialHash.cpp
--- a/Engine/Source/Runtime/Engine/Private/WorldPartition/WorldPartitionRuntimeSpatialHash.cpp
+++ b/Engine/Source/Runtime/Engine/Private/WorldPartition/WorldPartitionRuntimeSpatialHash.cpp
@@ -1502,12 +1502,13 @@ bool UWorldPartitionRuntimeSpatialHash::IsValidGrid(FName GridName) const
 
 #endif //WITH_EDITOR
 
+static const FString GOverrideLoadingRangeCommandName(TEXT("wp.Runtime.OverrideRuntimeSpatialHashLoadingRange"));
 FAutoConsoleCommand UWorldPartitionRuntimeSpatialHash::OverrideLoadingRangeCommand(
-	TEXT("wp.Runtime.OverrideRuntimeSpatialHashLoadingRange"),
+	*GOverrideLoadingRangeCommandName,
 	TEXT("Sets runtime loading range. Args -grid=[index] -range=[override_loading_range]"),
-	FConsoleCommandWithArgsDelegate::CreateLambda([](const TArray<FString>& Args)
+	FConsoleCommandWithArgsDelegate::CreateLambda([](const TArray<FString>& InArgs)
 	{
-		FString ArgString = FString::Join(Args, TEXT(" "));
+		FString ArgString = FString::Join(InArgs, TEXT(" "));
 		int32 GridIndex = 0;
 		float OverrideLoadingRange = -1.f;
 		FParse::Value(*ArgString, TEXT("grid="), GridIndex);
@@ -1518,6 +1519,7 @@ FAutoConsoleCommand UWorldPartitionRuntimeSpatialHash::OverrideLoadingRangeComma
 			UWorld* World = Context.World();
 			if (World && World->IsGameWorld())
 			{
+				FWorldPartitionHelpers::ServerExecConsoleCommand(World, GOverrideLoadingRangeCommandName, InArgs);
 				if (UWorldPartitionSubsystem* WorldPartitionSubsystem = World->GetSubsystem<UWorldPartitionSubsystem>())
 				{
 					WorldPartitionSubsystem->ForEachWorldPartition([GridIndex, OverrideLoadingRange](UWorldPartition* WorldPartition)
diff --git a/Engine/Source/Runtime/Engine/Public/WorldPartition/WorldPartitionHelpers.h b/Engine/Source/Runtime/Engine/Public/WorldPartition/WorldPartitionHelpers.h
--- a/Engine/Source/Runtime/Engine/Public/WorldPartition/WorldPartitionHelpers.h
+++ b/Engine/Source/Runtime/Engine/Public/WorldPartition/WorldPartitionHelpers.h
@@ -69,6 +69,9 @@ public:
 		return FWorldPartitionHelpersPrivate::GetWorldPartition(InObject);
 	}
 
+	/** Sends an RPC console command to the server (non shipping builds only). */
+	static ENGINE_API void ServerExecConsoleCommand(UWorld* InWorld, const FString& InConsoleCommandName, const TArray<FString>& InArgs);
+
 #if WITH_EDITOR
 private:
 	static ENGINE_API bool IsActorDescClassCompatibleWith(const FWorldPartitionActorDesc* ActorDesc, const UClass* Class);
