World Partition Data Layers
  - Added project setting DefaultDataLayerOperator that can be set to OR (default) or AND for actors data layers activation.
  - Added world partition setting DataLayersLogicOperator that can be set to ProjectDefault, OR or AND to overrride the default logic operator per world.

With these settings, a project (or a world) can change the default behavior of actor data layers activation logic so actors will be loaded if any data layer is activated (OR) or if all its data layers are activated (AND).

diff --git a/Engine/Source/Runtime/Engine/Private/WorldPartition/DataLayer/DataLayerManager.cpp b/Engine/Source/Runtime/Engine/Private/WorldPartition/DataLayer/DataLayerManager.cpp
--- a/Engine/Source/Runtime/Engine/Private/WorldPartition/DataLayer/DataLayerManager.cpp
+++ b/Engine/Source/Runtime/Engine/Private/WorldPartition/DataLayer/DataLayerManager.cpp
@@ -365,6 +365,34 @@ bool UDataLayerManager::IsAnyDataLayerInEffectiveRuntimeState(const TArray<FName
 	return false;
 }
 
+bool UDataLayerManager::IsAllDataLayerInEffectiveRuntimeState(const TArray<FName>& InDataLayerNames, EDataLayerRuntimeState InState) const
+{
+	const TSet<FName>& Activated = GetEffectiveActiveDataLayerNames();
+
+	if (InState == EDataLayerRuntimeState::Activated)
+	{
+		for (const FName& DataLayerName : InDataLayerNames)
+		{
+			if (!Activated.Contains(DataLayerName))
+			{
+				return false;
+			}
+		}
+	}
+	else if (InState == EDataLayerRuntimeState::Loaded)
+	{
+		const TSet<FName>& Loaded = GetEffectiveLoadedDataLayerNames();
+		for (const FName& DataLayerName : InDataLayerNames)
+		{
+			if (!Loaded.Contains(DataLayerName) && !Activated.Contains(DataLayerName))
+			{
+				return false;
+			}
+		}
+	}
+	return true;
+}
+
 void UDataLayerManager::DumpDataLayers(FOutputDevice& OutputDevice) const
 {
 	if (AWorldDataLayers* WorldDataLayers = GetWorldDataLayers())
diff --git a/Engine/Source/Runtime/Engine/Private/WorldPartition/RuntimeHashSet/WorldPartitionRuntimeHashSet.cpp b/Engine/Source/Runtime/Engine/Private/WorldPartition/RuntimeHashSet/WorldPartitionRuntimeHashSet.cpp
--- a/Engine/Source/Runtime/Engine/Private/WorldPartition/RuntimeHashSet/WorldPartitionRuntimeHashSet.cpp
+++ b/Engine/Source/Runtime/Engine/Private/WorldPartition/RuntimeHashSet/WorldPartitionRuntimeHashSet.cpp
@@ -394,20 +394,28 @@ void UWorldPartitionRuntimeHashSet::ForEachStreamingCellsSources(const TArray<FW
 					{
 						if (IsCellRelevantFor(Cell->GetClientOnlyVisible()))
 						{
-							if (!Cell->HasDataLayers() || Cell->HasAnyDataLayerInEffectiveRuntimeState(EDataLayerRuntimeState::Activated))
+							switch (Cell->GetCellEffectiveWantedState())
 							{
-								if (Source.TargetState == EStreamingSourceTargetState::Loaded)
+							case EDataLayerRuntimeState::Loaded:
+								LoadStreamingSourceCells.AddCell(Cell, Source, Shape);
+								break;
+							case EDataLayerRuntimeState::Activated:
+								switch (Source.TargetState)
 								{
+								case EStreamingSourceTargetState::Loaded:
 									LoadStreamingSourceCells.AddCell(Cell, Source, Shape);
-								}
-								else
-								{
+									break;
+								case EStreamingSourceTargetState::Activated:
 									ActivateStreamingSourceCells.AddCell(Cell, Source, Shape);
+									break;
+								default:
+									checkNoEntry();
 								}
-							}
-							else if (Cell->HasAnyDataLayerInEffectiveRuntimeState(EDataLayerRuntimeState::Loaded))
-							{
-								LoadStreamingSourceCells.AddCell(Cell, Source, Shape);
+								break;
+							case EDataLayerRuntimeState::Unloaded:
+								break;
+							default:
+								checkNoEntry();
 							}
 						}
 					});
@@ -422,13 +430,18 @@ void UWorldPartitionRuntimeHashSet::ForEachStreamingCellsSources(const TArray<FW
 		{
 			if (IsCellRelevantFor(Cell->GetClientOnlyVisible()))
 			{
-				if (!Cell->HasDataLayers() || Cell->HasAnyDataLayerInEffectiveRuntimeState(EDataLayerRuntimeState::Activated))
-				{
-					ActivateStreamingSourceCells.GetCells().Add(Cell);
-				}
-				else if (Cell->HasAnyDataLayerInEffectiveRuntimeState(EDataLayerRuntimeState::Loaded))
+				switch (Cell->GetCellEffectiveWantedState())
 				{
+				case EDataLayerRuntimeState::Loaded:
 					LoadStreamingSourceCells.GetCells().Add(Cell);
+					break;
+				case EDataLayerRuntimeState::Activated:
+					ActivateStreamingSourceCells.GetCells().Add(Cell);
+					break;
+				case EDataLayerRuntimeState::Unloaded:
+					break;
+				default:
+					checkNoEntry();
 				}
 			}
 		}
diff --git a/Engine/Source/Runtime/Engine/Private/WorldPartition/WorldPartition.cpp b/Engine/Source/Runtime/Engine/Private/WorldPartition/WorldPartition.cpp
--- a/Engine/Source/Runtime/Engine/Private/WorldPartition/WorldPartition.cpp
+++ b/Engine/Source/Runtime/Engine/Private/WorldPartition/WorldPartition.cpp
@@ -14,6 +14,7 @@
 #include "WorldPartition/WorldPartitionSubsystem.h"
 #include "WorldPartition/HLOD/HLODSubsystem.h"
 #include "WorldPartition/DataLayer/DataLayerManager.h"
+#include "WorldPartition/WorldPartitionSettings.h"
 #include "GameFramework/WorldSettings.h"
 #include "ProfilingDebugging/ScopedTimers.h"
 #include "UObject/UE5MainStreamObjectVersion.h"
@@ -293,6 +294,7 @@ UWorldPartition::UWorldPartition(const FObjectInitializer& ObjectInitializer)
 	bEnableStreaming = true;
 	ServerStreamingMode = EWorldPartitionServerStreamingMode::ProjectDefault;
 	ServerStreamingOutMode = EWorldPartitionServerStreamingOutMode::ProjectDefault;
+	DataLayersLogicOperator = EWorldPartitionDataLayersLogicOperator::ProjectDefault;
 	StreamingStateEpoch = 0;
 
 #if WITH_EDITOR
@@ -838,6 +840,24 @@ UDataLayerManager* UWorldPartition::GetDataLayerManager() const
 	return DataLayerManager;
 }
 
+EWorldPartitionDataLayersLogicOperator UWorldPartition::GetDataLayersLogicOperator() const
+{
+	if (DataLayersLogicOperator == EWorldPartitionDataLayersLogicOperator::ProjectDefault)
+	{
+		switch (UWorldPartitionSettings::Get()->GetDefaultDataLayerOperator())
+		{
+		case EDataLayerLogicOperator::Or:
+			return EWorldPartitionDataLayersLogicOperator::Or;
+		case EDataLayerLogicOperator::And:
+			return EWorldPartitionDataLayersLogicOperator::And;
+		default:
+			checkNoEntry();
+		}
+	}
+
+	return DataLayersLogicOperator;
+}
+
 bool UWorldPartition::IsInitialized() const
 {
 	return InitState == EWorldPartitionInitState::Initialized;
diff --git a/Engine/Source/Runtime/Engine/Private/WorldPartition/WorldPartitionRuntimeCell.cpp b/Engine/Source/Runtime/Engine/Private/WorldPartition/WorldPartitionRuntimeCell.cpp
--- a/Engine/Source/Runtime/Engine/Private/WorldPartition/WorldPartitionRuntimeCell.cpp
+++ b/Engine/Source/Runtime/Engine/Private/WorldPartition/WorldPartitionRuntimeCell.cpp
@@ -125,10 +125,45 @@ UDataLayerManager* UWorldPartitionRuntimeCell::GetDataLayerManager() const
 	return GetOuterWorld()->GetWorldPartition()->GetDataLayerManager();
 }
 
-bool UWorldPartitionRuntimeCell::HasAnyDataLayerInEffectiveRuntimeState(EDataLayerRuntimeState InState) const
+EDataLayerRuntimeState UWorldPartitionRuntimeCell::GetCellEffectiveWantedState() const
 {
-	const UDataLayerManager* DataLayerManager = HasDataLayers() ? GetDataLayerManager() : nullptr;
-	return DataLayerManager ? DataLayerManager->IsAnyDataLayerInEffectiveRuntimeState(GetDataLayers(), InState) : false;
+	if (!HasDataLayers())
+	{
+		return EDataLayerRuntimeState::Activated;
+	}
+
+	UWorldPartition* WorldPartition = GetOuterWorld()->GetWorldPartition();
+
+	if (const UDataLayerManager* DataLayerManager = WorldPartition->GetDataLayerManager())
+	{
+		switch (WorldPartition->GetDataLayersLogicOperator())
+		{
+		case EWorldPartitionDataLayersLogicOperator::Or:
+			if (DataLayerManager->IsAnyDataLayerInEffectiveRuntimeState(GetDataLayers(), EDataLayerRuntimeState::Activated))
+			{
+				return EDataLayerRuntimeState::Activated;
+			}
+			else if (DataLayerManager->IsAnyDataLayerInEffectiveRuntimeState(GetDataLayers(), EDataLayerRuntimeState::Loaded))
+			{
+				return EDataLayerRuntimeState::Loaded;
+			}
+			break;
+		case EWorldPartitionDataLayersLogicOperator::And:
+			if (DataLayerManager->IsAllDataLayerInEffectiveRuntimeState(GetDataLayers(), EDataLayerRuntimeState::Activated))
+			{
+				return EDataLayerRuntimeState::Activated;
+			}
+			else if (DataLayerManager->IsAllDataLayerInEffectiveRuntimeState(GetDataLayers(), EDataLayerRuntimeState::Loaded))
+			{
+				return EDataLayerRuntimeState::Loaded;
+			}
+			break;
+		default:
+			checkNoEntry();
+		}
+	}
+
+	return EDataLayerRuntimeState::Unloaded;
 }
 
 TArray<const UDataLayerInstance*> UWorldPartitionRuntimeCell::GetDataLayerInstances() const
diff --git a/Engine/Source/Runtime/Engine/Private/WorldPartition/WorldPartitionRuntimeSpatialHash.cpp b/Engine/Source/Runtime/Engine/Private/WorldPartition/WorldPartitionRuntimeSpatialHash.cpp
--- a/Engine/Source/Runtime/Engine/Private/WorldPartition/WorldPartitionRuntimeSpatialHash.cpp
+++ b/Engine/Source/Runtime/Engine/Private/WorldPartition/WorldPartitionRuntimeSpatialHash.cpp
@@ -263,22 +263,29 @@ void FSpatialHashStreamingGrid::GetCells(const TArray<FWorldPartitionStreamingSo
 						
 						if (bIncludeCell)
 						{
-							if (!Cell->HasDataLayers() || Cell->HasAnyDataLayerInEffectiveRuntimeState(EDataLayerRuntimeState::Activated))
+							switch (Cell->GetCellEffectiveWantedState())
 							{
-								if (Source.TargetState == EStreamingSourceTargetState::Loaded)
+							case EDataLayerRuntimeState::Loaded:
+								OutLoadCells.AddCell(Cell, Source, Shape);
+								break;
+							case EDataLayerRuntimeState::Activated:
+								switch (Source.TargetState)
 								{
+								case EStreamingSourceTargetState::Loaded:
 									OutLoadCells.AddCell(Cell, Source, Shape);
-								}
-								else
-								{
-									check(Source.TargetState == EStreamingSourceTargetState::Activated);
+									break;
+								case EStreamingSourceTargetState::Activated:
 									OutActivateCells.AddCell(Cell, Source, Shape);
 									bAddedActivatedCell = !GRuntimeSpatialHashUseAlignedGridLevelsEffective && GRuntimeSpatialHashSnapNonAlignedGridLevelsToLowerLevelsEffective;
+									break;
+								default:
+									checkNoEntry();
 								}
-							}
-							else if (Cell->HasAnyDataLayerInEffectiveRuntimeState(EDataLayerRuntimeState::Loaded))
-							{
-								OutLoadCells.AddCell(Cell, Source, Shape);
+								break;
+							case EDataLayerRuntimeState::Unloaded:
+								break;
+							default:
+								checkNoEntry();
 							}
 						}
 					});
@@ -348,12 +355,20 @@ void FSpatialHashStreamingGrid::GetCells(const TArray<FWorldPartitionStreamingSo
 		{
 			ForEachRuntimeCell(ParentCell.Key, [&](const UWorldPartitionRuntimeCell* Cell)
 			{
-				if (!Cell->HasDataLayers() || Cell->HasAnyDataLayerInEffectiveRuntimeState(EDataLayerRuntimeState::Activated))
+				switch (Cell->GetCellEffectiveWantedState())
 				{
+				case EDataLayerRuntimeState::Loaded:
+					break;
+				case EDataLayerRuntimeState::Activated:
 					for (const auto& Info : ParentCell.Value)
 					{
 						OutActivateCells.AddCell(Cell, Info.Source, Info.SourceShape);
 					}
+					break;
+				case EDataLayerRuntimeState::Unloaded:
+					break;
+				default:
+					checkNoEntry();
 				}
 			});
 		}
@@ -581,15 +596,20 @@ void FSpatialHashStreamingGrid::GetNonSpatiallyLoadedCells(TSet<const UWorldPart
 			{
 				for (const UWorldPartitionRuntimeCell* Cell : LayerCell.GridCells)
 				{
-					if (!Cell->HasDataLayers() || Cell->HasAnyDataLayerInEffectiveRuntimeState(EDataLayerRuntimeState::Activated))
-					{
-						check(Cell->IsAlwaysLoaded() || Cell->HasDataLayers() || Cell->GetContentBundleID().IsValid());
-						OutActivateCells.Add(Cell);
-					}
-					else if (Cell->HasAnyDataLayerInEffectiveRuntimeState(EDataLayerRuntimeState::Loaded))
+					switch (Cell->GetCellEffectiveWantedState())
 					{
+					case EDataLayerRuntimeState::Loaded:
 						check(Cell->HasDataLayers());
 						OutLoadCells.Add(Cell);
+						break;
+					case EDataLayerRuntimeState::Activated:
+						check(Cell->IsAlwaysLoaded() || Cell->HasDataLayers() || Cell->GetContentBundleID().IsValid());
+						OutActivateCells.Add(Cell);
+						break;
+					case EDataLayerRuntimeState::Unloaded:
+						break;
+					default:
+						checkNoEntry();
 					}
 				}
 			}
@@ -614,10 +634,19 @@ void FSpatialHashStreamingGrid::GetFilteredCellsForDebugDraw(const FSpatialHashS
 			{
 				EStreamingStatus StreamingStatus = Cell->GetStreamingStatus();
 				const TArray<FName>& DataLayers = Cell->GetDataLayers();
-				return (!Cell->HasDataLayers() ||
-					Cell->HasAnyDataLayerInEffectiveRuntimeState(EDataLayerRuntimeState::Loaded) ||
-					Cell->HasAnyDataLayerInEffectiveRuntimeState(EDataLayerRuntimeState::Activated) ||
-					((StreamingStatus != LEVEL_Unloaded) && (StreamingStatus != LEVEL_UnloadedButStillAround)));
+
+				switch (Cell->GetCellEffectiveWantedState())
+				{
+				case EDataLayerRuntimeState::Loaded:
+				case EDataLayerRuntimeState::Activated:
+					return true;
+				case EDataLayerRuntimeState::Unloaded:
+					break;
+				default:
+					checkNoEntry();
+				}
+
+				return (StreamingStatus != LEVEL_Unloaded) && (StreamingStatus != LEVEL_UnloadedButStillAround);
 			}
 			return false;
 		});
diff --git a/Engine/Source/Runtime/Engine/Private/WorldPartition/WorldPartitionSettings.cpp b/Engine/Source/Runtime/Engine/Private/WorldPartition/WorldPartitionSettings.cpp
new file mode 100644
--- /dev/null
+++ b/Engine/Source/Runtime/Engine/Private/WorldPartition/WorldPartitionSettings.cpp
@@ -0,0 +1,7 @@
+// Copyright Epic Games, Inc. All Rights Reserved.
+
+#include "WorldPartition/WorldPartitionSettings.h"
+
+UWorldPartitionSettings::UWorldPartitionSettings(const FObjectInitializer& ObjectInitializer)
+	: Super(ObjectInitializer)
+{}
\ No newline at end of file
diff --git a/Engine/Source/Runtime/Engine/Public/WorldPartition/DataLayer/DataLayerManager.h b/Engine/Source/Runtime/Engine/Public/WorldPartition/DataLayer/DataLayerManager.h
--- a/Engine/Source/Runtime/Engine/Public/WorldPartition/DataLayer/DataLayerManager.h
+++ b/Engine/Source/Runtime/Engine/Public/WorldPartition/DataLayer/DataLayerManager.h
@@ -91,6 +91,7 @@ public:
 	ENGINE_API const TSet<FName>& GetEffectiveActiveDataLayerNames() const;
 	ENGINE_API const TSet<FName>& GetEffectiveLoadedDataLayerNames() const;
 	ENGINE_API bool IsAnyDataLayerInEffectiveRuntimeState(const TArray<FName>& InDataLayerNames, EDataLayerRuntimeState InState) const;
+	ENGINE_API bool IsAllDataLayerInEffectiveRuntimeState(const TArray<FName>& InDataLayerNames, EDataLayerRuntimeState InState) const;
 	//~ End Runtime State
 
 private:
diff --git a/Engine/Source/Runtime/Engine/Public/WorldPartition/WorldPartition.h b/Engine/Source/Runtime/Engine/Public/WorldPartition/WorldPartition.h
--- a/Engine/Source/Runtime/Engine/Public/WorldPartition/WorldPartition.h
+++ b/Engine/Source/Runtime/Engine/Public/WorldPartition/WorldPartition.h
@@ -76,6 +76,14 @@ enum class EWorldPartitionServerStreamingOutMode : uint8
 	Enabled = 2 UMETA(ToolTip = "Server streaming out is enabled"),
 };
 
+UENUM()
+enum class EWorldPartitionDataLayersLogicOperator : uint8
+{
+	ProjectDefault,
+	Or,
+	And
+};
+
 #if WITH_EDITOR
 /**
  * Interface for the world partition editor
@@ -422,6 +430,8 @@ public:
 
 	ENGINE_API UDataLayerManager* GetDataLayerManager() const;
 
+	ENGINE_API EWorldPartitionDataLayersLogicOperator GetDataLayersLogicOperator() const;
+
 	UE_DEPRECATED(5.3, "UpdateStreamingState is deprecated, use UWorldPartitionSubsystem::UpdateStreamingState instead.")
 	void UpdateStreamingState() {}
 	UE_DEPRECATED(5.3, "CanAddLoadedLevelToWorld is deprecated, use CanAddCellToWorld instead.")
@@ -470,6 +480,9 @@ public:
 	EWorldPartitionServerStreamingOutMode ServerStreamingOutMode;
 
 private:
+	UPROPERTY(EditAnywhere, Category = WorldPartitionSetup, AdvancedDisplay)
+	EWorldPartitionDataLayersLogicOperator DataLayersLogicOperator;
+
 	TObjectPtr<UWorld> World;
 
 #if WITH_EDITOR
diff --git a/Engine/Source/Runtime/Engine/Public/WorldPartition/WorldPartitionRuntimeCell.h b/Engine/Source/Runtime/Engine/Public/WorldPartition/WorldPartitionRuntimeCell.h
--- a/Engine/Source/Runtime/Engine/Public/WorldPartition/WorldPartitionRuntimeCell.h
+++ b/Engine/Source/Runtime/Engine/Public/WorldPartition/WorldPartitionRuntimeCell.h
@@ -22,6 +22,8 @@ class UWorldPartition;
 class UDataLayerManager;
 struct FHierarchicalLogArchive;
 
+enum class EWorldPartitionDataLayersLogicOperator : uint8;
+
 enum class EWorldPartitionRuntimeCellVisualizeMode
 {
 	StreamingPriority,
@@ -215,7 +217,7 @@ class UWorldPartitionRuntimeCell : public UObject, public IWorldPartitionCell
 	inline bool HasDataLayers() const { return !DataLayers.IsEmpty(); }
 
 	ENGINE_API UDataLayerManager* GetDataLayerManager() const;
-	ENGINE_API bool HasAnyDataLayerInEffectiveRuntimeState(EDataLayerRuntimeState InState) const;
+	ENGINE_API EDataLayerRuntimeState GetCellEffectiveWantedState() const;
 
 	void SetBlockOnSlowLoading(bool bInBlockOnSlowLoading) { bBlockOnSlowLoading = bInBlockOnSlowLoading; }
 	bool GetBlockOnSlowLoading() const { return bBlockOnSlowLoading; }
diff --git a/Engine/Source/Runtime/Engine/Public/WorldPartition/WorldPartitionSettings.h b/Engine/Source/Runtime/Engine/Public/WorldPartition/WorldPartitionSettings.h
new file mode 100644
--- /dev/null
+++ b/Engine/Source/Runtime/Engine/Public/WorldPartition/WorldPartitionSettings.h
@@ -0,0 +1,29 @@
+// Copyright Epic Games, Inc. All Rights Reserved.
+#pragma once
+
+#include "Engine/DeveloperSettings.h"
+#include "WorldPartitionSettings.generated.h"
+
+UENUM()
+enum class EDataLayerLogicOperator : uint8
+{
+	Or,
+	And
+};
+
+UCLASS(config = Engine, defaultconfig, DisplayName = "World Partition")
+class ENGINE_API UWorldPartitionSettings : public UDeveloperSettings
+{
+	GENERATED_BODY()
+
+public:
+	UWorldPartitionSettings(const FObjectInitializer& ObjectInitializer = FObjectInitializer::Get());
+
+	static UWorldPartitionSettings* Get() { return CastChecked<UWorldPartitionSettings>(UWorldPartitionSettings::StaticClass()->GetDefaultObject()); }
+
+	EDataLayerLogicOperator GetDefaultDataLayerOperator() const { return DefaultDataLayerOperator; }
+
+protected:
+	UPROPERTY(EditAnywhere, config, Category = WorldPartition)
+	EDataLayerLogicOperator DefaultDataLayerOperator = EDataLayerLogicOperator::Or;
+};
\ No newline at end of file
