Disable pre-shadow warning in projects that don't use static ligthting.

diff --git a/Engine/Source/Runtime/Engine/Private/Components/StaticMeshComponent.cpp b/Engine/Source/Runtime/Engine/Private/Components/StaticMeshComponent.cpp
--- a/Engine/Source/Runtime/Engine/Private/Components/StaticMeshComponent.cpp
+++ b/Engine/Source/Runtime/Engine/Private/Components/StaticMeshComponent.cpp
@@ -585,7 +585,8 @@ void UStaticMeshComponent::CheckForErrors()
 		CastShadow && 
 		bCastDynamicShadow && 
 		IsRegistered() && 
-		Bounds.SphereRadius > 2000.0f )
+		Bounds.SphereRadius > 2000.0f &&
+		IsStaticLightingAllowed())
 	{
 		// Large shadow casting objects that create preshadows will cause a massive performance hit, since preshadows are meant for small shadow casters.
 		FMessageLog("MapCheck").PerformanceWarning()
