-Fix for WorldPartitionBlueprintLibrary::GetIntersectingActorDescs,
-Added missing TransformBy for Bounds of ActorDesc in container instances

diff --git a/Engine/Source/Runtime/Engine/Private/WorldPartitionBlueprintLibrary.cpp b/Engine/Source/Runtime/Engine/Private/WorldPartitionBlueprintLibrary.cpp
--- a/Engine/Source/Runtime/Engine/Private/WorldPartitionBlueprintLibrary.cpp
+++ b/Engine/Source/Runtime/Engine/Private/WorldPartitionBlueprintLibrary.cpp
@@ -163,7 +163,7 @@ bool UWorldPartitionBlueprintLibrary::GetIntersectingActorDescs(const UActorDesc
 
 	for (FActorDescList::TConstIterator<> ActorDescIt(InContainer); ActorDescIt; ++ActorDescIt)
 	{
-		if (ActorDescIt->GetEditorBounds().Intersect(InBox))
+		if (ActorDescIt->GetEditorBounds().TransformBy(InTransform).Intersect(InBox))
 		{
 			bResult &= HandleIntersectingActorDesc(*ActorDescIt, InBox, InTransform, OutActorDescs);
 		}
