Actor: bIsInLevelInstance transient property is initialized by ULevelInstanceLevelStreaming (will be used by next SceneOutliner change)

#rnx

diff --git a/Engine/Source/Runtime/Engine/Classes/GameFramework/Actor.h b/Engine/Source/Runtime/Engine/Classes/GameFramework/Actor.h
--- a/Engine/Source/Runtime/Engine/Classes/GameFramework/Actor.h
+++ b/Engine/Source/Runtime/Engine/Classes/GameFramework/Actor.h
@@ -295,6 +295,11 @@ private:
 	UPROPERTY(Transient)
 	uint8 bIsInEditingLevelInstance:1;
 
+	UPROPERTY(Transient)
+	uint8 bIsInLevelInstance:1;
+
+	friend struct FSetActorIsInLevelInstance;
+
 public:
 	UPROPERTY(EditAnywhere, AdvancedDisplay, Category = LevelInstance, meta = (Tooltip = "If checked, this Actor will only get loaded in a main world (persistent level), it will not be loaded through Level Instances."))
 	uint8 bIsMainWorldOnly : 1;
@@ -326,6 +331,12 @@ public:
 	{
 		return bIsInEditingLevelInstance;
 	}
+
+	/** If true, the actor belongs to a level instance. */
+	virtual bool IsInLevelInstance() const
+	{
+		return bIsInEditingLevelInstance || bIsInLevelInstance;
+	}
 #endif
 
 	/** Networking - Server - TearOff this actor to stop replication to clients. Will set bTearOff to true. */
@@ -4403,7 +4414,6 @@ private:
 	friend class UWorldPartitionConvertCommandlet;
 };
 
-#if WITH_EDITOR
 struct FSetActorReplicates
 {
 private:
@@ -4418,7 +4428,6 @@ private:
 	}
 	friend class FWorldPartitionLevelHelper;
 };
-#endif
 
 struct FSetActorInstanceGuid
 {
@@ -4467,6 +4476,17 @@ private:
 	friend class UDataLayerInstancePrivate;
 	friend class ULevelInstanceSubsystem;
 };
+
+struct FSetActorIsInLevelInstance
+{
+private:
+	FSetActorIsInLevelInstance(AActor* InActor)
+	{
+		InActor->bIsInLevelInstance = true;
+	}
+
+	friend class ULevelStreamingLevelInstance;
+};
 #endif
 
 /** Helper function for executing tick functions based on the normal conditions previous found in UActorComponent::ConditionalTick */
diff --git a/Engine/Source/Runtime/Engine/Private/Actor.cpp b/Engine/Source/Runtime/Engine/Private/Actor.cpp
--- a/Engine/Source/Runtime/Engine/Private/Actor.cpp
+++ b/Engine/Source/Runtime/Engine/Private/Actor.cpp
@@ -211,6 +211,7 @@ void AActor::InitializeDefaults()
 
 #if WITH_EDITORONLY_DATA
 	bIsInEditingLevelInstance = false;
+	bIsInLevelInstance = false;
 	PivotOffset = FVector::ZeroVector;
 #endif
 	SpawnCollisionHandlingMethod = ESpawnActorCollisionHandlingMethod::AlwaysSpawn;
diff --git a/Engine/Source/Runtime/Engine/Private/LevelInstance/LevelInstanceLevelStreaming.cpp b/Engine/Source/Runtime/Engine/Private/LevelInstance/LevelInstanceLevelStreaming.cpp
--- a/Engine/Source/Runtime/Engine/Private/LevelInstance/LevelInstanceLevelStreaming.cpp
+++ b/Engine/Source/Runtime/Engine/Private/LevelInstance/LevelInstanceLevelStreaming.cpp
@@ -198,6 +198,8 @@ void ULevelStreamingLevelInstance::PrepareLevelInstanceLoadedActor(AActor& InAct
 			InActor.AttachToActor(LevelInstanceEditorInstanceActor.Get(), FAttachmentTransformRules::KeepWorldTransform);
 		}
 	}
+
+	FSetActorIsInLevelInstance SetIsInLevelInstance(&InActor);
 }
 
 void ULevelStreamingLevelInstance::OnLoadedActorRemovedFromLevel(AActor& InActor)
