World Partition Data Layers Logical Operator: change per-project setting to be only applied to new maps and force maps saved prior to this change to use the OR operator.

diff --git a/Engine/Source/Runtime/Core/Public/UObject/FortniteMainBranchObjectVersion.h b/Engine/Source/Runtime/Core/Public/UObject/FortniteMainBranchObjectVersion.h
--- a/Engine/Source/Runtime/Core/Public/UObject/FortniteMainBranchObjectVersion.h
+++ b/Engine/Source/Runtime/Core/Public/UObject/FortniteMainBranchObjectVersion.h
@@ -363,6 +363,9 @@ struct FFortniteMainBranchObjectVersion
 		// GrassTypes are now per-component, rather than per-landscape proxy :
 		LandscapeSupportPerComponentGrassTypes,
 
+		// World partition actor data layers activation logic operator support defaults for old maps
+		WorldPartitionDataLayersLogicOperatorAdded,
+
 		// -----<new versions can be added above this line>-------------------------------------------------
 		VersionPlusOne,
 		LatestVersion = VersionPlusOne - 1
diff --git a/Engine/Source/Runtime/Engine/Private/WorldPartition/WorldPartition.cpp b/Engine/Source/Runtime/Engine/Private/WorldPartition/WorldPartition.cpp
--- a/Engine/Source/Runtime/Engine/Private/WorldPartition/WorldPartition.cpp
+++ b/Engine/Source/Runtime/Engine/Private/WorldPartition/WorldPartition.cpp
@@ -294,7 +294,7 @@ UWorldPartition::UWorldPartition(const FObjectInitializer& ObjectInitializer)
 	bEnableStreaming = true;
 	ServerStreamingMode = EWorldPartitionServerStreamingMode::ProjectDefault;
 	ServerStreamingOutMode = EWorldPartitionServerStreamingOutMode::ProjectDefault;
-	DataLayersLogicOperator = EWorldPartitionDataLayersLogicOperator::ProjectDefault;
+	DataLayersLogicOperator = UWorldPartitionSettings::Get()->GetNewMapsDataLayersLogicOperator();
 	StreamingStateEpoch = 0;
 
 #if WITH_EDITOR
@@ -840,24 +840,6 @@ UDataLayerManager* UWorldPartition::GetDataLayerManager() const
 	return DataLayerManager;
 }
 
-EWorldPartitionDataLayersLogicOperator UWorldPartition::GetDataLayersLogicOperator() const
-{
-	if (DataLayersLogicOperator == EWorldPartitionDataLayersLogicOperator::ProjectDefault)
-	{
-		switch (UWorldPartitionSettings::Get()->GetDefaultDataLayerOperator())
-		{
-		case EDataLayerLogicOperator::Or:
-			return EWorldPartitionDataLayersLogicOperator::Or;
-		case EDataLayerLogicOperator::And:
-			return EWorldPartitionDataLayersLogicOperator::And;
-		default:
-			checkNoEntry();
-		}
-	}
-
-	return DataLayersLogicOperator;
-}
-
 bool UWorldPartition::IsInitialized() const
 {
 	return InitState == EWorldPartitionInitState::Initialized;
@@ -1510,14 +1492,22 @@ void UWorldPartition::Serialize(FArchive& Ar)
 		Ar << bIsPIE;
 #endif
 	}
-	else if (Ar.CustomVer(FUE5MainStreamObjectVersion::GUID) >= FUE5MainStreamObjectVersion::WorldPartitionSerializeStreamingPolicyOnCook)
+	else
 	{
-		bool bCooked = Ar.IsCooking();
-		Ar << bCooked;
+		if (Ar.CustomVer(FUE5MainStreamObjectVersion::GUID) >= FUE5MainStreamObjectVersion::WorldPartitionSerializeStreamingPolicyOnCook)
+		{
+			bool bCooked = Ar.IsCooking();
+			Ar << bCooked;
 
-		if (bCooked)
+			if (bCooked)
+			{
+				Ar << StreamingPolicy;
+			}
+		}
+	
+		if (Ar.CustomVer(FFortniteMainBranchObjectVersion::GUID) < FFortniteMainBranchObjectVersion::WorldPartitionDataLayersLogicOperatorAdded)
 		{
-			Ar << StreamingPolicy;
+			DataLayersLogicOperator = EWorldPartitionDataLayersLogicOperator::Or;
 		}
 	}
 }
diff --git a/Engine/Source/Runtime/Engine/Public/WorldPartition/WorldPartition.h b/Engine/Source/Runtime/Engine/Public/WorldPartition/WorldPartition.h
--- a/Engine/Source/Runtime/Engine/Public/WorldPartition/WorldPartition.h
+++ b/Engine/Source/Runtime/Engine/Public/WorldPartition/WorldPartition.h
@@ -79,7 +79,6 @@ enum class EWorldPartitionServerStreamingOutMode : uint8
 UENUM()
 enum class EWorldPartitionDataLayersLogicOperator : uint8
 {
-	ProjectDefault,
 	Or,
 	And
 };
@@ -430,7 +429,7 @@ public:
 
 	ENGINE_API UDataLayerManager* GetDataLayerManager() const;
 
-	ENGINE_API EWorldPartitionDataLayersLogicOperator GetDataLayersLogicOperator() const;
+	inline EWorldPartitionDataLayersLogicOperator GetDataLayersLogicOperator() const { return DataLayersLogicOperator; }
 
 	UE_DEPRECATED(5.3, "UpdateStreamingState is deprecated, use UWorldPartitionSubsystem::UpdateStreamingState instead.")
 	void UpdateStreamingState() {}
diff --git a/Engine/Source/Runtime/Engine/Public/WorldPartition/WorldPartitionSettings.h b/Engine/Source/Runtime/Engine/Public/WorldPartition/WorldPartitionSettings.h
--- a/Engine/Source/Runtime/Engine/Public/WorldPartition/WorldPartitionSettings.h
+++ b/Engine/Source/Runtime/Engine/Public/WorldPartition/WorldPartitionSettings.h
@@ -2,15 +2,9 @@
 #pragma once
 
 #include "Engine/DeveloperSettings.h"
+#include "WorldPartition/WorldPartition.h"
 #include "WorldPartitionSettings.generated.h"
 
-UENUM()
-enum class EDataLayerLogicOperator : uint8
-{
-	Or,
-	And
-};
-
 UCLASS(config = Engine, defaultconfig, DisplayName = "World Partition")
 class ENGINE_API UWorldPartitionSettings : public UDeveloperSettings
 {
@@ -21,9 +15,10 @@ public:
 
 	static UWorldPartitionSettings* Get() { return CastChecked<UWorldPartitionSettings>(UWorldPartitionSettings::StaticClass()->GetDefaultObject()); }
 
-	EDataLayerLogicOperator GetDefaultDataLayerOperator() const { return DefaultDataLayerOperator; }
+	EWorldPartitionDataLayersLogicOperator GetNewMapsDataLayersLogicOperator() const { return NewMapsDataLayersLogicOperator; }
 
 protected:
+	/** Set the default logical operator for actor data layers activation for new maps */
 	UPROPERTY(EditAnywhere, config, Category = WorldPartition)
-	EDataLayerLogicOperator DefaultDataLayerOperator = EDataLayerLogicOperator::Or;
+	EWorldPartitionDataLayersLogicOperator NewMapsDataLayersLogicOperator = EWorldPartitionDataLayersLogicOperator::Or;
 };
\ No newline at end of file
