World Partition Data Layers: optimize AWorldDataLayers::GetEffectiveDataLayerNames functions by caching the value when dependencies changes.

diff --git a/Engine/Source/Runtime/Engine/Private/WorldPartition/DataLayer/WorldDataLayers.cpp b/Engine/Source/Runtime/Engine/Private/WorldPartition/DataLayer/WorldDataLayers.cpp
--- a/Engine/Source/Runtime/Engine/Private/WorldPartition/DataLayer/WorldDataLayers.cpp
+++ b/Engine/Source/Runtime/Engine/Private/WorldPartition/DataLayer/WorldDataLayers.cpp
@@ -59,6 +59,9 @@ AWorldDataLayers::AWorldDataLayers(const FObjectInitializer& ObjectInitializer)
 	// Avoid actor from being Destroyed/Recreated when scrubbing a replay
 	// instead AWorldDataLayers::RewindForReplay() gets called to reset this actors state
 	bReplayRewindable = true;
+
+	AllEffectiveActiveDataLayerNamesEpoch = MAX_int32;
+	AllEffectiveLoadedDataLayerNamesEpoch = MAX_int32;
 }
 
 void AWorldDataLayers::GetLifetimeReplicatedProps(TArray< FLifetimeProperty >& OutLifetimeProps) const
@@ -382,15 +385,23 @@ EDataLayerRuntimeState AWorldDataLayers::GetDataLayerEffectiveRuntimeStateByName
 
 const TSet<FName>& AWorldDataLayers::GetEffectiveActiveDataLayerNames() const
 {
-	AllEffectiveActiveDataLayerNames = EffectiveActiveDataLayerNames;
-	AllEffectiveActiveDataLayerNames.Append(LocalEffectiveActiveDataLayerNames);
+	if (AllEffectiveActiveDataLayerNamesEpoch != DataLayersStateEpoch)
+	{
+		AllEffectiveActiveDataLayerNames = EffectiveActiveDataLayerNames;
+		AllEffectiveActiveDataLayerNames.Append(LocalEffectiveActiveDataLayerNames);
+		AllEffectiveActiveDataLayerNamesEpoch = DataLayersStateEpoch;
+	}
 	return AllEffectiveActiveDataLayerNames;
 }
 
 const TSet<FName>& AWorldDataLayers::GetEffectiveLoadedDataLayerNames() const
 {
-	AllEffectiveLoadedDataLayerNames = EffectiveLoadedDataLayerNames;
-	AllEffectiveLoadedDataLayerNames.Append(LocalEffectiveLoadedDataLayerNames);
+	if (AllEffectiveLoadedDataLayerNamesEpoch != DataLayersStateEpoch)
+	{
+		AllEffectiveLoadedDataLayerNames = EffectiveLoadedDataLayerNames;
+		AllEffectiveLoadedDataLayerNames.Append(LocalEffectiveLoadedDataLayerNames);
+		AllEffectiveLoadedDataLayerNamesEpoch = DataLayersStateEpoch;
+	}
 	return AllEffectiveLoadedDataLayerNames;
 }
 
diff --git a/Engine/Source/Runtime/Engine/Public/WorldPartition/DataLayer/WorldDataLayers.h b/Engine/Source/Runtime/Engine/Public/WorldPartition/DataLayer/WorldDataLayers.h
--- a/Engine/Source/Runtime/Engine/Public/WorldPartition/DataLayer/WorldDataLayers.h
+++ b/Engine/Source/Runtime/Engine/Public/WorldPartition/DataLayer/WorldDataLayers.h
@@ -253,7 +253,9 @@ private:
 	TSet<FName> LocalEffectiveActiveDataLayerNames;
 	TSet<FName> LocalEffectiveLoadedDataLayerNames;
 
+	mutable int32 AllEffectiveActiveDataLayerNamesEpoch;
 	mutable TSet<FName> AllEffectiveActiveDataLayerNames;
+	mutable int32 AllEffectiveLoadedDataLayerNamesEpoch;
 	mutable TSet<FName> AllEffectiveLoadedDataLayerNames;
 
 	static ENGINE_API int32 DataLayersStateEpoch;
