World Partition Data Layers: optimize UWorldPartitionRuntimeCell::GetCellEffectiveWantedState by caching the value when dependencies changes.

diff --git a/Engine/Source/Runtime/Engine/Private/WorldPartition/WorldPartitionRuntimeCell.cpp b/Engine/Source/Runtime/Engine/Private/WorldPartition/WorldPartitionRuntimeCell.cpp
--- a/Engine/Source/Runtime/Engine/Private/WorldPartition/WorldPartitionRuntimeCell.cpp
+++ b/Engine/Source/Runtime/Engine/Private/WorldPartition/WorldPartitionRuntimeCell.cpp
@@ -22,7 +22,10 @@ UWorldPartitionRuntimeCell::UWorldPartitionRuntimeCell(const FObjectInitializer&
 	, DebugStreamingPriority(-1.f)
 #endif
 	, RuntimeCellData(nullptr)
-{}
+{
+	EffectiveWantedState = EDataLayerRuntimeState::Unloaded;
+	EffectiveWantedStateEpoch = MAX_int32;
+}
 
 UWorld* UWorldPartitionRuntimeCell::GetOwningWorld() const
 {
@@ -127,43 +130,51 @@ UDataLayerManager* UWorldPartitionRuntimeCell::GetDataLayerManager() const
 
 EDataLayerRuntimeState UWorldPartitionRuntimeCell::GetCellEffectiveWantedState() const
 {
-	if (!HasDataLayers())
-	{
-		return EDataLayerRuntimeState::Activated;
-	}
-
-	UWorldPartition* WorldPartition = GetOuterWorld()->GetWorldPartition();
-
-	if (const UDataLayerManager* DataLayerManager = WorldPartition->GetDataLayerManager())
+	if (EffectiveWantedStateEpoch != AWorldDataLayers::GetDataLayersStateEpoch())
 	{
-		switch (WorldPartition->GetDataLayersLogicOperator())
+		if (!HasDataLayers())
 		{
-		case EWorldPartitionDataLayersLogicOperator::Or:
-			if (DataLayerManager->IsAnyDataLayerInEffectiveRuntimeState(GetDataLayers(), EDataLayerRuntimeState::Activated))
-			{
-				return EDataLayerRuntimeState::Activated;
-			}
-			else if (DataLayerManager->IsAnyDataLayerInEffectiveRuntimeState(GetDataLayers(), EDataLayerRuntimeState::Loaded))
-			{
-				return EDataLayerRuntimeState::Loaded;
-			}
-			break;
-		case EWorldPartitionDataLayersLogicOperator::And:
-			if (DataLayerManager->IsAllDataLayerInEffectiveRuntimeState(GetDataLayers(), EDataLayerRuntimeState::Activated))
-			{
-				return EDataLayerRuntimeState::Activated;
-			}
-			else if (DataLayerManager->IsAllDataLayerInEffectiveRuntimeState(GetDataLayers(), EDataLayerRuntimeState::Loaded))
+			EffectiveWantedState = EDataLayerRuntimeState::Activated;
+		}
+		else
+		{
+			EffectiveWantedState = EDataLayerRuntimeState::Unloaded;
+
+			UWorldPartition* WorldPartition = GetOuterWorld()->GetWorldPartition();
+			if (const UDataLayerManager* DataLayerManager = WorldPartition->GetDataLayerManager())
 			{
-				return EDataLayerRuntimeState::Loaded;
+				switch (WorldPartition->GetDataLayersLogicOperator())
+				{
+				case EWorldPartitionDataLayersLogicOperator::Or:
+					if (DataLayerManager->IsAnyDataLayerInEffectiveRuntimeState(GetDataLayers(), EDataLayerRuntimeState::Activated))
+					{
+						EffectiveWantedState = EDataLayerRuntimeState::Activated;
+					}
+					else if (DataLayerManager->IsAnyDataLayerInEffectiveRuntimeState(GetDataLayers(), EDataLayerRuntimeState::Loaded))
+					{
+						EffectiveWantedState = EDataLayerRuntimeState::Loaded;
+					}
+					break;
+				case EWorldPartitionDataLayersLogicOperator::And:
+					if (DataLayerManager->IsAllDataLayerInEffectiveRuntimeState(GetDataLayers(), EDataLayerRuntimeState::Activated))
+					{
+						EffectiveWantedState = EDataLayerRuntimeState::Activated;
+					}
+					else if (DataLayerManager->IsAllDataLayerInEffectiveRuntimeState(GetDataLayers(), EDataLayerRuntimeState::Loaded))
+					{
+						EffectiveWantedState = EDataLayerRuntimeState::Loaded;
+					}
+					break;
+				default:
+					checkNoEntry();
+				}
 			}
-			break;
-		default:
-			checkNoEntry();
 		}
+
+		EffectiveWantedStateEpoch = AWorldDataLayers::GetDataLayersStateEpoch();
 	}
 
-	return EDataLayerRuntimeState::Unloaded;
+	return EffectiveWantedState;
 }
 
 TArray<const UDataLayerInstance*> UWorldPartitionRuntimeCell::GetDataLayerInstances() const
diff --git a/Engine/Source/Runtime/Engine/Public/WorldPartition/WorldPartitionRuntimeCell.h b/Engine/Source/Runtime/Engine/Public/WorldPartition/WorldPartitionRuntimeCell.h
--- a/Engine/Source/Runtime/Engine/Public/WorldPartition/WorldPartitionRuntimeCell.h
+++ b/Engine/Source/Runtime/Engine/Public/WorldPartition/WorldPartitionRuntimeCell.h
@@ -313,6 +313,9 @@ protected:
 	FName LevelPackageName;
 #endif
 
+	mutable EDataLayerRuntimeState EffectiveWantedState;
+	mutable int32 EffectiveWantedStateEpoch;
+
 public:
 	//~Begin UWorldPartitionRuntimeCellData Proxy
 	inline bool ShouldResetStreamingSourceInfo() const { return RuntimeCellData->ShouldResetStreamingSourceInfo(); }
