Add interface to add tags to BlueprintGeneratedClasses so BlueprintGeneratedClass assets can effectively replace BP assets.

diff --git a/Engine/Source/Editor/Blutility/Classes/EditorUtilityObject.h b/Engine/Source/Editor/Blutility/Classes/EditorUtilityObject.h
--- a/Engine/Source/Editor/Blutility/Classes/EditorUtilityObject.h
+++ b/Engine/Source/Editor/Blutility/Classes/EditorUtilityObject.h
@@ -9,16 +9,24 @@
 #include "UObject/Object.h"
 #include "UObject/ObjectMacros.h"
 #include "UObject/UObjectGlobals.h"
+#include "IAssetRegistryTagProviderInterface.h"
 
 #include "EditorUtilityObject.generated.h"
 
 
 UCLASS(Abstract, Blueprintable, meta = (ShowWorldContextPin))
-class BLUTILITY_API UEditorUtilityObject : public UObject
+class BLUTILITY_API UEditorUtilityObject : public UObject, public IAssetRegistryTagProviderInterface
 {
 	GENERATED_UCLASS_BODY()
 
 	// Standard function to execute
 	UFUNCTION(BlueprintCallable, BlueprintImplementableEvent, Category = "Editor")
 	void Run();
+
+	//~ Begin IAssetRegistryTagProviderInterface interface
+	virtual bool ShouldAddCDOTagsToBlueprintClass() const override 
+	{ 
+		return true; 
+	}
+	//~ End IAssetRegistryTagProviderInterface interface
 };
diff --git a/Engine/Source/Runtime/Engine/Private/BlueprintGeneratedClass.cpp b/Engine/Source/Runtime/Engine/Private/BlueprintGeneratedClass.cpp
--- a/Engine/Source/Runtime/Engine/Private/BlueprintGeneratedClass.cpp
+++ b/Engine/Source/Runtime/Engine/Private/BlueprintGeneratedClass.cpp
@@ -19,6 +19,7 @@
 #include "Engine/LevelScriptActor.h"
 #include "Engine/SCS_Node.h"
 #include "Engine/InheritableComponentHandler.h"
+#include "IAssetRegistryTagProviderInterface.h"
 #include "IFieldNotificationClassDescriptor.h"
 #include "INotifyFieldValueChanged.h"
 #include "Misc/ConfigCacheIni.h"
@@ -376,6 +377,17 @@ void UBlueprintGeneratedClass::GetAssetRegistryTags(TArray<FAssetRegistryTag>& O
 	{
 		OutTags.Add(FAssetRegistryTag(EditorTag.Key, EditorTag.Value, FAssetRegistryTag::TT_Hidden));
 	}
+
+	if (const UObject* CDO = GetDefaultObject())
+	{
+		if (const IAssetRegistryTagProviderInterface* AssetRegistryProvider = Cast<IAssetRegistryTagProviderInterface>(CDO))
+		{
+			if (AssetRegistryProvider->ShouldAddCDOTagsToBlueprintClass())
+			{
+				CDO->GetAssetRegistryTags(OutTags);
+			}
+		}
+	}
 #endif //#if WITH_EDITORONLY_DATA
 
 #if WITH_EDITOR
diff --git a/Engine/Source/Runtime/Engine/Public/IAssetRegistryTagProviderInterface.h b/Engine/Source/Runtime/Engine/Public/IAssetRegistryTagProviderInterface.h
new file mode 100644
--- /dev/null
+++ b/Engine/Source/Runtime/Engine/Public/IAssetRegistryTagProviderInterface.h
@@ -0,0 +1,26 @@
+// Copyright Epic Games, Inc. All Rights Reserved.
+
+#pragma once
+
+#include "UObject/Interface.h"
+
+#include "IAssetRegistryTagProviderInterface.generated.h"
+
+/**
+ * Interface to allow for CDO to add additional tags to generated class
+ */
+UINTERFACE(MinimalAPI, meta = (CannotImplementInterfaceInBlueprint))
+class UAssetRegistryTagProviderInterface : public UInterface
+{
+	GENERATED_BODY()
+};
+
+class IAssetRegistryTagProviderInterface
+{
+	GENERATED_BODY()
+
+public:
+
+	/** By default calls 'GetAssetRegistryTags' on the CDO, to add the BP's tags to the BPGC. */
+	ENGINE_API virtual bool ShouldAddCDOTagsToBlueprintClass() const { return false; }
+};
\ No newline at end of file
