#UE CIS quick fix

diff --git a/Engine/Source/Runtime/Engine/Private/ShaderCompiler/ShaderCompiler.cpp b/Engine/Source/Runtime/Engine/Private/ShaderCompiler/ShaderCompiler.cpp
--- a/Engine/Source/Runtime/Engine/Private/ShaderCompiler/ShaderCompiler.cpp
+++ b/Engine/Source/Runtime/Engine/Private/ShaderCompiler/ShaderCompiler.cpp
@@ -1614,7 +1614,7 @@ void FShaderJobCache::SubmitJob(FShaderCommonCompileJob* Job)
 		const FShaderCommonCompileJob::FInputHash& InputHash = Job->GetInputHash();
 		const bool bCheckDDC = !(Job->bIsDefaultMaterial || Job->bIsGlobalShader);
 		JobLock.WriteLock();
-		Job->JobCacheRef = FindOrAdd(InputHash, Job->Priority bCheckDDC, Job->RequestOwner, ExistingOutput);
+		Job->JobCacheRef = FindOrAdd(InputHash, Job->Priority, bCheckDDC, Job->RequestOwner, ExistingOutput);
 		bJobCacheLocked = true;
 	}
 
