Only initialize RayTracingGeometry resource in FSkeletalMeshObjectGPUSkin if IsRayTracingAllowed().

diff --git a/Engine/Source/Runtime/Engine/Private/SkeletalRenderGPUSkin.cpp b/Engine/Source/Runtime/Engine/Private/SkeletalRenderGPUSkin.cpp
--- a/Engine/Source/Runtime/Engine/Private/SkeletalRenderGPUSkin.cpp
+++ b/Engine/Source/Runtime/Engine/Private/SkeletalRenderGPUSkin.cpp
@@ -279,7 +279,7 @@ void FSkeletalMeshObjectGPUSkin::InitResources(USkinnedMeshComponent* InMeshComp
 	}
 
 #if RHI_RAYTRACING
-	if (bSupportRayTracing)
+	if (IsRayTracingAllowed() && bSupportRayTracing)
 	{
 		BeginInitResource(&RayTracingGeometry, &UE::RenderCommandPipe::SkeletalMesh);
 	}
