[Audio Engine] Enable source effects running on >2 channel sources + Opt-in virtual function for what a source effect supports, defaults to 2 (current behaviour) (on behalf of Jay Steen)

diff --git a/Engine/Source/Runtime/AudioMixer/Private/AudioMixerSource.cpp b/Engine/Source/Runtime/AudioMixer/Private/AudioMixerSource.cpp
--- a/Engine/Source/Runtime/AudioMixer/Private/AudioMixerSource.cpp
+++ b/Engine/Source/Runtime/AudioMixer/Private/AudioMixerSource.cpp
@@ -647,7 +647,15 @@ namespace Audio
 			// Source manager needs to know if this is a vorbis source for rebuilding speaker maps
 			InitParams.bIsVorbis = bIsVorbis;
 
-			if (InitParams.NumInputChannels <= 2)
+			// Support stereo by default
+			// Check the min number of channels the source effect chain supports
+			// We don't want to instantiate the effect chain if it has an effect that doesn't support its channel count
+			// E.g. we shouldn't instantiate a chain on a quad source if there is an effect that only supports stereo
+			InitParams.SourceEffectChainMaxSupportedChannels = WaveInstance->SourceEffectChain ? 
+				WaveInstance->SourceEffectChain->GetSupportedChannelCount() :
+				USoundEffectSourcePreset::DefaultSupportedChannels;
+
+			if (InitParams.NumInputChannels <= InitParams.SourceEffectChainMaxSupportedChannels)
 			{
 				if (WaveInstance->SourceEffectChain)
 				{
diff --git a/Engine/Source/Runtime/AudioMixer/Private/AudioMixerSourceManager.cpp b/Engine/Source/Runtime/AudioMixer/Private/AudioMixerSourceManager.cpp
--- a/Engine/Source/Runtime/AudioMixer/Private/AudioMixerSourceManager.cpp
+++ b/Engine/Source/Runtime/AudioMixer/Private/AudioMixerSourceManager.cpp
@@ -1009,8 +1009,8 @@ namespace Audio
 			SourceInfo.bEnableBaseSubmix = InitParams.bEnableBaseSubmix;
 			SourceInfo.bEnableSubmixSends = InitParams.bEnableSubmixSends;
 
-			// Copy the source effect chain if the channel count is 1 or 2
-			if (InitParams.NumInputChannels <= 2)
+			// Copy the source effect chain if the channel count is less than or equal to the number of channels supported by the effect chain
+			if (InitParams.NumInputChannels <= InitParams.SourceEffectChainMaxSupportedChannels)
 			{
 				// If we're told to care about effect chain tails, then we're not allowed
 				// to stop playing until the effect chain tails are finished
diff --git a/Engine/Source/Runtime/AudioMixer/Private/AudioMixerSourceManager.h b/Engine/Source/Runtime/AudioMixer/Private/AudioMixerSourceManager.h
--- a/Engine/Source/Runtime/AudioMixer/Private/AudioMixerSourceManager.h
+++ b/Engine/Source/Runtime/AudioMixer/Private/AudioMixerSourceManager.h
@@ -142,6 +142,7 @@ namespace Audio
 		float SourceBusDuration = 0.0f;
 		uint32 SourceEffectChainId = INDEX_NONE;
 		TArray<FSourceEffectChainEntry> SourceEffectChain;
+		int32 SourceEffectChainMaxSupportedChannels = 0;
 		FMixerSourceVoice* SourceVoice = nullptr;
 		int32 NumInputChannels = 0;
 		int32 NumInputFrames = 0;
diff --git a/Engine/Source/Runtime/Engine/Classes/Sound/SoundEffectSource.h b/Engine/Source/Runtime/Engine/Classes/Sound/SoundEffectSource.h
--- a/Engine/Source/Runtime/Engine/Classes/Sound/SoundEffectSource.h
+++ b/Engine/Source/Runtime/Engine/Classes/Sound/SoundEffectSource.h
@@ -19,6 +19,11 @@ UCLASS(config = Engine, abstract, editinlinenew, BlueprintType, MinimalAPI)
 class USoundEffectSourcePreset : public USoundEffectPreset
 {
 	GENERATED_BODY()
+
+public:
+	static constexpr int32 DefaultSupportedChannels = 2;
+
+	virtual int32 GetMaxSupportedChannels() const { return DefaultSupportedChannels; }
 };
 
 USTRUCT(BlueprintType)
@@ -56,6 +61,11 @@ public:
 
 	ENGINE_API void AddReferencedEffects(FReferenceCollector& Collector);
 
+	/**  Get the number of channels this chain supports, which is the lowest channel count of all effects in the chain */
+	int32 ENGINE_API GetSupportedChannelCount() const;
+
+	bool ENGINE_API SupportsChannelCount(const int32 InNumChannels) const;
+
 protected:
 
 #if WITH_EDITOR
diff --git a/Engine/Source/Runtime/Engine/Private/SoundEffectPreset.cpp b/Engine/Source/Runtime/Engine/Private/SoundEffectPreset.cpp
--- a/Engine/Source/Runtime/Engine/Private/SoundEffectPreset.cpp
+++ b/Engine/Source/Runtime/Engine/Private/SoundEffectPreset.cpp
@@ -108,6 +108,30 @@ void USoundEffectSourcePresetChain::AddReferencedEffects(FReferenceCollector& Co
 	}
 }
 
+int32 USoundEffectSourcePresetChain::GetSupportedChannelCount() const
+{
+	// Check the min number of channels the source effect chain supports
+	// We don't want to instantiate the effect chain if it has an effect that doesn't support its channel count
+	// E.g. we shouldn't instantiate a chain on a quad source if there is an effect in the chain that only supports stereo
+	constexpr int32 InvalidChannelCount = AUDIO_MIXER_MAX_OUTPUT_CHANNELS + 1;
+	int32 EffectChainSupportedChannels = InvalidChannelCount;
+
+	for (const FSourceEffectChainEntry& SourceEffectChainEntry : Chain)
+	{
+		if (SourceEffectChainEntry.Preset)
+		{
+			EffectChainSupportedChannels = FMath::Min(SourceEffectChainEntry.Preset->GetMaxSupportedChannels(), EffectChainSupportedChannels);
+		}
+	}
+
+	return (EffectChainSupportedChannels < InvalidChannelCount) ? EffectChainSupportedChannels : USoundEffectSourcePreset::DefaultSupportedChannels;
+}
+
+bool USoundEffectSourcePresetChain::SupportsChannelCount(const int32 InNumChannels) const
+{
+	return InNumChannels <= GetSupportedChannelCount();
+}
+
 void USoundEffectPreset::UnregisterInstance(TSoundEffectPtr InEffectPtr)
 {
 	if (ensure(IsInAudioThread() || IsInGameThread()))
diff --git a/Engine/Source/Runtime/Engine/Private/SoundWave.cpp b/Engine/Source/Runtime/Engine/Private/SoundWave.cpp
--- a/Engine/Source/Runtime/Engine/Private/SoundWave.cpp
+++ b/Engine/Source/Runtime/Engine/Private/SoundWave.cpp
@@ -3132,9 +3132,21 @@ void USoundWave::Parse(FAudioDevice* AudioDevice, const UPTRINT NodeWaveInstance
 	WaveInstance->ListenerToSoundDistanceForPanning = ParseParams.ListenerToSoundDistanceForPanning;
 	WaveInstance->AbsoluteAzimuth = ParseParams.AbsoluteAzimuth;
 
-	if (NumChannels <= 2)
+	if (ParseParams.SourceEffectChain)
 	{
-		WaveInstance->SourceEffectChain = ParseParams.SourceEffectChain;
+		// Only copy the source effect chain pointer if it supports the channel count
+		if (ParseParams.SourceEffectChain->SupportsChannelCount(NumChannels))
+		{
+			WaveInstance->SourceEffectChain = ParseParams.SourceEffectChain;
+		}
+		else
+		{
+			WaveInstance->SourceEffectChain = nullptr;
+			UE_LOG(LogAudio, Verbose, TEXT("Sound %s from Audio Component %s has a source effect chain but an incompatible channel count %d, so will not be instantiated"), 
+				ActiveSound.GetSound() ? *ActiveSound.GetSound()->GetFName().ToString() : TEXT("None"),
+				*ActiveSound.GetAudioComponentName(),
+				NumChannels);
+		}
 	}
 
 	bool bAlwaysPlay = false;
