Transformed MaterialTemplate.ush to use named parameters instead of "%s" positional parameters.

diff --git a/Engine/Shaders/Private/MaterialTemplate.ush b/Engine/Shaders/Private/MaterialTemplate.ush
--- a/Engine/Shaders/Private/MaterialTemplate.ush
+++ b/Engine/Shaders/Private/MaterialTemplate.ush
@@ -2,6 +2,10 @@
 
 /**
  * MaterialTemplate.usf: Filled in by FHLSLMaterialTranslator::GetMaterialShaderCode for each material being compiled.
+ 
+  $TemplateVersion{1}
+	Bump this number up when you make a change to this file that requires the material
+	translation DDC cache to be invalidated, forcing all materials to re-translate.
  */
 
 #ifndef ENABLE_NEW_HLSL_GENERATOR
@@ -75,13 +79,13 @@
 
 //////////////////////////////////////////////////////////////////////////
 
-#define NUM_MATERIAL_TEXCOORDS_VERTEX %s
-#define NUM_MATERIAL_TEXCOORDS %s
-#define NUM_CUSTOM_VERTEX_INTERPOLATORS %s
-#define NUM_TEX_COORD_INTERPOLATORS %s
+#define NUM_MATERIAL_TEXCOORDS_VERTEX %{num_material_texcoords_vertex}
+#define NUM_MATERIAL_TEXCOORDS %{num_material_texcoords}
+#define NUM_CUSTOM_VERTEX_INTERPOLATORS %{num_custom_vertex_interpolators}
+#define NUM_TEX_COORD_INTERPOLATORS %{num_tex_coord_interpolators}
 
 // Vertex interpolators offsets definition
-%s
+%{vertex_interpolators_offsets_definition}
 
 #if NUM_VIRTUALTEXTURE_SAMPLES || LIGHTMAP_VT_ENABLED
 	#include "/Engine/Private/VirtualTextureCommon.ush"
@@ -293,17 +297,17 @@ float4 GetDynamicParameter(FMaterialParticleParameters Parameters, float4 Defaul
 }
 
 /** Material declarations */
-%s
+%{material_declarations}
 
 /** FMaterialAttributes utilities */
-%s
+%{material_attributes_utilities}
 
 /** 
  * Parameters calculated from the pixel material inputs.
  */
 struct FPixelMaterialInputs
 {
-%s
+%{pixel_material_inputs}
 
 #if TEMPLATE_USES_SUBSTRATE
 FSubstrateData GetFrontSubstrateData()
@@ -2721,7 +2725,7 @@ float4 MaterialExpressionPathTracingBufferTextureLookup(FMaterialPixelParameters
 #define SwizzleDeriv4(_V, _MASK) ConstructFloatDeriv4(_V.Value._MASK, _V.Ddx._MASK, _V.Ddy._MASK)
 
 // Uniform material expressions.
-%s
+%{uniform_material_expressions}
 
 #if USE_ANALYTIC_DERIVATIVES && TEXTURE_SAMPLE_DEBUG
 MaterialFloat4 DebugTextureCommon(const int Mode, float2 UV, MaterialFloat2 DDX, MaterialFloat2 DDY, MaterialFloat Scale)
@@ -2844,7 +2848,7 @@ half3 GetMaterialEmissive(FPixelMaterialInputs PixelMaterialInputs)
 
 half3 GetMaterialEmissiveForCS(FMaterialPixelParameters Parameters)
 {
-%s;
+%{get_material_emissive_for_cs};
 }
 
 // Shading Model is an uint and represents a SHADINGMODELID_* in ShadingCommon.ush 
@@ -2917,55 +2921,55 @@ half GetMaterialAnisotropy(FPixelMaterialInputs PixelMaterialInputs)
 
 half GetMaterialTranslucencyDirectionalLightingIntensity()
 {
-%s;
+%{get_material_translucency_directional_lighting_intensity};
 }
 
 half GetMaterialTranslucentShadowDensityScale()
 {
-%s;
+%{get_material_translucent_shadow_density_scale};
 }
 
 half GetMaterialTranslucentSelfShadowDensityScale()
 {
-%s;
+%{get_material_translucent_self_shadow_density_scale};
 }
 
 half GetMaterialTranslucentSelfShadowSecondDensityScale()
 {
-%s;
+%{get_material_translucent_self_shadow_second_density_scale};
 }
 
 half GetMaterialTranslucentSelfShadowSecondOpacity()
 {
-%s;
+%{get_material_translucent_self_shadow_second_opacity};
 }
 
 half GetMaterialTranslucentBackscatteringExponent()
 {
-%s;
+%{get_material_translucent_backscattering_exponent};
 }
 
 half3 GetMaterialTranslucentMultipleScatteringExtinction()
 {
-%s;
+%{get_material_translucent_multiple_scattering_extinction};
 }
 
 // This is the clip value constant that is defined in the material (range 0..1)
 // Use GetMaterialMask() to get the Material Mask combined with this.
 half GetMaterialOpacityMaskClipValue()
 {
-%s;
+%{get_material_opacity_mask_clip_value};
 }
 
 #if USES_DISPLACEMENT && IS_NANITE_PASS
 float GetMaterialDisplacementMagnitude()
 {
-%s;
+%{get_material_displacement_magnitude};
 }
 
 float GetMaterialDisplacementCenter()
 {
-%s;
+%{get_material_displacement_center};
 }
 
 float GetMaterialDisplacement(FPixelMaterialInputs PixelMaterialInputs)
@@ -3047,7 +3051,7 @@ float3 ClampWorldPositionOffset(FMaterialVertexParameters Parameters, float3 InO
 
 float3 GetMaterialWorldPositionOffsetRaw(FMaterialVertexParameters Parameters)
 {
-%s;
+%{get_material_world_position_offset_raw};
 }
 
 float3 GetMaterialWorldPositionOffset(FMaterialVertexParameters Parameters)
@@ -3062,7 +3066,7 @@ float3 GetMaterialWorldPositionOffset(FMaterialVertexParameters Parameters)
 
 float3 GetMaterialPreviousWorldPositionOffsetRaw(FMaterialVertexParameters Parameters)
 {
-%s;
+%{get_material_previous_world_position_offset_raw};
 }
 
 float3 GetMaterialPreviousWorldPositionOffset(FMaterialVertexParameters Parameters)
@@ -3090,12 +3094,12 @@ half4 GetMaterialSubsurfaceData(FPixelMaterialInputs PixelMaterialInputs)
 
 half GetMaterialCustomData0(in out FMaterialPixelParameters Parameters)
 {
-%s;
+%{get_material_custom_data0};
 }
 
 half GetMaterialCustomData1(in out FMaterialPixelParameters Parameters)
 {
-%s;
+%{get_material_custom_data1};
 }
 
 half GetMaterialAmbientOcclusionRaw(FPixelMaterialInputs PixelMaterialInputs)
@@ -3152,12 +3156,12 @@ FMaterialRefractionData GetMaterialRefraction(FPixelMaterialInputs PixelMaterial
 #if NUM_TEX_COORD_INTERPOLATORS
 void GetMaterialCustomizedUVs(FMaterialVertexParameters Parameters, inout float2 OutTexCoords[NUM_TEX_COORD_INTERPOLATORS])
 {
-%s
+%{get_material_customized_u_vs}
 }
 
 void GetCustomInterpolators(FMaterialVertexParameters Parameters, inout float2 OutTexCoords[NUM_TEX_COORD_INTERPOLATORS])
 {
-%s
+%{get_custom_interpolators}
 }
 #endif
 
@@ -3239,7 +3243,7 @@ float GetRoughnessFromNormalCurvature(FMaterialPixelParameters InMaterialParamet
 }
 
 // EvaluateMaterialAttributes
-%s
+%{evaluate_material_attributes}
 
 // Be very, very careful about changing CalcPixelMaterialInputs() or CalcPixelMaterialInputsAnalyticDerivatives(). They apply the same basic calculation,
 // but CalcPixelMaterialInputs() uses finite difference derivates from hardware whereas CalcPixelMaterialInputsAnalyticDerivatives() uses analytic
@@ -3247,9 +3251,9 @@ float GetRoughnessFromNormalCurvature(FMaterialPixelParameters InMaterialParamet
 void CalcPixelMaterialInputs(in out FMaterialPixelParameters Parameters, in out FPixelMaterialInputs PixelMaterialInputs)
 {
 	// Initial calculations (required for Normal)
-%s
+%{calc_pixel_material_inputs_initial_calculations}
 	// The Normal is a special case as it might have its own expressions and also be used to calculate other inputs, so perform the assignment here
-%s
+%{calc_pixel_material_inputs_normal}
 
 #if TEMPLATE_USES_SUBSTRATE
 	Parameters.SubstratePixelFootprint = SubstrateGetPixelFootprint(Parameters.WorldPosition_CamRelative, GetRoughnessFromNormalCurvature(Parameters));
@@ -3293,7 +3297,7 @@ void CalcPixelMaterialInputs(in out FMaterialPixelParameters Parameters, in out
 #endif // !PARTICLE_SPRITE_FACTORY
 
 	// Now the rest of the inputs
-%s
+%{calc_pixel_material_inputs_other_inputs}
 
 #if MATERIAL_USES_ANISOTROPY
 	Parameters.WorldTangent = CalculateAnisotropyTangent(Parameters, PixelMaterialInputs);
@@ -3309,9 +3313,9 @@ void CalcPixelMaterialInputs(in out FMaterialPixelParameters Parameters, in out
 void CalcPixelMaterialInputsAnalyticDerivatives(in out FMaterialPixelParameters Parameters, in out FPixelMaterialInputs PixelMaterialInputs)
 {
 	// Initial calculations (required for Normal)
-%s
+%{calc_pixel_material_inputs_analytic_derivatives_initial}
 	// The Normal is a special case as it might have its own expressions and also be used to calculate other inputs, so perform the assignment here
-%s
+%{calc_pixel_material_inputs_analytic_derivatives_normal}
 
 #if TEMPLATE_USES_SUBSTRATE
 	Parameters.SubstratePixelFootprint = SubstrateGetPixelFootprint(Parameters.WorldPosition_CamRelative, GetRoughnessFromNormalCurvature(Parameters));
@@ -3355,7 +3359,7 @@ void CalcPixelMaterialInputsAnalyticDerivatives(in out FMaterialPixelParameters
 #endif // !PARTICLE_SPRITE_FACTORY
 
 	// Now the rest of the inputs
-%s
+%{calc_pixel_material_inputs_analytic_derivatives_other_inputs}
 
 #if MATERIAL_USES_ANISOTROPY
 	Parameters.WorldTangent = CalculateAnisotropyTangent(Parameters, PixelMaterialInputs);
@@ -3368,7 +3372,7 @@ void CalcPixelMaterialInputsAnalyticDerivatives(in out FMaterialPixelParameters
 
 // Programmatically set the line number after all the material inputs which have a variable number of line endings
 // This allows shader error line numbers after this point to be the same regardless of which material is being compiled
-#line %s
+#line %{line_number}
 
 void ClipLODTransition(float2 SvPosition, float DitherFactor)
 {
diff --git a/Engine/Source/Runtime/Engine/Private/Materials/HLSLMaterialTranslator.cpp b/Engine/Source/Runtime/Engine/Private/Materials/HLSLMaterialTranslator.cpp
--- a/Engine/Source/Runtime/Engine/Private/Materials/HLSLMaterialTranslator.cpp
+++ b/Engine/Source/Runtime/Engine/Private/Materials/HLSLMaterialTranslator.cpp
@@ -36,7 +36,9 @@
 #include "Materials/MaterialExpressionViewProperty.h"
 #include "Materials/MaterialExpressionVolumetricAdvancedMaterialOutput.h"
 #include "Materials/MaterialExpressionWorldPosition.h"
+#include "Materials/MaterialSourceTemplate.h"
 #include "Materials/SubstrateMaterial.h"
+#include "StringTemplate.h"
 #include "ParameterCollection.h"
 #include "RenderUtils.h"
 #include "Stats/StatsMisc.h"
@@ -333,7 +335,6 @@ FHLSLMaterialTranslator::FHLSLMaterialTranslator(FMaterial* InMaterial,
 ,	Platform(InPlatform)
 ,	QualityLevel(InQualityLevel)
 ,	FeatureLevel(InFeatureLevel)
-,	MaterialTemplateLineNumber(INDEX_NONE)
 ,	NextSymbolIndex(INDEX_NONE)
 ,	NextVertexInterpolatorIndex(0)
 ,	CurrentCustomVertexInterpolatorOffset(0)
@@ -1907,33 +1908,8 @@ bool FHLSLMaterialTranslator::Translate()
 			}
 		}
 
-		LoadShaderSourceFileChecked(TEXT("/Engine/Private/MaterialTemplate.ush"), GetShaderPlatform(), MaterialTemplate);
-
-		// Find the string index of the '#line' statement in MaterialTemplate.usf
-		const int32 LineIndex = MaterialTemplate.Find(TEXT("#line"), ESearchCase::CaseSensitive);
-		check(LineIndex != INDEX_NONE);
-
-		// Count line endings before the '#line' statement
-		MaterialTemplateLineNumber = INDEX_NONE;
-		int32 StartPosition = LineIndex + 1;
-		do 
-		{
-			MaterialTemplateLineNumber++;
-			// Using \n instead of LINE_TERMINATOR as not all of the lines are terminated consistently
-			// Subtract one from the last found line ending index to make sure we skip over it
-			StartPosition = MaterialTemplate.Find(TEXT("\n"), ESearchCase::CaseSensitive, ESearchDir::FromEnd, StartPosition - 1);
-		} 
-		while (StartPosition != INDEX_NONE);
-		check(MaterialTemplateLineNumber != INDEX_NONE);
-		// At this point MaterialTemplateLineNumber is one less than the line number of the '#line' statement
-		// For some reason we have to add 2 more to the #line value to get correct error line numbers from D3DXCompileShader
-		MaterialTemplateLineNumber += 3;
-
 		MaterialCompilationOutput.UniformExpressionSet.SetParameterCollections(ParameterCollections);
 
-		// This will be created shortly after the Translate call in FMaterial::BeginCompileShaderMap()
-		//MaterialCompilationOutput.UniformExpressionSet.CreateBufferStruct();
-
 		// Store the number of unique VT samples
 		MaterialCompilationOutput.EstimatedNumVirtualTextureLookups = NumVtSamples;
 	}
@@ -2591,7 +2567,14 @@ void FHLSLMaterialTranslator::GetSharedInputsMaterialCode(FString& PixelMembersD
 FString FHLSLMaterialTranslator::GetMaterialShaderCode()
 {	
 	// use "/Engine/Private/MaterialTemplate.ush" to create the functions to get data (e.g. material attributes) and code (e.g. material expressions to create specular color) from C++
-	FLazyPrintf LazyPrintf(*MaterialTemplate);
+	TMap<FString, FString> MaterialSourceTemplateParams;
+	int32 MaterialTemplateLineNumber;
+
+	FStringTemplateResolver Resolver = FMaterialSourceTemplate::Get().BeginResolve(GetShaderPlatform(), &MaterialTemplateLineNumber);
+	Resolver.SetParameterMap(&MaterialSourceTemplateParams);
+
+	MaterialSourceTemplateParams.Reserve(Resolver.GetTemplate().GetNumNamedParameters());
+	MaterialSourceTemplateParams.Add({ TEXT("line_number"), FString::Printf(TEXT("%u"), MaterialTemplateLineNumber) });
 
 	// Assign slots to vertex interpolators
 	FString VertexInterpolatorsOffsetsDefinition;
@@ -2602,19 +2585,21 @@ FString FHLSLMaterialTranslator::GetMaterialShaderCode()
 	const uint32 NumCustomVectors = FMath::DivideAndRoundUp((uint32)CurrentCustomVertexInterpolatorOffset, 2u);
 	const uint32 NumTexCoordVectors = FinalAllocatedCoords.FindLast(true) + 1;
 
-	LazyPrintf.PushParam(*FString::Printf(TEXT("%u"),NumUserVertexTexCoords));
-	LazyPrintf.PushParam(*FString::Printf(TEXT("%u"),NumUserTexCoords));
-	LazyPrintf.PushParam(*FString::Printf(TEXT("%u"),NumCustomVectors));
-	LazyPrintf.PushParam(*FString::Printf(TEXT("%u"),NumTexCoordVectors));
-
-	LazyPrintf.PushParam(*VertexInterpolatorsOffsetsDefinition);
+	MaterialSourceTemplateParams.Add({ TEXT("num_material_texcoords_vertex"), FString::Printf(TEXT("%u"), NumUserVertexTexCoords) });
+	MaterialSourceTemplateParams.Add({ TEXT("num_material_texcoords"), FString::Printf(TEXT("%u"), NumUserTexCoords) });
+	MaterialSourceTemplateParams.Add({ TEXT("num_custom_vertex_interpolators"), FString::Printf(TEXT("%u"), NumCustomVectors) });
+	MaterialSourceTemplateParams.Add({ TEXT("num_tex_coord_interpolators"), FString::Printf(TEXT("%u"), NumTexCoordVectors) });
+	MaterialSourceTemplateParams.Add({ TEXT("vertex_interpolators_offsets_definition"), VertexInterpolatorsOffsetsDefinition });
 
 	FString MaterialAttributesDeclaration;
 	FString MaterialAttributesUtilities;
 
-	const EMaterialShadingModel DefaultShadingModel = Material->GetShadingModels().GetFirstShadingModel();
+	// Reserve some bytes for fixed string content (e.g. "struct FMaterialAttributes { ... }")  and enough space for each declaration entry, using an estimation of the maximum length.
+	MaterialAttributesDeclaration.Reserve(30 + MaterialAttributesDeclaration.Len() * 128);
+	MaterialAttributesDeclaration.Append(TEXT("struct FMaterialAttributes\n{\n"));
 
-	MaterialAttributesDeclaration += TEXT("struct FMaterialAttributes\n{\n");
+	// Reserve enough space for each utility entry, using an estimation of the maximum length.
+	MaterialAttributesUtilities.Reserve(MaterialAttributesDeclaration.Len() * 256);
 
 	const TArray<FGuid>& OrderedVisibleAttributes = FMaterialAttributeDefinitionMap::GetOrderedVisibleAttributeList();
 	for (const FGuid& AttributeID : OrderedVisibleAttributes)
@@ -2641,18 +2626,18 @@ FString FHLSLMaterialTranslator::GetMaterialShaderCode()
 		{
 			const FVector4f DefaultValue = FMaterialAttributeDefinitionMap::GetDefaultValue(AttributeID);
 
-			MaterialAttributesDeclaration += FString::Printf(TEXT("\t%s %s;") LINE_TERMINATOR, HLSLType, *PropertyName);
+			MaterialAttributesDeclaration.Appendf(TEXT("\t%s %s;") LINE_TERMINATOR, HLSLType, *PropertyName);
 
 			// Chainable method to set the attribute
-			MaterialAttributesUtilities += FString::Printf(TEXT("FMaterialAttributes FMaterialAttributes_Set%s(FMaterialAttributes InAttributes, %s InValue) { InAttributes.%s = InValue; return InAttributes; }") LINE_TERMINATOR,
+			MaterialAttributesUtilities.Appendf(TEXT("FMaterialAttributes FMaterialAttributes_Set%s(FMaterialAttributes InAttributes, %s InValue) { InAttributes.%s = InValue; return InAttributes; }") LINE_TERMINATOR,
 				*PropertyName, HLSLType, *PropertyName);
 		}
 	}
 
-	MaterialAttributesDeclaration += TEXT("};\n");
+	MaterialAttributesDeclaration.Append(TEXT("};\n"));
 
-	LazyPrintf.PushParam(*MaterialAttributesDeclaration);
-	LazyPrintf.PushParam(*MaterialAttributesUtilities);
+	MaterialSourceTemplateParams.Add({ TEXT("material_declarations"), MaterialAttributesDeclaration });
+	MaterialSourceTemplateParams.Add({ TEXT("material_attributes_utilities"), MaterialAttributesUtilities });
 
 	// Stores the shared shader results member declarations
 	FString PixelMembersDeclaration[CompiledPDV_MAX];
@@ -2670,13 +2655,12 @@ FString FHLSLMaterialTranslator::GetMaterialShaderCode()
 	// PixelMembersDeclaration should be the same for all variations, but might change in the future. There are cases where work is shared
 	// between the pixel and vertex shader, but with Nanite all work has to be moved into the pixel shader, which means we will want
 	// different inputs. But for now, we are keeping them the same.
-	LazyPrintf.PushParam(*PixelMembersDeclaration[CompiledPDV_FiniteDifferences]);
+	MaterialSourceTemplateParams.Add({ TEXT("pixel_material_inputs"), PixelMembersDeclaration[CompiledPDV_FiniteDifferences] });
 
 	{
 		FString DerivativeHelpers = DerivativeAutogen.GenerateUsedFunctions(*this);
 		FString DerivativeHelpersAndResources = DerivativeHelpers + ResourcesString;
-		//LazyPrintf.PushParam(*ResourcesString);
-		LazyPrintf.PushParam(*DerivativeHelpersAndResources);
+		MaterialSourceTemplateParams.Add({ TEXT("uniform_material_expressions"), DerivativeHelpersAndResources });
 	}
 
 	// Anything used bye the GenerationFunctionCode() like WorldPositionOffset shouldn't be using texures, right?
@@ -2686,39 +2670,38 @@ FString FHLSLMaterialTranslator::GetMaterialShaderCode()
 
 	if (bCompileForComputeShader)
 	{
-		LazyPrintf.PushParam(*GenerateFunctionCode(CompiledMP_EmissiveColorCS, BaseDerivativeVariation));
+		MaterialSourceTemplateParams.Add({ TEXT("get_material_emissive_for_cs"), GenerateFunctionCode(CompiledMP_EmissiveColorCS, BaseDerivativeVariation) });
 	}
 	else
 	{
-		LazyPrintf.PushParam(TEXT("return 0"));
+		MaterialSourceTemplateParams.Add({ TEXT("get_material_emissive_for_cs"), TEXT("return 0") });
 	}
 
-	LazyPrintf.PushParam(*FString::Printf(TEXT("return %.5f"), Material->GetTranslucencyDirectionalLightingIntensity()));
-
-	LazyPrintf.PushParam(*FString::Printf(TEXT("return %.5f"), Material->GetTranslucentShadowDensityScale()));
-	LazyPrintf.PushParam(*FString::Printf(TEXT("return %.5f"), Material->GetTranslucentSelfShadowDensityScale()));
-	LazyPrintf.PushParam(*FString::Printf(TEXT("return %.5f"), Material->GetTranslucentSelfShadowSecondDensityScale()));
-	LazyPrintf.PushParam(*FString::Printf(TEXT("return %.5f"), Material->GetTranslucentSelfShadowSecondOpacity()));
-	LazyPrintf.PushParam(*FString::Printf(TEXT("return %.5f"), Material->GetTranslucentBackscatteringExponent()));
+	MaterialSourceTemplateParams.Add({ TEXT("get_material_translucency_directional_lighting_intensity"), FString::Printf(TEXT("return %.5f"), Material->GetTranslucencyDirectionalLightingIntensity()) });
+	MaterialSourceTemplateParams.Add({ TEXT("get_material_translucent_shadow_density_scale"), FString::Printf(TEXT("return %.5f"), Material->GetTranslucentShadowDensityScale()) });
+	MaterialSourceTemplateParams.Add({ TEXT("get_material_translucent_self_shadow_density_scale"), FString::Printf(TEXT("return %.5f"), Material->GetTranslucentSelfShadowDensityScale()) });
+	MaterialSourceTemplateParams.Add({ TEXT("get_material_translucent_self_shadow_second_density_scale"), FString::Printf(TEXT("return %.5f"), Material->GetTranslucentSelfShadowSecondDensityScale()) });
+	MaterialSourceTemplateParams.Add({ TEXT("get_material_translucent_self_shadow_second_opacity"), FString::Printf(TEXT("return %.5f"), Material->GetTranslucentSelfShadowSecondOpacity()) });
+	MaterialSourceTemplateParams.Add({ TEXT("get_material_translucent_backscattering_exponent"), FString::Printf(TEXT("return %.5f"), Material->GetTranslucentBackscatteringExponent()) });
 
 	{
 		FLinearColor Extinction = Material->GetTranslucentMultipleScatteringExtinction();
 
-		LazyPrintf.PushParam(*FString::Printf(TEXT("return MaterialFloat3(%.5f, %.5f, %.5f)"), Extinction.R, Extinction.G, Extinction.B));
+		MaterialSourceTemplateParams.Add({ TEXT("get_material_translucent_multiple_scattering_extinction"), FString::Printf(TEXT("return MaterialFloat3(%.5f, %.5f, %.5f)"), Extinction.R, Extinction.G, Extinction.B) });
 	}
 
-	LazyPrintf.PushParam(*FString::Printf(TEXT("return %.5f"), Material->GetOpacityMaskClipValue()));
+	MaterialSourceTemplateParams.Add({ TEXT("get_material_opacity_mask_clip_value"), FString::Printf(TEXT("return %.5f"), Material->GetOpacityMaskClipValue()) });
 
 	{
 		const FDisplacementScaling DisplacementScaling = Material->GetDisplacementScaling();
-		LazyPrintf.PushParam(*FString::Printf(TEXT("return %.5f"), FMath::Max(0.0f, DisplacementScaling.Magnitude)));
-		LazyPrintf.PushParam(*FString::Printf(TEXT("return %.5f"), FMath::Clamp(DisplacementScaling.Center, 0.0f, 1.0f)));
+		MaterialSourceTemplateParams.Add({ TEXT("get_material_displacement_magnitude"), FString::Printf(TEXT("return %.5f"), FMath::Max(0.0f, DisplacementScaling.Magnitude)) });
+		MaterialSourceTemplateParams.Add({ TEXT("get_material_displacement_center"), FString::Printf(TEXT("return %.5f"), FMath::Clamp(DisplacementScaling.Center, 0.0f, 1.0f)) });
 	}
 
-	LazyPrintf.PushParam(!bEnableExecutionFlow ? *GenerateFunctionCode(MP_WorldPositionOffset, BaseDerivativeVariation) : TEXT("return Parameters.MaterialAttributes.WorldPositionOffset"));
-	LazyPrintf.PushParam(!bEnableExecutionFlow ? *GenerateFunctionCode(CompiledMP_PrevWorldPositionOffset, BaseDerivativeVariation) : TEXT("return 0.0f"));
-	LazyPrintf.PushParam(!bEnableExecutionFlow ? *GenerateFunctionCode(MP_CustomData0, BaseDerivativeVariation) : TEXT("return 0.0f"));
-	LazyPrintf.PushParam(!bEnableExecutionFlow ? *GenerateFunctionCode(MP_CustomData1, BaseDerivativeVariation) : TEXT("return 0.0f"));
+	MaterialSourceTemplateParams.Add({ TEXT("get_material_world_position_offset_raw"), !bEnableExecutionFlow ? *GenerateFunctionCode(MP_WorldPositionOffset, BaseDerivativeVariation) : TEXT("return Parameters.MaterialAttributes.WorldPositionOffset") });
+	MaterialSourceTemplateParams.Add({ TEXT("get_material_previous_world_position_offset_raw"), !bEnableExecutionFlow ? *GenerateFunctionCode(CompiledMP_PrevWorldPositionOffset, BaseDerivativeVariation) : TEXT("return 0.0f") });
+	MaterialSourceTemplateParams.Add({ TEXT("get_material_custom_data0"), !bEnableExecutionFlow ? *GenerateFunctionCode(MP_CustomData0, BaseDerivativeVariation) : TEXT("return 0.0f") });
+	MaterialSourceTemplateParams.Add({ TEXT("get_material_custom_data1"), !bEnableExecutionFlow ? *GenerateFunctionCode(MP_CustomData1, BaseDerivativeVariation) : TEXT("return 0.0f") });
 
 	// Print custom texture coordinate assignments, should be fine with regular derivatives
 	FString CustomUVAssignments;
@@ -2729,24 +2712,25 @@ FString FHLSLMaterialTranslator::GetMaterialShaderCode()
 		if (bEnableExecutionFlow)
 		{
 			const FString AttributeName = FMaterialAttributeDefinitionMap::GetAttributeName((EMaterialProperty)(MP_CustomizedUVs0 + CustomUVIndex));
-			CustomUVAssignments += FString::Printf(TEXT("\tOutTexCoords[%u] = Parameters.MaterialAttributes.%s;") LINE_TERMINATOR, CustomUVIndex, *AttributeName);
+			CustomUVAssignments.Appendf(TEXT("\tOutTexCoords[%u] = Parameters.MaterialAttributes.%s;") LINE_TERMINATOR, CustomUVIndex, *AttributeName);
 		}
 		else
 		{
-		if (CustomUVIndex == 0)
-		{
-				CustomUVAssignments += DerivativeVariations[BaseDerivativeVariation].TranslatedCodeChunkDefinitions[MP_CustomizedUVs0 + CustomUVIndex];
-		}
+			if (CustomUVIndex == 0)
+			{
+				CustomUVAssignments.Append(DerivativeVariations[BaseDerivativeVariation].TranslatedCodeChunkDefinitions[MP_CustomizedUVs0 + CustomUVIndex]);
+			}
 
 			if (DerivativeVariations[BaseDerivativeVariation].TranslatedCodeChunkDefinitions[MP_CustomizedUVs0 + CustomUVIndex].Len() > 0)
 			{
 				LastProperty = MP_CustomizedUVs0 + CustomUVIndex;
 			}
-			CustomUVAssignments += FString::Printf(TEXT("\tOutTexCoords[%u] = %s;") LINE_TERMINATOR, CustomUVIndex, *DerivativeVariations[BaseDerivativeVariation].TranslatedCodeChunks[MP_CustomizedUVs0 + CustomUVIndex]);
+
+			CustomUVAssignments.Appendf(TEXT("\tOutTexCoords[%u] = %s;") LINE_TERMINATOR, CustomUVIndex, *DerivativeVariations[BaseDerivativeVariation].TranslatedCodeChunks[MP_CustomizedUVs0 + CustomUVIndex]);
 		}
 	}
 
-	LazyPrintf.PushParam(*CustomUVAssignments);
+	MaterialSourceTemplateParams.Add({ TEXT("get_material_customized_u_vs"), CustomUVAssignments });
 
 	// Print custom vertex shader interpolator assignments
 	FString CustomInterpolatorAssignments;
@@ -2764,26 +2748,26 @@ FString FHLSLMaterialTranslator::GetMaterialShaderCode()
 			const int32 Index = Interpolator->InterpolatorIndex;
 
 			// Note: We reference the UV define directly to avoid having to pre-accumulate UV counts before property translation
-			CustomInterpolatorAssignments += FString::Printf(TEXT("\tOutTexCoords[VERTEX_INTERPOLATOR_%i_TEXCOORDS_X].%s = VertexInterpolator%i(Parameters).x;") LINE_TERMINATOR, Index, Swizzle[Offset%2], Index);
+			CustomInterpolatorAssignments.Appendf(TEXT("\tOutTexCoords[VERTEX_INTERPOLATOR_%i_TEXCOORDS_X].%s = VertexInterpolator%i(Parameters).x;") LINE_TERMINATOR, Index, Swizzle[Offset%2], Index);
 				
 			if (Type >= MCT_Float2)
 			{
-				CustomInterpolatorAssignments += FString::Printf(TEXT("\tOutTexCoords[VERTEX_INTERPOLATOR_%i_TEXCOORDS_Y].%s = VertexInterpolator%i(Parameters).y;") LINE_TERMINATOR, Index, Swizzle[(Offset+1)%2], Index);
+				CustomInterpolatorAssignments.Appendf(TEXT("\tOutTexCoords[VERTEX_INTERPOLATOR_%i_TEXCOORDS_Y].%s = VertexInterpolator%i(Parameters).y;") LINE_TERMINATOR, Index, Swizzle[(Offset+1)%2], Index);
 
 				if (Type >= MCT_Float3)
 				{
-					CustomInterpolatorAssignments += FString::Printf(TEXT("\tOutTexCoords[VERTEX_INTERPOLATOR_%i_TEXCOORDS_Z].%s = VertexInterpolator%i(Parameters).z;") LINE_TERMINATOR, Index, Swizzle[(Offset+2)%2], Index);
+					CustomInterpolatorAssignments.Appendf(TEXT("\tOutTexCoords[VERTEX_INTERPOLATOR_%i_TEXCOORDS_Z].%s = VertexInterpolator%i(Parameters).z;") LINE_TERMINATOR, Index, Swizzle[(Offset+2)%2], Index);
 
 					if (Type == MCT_Float4)
 					{
-						CustomInterpolatorAssignments += FString::Printf(TEXT("\tOutTexCoords[VERTEX_INTERPOLATOR_%i_TEXCOORDS_W].%s = VertexInterpolator%i(Parameters).w;") LINE_TERMINATOR, Index, Swizzle[(Offset+3)%2], Index);
+						CustomInterpolatorAssignments.Appendf(TEXT("\tOutTexCoords[VERTEX_INTERPOLATOR_%i_TEXCOORDS_W].%s = VertexInterpolator%i(Parameters).w;") LINE_TERMINATOR, Index, Swizzle[(Offset+3)%2], Index);
 					}
 				}
 			}
 		}
 	}
 
-	LazyPrintf.PushParam(*CustomInterpolatorAssignments);
+	MaterialSourceTemplateParams.Add({ TEXT("get_custom_interpolators"), CustomInterpolatorAssignments });
 
 	if (bEnableExecutionFlow)
 	{
@@ -2794,12 +2778,12 @@ FString FHLSLMaterialTranslator::GetMaterialShaderCode()
 		{
 			const FString AttributeName = FMaterialAttributeDefinitionMap::GetAttributeName((EMaterialProperty)(MP_CustomizedUVs0 + TexCoordIndex));
 
-			EvaluateVertexCode += FString::Printf(TEXT("\tDefaultMaterialAttributes.%s = Parameters.TexCoords[%d];") LINE_TERMINATOR, *AttributeName, TexCoordIndex);
+			EvaluateVertexCode.Appendf(TEXT("\tDefaultMaterialAttributes.%s = Parameters.TexCoords[%d];") LINE_TERMINATOR, *AttributeName, TexCoordIndex);
 		}
 
-		EvaluateVertexCode += TranslatedAttributesCodeChunks[SF_Vertex];
+		EvaluateVertexCode.Append(TranslatedAttributesCodeChunks[SF_Vertex]);
 
-		LazyPrintf.PushParam(*TranslatedAttributesCodeChunks[SF_Pixel]);
+		MaterialSourceTemplateParams.Add({ TEXT("evaluate_material_attributes"), TranslatedAttributesCodeChunks[SF_Pixel] });
 
 		FString EvaluateMaterialAttributesCode = TEXT("    FMaterialAttributes MaterialAttributes = EvaluatePixelMaterialAttributes(Parameters);" LINE_TERMINATOR);
 
@@ -2819,42 +2803,44 @@ FString FHLSLMaterialTranslator::GetMaterialShaderCode()
 			if (PropertyIndex == MP_SubsurfaceColor)
 			{
 				// TODO - properly handle subsurface profile
-				EvaluateMaterialAttributesCode += FString::Printf("    PixelMaterialInputs.Subsurface = float4(MaterialAttributes.%s, 0.0f);" LINE_TERMINATOR, *PropertyName);
+				EvaluateMaterialAttributesCode.Appendf("    PixelMaterialInputs.Subsurface = float4(MaterialAttributes.%s, 0.0f);" LINE_TERMINATOR, *PropertyName);
 			}
 			else
 			{
-				EvaluateMaterialAttributesCode += FString::Printf("    PixelMaterialInputs.%s = MaterialAttributes.%s;" LINE_TERMINATOR, *PropertyName, *PropertyName);
+				EvaluateMaterialAttributesCode.Appendf("    PixelMaterialInputs.%s = MaterialAttributes.%s;" LINE_TERMINATOR, *PropertyName, *PropertyName);
 			}
 		}
 
-		// TODO - deriv
-		for (int32 Iter = 0; Iter < CompiledPDV_MAX; Iter++)
-		{
-			LazyPrintf.PushParam(*EvaluateMaterialAttributesCode);
-			LazyPrintf.PushParam(TEXT(""));
-			LazyPrintf.PushParam(TEXT(""));
-		}
+		MaterialSourceTemplateParams.Add({ TEXT("calc_pixel_material_inputs_initial_calculations"), EvaluateMaterialAttributesCode });
+		MaterialSourceTemplateParams.Add({ TEXT("calc_pixel_material_inputs_normal"), TEXT("") });
+		MaterialSourceTemplateParams.Add({ TEXT("calc_pixel_material_inputs_other_inputs"), TEXT("") });
+
+		MaterialSourceTemplateParams.Add({ TEXT("calc_pixel_material_inputs_analytic_derivatives_initial"), EvaluateMaterialAttributesCode });
+		MaterialSourceTemplateParams.Add({ TEXT("calc_pixel_material_inputs_analytic_derivatives_normal"), TEXT("") });
+		MaterialSourceTemplateParams.Add({ TEXT("calc_pixel_material_inputs_analytic_derivatives_other_inputs"), TEXT("") });
 	}
 	else
 	{
 		// skip material attributes code
-		LazyPrintf.PushParam(TEXT(""));
+		MaterialSourceTemplateParams.Add({ TEXT("evaluate_material_attributes"), TEXT("") });
 
-		for (int32 Iter = 0; Iter < CompiledPDV_MAX; Iter++)
 		{
-			ECompiledPartialDerivativeVariation Variation = (ECompiledPartialDerivativeVariation)Iter;
-
-	// Initializers required for Normal
-			LazyPrintf.PushParam(*DerivativeVariations[Variation].TranslatedCodeChunkDefinitions[MP_Normal]);
-			LazyPrintf.PushParam(*NormalAssignment[Variation]);
-	// Finally the rest of common code followed by assignment into each input
-			LazyPrintf.PushParam(*PixelMembersSetupAndAssignments[Variation]);
+			// Initializers required for Normal
+			MaterialSourceTemplateParams.Add({ TEXT("calc_pixel_material_inputs_initial_calculations"), DerivativeVariations[CompiledPDV_FiniteDifferences].TranslatedCodeChunkDefinitions[MP_Normal] });
+			MaterialSourceTemplateParams.Add({ TEXT("calc_pixel_material_inputs_normal"), NormalAssignment[CompiledPDV_FiniteDifferences] });
+			// Finally the rest of common code followed by assignment into each input
+			MaterialSourceTemplateParams.Add({ TEXT("calc_pixel_material_inputs_other_inputs"), PixelMembersSetupAndAssignments[CompiledPDV_FiniteDifferences] });
+		}
+		{
+			// Initializers required for Normal
+			MaterialSourceTemplateParams.Add({ TEXT("calc_pixel_material_inputs_analytic_derivatives_initial"), DerivativeVariations[CompiledPDV_Analytic].TranslatedCodeChunkDefinitions[MP_Normal] });
+			MaterialSourceTemplateParams.Add({ TEXT("calc_pixel_material_inputs_analytic_derivatives_normal"), NormalAssignment[CompiledPDV_Analytic] });
+			// Finally the rest of common code followed by assignment into each input
+			MaterialSourceTemplateParams.Add({ TEXT("calc_pixel_material_inputs_analytic_derivatives_other_inputs"), PixelMembersSetupAndAssignments[CompiledPDV_Analytic] });
 		}
 	}
 
-	LazyPrintf.PushParam(*FString::Printf(TEXT("%u"),MaterialTemplateLineNumber));
-
-	return LazyPrintf.GetResultString();
+	return Resolver.Finalize();
 }
 
 
diff --git a/Engine/Source/Runtime/Engine/Private/Materials/HLSLMaterialTranslator.h b/Engine/Source/Runtime/Engine/Private/Materials/HLSLMaterialTranslator.h
--- a/Engine/Source/Runtime/Engine/Private/Materials/HLSLMaterialTranslator.h
+++ b/Engine/Source/Runtime/Engine/Private/Materials/HLSLMaterialTranslator.h
@@ -33,6 +33,7 @@
 #include "Containers/Map.h"
 #include "Shader/ShaderTypes.h"
 #include "SparseVolumeTexture/SparseVolumeTexture.h"
+#include "Runtime/RenderCore/Internal/ShaderCompilerDefinitions.h"
 
 #if WITH_EDITORONLY_DATA
 #include "Materials/MaterialExpressionSceneTexture.h"
@@ -278,12 +279,6 @@ protected:
 
 	uint64 MaterialAttributesReturned[SF_NumFrequencies];
 
-	/** Line number of the #line in MaterialTemplate.usf */
-	int32 MaterialTemplateLineNumber;
-
-	/** Contents of the MaterialTemplate.usf file */
-	FString MaterialTemplate;
-
 	TArray<int32> ScopeStack;
 
 	TArray<int32> ReferencedCodeChunks;
@@ -1323,11 +1318,15 @@ protected:
 	/**Experimental access to the EyeAdaptation RT for applying an inverse. */
 	virtual int32 EyeAdaptationInverse(int32 LightValueArg, int32 AlphaArg) override;
 
-	// to only have one piece of code dealing with error handling if the Primitive constant buffer is not used.
-	// @param Name e.g. TEXT("ObjectWorldPositionAndRadius.w")
+	/**
+	 * To only have one piece of code dealing with error handling if the Primitive constant buffer is not used.
+	 * @param Name e.g. TEXT("ObjectWorldPositionAndRadius.w")
+	 */
 	int32 GetPrimitiveProperty(EMaterialValueType Type, const TCHAR* ExpressionName, const TCHAR* HLSLName);
 
-	// The compiler can run in a different state and this affects caching of sub expression, Expressions are different (e.g. View.PrevWorldViewOrigin) when using previous frame's values
+	/**
+	 * The compiler can run in a different state and this affects caching of sub expression, Expressions are different(e.g.View.PrevWorldViewOrigin) when using previous frame's values.
+	 */
 	virtual bool IsCurrentlyCompilingForPreviousFrame() const;
 
 	virtual bool IsDevelopmentFeatureEnabled(const FName& FeatureName) const override;
diff --git a/Engine/Source/Runtime/Engine/Private/Materials/StringTemplate.h b/Engine/Source/Runtime/Engine/Private/Materials/StringTemplate.h
--- a/Engine/Source/Runtime/Engine/Private/Materials/StringTemplate.h
+++ b/Engine/Source/Runtime/Engine/Private/Materials/StringTemplate.h
@@ -69,6 +69,12 @@ public:
 	 */
 	const FString& GetTemplateString() const { return TemplateString; }
 
+	/**
+	 * Returns the number of named parameters in this template string, useful to reserve
+	 * enough capacity in your parameters map.
+	 */
+	int32 GetNumNamedParameters() const { return NumNamedParameters; }
+
 	/**
 	 * Returns the array of named parameters contained in this template.
 	 */
@@ -120,6 +126,11 @@ public:
 	 */
 	FStringTemplateResolver(const FStringTemplate& Template, uint32 ResolvedStringSizeHint = 0);
 
+	/**
+	 * Returns the template this resolver is using.
+	 */
+	const FStringTemplate& GetTemplate() const { return Template; }
+
 	/**
 	 * Sets the map of named parameters.
 	 * Note: make sure all named parameters are set before calling Advance() or Finalize() otherwise
