Wrap some text in a TEXT macro to avoid string conversion and fix an issue I'm encountering locally.

#rnx

diff --git a/Engine/Source/Runtime/Engine/Private/DebugCameraHUD.cpp b/Engine/Source/Runtime/Engine/Private/DebugCameraHUD.cpp
--- a/Engine/Source/Runtime/Engine/Private/DebugCameraHUD.cpp
+++ b/Engine/Source/Runtime/Engine/Private/DebugCameraHUD.cpp
@@ -186,7 +186,7 @@ void ADebugCameraHUD::PostRender()
 				if ( bFoundMaterial == false )
 				{
 					yl += Y;
-					Canvas->DrawText(RenderFont, "Material: NULL", X + Y, yl, 1.f, 1.f, FontRenderInfo );
+					Canvas->DrawText(RenderFont, TEXT("Material: NULL"), X + Y, yl, 1.f, 1.f, FontRenderInfo );
 				}
 
 				if (!DCC->bIsOrbitingSelectedActor)
