Iris - Fix owning connection update.

diff --git a/Engine/Source/Runtime/Engine/Private/ActorReplication.cpp b/Engine/Source/Runtime/Engine/Private/ActorReplication.cpp
--- a/Engine/Source/Runtime/Engine/Private/ActorReplication.cpp
+++ b/Engine/Source/Runtime/Engine/Private/ActorReplication.cpp
@@ -1042,7 +1042,7 @@ void AActor::UpdateOwningNetConnection() const
 {
 	using namespace UE::Net;
 
-	UReplicationSystem* ReplicationSystem = FReplicationSystemUtil::GetReplicationSystem(GetNetOwner());
+	UReplicationSystem* ReplicationSystem = FReplicationSystemUtil::GetReplicationSystem(this);
 	UObjectReplicationBridge* ObjectReplicationBridge = (ReplicationSystem ? ReplicationSystem->GetReplicationBridgeAs<UObjectReplicationBridge>() : nullptr);
 	if (ObjectReplicationBridge == nullptr)
 	{
@@ -1056,22 +1056,14 @@ void AActor::UpdateOwningNetConnection() const
 	}
 
 	// If this actor isn't replicated there's no way for us to tell whether we need to update our children.
-	bool bUpdateChildren = !GetIsReplicated();
-	if (!bUpdateChildren)
+	bool bUpdateSelfAndChildren = !GetIsReplicated();
+	if (!bUpdateSelfAndChildren)
 	{
-#if DO_CHECK
-		// Sanity check
-		{
-			const UReplicationSystem* MyReplicationSystem = FReplicationSystemUtil::GetReplicationSystem(this);
-			check(MyReplicationSystem == ReplicationSystem || MyReplicationSystem == nullptr);
-		}
-#endif
-
-		FNetHandle NetHandle = FReplicationSystemUtil::GetNetHandle(this);
-		bUpdateChildren = NetHandle.IsValid();
+		const FNetHandle NetHandle = FReplicationSystemUtil::GetNetHandle(this);
+		bUpdateSelfAndChildren = NetHandle.IsValid();
 	}
 
-	if (!bUpdateChildren)
+	if (!bUpdateSelfAndChildren)
 	{
 		return;
 	}
