NonUnity fix

#rnx

diff --git a/Engine/Source/Runtime/Engine/Private/PhysicsEngine/Experimental/PhysScene_Chaos.cpp b/Engine/Source/Runtime/Engine/Private/PhysicsEngine/Experimental/PhysScene_Chaos.cpp
--- a/Engine/Source/Runtime/Engine/Private/PhysicsEngine/Experimental/PhysScene_Chaos.cpp
+++ b/Engine/Source/Runtime/Engine/Private/PhysicsEngine/Experimental/PhysScene_Chaos.cpp
@@ -36,6 +36,7 @@
 
 #if !UE_BUILD_SHIPPING
 #include "DrawDebugHelpers.h"
+#include "UObject/UObjectIterator.h"
 
 TAutoConsoleVariable<int32> CVar_ChaosDrawHierarchyEnable(TEXT("P.Chaos.DrawHierarchy.Enable"), 0, TEXT("Enable / disable drawing of the physics hierarchy"));
 TAutoConsoleVariable<int32> CVar_ChaosDrawHierarchyCells(TEXT("P.Chaos.DrawHierarchy.Cells"), 0, TEXT("Enable / disable drawing of the physics hierarchy cells"));
