Added detection in URuntimeHashExternalStreamingObjectBase::BeginDestroy that the object is not injected anymore.
WorldPartitionRuntimeSpatialHash's InjectedGridLevels is not a UPROPERTY(Transient) to avoid any dangling ref on destroyed external streaming object not removed by client code.

diff --git a/Engine/Source/Runtime/Engine/Private/WorldPartition/WorldPartition.cpp b/Engine/Source/Runtime/Engine/Private/WorldPartition/WorldPartition.cpp
--- a/Engine/Source/Runtime/Engine/Private/WorldPartition/WorldPartition.cpp
+++ b/Engine/Source/Runtime/Engine/Private/WorldPartition/WorldPartition.cpp
@@ -1703,6 +1703,11 @@ bool UWorldPartition::InjectExternalStreamingObject(URuntimeHashExternalStreamin
 		}
 		GetWorld()->GetSubsystem<UHLODSubsystem>()->OnExternalStreamingObjectInjected(InExternalStreamingObject);
 		++StreamingStateEpoch;
+
+#if DO_CHECK
+		check(InExternalStreamingObject->TargetInjectedWorldPartition.IsExplicitlyNull());
+		InExternalStreamingObject->TargetInjectedWorldPartition = this;
+#endif
 	}
 
 	return bInjected;
@@ -1713,6 +1718,11 @@ bool UWorldPartition::RemoveExternalStreamingObject(URuntimeHashExternalStreamin
 	bool bRemoved = RuntimeHash->RemoveExternalStreamingObject(InExternalStreamingObject);
 	if (bRemoved)
 	{
+#if DO_CHECK
+		check(InExternalStreamingObject->TargetInjectedWorldPartition.IsValid());
+		InExternalStreamingObject->TargetInjectedWorldPartition = nullptr;
+#endif
+
 		if (StreamingPolicy)
 		{
 			StreamingPolicy->RemoveExternalStreamingObject(InExternalStreamingObject);
diff --git a/Engine/Source/Runtime/Engine/Private/WorldPartition/WorldPartitionRuntimeHash.cpp b/Engine/Source/Runtime/Engine/Private/WorldPartition/WorldPartitionRuntimeHash.cpp
--- a/Engine/Source/Runtime/Engine/Private/WorldPartition/WorldPartitionRuntimeHash.cpp
+++ b/Engine/Source/Runtime/Engine/Private/WorldPartition/WorldPartitionRuntimeHash.cpp
@@ -18,6 +18,14 @@
 
 #define LOCTEXT_NAMESPACE "WorldPartition"
 
+#if DO_CHECK
+void URuntimeHashExternalStreamingObjectBase::BeginDestroy()
+{
+	checkf(!TargetInjectedWorldPartition.Get(), TEXT("Destroying external streaming object that is still injected."));
+	Super::BeginDestroy();
+}
+#endif
+
 void URuntimeHashExternalStreamingObjectBase::ForEachStreamingCells(TFunctionRef<void(UWorldPartitionRuntimeCell&)> Func)
 {
 	TArray<UObject*> Objects;
diff --git a/Engine/Source/Runtime/Engine/Public/WorldPartition/WorldPartitionRuntimeHash.h b/Engine/Source/Runtime/Engine/Public/WorldPartition/WorldPartitionRuntimeHash.h
--- a/Engine/Source/Runtime/Engine/Public/WorldPartition/WorldPartitionRuntimeHash.h
+++ b/Engine/Source/Runtime/Engine/Public/WorldPartition/WorldPartitionRuntimeHash.h
@@ -39,6 +39,9 @@ class URuntimeHashExternalStreamingObjectBase : public UObject
 public:
 	//~ Begin UObject Interface
 	virtual class UWorld* GetWorld() const override final { return GetOwningWorld(); }
+#if DO_CHECK
+	virtual void BeginDestroy() override;
+#endif
 	//~ End UObject Interface
 
 	UWorld* GetOwningWorld() const { return OwningWorld.Get(); }
@@ -64,6 +67,12 @@ protected:
 
 	UPROPERTY();
 	TMap<FName, FName> CellToLevelStreamingPackage;
+
+#if DO_CHECK
+	TWeakObjectPtr<UWorldPartition> TargetInjectedWorldPartition;
+#endif
+
+	friend class UWorldPartition;
 };
 
 struct FWorldPartitionQueryCache
diff --git a/Engine/Source/Runtime/Engine/Public/WorldPartition/WorldPartitionRuntimeSpatialHash.h b/Engine/Source/Runtime/Engine/Public/WorldPartition/WorldPartitionRuntimeSpatialHash.h
--- a/Engine/Source/Runtime/Engine/Public/WorldPartition/WorldPartitionRuntimeSpatialHash.h
+++ b/Engine/Source/Runtime/Engine/Public/WorldPartition/WorldPartitionRuntimeSpatialHash.h
@@ -170,6 +170,7 @@ private:
 	mutable FSquare2DGridHelper* GridHelper;
 
 	// Contains cells injected at runtime from content bundles
+	UPROPERTY(Transient)
 	mutable TArray<FSpatialHashStreamingGridLevel> InjectedGridLevels;
 };
 
