Substrate - renaming - renamed to MCT_Substrate

diff --git a/Engine/Source/Editor/GraphEditor/Private/MaterialNodes/SGraphSubstrateMaterial.cpp b/Engine/Source/Editor/GraphEditor/Private/MaterialNodes/SGraphSubstrateMaterial.cpp
--- a/Engine/Source/Editor/GraphEditor/Private/MaterialNodes/SGraphSubstrateMaterial.cpp
+++ b/Engine/Source/Editor/GraphEditor/Private/MaterialNodes/SGraphSubstrateMaterial.cpp
@@ -54,7 +54,7 @@ bool FSubstrateWidget::HasInputSubstrateType(const UEdGraphPin* InPin)
 			{
 				switch (PinNode->MaterialExpression->GetInputType(Index))
 				{
-					case MCT_Strata:
+					case MCT_Substrate:
 						return true;
 				}
 				break;
@@ -71,7 +71,7 @@ bool FSubstrateWidget::HasOutputSubstrateType(const UEdGraphPin* InPin)
 	if (UMaterialGraphNode* PinNode = Cast<UMaterialGraphNode>(InPin->GetOwningNode()))
 	{
 		const TArray<FExpressionOutput>& ExpressionOutputs = PinNode->MaterialExpression->GetOutputs();
-		if (InPin->SourceIndex < ExpressionOutputs.Num() && PinNode->MaterialExpression->GetOutputType(InPin->SourceIndex) == MCT_Strata)
+		if (InPin->SourceIndex < ExpressionOutputs.Num() && PinNode->MaterialExpression->GetOutputType(InPin->SourceIndex) == MCT_Substrate)
 		{
 			return true;
 		}
@@ -272,7 +272,7 @@ void FSubstrateWidget::GetPinColor(TSharedPtr<SGraphPin>& Out, const UMaterialGr
 		bHasColorModifier = true;
 		ColorModifier = FSubstrateWidget::GetConnectionColor();
 	}
-	else if (Pin && UMaterialGraphSchema::GetMaterialValueType(Pin) == MCT_Strata)
+	else if (Pin && UMaterialGraphSchema::GetMaterialValueType(Pin) == MCT_Substrate)
 	{
 		if (!Out->IsConnected())
 		{
diff --git a/Engine/Source/Editor/GraphEditor/Private/MaterialPins/SGraphPinMaterialInput.cpp b/Engine/Source/Editor/GraphEditor/Private/MaterialPins/SGraphPinMaterialInput.cpp
--- a/Engine/Source/Editor/GraphEditor/Private/MaterialPins/SGraphPinMaterialInput.cpp
+++ b/Engine/Source/Editor/GraphEditor/Private/MaterialPins/SGraphPinMaterialInput.cpp
@@ -17,7 +17,7 @@ void SGraphPinMaterialInput::Construct(const FArguments& InArgs, UEdGraphPin* In
 	bool bUsePinColor = true;
 	if (Substrate::IsSubstrateEnabled())
 	{
-		if (UMaterialGraphSchema::GetMaterialValueType(InGraphPinObj) == MCT_Strata)
+		if (UMaterialGraphSchema::GetMaterialValueType(InGraphPinObj) == MCT_Substrate)
 		{
 			bUsePinColor = false;
 		}			
@@ -35,7 +35,7 @@ FSlateColor SGraphPinMaterialInput::GetPinColor() const
 	{
 		if (Substrate::IsSubstrateEnabled())
 		{
-			if (UMaterialGraphSchema::GetMaterialValueType(GraphPinObj) == MCT_Strata)
+			if (UMaterialGraphSchema::GetMaterialValueType(GraphPinObj) == MCT_Substrate)
 			{
 				return FSlateColor(FSubstrateWidget::GetConnectionColor());
 			}			
diff --git a/Engine/Source/Editor/MaterialEditor/Private/MaterialEditor.cpp b/Engine/Source/Editor/MaterialEditor/Private/MaterialEditor.cpp
--- a/Engine/Source/Editor/MaterialEditor/Private/MaterialEditor.cpp
+++ b/Engine/Source/Editor/MaterialEditor/Private/MaterialEditor.cpp
@@ -4762,7 +4762,7 @@ UClass* FMaterialEditor::GetOnPromoteToParameterClass(const UEdGraphPin* TargetP
 					case MCT_VolumeTexture: 
 					case MCT_Texture: return UMaterialExpressionTextureObjectParameter::StaticClass();
 
-					case MCT_Strata: return nullptr;
+					case MCT_Substrate: return nullptr;
 				}
 
 				break;
diff --git a/Engine/Source/Editor/MaterialEditor/Private/MaterialEditorHelpers.cpp b/Engine/Source/Editor/MaterialEditor/Private/MaterialEditorHelpers.cpp
--- a/Engine/Source/Editor/MaterialEditor/Private/MaterialEditorHelpers.cpp
+++ b/Engine/Source/Editor/MaterialEditor/Private/MaterialEditorHelpers.cpp
@@ -240,7 +240,7 @@ void FMaterialEditorHelpers::CollapseToFunction(FMaterialEditor& MaterialEditor)
 				case MCT_MaterialAttributes: PinFunctionType = FunctionInput_MaterialAttributes; break;
 				case MCT_TextureExternal:    PinFunctionType = FunctionInput_TextureExternal;    break;
 				case MCT_Bool:				 PinFunctionType = FunctionInput_Bool;				 break;
-				case MCT_Strata:			 PinFunctionType = FunctionInput_Substrate;			 break;
+				case MCT_Substrate:			 PinFunctionType = FunctionInput_Substrate;			 break;
 				default:
 					// Will happen pretty often with MCT_Float, as the types are rarely fully resolved
 					// (eg Add nodes can take float1/2/3/4)
diff --git a/Engine/Source/Editor/UnrealEd/Private/MaterialGraphNode_Root.cpp b/Engine/Source/Editor/UnrealEd/Private/MaterialGraphNode_Root.cpp
--- a/Engine/Source/Editor/UnrealEd/Private/MaterialGraphNode_Root.cpp
+++ b/Engine/Source/Editor/UnrealEd/Private/MaterialGraphNode_Root.cpp
@@ -91,7 +91,7 @@ uint32 UMaterialGraphNode_Root::GetPinMaterialType(const UEdGraphPin* Pin) const
 	}
 	else if (Property == MP_FrontMaterial)
 	{
-		MaterialType = MCT_Strata;
+		MaterialType = MCT_Substrate;
 	}
 	else
 	{
diff --git a/Engine/Source/Editor/UnrealEd/Private/MaterialGraphSchema.cpp b/Engine/Source/Editor/UnrealEd/Private/MaterialGraphSchema.cpp
--- a/Engine/Source/Editor/UnrealEd/Private/MaterialGraphSchema.cpp
+++ b/Engine/Source/Editor/UnrealEd/Private/MaterialGraphSchema.cpp
@@ -122,7 +122,7 @@ void FMaterialGraphSchemaAction_NewNode::SetFunctionInputType(UMaterialExpressio
 	case MCT_MaterialAttributes:
 		FunctionInput->InputType = FunctionInput_MaterialAttributes;
 		break;
-	case MCT_Strata:
+	case MCT_Substrate:
 		FunctionInput->InputType = FunctionInput_Substrate;
 		break;
 	default:
diff --git a/Engine/Source/Runtime/Engine/Private/Materials/HLSLMaterialTranslator.cpp b/Engine/Source/Runtime/Engine/Private/Materials/HLSLMaterialTranslator.cpp
--- a/Engine/Source/Runtime/Engine/Private/Materials/HLSLMaterialTranslator.cpp
+++ b/Engine/Source/Runtime/Engine/Private/Materials/HLSLMaterialTranslator.cpp
@@ -2618,7 +2618,7 @@ FString FHLSLMaterialTranslator::GetMaterialShaderCode()
 		case MCT_UInt2: HLSLType = TEXT("uint2"); break;
 		case MCT_UInt3: HLSLType = TEXT("uint3"); break;
 		case MCT_UInt4: HLSLType = TEXT("uint4"); break;
-		case MCT_Strata: HLSLType = TEXT("FSubstrateData"); break;
+		case MCT_Substrate: HLSLType = TEXT("FSubstrateData"); break;
 		default: break;
 		}
 
@@ -3224,7 +3224,7 @@ const TCHAR* FHLSLMaterialTranslator::DescribeType(EMaterialValueType Type) cons
 	case MCT_UInt2:					return TEXT("uint2");
 	case MCT_UInt3:					return TEXT("uint3");
 	case MCT_UInt4:					return TEXT("uint4");
-	case MCT_Strata:				return TEXT("Substrate");
+	case MCT_Substrate:				return TEXT("Substrate");
 	case MCT_LWCScalar:				return TEXT("LWCScalar");
 	case MCT_LWCVector2:			return TEXT("LWCVector2");
 	case MCT_LWCVector3:			return TEXT("LWCVector3");
@@ -3261,7 +3261,7 @@ const TCHAR* FHLSLMaterialTranslator::HLSLTypeString(EMaterialValueType Type) co
 	case MCT_UInt2:					return TEXT("uint2");
 	case MCT_UInt3:					return TEXT("uint3");
 	case MCT_UInt4:					return TEXT("uint4");
-	case MCT_Strata:				return TEXT("FSubstrateData");
+	case MCT_Substrate:				return TEXT("FSubstrateData");
 	case MCT_LWCScalar:				return TEXT("FLWCScalar");
 	case MCT_LWCVector2:			return TEXT("FLWCVector2");
 	case MCT_LWCVector3:			return TEXT("FLWCVector3");
@@ -3298,7 +3298,7 @@ const TCHAR* FHLSLMaterialTranslator::HLSLTypeStringDeriv(EMaterialValueType Typ
 	case MCT_UInt2:					return TEXT("uint2");
 	case MCT_UInt3:					return TEXT("uint3");
 	case MCT_UInt4:					return TEXT("uint4");
-	case MCT_Strata:				return TEXT("FSubstrateData");
+	case MCT_Substrate:				return TEXT("FSubstrateData");
 	case MCT_LWCScalar:				return (DerivativeStatus == EDerivativeStatus::Valid) ? TEXT("FLWCScalarDeriv") : TEXT("FLWCScalar");
 	case MCT_LWCVector2:			return (DerivativeStatus == EDerivativeStatus::Valid) ? TEXT("FLWCVector2Deriv") : TEXT("FLWCVector2");
 	case MCT_LWCVector3:			return (DerivativeStatus == EDerivativeStatus::Valid) ? TEXT("FLWCVector3Deriv") : TEXT("FLWCVector3");
@@ -3406,7 +3406,7 @@ int32 FHLSLMaterialTranslator::AddCodeChunkInner(uint64 Hash, const TCHAR* Forma
 		new(*CurrentScopeChunks) FShaderCodeChunk(Hash, FormattedCode, FormattedCode, TEXT(""), Type, DerivativeStatus, true);
 	}
 	// Can only create temporaries for certain types
-	else if ((Type & (MCT_Float | MCT_LWCType | MCT_VTPageTableResult | MCT_UInt)) || Type == MCT_ShadingModel || Type == MCT_MaterialAttributes || Type == MCT_Strata || Type == MCT_UInt)
+	else if ((Type & (MCT_Float | MCT_LWCType | MCT_VTPageTableResult | MCT_UInt)) || Type == MCT_ShadingModel || Type == MCT_MaterialAttributes || Type == MCT_Substrate || Type == MCT_UInt)
 	{
 		// Check for existing
 		for (int32 i = 0; i < CurrentScopeChunks->Num(); ++i)
@@ -3486,7 +3486,7 @@ int32 FHLSLMaterialTranslator::AddCodeChunkInnerDeriv(const TCHAR* FormattedCode
 		new(*CurrentScopeChunks) FShaderCodeChunk(Hash, FormattedCodeFinite, FormattedCodeAnalytic, TEXT(""), Type, DerivativeStatus, true);
 	}
 	// Can only create temporaries for certain types
-	else if ((Type & (MCT_Float | MCT_LWCType | MCT_VTPageTableResult | MCT_UInt)) || Type == MCT_ShadingModel || Type == MCT_MaterialAttributes || Type == MCT_Strata)
+	else if ((Type & (MCT_Float | MCT_LWCType | MCT_VTPageTableResult | MCT_UInt)) || Type == MCT_ShadingModel || Type == MCT_MaterialAttributes || Type == MCT_Substrate)
 	{
 		// Check for existing
 		for (int32 i = 0; i < CurrentScopeChunks->Num(); ++i)
@@ -12394,7 +12394,7 @@ int32 FHLSLMaterialTranslator::SubstrateHorizontalMixingParameterBlendingBSDFCov
 
 int32 FHLSLMaterialTranslator::SubstrateCreateAndRegisterNullMaterial()
 {
-	int32 OutputCodeChunk = AddInlinedCodeChunk(MCT_Strata, TEXT("GetInitialisedSubstrateData()"));
+	int32 OutputCodeChunk = AddInlinedCodeChunk(MCT_Substrate, TEXT("GetInitialisedSubstrateData()"));
 	return OutputCodeChunk;
 }
 
@@ -12432,7 +12432,7 @@ int32 FHLSLMaterialTranslator::SubstrateSlabBSDF(
 		}
 
 		return AddCodeChunk(
-			MCT_Strata, TEXT("Parameters.%s.PromoteParameterBlendedBSDFToOperator(GetSubstrateSlabBSDF(Parameters.SubstratePixelFootprint, %s, %s, %s, %s, %s, %s, %s, %s, %s, %s, %s, %s, %s, %s, %s, %s, %s, %s, %s, %s, %s, %s, %s, %s, Parameters.%s.Types) /* Normal = %s ; Tangent = %s ; Thickness = %s */, %u, %u, %u, %u)"),
+			MCT_Substrate, TEXT("Parameters.%s.PromoteParameterBlendedBSDFToOperator(GetSubstrateSlabBSDF(Parameters.SubstratePixelFootprint, %s, %s, %s, %s, %s, %s, %s, %s, %s, %s, %s, %s, %s, %s, %s, %s, %s, %s, %s, %s, %s, %s, %s, %s, Parameters.%s.Types) /* Normal = %s ; Tangent = %s ; Thickness = %s */, %u, %u, %u, %u)"),
 			*GetParametersSubstrateTreeName(CurrentSubstrateCompilationContext),
 			*SubstrateGetCastParameterCode(DiffuseAlbedo,			MCT_Float3),
 			*SubstrateGetCastParameterCode(F0,						MCT_Float3),
@@ -12470,7 +12470,7 @@ int32 FHLSLMaterialTranslator::SubstrateSlabBSDF(
 	}
 	
 	return AddCodeChunk(
-		MCT_Strata, TEXT("GetSubstrateSlabBSDF(Parameters.SubstratePixelFootprint, %s, %s, %s, %s, %s, %s, %s, %s, %s, %s, %s, %s, %s, %s, %s, %s, %s, %s, %s, %s, %s, %s, %s, %s, Parameters.%s.Types) /* Normal = %s ; Tangent = %s ; Thickness = %s */"),
+		MCT_Substrate, TEXT("GetSubstrateSlabBSDF(Parameters.SubstratePixelFootprint, %s, %s, %s, %s, %s, %s, %s, %s, %s, %s, %s, %s, %s, %s, %s, %s, %s, %s, %s, %s, %s, %s, %s, %s, Parameters.%s.Types) /* Normal = %s ; Tangent = %s ; Thickness = %s */"),
 		*SubstrateGetCastParameterCode(DiffuseAlbedo,			MCT_Float3),
 		*SubstrateGetCastParameterCode(F0,						MCT_Float3),
 		*SubstrateGetCastParameterCode(F90,					MCT_Float3),
@@ -12552,7 +12552,7 @@ int32 FHLSLMaterialTranslator::SubstrateConversionFromLegacy(
 		}
 
 		return AddCodeChunk(
-			MCT_Strata, TEXT("Parameters.%s.PromoteParameterBlendedBSDFToOperator(SubstrateConvertLegacyMaterial%s(Parameters.SubstratePixelFootprint, %s, %s, %s, %s, %s, %s, %s, %s, %s, %s, %s, %s, %s, %s, %s, %s, %s, %s, %s, %s, %s, %s, %s, Parameters.%s.Types) /* Normal = %s ; Tangent = %s ; ClearCoat_Normal = %s ; ClearCoat_Tangent = %s */, %u, %u, %u, %u)"),
+			MCT_Substrate, TEXT("Parameters.%s.PromoteParameterBlendedBSDFToOperator(SubstrateConvertLegacyMaterial%s(Parameters.SubstratePixelFootprint, %s, %s, %s, %s, %s, %s, %s, %s, %s, %s, %s, %s, %s, %s, %s, %s, %s, %s, %s, %s, %s, %s, %s, Parameters.%s.Types) /* Normal = %s ; Tangent = %s ; ClearCoat_Normal = %s ; ClearCoat_Tangent = %s */, %u, %u, %u, %u)"),
 			*GetParametersSubstrateTreeName(CurrentSubstrateCompilationContext),
 			bHasDynamicShadingModels ? TEXT("Dynamic") : TEXT("Static"),
 			*SubstrateGetCastParameterCode(BaseColor,						MCT_Float3),
@@ -12594,7 +12594,7 @@ int32 FHLSLMaterialTranslator::SubstrateConversionFromLegacy(
 	}
 
 	return AddCodeChunk(
-		MCT_Strata, TEXT("SubstrateConvertLegacyMaterial%s(Parameters.SubstratePixelFootprint, %s, %s, %s, %s, %s, %s, %s, %s, %s, %s, %s, %s, %s, %s, %s, %s, %s, %s, %s, %s, %s, %s, %s, Parameters.%s.Types) /* Normal = %s ; Tangent = %s ; ClearCoat_Normal = %s ; ClearCoat_Tangent = %s */"),
+		MCT_Substrate, TEXT("SubstrateConvertLegacyMaterial%s(Parameters.SubstratePixelFootprint, %s, %s, %s, %s, %s, %s, %s, %s, %s, %s, %s, %s, %s, %s, %s, %s, %s, %s, %s, %s, %s, %s, %s, Parameters.%s.Types) /* Normal = %s ; Tangent = %s ; ClearCoat_Normal = %s ; ClearCoat_Tangent = %s */"),
 		bHasDynamicShadingModels ? TEXT("Dynamic") : TEXT("Static"),
 		*SubstrateGetCastParameterCode(BaseColor,						MCT_Float3),
 		*SubstrateGetCastParameterCode(Specular,						MCT_Float),
@@ -12633,7 +12633,7 @@ int32 FHLSLMaterialTranslator::SubstrateConversionFromLegacy(
 int32 FHLSLMaterialTranslator::SubstrateVolumetricFogCloudBSDF(int32 Albedo, int32 Extinction, int32 EmissiveColor, int32 AmbientOcclusion)
 {
 	return AddCodeChunk(
-		MCT_Strata, TEXT("GetSubstrateVolumeFogCloudBSDF(%s, %s, %s, %s)"),
+		MCT_Substrate, TEXT("GetSubstrateVolumeFogCloudBSDF(%s, %s, %s, %s)"),
 		*SubstrateGetCastParameterCode(Albedo,				MCT_Float3),
 		*SubstrateGetCastParameterCode(Extinction,			MCT_Float3),
 		*SubstrateGetCastParameterCode(EmissiveColor,		MCT_Float3),
@@ -12646,7 +12646,7 @@ int32 FHLSLMaterialTranslator::SubstrateUnlitBSDF(int32 EmissiveColor, int32 Tra
 	if (PromoteToOperator)
 	{
 		return AddCodeChunk(
-			MCT_Strata, TEXT("Parameters.%s.PromoteParameterBlendedBSDFToOperator(GetSubstrateUnlitBSDF(%s, %s, %s), %u, %u, %u, %u)"),
+			MCT_Substrate, TEXT("Parameters.%s.PromoteParameterBlendedBSDFToOperator(GetSubstrateUnlitBSDF(%s, %s, %s), %u, %u, %u, %u)"),
 			*GetParametersSubstrateTreeName(CurrentSubstrateCompilationContext),
 			*SubstrateGetCastParameterCode(EmissiveColor,		MCT_Float3),
 			*SubstrateGetCastParameterCode(TransmittanceColor, MCT_Float3),
@@ -12658,7 +12658,7 @@ int32 FHLSLMaterialTranslator::SubstrateUnlitBSDF(int32 EmissiveColor, int32 Tra
 	}
 	
 	return AddCodeChunk(
-		MCT_Strata, TEXT("GetSubstrateUnlitBSDF(%s, %s, %s)"),
+		MCT_Substrate, TEXT("GetSubstrateUnlitBSDF(%s, %s, %s)"),
 		*SubstrateGetCastParameterCode(EmissiveColor,		MCT_Float3),
 		*SubstrateGetCastParameterCode(TransmittanceColor, MCT_Float3),
 		*SubstrateGetCastParameterCode(Normal,				MCT_Float3)
@@ -12668,7 +12668,7 @@ int32 FHLSLMaterialTranslator::SubstrateUnlitBSDF(int32 EmissiveColor, int32 Tra
 int32 FHLSLMaterialTranslator::SubstrateHairBSDF(int32 BaseColor, int32 Scatter, int32 Specular, int32 Roughness, int32 Backlit, int32 EmissiveColor, int32 Tangent, const FString& SharedLocalBasisIndexMacro, FSubstrateOperator* PromoteToOperator)
 {
 	return AddCodeChunk(
-		MCT_Strata, TEXT("Parameters.%s.PromoteParameterBlendedBSDFToOperator(GetSubstrateHairBSDF(%s, %s, %s, %s, %s, %s, %s), %u, %u, %u, %u) /* Tangent:%s */"),
+		MCT_Substrate, TEXT("Parameters.%s.PromoteParameterBlendedBSDFToOperator(GetSubstrateHairBSDF(%s, %s, %s, %s, %s, %s, %s), %u, %u, %u, %u) /* Tangent:%s */"),
 		*GetParametersSubstrateTreeName(CurrentSubstrateCompilationContext),
 		*SubstrateGetCastParameterCode(BaseColor,			MCT_Float3),
 		*SubstrateGetCastParameterCode(Scatter,			MCT_Float),
@@ -12688,7 +12688,7 @@ int32 FHLSLMaterialTranslator::SubstrateHairBSDF(int32 BaseColor, int32 Scatter,
 int32 FHLSLMaterialTranslator::SubstrateEyeBSDF(int32 DiffuseAlbedo, int32 Roughness, int32 IrisMask, int32 IrisDistance, int32 IrisNormal, int32 IrisPlaneNormal, int32 SSSProfileId, int32 EmissiveColor, int32 CorneaNormal, const FString& SharedLocalBasisIndexMacro, FSubstrateOperator* PromoteToOperator)
 {
 	return AddCodeChunk(
-		MCT_Strata, TEXT("Parameters.%s.PromoteParameterBlendedBSDFToOperator(GetSubstrateEyeBSDF(%s, %s, %s, %s, %s, %s, %s, %s, %s), %u, %u, %u, %u) /* Cornea:%s Iris:%s */"),
+		MCT_Substrate, TEXT("Parameters.%s.PromoteParameterBlendedBSDFToOperator(GetSubstrateEyeBSDF(%s, %s, %s, %s, %s, %s, %s, %s, %s), %u, %u, %u, %u) /* Cornea:%s Iris:%s */"),
 		*GetParametersSubstrateTreeName(CurrentSubstrateCompilationContext),
 		*SubstrateGetCastParameterCode(DiffuseAlbedo, MCT_Float3),
 		*SubstrateGetCastParameterCode(Roughness, MCT_Float),
@@ -12714,7 +12714,7 @@ int32 FHLSLMaterialTranslator::SubstrateSingleLayerWaterBSDF(
 	FSubstrateOperator* PromoteToOperator)
 {
 	return AddCodeChunk(
-		MCT_Strata, TEXT("Parameters.%s.PromoteParameterBlendedBSDFToOperator(GetSubstrateSingleLayerWaterBSDF(%s, %s, %s, %s, %s, %s, %s, %s, %s, %s, %s), %u, %u, %u, %u) /* Normal:%s */"),
+		MCT_Substrate, TEXT("Parameters.%s.PromoteParameterBlendedBSDFToOperator(GetSubstrateSingleLayerWaterBSDF(%s, %s, %s, %s, %s, %s, %s, %s, %s, %s, %s), %u, %u, %u, %u) /* Normal:%s */"),
 		*GetParametersSubstrateTreeName(CurrentSubstrateCompilationContext),
 		*SubstrateGetCastParameterCode(BaseColor,				MCT_Float3),
 		*SubstrateGetCastParameterCode(Metallic,				MCT_Float),
@@ -12742,7 +12742,7 @@ int32 FHLSLMaterialTranslator::SubstrateHorizontalMixing(int32 Background, int32
 		return INDEX_NONE;
 	}
 	return AddCodeChunk(
-		MCT_Strata, TEXT("Parameters.%s.SubstrateHorizontalMixing(%s, %s, %s, %u, %u)"),
+		MCT_Substrate, TEXT("Parameters.%s.SubstrateHorizontalMixing(%s, %s, %s, %u, %u)"),
 		*GetParametersSubstrateTreeName(CurrentSubstrateCompilationContext),
 		*GetParameterCode(Background),
 		*GetParameterCode(Foreground),
@@ -12765,7 +12765,7 @@ int32 FHLSLMaterialTranslator::SubstrateHorizontalMixingParameterBlending(
 		check(PromoteToOperator->Index != INDEX_NONE);
 		check(PromoteToOperator->BSDFIndex != INDEX_NONE);
 		return AddCodeChunk(
-			MCT_Strata, TEXT("Parameters.%s.PromoteParameterBlendedBSDFToOperator(SubstrateHorizontalMixingParameterBlending(%s, %s, %s, %s, %s, Parameters.%s.Types), %u, %u, %u, %u)"),
+			MCT_Substrate, TEXT("Parameters.%s.PromoteParameterBlendedBSDFToOperator(SubstrateHorizontalMixingParameterBlending(%s, %s, %s, %s, %s, Parameters.%s.Types), %u, %u, %u, %u)"),
 			*GetParametersSubstrateTreeName(CurrentSubstrateCompilationContext),
 			*GetParameterCode(Background),
 			*GetParameterCode(Foreground),
@@ -12781,7 +12781,7 @@ int32 FHLSLMaterialTranslator::SubstrateHorizontalMixingParameterBlending(
 	}
 
 	return AddCodeChunk(
-		MCT_Strata, TEXT("SubstrateHorizontalMixingParameterBlending(%s, %s, %s, %s, %s, Parameters.%s.Types)"),
+		MCT_Substrate, TEXT("SubstrateHorizontalMixingParameterBlending(%s, %s, %s, %s, %s, Parameters.%s.Types)"),
 		*GetParameterCode(Background),
 		*GetParameterCode(Foreground),
 		*GetParameterCode(HorizontalMixCodeChunk),
@@ -12799,7 +12799,7 @@ int32 FHLSLMaterialTranslator::SubstrateVerticalLayering(int32 Top, int32 Base,
 	}
 	const FString ThicknessCode = Thickness != INDEX_NONE ? GetParameterCode(Thickness) : TEXT("NONE");;
 	return AddCodeChunk(
-		MCT_Strata, TEXT("Parameters.%s.SubstrateVerticalLayering(%s, %s, %u, %u) /* Thickness = %s */"),
+		MCT_Substrate, TEXT("Parameters.%s.SubstrateVerticalLayering(%s, %s, %u, %u) /* Thickness = %s */"),
 		*GetParametersSubstrateTreeName(CurrentSubstrateCompilationContext),
 		*GetParameterCode(Top),
 		*GetParameterCode(Base),
@@ -12823,7 +12823,7 @@ int32 FHLSLMaterialTranslator::SubstrateVerticalLayeringParameterBlending(int32
 		check(PromoteToOperator->Index != INDEX_NONE);
 		check(PromoteToOperator->BSDFIndex != INDEX_NONE);
 		return AddCodeChunk(
-			MCT_Strata, TEXT("Parameters.%s.PromoteParameterBlendedBSDFToOperator(SubstrateVerticalLayeringParameterBlending(%s, %s, %s, dot(%s, %s)), %u, %u, %u, %u) /* Thickness = %s */"),
+			MCT_Substrate, TEXT("Parameters.%s.PromoteParameterBlendedBSDFToOperator(SubstrateVerticalLayeringParameterBlending(%s, %s, %s, dot(%s, %s)), %u, %u, %u, %u) /* Thickness = %s */"),
 			*GetParametersSubstrateTreeName(CurrentSubstrateCompilationContext),
 			*GetParameterCode(Top),
 			*GetParameterCode(Base),
@@ -12839,7 +12839,7 @@ int32 FHLSLMaterialTranslator::SubstrateVerticalLayeringParameterBlending(int32
 	}
 
 	return AddCodeChunk(
-		MCT_Strata, TEXT("SubstrateVerticalLayeringParameterBlending(%s, %s, %s, dot(%s, %s)) /* Thickness = %s */"),
+		MCT_Substrate, TEXT("SubstrateVerticalLayeringParameterBlending(%s, %s, %s, dot(%s, %s)) /* Thickness = %s */"),
 		*GetParameterCode(Top),
 		*GetParameterCode(Base),
 		*SharedLocalBasisIndexMacro,
@@ -12856,7 +12856,7 @@ int32 FHLSLMaterialTranslator::SubstrateAdd(int32 A, int32 B, int OperatorIndex,
 		return INDEX_NONE;
 	}
 	return AddCodeChunk(
-		MCT_Strata, TEXT("Parameters.%s.SubstrateAdd(%s, %s, %u, %u)"),
+		MCT_Substrate, TEXT("Parameters.%s.SubstrateAdd(%s, %s, %u, %u)"),
 		*GetParametersSubstrateTreeName(CurrentSubstrateCompilationContext),
 		*GetParameterCode(A),
 		*GetParameterCode(B),
@@ -12877,7 +12877,7 @@ int32 FHLSLMaterialTranslator::SubstrateAddParameterBlending(int32 A, int32 B, i
 		check(PromoteToOperator->Index != INDEX_NONE);
 		check(PromoteToOperator->BSDFIndex != INDEX_NONE);
 		return AddCodeChunk(
-			MCT_Strata, TEXT("Parameters.%s.PromoteParameterBlendedBSDFToOperator(SubstrateAddParameterBlending(%s, %s, %s, %s), %u, %u, %u, %u)"),
+			MCT_Substrate, TEXT("Parameters.%s.PromoteParameterBlendedBSDFToOperator(SubstrateAddParameterBlending(%s, %s, %s, %s), %u, %u, %u, %u)"),
 			*GetParametersSubstrateTreeName(CurrentSubstrateCompilationContext),
 			*GetParameterCode(A),
 			*GetParameterCode(B),
@@ -12891,7 +12891,7 @@ int32 FHLSLMaterialTranslator::SubstrateAddParameterBlending(int32 A, int32 B, i
 	}
 
 	return AddCodeChunk(
-		MCT_Strata, TEXT("SubstrateAddParameterBlending(%s, %s, %s, %s)"),
+		MCT_Substrate, TEXT("SubstrateAddParameterBlending(%s, %s, %s, %s)"),
 		*GetParameterCode(A),
 		*GetParameterCode(B),
 		*GetParameterCode(AMixWeight),
@@ -12906,7 +12906,7 @@ int32 FHLSLMaterialTranslator::SubstrateWeight(int32 A, int32 Weight, int Operat
 		return INDEX_NONE;
 	}
 	return AddCodeChunk(
-		MCT_Strata, TEXT("Parameters.%s.SubstrateWeight(%s, %s, %u, %u)"),
+		MCT_Substrate, TEXT("Parameters.%s.SubstrateWeight(%s, %s, %u, %u)"),
 		*GetParametersSubstrateTreeName(CurrentSubstrateCompilationContext),
 		*GetParameterCode(A),
 		*GetParameterCode(Weight),
@@ -12927,7 +12927,7 @@ int32 FHLSLMaterialTranslator::SubstrateWeightParameterBlending(int32 A, int32 W
 		check(PromoteToOperator->Index != INDEX_NONE);
 		check(PromoteToOperator->BSDFIndex != INDEX_NONE);
 		return AddCodeChunk(
-			MCT_Strata, TEXT("Parameters.%s.PromoteParameterBlendedBSDFToOperator(SubstrateWeightParameterBlending(%s, %s), %u, %u, %u, %u)"),
+			MCT_Substrate, TEXT("Parameters.%s.PromoteParameterBlendedBSDFToOperator(SubstrateWeightParameterBlending(%s, %s), %u, %u, %u, %u)"),
 			*GetParametersSubstrateTreeName(CurrentSubstrateCompilationContext),
 			*GetParameterCode(A),
 			*GetParameterCode(Weight),
@@ -12939,7 +12939,7 @@ int32 FHLSLMaterialTranslator::SubstrateWeightParameterBlending(int32 A, int32 W
 	}
 
 	return AddCodeChunk(
-		MCT_Strata, TEXT("SubstrateWeightParameterBlending(%s, %s)"),
+		MCT_Substrate, TEXT("SubstrateWeightParameterBlending(%s, %s)"),
 		*GetParameterCode(A),
 		*GetParameterCode(Weight)
 	);
diff --git a/Engine/Source/Runtime/Engine/Private/Materials/MaterialAttributeDefinitionMap.cpp b/Engine/Source/Runtime/Engine/Private/Materials/MaterialAttributeDefinitionMap.cpp
--- a/Engine/Source/Runtime/Engine/Private/Materials/MaterialAttributeDefinitionMap.cpp
+++ b/Engine/Source/Runtime/Engine/Private/Materials/MaterialAttributeDefinitionMap.cpp
@@ -25,7 +25,7 @@ FMaterialAttributeDefintion::FMaterialAttributeDefintion(
 	, BlendFunction(InBlendFunction)
 	, bIsHidden(bInIsHidden)
 {
-	checkf(ValueType & MCT_Float || ValueType == MCT_ShadingModel || ValueType == MCT_Strata || ValueType == MCT_MaterialAttributes, TEXT("Unsupported material attribute type %d"), ValueType);
+	checkf(ValueType & MCT_Float || ValueType == MCT_ShadingModel || ValueType == MCT_Substrate || ValueType == MCT_MaterialAttributes, TEXT("Unsupported material attribute type %d"), ValueType);
 }
 
 int32 FMaterialAttributeDefintion::CompileDefaultValue(FMaterialCompiler* Compiler) const
@@ -66,7 +66,7 @@ int32 FMaterialAttributeDefintion::CompileDefaultValue(FMaterialCompiler* Compil
 
 	if (Property == MP_FrontMaterial)
 	{
-		check(ValueType == MCT_Strata);
+		check(ValueType == MCT_Substrate);
 		return Compiler->SubstrateCreateAndRegisterNullMaterial();
 	}
 
@@ -270,7 +270,7 @@ void FMaterialAttributeDefinitionMap::InitializeAttributeMap()
 	Add(FGuid(0x0AC97EC3, 0xE3D047BA, 0xB610167D, 0xC4D919FF), TEXT("PixelDepthOffset"),		MP_PixelDepthOffset,		MCT_Float,	FVector4(0,0,0,0),	SF_Pixel);
 	Add(FGuid(0xD9423FFF, 0xD77E4D82, 0x8FF9CF5E, 0x055D1255), TEXT("ShadingModel"),			MP_ShadingModel,			MCT_ShadingModel, FVector4(0, 0, 0, 0), SF_Pixel, INDEX_NONE, false, &CompileShadingModelBlendFunction);
 	Add(FGuid(0x42BDD2E0, 0xBE714189, 0xA0984BC3, 0xDD0BE872), TEXT("SurfaceThickness"),		MP_SurfaceThickness,		MCT_Float,  FVector4(SUBSTRATE_LAYER_DEFAULT_THICKNESS_CM, 0, 0, 0), SF_Pixel);
-	Add(FGuid(0x5973A03E, 0x13A74E08, 0x92D0CEDD, 0xF2936CF8), TEXT("FrontMaterial"),			MP_FrontMaterial,			MCT_Strata, FVector4(0,0,0,0),	SF_Pixel, INDEX_NONE, false, &CompileSubstrateBlendFunction);
+	Add(FGuid(0x5973A03E, 0x13A74E08, 0x92D0CEDD, 0xF2936CF8), TEXT("FrontMaterial"),			MP_FrontMaterial,			MCT_Substrate, FVector4(0,0,0,0),	SF_Pixel, INDEX_NONE, false, &CompileSubstrateBlendFunction);
 
 	// Used when compiling material with execution pins, which are compiling all attributes together
 	Add(FGuid(0xE0ED040B, 0x82794D93, 0xBD2D59B2, 0xA5BBF41C), TEXT("MaterialAttributes"),		MP_MaterialAttributes,		MCT_MaterialAttributes, FVector4(0,0,0,0), SF_Pixel, INDEX_NONE, bHideAttribute);
diff --git a/Engine/Source/Runtime/Engine/Private/Materials/MaterialExpressions.cpp b/Engine/Source/Runtime/Engine/Private/Materials/MaterialExpressions.cpp
--- a/Engine/Source/Runtime/Engine/Private/Materials/MaterialExpressions.cpp
+++ b/Engine/Source/Runtime/Engine/Private/Materials/MaterialExpressions.cpp
@@ -444,7 +444,7 @@ void GetMaterialValueTypeDescriptions(const uint32 MaterialValueType, TArray<FTe
 		OutDescriptions.Add(LOCTEXT("MaterialAttributes", "Material Attributes"));
 	if (MaterialValueType & MCT_ShadingModel)
 		OutDescriptions.Add(LOCTEXT("ShadingModel", "Shading Model"));
-	if (MaterialValueType & MCT_Strata)
+	if (MaterialValueType & MCT_Substrate)
 		OutDescriptions.Add(LOCTEXT("Substrate", "Substrate Material"));
 	if (MaterialValueType & MCT_Unknown)
 		OutDescriptions.Add(LOCTEXT("Unknown", "Unknown"));
@@ -1846,14 +1846,14 @@ uint32 UMaterialExpression::GetOutputType(int32 OutputIndex)
 {
 	// different outputs should be defined by sub classed expressions 
 
-	// Material attributes need to be tested first to work when plugged in main root node (to not return MCT_Strata when Substrate mateiral is fed)
+	// Material attributes need to be tested first to work when plugged in main root node (to not return MCT_Substrate when Substrate mateiral is fed)
 	if (IsResultMaterialAttributes(OutputIndex))
 	{
 		return MCT_MaterialAttributes;
 	}
 	else if (IsResultSubstrateMaterial(OutputIndex))
 	{
-		return MCT_Strata;
+		return MCT_Substrate;
 	}
 	else
 	{
@@ -7111,7 +7111,7 @@ uint32 UMaterialExpressionGetMaterialAttributes::GetOutputType(int32 OutputIndex
 		{
 			OutputType = PinType;
 		}
-		else if (PinType == MCT_Strata)
+		else if (PinType == MCT_Substrate)
 		{
 			OutputType = PinType;
 		}
@@ -7410,9 +7410,9 @@ uint32 UMaterialExpressionSetMaterialAttributes::GetInputType(int32 InputIndex)
 		{
 			InputType = MCT_ShadingModel;
 		}
-		else if (InputType == MCT_Strata)
+		else if (InputType == MCT_Substrate)
 		{
-			InputType = MCT_Strata;
+			InputType = MCT_Substrate;
 		}
 		else
 		{
@@ -17136,7 +17136,7 @@ int32 UMaterialExpressionFunctionInput::Compile(class FMaterialCompiler* Compile
 		MCT_MaterialAttributes,
 		MCT_TextureExternal,
 		MCT_Bool,
-		MCT_Strata
+		MCT_Substrate
 	};
 	check(InputType < FunctionInput_MAX);
 
@@ -17350,7 +17350,7 @@ uint32 UMaterialExpressionFunctionInput::GetInputType(int32 InputIndex)
 	case FunctionInput_MaterialAttributes:
 		return MCT_MaterialAttributes;
 	case FunctionInput_Substrate:
-		return MCT_Strata;
+		return MCT_Substrate;
 	default:
 		return MCT_Unknown;
 	}
@@ -17460,7 +17460,7 @@ void UMaterialExpressionFunctionOutput::GetExpressionToolTip(TArray<FString>& Ou
 uint32 UMaterialExpressionFunctionOutput::GetInputType(int32 InputIndex)
 {
 	// Acceptable types for material function outputs
-	return MCT_Float | MCT_MaterialAttributes | MCT_Strata;
+	return MCT_Float | MCT_MaterialAttributes | MCT_Substrate;
 }
 
 int32 UMaterialExpressionFunctionOutput::Compile(class FMaterialCompiler* Compiler, int32 OutputIndex)
@@ -23360,7 +23360,7 @@ void UMaterialExpressionSubstrateShadingModels::GetCaption(TArray<FString>& OutC
 
 uint32 UMaterialExpressionSubstrateShadingModels::GetOutputType(int32 OutputIndex)
 {
-	return MCT_Strata;
+	return MCT_Substrate;
 }
 
 uint32 UMaterialExpressionSubstrateShadingModels::GetInputType(int32 InputIndex)
@@ -23836,7 +23836,7 @@ void UMaterialExpressionSubstrateSlabBSDF::GetExpressionToolTip(TArray<FString>&
 
 uint32 UMaterialExpressionSubstrateSlabBSDF::GetOutputType(int32 OutputIndex)
 {
-	return MCT_Strata;
+	return MCT_Substrate;
 }
 
 uint32 UMaterialExpressionSubstrateSlabBSDF::GetInputType(int32 InputIndex)
@@ -24190,7 +24190,7 @@ void UMaterialExpressionSubstrateSimpleClearCoatBSDF::GetCaption(TArray<FString>
 
 uint32 UMaterialExpressionSubstrateSimpleClearCoatBSDF::GetOutputType(int32 OutputIndex)
 {
-	return MCT_Strata;
+	return MCT_Substrate;
 }
 
 uint32 UMaterialExpressionSubstrateSimpleClearCoatBSDF::GetInputType(int32 InputIndex)
@@ -24332,7 +24332,7 @@ void UMaterialExpressionSubstrateVolumetricFogCloudBSDF::GetCaption(TArray<FStri
 
 uint32 UMaterialExpressionSubstrateVolumetricFogCloudBSDF::GetOutputType(int32 OutputIndex)
 {
-	return MCT_Strata;
+	return MCT_Substrate;
 }
 
 uint32 UMaterialExpressionSubstrateVolumetricFogCloudBSDF::GetInputType(int32 InputIndex)
@@ -24411,7 +24411,7 @@ void UMaterialExpressionSubstrateLightFunction::GetCaption(TArray<FString>& OutC
 
 uint32 UMaterialExpressionSubstrateLightFunction::GetOutputType(int32 OutputIndex)
 {
-	return MCT_Strata;
+	return MCT_Substrate;
 }
 
 uint32 UMaterialExpressionSubstrateLightFunction::GetInputType(int32 InputIndex)
@@ -24484,7 +24484,7 @@ void UMaterialExpressionSubstratePostProcess::GetCaption(TArray<FString>& OutCap
 
 uint32 UMaterialExpressionSubstratePostProcess::GetOutputType(int32 OutputIndex)
 {
-	return MCT_Strata;
+	return MCT_Substrate;
 }
 
 uint32 UMaterialExpressionSubstratePostProcess::GetInputType(int32 InputIndex)
@@ -24560,7 +24560,7 @@ void UMaterialExpressionSubstrateUI::GetCaption(TArray<FString>& OutCaptions) co
 
 uint32 UMaterialExpressionSubstrateUI::GetOutputType(int32 OutputIndex)
 {
-	return MCT_Strata;
+	return MCT_Substrate;
 }
 
 uint32 UMaterialExpressionSubstrateUI::GetInputType(int32 InputIndex)
@@ -24656,7 +24656,7 @@ void UMaterialExpressionSubstrateConvertToDecal::GetCaption(TArray<FString>& Out
 
 uint32 UMaterialExpressionSubstrateConvertToDecal::GetOutputType(int32 OutputIndex)
 {
-	return MCT_Strata;
+	return MCT_Substrate;
 }
 
 uint32 UMaterialExpressionSubstrateConvertToDecal::GetInputType(int32 InputIndex)
@@ -24664,7 +24664,7 @@ uint32 UMaterialExpressionSubstrateConvertToDecal::GetInputType(int32 InputIndex
 	switch (InputIndex)
 	{
 	case 0:
-		return MCT_Strata;
+		return MCT_Substrate;
 		break;
 	case 1:
 		return MCT_Float1;
@@ -24757,7 +24757,7 @@ void UMaterialExpressionSubstrateUnlitBSDF::GetCaption(TArray<FString>& OutCapti
 
 uint32 UMaterialExpressionSubstrateUnlitBSDF::GetOutputType(int32 OutputIndex)
 {
-	return MCT_Strata;
+	return MCT_Substrate;
 }
 
 uint32 UMaterialExpressionSubstrateUnlitBSDF::GetInputType(int32 InputIndex)
@@ -24855,7 +24855,7 @@ void UMaterialExpressionSubstrateHairBSDF::GetCaption(TArray<FString>& OutCaptio
 
 uint32 UMaterialExpressionSubstrateHairBSDF::GetOutputType(int32 OutputIndex)
 {
-	return MCT_Strata;
+	return MCT_Substrate;
 }
 
 uint32 UMaterialExpressionSubstrateHairBSDF::GetInputType(int32 InputIndex)
@@ -24976,7 +24976,7 @@ void UMaterialExpressionSubstrateEyeBSDF::GetCaption(TArray<FString>& OutCaption
 
 uint32 UMaterialExpressionSubstrateEyeBSDF::GetOutputType(int32 OutputIndex)
 {
-	return MCT_Strata;
+	return MCT_Substrate;
 }
 
 uint32 UMaterialExpressionSubstrateEyeBSDF::GetInputType(int32 InputIndex)
@@ -25087,7 +25087,7 @@ void UMaterialExpressionSubstrateSingleLayerWaterBSDF::GetCaption(TArray<FString
 
 uint32 UMaterialExpressionSubstrateSingleLayerWaterBSDF::GetOutputType(int32 OutputIndex)
 {
-	return MCT_Strata;
+	return MCT_Substrate;
 }
 
 uint32 UMaterialExpressionSubstrateSingleLayerWaterBSDF::GetInputType(int32 InputIndex)
@@ -25271,12 +25271,12 @@ void UMaterialExpressionSubstrateHorizontalMixing::GetCaption(TArray<FString>& O
 
 uint32 UMaterialExpressionSubstrateHorizontalMixing::GetOutputType(int32 OutputIndex)
 {
-	return MCT_Strata;
+	return MCT_Substrate;
 }
 
 uint32 UMaterialExpressionSubstrateHorizontalMixing::GetInputType(int32 InputIndex)
 {
-	return InputIndex == 2 ? MCT_Float1 : MCT_Strata;
+	return InputIndex == 2 ? MCT_Float1 : MCT_Substrate;
 }
 
 bool UMaterialExpressionSubstrateHorizontalMixing::IsResultSubstrateMaterial(int32 OutputIndex)
@@ -25465,7 +25465,7 @@ FName UMaterialExpressionSubstrateVerticalLayering::GetInputName(int32 InputInde
 
 uint32 UMaterialExpressionSubstrateVerticalLayering::GetOutputType(int32 OutputIndex)
 {
-	return MCT_Strata;
+	return MCT_Substrate;
 }
 
 uint32 UMaterialExpressionSubstrateVerticalLayering::GetInputType(int32 InputIndex)
@@ -25474,7 +25474,7 @@ uint32 UMaterialExpressionSubstrateVerticalLayering::GetInputType(int32 InputInd
 	{
 		return MCT_Float;
 	}
-	return MCT_Strata;
+	return MCT_Substrate;
 }
 
 bool UMaterialExpressionSubstrateVerticalLayering::IsResultSubstrateMaterial(int32 OutputIndex)
@@ -25646,12 +25646,12 @@ void UMaterialExpressionSubstrateAdd::GetCaption(TArray<FString>& OutCaptions) c
 
 uint32 UMaterialExpressionSubstrateAdd::GetOutputType(int32 OutputIndex)
 {
-	return MCT_Strata;
+	return MCT_Substrate;
 }
 
 uint32 UMaterialExpressionSubstrateAdd::GetInputType(int32 InputIndex)
 {
-	return MCT_Strata;
+	return MCT_Substrate;
 }
 
 bool UMaterialExpressionSubstrateAdd::IsResultSubstrateMaterial(int32 OutputIndex)
@@ -25784,14 +25784,14 @@ void UMaterialExpressionSubstrateWeight::GetCaption(TArray<FString>& OutCaptions
 
 uint32 UMaterialExpressionSubstrateWeight::GetOutputType(int32 OutputIndex)
 {
-	return MCT_Strata;
+	return MCT_Substrate;
 }
 
 uint32 UMaterialExpressionSubstrateWeight::GetInputType(int32 InputIndex)
 {
 	if (InputIndex == 0)
 	{
-		return MCT_Strata;
+		return MCT_Substrate;
 	}
 	return MCT_Float1;
 }
diff --git a/Engine/Source/Runtime/Engine/Private/Materials/MaterialHLSLEmitter.cpp b/Engine/Source/Runtime/Engine/Private/Materials/MaterialHLSLEmitter.cpp
--- a/Engine/Source/Runtime/Engine/Private/Materials/MaterialHLSLEmitter.cpp
+++ b/Engine/Source/Runtime/Engine/Private/Materials/MaterialHLSLEmitter.cpp
@@ -47,7 +47,7 @@ static const TCHAR* HLSLTypeString(EMaterialValueType Type)
 	case MCT_UInt2:					return TEXT("uint2");
 	case MCT_UInt3:					return TEXT("uint3");
 	case MCT_UInt4:					return TEXT("uint4");
-	case MCT_Strata:				return TEXT("FSubstrateData");
+	case MCT_Substrate:				return TEXT("FSubstrateData");
 	default:						return TEXT("unknown");
 	};
 }
@@ -145,7 +145,7 @@ static FString GenerateMaterialTemplateHLSL(EShaderPlatform ShaderPlatform,
 		case MCT_UInt2:  HLSLType = TEXT("uint2"); break;
 		case MCT_UInt3:  HLSLType = TEXT("uint3"); break;
 		case MCT_UInt4:  HLSLType = TEXT("uint4"); break;
-		case MCT_Strata: HLSLType = TEXT("FSubstrateData"); break;
+		case MCT_Substrate: HLSLType = TEXT("FSubstrateData"); break;
 		default: break;
 		}
 
@@ -167,7 +167,7 @@ static FString GenerateMaterialTemplateHLSL(EShaderPlatform ShaderPlatform,
 			case MCT_Float3: MaterialAttributesDefault += FString::Printf(TEXT("\tResult.%s = MaterialFloat3(%0.8f,%0.8f,%0.8f);") LINE_TERMINATOR, *PropertyName, DefaultValue.X, DefaultValue.Y, DefaultValue.Z); break;
 			case MCT_Float4: MaterialAttributesDefault += FString::Printf(TEXT("\tResult.%s = MaterialFloat4(%0.8f,%0.8f,%0.8f,%0.8f);") LINE_TERMINATOR, *PropertyName, DefaultValue.X, DefaultValue.Y, DefaultValue.Z, DefaultValue.W); break;
 			case MCT_ShadingModel: MaterialAttributesDefault += FString::Printf(TEXT("\tResult.%s = %d;") LINE_TERMINATOR, *PropertyName, (int32)DefaultShadingModel); break;
-			case MCT_Strata: MaterialAttributesDefault += FString::Printf(TEXT("\tResult.%s = GetInitialisedSubstrateData();") LINE_TERMINATOR, *PropertyName); break; // TODO
+			case MCT_Substrate: MaterialAttributesDefault += FString::Printf(TEXT("\tResult.%s = GetInitialisedSubstrateData();") LINE_TERMINATOR, *PropertyName); break; // TODO
 			default: checkNoEntry(); break;
 			}
 		}
diff --git a/Engine/Source/Runtime/Engine/Private/Materials/MaterialShared.cpp b/Engine/Source/Runtime/Engine/Private/Materials/MaterialShared.cpp
--- a/Engine/Source/Runtime/Engine/Private/Materials/MaterialShared.cpp
+++ b/Engine/Source/Runtime/Engine/Private/Materials/MaterialShared.cpp
@@ -661,7 +661,7 @@ int32 FStrataMaterialInput::CompileWithDefault(class FMaterialCompiler* Compiler
 		}
 	}
 
-	return Compiler->ForceCast(FMaterialAttributeDefinitionMap::CompileDefaultExpression(Compiler, Property), MCT_Strata);
+	return Compiler->ForceCast(FMaterialAttributeDefinitionMap::CompileDefaultExpression(Compiler, Property), MCT_Substrate);
 }
 
 int32 FVectorMaterialInput::CompileWithDefault(class FMaterialCompiler* Compiler, EMaterialProperty Property)
diff --git a/Engine/Source/Runtime/Engine/Public/MaterialValueType.h b/Engine/Source/Runtime/Engine/Public/MaterialValueType.h
--- a/Engine/Source/Runtime/Engine/Public/MaterialValueType.h
+++ b/Engine/Source/Runtime/Engine/Public/MaterialValueType.h
@@ -45,7 +45,7 @@ enum EMaterialValueType
 	
 	MCT_ShadingModel		  = 1 << 16,
 
-	MCT_Strata			  = 1 << 17,
+	MCT_Substrate			  = 1 << 17,
 
 	MCT_LWCScalar			  = 1 << 18,
 	MCT_LWCVector2			  = 1 << 19,
