Substrate - renaming - FSubstrateMaterialInput

diff --git a/Engine/Config/BaseEngine.ini b/Engine/Config/BaseEngine.ini
--- a/Engine/Config/BaseEngine.ini
+++ b/Engine/Config/BaseEngine.ini
@@ -1204,6 +1204,7 @@ DefaultWorldPartitionHLODLayer=/Engine/Maps/Templates/HLODs/HLODLayer_Instanced.
 +ClassRedirects=(OldName="/Script/Engine.MaterialExpressionStrataMetalnessToDiffuseAlbedoF0",	NewName="/Script/Engine.MaterialExpressionSubstrateMetalnessToDiffuseAlbedoF0")
 +ClassRedirects=(OldName="/Script/Engine.MaterialExpressionStrataHazinessToSecondaryRoughness",	NewName="/Script/Engine.MaterialExpressionSubstrateHazinessToSecondaryRoughness")
 +ClassRedirects=(OldName="/Script/Engine.MaterialExpressionStrataThinFilm",						NewName="/Script/Engine.MaterialExpressionSubstrateThinFilm")
++StructRedirects=(OldName="/Script/Engine.StrataMaterialInput",									NewName="/Script/Engine.SubstrateMaterialInput")
 
 ; End of releases redirects
 
diff --git a/Engine/Source/Runtime/Engine/Classes/Materials/Material.h b/Engine/Source/Runtime/Engine/Classes/Materials/Material.h
--- a/Engine/Source/Runtime/Engine/Classes/Materials/Material.h
+++ b/Engine/Source/Runtime/Engine/Classes/Materials/Material.h
@@ -234,7 +234,7 @@ struct FShadingModelMaterialInput : public FMaterialInput
 
 #if !CPP      //noexport struct
 USTRUCT(noexport)
-struct FStrataMaterialInput : public FMaterialInput
+struct FSubstrateMaterialInput : public FMaterialInput
 {
 	// No support for constant
 };
@@ -395,7 +395,7 @@ public:
 	FScalarMaterialInput SurfaceThickness;
 
 	UPROPERTY()
-	FStrataMaterialInput FrontMaterial;
+	FSubstrateMaterialInput FrontMaterial;
 
 	UPROPERTY()
 	FMaterialExpressionCollection ExpressionCollection;
@@ -1988,7 +1988,7 @@ private:
 	FShadingModelMaterialInput ShadingModelFromMaterialExpression_DEPRECATED;
 
 	UPROPERTY()
-	FStrataMaterialInput FrontMaterial_DEPRECATED;
+	FSubstrateMaterialInput FrontMaterial_DEPRECATED;
 #endif // WITH_EDITORONLY_DATA
 };
 
diff --git a/Engine/Source/Runtime/Engine/Private/Materials/HLSLMaterialTranslator.cpp b/Engine/Source/Runtime/Engine/Private/Materials/HLSLMaterialTranslator.cpp
--- a/Engine/Source/Runtime/Engine/Private/Materials/HLSLMaterialTranslator.cpp
+++ b/Engine/Source/Runtime/Engine/Private/Materials/HLSLMaterialTranslator.cpp
@@ -902,7 +902,7 @@ bool FHLSLMaterialTranslator::Translate()
 		}
 		else
 		{
-			FStrataMaterialInput* FrontMaterialInput = Material->GetMaterialInterface() ? &Material->GetMaterialInterface()->GetMaterial()->GetEditorOnlyData()->FrontMaterial : nullptr;
+			FSubstrateMaterialInput* FrontMaterialInput = Material->GetMaterialInterface() ? &Material->GetMaterialInterface()->GetMaterial()->GetEditorOnlyData()->FrontMaterial : nullptr;
 			FrontMaterialExpr = FrontMaterialInput ? FrontMaterialInput->GetTracedInput().Expression : nullptr;
 			FrontMaterialOutputIndex = FrontMaterialInput ? FrontMaterialInput->OutputIndex : INDEX_NONE;
 		}
diff --git a/Engine/Source/Runtime/Engine/Private/Materials/Material.cpp b/Engine/Source/Runtime/Engine/Private/Materials/Material.cpp
--- a/Engine/Source/Runtime/Engine/Private/Materials/Material.cpp
+++ b/Engine/Source/Runtime/Engine/Private/Materials/Material.cpp
@@ -5785,7 +5785,7 @@ static void SetMaterialInputDescription(FMaterialAttributesInput& Input, bool bH
 	OutDescription.bHidden = bHidden;
 }
 
-static void SetMaterialInputDescription(FStrataMaterialInput& Input, bool bHidden, FMaterialInputDescription& OutDescription)
+static void SetMaterialInputDescription(FSubstrateMaterialInput& Input, bool bHidden, FMaterialInputDescription& OutDescription)
 {
 	OutDescription.Type = UE::Shader::EValueType::Void;
 	OutDescription.Input = &Input;
diff --git a/Engine/Source/Runtime/Engine/Private/Materials/MaterialShared.cpp b/Engine/Source/Runtime/Engine/Private/Materials/MaterialShared.cpp
--- a/Engine/Source/Runtime/Engine/Private/Materials/MaterialShared.cpp
+++ b/Engine/Source/Runtime/Engine/Private/Materials/MaterialShared.cpp
@@ -578,7 +578,7 @@ bool FShadingModelMaterialInput::Serialize(FArchive& Ar)
 	return SerializeMaterialInput<uint32>(Ar, *this);
 }
 
-bool FStrataMaterialInput::Serialize(FArchive& Ar)
+bool FSubstrateMaterialInput::Serialize(FArchive& Ar)
 {
 	return SerializeMaterialInput<uint32>(Ar, *this);
 }
@@ -650,7 +650,7 @@ int32 FShadingModelMaterialInput::CompileWithDefault(class FMaterialCompiler* Co
 	return Compiler->ForceCast(FMaterialAttributeDefinitionMap::CompileDefaultExpression(Compiler, Property), MCT_ShadingModel, MFCF_ExactMatch);
 }
 
-int32 FStrataMaterialInput::CompileWithDefault(class FMaterialCompiler* Compiler, EMaterialProperty Property)
+int32 FSubstrateMaterialInput::CompileWithDefault(class FMaterialCompiler* Compiler, EMaterialProperty Property)
 {
 	if (Expression)
 	{
diff --git a/Engine/Source/Runtime/Engine/Public/MaterialExpressionIO.h b/Engine/Source/Runtime/Engine/Public/MaterialExpressionIO.h
--- a/Engine/Source/Runtime/Engine/Public/MaterialExpressionIO.h
+++ b/Engine/Source/Runtime/Engine/Public/MaterialExpressionIO.h
@@ -243,7 +243,7 @@ struct TStructOpsTypeTraits<FShadingModelMaterialInput>
 	};
 };
 
-struct FStrataMaterialInput : FMaterialInput<uint32> // Still giving it a default type
+struct FSubstrateMaterialInput : FMaterialInput<uint32> // Still giving it a default type
 {
 #if WITH_EDITOR
 	ENGINE_API int32 CompileWithDefault(class FMaterialCompiler* Compiler, EMaterialProperty Property);
@@ -253,8 +253,8 @@ struct FStrataMaterialInput : FMaterialInput<uint32> // Still giving it a defaul
 };
 
 template<>
-struct TStructOpsTypeTraits<FStrataMaterialInput>
-	: public TStructOpsTypeTraitsBase2<FStrataMaterialInput>
+struct TStructOpsTypeTraits<FSubstrateMaterialInput>
+	: public TStructOpsTypeTraitsBase2<FSubstrateMaterialInput>
 {
 	enum
 	{
