Reinstancing fix, Tweak a bit an other change that I have made a while back. Let the optional mapping attempt first to replace the instances before we clear the external references.

diff --git a/Engine/Source/Runtime/Engine/Private/UnrealEngine.cpp b/Engine/Source/Runtime/Engine/Private/UnrealEngine.cpp
--- a/Engine/Source/Runtime/Engine/Private/UnrealEngine.cpp
+++ b/Engine/Source/Runtime/Engine/Private/UnrealEngine.cpp
@@ -17127,11 +17127,11 @@ void UEngine::CopyPropertiesForUnrelatedObjects(UObject* OldObject, UObject* New
 		}
 		else
 		{
-			FArchiveReplaceOrClearExternalReferences<UObject> ReplaceInCDOAr(NewObject, ReferenceReplacementMap, NewPackage);
 			if (Params.OptionalReplacementMappings)
 			{
-				FArchiveReplaceOrClearExternalReferences<UObject> OptionalReplaceInCDOAr(NewObject, *Params.OptionalReplacementMappings, NewPackage);
+				FArchiveReplaceObjectRef<UObject> OptionalReplaceInCDOAr(NewObject, *Params.OptionalReplacementMappings);
 			}
+			FArchiveReplaceOrClearExternalReferences<UObject> ReplaceInCDOAr(NewObject, ReferenceReplacementMap, NewPackage);
 		}
 
 		// Replace references inside each individual component. This is always required because if something is in ReferenceReplacementMap, the above replace code will skip fixing child properties
@@ -17144,11 +17144,11 @@ void UEngine::CopyPropertiesForUnrelatedObjects(UObject* OldObject, UObject* New
 			}
 			else
 			{
-				FArchiveReplaceOrClearExternalReferences<UObject> ReplaceInComponentAr(NewComponent, ReferenceReplacementMap, NewPackage);
 				if (Params.OptionalReplacementMappings)
 				{
-					FArchiveReplaceOrClearExternalReferences<UObject> OptionalReplaceInCDOAr(NewComponent, *Params.OptionalReplacementMappings, NewPackage);
+					FArchiveReplaceObjectRef<UObject> OptionalReplaceInCDOAr(NewComponent, *Params.OptionalReplacementMappings);
 				}
+				FArchiveReplaceOrClearExternalReferences<UObject> ReplaceInComponentAr(NewComponent, ReferenceReplacementMap, NewPackage);
 			}
 		}
 	}
