Merge Actor - Approximate Actors: Expose the FallbackRelativeError settings (Nanite option)

diff --git a/Engine/Plugins/Experimental/MeshModelingToolsetExp/Source/GeometryProcessingAdapters/Private/ApproximateActorsImpl.cpp b/Engine/Plugins/Experimental/MeshModelingToolsetExp/Source/GeometryProcessingAdapters/Private/ApproximateActorsImpl.cpp
--- a/Engine/Plugins/Experimental/MeshModelingToolsetExp/Source/GeometryProcessingAdapters/Private/ApproximateActorsImpl.cpp
+++ b/Engine/Plugins/Experimental/MeshModelingToolsetExp/Source/GeometryProcessingAdapters/Private/ApproximateActorsImpl.cpp
@@ -969,7 +969,10 @@ IGeometryProcessing_ApproximateActors::FOptions FApproximateActorsImpl::Construc
 
 	// Nanite settings
 	Options.bGenerateNaniteEnabledMesh = UseSettings.bGenerateNaniteEnabledMesh;
-	Options.NaniteProxyTrianglePercent = UseSettings.NaniteProxyTrianglePercent;
+	Options.NaniteFallbackTarget = UseSettings.NaniteFallbackTarget == ::ENaniteFallbackTarget::Auto ? IGeometryProcessing_ApproximateActors::ENaniteFallbackTarget::Auto :
+								  (UseSettings.NaniteFallbackTarget == ::ENaniteFallbackTarget::PercentTriangles ? IGeometryProcessing_ApproximateActors::ENaniteFallbackTarget::PercentTriangles : IGeometryProcessing_ApproximateActors::ENaniteFallbackTarget::RelativeError);
+	Options.NaniteFallbackPercentTriangles = UseSettings.NaniteFallbackPercentTriangles;
+	Options.NaniteFallbackRelativeError = UseSettings.NaniteFallbackRelativeError;
 
 	// Distance field
 	Options.bAllowDistanceField = UseSettings.bAllowDistanceField;
@@ -1294,7 +1297,10 @@ UStaticMesh* FApproximateActorsImpl::EmitGeneratedMeshAsset(
 
 	MeshAssetOptions.bGenerateNaniteEnabledMesh = Options.bGenerateNaniteEnabledMesh;
 	MeshAssetOptions.NaniteSettings.bEnabled = Options.bGenerateNaniteEnabledMesh;
-	MeshAssetOptions.NaniteSettings.FallbackPercentTriangles = Options.NaniteProxyTrianglePercent / 100.0;	// NaniteSettings wants value in range 0-1
+	MeshAssetOptions.NaniteSettings.FallbackTarget = Options.NaniteFallbackTarget == IGeometryProcessing_ApproximateActors::ENaniteFallbackTarget::Auto ? ::ENaniteFallbackTarget::Auto : 
+													 (Options.NaniteFallbackTarget == IGeometryProcessing_ApproximateActors::ENaniteFallbackTarget::PercentTriangles ? ::ENaniteFallbackTarget::PercentTriangles : ::ENaniteFallbackTarget::RelativeError);
+	MeshAssetOptions.NaniteSettings.FallbackPercentTriangles = Options.NaniteFallbackPercentTriangles;
+	MeshAssetOptions.NaniteSettings.FallbackRelativeError = Options.NaniteFallbackRelativeError;
 
 	MeshAssetOptions.bSupportRayTracing = Options.bSupportRayTracing;
 	MeshAssetOptions.bAllowDistanceField = Options.bAllowDistanceField;
diff --git a/Engine/Source/Developer/GeometryProcessingInterfaces/Public/GeometryProcessingInterfaces/ApproximateActors.h b/Engine/Source/Developer/GeometryProcessingInterfaces/Public/GeometryProcessingInterfaces/ApproximateActors.h
--- a/Engine/Source/Developer/GeometryProcessingInterfaces/Public/GeometryProcessingInterfaces/ApproximateActors.h
+++ b/Engine/Source/Developer/GeometryProcessingInterfaces/Public/GeometryProcessingInterfaces/ApproximateActors.h
@@ -6,6 +6,7 @@
 #include "Features/IModularFeature.h"
 
 class AActor;
+class UActorComponent;
 class UStaticMesh;
 class UMaterial;
 class UMaterialInterface;
@@ -91,6 +92,13 @@ public:
 		TexelDensity = 1
 	};
 
+	enum class ENaniteFallbackTarget : uint8
+	{
+		Auto,
+		PercentTriangles,
+		RelativeError
+	};
+
 
 	/**
 	 * Input options to Actor Approximation process
@@ -208,8 +216,14 @@ public:
 		// Whether to generate a nanite-enabled mesh
 		bool bGenerateNaniteEnabledMesh = false;
 
-		// Percentage of triangles to retain for the nanite coarse mesh
-		float NaniteProxyTrianglePercent = 100;
+		/** Which heuristic to use when generating the fallback mesh. */
+		ENaniteFallbackTarget NaniteFallbackTarget = ENaniteFallbackTarget::Auto;
+	
+		/** Percentage of triangles to keep from source mesh for fallback. 1.0 = no reduction, 0.0 = no triangles. */
+		float NaniteFallbackPercentTriangles = 1.0f;
+
+		/** Reduce until at least this amount of error is reached relative to size of the mesh */
+		float NaniteFallbackRelativeError = 1.0f;
 
 		// Whether ray tracing will be supported on this mesh. Disable this to save memory if the generated mesh will only be rendered in the distance
 		bool bSupportRayTracing = true;
diff --git a/Engine/Source/Runtime/Engine/Classes/Engine/MeshMerging.h b/Engine/Source/Runtime/Engine/Classes/Engine/MeshMerging.h
--- a/Engine/Source/Runtime/Engine/Classes/Engine/MeshMerging.h
+++ b/Engine/Source/Runtime/Engine/Classes/Engine/MeshMerging.h
@@ -894,9 +894,17 @@ struct FMeshApproximationSettings
 	UPROPERTY(BlueprintReadWrite, EditAnywhere, Category = MeshSettings)
 	bool bGenerateNaniteEnabledMesh = false;
 
-	/** Percentage of triangles to reduce down to for generating a coarse proxy mesh from the Nanite mesh */
-	UPROPERTY(BlueprintReadWrite, EditAnywhere, AdvancedDisplay, Category = MeshSettings, meta = (EditConditionHides, EditCondition = "bGenerateNaniteEnabledMesh", ClampMin = 0, ClampMax = 100))
-	float NaniteProxyTrianglePercent = 0;
+	/** Which heuristic to use when generating the Nanite fallback mesh. */
+	UPROPERTY(BlueprintReadWrite, EditAnywhere, AdvancedDisplay, Category = MeshSettings, meta = (EditConditionHides, EditCondition = "bGenerateNaniteEnabledMesh"))
+	ENaniteFallbackTarget NaniteFallbackTarget = ENaniteFallbackTarget::Auto;
+
+	/** Percentage of triangles to keep from source Nanite mesh for fallback. 1.0 = no reduction, 0.0 = no triangles. */
+	UPROPERTY(BlueprintReadWrite, EditAnywhere, AdvancedDisplay, Category = MeshSettings, meta = (EditConditionHides, EditCondition = "bGenerateNaniteEnabledMesh && NaniteFallbackTarget == ENaniteFallbackTarget::PercentTriangles", ClampMin = 0, ClampMax = 1))
+	float NaniteFallbackPercentTriangles = 1.0f;
+
+	/** Reduce Nanite fallback mesh until at least this amount of error is reached relative to size of the mesh. */
+	UPROPERTY(BlueprintReadWrite, EditAnywhere, AdvancedDisplay, Category = MeshSettings, meta = (EditConditionHides, EditCondition = "bGenerateNaniteEnabledMesh && NaniteFallbackTarget == ENaniteFallbackTarget::RelativeError", ClampMin = 0))
+	float NaniteFallbackRelativeError = 1.0f;
 
 	/** Whether ray tracing will be supported on this mesh. Disable this to save memory if the generated mesh will only be rendered in the distance. */
 	UPROPERTY(BlueprintReadWrite, EditAnywhere, Category = MeshSettings)
@@ -979,7 +987,9 @@ struct FMeshApproximationSettings
 			&& TrianglesPerM == Other.TrianglesPerM
 			&& GeometricDeviation == Other.GeometricDeviation
 			&& bGenerateNaniteEnabledMesh == Other.bGenerateNaniteEnabledMesh
-			&& NaniteProxyTrianglePercent == Other.NaniteProxyTrianglePercent
+			&& NaniteFallbackTarget == Other.NaniteFallbackTarget
+			&& NaniteFallbackPercentTriangles == Other.NaniteFallbackPercentTriangles
+			&& NaniteFallbackRelativeError == Other.NaniteFallbackRelativeError
 			&& bSupportRayTracing == Other.bSupportRayTracing
 			&& bAllowDistanceField == Other.bAllowDistanceField
 			&& MultiSamplingAA == Other.MultiSamplingAA
@@ -996,4 +1006,26 @@ struct FMeshApproximationSettings
 	{
 		return !(*this == Other);
 	}
+
+#if WITH_EDITORONLY_DATA
+	/** Handles deprecated properties */
+	void PostSerialize(const FArchive& Ar);
+#endif
+
+private:
+#if WITH_EDITORONLY_DATA
+	UPROPERTY()
+	float NaniteProxyTrianglePercent_DEPRECATED = 0;
+#endif
+};
+
+template<>
+struct TStructOpsTypeTraits<FMeshApproximationSettings> : public TStructOpsTypeTraitsBase2<FMeshApproximationSettings>
+{
+#if WITH_EDITORONLY_DATA
+	enum
+	{
+		WithPostSerialize = true,
+	};
+#endif
 };
diff --git a/Engine/Source/Runtime/Engine/Private/EngineTypes.cpp b/Engine/Source/Runtime/Engine/Private/EngineTypes.cpp
--- a/Engine/Source/Runtime/Engine/Private/EngineTypes.cpp
+++ b/Engine/Source/Runtime/Engine/Private/EngineTypes.cpp
@@ -79,6 +79,19 @@ void FMeshMergingSettings::PostSerialize(const FArchive& Ar)
 		}
 	}
 }
+
+void FMeshApproximationSettings::PostSerialize(const FArchive& Ar)
+{
+	if (Ar.IsLoading())
+	{
+		FMeshApproximationSettings DefaultObject;
+		if (NaniteProxyTrianglePercent_DEPRECATED != DefaultObject.NaniteProxyTrianglePercent_DEPRECATED)
+		{
+			NaniteFallbackTarget = ENaniteFallbackTarget::PercentTriangles;
+			NaniteFallbackPercentTriangles = NaniteProxyTrianglePercent_DEPRECATED / 100.0f;
+		}
+	}
+}
 #endif
 
 UEngineBaseTypes::UEngineBaseTypes(const FObjectInitializer& ObjectInitializer)
