[shaders]
- modify ddc key debug functionality to write files alongside other shader debug info, so it can be easily uploaded alongside other debug info build artifacts on CIS (which redirects shader debug info output to the automation folder)
- always dump WorldGridMaterial DDC key using above functionality; "Log" severity messages seem to no longer be printed to the logs on CIS so this will give us another path to getting this key for any CIS build without log spam

diff --git a/Engine/Source/Runtime/Engine/Private/Materials/MaterialShader.cpp b/Engine/Source/Runtime/Engine/Private/Materials/MaterialShader.cpp
--- a/Engine/Source/Runtime/Engine/Private/Materials/MaterialShader.cpp
+++ b/Engine/Source/Runtime/Engine/Private/Materials/MaterialShader.cpp
@@ -1384,6 +1384,11 @@ void FMaterialShaderMap::LoadFromDerivedDataCache(const FMaterial* Material, con
 	InOutShaderMap = BeginLoadFromDerivedDataCache(Material, ShaderMapId, InPlatform, TargetPlatform, InOutShaderMap, OutDDCKeyDesc)->Get();
 }
 
+FString GetMaterialDebugGroupName(EShaderPlatform ShaderPlatform, const FMaterial* Material, const FMaterialShaderMapId& ShaderMapId)
+{
+	return Material->GetUniqueAssetName(ShaderPlatform, ShaderMapId) / LexToString(Material->GetQualityLevel());
+}
+
 TSharedRef<FMaterialShaderMap::FAsyncLoadContext> FMaterialShaderMap::BeginLoadFromDerivedDataCache(const FMaterial* Material, const FMaterialShaderMapId& ShaderMapId, EShaderPlatform InPlatform, const ITargetPlatform* TargetPlatform, TRefCountPtr<FMaterialShaderMap>& InShaderMap, FString& OutDDCKeyDesc)
 {
 	TRACE_CPUPROFILER_EVENT_SCOPE(FMaterialShaderMap::BeginLoadFromDerivedDataCache);
@@ -1482,13 +1487,9 @@ TSharedRef<FMaterialShaderMap::FAsyncLoadContext> FMaterialShaderMap::BeginLoadF
 			FCacheKey CacheKey = GetMaterialShaderMapKey(Result->DataKey);
 			OutDDCKeyDesc = LexToString(CacheKey.Hash);
 
-			if (UNLIKELY(ShouldDumpShaderDDCKeys()))
-			{
-				const FString FileName = FString::Printf(TEXT("%s-%s-%s-%s.txt"), 
-					*Material->GetAssetName(), *LexToString(ShaderMapId.FeatureLevel), *LexToString(ShaderMapId.QualityLevel),
-					ShaderMapId.LayoutParams.WithEditorOnly() ? TEXT("Editor") : TEXT("Game")
-					);
-				DumpShaderDDCKeyToFile(InPlatform, ShaderMapId.LayoutParams.WithEditorOnly(), FileName, Result->DataKey);
+			if (UNLIKELY(ShouldDumpShaderDDCKeys()) || (Material->IsDefaultMaterial() && Material->GetMaterialDomain() == EMaterialDomain::MD_Surface))
+			{	
+				DumpShaderDDCKeyToFile(InPlatform, ShaderMapId.LayoutParams.WithEditorOnly(), *GetMaterialDebugGroupName(InPlatform, Material, ShaderMapId), Result->DataKey);
 			}
 
 			if (Material->IsDefaultMaterial() && Material->GetMaterialDomain() == EMaterialDomain::MD_Surface)
@@ -1861,7 +1862,7 @@ int32 FMaterialShaderMap::SubmitCompileJobs(uint32 CompilingShaderMapId,
 	const FMaterialShaderParameters MaterialParameters(Material);
 	const FMaterialShaderMapLayout& Layout = AcquireMaterialShaderMapLayout(ShaderPlatform, LocalPermutationFlags, MaterialParameters);
 
-	const FString DebugGroupName = Material->GetUniqueAssetName(ShaderPlatform, GetShaderMapId()) / LexToString(Material->GetQualityLevel());
+	const FString DebugGroupName = GetMaterialDebugGroupName(ShaderPlatform, Material, GetShaderMapId());
 
 #if ALLOW_SHADERMAP_DEBUG_DATA
 	FString DebugExtensionStr(TEXT(""));
diff --git a/Engine/Source/Runtime/Engine/Private/ShaderCompiler/ShaderCompiler.cpp b/Engine/Source/Runtime/Engine/Private/ShaderCompiler/ShaderCompiler.cpp
--- a/Engine/Source/Runtime/Engine/Private/ShaderCompiler/ShaderCompiler.cpp
+++ b/Engine/Source/Runtime/Engine/Private/ShaderCompiler/ShaderCompiler.cpp
@@ -1175,6 +1175,7 @@ bool ShouldDumpShaderDDCKeys()
 
 void DumpShaderDDCKeyToFile(const EShaderPlatform InPlatform, bool bWithEditor, const FString& FileName, const FString& DDCKey)
 {
+	// deprecated version
 	const FString SubDirectory = bWithEditor ? TEXT("Editor") : TEXT("Game");
 	const FString TempPath = FPaths::ProjectSavedDir() / TEXT("ShaderDDCKeys") / SubDirectory / LexToString(InPlatform);
 	IFileManager::Get().MakeDirectory(*TempPath, true);
@@ -1186,6 +1187,18 @@ void DumpShaderDDCKeyToFile(const EShaderPlatform InPlatform, bool bWithEditor,
 	DumpAr->Serialize(const_cast<TCHAR*>(*DDCKey), DDCKey.Len() * sizeof(TCHAR));
 }
 
+void DumpShaderDDCKeyToFile(const EShaderPlatform InPlatform, bool bEditorOnly, const TCHAR* DebugGroupName, const FString& DDCKey)
+{
+	const FString FileName = FString::Printf(TEXT("DDCKey-%s.txt"), bEditorOnly ? TEXT("Editor") : TEXT("Game"));
+
+	const FString TempPath = GShaderCompilingManager->GetAbsoluteShaderDebugInfoDirectory() / FGenericDataDrivenShaderPlatformInfo::GetName(InPlatform).ToString() / DebugGroupName;
+	IFileManager::Get().MakeDirectory(*TempPath, true);
+
+	const FString TempFile = TempPath / FileName;
+
+	FFileHelper::SaveStringToFile(DDCKey, *TempFile);
+}
+
 namespace ShaderCompiler
 {
 	bool IsJobCacheEnabled()
@@ -9342,10 +9355,17 @@ void CompileGlobalShaderMap(EShaderPlatform Platform, const ITargetPlatform* Tar
 
 					if (UNLIKELY(ShouldDumpShaderDDCKeys()))
 					{
-						const FString ShaderName = ShaderFilenameDependencies.Key.Replace(TEXT("/"), TEXT(".")).Replace(TEXT(".usf"), TEXT(""));
-						const FString FileName = FString::Printf(TEXT("GlobalShaderMap%s.txt"), *ShaderName);
 						const FString DataKey = GetGlobalShaderMapKeyString(ShaderMapId, Platform, ShaderFilenameDependencies.Value);
-						DumpShaderDDCKeyToFile(Platform, ShaderMapId.WithEditorOnly(), FileName, DataKey);
+						// For global shaders, we dump the key multiple times (once for each shader type) so they will live on disk alongside
+						// other shader debug artifacts.
+						for (const FShaderTypeDependency& ShaderTypeDependency : ShaderFilenameDependencies.Value)
+						{
+							const FShaderType* ShaderType = FindShaderTypeByName(ShaderTypeDependency.ShaderTypeName);
+							TStringBuilder<128> GroupNameBuilder;
+							GroupNameBuilder << TEXT("Global");
+							FPathViews::Append(GroupNameBuilder, ShaderType->GetName());
+							DumpShaderDDCKeyToFile(Platform, ShaderMapId.WithEditorOnly(), GroupNameBuilder.ToString(), DataKey);
+						}
 					}
 				}
 
diff --git a/Engine/Source/Runtime/Engine/Public/ShaderCompiler.h b/Engine/Source/Runtime/Engine/Public/ShaderCompiler.h
--- a/Engine/Source/Runtime/Engine/Public/ShaderCompiler.h
+++ b/Engine/Source/Runtime/Engine/Public/ShaderCompiler.h
@@ -1152,7 +1152,9 @@ extern ENGINE_API const FString& GetGlobalShaderMapDDCKey();
 extern ENGINE_API const FString& GetMaterialShaderMapDDCKey();
 
 extern ENGINE_API bool ShouldDumpShaderDDCKeys();
+UE_DEPRECATED(5.4, "DumpShaderDDCKeyToFile now takes DebugGroupName as parameter (these files now go into the ShaderDebugInfo folder alongside other debug artifacts).")
 extern ENGINE_API void DumpShaderDDCKeyToFile(const EShaderPlatform InPlatform, bool bWithEditor, const FString& FileName, const FString& DDCKey);
+extern ENGINE_API void DumpShaderDDCKeyToFile(const EShaderPlatform InPlatform, bool bWithEditor, const TCHAR* DebugGroupName, const FString& DDCKey);
 
 /**
 * Handles serializing in MeshMaterialMaps or GlobalShaderMap from a CookOnTheFly command and applying them to the in-memory shadermaps.
