Fix crash in non-editor builds with light actors PostLoad trying to access a null light component that has IsEditorOnly enabled.

diff --git a/Engine/Source/Runtime/Engine/Classes/Engine/Light.h b/Engine/Source/Runtime/Engine/Classes/Engine/Light.h
--- a/Engine/Source/Runtime/Engine/Classes/Engine/Light.h
+++ b/Engine/Source/Runtime/Engine/Classes/Engine/Light.h
@@ -61,6 +61,10 @@ public:
 	// END DEPRECATED
 
 public:
+	//~ Begin UObject Interface.
+	virtual void PostLoad() override;
+	//~ End UObject Interface.
+
 #if WITH_EDITOR
 	//~ Begin AActor Interface.
 	ENGINE_API virtual void CheckForErrors() override;
@@ -79,7 +83,4 @@ public:
 
 	/** Returns LightComponent subobject **/
 	class ULightComponent* GetLightComponent() const { return LightComponent; }
-};
-
-
-
+};
\ No newline at end of file
diff --git a/Engine/Source/Runtime/Engine/Classes/Engine/PointLight.h b/Engine/Source/Runtime/Engine/Classes/Engine/PointLight.h
--- a/Engine/Source/Runtime/Engine/Classes/Engine/PointLight.h
+++ b/Engine/Source/Runtime/Engine/Classes/Engine/PointLight.h
@@ -29,12 +29,8 @@ class APointLight : public ALight
 #endif
 
 	//~ Begin UObject Interface.
-	virtual void PostLoad() override;
 #if WITH_EDITOR
 	virtual void LoadedFromAnotherClass(const FName& OldClassName) override;
 #endif
 	//~ End UObject Interface.
-};
-
-
-
+};
\ No newline at end of file
diff --git a/Engine/Source/Runtime/Engine/Classes/Engine/RectLight.h b/Engine/Source/Runtime/Engine/Classes/Engine/RectLight.h
--- a/Engine/Source/Runtime/Engine/Classes/Engine/RectLight.h
+++ b/Engine/Source/Runtime/Engine/Classes/Engine/RectLight.h
@@ -20,11 +20,4 @@ class ARectLight : public ALight
 	virtual void EditorApplyScale(const FVector& DeltaScale, const FVector* PivotLocation, bool bAltDown, bool bShiftDown, bool bCtrlDown) override;
 	//~ End AActor Interface.
 #endif
-
-	//~ Begin UObject Interface.
-	virtual void PostLoad() override;
-	//~ End UObject Interface.
-};
-
-
-
+};
\ No newline at end of file
diff --git a/Engine/Source/Runtime/Engine/Private/Light.cpp b/Engine/Source/Runtime/Engine/Private/Light.cpp
--- a/Engine/Source/Runtime/Engine/Private/Light.cpp
+++ b/Engine/Source/Runtime/Engine/Private/Light.cpp
@@ -171,6 +171,16 @@ void ALight::SetAffectTranslucentLighting(bool bNewValue)
 	LightComponent->SetAffectTranslucentLighting(bNewValue);
 }
 
+void ALight::PostLoad()
+{
+	Super::PostLoad();
+
+	if (LightComponent && LightComponent->Mobility == EComponentMobility::Static)
+	{
+		LightComponent->ClearLightFunctionMaterial();
+	}
+}
+
 ///////////////////////////////////////////////////////////////////////////////////////////////////////////////////////
 
 APointLight::APointLight(const FObjectInitializer& ObjectInitializer)
@@ -182,16 +192,6 @@ APointLight::APointLight(const FObjectInitializer& ObjectInitializer)
 	RootComponent = PointLightComponent;
 }
 
-void APointLight::PostLoad()
-{
-	Super::PostLoad();
-
-	if (GetLightComponent()->Mobility == EComponentMobility::Static)
-	{
-		GetLightComponent()->ClearLightFunctionMaterial();
-	}
-}
-
 #if WITH_EDITOR
 void APointLight::LoadedFromAnotherClass(const FName& OldClassName)
 {
@@ -272,10 +272,6 @@ void ADirectionalLight::PostLoad()
 {
 	Super::PostLoad();
 
-	if (GetLightComponent()->Mobility == EComponentMobility::Static)
-	{
-		GetLightComponent()->ClearLightFunctionMaterial();
-	}
 #if WITH_EDITORONLY_DATA
 	if(ArrowComponent != nullptr)
 	{
diff --git a/Engine/Source/Runtime/Engine/Private/RectLight.cpp b/Engine/Source/Runtime/Engine/Private/RectLight.cpp
--- a/Engine/Source/Runtime/Engine/Private/RectLight.cpp
+++ b/Engine/Source/Runtime/Engine/Private/RectLight.cpp
@@ -14,16 +14,6 @@ ARectLight::ARectLight(const FObjectInitializer& ObjectInitializer)
 	RootComponent = RectLightComponent;
 }
 
-void ARectLight::PostLoad()
-{
-	Super::PostLoad();
-
-	if (GetLightComponent()->Mobility == EComponentMobility::Static)
-	{
-		GetLightComponent()->ClearLightFunctionMaterial();
-	}
-}
-
 #if WITH_EDITOR
 void ARectLight::EditorApplyScale(const FVector& DeltaScale, const FVector* PivotLocation, bool bAltDown, bool bShiftDown, bool bCtrlDown)
 {
diff --git a/Engine/Source/Runtime/Engine/Private/SpotLight.cpp b/Engine/Source/Runtime/Engine/Private/SpotLight.cpp
--- a/Engine/Source/Runtime/Engine/Private/SpotLight.cpp
+++ b/Engine/Source/Runtime/Engine/Private/SpotLight.cpp
@@ -48,11 +48,6 @@ void ASpotLight::PostLoad()
 {
 	Super::PostLoad();
 
-	if (GetLightComponent()->Mobility == EComponentMobility::Static)
-	{
-		GetLightComponent()->ClearLightFunctionMaterial();
-	}
-
 #if WITH_EDITORONLY_DATA
 	if(ArrowComponent)
 	{
