Add mechanism to ensure AnimationModifiersAssetUserData is always added/part of an AnimSequence, currently it would only be created when opening it in the AnimationEditor
Add script API to AssetUserData interface to check/add user-data (can be used to patch up existing content)

diff --git a/Engine/Source/Editor/AnimationModifiers/Private/AnimationModifiersAssetUserData.cpp b/Engine/Source/Editor/AnimationModifiers/Private/AnimationModifiersAssetUserData.cpp
--- a/Engine/Source/Editor/AnimationModifiers/Private/AnimationModifiersAssetUserData.cpp
+++ b/Engine/Source/Editor/AnimationModifiers/Private/AnimationModifiersAssetUserData.cpp
@@ -10,6 +10,14 @@
 #include "UObject/ObjectPtr.h"
 #include "Animation/AnimSequence.h"
 
+UAnimationModifiersAssetUserData::UAnimationModifiersAssetUserData(const FObjectInitializer& ObjectInitializer) : Super(ObjectInitializer)
+{
+	if (!HasAnyFlags(RF_ClassDefaultObject) && !HasAnyFlags(RF_Transactional))
+	{
+		SetFlags(RF_Transactional);
+	}
+}
+
 void UAnimationModifiersAssetUserData::AddAnimationModifier(UAnimationModifier* Instance)
 {
 	AnimationModifierInstances.Add(Instance);
@@ -99,6 +107,10 @@ void UAnimationModifiersAssetUserData::Serialize(FArchive& Ar)
 void UAnimationModifiersAssetUserData::PostLoad()
 {
 	Super::PostLoad();
+	if (!HasAnyFlags(RF_ClassDefaultObject) && !HasAnyFlags(RF_Transactional))
+	{		
+		SetFlags(RF_Transactional);
+	}
 	RemoveInvalidModifiers();
 }
 
diff --git a/Engine/Source/Editor/AnimationModifiers/Private/AnimationModifiersModule.cpp b/Engine/Source/Editor/AnimationModifiers/Private/AnimationModifiersModule.cpp
--- a/Engine/Source/Editor/AnimationModifiers/Private/AnimationModifiersModule.cpp
+++ b/Engine/Source/Editor/AnimationModifiers/Private/AnimationModifiersModule.cpp
@@ -42,6 +42,7 @@
 #include "ContentBrowserMenuContexts.h"
 #include "ToolMenuDelegates.h"
 #include "ToolMenus.h"
+#include "AssetRegistry/AssetRegistryModule.h"
 
 class UFactory;
 
@@ -72,6 +73,9 @@ void FAnimationModifiersModule::StartupModule()
 			GEditor->GetEditorSubsystem<UImportSubsystem>()->OnAssetPostImport.AddRaw(this, &FAnimationModifiersModule::OnAssetPostImport);
 			GEditor->GetEditorSubsystem<UImportSubsystem>()->OnAssetReimport.AddRaw(this, &FAnimationModifiersModule::OnAssetPostReimport);
 			RegisterMenus();
+
+			const FAssetRegistryModule& AssetRegistryModule = FModuleManager::LoadModuleChecked<FAssetRegistryModule>(TEXT("AssetRegistry"));
+			AssetRegistryModule.Get().OnInMemoryAssetCreated().AddRaw(this, &FAnimationModifiersModule::OnInMemoryAssetCreated);
 		}
 	});
 
@@ -110,6 +114,12 @@ void FAnimationModifiersModule::ShutdownModule()
 		AssetTools.UnregisterAssetTypeActions(AssetAction.ToSharedRef());
 	}
 
+	if (FModuleManager::Get().IsModuleLoaded(TEXT("AssetRegistry")))
+	{
+		FAssetRegistryModule& AssetRegistryModule = FModuleManager::GetModuleChecked<FAssetRegistryModule>(TEXT("AssetRegistry"));
+		AssetRegistryModule.Get().OnInMemoryAssetCreated().RemoveAll(this);
+	}
+	
 	RegisteredApplicationModes.Empty();
 
 	if (GEditor)
@@ -277,6 +287,14 @@ void FAnimationModifiersModule::OnAssetPostReimport(UObject* ReimportedObject)
 	}
 }
 
+void FAnimationModifiersModule::OnInMemoryAssetCreated(UObject* Object)
+{
+	if (Object->GetClass() == UAnimSequence::StaticClass())
+	{
+		FAnimationModifierHelpers::RetrieveOrCreateModifierUserData(Cast<UAnimSequence>(Object));
+	}
+}
+
 void FAnimationModifiersModule::ShowAddAnimationModifierWindow(const TArray<UAnimSequence*>& InSequences)
 {
 	TSharedPtr<SAnimationModifierContentBrowserWindow> WindowContent;
diff --git a/Engine/Source/Editor/AnimationModifiers/Private/AnimationModifiersModule.h b/Engine/Source/Editor/AnimationModifiers/Private/AnimationModifiersModule.h
--- a/Engine/Source/Editor/AnimationModifiers/Private/AnimationModifiersModule.h
+++ b/Engine/Source/Editor/AnimationModifiers/Private/AnimationModifiersModule.h
@@ -61,6 +61,7 @@ protected:
 	/** Callbacks used to add and apply default animation modifier classes */
 	void OnAssetPostImport(UFactory* ImportFactory, UObject* ImportedObject);
 	void OnAssetPostReimport(UObject* ReimportedObject);
+	void OnInMemoryAssetCreated(UObject* Object);
 	
 	TSharedPtr<FAssetTypeActions_AnimationModifier> AssetAction;
 	FDelegateHandle DelegateHandle;
diff --git a/Engine/Source/Editor/AnimationModifiers/Public/AnimationModifiersAssetUserData.h b/Engine/Source/Editor/AnimationModifiers/Public/AnimationModifiersAssetUserData.h
--- a/Engine/Source/Editor/AnimationModifiers/Public/AnimationModifiersAssetUserData.h
+++ b/Engine/Source/Editor/AnimationModifiers/Public/AnimationModifiersAssetUserData.h
@@ -27,6 +27,7 @@ class ANIMATIONMODIFIERS_API UAnimationModifiersAssetUserData : public UAssetUse
 	friend class FAnimationModifiersModule;
 	friend class UAnimationModifier;
 public:
+	UAnimationModifiersAssetUserData(const FObjectInitializer& ObjectInitializer);
 	const TArray<UAnimationModifier*>& GetAnimationModifierInstances() const;
 protected:	 
 	/** Begin UAssetUserData overrides */
diff --git a/Engine/Source/Runtime/Engine/Classes/Interfaces/Interface_AssetUserData.h b/Engine/Source/Runtime/Engine/Classes/Interfaces/Interface_AssetUserData.h
--- a/Engine/Source/Runtime/Engine/Classes/Interfaces/Interface_AssetUserData.h
+++ b/Engine/Source/Runtime/Engine/Classes/Interfaces/Interface_AssetUserData.h
@@ -6,10 +6,9 @@
 #include "Templates/SubclassOf.h"
 #include "Templates/Casts.h"
 #include "UObject/Interface.h"
+#include "Engine/AssetUserData.h"
 #include "Interface_AssetUserData.generated.h"
 
-class UAssetUserData;
-
 /** Interface for assets/objects that can own UserData **/
 UINTERFACE(MinimalApi, meta=(CannotImplementInterfaceInBlueprint))
 class UInterface_AssetUserData : public UInterface
@@ -22,12 +21,40 @@ class IInterface_AssetUserData
 	GENERATED_IINTERFACE_BODY()
 
 	virtual void AddAssetUserData(UAssetUserData* InUserData) {}
-
+	
 	virtual UAssetUserData* GetAssetUserDataOfClass(TSubclassOf<UAssetUserData> InUserDataClass)
 	{
 		return nullptr;
 	}
 
+	/**
+	* Checks whether or not an instance of the provided AssetUserData class is contained.
+	*
+	* @param	InUserDataClass		UAssetUserData sub class to check for
+	*
+	* @return	Whether or not an instance of InUserDataClass was found
+	*/
+	UFUNCTION(BlueprintCallable, Category=AssetUserData)
+	virtual bool HasAssetUserDataOfClass(TSubclassOf<UAssetUserData> InUserDataClass)
+	{		
+		return GetAssetUserDataOfClass(InUserDataClass) != nullptr;
+	}
+	
+/**
+	* Creates and adds an instance of the provided AssetUserData class to the target asset.
+	*
+	* @param	InUserDataClass		UAssetUserData sub class to create
+	*
+	* @return	Whether or not an instance of InUserDataClass was succesfully added
+	*/
+	UFUNCTION(BlueprintCallable, Category=AssetUserData)
+	virtual bool AddAssetUserDataOfClass(TSubclassOf<UAssetUserData> InUserDataClass)
+	{
+		check(InUserDataClass.Get() != nullptr);
+		AddAssetUserData(NewObject<UAssetUserData>(_getUObject(), InUserDataClass.Get()));
+		return GetAssetUserDataOfClass(InUserDataClass.Get()) != nullptr;
+	}
+
 	virtual const TArray<UAssetUserData*>* GetAssetUserDataArray() const
 	{
 		return nullptr;
