Implement GetOptions for `UPARAM` and add it for `ByProfile` traces

diff --git a/Engine/Source/Editor/BlueprintGraph/Classes/EdGraphSchema_K2.h b/Engine/Source/Editor/BlueprintGraph/Classes/EdGraphSchema_K2.h
--- a/Engine/Source/Editor/BlueprintGraph/Classes/EdGraphSchema_K2.h
+++ b/Engine/Source/Editor/BlueprintGraph/Classes/EdGraphSchema_K2.h
@@ -279,6 +279,9 @@ public:
 	/** Metadata that should be used with UPARAM to specify whether a TSubclassOf argument allows abstract classes */
 	static const FName MD_AllowAbstractClasses;
 
+	/** Metadata that should be used with UPARAM to specify a function name that generates a list of available values */
+	static const FName MD_GetOptions;
+
 	/** Namespace into which a type can be optionally defined; if empty or not set, the type will belong to the global namespace (default). */
 	static const FName MD_Namespace;
 
diff --git a/Engine/Source/Editor/BlueprintGraph/Private/EdGraphSchema_K2.cpp b/Engine/Source/Editor/BlueprintGraph/Private/EdGraphSchema_K2.cpp
--- a/Engine/Source/Editor/BlueprintGraph/Private/EdGraphSchema_K2.cpp
+++ b/Engine/Source/Editor/BlueprintGraph/Private/EdGraphSchema_K2.cpp
@@ -275,6 +275,7 @@ const FName FBlueprintMetadata::MD_UseEnumValuesAsMaskValuesInEditor(TEXT("UseEn
 const FName FBlueprintMetadata::MD_AnimBlueprintFunction(TEXT("AnimBlueprintFunction"));
 
 const FName FBlueprintMetadata::MD_AllowAbstractClasses(TEXT("AllowAbstract"));
+const FName FBlueprintMetadata::MD_GetOptions(TEXT("GetOptions"));
 
 const FName FBlueprintMetadata::MD_Namespace(TEXT("Namespace"));
 
diff --git a/Engine/Source/Editor/GraphEditor/GraphEditor.Build.cs b/Engine/Source/Editor/GraphEditor/GraphEditor.Build.cs
--- a/Engine/Source/Editor/GraphEditor/GraphEditor.Build.cs
+++ b/Engine/Source/Editor/GraphEditor/GraphEditor.Build.cs
@@ -33,6 +33,7 @@ public class GraphEditor : ModuleRules
 				"BlueprintGraph",
 				"Documentation",
 				"Persona",
+				"PropertyEditor",
 				"RenderCore",
 				"RHI",
 				"ToolMenus",
diff --git a/Engine/Source/Editor/GraphEditor/Private/KismetPins/SGraphPinString.cpp b/Engine/Source/Editor/GraphEditor/Private/KismetPins/SGraphPinString.cpp
--- a/Engine/Source/Editor/GraphEditor/Private/KismetPins/SGraphPinString.cpp
+++ b/Engine/Source/Editor/GraphEditor/Private/KismetPins/SGraphPinString.cpp
@@ -10,8 +10,11 @@
 #include "EdGraphSchema_K2.h"
 #include "GenericPlatform/GenericApplication.h"
 #include "Internationalization/Internationalization.h"
+#include "K2Node_CallFunction.h"
+#include "PropertyEditorUtils.h"
 #include "Misc/Attribute.h"
 #include "ScopedTransaction.h"
+#include "SSearchableComboBox.h"
 #include "Styling/AppStyle.h"
 #include "Styling/SlateColor.h"
 #include "Types/SlateStructs.h"
@@ -53,6 +56,12 @@ TSharedRef<SWidget>	SGraphPinString::GetDefaultValueWidget()
 	}
 	else
 	{
+		const TSharedPtr<SWidget> DropdownWidget = TryBuildComboBoxWidget();
+		if (DropdownWidget.Get())
+		{
+			return DropdownWidget.ToSharedRef();
+		}
+
 		return SNew(SBox)
 			.MinDesiredWidth(18)
 			.MaxDesiredWidth(400)
@@ -94,3 +103,98 @@ bool SGraphPinString::GetDefaultValueIsReadOnly() const
 {
 	return GraphPinObj->bDefaultValueIsReadOnly;
 }
+
+TSharedRef<SWidget> SGraphPinString::GenerateComboBoxEntry(const TSharedPtr<FString> Value)
+{
+	return SNew(STextBlock)
+		.Text(FText::FromString(*Value))
+		.Font(FAppStyle::GetFontStyle("PropertyWindow.NormalFont"));
+}
+
+void SGraphPinString::HandleComboBoxSelectionChanged(const TSharedPtr<FString> Value,
+	const ESelectInfo::Type InSelectInfo)
+{
+	if (GraphPinObj->IsPendingKill())
+	{
+		return;
+	}
+
+	if (!GraphPinObj->GetDefaultAsString().Equals(*Value))
+	{
+		const FScopedTransaction Transaction(
+			NSLOCTEXT("GraphEditor", "ChangeStringPinValue", "Change String Pin Value")
+		);
+
+		GraphPinObj->Modify();
+		GraphPinObj->GetSchema()->TrySetDefaultValue(*GraphPinObj, *Value);
+	}
+}
+
+TSharedPtr<SWidget> SGraphPinString::TryBuildComboBoxWidget()
+{
+	FString GetOptionsFunctionName = GraphPinObj->GetOwningNode()->GetPinMetaData(
+		GraphPinObj->PinName,
+		FBlueprintMetadata::MD_GetOptions
+	);
+
+	if (GetOptionsFunctionName.IsEmpty())
+	{
+		return nullptr;
+	}
+
+	const UK2Node_CallFunction* CallFunctionNode = Cast<UK2Node_CallFunction>(GraphPinObj->GetOwningNode());
+	if (!CallFunctionNode)
+	{
+		return nullptr;
+	}
+
+	const UClass* FunctionParentClass = CallFunctionNode->FunctionReference.GetMemberParentClass();
+	if (!FunctionParentClass)
+	{
+		// If we don't have a parent class, then we have to be in self context
+		if (!CallFunctionNode->FunctionReference.IsSelfContext())
+		{
+			return nullptr;
+		}
+
+		// If we're in self context, let's get the Blueprint we're residing in
+		const UBlueprint* Blueprint = GraphPinObj->GetOwningNode()->GetTypedOuter<UBlueprint>();
+		if (!Blueprint || !Blueprint->GeneratedClass)
+		{
+			return nullptr;
+		}
+
+		// Work with the blueprint class the node is in
+		FunctionParentClass = Blueprint->GeneratedClass;
+	}
+
+	ComboBoxOptions.Empty();
+
+	TArray<UObject*> PropertyContainers = { FunctionParentClass->GetDefaultObject() };
+	PropertyEditorUtils::GetPropertyOptions(PropertyContainers, GetOptionsFunctionName, ComboBoxOptions);
+
+	const FString DefaultString = GraphPinObj->GetDefaultAsString();
+	const TSharedPtr<FString>* DefaultValue = ComboBoxOptions.FindByPredicate([&DefaultString](
+		const TSharedPtr<FString>& Value)
+	{
+		return Value.IsValid() ? *Value == DefaultString : false;
+	});
+
+	return SNew(SBox)
+		.MinDesiredWidth(18)
+		.MaxDesiredWidth(400)
+		[
+			SNew(SSearchableComboBox)
+				.OptionsSource(&ComboBoxOptions)
+				.OnGenerateWidget(this, &SGraphPinString::GenerateComboBoxEntry)
+				.OnSelectionChanged(this, &SGraphPinString::HandleComboBoxSelectionChanged)
+				.ContentPadding(3.0f)
+				.InitiallySelectedItem(DefaultValue ? *DefaultValue : nullptr)
+				.Content()
+				[
+					SNew(STextBlock)
+						.Text(this, &SGraphPinString::GetTypeInValue)
+						.Font(FAppStyle::GetFontStyle("PropertyWindow.NormalFont"))
+				]
+		];
+}
diff --git a/Engine/Source/Editor/GraphEditor/Public/KismetPins/SGraphPinString.h b/Engine/Source/Editor/GraphEditor/Public/KismetPins/SGraphPinString.h
--- a/Engine/Source/Editor/GraphEditor/Public/KismetPins/SGraphPinString.h
+++ b/Engine/Source/Editor/GraphEditor/Public/KismetPins/SGraphPinString.h
@@ -31,4 +31,11 @@ protected:
 
 	/** @return True if the pin default value field is read-only */
 	bool GetDefaultValueIsReadOnly() const;
+
+private:
+	TSharedRef<SWidget> GenerateComboBoxEntry(TSharedPtr<FString> Value);
+	void HandleComboBoxSelectionChanged(TSharedPtr<FString> Value, ESelectInfo::Type InSelectInfo);
+	TSharedPtr<SWidget> TryBuildComboBoxWidget();
+
+	TArray<TSharedPtr<FString>> ComboBoxOptions;
 };
diff --git a/Engine/Source/Editor/PropertyEditor/Private/PropertyEditorUtils.cpp b/Engine/Source/Editor/PropertyEditor/Private/PropertyEditorUtils.cpp
new file mode 100644
--- /dev/null
+++ b/Engine/Source/Editor/PropertyEditor/Private/PropertyEditorUtils.cpp
@@ -0,0 +1,70 @@
+﻿// Copyright Epic Games, Inc. All Rights Reserved.
+
+#include "PropertyEditorUtils.h"
+#include "PropertyPathHelpers.h"
+
+namespace PropertyEditorUtils
+{
+	void GetPropertyOptions(TArray<UObject*>& InOutContainers, FString& InOutPropertyPath,
+		TArray<TSharedPtr<FString>>& InOutOptions)
+	{
+		// Check for external function references
+		if (InOutPropertyPath.Contains(TEXT(".")))
+		{
+			InOutContainers.Empty();
+			UFunction* GetOptionsFunction = FindObject<UFunction>(nullptr, *InOutPropertyPath, true);
+
+			if (ensureMsgf(GetOptionsFunction && GetOptionsFunction->HasAnyFunctionFlags(EFunctionFlags::FUNC_Static), TEXT("Invalid GetOptions: %s"), *InOutPropertyPath))
+			{
+				UObject* GetOptionsCDO = GetOptionsFunction->GetOuterUClass()->GetDefaultObject();
+				GetOptionsFunction->GetName(InOutPropertyPath);
+				InOutContainers.Add(GetOptionsCDO);
+			}
+		}
+
+		if (InOutContainers.Num() > 0)
+		{
+			TArray<FString> OptionIntersection;
+			TSet<FString> OptionIntersectionSet;
+
+			for (UObject* Target : InOutContainers)
+			{
+				TArray<FString> StringOptions;
+				{
+					FEditorScriptExecutionGuard ScriptExecutionGuard;
+
+					FCachedPropertyPath Path(InOutPropertyPath);
+					if (!PropertyPathHelpers::GetPropertyValue(Target, Path, StringOptions))
+					{
+						TArray<FName> NameOptions;
+						if (PropertyPathHelpers::GetPropertyValue(Target, Path, NameOptions))
+						{
+							Algo::Transform(NameOptions, StringOptions, [](const FName& InName) { return InName.ToString(); });
+						}
+					}
+				}
+
+				// If this is the first time there won't be any options.
+				if (OptionIntersection.Num() == 0)
+				{
+					OptionIntersection = StringOptions;
+					OptionIntersectionSet = TSet<FString>(StringOptions);
+				}
+				else
+				{
+					TSet<FString> StringOptionsSet(StringOptions);
+					OptionIntersectionSet = StringOptionsSet.Intersect(OptionIntersectionSet);
+					OptionIntersection.RemoveAll([&OptionIntersectionSet](const FString& Option){ return !OptionIntersectionSet.Contains(Option); });
+				}
+
+				// If we're out of possible intersected options, we can stop.
+				if (OptionIntersection.Num() == 0)
+				{
+					break;
+				}
+			}
+
+			Algo::Transform(OptionIntersection, InOutOptions, [](const FString& InString) { return MakeShared<FString>(InString); });
+		}
+	}
+}
diff --git a/Engine/Source/Editor/PropertyEditor/Private/PropertyHandleImpl.cpp b/Engine/Source/Editor/PropertyEditor/Private/PropertyHandleImpl.cpp
--- a/Engine/Source/Editor/PropertyEditor/Private/PropertyHandleImpl.cpp
+++ b/Engine/Source/Editor/PropertyEditor/Private/PropertyHandleImpl.cpp
@@ -13,6 +13,7 @@
 #include "Editor/UnrealEdEngine.h"
 #include "UnrealEdGlobals.h"
 #include "PropertyEditorHelpers.h"
+#include "PropertyEditorUtils.h"
 #include "StructurePropertyNode.h"
 #include "ScopedTransaction.h"
 #include "Kismet2/KismetEditorUtilities.h"
@@ -3146,64 +3147,7 @@ bool FPropertyHandleBase::GeneratePossibleValues(TArray< TSharedPtr<FString> >&
 			TArray<UObject*> OutObjects;
 			GetOuterObjects(OutObjects);
 
-			// Check for external function references
-			if (GetOptionsFunctionName.Contains(TEXT(".")))
-			{
-				OutObjects.Empty();
-				UFunction* GetOptionsFunction = FindObject<UFunction>(nullptr, *GetOptionsFunctionName, true);
-
-				if (ensureMsgf(GetOptionsFunction && GetOptionsFunction->HasAnyFunctionFlags(EFunctionFlags::FUNC_Static), TEXT("Invalid GetOptions: %s"), *GetOptionsFunctionName))
-				{
-					UObject* GetOptionsCDO = GetOptionsFunction->GetOuterUClass()->GetDefaultObject();
-					GetOptionsFunction->GetName(GetOptionsFunctionName);
-					OutObjects.Add(GetOptionsCDO);
-				}
-			}
-
-			if (OutObjects.Num() > 0)
-			{
-				TArray<FString> OptionIntersection;
-				TSet<FString> OptionIntersectionSet;
-
-				for (UObject* Target : OutObjects)
-				{
-					TArray<FString> StringOptions;
-					{
-						FEditorScriptExecutionGuard ScriptExecutionGuard;
-
-						FCachedPropertyPath Path(GetOptionsFunctionName);
-						if (!PropertyPathHelpers::GetPropertyValue(Target, Path, StringOptions))
-						{
-							TArray<FName> NameOptions;
-							if (PropertyPathHelpers::GetPropertyValue(Target, Path, NameOptions))
-							{
-								Algo::Transform(NameOptions, StringOptions, [](const FName& InName) { return InName.ToString(); });
-							}
-						}
-					}
-
-					// If this is the first time there won't be any options.
-					if (OptionIntersection.Num() == 0)
-					{
-						OptionIntersection = StringOptions;
-						OptionIntersectionSet = TSet<FString>(StringOptions);
-					}
-					else
-					{
-						TSet<FString> StringOptionsSet(StringOptions);
-						OptionIntersectionSet = StringOptionsSet.Intersect(OptionIntersectionSet);
-						OptionIntersection.RemoveAll([&OptionIntersectionSet](const FString& Option){ return !OptionIntersectionSet.Contains(Option); });
-					}
-
-					// If we're out of possible intersected options, we can stop.
-					if (OptionIntersection.Num() == 0)
-					{
-						break;
-					}
-				}
-
-				Algo::Transform(OptionIntersection, OutOptionStrings, [](const FString& InString) { return MakeShared<FString>(InString); });
-			}
+			PropertyEditorUtils::GetPropertyOptions(OutObjects, GetOptionsFunctionName, OutOptionStrings);
 		}
 	}
 	else if( Property->IsA(FClassProperty::StaticClass()) || Property->IsA(FSoftClassProperty::StaticClass()) )		
diff --git a/Engine/Source/Editor/PropertyEditor/Public/PropertyEditorUtils.h b/Engine/Source/Editor/PropertyEditor/Public/PropertyEditorUtils.h
new file mode 100644
--- /dev/null
+++ b/Engine/Source/Editor/PropertyEditor/Public/PropertyEditorUtils.h
@@ -0,0 +1,16 @@
+﻿// Copyright Epic Games, Inc. All Rights Reserved.
+
+#pragma once
+
+#include "UObject/Object.h"
+
+namespace PropertyEditorUtils
+{
+	/**
+	 * Calculates the possible drop-down options for the specified property path
+	 * @param	InOutContainers		The container objects to resolve the property path against
+	 * @param	InOutPropertyPath	The property path
+	 * @param	InOutOptions		The resulting options
+	 */
+	PROPERTYEDITOR_API void GetPropertyOptions(TArray<UObject*>& InOutContainers, FString& InOutPropertyPath, TArray<TSharedPtr<FString>>& InOutOptions);
+}
diff --git a/Engine/Source/Runtime/CoreUObject/Public/UObject/ObjectMacros.h b/Engine/Source/Runtime/CoreUObject/Public/UObject/ObjectMacros.h
--- a/Engine/Source/Runtime/CoreUObject/Public/UObject/ObjectMacros.h
+++ b/Engine/Source/Runtime/CoreUObject/Public/UObject/ObjectMacros.h
@@ -1460,7 +1460,8 @@ namespace UM
 		/// [PropertyMetadata] Causes FString and FName properties to have a limited set of options generated dynamically, e.g. meta=(GetOptions="FuncName"). Supports external static function references via "Module.Class.Function" syntax.
 		///
 		/// UFUNCTION()
-		/// TArray<FString> FuncName() const; // Always return string array even if FName property.
+		/// static TArray<FName> FuncName() const; // If the field wants an FName value
+		/// static TArray<FString> FuncName() const; // If the field wants an FString value
 		GetOptions,
 
 		/// [PropertyMetadata] The property can be exposed as a data pin, but is hidden by default.
diff --git a/Engine/Source/Runtime/Engine/Classes/Kismet/KismetSystemLibrary.h b/Engine/Source/Runtime/Engine/Classes/Kismet/KismetSystemLibrary.h
--- a/Engine/Source/Runtime/Engine/Classes/Kismet/KismetSystemLibrary.h
+++ b/Engine/Source/Runtime/Engine/Classes/Kismet/KismetSystemLibrary.h
@@ -1377,7 +1377,7 @@ class UKismetSystemLibrary : public UBlueprintFunctionLibrary
 	* @return				True if there was a hit, false otherwise.
 	*/
 	UFUNCTION(BlueprintCallable, Category = "Collision", meta = (bIgnoreSelf = "true", WorldContext = "WorldContextObject", AutoCreateRefTerm = "ActorsToIgnore", DisplayName = "Line Trace By Profile", AdvancedDisplay = "TraceColor,TraceHitColor,DrawTime", Keywords = "raycast"))
-	static ENGINE_API bool LineTraceSingleByProfile(const UObject* WorldContextObject, const FVector Start, const FVector End, FName ProfileName, bool bTraceComplex, const TArray<AActor*>& ActorsToIgnore, EDrawDebugTrace::Type DrawDebugType, FHitResult& OutHit, bool bIgnoreSelf, FLinearColor TraceColor = FLinearColor::Red, FLinearColor TraceHitColor = FLinearColor::Green, float DrawTime = 5.0f);
+	static ENGINE_API bool LineTraceSingleByProfile(const UObject* WorldContextObject, const FVector Start, const FVector End, UPARAM(Meta=(GetOptions="Engine.KismetSystemLibrary.GetCollisionProfileNames")) FName ProfileName, bool bTraceComplex, const TArray<AActor*>& ActorsToIgnore, EDrawDebugTrace::Type DrawDebugType, FHitResult& OutHit, bool bIgnoreSelf, FLinearColor TraceColor = FLinearColor::Red, FLinearColor TraceHitColor = FLinearColor::Green, float DrawTime = 5.0f);
 
 	/**
 	*  Trace a ray against the world using a specific profile and return overlapping hits and then first blocking hit
@@ -1393,7 +1393,7 @@ class UKismetSystemLibrary : public UBlueprintFunctionLibrary
 	* @return				True if there was a blocking hit, false otherwise.
 	*/
 	UFUNCTION(BlueprintCallable, Category = "Collision", meta = (bIgnoreSelf = "true", WorldContext = "WorldContextObject", AutoCreateRefTerm = "ActorsToIgnore", DisplayName = "Multi Line Trace By Profile", AdvancedDisplay = "TraceColor,TraceHitColor,DrawTime", Keywords = "raycast"))
-	static ENGINE_API bool LineTraceMultiByProfile(const UObject* WorldContextObject, const FVector Start, const FVector End, FName ProfileName, bool bTraceComplex, const TArray<AActor*>& ActorsToIgnore, EDrawDebugTrace::Type DrawDebugType, TArray<FHitResult>& OutHits, bool bIgnoreSelf, FLinearColor TraceColor = FLinearColor::Red, FLinearColor TraceHitColor = FLinearColor::Green, float DrawTime = 5.0f);
+	static ENGINE_API bool LineTraceMultiByProfile(const UObject* WorldContextObject, const FVector Start, const FVector End, UPARAM(Meta=(GetOptions=GetCollisionProfileNames)) FName ProfileName, bool bTraceComplex, const TArray<AActor*>& ActorsToIgnore, EDrawDebugTrace::Type DrawDebugType, TArray<FHitResult>& OutHits, bool bIgnoreSelf, FLinearColor TraceColor = FLinearColor::Red, FLinearColor TraceHitColor = FLinearColor::Green, float DrawTime = 5.0f);
 
 	/**
 	*  Sweep a sphere against the world and return the first blocking hit using a specific profile
@@ -1407,7 +1407,7 @@ class UKismetSystemLibrary : public UBlueprintFunctionLibrary
 	* @return				True if there was a hit, false otherwise.
 	*/
 	UFUNCTION(BlueprintCallable, Category = "Collision", meta = (bIgnoreSelf = "true", WorldContext = "WorldContextObject", AutoCreateRefTerm = "ActorsToIgnore", DisplayName = "Sphere Trace By Profile", AdvancedDisplay = "TraceColor,TraceHitColor,DrawTime", Keywords = "sweep"))
-	static ENGINE_API bool SphereTraceSingleByProfile(const UObject* WorldContextObject, const FVector Start, const FVector End, float Radius, FName ProfileName, bool bTraceComplex, const TArray<AActor*>& ActorsToIgnore, EDrawDebugTrace::Type DrawDebugType, FHitResult& OutHit, bool bIgnoreSelf, FLinearColor TraceColor = FLinearColor::Red, FLinearColor TraceHitColor = FLinearColor::Green, float DrawTime = 5.0f);
+	static ENGINE_API bool SphereTraceSingleByProfile(const UObject* WorldContextObject, const FVector Start, const FVector End, float Radius, UPARAM(Meta=(GetOptions=GetCollisionProfileNames)) FName ProfileName, bool bTraceComplex, const TArray<AActor*>& ActorsToIgnore, EDrawDebugTrace::Type DrawDebugType, FHitResult& OutHit, bool bIgnoreSelf, FLinearColor TraceColor = FLinearColor::Red, FLinearColor TraceHitColor = FLinearColor::Green, float DrawTime = 5.0f);
 
 	/**
 	*  Sweep a sphere against the world and return all initial overlaps using a specific profile, then overlapping hits and then first blocking hit
@@ -1424,7 +1424,7 @@ class UKismetSystemLibrary : public UBlueprintFunctionLibrary
 	* @return				True if there was a blocking hit, false otherwise.
 	*/
 	UFUNCTION(BlueprintCallable, Category = "Collision", meta = (bIgnoreSelf = "true", WorldContext = "WorldContextObject", AutoCreateRefTerm = "ActorsToIgnore", DisplayName = "Multi Sphere Trace By Profile", AdvancedDisplay = "TraceColor,TraceHitColor,DrawTime", Keywords = "sweep"))
-	static ENGINE_API bool SphereTraceMultiByProfile(const UObject* WorldContextObject, const FVector Start, const FVector End, float Radius, FName ProfileName, bool bTraceComplex, const TArray<AActor*>& ActorsToIgnore, EDrawDebugTrace::Type DrawDebugType, TArray<FHitResult>& OutHits, bool bIgnoreSelf, FLinearColor TraceColor = FLinearColor::Red, FLinearColor TraceHitColor = FLinearColor::Green, float DrawTime = 5.0f);
+	static ENGINE_API bool SphereTraceMultiByProfile(const UObject* WorldContextObject, const FVector Start, const FVector End, float Radius, UPARAM(Meta=(GetOptions=GetCollisionProfileNames)) FName ProfileName, bool bTraceComplex, const TArray<AActor*>& ActorsToIgnore, EDrawDebugTrace::Type DrawDebugType, TArray<FHitResult>& OutHits, bool bIgnoreSelf, FLinearColor TraceColor = FLinearColor::Red, FLinearColor TraceHitColor = FLinearColor::Green, float DrawTime = 5.0f);
 
 	/**
 	*  Sweep a box against the world and return the first blocking hit using a specific profile
@@ -1439,7 +1439,7 @@ class UKismetSystemLibrary : public UBlueprintFunctionLibrary
 	* @return				True if there was a hit, false otherwise.
 	*/
 	UFUNCTION(BlueprintCallable, Category = "Collision", meta = (bIgnoreSelf = "true", WorldContext = "WorldContextObject", AutoCreateRefTerm = "ActorsToIgnore", DisplayName = "Box Trace By Profile", AdvancedDisplay = "TraceColor,TraceHitColor,DrawTime", Keywords = "sweep"))
-	static ENGINE_API bool BoxTraceSingleByProfile(const UObject* WorldContextObject, const FVector Start, const FVector End, const FVector HalfSize, const FRotator Orientation, FName ProfileName, bool bTraceComplex, const TArray<AActor*>& ActorsToIgnore, EDrawDebugTrace::Type DrawDebugType, FHitResult& OutHit, bool bIgnoreSelf, FLinearColor TraceColor = FLinearColor::Red, FLinearColor TraceHitColor = FLinearColor::Green, float DrawTime = 5.0f);
+	static ENGINE_API bool BoxTraceSingleByProfile(const UObject* WorldContextObject, const FVector Start, const FVector End, const FVector HalfSize, const FRotator Orientation, UPARAM(Meta=(GetOptions=GetCollisionProfileNames)) FName ProfileName, bool bTraceComplex, const TArray<AActor*>& ActorsToIgnore, EDrawDebugTrace::Type DrawDebugType, FHitResult& OutHit, bool bIgnoreSelf, FLinearColor TraceColor = FLinearColor::Red, FLinearColor TraceHitColor = FLinearColor::Green, float DrawTime = 5.0f);
 
 	/**
 	*  Sweep a box against the world and return all initial overlaps using a specific profile, then overlapping hits and then first blocking hit
@@ -1456,7 +1456,7 @@ class UKismetSystemLibrary : public UBlueprintFunctionLibrary
 	* @return				True if there was a blocking hit, false otherwise.
 	*/
 	UFUNCTION(BlueprintCallable, Category = "Collision", meta = (bIgnoreSelf = "true", WorldContext = "WorldContextObject", AutoCreateRefTerm = "ActorsToIgnore", DisplayName = "Multi Box Trace By Profile", AdvancedDisplay = "TraceColor,TraceHitColor,DrawTime", Keywords = "sweep"))
-	static ENGINE_API bool BoxTraceMultiByProfile(const UObject* WorldContextObject, const FVector Start, const FVector End, FVector HalfSize, const FRotator Orientation, FName ProfileName, bool bTraceComplex, const TArray<AActor*>& ActorsToIgnore, EDrawDebugTrace::Type DrawDebugType, TArray<FHitResult>& OutHits, bool bIgnoreSelf, FLinearColor TraceColor = FLinearColor::Red, FLinearColor TraceHitColor = FLinearColor::Green, float DrawTime = 5.0f);
+	static ENGINE_API bool BoxTraceMultiByProfile(const UObject* WorldContextObject, const FVector Start, const FVector End, FVector HalfSize, const FRotator Orientation, UPARAM(Meta=(GetOptions=GetCollisionProfileNames)) FName ProfileName, bool bTraceComplex, const TArray<AActor*>& ActorsToIgnore, EDrawDebugTrace::Type DrawDebugType, TArray<FHitResult>& OutHits, bool bIgnoreSelf, FLinearColor TraceColor = FLinearColor::Red, FLinearColor TraceHitColor = FLinearColor::Green, float DrawTime = 5.0f);
 
 
 	/**
@@ -1473,7 +1473,7 @@ class UKismetSystemLibrary : public UBlueprintFunctionLibrary
 	* @return				True if there was a hit, false otherwise.
 	*/
 	UFUNCTION(BlueprintCallable, Category = "Collision", meta = (bIgnoreSelf = "true", WorldContext = "WorldContextObject", AutoCreateRefTerm = "ActorsToIgnore", DisplayName = "Capsule Trace By Profile", AdvancedDisplay = "TraceColor,TraceHitColor,DrawTime", Keywords = "sweep"))
-	static ENGINE_API bool CapsuleTraceSingleByProfile(const UObject* WorldContextObject, const FVector Start, const FVector End, float Radius, float HalfHeight, FName ProfileName, bool bTraceComplex, const TArray<AActor*>& ActorsToIgnore, EDrawDebugTrace::Type DrawDebugType, FHitResult& OutHit, bool bIgnoreSelf, FLinearColor TraceColor = FLinearColor::Red, FLinearColor TraceHitColor = FLinearColor::Green, float DrawTime = 5.0f);
+	static ENGINE_API bool CapsuleTraceSingleByProfile(const UObject* WorldContextObject, const FVector Start, const FVector End, float Radius, float HalfHeight, UPARAM(Meta=(GetOptions=GetCollisionProfileNames)) FName ProfileName, bool bTraceComplex, const TArray<AActor*>& ActorsToIgnore, EDrawDebugTrace::Type DrawDebugType, FHitResult& OutHit, bool bIgnoreSelf, FLinearColor TraceColor = FLinearColor::Red, FLinearColor TraceHitColor = FLinearColor::Green, float DrawTime = 5.0f);
 
 	/**
 	*  Sweep a capsule against the world and return all initial overlaps using a specific profile, then overlapping hits and then first blocking hit
@@ -1491,8 +1491,12 @@ class UKismetSystemLibrary : public UBlueprintFunctionLibrary
 	* @return				True if there was a blocking hit, false otherwise.
 	*/
 	UFUNCTION(BlueprintCallable, Category = "Collision", meta = (bIgnoreSelf = "true", WorldContext = "WorldContextObject", AutoCreateRefTerm = "ActorsToIgnore", DisplayName = "Multi Capsule Trace By Profile", AdvancedDisplay = "TraceColor,TraceHitColor,DrawTime", Keywords = "sweep"))
-	static ENGINE_API bool CapsuleTraceMultiByProfile(const UObject* WorldContextObject, const FVector Start, const FVector End, float Radius, float HalfHeight, FName ProfileName, bool bTraceComplex, const TArray<AActor*>& ActorsToIgnore, EDrawDebugTrace::Type DrawDebugType, TArray<FHitResult>& OutHits, bool bIgnoreSelf, FLinearColor TraceColor = FLinearColor::Red, FLinearColor TraceHitColor = FLinearColor::Green, float DrawTime = 5.0f);
+	static ENGINE_API bool CapsuleTraceMultiByProfile(const UObject* WorldContextObject, const FVector Start, const FVector End, float Radius, float HalfHeight, UPARAM(Meta=(GetOptions=GetCollisionProfileNames)) FName ProfileName, bool bTraceComplex, const TArray<AActor*>& ActorsToIgnore, EDrawDebugTrace::Type DrawDebugType, TArray<FHitResult>& OutHits, bool bIgnoreSelf, FLinearColor TraceColor = FLinearColor::Red, FLinearColor TraceHitColor = FLinearColor::Green, float DrawTime = 5.0f);
 
+#if WITH_EDITOR
+	UFUNCTION()
+	static TArray<FName> GetCollisionProfileNames();
+#endif
 
 
 	/**
diff --git a/Engine/Source/Runtime/Engine/Private/KismetSystemLibrary.cpp b/Engine/Source/Runtime/Engine/Private/KismetSystemLibrary.cpp
--- a/Engine/Source/Runtime/Engine/Private/KismetSystemLibrary.cpp
+++ b/Engine/Source/Runtime/Engine/Private/KismetSystemLibrary.cpp
@@ -2221,6 +2221,25 @@ bool UKismetSystemLibrary::CapsuleTraceMultiByProfile(const UObject* WorldContex
 	return bHit;
 }
 
+#if WITH_EDITOR
+TArray<FName> UKismetSystemLibrary::GetCollisionProfileNames()
+{
+	TArray<TSharedPtr<FName>> SharedNames;
+	UCollisionProfile::GetProfileNames(SharedNames);
+
+	TArray<FName> Names;
+	Names.Reserve(SharedNames.Num());
+	for (const TSharedPtr<FName>& SharedName : SharedNames)
+	{
+		if (const FName* Name = SharedName.Get())
+		{
+			Names.Add(*Name);
+		}
+	}
+
+	return Names;
+}
+#endif
 
 /** Draw a debug line */
 void UKismetSystemLibrary::DrawDebugLine(const UObject* WorldContextObject, FVector const LineStart, FVector const LineEnd, FLinearColor Color, float LifeTime, float Thickness)
