Fix modulation crash on PIE/Shutdown

diff --git a/Engine/Source/Runtime/Engine/Classes/Sound/SoundModulationDestination.h b/Engine/Source/Runtime/Engine/Classes/Sound/SoundModulationDestination.h
--- a/Engine/Source/Runtime/Engine/Classes/Sound/SoundModulationDestination.h
+++ b/Engine/Source/Runtime/Engine/Classes/Sound/SoundModulationDestination.h
@@ -179,40 +179,50 @@ namespace Audio
 			ENGINE_API void UpdateModulators(const TSet<TObjectPtr<USoundModulatorBase>>& InModulators);
 			ENGINE_API void UpdateModulators(const TSet<USoundModulatorBase*>& InModulators);
 			ENGINE_API void UpdateModulators(const TSet<const USoundModulatorBase*>& InModulators);
-	
+
 	private:
 			ENGINE_API void UpdateModulatorsInternal(TArray<TUniquePtr<Audio::IModulatorSettings>>&& ProxySettings);
 
-			ENGINE_API void ResetHandles();
+			struct FModulationDestinationData
+			{
+				FDeviceId DeviceId = INDEX_NONE;
+
+				float ValueTarget = 1.0f;
 
-			FDeviceId DeviceId = INDEX_NONE;
+				bool bIsBuffered = false;
+				bool bValueNormalized = false;
+				bool bHasProcessed = false;
 
-			float ValueTarget = 1.0f;
+				FAlignedFloatBuffer OutputBuffer;
 
-			bool bIsBuffered = false;
-			bool bValueNormalized = false;
-			bool bHasProcessed = false;
+				TSet<FModulatorHandle> Handles;
 
-			FAlignedFloatBuffer OutputBuffer;
+				FModulationParameter Parameter;
 
-			TSet<FModulatorHandle> Handles;
+				mutable FCriticalSection HandleCritSection;
 
-			FModulationParameter Parameter;
+				FModulationDestinationData& operator=(const FModulationDestinationData& InDestInfo);
+				FModulationDestinationData& operator=(FModulationDestinationData&& InDestInfo);
+				const FDeviceId& GetDeviceId() const;
+				const FModulationParameter& GetParameter() const;
+				void SetHandles(TSet<FModulatorHandle>&& Handles);
+				void ResetHandles();
+			};
 
-			mutable FCriticalSection HandleCritSection;
+			TSharedRef<FModulationDestinationData> DestinationData{ MakeShared<FModulationDestinationData>() };
 
 		public:
 			/** Returns buffer of interpolated modulation values. If not set to "IsBuffered" when initialized, returns an empty array. */
 			FORCEINLINE const FAlignedFloatBuffer& GetBuffer() const
 			{
-				return OutputBuffer;
+				return DestinationData->OutputBuffer;
 			}
 
 			/** Returns whether or not the destination has requested to 
 			  * process the control or not. */
 			FORCEINLINE bool GetHasProcessed() const
 			{
-				return bHasProcessed;
+				return DestinationData->bHasProcessed;
 			}
 
 			/** Returns sample value last reported by modulator. Returns value in unit space, unless
@@ -220,7 +230,7 @@ namespace Audio
 			 */
 			FORCEINLINE float GetValue() const
 			{
-				return ValueTarget;
+				return DestinationData->ValueTarget;
 			}
 	};
 } // namespace Audio
diff --git a/Engine/Source/Runtime/Engine/Private/SoundModulationDestination.cpp b/Engine/Source/Runtime/Engine/Private/SoundModulationDestination.cpp
--- a/Engine/Source/Runtime/Engine/Private/SoundModulationDestination.cpp
+++ b/Engine/Source/Runtime/Engine/Private/SoundModulationDestination.cpp
@@ -50,78 +50,30 @@ FSoundModulationDefaultRoutingSettings::FSoundModulationDefaultRoutingSettings()
 namespace Audio
 {
 	FModulationDestination::FModulationDestination(const FModulationDestination& InModulationDestination)
-		: DeviceId(InModulationDestination.DeviceId)
-		, ValueTarget(InModulationDestination.ValueTarget)
-		, bIsBuffered(InModulationDestination.bIsBuffered)
-		, bValueNormalized(InModulationDestination.bValueNormalized)
-		, OutputBuffer(InModulationDestination.OutputBuffer)
-		, Parameter(InModulationDestination.Parameter)
 	{
-		FScopeLock OtherLock(&InModulationDestination.HandleCritSection);
-		Handles = InModulationDestination.Handles;
+		FScopeLock OtherLock(&InModulationDestination.DestinationData->HandleCritSection);
+		*DestinationData = *InModulationDestination.DestinationData;
 	}
 
 	FModulationDestination::FModulationDestination(FModulationDestination&& InModulationDestination)
-		: DeviceId(MoveTemp(InModulationDestination.DeviceId))
-		, ValueTarget(MoveTemp(InModulationDestination.ValueTarget))
-		, bIsBuffered(MoveTemp(InModulationDestination.bIsBuffered))
-		, bValueNormalized(MoveTemp(InModulationDestination.bValueNormalized))
-		, OutputBuffer(MoveTemp(InModulationDestination.OutputBuffer))
-		, Parameter(MoveTemp(InModulationDestination.Parameter))
 	{
-		FScopeLock OtherLock(&InModulationDestination.HandleCritSection);
-		Handles = MoveTemp(InModulationDestination.Handles);
+		FScopeLock OtherLock(&InModulationDestination.DestinationData->HandleCritSection);
+		*DestinationData = MoveTemp(*InModulationDestination.DestinationData);
 	}
 
 	FModulationDestination& FModulationDestination::operator=(const FModulationDestination& InModulationDestination)
 	{
-		DeviceId = InModulationDestination.DeviceId;
-		ValueTarget = InModulationDestination.ValueTarget;
-		bIsBuffered = InModulationDestination.bIsBuffered;
-		bValueNormalized = InModulationDestination.bValueNormalized;
-		OutputBuffer = InModulationDestination.OutputBuffer;
-
-		TSet<FModulatorHandle> NewHandles;
-		{
-			FScopeLock OtherLock(&InModulationDestination.HandleCritSection);
-			NewHandles = InModulationDestination.Handles;
-		}
-
-		{
-			FScopeLock Lock(&HandleCritSection);
-			Handles = MoveTemp(NewHandles);
-		}
-
-		Parameter = InModulationDestination.Parameter;
-
+		*DestinationData = *(InModulationDestination.DestinationData);
 		return *this;
 	}
 
 	FModulationDestination& FModulationDestination::operator=(FModulationDestination&& InModulationDestination)
 	{
-		DeviceId = MoveTemp(InModulationDestination.DeviceId);
-		ValueTarget = MoveTemp(InModulationDestination.ValueTarget);
-		bIsBuffered = MoveTemp(InModulationDestination.bIsBuffered);
-		bValueNormalized = MoveTemp(InModulationDestination.bValueNormalized);
-		bHasProcessed = MoveTemp(InModulationDestination.bHasProcessed);
-		OutputBuffer = MoveTemp(InModulationDestination.OutputBuffer);
-
-		TSet<FModulatorHandle> NewHandles;
-		{
-			FScopeLock OtherLock(&InModulationDestination.HandleCritSection);
-			NewHandles = MoveTemp(InModulationDestination.Handles);
-		}
-		{
-			FScopeLock Lock(&HandleCritSection);
-			Handles = MoveTemp(NewHandles);
-		}
-
-		Parameter = MoveTemp(InModulationDestination.Parameter);
-
+		*DestinationData = MoveTemp(*InModulationDestination.DestinationData);
 		return *this;
 	}
 
-	void FModulationDestination::ResetHandles()
+	void FModulationDestination::FModulationDestinationData::ResetHandles()
 	{
 		Audio::FModulationParameter ParameterCopy = Parameter;
 
@@ -137,36 +89,41 @@ namespace Audio
 
 	void FModulationDestination::Init(FDeviceId InDeviceId, FName InParameterName, bool bInIsBuffered, bool bInValueNormalized)
 	{
-		DeviceId = InDeviceId;
-		bIsBuffered = bInIsBuffered;
-		bValueNormalized = bInValueNormalized;
+		DestinationData->DeviceId = InDeviceId;
+		DestinationData->bIsBuffered = bInIsBuffered;
+		DestinationData->bValueNormalized = bInValueNormalized;
 
-		OutputBuffer.Reset();
-		Parameter = Audio::GetModulationParameter(InParameterName);
+		DestinationData->OutputBuffer.Reset();
+		DestinationData->Parameter = Audio::GetModulationParameter(InParameterName);
 
-		ResetHandles();
+		DestinationData->ResetHandles();
 	}
 
 	bool FModulationDestination::IsActive() const
 	{
-		FScopeLock Lock(&HandleCritSection);
-		return Algo::AnyOf(Handles, [](const FModulatorHandle& Handle) { return Handle.IsValid(); });
+		FScopeLock Lock(&DestinationData->HandleCritSection);
+		return Algo::AnyOf(DestinationData->Handles, [](const FModulatorHandle& Handle) { return Handle.IsValid(); });
 	}
 
 	bool FModulationDestination::ProcessControl(float InValueUnitBase, int32 InNumSamples)
 	{
-		bHasProcessed = true;
+		FModulationParameter& Parameter = DestinationData->Parameter;
+		float& ValueTarget = DestinationData->ValueTarget;
+		FAlignedFloatBuffer& OutputBuffer = DestinationData->OutputBuffer;
+		
+		DestinationData->bHasProcessed = true;
 		float LastTarget = ValueTarget;
 
+
 		float NewTargetNormalized = Parameter.DefaultValue;
 		if (Parameter.bRequiresConversion)
 		{
 			Parameter.NormalizedFunction(NewTargetNormalized);
 		}
 
-		FScopeLock Lock(&HandleCritSection);
+		FScopeLock Lock(&DestinationData->HandleCritSection);
 		{
-			for (const FModulatorHandle& Handle : Handles)
+			for (const FModulatorHandle& Handle : DestinationData->Handles)
 			{
 				if (Handle.IsValid())
 				{
@@ -189,12 +146,12 @@ namespace Audio
 		ValueTarget = NewTargetNormalized;
 
 		// Convert target to unit space if required
-		if (Parameter.bRequiresConversion && !bValueNormalized)
+		if (Parameter.bRequiresConversion && !DestinationData->bValueNormalized)
 		{
 			Parameter.UnitFunction(ValueTarget);
 		}
 
-		if (bIsBuffered)
+		if (DestinationData->bIsBuffered)
 		{
 			if (OutputBuffer.Num() != InNumSamples)
 			{
@@ -284,34 +241,98 @@ namespace Audio
 		UpdateModulatorsInternal(MoveTemp(ProxySettings));
 	}
 
+	void FModulationDestination::FModulationDestinationData::SetHandles(TSet<FModulatorHandle>&& NewHandles)
+	{
+		FScopeLock Lock(&HandleCritSection);
+		Handles = MoveTemp(NewHandles);
+	}
+
+	FModulationDestination::FModulationDestinationData& FModulationDestination::FModulationDestinationData::operator=(const FModulationDestinationData& InDestinationInfo)
+	{
+		DeviceId = InDestinationInfo.DeviceId;
+		ValueTarget = InDestinationInfo.ValueTarget;
+		bIsBuffered = InDestinationInfo.bIsBuffered;
+		bValueNormalized = InDestinationInfo.bValueNormalized;
+		OutputBuffer = InDestinationInfo.OutputBuffer;
+
+		TSet<FModulatorHandle> NewHandles;
+		{
+			FScopeLock OtherLock(&(InDestinationInfo.HandleCritSection));
+			NewHandles = InDestinationInfo.Handles;
+		}
+
+		{
+			FScopeLock Lock(&HandleCritSection);
+			Handles = MoveTemp(NewHandles);
+		}
+
+		Parameter = InDestinationInfo.Parameter;
+
+		return *this;
+	}
+
+	FModulationDestination::FModulationDestinationData& FModulationDestination::FModulationDestinationData::operator=(FModulationDestinationData&& InDestinationInfo)
+	{
+		DeviceId = MoveTemp(InDestinationInfo.DeviceId);
+		ValueTarget = MoveTemp(InDestinationInfo.ValueTarget);
+		bIsBuffered = MoveTemp(InDestinationInfo.bIsBuffered);
+		bValueNormalized = MoveTemp(InDestinationInfo.bValueNormalized);
+		bHasProcessed = MoveTemp(InDestinationInfo.bHasProcessed);
+		OutputBuffer = MoveTemp(InDestinationInfo.OutputBuffer);
+
+		TSet<FModulatorHandle> NewHandles;
+		{
+			FScopeLock OtherLock(&InDestinationInfo.HandleCritSection);
+			NewHandles = MoveTemp(InDestinationInfo.Handles);
+		}
+		{
+			FScopeLock Lock(&HandleCritSection);
+			Handles = MoveTemp(NewHandles);
+		}
+
+		Parameter = MoveTemp(InDestinationInfo.Parameter);
+
+		return *this;
+	}
+
+	const FDeviceId& FModulationDestination::FModulationDestinationData::GetDeviceId() const
+	{
+		return DeviceId;
+	}
+
+	const FModulationParameter& FModulationDestination::FModulationDestinationData::GetParameter() const
+	{
+		return Parameter;
+	}
+
 	void FModulationDestination::UpdateModulatorsInternal(TArray<TUniquePtr<Audio::IModulatorSettings>>&& ProxySettings)
 	{
-		auto UpdateHandleLambda = [this, ModSettings = MoveTemp(ProxySettings)]() mutable
+		auto UpdateHandleLambda = [DestinationDataPtr = TSharedPtr<FModulationDestinationData>(DestinationData), ModSettings = MoveTemp(ProxySettings)]() mutable
 		{
-			if (FAudioDeviceManager* AudioDeviceManager = FAudioDeviceManager::Get())
+			if (DestinationDataPtr.IsValid())
 			{
-				if (FAudioDevice* AudioDevice = AudioDeviceManager->GetAudioDeviceRaw(DeviceId))
+				if (FAudioDeviceManager* AudioDeviceManager = FAudioDeviceManager::Get())
 				{
-					if (AudioDevice->IsModulationPluginEnabled() && AudioDevice->ModulationInterface.IsValid())
+					if (FAudioDevice* AudioDevice = AudioDeviceManager->GetAudioDeviceRaw(DestinationDataPtr->GetDeviceId()))
 					{
-						if (IAudioModulationManager* Modulation = AudioDevice->ModulationInterface.Get())
+						if (AudioDevice->IsModulationPluginEnabled() && AudioDevice->ModulationInterface.IsValid())
 						{
-							TSet<FModulatorHandle> NewHandles;
-							for (TUniquePtr<Audio::IModulatorSettings>& ModSetting : ModSettings)
+							if (IAudioModulationManager* Modulation = AudioDevice->ModulationInterface.Get())
 							{
-								Audio::FModulationParameter HandleParam = Parameter;
-								NewHandles.Add(FModulatorHandle{ *Modulation, *ModSetting.Get(), MoveTemp(HandleParam) });
+								TSet<FModulatorHandle> NewHandles;
+								for (TUniquePtr<Audio::IModulatorSettings>& ModSetting : ModSettings)
+								{
+									Audio::FModulationParameter HandleParam = DestinationDataPtr->GetParameter();
+									NewHandles.Add(FModulatorHandle{ *Modulation, *ModSetting.Get(), MoveTemp(HandleParam) });
+								}
+								DestinationDataPtr->SetHandles(MoveTemp(NewHandles));
 							}
-
-							FScopeLock Lock(&HandleCritSection);
-							Handles = MoveTemp(NewHandles);
+							return;
 						}
-						return;
 					}
 				}
+				DestinationDataPtr->ResetHandles();
 			}
-
-			ResetHandles();
 		};
 
 		FAudioThread::RunCommandOnAudioThread(MoveTemp(UpdateHandleLambda));
