Fix bug in AudioDecompress caused by `SeekToTime`. Need to update the `CurrentSampleCount` value to the newly sample that we're seeking to.This happens in "StreamCompressedData" where we do the actual seek (and confirm that it's a valid seek)

Also, reset PCM offset values for last decode as seeking should invalidate that data. The new data that we should be reading should start fromt he new seek time.

We haven't seen any issues with the currenty behavior because the SoundWaveProxyReader currently creates a completely new decoder when seeking.

diff --git a/Engine/Source/Runtime/Engine/Private/AudioDecompress.cpp b/Engine/Source/Runtime/Engine/Private/AudioDecompress.cpp
--- a/Engine/Source/Runtime/Engine/Private/AudioDecompress.cpp
+++ b/Engine/Source/Runtime/Engine/Private/AudioDecompress.cpp
@@ -300,7 +300,12 @@ bool IStreamedCompressedInfo::StreamCompressedData(uint8* Destination, bool bLoo
 
 							// Set our offset from the table.
 							SrcBufferOffset = Offset + TableOffset;
-							SrcBufferOffset += CurrentChunkIndex == 0 ? AudioDataOffset : 0;							
+							SrcBufferOffset += CurrentChunkIndex == 0 ? AudioDataOffset : 0;
+
+							// Convert frames to samples and update "current sample count"
+							CurrentSampleCount = StreamSeekToAudioFrames * NumChannels;
+							LastPCMByteSize = 0;
+							LastPCMOffset = 0;
 						}
 						else // Seek failed (off the end of the chunk).
 						{
