Added BP util to test whether a given class implements a given interface

diff --git a/Engine/Source/Runtime/Engine/Classes/Kismet/KismetSystemLibrary.h b/Engine/Source/Runtime/Engine/Classes/Kismet/KismetSystemLibrary.h
--- a/Engine/Source/Runtime/Engine/Classes/Kismet/KismetSystemLibrary.h
+++ b/Engine/Source/Runtime/Engine/Classes/Kismet/KismetSystemLibrary.h
@@ -202,10 +202,14 @@ class UKismetSystemLibrary : public UBlueprintFunctionLibrary
 	UFUNCTION(BlueprintPure, Category = "Utilities|Platform")
 	static ENGINE_API FString GetPlatformUserDir();
 
-	/** Checks if this object implements a specific interface, works for both native and blueprint interfacse */
-	UFUNCTION(BlueprintPure, Category="Utilities")
+	/** Checks if the given object implements a specific interface, works for both native and blueprint interfacse */
+	UFUNCTION(BlueprintPure, Category="Utilities", meta=(DisplayName = "Does Object Implement Interface"))
 	static ENGINE_API bool DoesImplementInterface(const UObject* TestObject, TSubclassOf<UInterface> Interface);
 
+	/** Checks if the given class implements a specific interface, works for both native and blueprint interfacse */
+	UFUNCTION(BlueprintPure, Category="Utilities")
+	static ENGINE_API bool DoesClassImplementInterface(const UClass* TestClass, TSubclassOf<UInterface> Interface);
+
 	/** 
 	 * Get the current game time, in seconds. This stops when the game is paused and is affected by slomo. 
 	 * 
diff --git a/Engine/Source/Runtime/Engine/Private/KismetSystemLibrary.cpp b/Engine/Source/Runtime/Engine/Private/KismetSystemLibrary.cpp
--- a/Engine/Source/Runtime/Engine/Private/KismetSystemLibrary.cpp
+++ b/Engine/Source/Runtime/Engine/Private/KismetSystemLibrary.cpp
@@ -256,10 +256,25 @@ FString UKismetSystemLibrary::GetPlatformUserDir()
 
 bool UKismetSystemLibrary::DoesImplementInterface(const UObject* TestObject, TSubclassOf<UInterface> Interface)
 {
-	if (Interface != NULL && TestObject != NULL)
+	if (TestObject)
 	{
-		checkf(Interface->IsChildOf(UInterface::StaticClass()), TEXT("Interface parameter %s is not actually an interface."), *Interface->GetName());
-		return TestObject->GetClass()->ImplementsInterface(Interface);
+		return UKismetSystemLibrary::DoesClassImplementInterface(TestObject->GetClass(), Interface);
+	}
+
+	return false;
+}
+
+bool UKismetSystemLibrary::DoesClassImplementInterface(const UClass* TestClass, TSubclassOf<UInterface> Interface)
+{
+	if (TestClass && Interface)
+	{
+		if (!Interface->IsChildOf(UInterface::StaticClass()))
+		{
+			LogRuntimeError(FText::Format(LOCTEXT("DoesClassImplementInterface.InvalidInterface", "Interface parameter {0} is not actually an interface."), FText::AsCultureInvariant(Interface->GetName())));
+			return false;
+		}
+
+		return TestClass->ImplementsInterface(Interface);
 	}
 
 	return false;
