Work-in-progress changes for the new HLSL translator. This change doesn't impact current or future projects unless the new translator is enabled. Implemented GenerateHLSLExpression interface for numerous UMaterialExpression variants. Added a hash map to reduce the time needed to prepare complex HLSL trees. Fixed a bug where FExpressionSetStructField::PrepareValue can cause unrequested struct fields be prepared. Fixed an error in FExpressionLocalPHI::PrepareValue caused by some live values being nullptrs. Updated GenerateMaterialTemplateHLSL to also use named paramters.

#tests loaded a level in editor with the new translator disabled and checked visuals

diff --git a/Engine/Plugins/Interchange/Runtime/Source/Import/Private/MaterialX/MaterialExpressions/MaterialExpressionRemap.cpp b/Engine/Plugins/Interchange/Runtime/Source/Import/Private/MaterialX/MaterialExpressions/MaterialExpressionRemap.cpp
--- a/Engine/Plugins/Interchange/Runtime/Source/Import/Private/MaterialX/MaterialExpressions/MaterialExpressionRemap.cpp
+++ b/Engine/Plugins/Interchange/Runtime/Source/Import/Private/MaterialX/MaterialExpressions/MaterialExpressionRemap.cpp
@@ -1,6 +1,7 @@
 // Copyright Epic Games, Inc. All Rights Reserved.
 #include "MaterialExpressionRemap.h"
 #include "MaterialCompiler.h"
+#include "MaterialHLSLGenerator.h"
 
 #include UE_INLINE_GENERATED_CPP_BY_NAME(MaterialExpressionRemap)
 
@@ -51,6 +52,32 @@ int32 UMaterialExpressionMaterialXRemap::Compile(FMaterialCompiler* Compiler, in
 	return Compiler->Add(TargetLowIndex, Compiler->Mul(Compiler->Sub(Input.Compile(Compiler), InputLowIndex), Delta));
 }
 
+bool UMaterialExpressionMaterialXRemap::GenerateHLSLExpression(FMaterialHLSLGenerator& Generator, UE::HLSLTree::FScope& Scope, int32 OutputIndex, UE::HLSLTree::FExpression const*& OutExpression) const
+{
+	using namespace UE::HLSLTree;
+
+	if (!Input.GetTracedInput().Expression)
+	{
+		return Generator.Errorf(TEXT("Missing remap input"));
+	}
+
+	// Min2 + (Input - Min1) x (Max2 - Min2) / (Max1 - Min1)
+	const FExpression* InputLowExpression = InputLow.AcquireHLSLExpressionOrConstant(Generator, Scope, InputLowDefault);
+	const FExpression* InputHighExpression = InputHigh.AcquireHLSLExpressionOrConstant(Generator, Scope, InputHighDefault);
+
+	const FExpression* TargetLowExpression = TargetLow.AcquireHLSLExpressionOrConstant(Generator, Scope, TargetLowDefault);
+	const FExpression* TargetHighExpression = TargetHigh.AcquireHLSLExpressionOrConstant(Generator, Scope, TargetHighDefault);
+
+	const FExpression* InputDeltaExpression = Generator.GetTree().NewSub(InputHighExpression, InputLowExpression);
+	const FExpression* TargetDeltaExpression = Generator.GetTree().NewSub(TargetHighExpression, TargetLowExpression);
+	const FExpression* ScaleExpression = Generator.GetTree().NewDiv(TargetDeltaExpression, InputDeltaExpression);
+
+	const FExpression* InputExpression = Input.AcquireHLSLExpression(Generator, Scope);
+
+	OutExpression = Generator.GetTree().NewAdd(TargetLowExpression, Generator.GetTree().NewMul(Generator.GetTree().NewSub(InputExpression, InputLowExpression), ScaleExpression));
+	return true;
+}
+
 void UMaterialExpressionMaterialXRemap::GetCaption(TArray<FString>& OutCaptions) const
 {
 	OutCaptions.Add(TEXT("MaterialX Remap"));
diff --git a/Engine/Plugins/Interchange/Runtime/Source/Import/Private/MaterialX/MaterialExpressions/MaterialExpressionRemap.h b/Engine/Plugins/Interchange/Runtime/Source/Import/Private/MaterialX/MaterialExpressions/MaterialExpressionRemap.h
--- a/Engine/Plugins/Interchange/Runtime/Source/Import/Private/MaterialX/MaterialExpressions/MaterialExpressionRemap.h
+++ b/Engine/Plugins/Interchange/Runtime/Source/Import/Private/MaterialX/MaterialExpressions/MaterialExpressionRemap.h
@@ -45,6 +45,8 @@ class UMaterialExpressionMaterialXRemap : public UMaterialExpression
 #if WITH_EDITOR
 	virtual int32 Compile(class FMaterialCompiler* Compiler, int32 OutputIndex) override;
 	virtual void GetCaption(TArray<FString>& OutCaptions) const override;
+
+	virtual bool GenerateHLSLExpression(class FMaterialHLSLGenerator& Generator, class UE::HLSLTree::FScope& Scope, int32 OutputIndex, class UE::HLSLTree::FExpression const*& OutExpression) const override;
 #endif
 	//~ End UMaterialExpressionMaterialX Interface
 };
diff --git a/Engine/Source/Runtime/Engine/Classes/Materials/MaterialExpressionAtmosphericLightVector.h b/Engine/Source/Runtime/Engine/Classes/Materials/MaterialExpressionAtmosphericLightVector.h
--- a/Engine/Source/Runtime/Engine/Classes/Materials/MaterialExpressionAtmosphericLightVector.h
+++ b/Engine/Source/Runtime/Engine/Classes/Materials/MaterialExpressionAtmosphericLightVector.h
@@ -17,6 +17,8 @@ class UMaterialExpressionAtmosphericLightVector : public UMaterialExpression
 #if WITH_EDITOR
 	virtual int32 Compile(class FMaterialCompiler* Compiler, int32 OutputIndex) override;
 	virtual void GetCaption(TArray<FString>& OutCaptions) const override;
+
+	virtual bool GenerateHLSLExpression(FMaterialHLSLGenerator& Generator, UE::HLSLTree::FScope& Scope, int32 OutputIndex, UE::HLSLTree::FExpression const*& OutExpression) const override;
 #endif
 	//~ End UMaterialExpression Interface
 };
diff --git a/Engine/Source/Runtime/Engine/Classes/Materials/MaterialExpressionDepthFade.h b/Engine/Source/Runtime/Engine/Classes/Materials/MaterialExpressionDepthFade.h
--- a/Engine/Source/Runtime/Engine/Classes/Materials/MaterialExpressionDepthFade.h
+++ b/Engine/Source/Runtime/Engine/Classes/Materials/MaterialExpressionDepthFade.h
@@ -52,6 +52,8 @@ class UMaterialExpressionDepthFade : public UMaterialExpression
 	{
 		OutCaptions.Add(TEXT("Depth Fade"));
 	}
+
+	virtual bool GenerateHLSLExpression(FMaterialHLSLGenerator& Generator, UE::HLSLTree::FScope& Scope, int32 OutputIndex, UE::HLSLTree::FExpression const*& OutExpression) const override;
 #endif
 	//~ End UMaterialExpression Interface#endif
 };
diff --git a/Engine/Source/Runtime/Engine/Classes/Materials/MaterialExpressionDeriveNormalZ.h b/Engine/Source/Runtime/Engine/Classes/Materials/MaterialExpressionDeriveNormalZ.h
--- a/Engine/Source/Runtime/Engine/Classes/Materials/MaterialExpressionDeriveNormalZ.h
+++ b/Engine/Source/Runtime/Engine/Classes/Materials/MaterialExpressionDeriveNormalZ.h
@@ -22,6 +22,8 @@ class UMaterialExpressionDeriveNormalZ : public UMaterialExpression
 #if WITH_EDITOR
 	virtual int32 Compile(class FMaterialCompiler* Compiler, int32 OutputIndex) override;
 	virtual void GetCaption(TArray<FString>& OutCaptions) const override;
+
+	virtual bool GenerateHLSLExpression(FMaterialHLSLGenerator& Generator, UE::HLSLTree::FScope& Scope, int32 OutputIndex, UE::HLSLTree::FExpression const*& OutExpression) const override;
 #endif
 	//~ End UMaterialExpression Interface
 };
diff --git a/Engine/Source/Runtime/Engine/Classes/Materials/MaterialExpressionDistanceFieldsRenderingSwitch.h b/Engine/Source/Runtime/Engine/Classes/Materials/MaterialExpressionDistanceFieldsRenderingSwitch.h
--- a/Engine/Source/Runtime/Engine/Classes/Materials/MaterialExpressionDistanceFieldsRenderingSwitch.h
+++ b/Engine/Source/Runtime/Engine/Classes/Materials/MaterialExpressionDistanceFieldsRenderingSwitch.h
@@ -29,6 +29,8 @@ class UMaterialExpressionDistanceFieldsRenderingSwitch : public UMaterialExpress
 	virtual uint32 GetInputType(int32 InputIndex) override { return MCT_Unknown; }
 	virtual uint32 GetOutputType(int32 OutputIndex) override { return MCT_Unknown; }
 	virtual void GetCaption(TArray<FString>& OutCaptions) const override;
+
+	virtual bool GenerateHLSLExpression(FMaterialHLSLGenerator& Generator, UE::HLSLTree::FScope& Scope, int32 OutputIndex, UE::HLSLTree::FExpression const*& OutExpression) const override;
 #endif
 	//~ End UMaterialExpression Interface
 };
diff --git a/Engine/Source/Runtime/Engine/Classes/Materials/MaterialExpressionMaterialProxyReplace.h b/Engine/Source/Runtime/Engine/Classes/Materials/MaterialExpressionMaterialProxyReplace.h
--- a/Engine/Source/Runtime/Engine/Classes/Materials/MaterialExpressionMaterialProxyReplace.h
+++ b/Engine/Source/Runtime/Engine/Classes/Materials/MaterialExpressionMaterialProxyReplace.h
@@ -29,6 +29,8 @@ class UMaterialExpressionMaterialProxyReplace : public UMaterialExpression
 	virtual uint32 GetInputType(int32 InputIndex) override { return MCT_Unknown; }
 	virtual uint32 GetOutputType(int32 OutputIndex) override { return MCT_Unknown; }
 	virtual void GetCaption(TArray<FString>& OutCaptions) const override;
+
+	virtual bool GenerateHLSLExpression(FMaterialHLSLGenerator& Generator, UE::HLSLTree::FScope& Scope, int32 OutputIndex, UE::HLSLTree::FExpression const*& OutExpression) const override;
 #endif
 	// End UMaterialExpression Interface
 };
diff --git a/Engine/Source/Runtime/Engine/Classes/Materials/MaterialExpressionParticleDirection.h b/Engine/Source/Runtime/Engine/Classes/Materials/MaterialExpressionParticleDirection.h
--- a/Engine/Source/Runtime/Engine/Classes/Materials/MaterialExpressionParticleDirection.h
+++ b/Engine/Source/Runtime/Engine/Classes/Materials/MaterialExpressionParticleDirection.h
@@ -22,6 +22,8 @@ class UMaterialExpressionParticleDirection : public UMaterialExpression
 #if WITH_EDITOR
 	virtual int32 Compile(class FMaterialCompiler* Compiler, int32 OutputIndex) override;
 	virtual void GetCaption(TArray<FString>& OutCaptions) const override;
+
+	virtual bool GenerateHLSLExpression(FMaterialHLSLGenerator& Generator, UE::HLSLTree::FScope& Scope, int32 OutputIndex, UE::HLSLTree::FExpression const*& OutExpression) const override;
 #endif
 	//~ End UMaterialExpression Interface
 };
diff --git a/Engine/Source/Runtime/Engine/Classes/Materials/MaterialExpressionParticleRandom.h b/Engine/Source/Runtime/Engine/Classes/Materials/MaterialExpressionParticleRandom.h
--- a/Engine/Source/Runtime/Engine/Classes/Materials/MaterialExpressionParticleRandom.h
+++ b/Engine/Source/Runtime/Engine/Classes/Materials/MaterialExpressionParticleRandom.h
@@ -22,6 +22,8 @@ class UMaterialExpressionParticleRandom: public UMaterialExpression
 #if WITH_EDITOR
 	virtual int32 Compile(class FMaterialCompiler* Compiler, int32 OutputIndex) override;
 	virtual void GetCaption(TArray<FString>& OutCaptions) const override;
+
+	virtual bool GenerateHLSLExpression(FMaterialHLSLGenerator& Generator, UE::HLSLTree::FScope& Scope, int32 OutputIndex, UE::HLSLTree::FExpression const*& OutExpression) const override;
 #endif
 	//~ End UMaterialExpression Interface
 };
diff --git a/Engine/Source/Runtime/Engine/Classes/Materials/MaterialExpressionParticleRelativeTime.h b/Engine/Source/Runtime/Engine/Classes/Materials/MaterialExpressionParticleRelativeTime.h
--- a/Engine/Source/Runtime/Engine/Classes/Materials/MaterialExpressionParticleRelativeTime.h
+++ b/Engine/Source/Runtime/Engine/Classes/Materials/MaterialExpressionParticleRelativeTime.h
@@ -22,6 +22,8 @@ class UMaterialExpressionParticleRelativeTime : public UMaterialExpression
 #if WITH_EDITOR
 	virtual int32 Compile(class FMaterialCompiler* Compiler, int32 OutputIndex) override;
 	virtual void GetCaption(TArray<FString>& OutCaptions) const override;
+
+	virtual bool GenerateHLSLExpression(FMaterialHLSLGenerator& Generator, UE::HLSLTree::FScope& Scope, int32 OutputIndex, UE::HLSLTree::FExpression const*& OutExpression) const override;
 #endif
 	//~ End UMaterialExpression Interface
 };
diff --git a/Engine/Source/Runtime/Engine/Classes/Materials/MaterialExpressionParticleSize.h b/Engine/Source/Runtime/Engine/Classes/Materials/MaterialExpressionParticleSize.h
--- a/Engine/Source/Runtime/Engine/Classes/Materials/MaterialExpressionParticleSize.h
+++ b/Engine/Source/Runtime/Engine/Classes/Materials/MaterialExpressionParticleSize.h
@@ -22,6 +22,8 @@ class UMaterialExpressionParticleSize : public UMaterialExpression
 #if WITH_EDITOR
 	virtual int32 Compile(class FMaterialCompiler* Compiler, int32 OutputIndex) override;
 	virtual void GetCaption(TArray<FString>& OutCaptions) const override;
+
+	virtual bool GenerateHLSLExpression(FMaterialHLSLGenerator& Generator, UE::HLSLTree::FScope& Scope, int32 OutputIndex, UE::HLSLTree::FExpression const*& OutExpression) const override;
 #endif
 	//~ End UMaterialExpression Interface
 };
diff --git a/Engine/Source/Runtime/Engine/Classes/Materials/MaterialExpressionParticleSpeed.h b/Engine/Source/Runtime/Engine/Classes/Materials/MaterialExpressionParticleSpeed.h
--- a/Engine/Source/Runtime/Engine/Classes/Materials/MaterialExpressionParticleSpeed.h
+++ b/Engine/Source/Runtime/Engine/Classes/Materials/MaterialExpressionParticleSpeed.h
@@ -22,6 +22,8 @@ class UMaterialExpressionParticleSpeed : public UMaterialExpression
 #if WITH_EDITOR
 	virtual int32 Compile(class FMaterialCompiler* Compiler, int32 OutputIndex) override;
 	virtual void GetCaption(TArray<FString>& OutCaptions) const override;
+
+	virtual bool GenerateHLSLExpression(FMaterialHLSLGenerator& Generator, UE::HLSLTree::FScope& Scope, int32 OutputIndex, UE::HLSLTree::FExpression const*& OutExpression) const override;
 #endif
 	//~ End UMaterialExpression Interface
 };
diff --git a/Engine/Source/Runtime/Engine/Classes/Materials/MaterialExpressionPreviousFrameSwitch.h b/Engine/Source/Runtime/Engine/Classes/Materials/MaterialExpressionPreviousFrameSwitch.h
--- a/Engine/Source/Runtime/Engine/Classes/Materials/MaterialExpressionPreviousFrameSwitch.h
+++ b/Engine/Source/Runtime/Engine/Classes/Materials/MaterialExpressionPreviousFrameSwitch.h
@@ -27,6 +27,8 @@ class UMaterialExpressionPreviousFrameSwitch : public UMaterialExpression
 	virtual bool IsResultMaterialAttributes(int32 OutputIndex) override;
 	virtual uint32 GetInputType(int32 InputIndex) override;
 	virtual uint32 GetOutputType(int32 OutputIndex) override {return MCT_Unknown;}
+
+	virtual bool GenerateHLSLExpression(FMaterialHLSLGenerator& Generator, UE::HLSLTree::FScope& Scope, int32 OutputIndex, UE::HLSLTree::FExpression const*& OutExpression) const override;
 #endif // WITH_EDITOR
 	//~ End UMaterialExpression Interface
 };
diff --git a/Engine/Source/Runtime/Engine/Classes/Materials/MaterialExpressionRayTracingQualitySwitch.h b/Engine/Source/Runtime/Engine/Classes/Materials/MaterialExpressionRayTracingQualitySwitch.h
--- a/Engine/Source/Runtime/Engine/Classes/Materials/MaterialExpressionRayTracingQualitySwitch.h
+++ b/Engine/Source/Runtime/Engine/Classes/Materials/MaterialExpressionRayTracingQualitySwitch.h
@@ -33,6 +33,8 @@ class UMaterialExpressionRayTracingQualitySwitch : public UMaterialExpression
 	virtual bool IsResultSubstrateMaterial(int32 OutputIndex) override;
 	virtual void GatherSubstrateMaterialInfo(FSubstrateMaterialInfo& SubstrateMaterialInfo, int32 OutputIndex) override;
 	virtual FSubstrateOperator* SubstrateGenerateMaterialTopologyTree(class FMaterialCompiler* Compiler, class UMaterialExpression* Parent, int32 OutputIndex) override;
+
+	virtual bool GenerateHLSLExpression(FMaterialHLSLGenerator& Generator, UE::HLSLTree::FScope& Scope, int32 OutputIndex, UE::HLSLTree::FExpression const*& OutExpression) const override;
 #endif
 	//~ End UMaterialExpression Interface
 };
diff --git a/Engine/Source/Runtime/Engine/Classes/Materials/MaterialExpressionRuntimeVirtualTextureSample.h b/Engine/Source/Runtime/Engine/Classes/Materials/MaterialExpressionRuntimeVirtualTextureSample.h
--- a/Engine/Source/Runtime/Engine/Classes/Materials/MaterialExpressionRuntimeVirtualTextureSample.h
+++ b/Engine/Source/Runtime/Engine/Classes/Materials/MaterialExpressionRuntimeVirtualTextureSample.h
@@ -124,6 +124,8 @@ protected:
 	ENGINE_API virtual bool CanEditChange(const FProperty* InProperty) const override;
 public:
 	ENGINE_API virtual void PostEditChangeProperty(FPropertyChangedEvent& PropertyChangedEvent) override;
+
+	virtual bool GenerateHLSLExpression(FMaterialHLSLGenerator& Generator, UE::HLSLTree::FScope& Scope, int32 OutputIndex, UE::HLSLTree::FExpression const*& OutExpression) const override;
 #endif
 	//~ End UMaterialExpression Interface
 };
diff --git a/Engine/Source/Runtime/Engine/Classes/Materials/MaterialExpressionSceneColor.h b/Engine/Source/Runtime/Engine/Classes/Materials/MaterialExpressionSceneColor.h
--- a/Engine/Source/Runtime/Engine/Classes/Materials/MaterialExpressionSceneColor.h
+++ b/Engine/Source/Runtime/Engine/Classes/Materials/MaterialExpressionSceneColor.h
@@ -55,6 +55,8 @@ class UMaterialExpressionSceneColor : public UMaterialExpression
 	virtual int32 Compile(class FMaterialCompiler* Compiler, int32 OutputIndex) override;
 	virtual void GetCaption(TArray<FString>& OutCaptions) const override;
 	virtual FName GetInputName(int32 InputIndex) const override;
+
+	virtual bool GenerateHLSLExpression(FMaterialHLSLGenerator& Generator, UE::HLSLTree::FScope& Scope, int32 OutputIndex, UE::HLSLTree::FExpression const*& OutExpression) const override;
 #endif
 	//~ End UMaterialExpression Interface
 };
diff --git a/Engine/Source/Runtime/Engine/Classes/Materials/MaterialExpressionShadowReplace.h b/Engine/Source/Runtime/Engine/Classes/Materials/MaterialExpressionShadowReplace.h
--- a/Engine/Source/Runtime/Engine/Classes/Materials/MaterialExpressionShadowReplace.h
+++ b/Engine/Source/Runtime/Engine/Classes/Materials/MaterialExpressionShadowReplace.h
@@ -25,6 +25,8 @@ class UMaterialExpressionShadowReplace : public UMaterialExpression
 	virtual int32 Compile(class FMaterialCompiler* Compiler, int32 OutputIndex) override;
 	virtual void GetCaption(TArray<FString>& OutCaptions) const override;
 	virtual void GetExpressionToolTip(TArray<FString>& OutToolTip) override;
+
+	virtual bool GenerateHLSLExpression(FMaterialHLSLGenerator& Generator, UE::HLSLTree::FScope& Scope, int32 OutputIndex, UE::HLSLTree::FExpression const*& OutExpression) const override;
 #endif
 	//~ End UMaterialExpression Interface
 };
diff --git a/Engine/Source/Runtime/Engine/Classes/Materials/MaterialExpressionVectorNoise.h b/Engine/Source/Runtime/Engine/Classes/Materials/MaterialExpressionVectorNoise.h
--- a/Engine/Source/Runtime/Engine/Classes/Materials/MaterialExpressionVectorNoise.h
+++ b/Engine/Source/Runtime/Engine/Classes/Materials/MaterialExpressionVectorNoise.h
@@ -80,6 +80,8 @@ class UMaterialExpressionVectorNoise : public UMaterialExpression
 	virtual bool CanEditChange(const FProperty* InProperty) const override;
 	virtual int32 Compile(class FMaterialCompiler* Compiler, int32 OutputIndex) override;
 	virtual void GetCaption(TArray<FString>& OutCaptions) const override;
+
+	virtual bool GenerateHLSLExpression(FMaterialHLSLGenerator& Generator, UE::HLSLTree::FScope& Scope, int32 OutputIndex, UE::HLSLTree::FExpression const*& OutExpression) const override;
 #endif
 	//~ End UMaterialExpression Interface
 };
diff --git a/Engine/Source/Runtime/Engine/Classes/Materials/MaterialExpressionVirtualTextureFeatureSwitch.h b/Engine/Source/Runtime/Engine/Classes/Materials/MaterialExpressionVirtualTextureFeatureSwitch.h
--- a/Engine/Source/Runtime/Engine/Classes/Materials/MaterialExpressionVirtualTextureFeatureSwitch.h
+++ b/Engine/Source/Runtime/Engine/Classes/Materials/MaterialExpressionVirtualTextureFeatureSwitch.h
@@ -30,6 +30,8 @@ class UMaterialExpressionVirtualTextureFeatureSwitch : public UMaterialExpressio
 	virtual uint32 GetInputType(int32 InputIndex) override { return MCT_Unknown; }
 	virtual uint32 GetOutputType(int32 OutputIndex) override { return MCT_Unknown; }
 	virtual void GetCaption(TArray<FString>& OutCaptions) const override;
+
+	virtual bool GenerateHLSLExpression(FMaterialHLSLGenerator& Generator, UE::HLSLTree::FScope& Scope, int32 OutputIndex, UE::HLSLTree::FExpression const*& OutExpression) const override;
 #endif
 	//~ End UMaterialExpression Interface
 };
diff --git a/Engine/Source/Runtime/Engine/Private/HLSLTree/HLSLTree.cpp b/Engine/Source/Runtime/Engine/Private/HLSLTree/HLSLTree.cpp
--- a/Engine/Source/Runtime/Engine/Private/HLSLTree/HLSLTree.cpp
+++ b/Engine/Source/Runtime/Engine/Private/HLSLTree/HLSLTree.cpp
@@ -1443,13 +1443,28 @@ bool FTree::Finalize()
 				{
 					const FExpressionDerivatives Derivatives = GetAnalyticDerivatives(LocalValue);
 					LocalValue = (ChainType == ELocalPHIChainType::Ddx) ? Derivatives.ExpressionDdx : Derivatives.ExpressionDdy;
+					
+					// TODO: Revisit. Can we fallback to hardware derivatives? Handle fallback in FExpression::ComputeAnalyticDerivatives better?
+					if (!LocalValue)
+					{
+						LocalValue = NewConstant(0.f);
+					}
 				}
 				else
 				{
-					check(ChainType == ELocalPHIChainType::PreviousFrame);
+					check(ChainType == ELocalPHIChainType::PreviousFrame && LocalValue);
+
+					const FExpression* PrevLocalValue = LocalValue;
 					LocalValue = GetPreviousFrame(LocalValue, Entry.RequestedType);
+
+					// TODO: Revisit. Is this correct? The intention is falling back to current frame if previous frame value is invalid
+					if (!LocalValue)
+					{
+						LocalValue = PrevLocalValue;
+					}
 				}
 			}
+
 			// May be nullptr if derivatives are not valid
 			Expression->Values[i] = LocalValue;
 		}
diff --git a/Engine/Source/Runtime/Engine/Private/HLSLTree/HLSLTreeCommon.cpp b/Engine/Source/Runtime/Engine/Private/HLSLTree/HLSLTreeCommon.cpp
--- a/Engine/Source/Runtime/Engine/Private/HLSLTree/HLSLTreeCommon.cpp
+++ b/Engine/Source/Runtime/Engine/Private/HLSLTree/HLSLTreeCommon.cpp
@@ -9,6 +9,7 @@
 #include "MaterialShared.h"
 #include "MaterialHLSLTree.h"
 #include "MaterialCachedData.h"
+#include "DataDrivenShaderPlatformInfo.h"
 
 namespace UE::HLSLTree
 {
@@ -107,6 +108,11 @@ void FExpressionForward::EmitValuePreshader(FEmitContext& Context, FEmitScope& S
 	return Expression->EmitValuePreshader(Context, Scope, RequestedType, OutResult);
 }
 
+const FExpression* FExpressionPreviousFrameSwitch::ComputePreviousFrame(FTree& Tree, const FRequestedType& RequestedType) const
+{
+	return PreviousFrameExpression->ComputePreviousFrame(Tree, RequestedType);
+}
+
 const FExpression* FTree::NewConstant(const Shader::FValue& Value)
 {
 	return NewExpression<FExpressionConstant>(Value);
@@ -300,7 +306,7 @@ bool FExpressionSetStructField::PrepareValue(FEmitContext& Context, FEmitScope&
 	}
 	
 	const FRequestedType RequestedFieldType = RequestedType.GetField(Field);
-	if (!RequestedType.IsEmpty())
+	if (!RequestedFieldType.IsEmpty())
 	{
 		const FPreparedType FieldPreparedType = Context.PrepareExpression(FieldExpression, Scope, RequestedFieldType);
 		ResultType.SetField(Field, FieldPreparedType);
@@ -487,7 +493,8 @@ void FExpressionSelect::EmitValuePreshader(FEmitContext& Context, FEmitScope& Sc
 		ConditionExpression->GetValuePreshader(Context, Scope, Shader::EValueType::Bool1, OutResult.Preshader);
 
 		check(Context.PreshaderStackPosition > 0);
-		Context.PreshaderStackPosition--;
+		// -1 due to JumpIfFalse consuming condition value. Another -1 due to only one branch is evaluated
+		Context.PreshaderStackPosition -= 2;
 		const Shader::FPreshaderLabel Label0 = OutResult.Preshader.WriteJump(Shader::EPreshaderOpcode::JumpIfFalse);
 
 		const Shader::FType TrueType = TrueExpression->GetValuePreshader(Context, Scope, RequestedType, OutResult.Preshader);
@@ -1110,6 +1117,63 @@ bool FExpressionShaderStageSwitch::IsInputActive(const FEmitContext& Context, in
 	return (Context.ShaderFrequency == SF_Pixel && Index == 0) || (Context.ShaderFrequency == SF_Vertex && Index == 1);
 }
 
+FExpressionVirtualTextureFeatureSwitch::FExpressionVirtualTextureFeatureSwitch(TConstArrayView<const FExpression*> InInputs)
+	: FExpressionSwitchBase(InInputs)
+{
+	static_assert(MaxInputs >= 2, "FExpressionSwitchBase is too small for FExpressionVirtualTextureFeatureSwitch");
+	check(InInputs.Num() == 2);
+}
+
+bool FExpressionVirtualTextureFeatureSwitch::IsInputActive(const FEmitContext& Context, int32 Index) const
+{
+	if (UseVirtualTexturing(Context.TargetParameters.FeatureLevel, Context.TargetParameters.TargetPlatform))
+	{
+		return Index == 0;
+	}
+	else
+	{
+		return Index == 1;
+	}
+}
+
+FExpressionDistanceFieldsRenderingSwitch::FExpressionDistanceFieldsRenderingSwitch(TConstArrayView<const FExpression*> InInputs)
+	: FExpressionSwitchBase(InInputs)
+{
+	static_assert(MaxInputs >= 2, "FExpressionSwitchBase is too small for FExpressionDistanceFieldsRenderingSwitch");
+	check(InInputs.Num() == 2);
+}
+
+bool FExpressionDistanceFieldsRenderingSwitch::IsInputActive(const FEmitContext& Context, int32 Index) const
+{
+	if (Context.TargetParameters.IsGenericTarget())
+	{
+		return true;
+	}
+
+	if (IsMobilePlatform(Context.TargetParameters.ShaderPlatform))
+	{
+		if (IsMobileDistanceFieldEnabled(Context.TargetParameters.ShaderPlatform))
+		{
+			return Index == 0;
+		}
+		else
+		{
+			return Index == 1;
+		}
+	}
+	else
+	{
+		if (IsUsingDistanceFields(Context.TargetParameters.ShaderPlatform))
+		{
+			return Index == 0;
+		}
+		else
+		{
+			return Index == 1;
+		}
+	}
+}
+
 bool FExpressionInlineCustomHLSL::PrepareValue(FEmitContext& Context, FEmitScope& Scope, const FRequestedType& RequestedType, FPrepareValueResult& OutResult) const
 {
 	return OutResult.SetType(Context, RequestedType, EExpressionEvaluation::Shader, Shader::FType(ResultType));
diff --git a/Engine/Source/Runtime/Engine/Private/HLSLTree/HLSLTreeEmit.cpp b/Engine/Source/Runtime/Engine/Private/HLSLTree/HLSLTreeEmit.cpp
--- a/Engine/Source/Runtime/Engine/Private/HLSLTree/HLSLTreeEmit.cpp
+++ b/Engine/Source/Runtime/Engine/Private/HLSLTree/HLSLTreeEmit.cpp
@@ -208,6 +208,15 @@ FEmitContext::~FEmitContext()
 	{
 		delete It.Value;
 	}
+	for (TMap<const FExpression*, FPrepareValueResult*>::TIterator It(PrepareValueMap); It; ++It)
+	{
+		FPrepareValueResult* Value = It.Value();
+		if (Value)
+		{
+			Value->~FPrepareValueResult();
+		}
+	}
+	ResetPastRequestedTypes();
 }
 
 bool FEmitContext::InternalError(FStringView ErrorMessage)
@@ -381,6 +390,30 @@ EExpressionEvaluation FEmitContext::GetEvaluation(const FExpression* Expression,
 	return GetPreparedType(Expression).GetEvaluation(Scope, RequestedType);
 }
 
+bool AllComponentsRequestedInPast(const FRequestedType& RequestedType, const FRequestedType& PastRequestedType)
+{
+	bool bResult = true;
+	TBitArray<>::FConstWordIterator ItA(RequestedType.RequestedComponents);
+	TBitArray<>::FConstWordIterator ItB(PastRequestedType.RequestedComponents);
+	for (; ItA && ItB; ++ItA, ++ItB)
+	{
+		if ((ItA.GetWord() & ItB.GetWord()) != ItA.GetWord())
+		{
+			bResult = false;
+			break;
+		}
+	}
+	for (; ItA; ++ItA)
+	{
+		if (ItA.GetWord() != 0u)
+		{
+			bResult = false;
+			break;
+		}
+	}
+	return bResult;
+}
+
 FPreparedType FEmitContext::PrepareExpression(const FExpression* InExpression, FEmitScope& Scope, const FRequestedType& RequestedType)
 {
 	if (!InExpression || RequestedType.IsVoid())
@@ -417,6 +450,52 @@ FPreparedType FEmitContext::PrepareExpression(const FExpression* InExpression, F
 		return Result->PreparedType;
 	}
 
+	FXxHash64Builder Hasher;
+	Hasher.Update(&InExpression, sizeof(InExpression));
+	Hasher.Update(&bMarkLiveValues, sizeof(bMarkLiveValues));
+	const FXxHash64 Hash = Hasher.Finalize();
+
+	FRequestedType** PastRequestedTypePtr = RequestedTypeTracker.Find(Hash);
+	if (PastRequestedTypePtr)
+	{
+		FRequestedType& PastRequestedType = **PastRequestedTypePtr;
+		if (RequestedType.Type.IsAny() || PastRequestedType.Type.IsAny())
+		{
+			if (PastRequestedType.Type.IsAny())
+			{
+				return Result->PreparedType;
+			}
+			else
+			{
+				PastRequestedType = FRequestedType(Shader::EValueType::Any, false);
+			}
+		}
+		else
+		{
+			const Shader::FType CombinedType = Shader::CombineTypes(RequestedType.Type, PastRequestedType.Type);
+			check(!CombinedType.IsVoid());
+
+			if (CombinedType == PastRequestedType.Type && AllComponentsRequestedInPast(RequestedType, PastRequestedType))
+			{
+				return Result->PreparedType;
+			}
+			else
+			{
+				PastRequestedType.Type = CombinedType;
+				PastRequestedType.RequestedComponents.CombineWithBitwiseOR(RequestedType.RequestedComponents, EBitwiseOperatorFlags::MaxSize);
+			}
+		}
+	}
+	else
+	{
+		check(Result->PreparedType.IsVoid() || bMarkLiveValues);
+		FRequestedType* StackRequestedType = new (*Allocator) FRequestedType(RequestedType);
+		// Turn Bool into Int
+		StackRequestedType->Type = Shader::CombineTypes(StackRequestedType->Type, StackRequestedType->Type);
+		check(!StackRequestedType->Type.IsVoid());
+		RequestedTypeTracker.Add(Hash, StackRequestedType);
+	}
+
 	bool bResult = false;
 	{
 		FEmitOwnerScope OwnerScope(*this, InExpression);
@@ -433,10 +512,7 @@ FPreparedType FEmitContext::PrepareExpression(const FExpression* InExpression, F
 		check(!ResultType.IsVoid());
 		MarkInputType(InExpression, RequestedType.Type.GetConcreteType());
 	}
-	else
-	{
-		int a = 0;
-	}
+	
 	return ResultType;
 }
 
@@ -1369,7 +1445,16 @@ void FEmitContext::Finalize()
 	
 	// Don't reset Expression/Preshader maps, allow future passes to share matching preshaders/expressions
 
+	for (TMap<const FExpression*, FPrepareValueResult*>::TIterator It(PrepareValueMap); It; ++It)
+	{
+		FPrepareValueResult* Value = It.Value();
+		if (Value)
+		{
+			Value->~FPrepareValueResult();
+		}
+	}
 	PrepareValueMap.Reset();
+	ResetPastRequestedTypes();
 	EmitScopeMap.Reset();
 	EmitFunctionMap.Reset();
 	EmitLocalPHIMap.Reset();
@@ -1377,6 +1462,19 @@ void FEmitContext::Finalize()
 	MaterialCompilationOutput->UniformExpressionSet.UniformPreshaderBufferSize = (UniformPreshaderOffset + 3u) / 4u;
 }
 
+void FEmitContext::ResetPastRequestedTypes()
+{
+	for (TMap<FXxHash64, FRequestedType*>::TIterator It(RequestedTypeTracker); It; ++It)
+	{
+		FRequestedType* Value = It.Value();
+		if (Value)
+		{
+			Value->~FRequestedType();
+		}
+	}
+	RequestedTypeTracker.Reset();
+}
+
 } // namespace UE::HLSLTree
 
 #endif // WITH_EDITOR
diff --git a/Engine/Source/Runtime/Engine/Private/Materials/MaterialExpressionHLSL.cpp b/Engine/Source/Runtime/Engine/Private/Materials/MaterialExpressionHLSL.cpp
--- a/Engine/Source/Runtime/Engine/Private/Materials/MaterialExpressionHLSL.cpp
+++ b/Engine/Source/Runtime/Engine/Private/Materials/MaterialExpressionHLSL.cpp
@@ -53,26 +53,42 @@
 #include "Materials/MaterialExpressionParticleColor.h"
 #include "Materials/MaterialExpressionParticlePositionWS.h"
 #include "Materials/MaterialExpressionParticleRadius.h"
+#include "Materials/MaterialExpressionParticleDirection.h"
+#include "Materials/MaterialExpressionParticleSpeed.h"
+#include "Materials/MaterialExpressionParticleRelativeTime.h"
+#include "Materials/MaterialExpressionParticleRandom.h"
+#include "Materials/MaterialExpressionParticleSize.h"
 #include "Materials/MaterialExpressionSkyLightEnvMapSample.h"
 #include "Materials/MaterialExpressionTextureSample.h"
 #include "Materials/MaterialExpressionTextureSampleParameter.h"
+#include "Materials/MaterialExpressionRuntimeVirtualTextureSample.h"
+#include "Materials/MaterialExpressionRuntimeVirtualTextureSampleParameter.h"
 #include "Materials/MaterialExpressionFontSample.h"
 #include "Materials/MaterialExpressionTextureObject.h"
 #include "Materials/MaterialExpressionTextureObjectParameter.h"
 #include "Materials/MaterialExpressionSceneTexture.h"
+#include "Materials/MaterialExpressionSceneColor.h"
 #include "Materials/MaterialExpressionSceneDepth.h"
+#include "Materials/MaterialExpressionDepthFade.h"
 #include "Materials/MaterialExpressionNoise.h"
+#include "Materials/MaterialExpressionVectorNoise.h"
 #include "Materials/MaterialExpressionVertexInterpolator.h"
 #include "Materials/MaterialExpressionFunctionInput.h"
 #include "Materials/MaterialExpressionFunctionOutput.h"
 #include "Materials/MaterialExpressionMaterialFunctionCall.h"
 #include "Materials/MaterialExpressionMaterialAttributeLayers.h"
 #include "Materials/MaterialExpressionBlendMaterialAttributes.h"
+#include "Materials/MaterialExpressionPreviousFrameSwitch.h"
 #include "Materials/MaterialExpressionStaticSwitch.h"
 #include "Materials/MaterialExpressionFeatureLevelSwitch.h"
 #include "Materials/MaterialExpressionShadingPathSwitch.h"
 #include "Materials/MaterialExpressionQualitySwitch.h"
 #include "Materials/MaterialExpressionShaderStageSwitch.h"
+#include "Materials/MaterialExpressionRayTracingQualitySwitch.h"
+#include "Materials/MaterialExpressionVirtualTextureFeatureSwitch.h"
+#include "Materials/MaterialExpressionDistanceFieldsRenderingSwitch.h"
+#include "Materials/MaterialExpressionShadowReplace.h"
+#include "Materials/MaterialExpressionMaterialProxyReplace.h"
 #include "Materials/MaterialExpressionGetLocal.h"
 #include "Materials/MaterialExpressionOneMinus.h"
 #include "Materials/MaterialExpressionAbs.h"
@@ -138,6 +154,7 @@
 #include "Materials/MaterialExpressionBumpOffset.h"
 #include "Materials/MaterialExpressionPanner.h"
 #include "Materials/MaterialExpressionRotator.h"
+#include "Materials/MaterialExpressionAtmosphericLightVector.h"
 #include "Materials/MaterialExpressionSkyAtmosphereLightDirection.h"
 #include "Materials/MaterialExpressionVolumetricAdvancedMaterialInput.h"
 #include "Materials/MaterialExpressionVolumetricAdvancedMaterialOutput.h"
@@ -154,6 +171,7 @@
 #include "Materials/MaterialExpressionIfThenElse.h"
 #include "Materials/MaterialExpressionForLoop.h"
 #include "Materials/MaterialExpressionWhileLoop.h"
+#include "Materials/MaterialExpressionDeriveNormalZ.h"
 #include "Materials/MaterialFunctionInterface.h"
 #include "Materials/Material.h"
 #include "MaterialHLSLTree.h"
@@ -248,6 +266,16 @@ bool UMaterialExpressionShadingModel::GenerateHLSLExpression(FMaterialHLSLGenera
 	return true;
 }
 
+bool UMaterialExpressionPreviousFrameSwitch::GenerateHLSLExpression(FMaterialHLSLGenerator& Generator, UE::HLSLTree::FScope& Scope, int32 OutputIndex, UE::HLSLTree::FExpression const*& OutExpression) const
+{
+	using namespace UE::HLSLTree;
+
+	const FExpression* CurrentFrameExpression = CurrentFrame.AcquireHLSLExpression(Generator, Scope);
+	const FExpression* PreviousFrameExpression = PreviousFrame.AcquireHLSLExpression(Generator, Scope);
+	OutExpression = Generator.GetTree().NewExpression<FExpressionPreviousFrameSwitch>(CurrentFrameExpression, PreviousFrameExpression);
+	return true;
+}
+
 bool UMaterialExpressionStaticSwitch::GenerateHLSLExpression(FMaterialHLSLGenerator& Generator, UE::HLSLTree::FScope& Scope, int32 OutputIndex, UE::HLSLTree::FExpression const*& OutExpression) const
 {
 	using namespace UE::HLSLTree;
@@ -394,6 +422,115 @@ bool UMaterialExpressionShaderStageSwitch::GenerateHLSLExpression(FMaterialHLSLG
 	}
 }
 
+bool UMaterialExpressionRayTracingQualitySwitch::GenerateHLSLExpression(FMaterialHLSLGenerator& Generator, UE::HLSLTree::FScope& Scope, int32 OutputIndex, UE::HLSLTree::FExpression const*& OutExpression) const
+{
+	if (!Normal.GetTracedInput().Expression)
+	{
+		return Generator.Error(TEXT("Missing RayTracingQualitySwitch input 'Normal'"));
+	}
+	else if (!RayTraced.GetTracedInput().Expression)
+	{
+		return Generator.Error(TEXT("Missing RayTracingQualitySwitch input 'RayTraced'"));
+	}
+	else
+	{
+		using namespace UE::HLSLTree;
+		const FExpression* NormalExpression = Normal.AcquireHLSLExpression(Generator, Scope);
+		const FExpression* RayTracedExpression = RayTraced.AcquireHLSLExpression(Generator, Scope);
+		const FExpression* ConditionExpression = Generator.GetTree().NewExpression<FExpressionInlineCustomHLSL>(UE::Shader::EValueType::Bool1, TEXT("GetRayTracingQualitySwitch()"));
+		
+		OutExpression = Generator.GetTree().NewExpression<FExpressionSelect>(ConditionExpression, RayTracedExpression, NormalExpression);
+		return true;
+	}
+}
+
+bool UMaterialExpressionVirtualTextureFeatureSwitch::GenerateHLSLExpression(FMaterialHLSLGenerator& Generator, UE::HLSLTree::FScope& Scope, int32 OutputIndex, UE::HLSLTree::FExpression const*& OutExpression) const
+{
+	if (!Yes.GetTracedInput().Expression)
+	{
+		return Generator.Error(TEXT("Missing VirtualTextureFeatureSwitch input 'Yes'"));
+	}
+	else if (!No.GetTracedInput().Expression)
+	{
+		return Generator.Error(TEXT("Missing VirtualTextureFeatureSwitch input 'No'"));
+	}
+	else
+	{
+		using namespace UE::HLSLTree;
+		const FExpression* ExpressionInputs[2];
+		ExpressionInputs[0] = Yes.AcquireHLSLExpression(Generator, Scope);
+		ExpressionInputs[1] = No.AcquireHLSLExpression(Generator, Scope);
+
+		OutExpression = Generator.GetTree().NewExpression<FExpressionVirtualTextureFeatureSwitch>(ExpressionInputs);
+		return true;
+	}
+}
+
+bool UMaterialExpressionDistanceFieldsRenderingSwitch::GenerateHLSLExpression(FMaterialHLSLGenerator& Generator, UE::HLSLTree::FScope& Scope, int32 OutputIndex, UE::HLSLTree::FExpression const*& OutExpression) const
+{
+	if (!Yes.GetTracedInput().Expression)
+	{
+		return Generator.Error(TEXT("Missing DistanceFieldsRenderingSwitch input 'Yes'"));
+	}
+	else if (!No.GetTracedInput().Expression)
+	{
+		return Generator.Error(TEXT("Missing DistanceFieldsRenderingSwitch input 'No'"));
+	}
+	else
+	{
+		using namespace UE::HLSLTree;
+		const FExpression* ExpressionInputs[2];
+		ExpressionInputs[0] = Yes.AcquireHLSLExpression(Generator, Scope);
+		ExpressionInputs[1] = No.AcquireHLSLExpression(Generator, Scope);
+
+		OutExpression = Generator.GetTree().NewExpression<FExpressionDistanceFieldsRenderingSwitch>(ExpressionInputs);
+		return true;
+	}
+}
+
+bool UMaterialExpressionShadowReplace::GenerateHLSLExpression(FMaterialHLSLGenerator& Generator, UE::HLSLTree::FScope& Scope, int32 OutputIndex, UE::HLSLTree::FExpression const*& OutExpression) const
+{
+	if (!Default.GetTracedInput().Expression)
+	{
+		return Generator.Error(TEXT("Missing input Default"));
+	}
+	else if (!Shadow.GetTracedInput().Expression)
+	{
+		return Generator.Error(TEXT("Missing input Shadow"));
+	}
+	else
+	{
+		using namespace UE::HLSLTree;
+		const FExpression* ConditionExpression = Generator.GetTree().NewExpression<FExpressionInlineCustomHLSL>(UE::Shader::EValueType::Bool1, TEXT("GetShadowReplaceState()"));
+		const FExpression* TrueExpression = Shadow.AcquireHLSLExpression(Generator, Scope);
+		const FExpression* FalseExpression = Default.AcquireHLSLExpression(Generator, Scope);
+		OutExpression = Generator.GetTree().NewExpression<FExpressionSelect>(ConditionExpression, TrueExpression, FalseExpression);
+		return true;
+	}
+}
+
+bool UMaterialExpressionMaterialProxyReplace::GenerateHLSLExpression(FMaterialHLSLGenerator& Generator, UE::HLSLTree::FScope& Scope, int32 OutputIndex, UE::HLSLTree::FExpression const*& OutExpression) const
+{
+	if (!Realtime.GetTracedInput().Expression)
+	{
+		return Generator.Error(TEXT("Missing MaterialProxyReplace input Realtime"));
+	}
+	else if (!MaterialProxy.GetTracedInput().Expression)
+	{
+		return Generator.Error(TEXT("Missing MaterialProxyReplace input MaterialProxy"));
+	}
+	else
+	{
+		// FMaterialHLSLGenerator replaces FHLSLMaterialTranslator which is EMaterialCompilerType::Standard (i.e. not a proxy compiler)
+		OutExpression = Realtime.AcquireHLSLExpression(Generator, Scope);
+		if (!OutExpression)
+		{
+			return Generator.Error(TEXT("Input Realtime failed to generate HLSL expression"));
+		}
+		return true;
+	}
+}
+
 bool UMaterialExpressionGetLocal::GenerateHLSLExpression(FMaterialHLSLGenerator& Generator, UE::HLSLTree::FScope& Scope, int32 OutputIndex, UE::HLSLTree::FExpression const*& OutExpression) const
 {
 	OutExpression = Generator.GetTree().AcquireLocal(Scope, LocalName);
@@ -891,6 +1028,42 @@ bool UMaterialExpressionParticleRadius::GenerateHLSLExpression(FMaterialHLSLGene
 	return true;
 }
 
+bool UMaterialExpressionParticleDirection::GenerateHLSLExpression(FMaterialHLSLGenerator& Generator, UE::HLSLTree::FScope& Scope, int32 OutputIndex, UE::HLSLTree::FExpression const*& OutExpression) const
+{
+	using namespace UE::HLSLTree::Material;
+	OutExpression = Generator.NewExternalInput(EExternalInput::ParticleDirection);
+	return true;
+}
+
+bool UMaterialExpressionParticleSpeed::GenerateHLSLExpression(FMaterialHLSLGenerator& Generator, UE::HLSLTree::FScope& Scope, int32 OutputIndex, UE::HLSLTree::FExpression const*& OutExpression) const
+{
+	using namespace UE::HLSLTree::Material;
+	OutExpression = Generator.NewExternalInput(EExternalInput::ParticleSpeed);
+	OutExpression = Generator.GetTree().NewMax(OutExpression, Generator.NewConstant(0.001f));
+	return true;
+}
+
+bool UMaterialExpressionParticleRelativeTime::GenerateHLSLExpression(FMaterialHLSLGenerator& Generator, UE::HLSLTree::FScope& Scope, int32 OutputIndex, UE::HLSLTree::FExpression const*& OutExpression) const
+{
+	using namespace UE::HLSLTree::Material;
+	OutExpression = Generator.NewExternalInput(EExternalInput::ParticleRelativeTime);
+	return true;
+}
+
+bool UMaterialExpressionParticleRandom::GenerateHLSLExpression(FMaterialHLSLGenerator& Generator, UE::HLSLTree::FScope& Scope, int32 OutputIndex, UE::HLSLTree::FExpression const*& OutExpression) const
+{
+	using namespace UE::HLSLTree::Material;
+	OutExpression = Generator.NewExternalInput(EExternalInput::ParticleRandom);
+	return true;
+}
+
+bool UMaterialExpressionParticleSize::GenerateHLSLExpression(FMaterialHLSLGenerator& Generator, UE::HLSLTree::FScope& Scope, int32 OutputIndex, UE::HLSLTree::FExpression const*& OutExpression) const
+{
+	using namespace UE::HLSLTree::Material;
+	OutExpression = Generator.NewExternalInput(EExternalInput::ParticleSize);
+	return true;
+}
+
 bool UMaterialExpressionTextureObject::GenerateHLSLExpression(FMaterialHLSLGenerator& Generator, UE::HLSLTree::FScope& Scope, int32 OutputIndex, UE::HLSLTree::FExpression const*& OutExpression) const
 {
 	using namespace UE::HLSLTree;
@@ -1022,6 +1195,317 @@ bool UMaterialExpressionTextureSampleParameter::GenerateHLSLExpression(FMaterial
 	return GenerateHLSLExpressionBase(Generator, Scope, TextureExpression, OutExpression);
 }
 
+bool UMaterialExpressionRuntimeVirtualTextureSample::GenerateHLSLExpression(FMaterialHLSLGenerator& Generator, UE::HLSLTree::FScope& Scope, int32 OutputIndex, UE::HLSLTree::FExpression const*& OutExpression) const
+{
+	using namespace UE::HLSLTree;
+
+	// Is this a valid UMaterialExpressionRuntimeVirtualTextureSampleParameter?
+	const bool bIsParameter = HasAParameterName() && GetParameterName().IsValid() && !GetParameterName().IsNone();
+
+	// Check validity of current virtual texture
+	bool bIsVirtualTextureValid = VirtualTexture != nullptr;
+	if (!bIsVirtualTextureValid)
+	{
+		if (bIsParameter)
+		{
+			return Generator.Error(TEXT("Missing input Virtual Texture"));
+		}
+	}
+	else if (VirtualTexture->GetMaterialType() != MaterialType)
+	{
+		UEnum const* Enum = StaticEnum<ERuntimeVirtualTextureMaterialType>();
+		const FString MaterialTypeDisplayName = Enum->GetDisplayNameTextByValue((int64)MaterialType).ToString();
+		const FString TextureTypeDisplayName = Enum->GetDisplayNameTextByValue((int64)VirtualTexture->GetMaterialType()).ToString();
+
+		Generator.Errorf(TEXT("%Material type is '%s', should be '%s' to match %s"),
+			*MaterialTypeDisplayName,
+			*TextureTypeDisplayName,
+			*VirtualTexture->GetName());
+
+		bIsVirtualTextureValid = false;
+	}
+	else if (VirtualTexture->GetSinglePhysicalSpace() != bSinglePhysicalSpace)
+	{
+		Generator.Errorf(TEXT("%Page table packing is '%d', should be '%d' to match %s"),
+			bSinglePhysicalSpace ? 1 : 0,
+			VirtualTexture->GetSinglePhysicalSpace() ? 1 : 0,
+			*VirtualTexture->GetName());
+
+		bIsVirtualTextureValid = false;
+	}
+	else if ((VirtualTexture->GetAdaptivePageTable()) != bAdaptive)
+	{
+		Generator.Errorf(TEXT("Adaptive page table is '%d', should be '%d' to match %s"),
+			bAdaptive ? 1 : 0,
+			VirtualTexture->GetAdaptivePageTable() ? 1 : 0,
+			*VirtualTexture->GetName());
+
+		bIsVirtualTextureValid = false;
+	}
+
+	// Calculate the virtual texture layer and sampling/unpacking functions for this output
+	// Fallback to a sensible default value if the output isn't valid for the bound virtual texture
+	uint32 UnpackTarget = 0;
+	uint32 UnpackMask = 0;
+	EVirtualTextureUnpackType UnpackType = EVirtualTextureUnpackType::None;
+
+	bool bIsBaseColorValid = false;
+	bool bIsSpecularValid = false;
+	bool bIsRoughnessValid = false;
+	bool bIsNormalValid = false;
+	bool bIsWorldHeightValid = false;
+	bool bIsMaskValid = false;
+
+	switch (MaterialType)
+	{
+	case ERuntimeVirtualTextureMaterialType::BaseColor: bIsBaseColorValid = true; break;
+	case ERuntimeVirtualTextureMaterialType::BaseColor_Normal_Roughness: bIsBaseColorValid = bIsNormalValid = bIsRoughnessValid = true; break;
+	case ERuntimeVirtualTextureMaterialType::BaseColor_Normal_Specular: bIsRoughnessValid = bIsBaseColorValid = bIsNormalValid = bIsSpecularValid = true; break;
+	case ERuntimeVirtualTextureMaterialType::BaseColor_Normal_Specular_YCoCg: bIsRoughnessValid = bIsBaseColorValid = bIsNormalValid = bIsSpecularValid = true; break;
+	case ERuntimeVirtualTextureMaterialType::BaseColor_Normal_Specular_Mask_YCoCg: bIsRoughnessValid = bIsBaseColorValid = bIsNormalValid = bIsSpecularValid = bIsMaskValid = true; break;
+	case ERuntimeVirtualTextureMaterialType::WorldHeight: bIsWorldHeightValid = true; break;
+	}
+
+	switch (OutputIndex)
+	{
+	case 0:
+		if (bIsVirtualTextureValid && bIsBaseColorValid)
+		{
+			switch (MaterialType)
+			{
+			case ERuntimeVirtualTextureMaterialType::BaseColor_Normal_Specular_YCoCg:
+			case ERuntimeVirtualTextureMaterialType::BaseColor_Normal_Specular_Mask_YCoCg: UnpackType = EVirtualTextureUnpackType::BaseColorYCoCg; break;
+			case ERuntimeVirtualTextureMaterialType::BaseColor_Normal_Roughness: UnpackType = EVirtualTextureUnpackType::BaseColorSRGB; break;
+			default: UnpackTarget = 0; UnpackMask = 0x7; break;
+			}
+		}
+		else
+		{
+			OutExpression = Generator.GetTree().NewConstant(FVector3f::ZeroVector);
+			return true;
+		}
+		break;
+	case 1:
+		if (bIsVirtualTextureValid && bIsSpecularValid)
+		{
+			switch (MaterialType)
+			{
+			case ERuntimeVirtualTextureMaterialType::BaseColor_Normal_Specular:  UnpackTarget = 1; UnpackMask = 0x1; break;
+			case ERuntimeVirtualTextureMaterialType::BaseColor_Normal_Specular_YCoCg:
+			case ERuntimeVirtualTextureMaterialType::BaseColor_Normal_Specular_Mask_YCoCg: UnpackTarget = 2; UnpackMask = 0x1; break;
+			}
+		}
+		else
+		{
+			OutExpression = Generator.GetTree().NewConstant(0.5f);
+			return true;
+		}
+		break;
+	case 2:
+		if (bIsVirtualTextureValid && bIsRoughnessValid)
+		{
+			switch (MaterialType)
+			{
+			case ERuntimeVirtualTextureMaterialType::BaseColor_Normal_Roughness: UnpackTarget = 1; UnpackMask = 0x2; break;
+			case ERuntimeVirtualTextureMaterialType::BaseColor_Normal_Specular: UnpackTarget = 1; UnpackMask = 0x2; break;
+			case ERuntimeVirtualTextureMaterialType::BaseColor_Normal_Specular_YCoCg:
+			case ERuntimeVirtualTextureMaterialType::BaseColor_Normal_Specular_Mask_YCoCg: UnpackTarget = 2; UnpackMask = 0x2; break;
+			}
+		}
+		else
+		{
+			OutExpression = Generator.GetTree().NewConstant(0.5f);
+			return true;
+		}
+		break;
+	case 3:
+		if (bIsVirtualTextureValid && bIsNormalValid)
+		{
+			switch (MaterialType)
+			{
+			case ERuntimeVirtualTextureMaterialType::BaseColor_Normal_Roughness: UnpackType = EVirtualTextureUnpackType::NormalBGR565; break;
+			case ERuntimeVirtualTextureMaterialType::BaseColor_Normal_Specular: UnpackType = EVirtualTextureUnpackType::NormalBC3BC3; break;
+			case ERuntimeVirtualTextureMaterialType::BaseColor_Normal_Specular_YCoCg:
+			case ERuntimeVirtualTextureMaterialType::BaseColor_Normal_Specular_Mask_YCoCg: UnpackType = EVirtualTextureUnpackType::NormalBC5BC1; break;
+			}
+		}
+		else
+		{
+			OutExpression = Generator.GetTree().NewConstant(FVector3f(0.f, 0.f, 1.f));
+			return true;
+		}
+		break;
+	case 4:
+		if (bIsVirtualTextureValid && bIsWorldHeightValid)
+		{
+			UnpackType = EVirtualTextureUnpackType::HeightR16;
+		}
+		else
+		{
+			OutExpression = Generator.GetTree().NewConstant(0.f);
+			return true;
+		}
+		break;
+	case 5:
+		if (bIsVirtualTextureValid && bIsMaskValid)
+		{
+			UnpackTarget = 2; UnpackMask = 0x8; break;
+		}
+		else
+		{
+			OutExpression = Generator.GetTree().NewConstant(1.f);
+			return true;
+		}
+		break;
+	default:
+		return false;
+	}
+
+	// We don't know what the texture reference index is until the HLSL tree is prepared
+	const FExpression* TextureExpression;
+	if (bIsParameter)
+	{
+		FMaterialParameterMetadata ParameterMeta;
+		if (!GetParameterValue(ParameterMeta))
+		{
+			return Generator.Error(TEXT("Failed to get parameter value"));
+		}
+		TextureExpression = Generator.GenerateMaterialParameter(GetParameterName(), ParameterMeta, SAMPLERTYPE_VirtualMasks);
+	}
+	else
+	{
+		const FMaterialParameterMetadata ParameterMeta(VirtualTexture);
+		TextureExpression = Generator.GenerateMaterialParameter(FName(), ParameterMeta, SAMPLERTYPE_VirtualMasks);
+	}
+
+	// Compile the runtime virtual texture uniforms
+	const FExpression* UniformExpressions[ERuntimeVirtualTextureShaderUniform_Count];
+	for (int32 UniformIndex = 0; UniformIndex < ERuntimeVirtualTextureShaderUniform_Count; ++UniformIndex)
+	{
+		UniformExpressions[UniformIndex] = Generator.GetTree().NewExpression<Material::FExpressionRuntimeVirtualTextureUniform>(
+			Generator.GetParameterInfo(bIsParameter ? GetParameterName() : FName()),
+			TextureExpression,
+			(ERuntimeVirtualTextureShaderUniform)UniformIndex);
+	}
+
+	// Compile the coordinates
+	// We use the virtual texture world space transform by default
+	const FExpression* CoordinateExpression = nullptr;
+
+	if (Coordinates.GetTracedInput().Expression != nullptr && WorldPosition.GetTracedInput().Expression != nullptr)
+	{
+		return Generator.Error(TEXT("Only one of 'Coordinates' and 'WorldPosition' can be used"));
+	}
+
+	if (Coordinates.GetTracedInput().Expression != nullptr)
+	{
+		CoordinateExpression = Coordinates.AcquireHLSLExpression(Generator, Scope);
+	}
+	else
+	{
+		const FExpression* WorldPositionExpression = WorldPosition.AcquireHLSLExpressionOrExternalInput(Generator, Scope, Material::EExternalInput::WorldPosition);
+
+		check(WorldPositionExpression);
+		// Same as VirtualTextureWorldToUV
+		const FExpression* TempExpression = Generator.GetTree().NewSub(WorldPositionExpression, UniformExpressions[ERuntimeVirtualTextureShaderUniform_WorldToUVTransform0]);
+		const FExpression* UExpression = Generator.GetTree().NewDot(TempExpression, UniformExpressions[ERuntimeVirtualTextureShaderUniform_WorldToUVTransform1]);
+		const FExpression* VExpression = Generator.GetTree().NewDot(TempExpression, UniformExpressions[ERuntimeVirtualTextureShaderUniform_WorldToUVTransform2]);
+		CoordinateExpression = Generator.GetTree().NewExpression<FExpressionAppend>(UExpression, VExpression);
+	}
+
+	// Compile the mip level for the current mip value mode
+	ETextureMipValueMode TextureMipLevelMode = TMVM_None;
+	const FExpression* MipValueExpression = nullptr;
+	FExpressionDerivatives CoordinateExpressionDerivatives;
+	const bool bMipValueExpressionValid = MipValue.GetTracedInput().Expression != nullptr;
+	if (MipValueMode == RVTMVM_MipLevel)
+	{
+		TextureMipLevelMode = TMVM_MipLevel;
+		MipValueExpression = bMipValueExpressionValid ? MipValue.AcquireHLSLExpression(Generator, Scope) : Generator.GetTree().NewConstant(0);
+	}
+	else if (MipValueMode == RVTMVM_MipBias)
+	{
+		TextureMipLevelMode = TMVM_MipBias;
+		MipValueExpression = bMipValueExpressionValid ? MipValue.AcquireHLSLExpression(Generator, Scope) : Generator.GetTree().NewConstant(0);
+	}
+	else if (MipValueMode == RVTMVM_RecalculateDerivatives)
+	{
+		// Calculate derivatives from world position.
+		TextureMipLevelMode = TMVM_Derivative;
+		const FExpression* WorldPos = Generator.NewExternalInput(Material::EExternalInput::TranslatedWorldPosition);
+		const FExpression* WorldPositionDdx = Generator.GetTree().NewExpression<FExpressionDerivative>(EDerivativeCoordinate::Ddx, WorldPos);
+		const FExpression* UDdx = Generator.GetTree().NewDot(WorldPositionDdx, UniformExpressions[ERuntimeVirtualTextureShaderUniform_WorldToUVTransform1]);
+		const FExpression* VDdx = Generator.GetTree().NewDot(WorldPositionDdx, UniformExpressions[ERuntimeVirtualTextureShaderUniform_WorldToUVTransform2]);
+		CoordinateExpressionDerivatives.ExpressionDdx = Generator.GetTree().NewExpression<FExpressionAppend>(UDdx, VDdx);
+		const FExpression* WorldPositionDdy = Generator.GetTree().NewExpression<FExpressionDerivative>(EDerivativeCoordinate::Ddy, WorldPos);
+		const FExpression* UDdy = Generator.GetTree().NewDot(WorldPositionDdy, UniformExpressions[ERuntimeVirtualTextureShaderUniform_WorldToUVTransform1]);
+		const FExpression* VDdy = Generator.GetTree().NewDot(WorldPositionDdy, UniformExpressions[ERuntimeVirtualTextureShaderUniform_WorldToUVTransform2]);
+		CoordinateExpressionDerivatives.ExpressionDdy = Generator.GetTree().NewExpression<FExpressionAppend>(UDdy, VDdy);
+	}
+
+	// Convert texture address mode to matching sampler source mode.
+	// Would be better if ESamplerSourceMode had a Mirror enum that we could also use...
+	ESamplerSourceMode SamplerSourceMode = SSM_Clamp_WorldGroupSettings;
+	switch (TextureAddressMode)
+	{
+	case RVTTA_Clamp:
+		SamplerSourceMode = SSM_Clamp_WorldGroupSettings;
+		break;
+	case RVTTA_Wrap:
+		SamplerSourceMode = SSM_Wrap_WorldGroupSettings;
+		break;
+	}
+
+	// We can support disabling feedback for MipLevel mode.
+	const bool bForceEnableFeedback = TextureMipLevelMode != TMVM_MipLevel;
+
+	// Compile the texture sample code
+	const FExpression* AutomaticMipViewBiasExpression = Generator.GetTree().NewConstant(true);
+	const FExpression* SampleLayerExpressions[RuntimeVirtualTexture::MaxTextureLayers] = {};
+	const int32 TextureLayerCount = URuntimeVirtualTexture::GetLayerCount(MaterialType);
+	check(TextureLayerCount <= RuntimeVirtualTexture::MaxTextureLayers);
+	for (int32 TextureLayerIndex = 0; TextureLayerIndex < TextureLayerCount; ++TextureLayerIndex)
+	{
+		const int32 PageTableLayerIndex = bSinglePhysicalSpace ? 0 : TextureLayerIndex;
+		SampleLayerExpressions[TextureLayerIndex] = Generator.GetTree().NewExpression<Material::FExpressionTextureSample>(
+			TextureExpression,
+			CoordinateExpression,
+			MipValueExpression,
+			AutomaticMipViewBiasExpression,
+			CoordinateExpressionDerivatives,
+			SamplerSourceMode,
+			TextureMipLevelMode,
+			TextureLayerIndex,
+			PageTableLayerIndex,
+			bAdaptive,
+			bEnableFeedback || bForceEnableFeedback);
+	}
+
+	// Compile any unpacking code
+	const FExpression* UnpackExpression = nullptr;
+	if (UnpackType != EVirtualTextureUnpackType::None)
+	{
+		UnpackExpression = Generator.GetTree().NewExpression<Material::FExpressionVirtualTextureUnpack>(
+			SampleLayerExpressions[0],
+			SampleLayerExpressions[1],
+			SampleLayerExpressions[2],
+			UniformExpressions[ERuntimeVirtualTextureShaderUniform_WorldHeightUnpack],
+			UnpackType);
+	}
+	else
+	{
+		if (!SampleLayerExpressions[UnpackTarget])
+		{
+			return Generator.Errorf(TEXT("Virtual texture layer %d required but missing."), UnpackTarget);
+		}
+
+		UnpackExpression = Generator.GetTree().NewSwizzle(MakeSwizzleMask(UnpackMask & 0x1, UnpackMask & 0x2, UnpackMask & 0x4, UnpackMask & 0x8), SampleLayerExpressions[UnpackTarget]);
+	}
+
+	OutExpression = UnpackExpression;
+	return true;
+}
+
 bool UMaterialExpressionFontSample::GenerateHLSLExpression(FMaterialHLSLGenerator& Generator, UE::HLSLTree::FScope& Scope, int32 OutputIndex, UE::HLSLTree::FExpression const*& OutExpression) const
 {
 	using namespace UE::HLSLTree;
@@ -1129,9 +1613,101 @@ bool UMaterialExpressionNoise::GenerateHLSLExpression(FMaterialHLSLGenerator& Ge
 	return true;
 }
 
+bool UMaterialExpressionVectorNoise::GenerateHLSLExpression(FMaterialHLSLGenerator& Generator, UE::HLSLTree::FScope& Scope, int32 OutputIndex, UE::HLSLTree::FExpression const*& OutExpression) const
+{
+	using namespace UE::HLSLTree;
+
+	const FExpression* PositionExpression = Position.AcquireHLSLExpressionOrExternalInput(Generator, Scope, Material::EExternalInput::WorldPosition);
+
+	Material::FVectorNoiseParameters VectorNoiseParameters;
+	VectorNoiseParameters.Quality = Quality;
+	VectorNoiseParameters.TileSize = TileSize;
+	VectorNoiseParameters.Function = NoiseFunction;
+	VectorNoiseParameters.bTiling = bTiling;
+
+	OutExpression = Generator.GetTree().NewExpression<Material::FExpressionVectorNoise>(VectorNoiseParameters, PositionExpression);
+	return true;
+}
+
+bool UMaterialExpressionSceneColor::GenerateHLSLExpression(FMaterialHLSLGenerator& Generator, UE::HLSLTree::FScope& Scope, int32 OutputIndex, UE::HLSLTree::FExpression const*& OutExpression) const
+{
+	using namespace UE::HLSLTree;
+
+	const FExpression* OffsetExpression = nullptr;
+	const FExpression* CoordinateExpression = nullptr;
+
+	if (InputMode == EMaterialSceneAttributeInputMode::OffsetFraction)
+	{
+		OffsetExpression = Input.AcquireHLSLExpressionOrConstant(Generator, Scope, FVector2f(ConstInput.X, ConstInput.Y));
+	}
+	else if (InputMode == EMaterialSceneAttributeInputMode::Coordinates)
+	{
+		CoordinateExpression = Input.TryAcquireHLSLExpression(Generator, Scope);
+	}
+
+	const FExpression* ScreenUVExpression = Generator.GetTree().NewExpression<Material::FExpressionScreenAlignedUV>(OffsetExpression, CoordinateExpression);
+	OutExpression = Generator.GetTree().NewExpression<Material::FExpressionSceneColor>(ScreenUVExpression);
+	return true;
+}
+
 bool UMaterialExpressionSceneDepth::GenerateHLSLExpression(FMaterialHLSLGenerator& Generator, UE::HLSLTree::FScope& Scope, int32 OutputIndex, UE::HLSLTree::FExpression const*& OutExpression) const
 {
-	OutExpression = Generator.NewConstant(0.0f); // TODO
+	using namespace UE::HLSLTree;
+
+	const FExpression* OffsetExpression = nullptr;
+	const FExpression* CoordinateExpression = nullptr;
+
+	if (InputMode == EMaterialSceneAttributeInputMode::OffsetFraction)
+	{
+		OffsetExpression = Input.AcquireHLSLExpressionOrConstant(Generator, Scope, FVector2f(ConstInput.X, ConstInput.Y));
+	}
+	else if (InputMode == EMaterialSceneAttributeInputMode::Coordinates)
+	{
+		CoordinateExpression = Input.TryAcquireHLSLExpression(Generator, Scope);
+	}
+
+	const FExpression* ScreenUVExpression = Generator.GetTree().NewExpression<Material::FExpressionScreenAlignedUV>(OffsetExpression, CoordinateExpression);
+	OutExpression = Generator.GetTree().NewExpression<Material::FExpressionSceneDepth>(ScreenUVExpression);
+	return true;
+}
+
+bool UMaterialExpressionDepthFade::GenerateHLSLExpression(FMaterialHLSLGenerator& Generator, UE::HLSLTree::FScope& Scope, int32 OutputIndex, UE::HLSLTree::FExpression const*& OutExpression) const
+{
+	using namespace UE::HLSLTree;
+
+	FTree& Tree = Generator.GetTree();
+
+	// Scales Opacity by a Linear fade based on SceneDepth, from 0 at PixelDepth to 1 at FadeDistance
+	// Result = Opacity * saturate((SceneDepth - PixelDepth) / max(FadeDistance, DELTA))
+	const FExpression* OpacityExpression = InOpacity.AcquireHLSLExpressionOrConstant(Generator, Scope, OpacityDefault);
+	const FExpression* FadeDistanceExpression = FadeDistance.AcquireHLSLExpressionOrConstant(Generator, Scope, FadeDistanceDefault);
+	FadeDistanceExpression = Tree.NewMax(FadeDistanceExpression, Tree.NewConstant(UE_DELTA));
+
+	
+	// On mobile scene depth is limited to 65500 
+	// to avoid false fading on objects that are close or exceed this limit we clamp pixel depth to (65500 - FadeDistance)
+	const FExpression* PixelDepthExpression = Generator.NewExternalInput(Material::EExternalInput::PixelDepth);
+	const FExpression* MobilePixelDepthExpression = Tree.NewMin(PixelDepthExpression, Tree.NewSub(Tree.NewConstant(65500.f), FadeDistanceExpression));
+	
+	const FExpression* FeatureLevelPixelDepthExpressions[ERHIFeatureLevel::Num];
+	for (int32 FeatureLevelIndex = 0; FeatureLevelIndex < ERHIFeatureLevel::Num; ++FeatureLevelIndex)
+	{
+		if (FeatureLevelIndex <= ERHIFeatureLevel::ES3_1)
+		{
+			FeatureLevelPixelDepthExpressions[FeatureLevelIndex] = MobilePixelDepthExpression;
+		}
+		else
+		{
+			FeatureLevelPixelDepthExpressions[FeatureLevelIndex] = PixelDepthExpression;
+		}
+	}
+	PixelDepthExpression = Tree.NewExpression<FExpressionFeatureLevelSwitch>(FeatureLevelPixelDepthExpressions);
+
+	const FExpression* ScreenUVExpression = Tree.NewExpression<Material::FExpressionScreenAlignedUV>(nullptr, nullptr);
+	const FExpression* SceneDepthExpression = Tree.NewExpression<Material::FExpressionSceneDepth>(ScreenUVExpression);
+	const FExpression* FadeExpression = Tree.NewSaturate(Tree.NewDiv(Tree.NewSub(SceneDepthExpression, PixelDepthExpression), FadeDistanceExpression));
+
+	OutExpression = Tree.NewMul(OpacityExpression, FadeExpression);
 	return true;
 }
 
@@ -2518,6 +3094,12 @@ bool UMaterialExpressionSkyAtmosphereLightDirection::GenerateHLSLExpression(FMat
 	return true;
 }
 
+bool UMaterialExpressionAtmosphericLightVector::GenerateHLSLExpression(FMaterialHLSLGenerator& Generator, UE::HLSLTree::FScope& Scope, int32 OutputIndex, UE::HLSLTree::FExpression const*& OutExpression) const
+{
+	OutExpression = Generator.GetTree().NewExpression<UE::HLSLTree::FExpressionInlineCustomHLSL>(UE::Shader::EValueType::Float3, TEXT("MaterialExpressionAtmosphericLightVector(Parameters)"));
+	return true;
+}
+
 static UE::Shader::FType GetCustomOutputType(const FMaterialHLSLGenerator& Generator, ECustomMaterialOutputType Type)
 {
 	using namespace UE::Shader;
@@ -2536,8 +3118,8 @@ bool UMaterialExpressionVolumetricAdvancedMaterialInput::GenerateHLSLExpression(
 {
 	switch (OutputIndex)
 	{
-	case 0: OutExpression = Generator.GetTree().NewExpression<UE::HLSLTree::FExpressionInlineCustomHLSL>(UE::Shader::EValueType::Float3, TEXT("MaterialExpressionVolumeSampleConservativeDensity(Parameters).rgb"));
-	case 1: OutExpression = Generator.GetTree().NewExpression<UE::HLSLTree::FExpressionInlineCustomHLSL>(UE::Shader::EValueType::Float4, TEXT("MaterialExpressionVolumeSampleConservativeDensity(Parameters)"));
+	case 0: OutExpression = Generator.GetTree().NewExpression<UE::HLSLTree::FExpressionInlineCustomHLSL>(UE::Shader::EValueType::Float3, TEXT("MaterialExpressionVolumeSampleConservativeDensity(Parameters).rgb")); break;
+	case 1: OutExpression = Generator.GetTree().NewExpression<UE::HLSLTree::FExpressionInlineCustomHLSL>(UE::Shader::EValueType::Float4, TEXT("MaterialExpressionVolumeSampleConservativeDensity(Parameters)")); break;
 	default: return Generator.Error(TEXT("Invlid output"));
 	}
 	return OutExpression != nullptr;
@@ -2929,4 +3511,25 @@ bool UMaterialExpressionForLoop::GenerateHLSLStatements(FMaterialHLSLGenerator&
 	return true;
 }
 
+bool UMaterialExpressionDeriveNormalZ::GenerateHLSLExpression(FMaterialHLSLGenerator& Generator, UE::HLSLTree::FScope& Scope, int32 OutputIndex, UE::HLSLTree::FExpression const*& OutExpression) const
+{
+	using namespace UE::HLSLTree;
+
+	if (!InXY.GetTracedInput().Expression)
+	{
+		return Generator.Error(TEXT("Missing input normal xy vector whose z should be derived."));
+	}
+
+	// z = sqrt(saturate(1 - ( x * x + y * y)));
+	FTree& Tree = Generator.GetTree();
+	const FExpression* InputVector = InXY.AcquireHLSLExpression(Generator, Scope);
+	const FExpression* DotResult = Tree.NewDot(InputVector, InputVector);
+	const FExpression* InnerResult = Tree.NewSub(Tree.NewConstant(1.f), DotResult);
+	const FExpression* SaturatedInnerResult = Tree.NewSaturate(InnerResult);
+	const FExpression* DerivedZ = Tree.NewSqrt(SaturatedInnerResult);
+	OutExpression = Tree.NewExpression<FExpressionAppend>(InputVector, DerivedZ);
+
+	return true;
+}
+
 #endif // WITH_EDITOR
diff --git a/Engine/Source/Runtime/Engine/Private/Materials/MaterialHLSLEmitter.cpp b/Engine/Source/Runtime/Engine/Private/Materials/MaterialHLSLEmitter.cpp
--- a/Engine/Source/Runtime/Engine/Private/Materials/MaterialHLSLEmitter.cpp
+++ b/Engine/Source/Runtime/Engine/Private/Materials/MaterialHLSLEmitter.cpp
@@ -10,6 +10,8 @@
 #include "Materials/MaterialExpressionVolumetricAdvancedMaterialOutput.h"
 #include "Materials/MaterialParameterCollection.h"
 #include "Materials/Material.h"
+#include "Materials/MaterialSourceTemplate.h"
+#include "StringTemplate.h"
 #include "MaterialDomain.h"
 #include "MaterialHLSLTree.h"
 #include "MaterialCachedHLSLTree.h"
@@ -67,29 +69,14 @@ static FString GenerateMaterialTemplateHLSL(EShaderPlatform ShaderPlatform,
 
 	const Material::FEmitData& EmitMaterialData = EmitContext.FindData<Material::FEmitData>();
 
-	FString MaterialTemplateSource;
-	LoadShaderSourceFileChecked(TEXT("/Engine/Private/MaterialTemplate.ush"), ShaderPlatform, MaterialTemplateSource);
+	TMap<FString, FString> MaterialSourceTemplateParams;
+	int32 MaterialTemplateLineNumber;
 
-	// Find the string index of the '#line' statement in MaterialTemplate.usf
-	const int32 LineIndex = MaterialTemplateSource.Find(TEXT("#line"), ESearchCase::CaseSensitive);
-	check(LineIndex != INDEX_NONE);
+	FStringTemplateResolver Resolver = FMaterialSourceTemplate::Get().BeginResolve(ShaderPlatform, &MaterialTemplateLineNumber);
+	Resolver.SetParameterMap(&MaterialSourceTemplateParams);
 
-	// Count line endings before the '#line' statement
-	int32 MaterialTemplateLineNumber = INDEX_NONE;
-	int32 StartPosition = LineIndex + 1;
-	do
-	{
-		MaterialTemplateLineNumber++;
-		// Using \n instead of LINE_TERMINATOR as not all of the lines are terminated consistently
-		// Subtract one from the last found line ending index to make sure we skip over it
-		StartPosition = MaterialTemplateSource.Find(TEXT("\n"), ESearchCase::CaseSensitive, ESearchDir::FromEnd, StartPosition - 1);
-	} while (StartPosition != INDEX_NONE);
-	check(MaterialTemplateLineNumber != INDEX_NONE);
-	// At this point MaterialTemplateLineNumber is one less than the line number of the '#line' statement
-	// For some reason we have to add 2 more to the #line value to get correct error line numbers from D3DXCompileShader
-	MaterialTemplateLineNumber += 3;
-
-	FLazyPrintf LazyPrintf(*MaterialTemplateSource);
+	MaterialSourceTemplateParams.Reserve(Resolver.GetTemplate().GetNumNamedParameters());
+	MaterialSourceTemplateParams.Add({ TEXT("line_number"), FString::Printf(TEXT("%u"), MaterialTemplateLineNumber) });
 
 	uint32 NumVertexTexCoords = 0u;
 	uint32 NumPixelTexCoords = 0u;
@@ -114,70 +101,17 @@ static FString GenerateMaterialTemplateHLSL(EShaderPlatform ShaderPlatform,
 	OutCompilationOutput.NumUsedUVScalars = NumPixelTexCoords * 2;
 	OutCompilationOutput.NumUsedCustomInterpolatorScalars = EmitMaterialData.NumInterpolatorComponents;
 
-	LazyPrintf.PushParam(*FString::Printf(TEXT("%u"), NumVertexTexCoords));
-	LazyPrintf.PushParam(*FString::Printf(TEXT("%u"), NumPixelTexCoords));
-	LazyPrintf.PushParam(*FString::Printf(TEXT("%u"), NumCustomVectors));
-	LazyPrintf.PushParam(*FString::Printf(TEXT("%u"), NumTexCoordVectors));
+	MaterialSourceTemplateParams.Add({ TEXT("num_material_texcoords_vertex"), FString::Printf(TEXT("%u"), NumVertexTexCoords) });
+	MaterialSourceTemplateParams.Add({ TEXT("num_material_texcoords"), FString::Printf(TEXT("%u"), NumPixelTexCoords) });
+	MaterialSourceTemplateParams.Add({ TEXT("num_custom_vertex_interpolators"), FString::Printf(TEXT("%u"), NumCustomVectors) });
+	MaterialSourceTemplateParams.Add({ TEXT("num_tex_coord_interpolators"), FString::Printf(TEXT("%u"), NumTexCoordVectors) });
 
 	FString VertexInterpolatorsOffsetsDefinition;
-	LazyPrintf.PushParam(*VertexInterpolatorsOffsetsDefinition);
-
-	FString MaterialAttributesDeclaration;
-	FString MaterialAttributesUtilities;
-	FString MaterialAttributesDefault;
+	MaterialSourceTemplateParams.Add({ TEXT("vertex_interpolators_offsets_definition"), VertexInterpolatorsOffsetsDefinition });
 
-	const EMaterialShadingModel DefaultShadingModel = Material.GetShadingModels().GetFirstShadingModel();
-
-	const TArray<FGuid>& OrderedVisibleAttributes = FMaterialAttributeDefinitionMap::GetOrderedVisibleAttributeList();
-	for (const FGuid& AttributeID : OrderedVisibleAttributes)
-	{
-		const FString PropertyName = FMaterialAttributeDefinitionMap::GetAttributeName(AttributeID);
-		const EMaterialValueType PropertyType = FMaterialAttributeDefinitionMap::GetValueType(AttributeID);
-		const TCHAR* HLSLType = nullptr;
-
-		switch (PropertyType)
-		{
-		case MCT_Float1: case MCT_Float: HLSLType = TEXT("float"); break;
-		case MCT_Float2: HLSLType = TEXT("float2"); break;
-		case MCT_Float3: HLSLType = TEXT("float3"); break;
-		case MCT_Float4: HLSLType = TEXT("float4"); break;
-		case MCT_ShadingModel: case MCT_UInt:  case MCT_UInt1:  HLSLType = TEXT("uint"); break;
-		case MCT_UInt2:  HLSLType = TEXT("uint2"); break;
-		case MCT_UInt3:  HLSLType = TEXT("uint3"); break;
-		case MCT_UInt4:  HLSLType = TEXT("uint4"); break;
-		case MCT_Substrate: HLSLType = TEXT("FSubstrateData"); break;
-		default: break;
-		}
-
-		if (HLSLType)
-		{
-			const FVector4f DefaultValue = FMaterialAttributeDefinitionMap::GetDefaultValue(AttributeID);
-
-			MaterialAttributesDeclaration += FString::Printf(TEXT("\t%s %s;") LINE_TERMINATOR, HLSLType, *PropertyName);
-
-			// Chainable method to set the attribute
-			MaterialAttributesUtilities += FString::Printf(TEXT("FMaterialAttributes FMaterialAttributes_Set%s(FMaterialAttributes InAttributes, %s InValue) { InAttributes.%s = InValue; return InAttributes; }") LINE_TERMINATOR,
-				*PropertyName, HLSLType, *PropertyName);
-
-			// Standard value type
-			switch (PropertyType)
-			{
-			case MCT_Float: case MCT_Float1: MaterialAttributesDefault += FString::Printf(TEXT("\tResult.%s = %0.8f;") LINE_TERMINATOR, *PropertyName, DefaultValue.X); break;
-			case MCT_Float2: MaterialAttributesDefault += FString::Printf(TEXT("\tResult.%s = MaterialFloat2(%0.8f,%0.8f);") LINE_TERMINATOR, *PropertyName, DefaultValue.X, DefaultValue.Y); break;
-			case MCT_Float3: MaterialAttributesDefault += FString::Printf(TEXT("\tResult.%s = MaterialFloat3(%0.8f,%0.8f,%0.8f);") LINE_TERMINATOR, *PropertyName, DefaultValue.X, DefaultValue.Y, DefaultValue.Z); break;
-			case MCT_Float4: MaterialAttributesDefault += FString::Printf(TEXT("\tResult.%s = MaterialFloat4(%0.8f,%0.8f,%0.8f,%0.8f);") LINE_TERMINATOR, *PropertyName, DefaultValue.X, DefaultValue.Y, DefaultValue.Z, DefaultValue.W); break;
-			case MCT_ShadingModel: MaterialAttributesDefault += FString::Printf(TEXT("\tResult.%s = %d;") LINE_TERMINATOR, *PropertyName, (int32)DefaultShadingModel); break;
-			case MCT_Substrate: MaterialAttributesDefault += FString::Printf(TEXT("\tResult.%s = GetInitialisedSubstrateData();") LINE_TERMINATOR, *PropertyName); break; // TODO
-			default: checkNoEntry(); break;
-			}
-		}
-	}
-
-	//LazyPrintf.PushParam(*MaterialAttributesDeclaration);
-	//LazyPrintf.PushParam(*MaterialAttributesUtilities);
-
-	LazyPrintf.PushParam(DeclarationsCode);
-	LazyPrintf.PushParam(TEXT(""));
+	MaterialSourceTemplateParams.Add({ TEXT("material_declarations"), DeclarationsCode });
+	// Empty string because setters included in DeclarationsCode
+	MaterialSourceTemplateParams.Add({ TEXT("material_attributes_utilities"), TEXT("") });
 
 	// Stores the shared shader results member declarations
 	// PixelMembersDeclaration should be the same for all variations, but might change in the future. There are cases where work is shared
@@ -208,9 +142,10 @@ static FString GenerateMaterialTemplateHLSL(EShaderPlatform ShaderPlatform,
 	//	GetSharedInputsMaterialCode(PixelMembersDeclaration[Index], NormalAssignment[Index], PixelMembersSetupAndAssignments[Index], (ECompiledPartialDerivativeVariation)Index);
 	//}
 
-	LazyPrintf.PushParam(*PixelMembersDeclaration);
+	MaterialSourceTemplateParams.Add({ TEXT("pixel_material_inputs"), PixelMembersDeclaration });
 
-	LazyPrintf.PushParam(SharedShaderCode);
+	// This is actually not uniform expressions but material custom outputs, custom HLSL code, and user vertex interpolators
+	MaterialSourceTemplateParams.Add({ TEXT("uniform_material_expressions"), SharedShaderCode });
 
 	// Anything used bye the GenerationFunctionCode() like WorldPositionOffset shouldn't be using texures, right?
 	// Let those use the standard finite differences textures, since they should be the same. If we actually want
@@ -223,32 +158,37 @@ static FString GenerateMaterialTemplateHLSL(EShaderPlatform ShaderPlatform,
 	//}
 	//else
 	{
-		LazyPrintf.PushParam(TEXT("return 0"));
+		MaterialSourceTemplateParams.Add({ TEXT("get_material_emissive_for_cs"), TEXT("return 0") });
 	}
 
-	LazyPrintf.PushParam(*FString::Printf(TEXT("return %.5f"), Material.GetTranslucencyDirectionalLightingIntensity()));
+	MaterialSourceTemplateParams.Add({ TEXT("get_material_translucency_directional_lighting_intensity"), FString::Printf(TEXT("return %.5f"), Material.GetTranslucencyDirectionalLightingIntensity()) });
 
-	LazyPrintf.PushParam(*FString::Printf(TEXT("return %.5f"), Material.GetTranslucentShadowDensityScale()));
-	LazyPrintf.PushParam(*FString::Printf(TEXT("return %.5f"), Material.GetTranslucentSelfShadowDensityScale()));
-	LazyPrintf.PushParam(*FString::Printf(TEXT("return %.5f"), Material.GetTranslucentSelfShadowSecondDensityScale()));
-	LazyPrintf.PushParam(*FString::Printf(TEXT("return %.5f"), Material.GetTranslucentSelfShadowSecondOpacity()));
-	LazyPrintf.PushParam(*FString::Printf(TEXT("return %.5f"), Material.GetTranslucentBackscatteringExponent()));
+	MaterialSourceTemplateParams.Add({ TEXT("get_material_translucent_shadow_density_scale"), FString::Printf(TEXT("return %.5f"), Material.GetTranslucentShadowDensityScale()) });
+	MaterialSourceTemplateParams.Add({ TEXT("get_material_translucent_self_shadow_density_scale"), FString::Printf(TEXT("return %.5f"), Material.GetTranslucentSelfShadowDensityScale()) });
+	MaterialSourceTemplateParams.Add({ TEXT("get_material_translucent_self_shadow_second_density_scale"), FString::Printf(TEXT("return %.5f"), Material.GetTranslucentSelfShadowSecondDensityScale()) });
+	MaterialSourceTemplateParams.Add({ TEXT("get_material_translucent_self_shadow_second_opacity"), FString::Printf(TEXT("return %.5f"), Material.GetTranslucentSelfShadowSecondOpacity()) });
+	MaterialSourceTemplateParams.Add({ TEXT("get_material_translucent_backscattering_exponent"), FString::Printf(TEXT("return %.5f"), Material.GetTranslucentBackscatteringExponent()) });
 
 	{
 		FLinearColor Extinction = Material.GetTranslucentMultipleScatteringExtinction();
 
-		LazyPrintf.PushParam(*FString::Printf(TEXT("return MaterialFloat3(%.5f, %.5f, %.5f)"), Extinction.R, Extinction.G, Extinction.B));
+		MaterialSourceTemplateParams.Add({ TEXT("get_material_translucent_multiple_scattering_extinction"), FString::Printf(TEXT("return MaterialFloat3(%.5f, %.5f, %.5f)"), Extinction.R, Extinction.G, Extinction.B) });
 	}
 
-	LazyPrintf.PushParam(*FString::Printf(TEXT("return %.5f"), Material.GetOpacityMaskClipValue()));
+	MaterialSourceTemplateParams.Add({ TEXT("get_material_opacity_mask_clip_value"), FString::Printf(TEXT("return %.5f"), Material.GetOpacityMaskClipValue()) });
+
+	{
+		const FDisplacementScaling DisplacementScaling = Material.GetDisplacementScaling();
+		MaterialSourceTemplateParams.Add({ TEXT("get_material_displacement_magnitude"), FString::Printf(TEXT("return %.5f"), FMath::Max(0.0f, DisplacementScaling.Magnitude)) });
+		MaterialSourceTemplateParams.Add({ TEXT("get_material_displacement_center"), FString::Printf(TEXT("return %.5f"), FMath::Clamp(DisplacementScaling.Center, 0.0f, 1.0f)) });
+	}
 
-	LazyPrintf.PushParam(TEXT("return Parameters.MaterialAttributes.Displacement"));
-	LazyPrintf.PushParam(TEXT("return Parameters.MaterialAttributes.WorldPositionOffset"));
-	LazyPrintf.PushParam(TEXT("return Parameters.MaterialAttributes.PrevWorldPositionOffset"));
+	MaterialSourceTemplateParams.Add({ TEXT("get_material_world_position_offset_raw"), TEXT("return Parameters.MaterialAttributes.WorldPositionOffset") });
+	MaterialSourceTemplateParams.Add({ TEXT("get_material_previous_world_position_offset_raw"), TEXT("return Parameters.MaterialAttributes.PrevWorldPositionOffset") });
 	
 	// CustomData0/1 are named ClearCoat/ClearCoatRoughness
-	LazyPrintf.PushParam(TEXT("return Parameters.MaterialAttributes.ClearCoat"));
-	LazyPrintf.PushParam(TEXT("return Parameters.MaterialAttributes.ClearCoatRoughness"));
+	MaterialSourceTemplateParams.Add({ TEXT("get_material_custom_data0"), TEXT("return Parameters.MaterialAttributes.ClearCoat") });
+	MaterialSourceTemplateParams.Add({ TEXT("get_material_custom_data1"), TEXT("return Parameters.MaterialAttributes.ClearCoatRoughness") });
 
 	// Print custom texture coordinate assignments, should be fine with regular derivatives
 	{
@@ -258,7 +198,7 @@ static FString GenerateMaterialTemplateHLSL(EShaderPlatform ShaderPlatform,
 			const FString AttributeName = FMaterialAttributeDefinitionMap::GetAttributeName((EMaterialProperty)(MP_CustomizedUVs0 + CustomUVIndex));
 			CustomUVAssignments += FString::Printf(TEXT("\tOutTexCoords[%u] = Parameters.MaterialAttributes.%s;") LINE_TERMINATOR, CustomUVIndex, *AttributeName);
 		}
-		LazyPrintf.PushParam(*CustomUVAssignments);
+		MaterialSourceTemplateParams.Add({ TEXT("get_material_customized_u_vs"), CustomUVAssignments });
 	}
 
 	// Print custom vertex shader interpolator assignments
@@ -268,11 +208,9 @@ static FString GenerateMaterialTemplateHLSL(EShaderPlatform ShaderPlatform,
 		{
 			CustomInterpolatorAssignments += FString::Printf(TEXT("\tOutTexCoords[NUM_MATERIAL_TEXCOORDS + %u] = Parameters.CustomInterpolators[%u];") LINE_TERMINATOR, InterpolatorIndex, InterpolatorIndex);
 		}
-		LazyPrintf.PushParam(*CustomInterpolatorAssignments);
+		MaterialSourceTemplateParams.Add({ TEXT("get_custom_interpolators"), CustomInterpolatorAssignments });
 	}
 
-	//LazyPrintf.PushParam(*MaterialAttributesDefault);
-
 	//if (bEnableExecutionFlow)
 	{
 		FString EvaluateMaterialDeclaration;
@@ -328,7 +266,7 @@ static FString GenerateMaterialTemplateHLSL(EShaderPlatform ShaderPlatform,
 			}
 		}
 
-		LazyPrintf.PushParam(*EvaluateMaterialDeclaration);
+		MaterialSourceTemplateParams.Add({ TEXT("evaluate_material_attributes"), EvaluateMaterialDeclaration });
 
 		for (int32 DerivativeIndex = 0; DerivativeIndex < 2; ++DerivativeIndex)
 		{
@@ -368,15 +306,22 @@ static FString GenerateMaterialTemplateHLSL(EShaderPlatform ShaderPlatform,
 				}
 			}
 
-			LazyPrintf.PushParam(TEXT(""));
-			LazyPrintf.PushParam(*EvaluateMaterialAttributesPhase0[DerivativeIndex]);
-			LazyPrintf.PushParam(*EvaluateMaterialAttributesPhase1[DerivativeIndex]);
+			if (DerivativeIndex == 0)
+			{
+				MaterialSourceTemplateParams.Add({ TEXT("calc_pixel_material_inputs_initial_calculations"), TEXT("") });
+				MaterialSourceTemplateParams.Add({ TEXT("calc_pixel_material_inputs_normal"), EvaluateMaterialAttributesPhase0[DerivativeIndex] });
+				MaterialSourceTemplateParams.Add({ TEXT("calc_pixel_material_inputs_other_inputs"), EvaluateMaterialAttributesPhase1[DerivativeIndex] });
+			}
+			else // if (DerivativeIndex == 1)
+			{
+				MaterialSourceTemplateParams.Add({ TEXT("calc_pixel_material_inputs_analytic_derivatives_initial"), TEXT("") });
+				MaterialSourceTemplateParams.Add({ TEXT("calc_pixel_material_inputs_analytic_derivatives_normal"), EvaluateMaterialAttributesPhase0[DerivativeIndex] });
+				MaterialSourceTemplateParams.Add({ TEXT("calc_pixel_material_inputs_analytic_derivatives_other_inputs"), EvaluateMaterialAttributesPhase1[DerivativeIndex] });
+			}
 		}
 	}
 
-	LazyPrintf.PushParam(*FString::Printf(TEXT("%u"), MaterialTemplateLineNumber));
-
-	return LazyPrintf.GetResultString();
+	return Resolver.Finalize();
 }
 
 static void GetMaterialEnvironment(EShaderPlatform InPlatform,
@@ -399,12 +344,15 @@ static void GetMaterialEnvironment(EShaderPlatform InPlatform,
 	const bool bNeedsParticlePosition = EmitMaterialData.IsExternalInputUsed(Material::EExternalInput::ParticleTranslatedWorldPosition) ||
 		EmitMaterialData.IsExternalInputUsed(Material::EExternalInput::ParticleRadius);
 
-	if (bNeedsParticlePosition)// || Material->ShouldGenerateSphericalParticleNormals() || bUsesSphericalParticleOpacity)
+	if (bNeedsParticlePosition || InMaterial.ShouldGenerateSphericalParticleNormals()) // || bUsesSphericalParticleOpacity)
 	{
 		OutEnvironment.SetDefine(TEXT("NEEDS_PARTICLE_POSITION"), 1);
 	}
 
-	if (false)//bNeedsParticleVelocity || Material->IsUsedWithNiagaraMeshParticles())
+	const bool bNeedsParticleVelocity = EmitMaterialData.IsExternalInputUsed(Material::EExternalInput::ParticleDirection)
+		|| EmitMaterialData.IsExternalInputUsed(Material::EExternalInput::ParticleSpeed);
+
+	if (bNeedsParticleVelocity || InMaterial.IsUsedWithNiagaraMeshParticles())
 	{
 		OutEnvironment.SetDefine(TEXT("NEEDS_PARTICLE_VELOCITY"), 1);
 	}
@@ -415,7 +363,9 @@ static void GetMaterialEnvironment(EShaderPlatform InPlatform,
 		OutEnvironment.SetDefine(TEXT("DYNAMIC_PARAMETERS_MASK"), EmitContext.DynamicParticleParameterMask);
 	}
 
-	if (false)//bNeedsParticleTime)
+	const bool bNeedsParticleTime = EmitMaterialData.IsExternalInputUsed(Material::EExternalInput::ParticleRelativeTime);
+
+	if (bNeedsParticleTime)
 	{
 		OutEnvironment.SetDefine(TEXT("NEEDS_PARTICLE_TIME"), 1);
 	}
@@ -425,7 +375,9 @@ static void GetMaterialEnvironment(EShaderPlatform InPlatform,
 		OutEnvironment.SetDefine(TEXT("USES_PARTICLE_MOTION_BLUR"), 1);
 	}
 
-	if (false)//bNeedsParticleRandom)
+	const bool bNeedsParticleRandom = EmitMaterialData.IsExternalInputUsed(Material::EExternalInput::ParticleRandom);
+
+	if (bNeedsParticleRandom)
 	{
 		OutEnvironment.SetDefine(TEXT("NEEDS_PARTICLE_RANDOM"), 1);
 	}
@@ -460,7 +412,9 @@ static void GetMaterialEnvironment(EShaderPlatform InPlatform,
 		OutEnvironment.SetDefine(TEXT("NEEDS_WORLD_POSITION_EXCLUDING_SHADER_OFFSETS"), TEXT("1"));
 	}
 
-	if (false)//bNeedsParticleSize)
+	const bool bNeedsParticleSize = EmitMaterialData.IsExternalInputUsed(Material::EExternalInput::ParticleSize);
+
+	if (bNeedsParticleSize)
 	{
 		OutEnvironment.SetDefine(TEXT("NEEDS_PARTICLE_SIZE"), TEXT("1"));
 	}
diff --git a/Engine/Source/Runtime/Engine/Private/Materials/MaterialHLSLGenerator.cpp b/Engine/Source/Runtime/Engine/Private/Materials/MaterialHLSLGenerator.cpp
--- a/Engine/Source/Runtime/Engine/Private/Materials/MaterialHLSLGenerator.cpp
+++ b/Engine/Source/Runtime/Engine/Private/Materials/MaterialHLSLGenerator.cpp
@@ -234,6 +234,11 @@ bool FMaterialHLSLGenerator::GenerateResult(UE::HLSLTree::FScope& Scope)
 				{
 					const EMaterialProperty Property = (EMaterialProperty)PropertyIndex;
 
+					if (Property == MP_CustomOutput)
+					{
+						continue;
+					}
+
 					// We're only interesting in attributes that map to valid fields
 					const UE::Shader::FStructField* AttributeField = CachedTree.GetMaterialAttributesType()->FindFieldByName(*FMaterialAttributeDefinitionMap::GetAttributeName(Property));
 					if (AttributeField)
diff --git a/Engine/Source/Runtime/Engine/Private/Materials/MaterialHLSLTree.cpp b/Engine/Source/Runtime/Engine/Private/Materials/MaterialHLSLTree.cpp
--- a/Engine/Source/Runtime/Engine/Private/Materials/MaterialHLSLTree.cpp
+++ b/Engine/Source/Runtime/Engine/Private/Materials/MaterialHLSLTree.cpp
@@ -8,8 +8,10 @@
 #include "MaterialShared.h"
 #include "MaterialCachedData.h"
 #include "Materials/MaterialFunctionInterface.h"
+#include "Materials/MaterialExpressionVectorNoise.h"
 #include "Engine/BlendableInterface.h" // BL_AfterTonemapping
 #include "VT/VirtualTextureScalability.h"
+#include "VT/RuntimeVirtualTexture.h"
 
 namespace UE::HLSLTree::Material
 {
@@ -142,6 +144,11 @@ FExternalInputDescription GetExternalInputDescription(EExternalInput Input)
 	case EExternalInput::ParticleColor: return FExternalInputDescription(TEXT("ParticleColor"), Shader::EValueType::Float4);
 	case EExternalInput::ParticleTranslatedWorldPosition: return FExternalInputDescription(TEXT("ParticleTranslatedWorldPosition"), Shader::EValueType::Float3);
 	case EExternalInput::ParticleRadius: return FExternalInputDescription(TEXT("ParticleRadius"), Shader::EValueType::Float1);
+	case EExternalInput::ParticleDirection: return FExternalInputDescription(TEXT("ParticleDirection"), Shader::EValueType::Float3);
+	case EExternalInput::ParticleSpeed: return FExternalInputDescription(TEXT("ParticleSpeed"), Shader::EValueType::Float1);
+	case EExternalInput::ParticleRelativeTime: return FExternalInputDescription(TEXT("ParticleRelativeTime"), Shader::EValueType::Float1);
+	case EExternalInput::ParticleRandom: return FExternalInputDescription(TEXT("ParticleRandom"), Shader::EValueType::Float1);
+	case EExternalInput::ParticleSize: return FExternalInputDescription(TEXT("ParticleSize"), Shader::EValueType::Float2);
 
 	case EExternalInput::IsOrthographic: return FExternalInputDescription(TEXT("IsOrthographic"), Shader::EValueType::Float1);
 
@@ -332,6 +339,11 @@ void FExpressionExternalInput::EmitValueShader(FEmitContext& Context, FEmitScope
 		case EExternalInput::ParticleColor: Code = TEXT("Parameters.Particle.Color"); break;
 		case EExternalInput::ParticleTranslatedWorldPosition: Code = TEXT("Parameters.Particle.TranslatedWorldPositionAndSize.xyz"); break;
 		case EExternalInput::ParticleRadius: Code = TEXT("Parameters.Particle.TranslatedWorldPositionAndSize.w"); break;
+		case EExternalInput::ParticleDirection: Code = TEXT("Parameters.Particle.Velocity.xyz"); break;
+		case EExternalInput::ParticleSpeed: Code = TEXT("Parameters.Particle.Velocity.w"); break;
+		case EExternalInput::ParticleRelativeTime: Code = TEXT("Parameters.Particle.RelativeTime"); break;
+		case EExternalInput::ParticleRandom: Code = TEXT("Parameters.Particle.Random"); break;
+		case EExternalInput::ParticleSize: Code = TEXT("Parameters.Particle.Size"); break;
 
 		case EExternalInput::IsOrthographic: Code = TEXT("((View.ViewToClip[3][3] < 1.0f) ? 0.0f : 1.0f)"); break;
 
@@ -976,29 +988,56 @@ void FExpressionTextureSample::EmitValueShader(FEmitContext& Context, FEmitScope
 		check(TextureParameterInfo.TextureIndex != INDEX_NONE);
 		const int32 TextureParameterIndex = Context.MaterialCompilationOutput->UniformExpressionSet.FindOrAddTextureParameter(EMaterialTextureParameterType::Virtual, TextureParameterInfo);
 
-		// Using Source size because we care about the aspect ratio of each block (each block of multi-block texture must have same aspect ratio)
-		// We can still combine multi-block textures of different block aspect ratios, as long as each block has the same ratio
-		// This is because we only need to overlay VT pages from within a given block
-		const float TextureAspectRatio = (float)Texture->Source.GetSizeX() / (float)Texture->Source.GetSizeY();
-
-		const bool AdaptiveVirtualTexture = false;
-		const bool bGenerateFeedback = (Context.ShaderFrequency == SF_Pixel);
-		int32 VTStackIndex = Private::AcquireVTStackIndex(Context, Scope, EmitMaterialData,
-			LocalMipValueMode,
-			StaticAddressX,
-			StaticAddressY,
-			TextureAspectRatio,
-			EmitTexCoordValue,
-			EmitTexCoordValueDdx,
-			EmitTexCoordValueDdy,
-			EmitMipValue,
-			INDEX_NONE,
-			AdaptiveVirtualTexture,
-			bGenerateFeedback);
+		const bool AdaptiveVirtualTexture = bAdaptive;
+		const bool bGenerateFeedback = bEnableFeedback && Context.ShaderFrequency == SF_Pixel;
+		int32 VTLayerIndex = TextureLayerIndex;
+		int32 VTPageTableIndex = PageTableLayerIndex;
+		int32 VTStackIndex;
+
+		if (VTLayerIndex != INDEX_NONE)
+		{
+			// The layer index in the virtual texture stack is already known
+			// Create a page table sample for each new combination of virtual texture and sample parameters
+			VTStackIndex = Private::AcquireVTStackIndex(Context, Scope, EmitMaterialData,
+				LocalMipValueMode,
+				StaticAddressX,
+				StaticAddressY,
+				1.0f,
+				EmitTexCoordValue,
+				EmitTexCoordValueDdx,
+				EmitTexCoordValueDdy,
+				EmitMipValue,
+				TextureParameterInfo.TextureIndex,
+				AdaptiveVirtualTexture,
+				bGenerateFeedback);
 
-		// Allocate a layer in the virtual texture stack for this physical sample
-		int32 VTLayerIndex = Context.MaterialCompilationOutput->UniformExpressionSet.AddVTLayer(VTStackIndex, TextureParameterIndex);
-		int32 VTPageTableIndex = VTLayerIndex;
+			Context.MaterialCompilationOutput->UniformExpressionSet.SetVTLayer(VTStackIndex, VTLayerIndex, TextureParameterIndex);
+		}
+		else
+		{
+			// Using Source size because we care about the aspect ratio of each block (each block of multi-block texture must have same aspect ratio)
+			// We can still combine multi-block textures of different block aspect ratios, as long as each block has the same ratio
+			// This is because we only need to overlay VT pages from within a given block
+			const float TextureAspectRatio = (float)Texture->Source.GetSizeX() / (float)Texture->Source.GetSizeY();
+
+			// Create a page table sample for each new set of sample parameters
+			VTStackIndex = Private::AcquireVTStackIndex(Context, Scope, EmitMaterialData,
+				LocalMipValueMode,
+				StaticAddressX,
+				StaticAddressY,
+				TextureAspectRatio,
+				EmitTexCoordValue,
+				EmitTexCoordValueDdx,
+				EmitTexCoordValueDdy,
+				EmitMipValue,
+				INDEX_NONE,
+				AdaptiveVirtualTexture,
+				bGenerateFeedback);
+
+			// Allocate a layer in the virtual texture stack for this physical sample
+			VTLayerIndex = Context.MaterialCompilationOutput->UniformExpressionSet.AddVTLayer(VTStackIndex, TextureParameterIndex);
+			VTPageTableIndex = VTLayerIndex;
+		}
 
 		TStringBuilder<64> FormattedTexture;
 		FormattedTexture.Appendf(TEXT("Material.VirtualTexturePhysical_%d"), TextureParameterIndex);
@@ -1153,6 +1192,159 @@ void FExpressionTextureSize::EmitValuePreshader(FEmitContext& Context, FEmitScop
 	OutResult.Preshader.WriteOpcode(Shader::EPreshaderOpcode::TextureSize).Write<FMemoryImageMaterialParameterInfo>(TextureValue.ParameterInfo).Write(TextureIndex);
 }
 
+bool FExpressionRuntimeVirtualTextureUniform::PrepareValue(FEmitContext& Context, FEmitScope& Scope, const FRequestedType& RequestedType, FPrepareValueResult& OutResult) const
+{
+	const FPreparedType& TextureType = Context.PrepareExpression(TextureExpression, Scope, FMaterialTextureValue::GetTypeName());
+	if (TextureType.Type.ObjectType != FMaterialTextureValue::GetTypeName())
+	{
+		return Context.Error(TEXT("Expected texture"));
+	}
+
+	return OutResult.SetType(Context, RequestedType, EExpressionEvaluation::Preshader, URuntimeVirtualTexture::GetUniformParameterType((int32)UniformType));
+}
+
+void FExpressionRuntimeVirtualTextureUniform::EmitValuePreshader(FEmitContext& Context, FEmitScope& Scope, const FRequestedType& RequestedType, FEmitValuePreshaderResult& OutResult) const
+{
+	FMaterialTextureValue TextureValue;
+	TextureExpression->GetValueObject(Context, Scope, TextureValue);
+
+	const int32 TextureIndex = Context.Material->GetReferencedTextures().Find(TextureValue.Texture);
+	check(TextureIndex != INDEX_NONE);
+	const int32 VectorIndex = (int32)UniformType;
+
+	++Context.PreshaderStackPosition;
+	OutResult.Type = URuntimeVirtualTexture::GetUniformParameterType(VectorIndex);
+	OutResult.Preshader.WriteOpcode(Shader::EPreshaderOpcode::RuntimeVirtualTextureUniform).Write(ParameterInfo).Write(TextureIndex).Write(VectorIndex);
+}
+
+bool FExpressionVirtualTextureUnpack::PrepareValue(FEmitContext& Context, FEmitScope& Scope, const FRequestedType& RequestedType, FPrepareValueResult& OutResult) const
+{
+	if (UnpackType == EVirtualTextureUnpackType::BaseColorYCoCg)
+	{
+		const FPreparedType& SampleType = Context.PrepareExpression(SampleLayer0Expression, Scope, Shader::EValueType::Float4);
+		if (SampleType.IsVoid())
+		{
+			return false;
+		}
+		return OutResult.SetType(Context, RequestedType, EExpressionEvaluation::Shader, Shader::EValueType::Float3);
+	}
+	else if (UnpackType == EVirtualTextureUnpackType::NormalBC3)
+	{
+		const FPreparedType& SampleType = Context.PrepareExpression(SampleLayer1Expression, Scope, Shader::EValueType::Float4);
+		if (SampleType.IsVoid())
+		{
+			return false;
+		}
+		return OutResult.SetType(Context, RequestedType, EExpressionEvaluation::Shader, Shader::EValueType::Float3);
+	}
+	else if (UnpackType == EVirtualTextureUnpackType::NormalBC5)
+	{
+		const FPreparedType& SampleType = Context.PrepareExpression(SampleLayer1Expression, Scope, Shader::EValueType::Float4);
+		if (SampleType.IsVoid())
+		{
+			return false;
+		}
+		return OutResult.SetType(Context, RequestedType, EExpressionEvaluation::Shader, Shader::EValueType::Float3);
+	}
+	else if (UnpackType == EVirtualTextureUnpackType::NormalBC3BC3)
+	{
+		const FPreparedType& Sample0Type = Context.PrepareExpression(SampleLayer0Expression, Scope, Shader::EValueType::Float4);
+		const FPreparedType& Sample1Type = Context.PrepareExpression(SampleLayer1Expression, Scope, Shader::EValueType::Float4);
+		if (Sample0Type.IsVoid() || Sample1Type.IsVoid())
+		{
+			return false;
+		}
+		return OutResult.SetType(Context, RequestedType, EExpressionEvaluation::Shader, Shader::EValueType::Float3);
+	}
+	else if (UnpackType == EVirtualTextureUnpackType::NormalBC5BC1)
+	{
+		const FPreparedType& Sample1Type = Context.PrepareExpression(SampleLayer1Expression, Scope, Shader::EValueType::Float4);
+		const FPreparedType& Sample2Type = Context.PrepareExpression(SampleLayer2Expression, Scope, Shader::EValueType::Float4);
+		if (Sample1Type.IsVoid() || Sample2Type.IsVoid())
+		{
+			return false;
+		}
+		return OutResult.SetType(Context, RequestedType, EExpressionEvaluation::Shader, Shader::EValueType::Float3);
+	}
+	else if (UnpackType == EVirtualTextureUnpackType::HeightR16)
+	{
+		const FPreparedType& Sample0Type = Context.PrepareExpression(SampleLayer0Expression, Scope, Shader::EValueType::Float4);
+		const FPreparedType& HeightScaleBiasType = Context.PrepareExpression(WorldHeightUnpackUniformExpression, Scope, Shader::EValueType::Float2);
+		if (Sample0Type.IsVoid() || HeightScaleBiasType.IsVoid())
+		{
+			return false;
+		}
+		return OutResult.SetType(Context, RequestedType, EExpressionEvaluation::Shader, Shader::EValueType::Float1);
+	}
+	else if (UnpackType == EVirtualTextureUnpackType::NormalBGR565)
+	{
+		const FPreparedType& SampleType = Context.PrepareExpression(SampleLayer1Expression, Scope, Shader::EValueType::Float4);
+		if (SampleType.IsVoid())
+		{
+			return false;
+		}
+		return OutResult.SetType(Context, RequestedType, EExpressionEvaluation::Shader, Shader::EValueType::Float3);
+	}
+	else if (UnpackType == EVirtualTextureUnpackType::BaseColorSRGB)
+	{
+		const FPreparedType& SampleType = Context.PrepareExpression(SampleLayer0Expression, Scope, Shader::EValueType::Float4);
+		if (SampleType.IsVoid())
+		{
+			return false;
+		}
+		return OutResult.SetType(Context, RequestedType, EExpressionEvaluation::Shader, Shader::EValueType::Float3);
+	}
+	
+	return Context.Error(TEXT("Unexpected virtual texture unpack type."));
+}
+
+void FExpressionVirtualTextureUnpack::EmitValueShader(FEmitContext& Context, FEmitScope& Scope, const FRequestedType& RequestedType, FEmitValueShaderResult& OutResult) const
+{
+	if (UnpackType == EVirtualTextureUnpackType::BaseColorYCoCg)
+	{
+		FEmitShaderExpression* EmitSampleLayer0 = SampleLayer0Expression->GetValueShader(Context, Scope, Shader::EValueType::Float4);
+		OutResult.Code = Context.EmitExpression(Scope, Shader::EValueType::Float3, TEXT("VirtualTextureUnpackBaseColorYCoCg(%)"), EmitSampleLayer0);
+	}
+	else if (UnpackType == EVirtualTextureUnpackType::NormalBC3)
+	{
+		FEmitShaderExpression* EmitSampleLayer1 = SampleLayer1Expression->GetValueShader(Context, Scope, Shader::EValueType::Float4);
+		OutResult.Code = Context.EmitExpression(Scope, Shader::EValueType::Float3, TEXT("VirtualTextureUnpackNormalBC3(%)"), EmitSampleLayer1);
+	}
+	else if (UnpackType == EVirtualTextureUnpackType::NormalBC5)
+	{
+		FEmitShaderExpression* EmitSampleLayer1 = SampleLayer1Expression->GetValueShader(Context, Scope, Shader::EValueType::Float4);
+		OutResult.Code = Context.EmitExpression(Scope, Shader::EValueType::Float3, TEXT("VirtualTextureUnpackNormalBC5(%)"), EmitSampleLayer1);
+	}
+	else if (UnpackType == EVirtualTextureUnpackType::NormalBC3BC3)
+	{
+		FEmitShaderExpression* EmitSampleLayer0 = SampleLayer0Expression->GetValueShader(Context, Scope, Shader::EValueType::Float4);
+		FEmitShaderExpression* EmitSampleLayer1 = SampleLayer1Expression->GetValueShader(Context, Scope, Shader::EValueType::Float4);
+		OutResult.Code = Context.EmitExpression(Scope, Shader::EValueType::Float3, TEXT("VirtualTextureUnpackNormalBC3BC3(%, %)"), EmitSampleLayer0, EmitSampleLayer1);
+	}
+	else if (UnpackType == EVirtualTextureUnpackType::NormalBC5BC1)
+	{
+		FEmitShaderExpression* EmitSampleLayer1 = SampleLayer1Expression->GetValueShader(Context, Scope, Shader::EValueType::Float4);
+		FEmitShaderExpression* EmitSampleLayer2 = SampleLayer2Expression->GetValueShader(Context, Scope, Shader::EValueType::Float4);
+		OutResult.Code = Context.EmitExpression(Scope, Shader::EValueType::Float3, TEXT("VirtualTextureUnpackNormalBC5BC1(%, %)"), EmitSampleLayer1, EmitSampleLayer2);
+	}
+	else if (UnpackType == EVirtualTextureUnpackType::HeightR16)
+	{
+		FEmitShaderExpression* EmitSampleLayer0 = SampleLayer0Expression->GetValueShader(Context, Scope, Shader::EValueType::Float4);
+		FEmitShaderExpression* EmitHeightScaleBias = WorldHeightUnpackUniformExpression->GetValueShader(Context, Scope, Shader::EValueType::Float2);
+		OutResult.Code = Context.EmitExpression(Scope, Shader::EValueType::Float1, TEXT("VirtualTextureUnpackHeight(%, %)"), EmitSampleLayer0, EmitHeightScaleBias);
+	}
+	else if (UnpackType == EVirtualTextureUnpackType::NormalBGR565)
+	{
+		FEmitShaderExpression* EmitSampleLayer1 = SampleLayer1Expression->GetValueShader(Context, Scope, Shader::EValueType::Float4);
+		OutResult.Code = Context.EmitExpression(Scope, Shader::EValueType::Float3, TEXT("VirtualTextureUnpackNormalBGR565(%)"), EmitSampleLayer1);
+	}
+	else if (UnpackType == EVirtualTextureUnpackType::BaseColorSRGB)
+	{
+		FEmitShaderExpression* EmitSampleLayer0 = SampleLayer0Expression->GetValueShader(Context, Scope, Shader::EValueType::Float4);
+		OutResult.Code = Context.EmitExpression(Scope, Shader::EValueType::Float3, TEXT("VirtualTextureUnpackBaseColorSRGB(%)"), EmitSampleLayer0);
+	}
+}
+
 bool FExpressionFunctionCall::PrepareValue(FEmitContext& Context, FEmitScope& Scope, const FRequestedType& RequestedType, FPrepareValueResult& OutResult) const
 {
 	if (Context.bMarkLiveValues)
@@ -1244,6 +1436,125 @@ void FExpressionSceneTexture::EmitValueShader(FEmitContext& Context, FEmitScope&
 	OutResult.Code = EmitLookup;
 }
 
+bool FExpressionScreenAlignedUV::PrepareValue(FEmitContext& Context, FEmitScope& Scope, const FRequestedType& RequestedType, FPrepareValueResult& OutResult) const
+{
+	if (OffsetExpression)
+	{
+		const FPreparedType& OffsetType = Context.PrepareExpression(OffsetExpression, Scope, Shader::EValueType::Float2);
+		if (OffsetType.IsVoid())
+		{
+			return false;
+		}
+	}
+	else if (ViewportUVExpression)
+	{
+		const FPreparedType& ViewportUVType = Context.PrepareExpression(ViewportUVExpression, Scope, Shader::EValueType::Float2);
+		if (ViewportUVType.IsVoid())
+		{
+			return false;
+		}
+	}
+
+	return OutResult.SetType(Context, RequestedType, EExpressionEvaluation::Shader, Shader::EValueType::Float2);
+}
+
+void FExpressionScreenAlignedUV::EmitValueShader(FEmitContext& Context, FEmitScope& Scope, const FRequestedType& RequestedType, FEmitValueShaderResult& OutResult) const
+{
+	if (OffsetExpression)
+	{
+		FEmitShaderExpression* EmitOffset = OffsetExpression->GetValueShader(Context, Scope, Shader::EValueType::Float2);
+		OutResult.Code = Context.EmitExpression(Scope, Shader::EValueType::Float2, TEXT("CalcScreenUVFromOffsetFraction(GetScreenPosition(Parameters), %)"), EmitOffset);
+	}
+	else if (ViewportUVExpression)
+	{
+		FEmitShaderExpression* EmitViewportUV = ViewportUVExpression->GetValueShader(Context, Scope, Shader::EValueType::Float2);
+
+		check(Context.Material);
+		const EMaterialDomain MaterialDomain = Context.Material->GetMaterialDomain();
+		OutResult.Code = Context.EmitExpression(
+			Scope,
+			Shader::EValueType::Float2,
+			TEXT("clamp(ViewportUVToBufferUV(%),%,%)"),
+			EmitViewportUV,
+			MaterialDomain == MD_Surface ? TEXT("ResolvedView.BufferBilinearUVMinMax.xy") : TEXT("View.BufferBilinearUVMinMax.xy"),
+			MaterialDomain == MD_Surface ? TEXT("ResolvedView.BufferBilinearUVMinMax.zw") : TEXT("View.BufferBilinearUVMinMax.zw"));
+	}
+	else
+	{
+		OutResult.Code = Context.EmitExpression(Scope, Shader::EValueType::Float2, TEXT("ScreenAlignedPosition(GetScreenPosition(Parameters))"));
+	}
+}
+
+bool FExpressionSceneDepth::PrepareValue(FEmitContext& Context, FEmitScope& Scope, const FRequestedType& RequestedType, FPrepareValueResult& OutResult) const
+{
+	if (Context.ShaderFrequency == SF_Vertex && Context.TargetParameters.FeatureLevel <= ERHIFeatureLevel::ES3_1)
+	{
+		// mobile currently does not support this, we need to read a separate copy of the depth, we must disable framebuffer fetch and force scene texture reads.
+		return Context.Error(TEXT("Cannot read scene depth from the vertex shader with the Mobile feature level"));
+	}
+
+	if (Context.Material && Context.Material->IsTranslucencyWritingVelocity())
+	{
+		return Context.Error(TEXT("Translucenct material with 'Output Velocity' enabled will write to depth buffer, therefore cannot read from depth buffer at the same time."));
+	}
+
+	const FPreparedType& ScreenUVType = Context.PrepareExpression(ScreenUVExpression, Scope, Shader::EValueType::Float2);
+	if (ScreenUVType.IsVoid())
+	{
+		return false;
+	}
+
+	if (Context.bMarkLiveValues && Context.MaterialCompilationOutput)
+	{
+		Context.MaterialCompilationOutput->SetIsSceneTextureUsed(PPI_SceneDepth);
+	}
+
+	return OutResult.SetType(Context, RequestedType, EExpressionEvaluation::Shader, Shader::EValueType::Float1);
+}
+
+void FExpressionSceneDepth::EmitValueShader(FEmitContext& Context, FEmitScope& Scope, const FRequestedType& RequestedType, FEmitValueShaderResult& OutResult) const
+{
+	FEmitShaderExpression* EmitScreenUV = ScreenUVExpression->GetValueShader(Context, Scope, Shader::EValueType::Float2);
+	OutResult.Code = Context.EmitExpression(Scope, Shader::EValueType::Float1, TEXT("CalcSceneDepth(%)"), EmitScreenUV);
+}
+
+bool FExpressionSceneColor::PrepareValue(FEmitContext& Context, FEmitScope& Scope, const FRequestedType& RequestedType, FPrepareValueResult& OutResult) const
+{
+	if (Context.ShaderFrequency != SF_Pixel)
+	{
+		return Context.Error(TEXT("Invalid node used in vertex/hull/domain shader input!"));
+	}
+
+	if (Context.Material && Context.Material->GetMaterialDomain() != MD_Surface)
+	{
+		return Context.Error(TEXT("SceneColor lookups are only available when MaterialDomain = Surface."));
+	}
+
+	if (Context.TargetParameters.FeatureLevel < ERHIFeatureLevel::SM5)
+	{
+		return Context.Errorf(TEXT("Node not supported in feature level %d. %d required."), Context.TargetParameters.FeatureLevel, ERHIFeatureLevel::SM5);
+	}
+
+	const FPreparedType& ScreenUVType = Context.PrepareExpression(ScreenUVExpression, Scope, Shader::EValueType::Float2);
+	if (ScreenUVType.IsVoid())
+	{
+		return false;
+	}
+
+	if (Context.bMarkLiveValues && Context.MaterialCompilationOutput)
+	{
+		Context.MaterialCompilationOutput->SetIsSceneTextureUsed(PPI_SceneColor);
+	}
+
+	return OutResult.SetType(Context, RequestedType, EExpressionEvaluation::Shader, Shader::EValueType::Float3);
+}
+
+void FExpressionSceneColor::EmitValueShader(FEmitContext& Context, FEmitScope& Scope, const FRequestedType& RequestedType, FEmitValueShaderResult& OutResult) const
+{
+	FEmitShaderExpression* EmitScreenUV = ScreenUVExpression->GetValueShader(Context, Scope, Shader::EValueType::Float2);
+	OutResult.Code = Context.EmitExpression(Scope, Shader::EValueType::Float3, TEXT("DecodeSceneColorForMaterialNode(%)"), EmitScreenUV);
+}
+
 bool FExpressionNoise::PrepareValue(FEmitContext& Context, FEmitScope& Scope, const FRequestedType& RequestedType, FPrepareValueResult& OutResult) const
 {
 	// TODO - we support Float3 or Double3 position input
@@ -1286,6 +1597,61 @@ void FExpressionNoise::EmitValueShader(FEmitContext& Context, FEmitScope& Scope,
 		Parameters.RepeatSize);
 }
 
+bool FExpressionVectorNoise::PrepareValue(FEmitContext& Context, FEmitScope& Scope, const FRequestedType& RequestedType, FPrepareValueResult& OutResult) const
+{
+	// TODO - we support Float3 or Double3 position input
+	const FPreparedType& PositionType = Context.PrepareExpression(PositionExpression, Scope, Shader::EValueType::Float3);
+	if (PositionType.IsVoid())
+	{
+		return false;
+	}
+
+	const EVectorNoiseFunction NoiseFunction = (EVectorNoiseFunction)Parameters.Function;
+	Shader::EValueType ResultType;
+	if (NoiseFunction == VNF_GradientALU || NoiseFunction == VNF_VoronoiALU)
+	{
+		ResultType = Shader::EValueType::Float4;
+	}
+	else
+	{
+		ResultType = Shader::EValueType::Float3;
+	}
+
+	return OutResult.SetType(Context, RequestedType, EExpressionEvaluation::Shader, ResultType);
+}
+
+void FExpressionVectorNoise::EmitValueShader(FEmitContext& Context, FEmitScope& Scope, const FRequestedType& RequestedType, FEmitValueShaderResult& OutResult) const
+{
+	FEmitShaderExpression* EmitPosition = PositionExpression->GetValueShader(Context, Scope, Shader::EValueType::Float3);
+
+	// LWC_TODO - maybe possible/useful to add LWC-aware noise functions
+	const EVectorNoiseFunction NoiseFunction = (EVectorNoiseFunction)Parameters.Function;
+	if (NoiseFunction == VNF_GradientALU || NoiseFunction == VNF_VoronoiALU)
+	{
+		OutResult.Code = Context.EmitExpression(
+			Scope,
+			Shader::EValueType::Float4,
+			TEXT("MaterialExpressionVectorNoise(%,%,%,%,%)"),
+			EmitPosition,
+			Parameters.Quality,
+			Parameters.Function,
+			Parameters.bTiling,
+			Parameters.TileSize);
+	}
+	else
+	{
+		OutResult.Code = Context.EmitExpression(
+			Scope,
+			Shader::EValueType::Float3,
+			TEXT("MaterialExpressionVectorNoise(%,%,%,%,%).xyz"),
+			EmitPosition,
+			Parameters.Quality,
+			Parameters.Function,
+			Parameters.bTiling,
+			Parameters.TileSize);
+	}
+}
+
 void FExpressionVertexInterpolator::ComputeAnalyticDerivatives(FTree& Tree, FExpressionDerivatives& OutResult) const
 {
 	// TODO - we can access interpolator derivatives via TexCoord_DDX/DDY
diff --git a/Engine/Source/Runtime/Engine/Private/Materials/MaterialUniformExpressions.cpp b/Engine/Source/Runtime/Engine/Private/Materials/MaterialUniformExpressions.cpp
--- a/Engine/Source/Runtime/Engine/Private/Materials/MaterialUniformExpressions.cpp
+++ b/Engine/Source/Runtime/Engine/Private/Materials/MaterialUniformExpressions.cpp
@@ -676,6 +676,11 @@ int32 FUniformExpressionSet::AddVTLayer(int32 StackIndex, int32 TextureIndex)
 	return VTLayerIndex;
 }
 
+void FUniformExpressionSet::SetVTLayer(int32 StackIndex, int32 VTLayerIndex, int32 TextureIndex)
+{
+	VTStacks[StackIndex].SetLayer(VTLayerIndex, TextureIndex);
+}
+
 void FUniformExpressionSet::GetGameThreadTextureValue(EMaterialTextureParameterType Type, int32 Index, const UMaterialInterface* MaterialInterface, const FMaterial& Material, UTexture*& OutValue, bool bAllowOverride) const
 {
 	check(IsInGameThread());
diff --git a/Engine/Source/Runtime/Engine/Public/HLSLTree/HLSLTree.h b/Engine/Source/Runtime/Engine/Public/HLSLTree/HLSLTree.h
--- a/Engine/Source/Runtime/Engine/Public/HLSLTree/HLSLTree.h
+++ b/Engine/Source/Runtime/Engine/Public/HLSLTree/HLSLTree.h
@@ -507,6 +507,7 @@ private:
 	friend class FTree;
 	friend class FEmitContext;
 	friend class FExpressionForward;
+	friend class FExpressionPreviousFrameSwitch;
 	friend class FExpressionFunctionCall;
 	friend class FExpressionOperation;
 };
@@ -635,7 +636,7 @@ public:
 	const FExpression* NewConstant(const Shader::FValue& Value);
 	const FExpression* NewSwizzle(const FSwizzleParameters& Params, const FExpression* Input);
 	const FExpression* NewUnaryOp(EOperation Op, const FExpression* Input);
-	const FExpression* NewBinaryOp(EOperation Op, const FExpression* Lhs, const FExpression* Rhs);
+	ENGINE_API const FExpression* NewBinaryOp(EOperation Op, const FExpression* Lhs, const FExpression* Rhs);
 	const FExpression* NewTernaryOp(EOperation Op, const FExpression* Input0, const FExpression* Input1, const FExpression* Input2);
 
 	const FExpression* NewAbs(const FExpression* Input) { return NewUnaryOp(EOperation::Abs, Input); }
diff --git a/Engine/Source/Runtime/Engine/Public/HLSLTree/HLSLTreeCommon.h b/Engine/Source/Runtime/Engine/Public/HLSLTree/HLSLTreeCommon.h
--- a/Engine/Source/Runtime/Engine/Public/HLSLTree/HLSLTreeCommon.h
+++ b/Engine/Source/Runtime/Engine/Public/HLSLTree/HLSLTreeCommon.h
@@ -41,6 +41,22 @@ public:
 	virtual void EmitValuePreshader(FEmitContext& Context, FEmitScope& Scope, const FRequestedType& RequestedType, FEmitValuePreshaderResult& OutResult) const override;
 };
 
+/**
+ * Similar to FExpressionForward but switch to a different expression when computing previous frame
+ */
+class FExpressionPreviousFrameSwitch : public FExpressionForward
+{
+public:
+	explicit FExpressionPreviousFrameSwitch(const FExpression* InCurrentFrameExpression, const FExpression* InPreviousFrameExpression)
+		: FExpressionForward(InCurrentFrameExpression)
+		, PreviousFrameExpression(InPreviousFrameExpression)
+	{}
+
+	const FExpression* PreviousFrameExpression;
+
+	virtual const FExpression* ComputePreviousFrame(FTree& Tree, const FRequestedType& RequestedType) const override;
+};
+
 class FExpressionConstant : public FExpression
 {
 public:
@@ -281,6 +297,24 @@ public:
 	virtual bool IsInputActive(const FEmitContext& Context, int32 Index) const override;
 };
 
+class FExpressionVirtualTextureFeatureSwitch : public FExpressionSwitchBase
+{
+public:
+	FExpressionVirtualTextureFeatureSwitch(TConstArrayView<const FExpression*> InInputs);
+
+	virtual const FExpression* NewSwitch(FTree& Tree, TConstArrayView<const FExpression*> InInputs) const override { return Tree.NewExpression<FExpressionVirtualTextureFeatureSwitch>(InInputs); }
+	virtual bool IsInputActive(const FEmitContext& Context, int32 Index) const override;
+};
+
+class FExpressionDistanceFieldsRenderingSwitch : public FExpressionSwitchBase
+{
+public:
+	FExpressionDistanceFieldsRenderingSwitch(TConstArrayView<const FExpression*> InInputs);
+
+	virtual const FExpression* NewSwitch(FTree& Tree, TConstArrayView<const FExpression*> InInputs) const override { return Tree.NewExpression<FExpressionDistanceFieldsRenderingSwitch>(InInputs); }
+	virtual bool IsInputActive(const FEmitContext& Context, int32 Index) const override;
+};
+
 /** Can be used to emit HLSL chunks with no inputs, where it's not worth the trouble of defining a new expression type */
 class FExpressionInlineCustomHLSL : public FExpression
 {
diff --git a/Engine/Source/Runtime/Engine/Public/HLSLTree/HLSLTreeEmit.h b/Engine/Source/Runtime/Engine/Public/HLSLTree/HLSLTreeEmit.h
--- a/Engine/Source/Runtime/Engine/Public/HLSLTree/HLSLTreeEmit.h
+++ b/Engine/Source/Runtime/Engine/Public/HLSLTree/HLSLTreeEmit.h
@@ -400,6 +400,8 @@ public:
 
 	ENGINE_API void Finalize();
 
+	void ResetPastRequestedTypes();
+
 	ENGINE_API FEmitShaderExpression* InternalEmitExpression(FEmitScope& Scope, TArrayView<FEmitShaderNode*> Dependencies, bool bInline, const Shader::FType& Type, FStringView Code);
 
 	/**
@@ -585,6 +587,7 @@ public:
 	TArray<FEmitShaderNode*> EmitNodes;
 	TMap<const FScope*, FEmitScope*> EmitScopeMap;
 	TMap<const FExpression*, FPrepareValueResult*> PrepareValueMap;
+	TMap<FXxHash64, FRequestedType*> RequestedTypeTracker;
 	TMap<const FExpression*, FEmitScope*> PrepareLocalPHIMap;
 	TMap<const FExpression*, FEmitShaderExpression*> EmitLocalPHIMap;
 	TMap<FXxHash64, FEmitShaderExpression*> EmitExpressionMap;
diff --git a/Engine/Source/Runtime/Engine/Public/MaterialHLSLGenerator.h b/Engine/Source/Runtime/Engine/Public/MaterialHLSLGenerator.h
--- a/Engine/Source/Runtime/Engine/Public/MaterialHLSLGenerator.h
+++ b/Engine/Source/Runtime/Engine/Public/MaterialHLSLGenerator.h
@@ -75,7 +75,7 @@ public:
 	const FMaterialLayersFunctions* GetLayerOverrides() const { return LayerOverrides; }
 
 	const UMaterial* GetTargetMaterial() const;
-	UE::HLSLTree::FTree& GetTree() const;
+	ENGINE_API UE::HLSLTree::FTree& GetTree() const;
 	UE::Shader::FStructTypeRegistry& GetTypeRegistry() const;
 	const UE::Shader::FStructType* GetMaterialAttributesType() const;
 	const UE::Shader::FValue& GetMaterialAttributesDefaultValue() const;
@@ -240,7 +240,7 @@ private:
 	bool InternalGenerate();
 
 	const UE::HLSLTree::FExpression* InternalNewErrorExpression(FStringView Error);
-	bool InternalError(FStringView ErrorMessage);
+	ENGINE_API bool InternalError(FStringView ErrorMessage);
 
 	void InternalRegisterExpressionData(const FName& Type, const UMaterialExpression* MaterialExpression, void* Data);
 	void* InternalFindExpressionData(const FName& Type, const UMaterialExpression* MaterialExpression);
diff --git a/Engine/Source/Runtime/Engine/Public/MaterialHLSLTree.h b/Engine/Source/Runtime/Engine/Public/MaterialHLSLTree.h
--- a/Engine/Source/Runtime/Engine/Public/MaterialHLSLTree.h
+++ b/Engine/Source/Runtime/Engine/Public/MaterialHLSLTree.h
@@ -8,6 +8,8 @@
 #include "Materials/MaterialLayersFunctions.h"
 #include "RHIDefinitions.h"
 #include "Engine/Texture.h" // enum TextureAddress for VTStackEntry
+#include "VT/RuntimeVirtualTextureEnum.h"
+#include "MaterialCompiler.h"
 
 class UTexture;
 enum class EMaterialParameterType : uint8;
@@ -180,6 +182,11 @@ enum class EExternalInput : uint8
 	ParticleColor,
 	ParticleTranslatedWorldPosition,
 	ParticleRadius,
+	ParticleDirection,
+	ParticleSpeed,
+	ParticleRelativeTime,
+	ParticleRandom,
+	ParticleSize,
 
 	IsOrthographic,
 
@@ -327,7 +334,11 @@ public:
 		const FExpression* InAutomaticMipBiasExpression,
 		const FExpressionDerivatives& InTexCoordDerivatives,
 		ESamplerSourceMode InSamplerSource,
-		ETextureMipValueMode InMipValueMode)
+		ETextureMipValueMode InMipValueMode,
+		int16 InTextureLayerIndex = INDEX_NONE,
+		int16 InPageTableLayerIndex = INDEX_NONE,
+		bool bInAdaptive = false,
+		bool bInEnableFeedback = true)
 		: TextureExpression(InTextureExpression)
 		, TexCoordExpression(InTexCoordExpression)
 		, MipValueExpression(InMipValueExpression)
@@ -335,6 +346,10 @@ public:
 		, TexCoordDerivatives(InTexCoordDerivatives)
 		, SamplerSource(InSamplerSource)
 		, MipValueMode(InMipValueMode)
+		, TextureLayerIndex(InTextureLayerIndex)
+		, PageTableLayerIndex(InPageTableLayerIndex)
+		, bAdaptive(bInAdaptive)
+		, bEnableFeedback(bInEnableFeedback)
 	{}
 
 	const FExpression* TextureExpression;
@@ -344,6 +359,12 @@ public:
 	FExpressionDerivatives TexCoordDerivatives;
 	ESamplerSourceMode SamplerSource;
 	ETextureMipValueMode MipValueMode;
+	
+	// Only used for virtual textures
+	int16 TextureLayerIndex;
+	int16 PageTableLayerIndex;
+	bool bAdaptive;
+	bool bEnableFeedback;
 
 	virtual bool PrepareValue(FEmitContext& Context, FEmitScope& Scope, const FRequestedType& RequestedType, FPrepareValueResult& OutResult) const override;
 	virtual void EmitValueShader(FEmitContext& Context, FEmitScope& Scope, const FRequestedType& RequestedType, FEmitValueShaderResult& OutResult) const override;
@@ -362,6 +383,49 @@ public:
 	virtual void EmitValuePreshader(FEmitContext& Context, FEmitScope& Scope, const FRequestedType& RequestedType, FEmitValuePreshaderResult& OutResult) const override;
 };
 
+class FExpressionRuntimeVirtualTextureUniform : public FExpression
+{
+public:
+	FExpressionRuntimeVirtualTextureUniform(const FMaterialParameterInfo& InParameterInfo, const FExpression* InTextureExpression, ERuntimeVirtualTextureShaderUniform InUniformType)
+		: ParameterInfo(InParameterInfo)
+		, TextureExpression(InTextureExpression)
+		, UniformType(InUniformType)
+	{}
+
+	FHashedMaterialParameterInfo ParameterInfo;
+	const FExpression* TextureExpression;
+	ERuntimeVirtualTextureShaderUniform UniformType;
+
+	virtual bool PrepareValue(FEmitContext& Context, FEmitScope& Scope, const FRequestedType& RequestedType, FPrepareValueResult& OutResult) const override;
+	virtual void EmitValuePreshader(FEmitContext& Context, FEmitScope& Scope, const FRequestedType& RequestedType, FEmitValuePreshaderResult& OutResult) const override;
+};
+
+class FExpressionVirtualTextureUnpack : public FExpression
+{
+public:
+	FExpressionVirtualTextureUnpack(
+		const FExpression* InSampleLayer0Expression,
+		const FExpression* InSampleLayer1Expression,
+		const FExpression* InSampleLayer2Expression,
+		const FExpression* InWorldHeightUnpackUniformExpression,
+		EVirtualTextureUnpackType InUnpackType)
+		: SampleLayer0Expression(InSampleLayer0Expression)
+		, SampleLayer1Expression(InSampleLayer1Expression)
+		, SampleLayer2Expression(InSampleLayer2Expression)
+		, WorldHeightUnpackUniformExpression(InWorldHeightUnpackUniformExpression)
+		, UnpackType(InUnpackType)
+	{}
+
+	const FExpression* SampleLayer0Expression;
+	const FExpression* SampleLayer1Expression;
+	const FExpression* SampleLayer2Expression;
+	const FExpression* WorldHeightUnpackUniformExpression;
+	EVirtualTextureUnpackType UnpackType;
+
+	virtual bool PrepareValue(FEmitContext& Context, FEmitScope& Scope, const FRequestedType& RequestedType, FPrepareValueResult& OutResult) const override;
+	virtual void EmitValueShader(FEmitContext& Context, FEmitScope& Scope, const FRequestedType& RequestedType, FEmitValueShaderResult& OutResult) const override;
+};
+
 class FExpressionFunctionCall : public FExpressionForward
 {
 public:
@@ -405,6 +469,47 @@ public:
 	virtual void EmitValueShader(FEmitContext& Context, FEmitScope& Scope, const FRequestedType& RequestedType, FEmitValueShaderResult& OutResult) const override;
 };
 
+class FExpressionScreenAlignedUV : public FExpression
+{
+public:
+	 FExpressionScreenAlignedUV(const FExpression* InOffsetExpression, const FExpression* InViewportUVExpression)
+		: OffsetExpression(InOffsetExpression)
+		, ViewportUVExpression(InViewportUVExpression)
+	{}
+
+	 const FExpression* OffsetExpression;
+	 const FExpression* ViewportUVExpression;
+
+	 virtual bool PrepareValue(FEmitContext& Context, FEmitScope& Scope, const FRequestedType& RequestedType, FPrepareValueResult& OutResult) const override;
+	 virtual void EmitValueShader(FEmitContext& Context, FEmitScope& Scope, const FRequestedType& RequestedType, FEmitValueShaderResult& OutResult) const override;
+};
+
+class FExpressionSceneDepth : public FExpression
+{
+public:
+	FExpressionSceneDepth(const FExpression* InScreenUVExpression)
+		: ScreenUVExpression(InScreenUVExpression)
+	{}
+
+	const FExpression* ScreenUVExpression;
+
+	virtual bool PrepareValue(FEmitContext& Context, FEmitScope& Scope, const FRequestedType& RequestedType, FPrepareValueResult& OutResult) const override;
+	virtual void EmitValueShader(FEmitContext& Context, FEmitScope& Scope, const FRequestedType& RequestedType, FEmitValueShaderResult& OutResult) const override;
+};
+
+class FExpressionSceneColor : public FExpression
+{
+public:
+	FExpressionSceneColor(const FExpression* InScreenUVExpression)
+		: ScreenUVExpression(InScreenUVExpression)
+	{}
+
+	const FExpression* ScreenUVExpression;
+
+	virtual bool PrepareValue(FEmitContext& Context, FEmitScope& Scope, const FRequestedType& RequestedType, FPrepareValueResult& OutResult) const override;
+	virtual void EmitValueShader(FEmitContext& Context, FEmitScope& Scope, const FRequestedType& RequestedType, FEmitValueShaderResult& OutResult) const override;
+};
+
 struct FNoiseParameters
 {
 	FNoiseParameters() { FMemory::Memzero(*this); }
@@ -438,6 +543,34 @@ public:
 	virtual void EmitValueShader(FEmitContext& Context, FEmitScope& Scope, const FRequestedType& RequestedType, FEmitValueShaderResult& OutResult) const override;
 };
 
+struct FVectorNoiseParameters
+{
+	FVectorNoiseParameters()
+	{
+		FMemory::Memzero(*this);
+	}
+
+	int32 Quality;
+	uint32 TileSize;
+	uint8 Function;
+	bool bTiling;
+};
+
+class FExpressionVectorNoise : public FExpression
+{
+public:
+	FExpressionVectorNoise(const FVectorNoiseParameters& InParams, const FExpression* InPositionExpression)
+		: PositionExpression(InPositionExpression)
+		, Parameters(InParams)
+	{}
+
+	const FExpression* PositionExpression;
+	FVectorNoiseParameters Parameters;
+
+	virtual bool PrepareValue(FEmitContext& Context, FEmitScope& Scope, const FRequestedType& RequestedType, FPrepareValueResult& OutResult) const override;
+	virtual void EmitValueShader(FEmitContext& Context, FEmitScope& Scope, const FRequestedType& RequestedType, FEmitValueShaderResult& OutResult) const override;
+};
+
 /**
  * Mechanics of passing values from PS->VS are specific to material shaders, to VertexInterpolator is included in the Material module
  * Possible that parts of this could be moved to the common HLSLTree module, if interpolators are needed by another system
diff --git a/Engine/Source/Runtime/Engine/Public/MaterialShared.h b/Engine/Source/Runtime/Engine/Public/MaterialShared.h
--- a/Engine/Source/Runtime/Engine/Public/MaterialShared.h
+++ b/Engine/Source/Runtime/Engine/Public/MaterialShared.h
@@ -635,6 +635,7 @@ public:
 	const FMaterialVirtualTextureStack& GetVTStack(int32 Index) const { return VTStacks[Index]; }
 	int32 AddVTStack(int32 InPreallocatedStackTextureIndex);
 	int32 AddVTLayer(int32 StackIndex, int32 TextureIndex);
+	void SetVTLayer(int32 StackIndex, int32 VTLayerIndex, int32 TextureIndex);
 
 protected:
 	union FVTPackedStackAndLayerIndex
