Shader Compiler:  fix for hang when retrying shader compilation with async DDC queries enabled.  Need to release original DDC query before reissuing jobs.

#rnx

diff --git a/Engine/Source/Runtime/Engine/Private/ShaderCompiler/ShaderCompiler.cpp b/Engine/Source/Runtime/Engine/Private/ShaderCompiler/ShaderCompiler.cpp
--- a/Engine/Source/Runtime/Engine/Private/ShaderCompiler/ShaderCompiler.cpp
+++ b/Engine/Source/Runtime/Engine/Private/ShaderCompiler/ShaderCompiler.cpp
@@ -6735,6 +6735,9 @@ bool FShaderCompilingManager::HandlePotentialRetryOnError(TMap<int32, FShaderMap
 							CommonJob->bFinalized = false;
 						}
 					}
+					
+					// Reset DDC query request owner
+					CurrentJob.RequestOwner.Reset();
 					CurrentJob.bFinalized = false;
 				}
 
