Truncation warning fix for FAnimationStreamingManager in the Engine module.

#rnx

- The memory buffer returned by IBulkDataIORequest has a size type of int64 but TArrayView will truncate that to int32 and potentially incorrectly resize the buffer.
- Switched to use FMemoryView which supports 64bit types as FMemoryReaderView does not accept TArrayView64.
- Fixed some truncation/casting warnings related to time (double to float)

diff --git a/Engine/Source/Runtime/Engine/Private/Animation/AnimationStreaming.cpp b/Engine/Source/Runtime/Engine/Private/Animation/AnimationStreaming.cpp
--- a/Engine/Source/Runtime/Engine/Private/Animation/AnimationStreaming.cpp
+++ b/Engine/Source/Runtime/Engine/Private/Animation/AnimationStreaming.cpp
@@ -245,12 +245,12 @@ bool FStreamingAnimationData::BlockTillAllRequestsFinished(float TimeLimit)
 	}
 	else
 	{
-		double EndTime = FPlatformTime::Seconds() + TimeLimit;
+		const double EndTime = FPlatformTime::Seconds() + TimeLimit;
 		for (FLoadedAnimationChunk& LoadedChunk : LoadedChunks)
 		{
 			if (LoadedChunk.IORequest)
 			{
-				float ThisTimeLimit = EndTime - FPlatformTime::Seconds();
+				const float ThisTimeLimit = static_cast<float>(EndTime - FPlatformTime::Seconds());
 				if (ThisTimeLimit < .001f || // one ms is the granularity of the platform event system
 					!LoadedChunk.IORequest->WaitCompletion(ThisTimeLimit))
 				{
@@ -347,7 +347,7 @@ void FAnimationStreamingManager::OnAsyncFileCallback(FStreamingAnimationData* St
 
 		FCompressedAnimSequence* NewCompressedData = new FCompressedAnimSequence();
 
-		TArrayView<const uint8> MemView(Mem, ReadSize);
+		FMemoryView MemView(Mem, ReadSize);
 		FMemoryReaderView Reader(MemView);
 
 		UAnimStreamable* Anim = StreamingAnimData->StreamableAnim;
@@ -420,10 +420,10 @@ int32 FAnimationStreamingManager::BlockTillAllRequestsFinished(float TimeLimit,
 		}
 		else
 		{
-			double EndTime = FPlatformTime::Seconds() + TimeLimit;
+			const double EndTime = FPlatformTime::Seconds() + TimeLimit;
 			for (TPair<UAnimStreamable*, FStreamingAnimationData*>& AnimPair : StreamingAnimations)
 			{
-				float ThisTimeLimit = EndTime - FPlatformTime::Seconds();
+				const float ThisTimeLimit = static_cast<float>(EndTime - FPlatformTime::Seconds());
 				if (ThisTimeLimit < .001f || // one ms is the granularity of the platform event system
 					!AnimPair.Value->BlockTillAllRequestsFinished(ThisTimeLimit))
 				{
