Adapted FLevelUtils::IsValidStreamingLevel to properly search in all registered WorldPartition.

diff --git a/Engine/Source/Runtime/Engine/Private/LevelUtils.cpp b/Engine/Source/Runtime/Engine/Private/LevelUtils.cpp
--- a/Engine/Source/Runtime/Engine/Private/LevelUtils.cpp
+++ b/Engine/Source/Runtime/Engine/Private/LevelUtils.cpp
@@ -17,6 +17,7 @@
 #if WITH_EDITOR
 #include "ScopedTransaction.h"
 #include "WorldPartition/WorldPartition.h"
+#include "WorldPartition/WorldPartitionSubsystem.h"
 #endif
 
 #define LOCTEXT_NAMESPACE "LevelUtils"
@@ -108,9 +109,15 @@ bool FLevelUtils::IsValidStreamingLevel(UWorld* InWorld, const TCHAR* InPackageN
 	}
 
 #if WITH_EDITOR
-	if (UWorldPartition* WorldPartition = InWorld ? InWorld->GetWorldPartition() : nullptr)
+	if (UWorldPartitionSubsystem* WorldPartitionSubsystem = InWorld ? InWorld->GetSubsystem<UWorldPartitionSubsystem>() : nullptr)
 	{
-		return WorldPartition->IsValidPackageName(InPackageName);
+		bool bIsValidStreamingLevel = false;
+		WorldPartitionSubsystem->ForEachWorldPartition([&bIsValidStreamingLevel, InPackageName](UWorldPartition* WorldPartition)
+		{
+			bIsValidStreamingLevel = WorldPartition->IsValidPackageName(InPackageName);
+			return !bIsValidStreamingLevel;
+		});
+		return bIsValidStreamingLevel;
 	}
 #endif
 
