Added bCastHiddenShadow property to FISMComponentDescriptor

diff --git a/Engine/Source/Runtime/Engine/Private/ISMPartition/ISMComponentDescriptor.cpp b/Engine/Source/Runtime/Engine/Private/ISMPartition/ISMComponentDescriptor.cpp
--- a/Engine/Source/Runtime/Engine/Private/ISMPartition/ISMComponentDescriptor.cpp
+++ b/Engine/Source/Runtime/Engine/Private/ISMPartition/ISMComponentDescriptor.cpp
@@ -94,6 +94,7 @@ void FISMComponentDescriptorBase::InitFrom(const UStaticMeshComponent* Template,
 	bCastDynamicShadow = Template->bCastDynamicShadow;
 	bCastContactShadow = Template->bCastContactShadow;
 	bCastShadowAsTwoSided = Template->bCastShadowAsTwoSided;
+	bCastHiddenShadow = Template->bCastHiddenShadow;
 	bAffectDynamicIndirectLighting = Template->bAffectDynamicIndirectLighting;
 	bAffectDynamicIndirectLightingWhileHidden = Template->bAffectIndirectLightingWhileHidden;
 	bAffectDistanceFieldLighting = Template->bAffectDistanceFieldLighting;
@@ -199,6 +200,7 @@ bool FISMComponentDescriptorBase::operator==(const FISMComponentDescriptorBase&
 	bCastDynamicShadow == Other.bCastDynamicShadow &&
 	bCastContactShadow == Other.bCastContactShadow &&
 	bCastShadowAsTwoSided == Other.bCastShadowAsTwoSided &&
+	bCastHiddenShadow == Other.bCastHiddenShadow &&
 	bAffectDynamicIndirectLighting == Other.bAffectDynamicIndirectLighting &&
 	bAffectDynamicIndirectLightingWhileHidden == Other.bAffectDynamicIndirectLightingWhileHidden &&
 	bAffectDistanceFieldLighting == Other.bAffectDistanceFieldLighting &&
@@ -313,6 +315,7 @@ void FISMComponentDescriptorBase::InitComponent(UInstancedStaticMeshComponent* I
 	ISMComponent->bCastDynamicShadow = bCastDynamicShadow;
 	ISMComponent->bCastContactShadow = bCastContactShadow;
 	ISMComponent->bCastShadowAsTwoSided = bCastShadowAsTwoSided;
+	ISMComponent->bCastHiddenShadow = bCastHiddenShadow;
 	ISMComponent->bAffectDynamicIndirectLighting = bAffectDynamicIndirectLighting;
 	ISMComponent->bAffectIndirectLightingWhileHidden = bAffectDynamicIndirectLightingWhileHidden;
 	ISMComponent->bAffectDistanceFieldLighting = bAffectDistanceFieldLighting;
diff --git a/Engine/Source/Runtime/Engine/Public/ISMPartition/ISMComponentDescriptor.h b/Engine/Source/Runtime/Engine/Public/ISMPartition/ISMComponentDescriptor.h
--- a/Engine/Source/Runtime/Engine/Public/ISMPartition/ISMComponentDescriptor.h
+++ b/Engine/Source/Runtime/Engine/Public/ISMPartition/ISMComponentDescriptor.h
@@ -107,6 +107,9 @@ public:
 	UPROPERTY(EditAnywhere, Category = "Component Settings")
 	uint8 bCastShadowAsTwoSided : 1;
 
+	UPROPERTY(EditAnywhere, Category = "Component Settings")
+	uint8 bCastHiddenShadow : 1;
+
 	UPROPERTY(EditAnywhere, Category = "Component Settings")
 	uint8 bAffectDynamicIndirectLighting : 1;
 
