[shaders] fix retry case when preprocessed cache is enabled (need to reset preprocess output struct)

diff --git a/Engine/Source/Runtime/Engine/Private/ShaderCompiler/ShaderCompiler.cpp b/Engine/Source/Runtime/Engine/Private/ShaderCompiler/ShaderCompiler.cpp
--- a/Engine/Source/Runtime/Engine/Private/ShaderCompiler/ShaderCompiler.cpp
+++ b/Engine/Source/Runtime/Engine/Private/ShaderCompiler/ShaderCompiler.cpp
@@ -6725,6 +6725,7 @@ bool FShaderCompilingManager::HandlePotentialRetryOnError(TMap<int32, FShaderMap
 							SingleJob->Input.DumpDebugInfoPath = GShaderCompilingManager->CreateShaderDebugInfoPath(SingleJob->Input);
 						}
 						SingleJob->Output = FShaderCompilerOutput();
+						SingleJob->PreprocessOutput = FShaderPreprocessOutput();
 					}
 					else
 					{
@@ -6732,6 +6733,7 @@ bool FShaderCompilingManager::HandlePotentialRetryOnError(TMap<int32, FShaderMap
 						for (auto CommonJob : PipelineJob->StageJobs)
 						{
 							CommonJob->Output = FShaderCompilerOutput();
+							CommonJob->PreprocessOutput = FShaderPreprocessOutput();
 							CommonJob->bFinalized = false;
 						}
 					}
