Fixed wp.Runtime.OverrideRuntimeSpatialHashLoadingRange to properly apply the value when a new instanced world partition is loaded
[RN] patrick.enfedaque

diff --git a/Engine/Source/Runtime/Engine/Private/WorldPartition/WorldPartitionRuntimeSpatialHash.cpp b/Engine/Source/Runtime/Engine/Private/WorldPartition/WorldPartitionRuntimeSpatialHash.cpp
--- a/Engine/Source/Runtime/Engine/Private/WorldPartition/WorldPartitionRuntimeSpatialHash.cpp
+++ b/Engine/Source/Runtime/Engine/Private/WorldPartition/WorldPartitionRuntimeSpatialHash.cpp
@@ -106,7 +106,7 @@ FSpatialHashStreamingGrid::FSpatialHashStreamingGrid()
 	, WorldBounds(ForceInitToZero)
 	, bClientOnlyVisible(false)
 	, HLODLayer(nullptr)
-	, OverrideLoadingRange(-1.f)
+	, GridIndex(INDEX_NONE)
 	, GridHelper(nullptr)
 {
 }
@@ -134,6 +134,65 @@ const FSquare2DGridHelper& FSpatialHashStreamingGrid::GetGridHelper() const
 	return *GridHelper;
 }
 
+static const FString GOverrideLoadingRangeCommandName(TEXT("wp.Runtime.OverrideRuntimeSpatialHashLoadingRange"));
+bool FSpatialHashStreamingGrid::bAddedWorldPartitionSubsystemDeinitializedCallback = false;
+TMap<int32, float> FSpatialHashStreamingGrid::OverriddenLoadingRanges;
+FAutoConsoleCommand FSpatialHashStreamingGrid::OverrideLoadingRangeCommand(
+	*GOverrideLoadingRangeCommandName,
+	TEXT("Sets runtime loading range. Args -grid=[index] -range=[override_loading_range]"),
+	FConsoleCommandWithArgsDelegate::CreateLambda([](const TArray<FString>& InArgs)
+	{
+		FString ArgString = FString::Join(InArgs, TEXT(" "));
+		int32 OverrideGridIndex = 0;
+		float OverrideLoadingRange = -1.f;
+		FParse::Value(*ArgString, TEXT("grid="), OverrideGridIndex);
+		FParse::Value(*ArgString, TEXT("range="), OverrideLoadingRange);
+
+		if (!bAddedWorldPartitionSubsystemDeinitializedCallback)
+		{
+			bAddedWorldPartitionSubsystemDeinitializedCallback = true;
+			UWorldPartitionSubsystem::OnWorldPartitionSubsystemDeinitialized.AddLambda([](UWorldPartitionSubsystem* InWorldPartitionSubsystem, UWorld* InWorld)
+			{
+				if (InWorld && InWorld->IsGameWorld())
+				{
+					FSpatialHashStreamingGrid::OverriddenLoadingRanges.Reset();
+				}
+			});
+		}
+
+		for (const FWorldContext& Context : GEngine->GetWorldContexts())
+		{
+			UWorld* World = Context.World();
+			if (World && World->IsGameWorld())
+			{
+				FWorldPartitionHelpers::ServerExecConsoleCommand(World, GOverrideLoadingRangeCommandName, InArgs);
+				if (UWorld::HasSubsystem<UWorldPartitionSubsystem>(World))
+				{
+					if (OverrideLoadingRange >= 0.f)
+					{
+						FSpatialHashStreamingGrid::OverriddenLoadingRanges.Add(OverrideGridIndex, OverrideLoadingRange);
+					}
+					else
+					{
+						FSpatialHashStreamingGrid::OverriddenLoadingRanges.Remove(OverrideGridIndex);
+					}
+					break;
+				}
+			}
+		}
+	})
+);
+
+float FSpatialHashStreamingGrid::GetLoadingRange() const
+{
+	if (float* OverrideLoadingRange = FSpatialHashStreamingGrid::OverriddenLoadingRanges.Find(GridIndex))
+	{
+		check(*OverrideLoadingRange >= 0.f);
+		return *OverrideLoadingRange;
+	}
+	return LoadingRange;
+}
+
 int64 FSpatialHashStreamingGrid::GetCellSize(int32 Level) const
 {
 	return GetGridHelper().Levels[Level].CellSize;
@@ -1425,6 +1484,7 @@ bool UWorldPartitionRuntimeSpatialHash::CreateStreamingGrid(const FSpatialHashRu
 	CurrentStreamingGrid.DebugColor = RuntimeGrid.DebugColor;
 	CurrentStreamingGrid.bClientOnlyVisible = RuntimeGrid.bClientOnlyVisible;
 	CurrentStreamingGrid.HLODLayer = RuntimeGrid.HLODLayer;
+	CurrentStreamingGrid.GridIndex = (StreamingGrids.Num() - 1);
 
 	// Move actors into the final streaming grids
 	CurrentStreamingGrid.GridLevels.Reserve(PartionedActors.Levels.Num());
@@ -1531,43 +1591,6 @@ bool UWorldPartitionRuntimeSpatialHash::IsValidGrid(FName GridName) const
 
 #endif //WITH_EDITOR
 
-static const FString GOverrideLoadingRangeCommandName(TEXT("wp.Runtime.OverrideRuntimeSpatialHashLoadingRange"));
-FAutoConsoleCommand UWorldPartitionRuntimeSpatialHash::OverrideLoadingRangeCommand(
-	*GOverrideLoadingRangeCommandName,
-	TEXT("Sets runtime loading range. Args -grid=[index] -range=[override_loading_range]"),
-	FConsoleCommandWithArgsDelegate::CreateLambda([](const TArray<FString>& InArgs)
-	{
-		FString ArgString = FString::Join(InArgs, TEXT(" "));
-		int32 GridIndex = 0;
-		float OverrideLoadingRange = -1.f;
-		FParse::Value(*ArgString, TEXT("grid="), GridIndex);
-		FParse::Value(*ArgString, TEXT("range="), OverrideLoadingRange);
-
-		for (const FWorldContext& Context : GEngine->GetWorldContexts())
-		{
-			UWorld* World = Context.World();
-			if (World && World->IsGameWorld())
-			{
-				FWorldPartitionHelpers::ServerExecConsoleCommand(World, GOverrideLoadingRangeCommandName, InArgs);
-				if (UWorldPartitionSubsystem* WorldPartitionSubsystem = World->GetSubsystem<UWorldPartitionSubsystem>())
-				{
-					WorldPartitionSubsystem->ForEachWorldPartition([GridIndex, OverrideLoadingRange](UWorldPartition* WorldPartition)
-					{
-						if (UWorldPartitionRuntimeSpatialHash* RuntimeSpatialHash = Cast<UWorldPartitionRuntimeSpatialHash>(WorldPartition->RuntimeHash))
-						{
-							if (RuntimeSpatialHash->StreamingGrids.IsValidIndex(GridIndex))
-							{
-								RuntimeSpatialHash->StreamingGrids[GridIndex].OverrideLoadingRange = OverrideLoadingRange;
-							}
-						}
-						return true;
-					});
-				}
-			}
-		}
-	})
-);
-
 // Streaming interface
 void UWorldPartitionRuntimeSpatialHash::ForEachStreamingCells(TFunctionRef<bool(const UWorldPartitionRuntimeCell*)> Func) const
 {
diff --git a/Engine/Source/Runtime/Engine/Private/WorldPartition/WorldPartitionSubsystem.cpp b/Engine/Source/Runtime/Engine/Private/WorldPartition/WorldPartitionSubsystem.cpp
--- a/Engine/Source/Runtime/Engine/Private/WorldPartition/WorldPartitionSubsystem.cpp
+++ b/Engine/Source/Runtime/Engine/Private/WorldPartition/WorldPartitionSubsystem.cpp
@@ -162,6 +162,9 @@ static FAutoConsoleVariableRef CVarUdateStreamingStateTimeLimit(
 	ECVF_Default
 );
 
+TMulticastDelegate<void(UWorldPartitionSubsystem*, UWorld*)> UWorldPartitionSubsystem::OnWorldPartitionSubsystemInitialized;
+TMulticastDelegate<void(UWorldPartitionSubsystem*, UWorld*)> UWorldPartitionSubsystem::OnWorldPartitionSubsystemDeinitialized;
+
 UWorldPartitionSubsystem::UWorldPartitionSubsystem()
 : StreamingSourcesHash(0)
 , NumWorldPartitionServerStreamingEnabled(0)
@@ -468,6 +471,8 @@ void UWorldPartitionSubsystem::Initialize(FSubsystemCollectionBase& Collection)
 		FLevelStreamingDelegates::OnLevelStreamingTargetStateChanged.AddUObject(this, &UWorldPartitionSubsystem::OnLevelStreamingTargetStateChanged);
 		FLevelStreamingDelegates::OnLevelStreamingStateChanged.AddUObject(this, &UWorldPartitionSubsystem::OnLevelStreamingStateChanged);
 	}
+
+	OnWorldPartitionSubsystemInitialized.Broadcast(this, GetWorld());
 }
 
 void UWorldPartitionSubsystem::Deinitialize()
@@ -478,7 +483,9 @@ void UWorldPartitionSubsystem::Deinitialize()
 		Super::Deinitialize();
 		return;
 	}
-#endif 
+#endif
+
+	OnWorldPartitionSubsystemDeinitialized.Broadcast(this, GetWorld());
 
 	GetWorld()->OnWorldPartitionInitialized().RemoveAll(this);
 	GetWorld()->OnWorldPartitionUninitialized().RemoveAll(this);
diff --git a/Engine/Source/Runtime/Engine/Public/WorldPartition/WorldPartitionRuntimeSpatialHash.h b/Engine/Source/Runtime/Engine/Public/WorldPartition/WorldPartitionRuntimeSpatialHash.h
--- a/Engine/Source/Runtime/Engine/Public/WorldPartition/WorldPartitionRuntimeSpatialHash.h
+++ b/Engine/Source/Runtime/Engine/Public/WorldPartition/WorldPartitionRuntimeSpatialHash.h
@@ -134,16 +134,8 @@ struct FSpatialHashStreamingGrid
 	UPROPERTY()
 	TObjectPtr<const UHLODLayer> HLODLayer;
 
-	float OverrideLoadingRange;
-
-	float GetLoadingRange() const
-	{
-		if (OverrideLoadingRange >= 0.f)
-		{
-			return OverrideLoadingRange;
-		}
-		return LoadingRange;
-	}
+	UPROPERTY()
+	int32 GridIndex;
 
 	ENGINE_API void InjectExternalStreamingObjectGrid(const FSpatialHashStreamingGrid& InExternalObjectStreamingGrid) const;
 	ENGINE_API void RemoveExternalStreamingObjectGrid(const FSpatialHashStreamingGrid& InExternalObjectStreamingGrid) const;
@@ -159,6 +151,7 @@ struct FSpatialHashStreamingGrid
 	ENGINE_API void Draw3D(const class UWorldPartitionRuntimeSpatialHash* Owner, const TArray<FWorldPartitionStreamingSource>& Sources, const FTransform& Transform) const;
 	ENGINE_API void ForEachRuntimeCell(TFunctionRef<bool(const UWorldPartitionRuntimeCell*)> Func) const;
 	ENGINE_API const FSquare2DGridHelper& GetGridHelper() const;
+	ENGINE_API float GetLoadingRange() const;
 
 private:
 	ENGINE_API void ForEachRuntimeCell(const FGridCellCoord& Coords, TFunctionRef<void(const UWorldPartitionRuntimeCell*)> Func) const;
@@ -172,6 +165,11 @@ private:
 	// Contains cells injected at runtime from content bundles
 	UPROPERTY(Transient)
 	mutable TArray<FSpatialHashStreamingGridLevel> InjectedGridLevels;
+
+	/** Console command used to change loading range for a given streaming grid */
+	static class FAutoConsoleCommand OverrideLoadingRangeCommand;
+	static TMap<int32, float> OverriddenLoadingRanges;
+	static bool bAddedWorldPartitionSubsystemDeinitializedCallback;
 };
 
 // Disable copy of this USTRUCT by the autogenerated code.
@@ -322,9 +320,6 @@ protected:
 #endif
 
 private:
-	/** Console command used to change loading range for a given streaming grid */
-	static ENGINE_API class FAutoConsoleCommand OverrideLoadingRangeCommand;
-
 #if WITH_EDITORONLY_DATA
 	UPROPERTY(EditAnywhere, Config, Category = RuntimeSettings)
 	TArray<FSpatialHashRuntimeGrid> Grids;
diff --git a/Engine/Source/Runtime/Engine/Public/WorldPartition/WorldPartitionSubsystem.h b/Engine/Source/Runtime/Engine/Public/WorldPartition/WorldPartitionSubsystem.h
--- a/Engine/Source/Runtime/Engine/Public/WorldPartition/WorldPartitionSubsystem.h
+++ b/Engine/Source/Runtime/Engine/Public/WorldPartition/WorldPartitionSubsystem.h
@@ -102,6 +102,9 @@ public:
 	DECLARE_DELEGATE_RetVal_OneParam(bool, FWorldPartitionStreamingSourceProviderFilter, const IWorldPartitionStreamingSourceProvider*);
 	FWorldPartitionStreamingSourceProviderFilter& OnIsStreamingSourceProviderFiltered() { return IsStreamingSourceProviderFiltered; }
 
+	static ENGINE_API TMulticastDelegate<void(UWorldPartitionSubsystem*, UWorld*)> OnWorldPartitionSubsystemInitialized;
+	static ENGINE_API TMulticastDelegate<void(UWorldPartitionSubsystem*, UWorld*)> OnWorldPartitionSubsystemDeinitialized;
+
 	ENGINE_API void ForEachWorldPartition(TFunctionRef<bool(UWorldPartition*)> Func);
 
 #if WITH_EDITOR
