Data Layers: don't update the CDO when changing the per-project data layers operator, since it's going to be delta-serialized and affect all maps.
Just set the default when creating a new map from an empty template, as templates could have data layers set up.

diff --git a/Engine/Source/Runtime/Engine/Private/WorldPartition/WorldPartition.cpp b/Engine/Source/Runtime/Engine/Private/WorldPartition/WorldPartition.cpp
--- a/Engine/Source/Runtime/Engine/Private/WorldPartition/WorldPartition.cpp
+++ b/Engine/Source/Runtime/Engine/Private/WorldPartition/WorldPartition.cpp
@@ -295,7 +295,7 @@ UWorldPartition::UWorldPartition(const FObjectInitializer& ObjectInitializer)
 	bEnableStreaming = true;
 	ServerStreamingMode = EWorldPartitionServerStreamingMode::ProjectDefault;
 	ServerStreamingOutMode = EWorldPartitionServerStreamingOutMode::ProjectDefault;
-	DataLayersLogicOperator = UWorldPartitionSettings::Get()->GetNewMapsDataLayersLogicOperator();
+	DataLayersLogicOperator = EWorldPartitionDataLayersLogicOperator::Or;
 	StreamingStateEpoch = 0;
 
 #if WITH_EDITOR
@@ -1064,6 +1064,8 @@ UWorldPartition* UWorldPartition::CreateOrRepairWorldPartition(AWorldSettings* W
 		}
 
 		FWorldPartitionMiniMapHelper::GetWorldPartitionMiniMap(OuterWorld, true);
+
+		WorldPartition->DataLayersLogicOperator = UWorldPartitionSettings::Get()->GetNewMapsDataLayersLogicOperator();
 	}
 
 	if (!WorldPartition->EditorHash)
