Editor: ULevelInstanceSubsystem::UpdateStreamingState will complete all loading in a Editor World (no need to call it multiple times even for recursive loading)

#rnx

diff --git a/Engine/Source/Runtime/Engine/Private/LevelInstance/LevelInstanceSubsystem.cpp b/Engine/Source/Runtime/Engine/Private/LevelInstance/LevelInstanceSubsystem.cpp
--- a/Engine/Source/Runtime/Engine/Private/LevelInstance/LevelInstanceSubsystem.cpp
+++ b/Engine/Source/Runtime/Engine/Private/LevelInstance/LevelInstanceSubsystem.cpp
@@ -229,6 +229,22 @@ bool ULevelInstanceSubsystem::IsLoading(const ILevelInstanceInterface* LevelInst
 }
 
 void ULevelInstanceSubsystem::UpdateStreamingState()
+{
+	UpdateStreamingStateInternal();
+
+#if WITH_EDITOR
+	if (!GetWorld()->IsGameWorld())
+	{
+		// For Editor Worlds make sure UpdateStreamingState completes all recursive loading/unloading
+		while (LevelInstancesToLoadOrUpdate.Num() || LevelInstancesToUnload.Num())
+		{
+			UpdateStreamingStateInternal();
+		}
+	}
+#endif
+}
+
+void ULevelInstanceSubsystem::UpdateStreamingStateInternal()
 {
 	if (!LevelInstancesToUnload.Num() && !LevelInstancesToLoadOrUpdate.Num())
 	{
@@ -242,7 +258,7 @@ void ULevelInstanceSubsystem::UpdateStreamingState()
 		return;
 	}
 
-	FScopedSlowTask SlowTask(LevelInstancesToUnload.Num() + LevelInstancesToLoadOrUpdate.Num() * 2, LOCTEXT("UpdatingLevelInstances", "Updating Level Instances..."), !GetWorld()->IsGameWorld());
+	FScopedSlowTask SlowTask(LevelInstancesToUnload.Num() + LevelInstancesToLoadOrUpdate.Num() * 2, LOCTEXT("UpdatingLevelInstances", "Updating Level Instances..."), !GetWorld()->IsGameWorld() && !GetWorld()->GetIsInBlockTillLevelStreamingCompleted());
 	SlowTask.MakeDialogDelayed(1.0f);
 
 	check(!LevelsToRemoveScope);
diff --git a/Engine/Source/Runtime/Engine/Public/LevelInstance/LevelInstanceSubsystem.h b/Engine/Source/Runtime/Engine/Public/LevelInstance/LevelInstanceSubsystem.h
--- a/Engine/Source/Runtime/Engine/Public/LevelInstance/LevelInstanceSubsystem.h
+++ b/Engine/Source/Runtime/Engine/Public/LevelInstance/LevelInstanceSubsystem.h
@@ -136,6 +136,7 @@ public:
 #endif
 
 private:
+	ENGINE_API void UpdateStreamingStateInternal();
 	ENGINE_API void BlockOnLoading();
 	ENGINE_API void LoadLevelInstance(ILevelInstanceInterface* LevelInstance);
 	ENGINE_API void UnloadLevelInstance(const FLevelInstanceID& LevelInstanceID);
