Fixed crash in CPU skinning VF when skeletal mesh pipe is enabled.

diff --git a/Engine/Source/Runtime/Engine/Private/SkeletalRenderCPUSkin.cpp b/Engine/Source/Runtime/Engine/Private/SkeletalRenderCPUSkin.cpp
--- a/Engine/Source/Runtime/Engine/Private/SkeletalRenderCPUSkin.cpp
+++ b/Engine/Source/Runtime/Engine/Private/SkeletalRenderCPUSkin.cpp
@@ -354,8 +354,8 @@ void FSkeletalMeshObjectCPUSkin::FSkeletalMeshObjectLOD::InitResources(FSkelMesh
 	FSkeletalMeshObjectCPUSkin::FSkeletalMeshObjectLOD* Self = this;
 	FLocalVertexFactory* VertexFactoryPtr = &VertexFactory;
 	// update vertex factory components and sync it
-	ENQUEUE_RENDER_COMMAND(InitSkeletalMeshCPUSkinVertexFactory)(
-		[VertexFactoryPtr, Self](FRHICommandListImmediate& RHICmdList)
+	ENQUEUE_RENDER_COMMAND(InitSkeletalMeshCPUSkinVertexFactory)(UE::RenderCommandPipe::SkeletalMesh,
+		[VertexFactoryPtr, Self](FRHICommandListBase& RHICmdList)
 		{
 			FLocalVertexFactory::FDataType Data;
 
@@ -387,8 +387,8 @@ void FSkeletalMeshObjectCPUSkin::FSkeletalMeshObjectLOD::InitResources(FSkelMesh
 		}
 
 		TArray<FSkelMeshRenderSection>* RenderSections = &LODModel.RenderSections;
-		ENQUEUE_RENDER_COMMAND(InitSkeletalRenderCPUSkinRayTracingGeometry)(
-			[this, VertexBufferRHI, IndexBufferRHI, VertexBufferStride, TrianglesCount, RenderSections, &SourceGeometry = LODModel.SourceRayTracingGeometry](FRHICommandListImmediate& RHICmdList)
+		ENQUEUE_RENDER_COMMAND(InitSkeletalRenderCPUSkinRayTracingGeometry)(UE::RenderCommandPipe::SkeletalMesh,
+			[this, VertexBufferRHI, IndexBufferRHI, VertexBufferStride, TrianglesCount, RenderSections, &SourceGeometry = LODModel.SourceRayTracingGeometry](FRHICommandListBase& RHICmdList)
 			{
 				FRayTracingGeometryInitializer Initializer;
 				static const FName DebugName("FSkeletalMeshObjectCPUSkin");
