Cameras: rename camera shake API parameter structs to make room for new ones

This will make it clearer what structs are meant for what, avoiding name collisions or confusing similarities.

diff --git a/Engine/Config/BaseEngine.ini b/Engine/Config/BaseEngine.ini
--- a/Engine/Config/BaseEngine.ini
+++ b/Engine/Config/BaseEngine.ini
@@ -1166,6 +1166,13 @@ DefaultWorldPartitionHLODLayer=/Engine/Maps/Templates/HLODs/HLODLayer_Instanced.
 +ClassRedirects=(OldName="XRDeviceVisualizationComponent", NewName="/Script/XRBase.XRDeviceVisualizationComponent")
 +ClassRedirects=(OldName="XRLoadingScreenFunctionLibrary", NewName="/Script/XRBase.XRLoadingScreenFunctionLibrary")
 
+; Renamed camera shake parameter structs.
++EnumRedirects=(OldName="ECameraShakeUpdateResultFlags",NewName="/Script/Engine.ECameraShakePatternUpdateResultFlags")
++StructRedirects=(OldName="CameraShakeStartParams",NewName="/Script/Engine.CameraShakePatternStartParams")
++StructRedirects=(OldName="CameraShakeUpdateParams",NewName="/Script/Engine.CameraShakePatternUpdateParams")
++StructRedirects=(OldName="CameraShakeScrubParams",NewName="/Script/Engine.CameraShakePatternScrubParams")
++StructRedirects=(OldName="CameraShakeStopParams",NewName="/Script/Engine.CameraShakePatternStopParams")
++StructRedirects=(OldName="CameraShakeUpdateResult",NewName="/Script/Engine.CameraShakePatternUpdateResult")
 
 ; Moving some MaterialX Material Expressions into Engine
 +ClassRedirects=(OldName="/Script/InterchangeImport.MaterialExpressionMaterialXExponential", NewName="/Script/Engine.MaterialExpressionExponential")
diff --git a/Engine/Plugins/Cameras/GameplayCameras/Source/GameplayCameras/Private/CompositeCameraShakePattern.cpp b/Engine/Plugins/Cameras/GameplayCameras/Source/GameplayCameras/Private/CompositeCameraShakePattern.cpp
--- a/Engine/Plugins/Cameras/GameplayCameras/Source/GameplayCameras/Private/CompositeCameraShakePattern.cpp
+++ b/Engine/Plugins/Cameras/GameplayCameras/Source/GameplayCameras/Private/CompositeCameraShakePattern.cpp
@@ -36,7 +36,7 @@ void UCompositeCameraShakePattern::GetShakePatternInfoImpl(FCameraShakeInfo& Out
 	}
 }
 
-void UCompositeCameraShakePattern::StartShakePatternImpl(const FCameraShakeStartParams& Params)
+void UCompositeCameraShakePattern::StartShakePatternImpl(const FCameraShakePatternStartParams& Params)
 {
 	for (UCameraShakePattern* Pattern : ChildPatterns)
 	{
@@ -48,19 +48,19 @@ void UCompositeCameraShakePattern::StartShakePatternImpl(const FCameraShakeStart
 	}
 }
 
-void UCompositeCameraShakePattern::UpdateShakePatternImpl(const FCameraShakeUpdateParams& Params, FCameraShakeUpdateResult& OutResult)
+void UCompositeCameraShakePattern::UpdateShakePatternImpl(const FCameraShakePatternUpdateParams& Params, FCameraShakePatternUpdateResult& OutResult)
 {
 	UCameraShakeBase* ShakeInstance = GetShakeInstance();
 	checkf(ShakeInstance, TEXT("Running a shake pattern without an outer shake instance"));
 
-	FCameraShakeUpdateParams ChildParams(Params);
+	FCameraShakePatternUpdateParams ChildParams(Params);
 
 	for (UCameraShakePattern* Pattern : ChildPatterns)
 	{
 		if (Pattern != nullptr && !Pattern->IsFinished())
 		{
 			// Let the child pattern run on the current result, with its own blending weight.
-			FCameraShakeUpdateResult ChildResult;
+			FCameraShakePatternUpdateResult ChildResult;
 			Pattern->UpdateShakePattern(ChildParams, ChildResult);
 
 			if (!Pattern->IsFinished())
@@ -83,10 +83,10 @@ void UCompositeCameraShakePattern::UpdateShakePatternImpl(const FCameraShakeUpda
 	OutResult.Location = ChildParams.POV.Location;
 	OutResult.Rotation = ChildParams.POV.Rotation;
 	OutResult.FOV = ChildParams.POV.FOV;
-	OutResult.Flags = ECameraShakeUpdateResultFlags::ApplyAsAbsolute;
+	OutResult.Flags = ECameraShakePatternUpdateResultFlags::ApplyAsAbsolute;
 }
 
-void UCompositeCameraShakePattern::ScrubShakePatternImpl(const FCameraShakeScrubParams& Params, FCameraShakeUpdateResult& OutResult)
+void UCompositeCameraShakePattern::ScrubShakePatternImpl(const FCameraShakePatternScrubParams& Params, FCameraShakePatternUpdateResult& OutResult)
 {
 	// This method is similar to the UpdateShakePatternImpl method above, but calls the scrub methods
 	// instead of the update methods.
@@ -94,14 +94,14 @@ void UCompositeCameraShakePattern::ScrubShakePatternImpl(const FCameraShakeScrub
 	UCameraShakeBase* ShakeInstance = GetShakeInstance();
 	checkf(ShakeInstance, TEXT("Running a shake pattern without an outer shake instance"));
 
-	FCameraShakeScrubParams ChildParams(Params);
+	FCameraShakePatternScrubParams ChildParams(Params);
 
 	for (UCameraShakePattern* Pattern : ChildPatterns)
 	{
 		if (Pattern != nullptr) // Don't check for IsFinished here, we might scrub anywhere.
 		{
 			// Let the child pattern run on the current result, with its own blending weight.
-			FCameraShakeUpdateResult ChildResult;
+			FCameraShakePatternUpdateResult ChildResult;
 			Pattern->ScrubShakePattern(Params, ChildResult);
 
 			if (!Pattern->IsFinished())
@@ -123,7 +123,7 @@ void UCompositeCameraShakePattern::ScrubShakePatternImpl(const FCameraShakeScrub
 	OutResult.Location = ChildParams.POV.Location;
 	OutResult.Rotation = ChildParams.POV.Rotation;
 	OutResult.FOV = ChildParams.POV.FOV;
-	OutResult.Flags = ECameraShakeUpdateResultFlags::ApplyAsAbsolute;
+	OutResult.Flags = ECameraShakePatternUpdateResultFlags::ApplyAsAbsolute;
 }
 
 bool UCompositeCameraShakePattern::IsFinishedImpl() const
@@ -139,7 +139,7 @@ bool UCompositeCameraShakePattern::IsFinishedImpl() const
 	return true;
 }
 
-void UCompositeCameraShakePattern::StopShakePatternImpl(const FCameraShakeStopParams& Params)
+void UCompositeCameraShakePattern::StopShakePatternImpl(const FCameraShakePatternStopParams& Params)
 {
 	// Stop all our children.
 	for (UCameraShakePattern* Pattern : ChildPatterns)
diff --git a/Engine/Plugins/Cameras/GameplayCameras/Source/GameplayCameras/Private/LegacyCameraShake.cpp b/Engine/Plugins/Cameras/GameplayCameras/Source/GameplayCameras/Private/LegacyCameraShake.cpp
--- a/Engine/Plugins/Cameras/GameplayCameras/Source/GameplayCameras/Private/LegacyCameraShake.cpp
+++ b/Engine/Plugins/Cameras/GameplayCameras/Source/GameplayCameras/Private/LegacyCameraShake.cpp
@@ -99,7 +99,7 @@ void ULegacyCameraShake::DoStopShake(bool bImmediately)
 	
 	if (SequenceShakePattern)
 	{
-		FCameraShakeStopParams StopParams;
+		FCameraShakePatternStopParams StopParams;
 		StopParams.bImmediately = bImmediately;
 		SequenceShakePattern->StopShakePattern(StopParams);
 	}
@@ -109,7 +109,7 @@ void ULegacyCameraShake::DoStopShake(bool bImmediately)
 	ReceiveStopShake(bImmediately);
 }
 
-void ULegacyCameraShake::DoStartShake(const FCameraShakeStartParams& Params)
+void ULegacyCameraShake::DoStartShake(const FCameraShakePatternStartParams& Params)
 {
 	const float EffectiveOscillationDuration = (OscillationDuration > 0.f) ? OscillationDuration : TNumericLimits<float>::Max();
 
@@ -191,7 +191,7 @@ void ULegacyCameraShake::DoStartShake(const FCameraShakeStartParams& Params)
 	ReceivePlayShake(ShakeScale);
 }
 
-void ULegacyCameraShake::DoUpdateShake(const FCameraShakeUpdateParams& Params, FCameraShakeUpdateResult& OutResult)
+void ULegacyCameraShake::DoUpdateShake(const FCameraShakePatternUpdateParams& Params, FCameraShakePatternUpdateResult& OutResult)
 {
 	const float DeltaTime = Params.DeltaTime;
 
@@ -299,7 +299,7 @@ void ULegacyCameraShake::DoUpdateShake(const FCameraShakeUpdateParams& Params, F
 	// Update the sequence animation if there's one.
 	if (SequenceShakePattern != nullptr && !SequenceShakePattern->IsFinished())
 	{
-		FCameraShakeUpdateResult ChildResult;
+		FCameraShakePatternUpdateResult ChildResult;
 		SequenceShakePattern->UpdateShakePattern(Params, ChildResult);
 
 		// The sequence shake pattern returns a local, additive result. So we should be able to
@@ -313,11 +313,11 @@ void ULegacyCameraShake::DoUpdateShake(const FCameraShakeUpdateParams& Params, F
 	}
 
 	// Apply the scaling, limits, and playspace so we have an absolute result we can pass to the legacy blueprint API.
-	check(OutResult.Flags == ECameraShakeUpdateResultFlags::Default);
+	check(OutResult.Flags == ECameraShakePatternUpdateResultFlags::Default);
 	ApplyScale(Params.GetTotalScale(), OutResult);
 	ApplyLimits(Params.POV, OutResult);
 	ApplyPlaySpace(Params, OutResult);
-	check(EnumHasAnyFlags(OutResult.Flags, ECameraShakeUpdateResultFlags::ApplyAsAbsolute));
+	check(EnumHasAnyFlags(OutResult.Flags, ECameraShakePatternUpdateResultFlags::ApplyAsAbsolute));
 
 	// Call the legacy blueprint API. We need to convert back and forth.
 	{
@@ -336,7 +336,7 @@ void ULegacyCameraShake::DoUpdateShake(const FCameraShakeUpdateParams& Params, F
 	UE_LOG(LogLegacyCameraShake, Verbose, TEXT("ULegacyCameraShake::DoUpdateShake %s Finished: %i Duration: %f Remaining: %f"), *GetNameSafe(this), bOscillationFinished, OscillationDuration, OscillatorTimeRemaining);
 }
 
-void ULegacyCameraShake::DoScrubShake(const FCameraShakeScrubParams& Params, FCameraShakeUpdateResult& OutResult)
+void ULegacyCameraShake::DoScrubShake(const FCameraShakePatternScrubParams& Params, FCameraShakePatternUpdateResult& OutResult)
 {
 	const float NewTime = Params.AbsoluteTime;
 
@@ -370,11 +370,11 @@ void ULegacyCameraShake::DoScrubShake(const FCameraShakeScrubParams& Params, FCa
 		}
 	}
 
-	FCameraShakeUpdateParams UpdateParams = Params.ToUpdateParams();
+	FCameraShakePatternUpdateParams UpdateParams = Params.ToUpdateParams();
 
 	DoUpdateShake(UpdateParams, OutResult);
 
-	check(EnumHasAnyFlags(OutResult.Flags, ECameraShakeUpdateResultFlags::ApplyAsAbsolute));
+	check(EnumHasAnyFlags(OutResult.Flags, ECameraShakePatternUpdateResultFlags::ApplyAsAbsolute));
 }
 
 bool ULegacyCameraShake::DoGetIsFinished() const
@@ -454,25 +454,25 @@ void ULegacyCameraShakePattern::GetShakePatternInfoImpl(FCameraShakeInfo& OutInf
 	}
 }
 
-void ULegacyCameraShakePattern::StopShakePatternImpl(const FCameraShakeStopParams& Params)
+void ULegacyCameraShakePattern::StopShakePatternImpl(const FCameraShakePatternStopParams& Params)
 {
 	ULegacyCameraShake* Shake = GetShakeInstance<ULegacyCameraShake>();
 	Shake->DoStopShake(Params.bImmediately);
 }
 
-void ULegacyCameraShakePattern::StartShakePatternImpl(const FCameraShakeStartParams& Params)
+void ULegacyCameraShakePattern::StartShakePatternImpl(const FCameraShakePatternStartParams& Params)
 {
 	ULegacyCameraShake* Shake = GetShakeInstance<ULegacyCameraShake>();
 	Shake->DoStartShake(Params);
 }
 
-void ULegacyCameraShakePattern::UpdateShakePatternImpl(const FCameraShakeUpdateParams& Params, FCameraShakeUpdateResult& OutResult)
+void ULegacyCameraShakePattern::UpdateShakePatternImpl(const FCameraShakePatternUpdateParams& Params, FCameraShakePatternUpdateResult& OutResult)
 {
 	ULegacyCameraShake* Shake = GetShakeInstance<ULegacyCameraShake>();
 	Shake->DoUpdateShake(Params, OutResult);
 }
 
-void ULegacyCameraShakePattern::ScrubShakePatternImpl(const FCameraShakeScrubParams& Params, FCameraShakeUpdateResult& OutResult)
+void ULegacyCameraShakePattern::ScrubShakePatternImpl(const FCameraShakePatternScrubParams& Params, FCameraShakePatternUpdateResult& OutResult)
 {
 	ULegacyCameraShake* Shake = GetShakeInstance<ULegacyCameraShake>();
 	Shake->DoScrubShake(Params, OutResult);
diff --git a/Engine/Plugins/Cameras/GameplayCameras/Source/GameplayCameras/Private/PerlinNoiseCameraShakePattern.cpp b/Engine/Plugins/Cameras/GameplayCameras/Source/GameplayCameras/Private/PerlinNoiseCameraShakePattern.cpp
--- a/Engine/Plugins/Cameras/GameplayCameras/Source/GameplayCameras/Private/PerlinNoiseCameraShakePattern.cpp
+++ b/Engine/Plugins/Cameras/GameplayCameras/Source/GameplayCameras/Private/PerlinNoiseCameraShakePattern.cpp
@@ -23,7 +23,7 @@ UPerlinNoiseCameraShakePattern::UPerlinNoiseCameraShakePattern(const FObjectInit
 	FOV.Amplitude = 0.f;
 }
 
-void UPerlinNoiseCameraShakePattern::StartShakePatternImpl(const FCameraShakeStartParams& Params)
+void UPerlinNoiseCameraShakePattern::StartShakePatternImpl(const FCameraShakePatternStartParams& Params)
 {
 	Super::StartShakePatternImpl(Params);
 
@@ -42,7 +42,7 @@ void UPerlinNoiseCameraShakePattern::StartShakePatternImpl(const FCameraShakeSta
 	}
 }
 
-void UPerlinNoiseCameraShakePattern::UpdateShakePatternImpl(const FCameraShakeUpdateParams& Params, FCameraShakeUpdateResult& OutResult)
+void UPerlinNoiseCameraShakePattern::UpdateShakePatternImpl(const FCameraShakePatternUpdateParams& Params, FCameraShakePatternUpdateResult& OutResult)
 {
 	UpdatePerlinNoise(Params.DeltaTime, OutResult);
 
@@ -50,7 +50,7 @@ void UPerlinNoiseCameraShakePattern::UpdateShakePatternImpl(const FCameraShakeUp
 	OutResult.ApplyScale(BlendWeight);
 }
 
-void UPerlinNoiseCameraShakePattern::ScrubShakePatternImpl(const FCameraShakeScrubParams& Params, FCameraShakeUpdateResult& OutResult)
+void UPerlinNoiseCameraShakePattern::ScrubShakePatternImpl(const FCameraShakePatternScrubParams& Params, FCameraShakePatternUpdateResult& OutResult)
 {
 	// Scrubbing is like going back to our initial state and updating directly to the scrub time.
 	CurrentLocationOffset = InitialLocationOffset;
@@ -63,7 +63,7 @@ void UPerlinNoiseCameraShakePattern::ScrubShakePatternImpl(const FCameraShakeScr
 	OutResult.ApplyScale(BlendWeight);
 }
 
-void UPerlinNoiseCameraShakePattern::UpdatePerlinNoise(float DeltaTime, FCameraShakeUpdateResult& OutResult)
+void UPerlinNoiseCameraShakePattern::UpdatePerlinNoise(float DeltaTime, FCameraShakePatternUpdateResult& OutResult)
 {
 	OutResult.Location.X = X.Update(DeltaTime, LocationAmplitudeMultiplier, LocationFrequencyMultiplier, CurrentLocationOffset.X);
 	OutResult.Location.Y = Y.Update(DeltaTime, LocationAmplitudeMultiplier, LocationFrequencyMultiplier, CurrentLocationOffset.Y);
diff --git a/Engine/Plugins/Cameras/GameplayCameras/Source/GameplayCameras/Private/SimpleCameraShakePattern.cpp b/Engine/Plugins/Cameras/GameplayCameras/Source/GameplayCameras/Private/SimpleCameraShakePattern.cpp
--- a/Engine/Plugins/Cameras/GameplayCameras/Source/GameplayCameras/Private/SimpleCameraShakePattern.cpp
+++ b/Engine/Plugins/Cameras/GameplayCameras/Source/GameplayCameras/Private/SimpleCameraShakePattern.cpp
@@ -19,7 +19,7 @@ void USimpleCameraShakePattern::GetShakePatternInfoImpl(FCameraShakeInfo& OutInf
 	OutInfo.BlendOut = BlendOutTime;
 }
 
-void USimpleCameraShakePattern::StartShakePatternImpl(const FCameraShakeStartParams& Params)
+void USimpleCameraShakePattern::StartShakePatternImpl(const FCameraShakePatternStartParams& Params)
 {
 	State.Start(this);
 }
@@ -29,7 +29,7 @@ bool USimpleCameraShakePattern::IsFinishedImpl() const
 	return !State.IsPlaying();
 }
 
-void USimpleCameraShakePattern::StopShakePatternImpl(const FCameraShakeStopParams& Params)
+void USimpleCameraShakePattern::StopShakePatternImpl(const FCameraShakePatternStopParams& Params)
 {
 	State.Stop(Params.bImmediately);
 }
diff --git a/Engine/Plugins/Cameras/GameplayCameras/Source/GameplayCameras/Private/Tests/CameraShakeTestObjects.cpp b/Engine/Plugins/Cameras/GameplayCameras/Source/GameplayCameras/Private/Tests/CameraShakeTestObjects.cpp
--- a/Engine/Plugins/Cameras/GameplayCameras/Source/GameplayCameras/Private/Tests/CameraShakeTestObjects.cpp
+++ b/Engine/Plugins/Cameras/GameplayCameras/Source/GameplayCameras/Private/Tests/CameraShakeTestObjects.cpp
@@ -9,7 +9,7 @@ UConstantCameraShakePattern::UConstantCameraShakePattern(const FObjectInitialize
 {
 }
 
-void UConstantCameraShakePattern::UpdateShakePatternImpl(const FCameraShakeUpdateParams& Params, FCameraShakeUpdateResult& OutResult)
+void UConstantCameraShakePattern::UpdateShakePatternImpl(const FCameraShakePatternUpdateParams& Params, FCameraShakePatternUpdateResult& OutResult)
 {
 	OutResult.Location = LocationOffset;
 	OutResult.Rotation = RotationOffset;
diff --git a/Engine/Plugins/Cameras/GameplayCameras/Source/GameplayCameras/Private/Tests/CameraShakeTestObjects.h b/Engine/Plugins/Cameras/GameplayCameras/Source/GameplayCameras/Private/Tests/CameraShakeTestObjects.h
--- a/Engine/Plugins/Cameras/GameplayCameras/Source/GameplayCameras/Private/Tests/CameraShakeTestObjects.h
+++ b/Engine/Plugins/Cameras/GameplayCameras/Source/GameplayCameras/Private/Tests/CameraShakeTestObjects.h
@@ -42,6 +42,6 @@ public:
 	FRotator RotationOffset;
 
 private:
-	virtual void UpdateShakePatternImpl(const FCameraShakeUpdateParams& Params, FCameraShakeUpdateResult& OutResult) override;
+	virtual void UpdateShakePatternImpl(const FCameraShakePatternUpdateParams& Params, FCameraShakePatternUpdateResult& OutResult) override;
 };
 
diff --git a/Engine/Plugins/Cameras/GameplayCameras/Source/GameplayCameras/Private/WaveOscillatorCameraShakePattern.cpp b/Engine/Plugins/Cameras/GameplayCameras/Source/GameplayCameras/Private/WaveOscillatorCameraShakePattern.cpp
--- a/Engine/Plugins/Cameras/GameplayCameras/Source/GameplayCameras/Private/WaveOscillatorCameraShakePattern.cpp
+++ b/Engine/Plugins/Cameras/GameplayCameras/Source/GameplayCameras/Private/WaveOscillatorCameraShakePattern.cpp
@@ -31,7 +31,7 @@ UWaveOscillatorCameraShakePattern::UWaveOscillatorCameraShakePattern(const FObje
 	FOV.Amplitude = 0.f;
 }
 
-void UWaveOscillatorCameraShakePattern::StartShakePatternImpl(const FCameraShakeStartParams& Params)
+void UWaveOscillatorCameraShakePattern::StartShakePatternImpl(const FCameraShakePatternStartParams& Params)
 {
 	Super::StartShakePatternImpl(Params);
 
@@ -55,7 +55,7 @@ void UWaveOscillatorCameraShakePattern::StartShakePatternImpl(const FCameraShake
 	}
 }
 
-void UWaveOscillatorCameraShakePattern::UpdateShakePatternImpl(const FCameraShakeUpdateParams& Params, FCameraShakeUpdateResult& OutResult)
+void UWaveOscillatorCameraShakePattern::UpdateShakePatternImpl(const FCameraShakePatternUpdateParams& Params, FCameraShakePatternUpdateResult& OutResult)
 {
 	UpdateOscillators(Params.DeltaTime, OutResult);
 
@@ -63,7 +63,7 @@ void UWaveOscillatorCameraShakePattern::UpdateShakePatternImpl(const FCameraShak
 	OutResult.ApplyScale(BlendWeight);
 }
 
-void UWaveOscillatorCameraShakePattern::ScrubShakePatternImpl(const FCameraShakeScrubParams& Params, FCameraShakeUpdateResult& OutResult)
+void UWaveOscillatorCameraShakePattern::ScrubShakePatternImpl(const FCameraShakePatternScrubParams& Params, FCameraShakePatternUpdateResult& OutResult)
 {
 	// Scrubbing is like going back to our initial state and updating directly to the scrub time.
 	CurrentLocationOffset = InitialLocationOffset;
@@ -76,7 +76,7 @@ void UWaveOscillatorCameraShakePattern::ScrubShakePatternImpl(const FCameraShake
 	OutResult.ApplyScale(BlendWeight);
 }
 
-void UWaveOscillatorCameraShakePattern::UpdateOscillators(float DeltaTime, FCameraShakeUpdateResult& OutResult)
+void UWaveOscillatorCameraShakePattern::UpdateOscillators(float DeltaTime, FCameraShakePatternUpdateResult& OutResult)
 {
 	OutResult.Location.X = X.Update(DeltaTime, LocationAmplitudeMultiplier, LocationFrequencyMultiplier, CurrentLocationOffset.X);
 	OutResult.Location.Y = Y.Update(DeltaTime, LocationAmplitudeMultiplier, LocationFrequencyMultiplier, CurrentLocationOffset.Y);
diff --git a/Engine/Plugins/Cameras/GameplayCameras/Source/GameplayCameras/Public/CompositeCameraShakePattern.h b/Engine/Plugins/Cameras/GameplayCameras/Source/GameplayCameras/Public/CompositeCameraShakePattern.h
--- a/Engine/Plugins/Cameras/GameplayCameras/Source/GameplayCameras/Public/CompositeCameraShakePattern.h
+++ b/Engine/Plugins/Cameras/GameplayCameras/Source/GameplayCameras/Public/CompositeCameraShakePattern.h
@@ -35,11 +35,11 @@ private:
 
 	// UCameraShakePattern interface
 	virtual void GetShakePatternInfoImpl(FCameraShakeInfo& OutInfo) const override;
-	virtual void StartShakePatternImpl(const FCameraShakeStartParams& Params) override;
-	virtual void UpdateShakePatternImpl(const FCameraShakeUpdateParams& Params, FCameraShakeUpdateResult& OutResult) override;
-	virtual void ScrubShakePatternImpl(const FCameraShakeScrubParams& Params, FCameraShakeUpdateResult& OutResult) override;
+	virtual void StartShakePatternImpl(const FCameraShakePatternStartParams& Params) override;
+	virtual void UpdateShakePatternImpl(const FCameraShakePatternUpdateParams& Params, FCameraShakePatternUpdateResult& OutResult) override;
+	virtual void ScrubShakePatternImpl(const FCameraShakePatternScrubParams& Params, FCameraShakePatternUpdateResult& OutResult) override;
 	virtual bool IsFinishedImpl() const override;
-	virtual void StopShakePatternImpl(const FCameraShakeStopParams& Params) override;
+	virtual void StopShakePatternImpl(const FCameraShakePatternStopParams& Params) override;
 	virtual void TeardownShakePatternImpl() override;
 };
 
diff --git a/Engine/Plugins/Cameras/GameplayCameras/Source/GameplayCameras/Public/LegacyCameraShake.h b/Engine/Plugins/Cameras/GameplayCameras/Source/GameplayCameras/Public/LegacyCameraShake.h
--- a/Engine/Plugins/Cameras/GameplayCameras/Source/GameplayCameras/Public/LegacyCameraShake.h
+++ b/Engine/Plugins/Cameras/GameplayCameras/Source/GameplayCameras/Public/LegacyCameraShake.h
@@ -242,9 +242,9 @@ public:
 
 private:
 
-	void DoStartShake(const FCameraShakeStartParams& Params);
-	void DoUpdateShake(const FCameraShakeUpdateParams& Params, FCameraShakeUpdateResult& OutResult);
-	void DoScrubShake(const FCameraShakeScrubParams& Params, FCameraShakeUpdateResult& OutResult);
+	void DoStartShake(const FCameraShakePatternStartParams& Params);
+	void DoUpdateShake(const FCameraShakePatternUpdateParams& Params, FCameraShakePatternUpdateResult& OutResult);
+	void DoScrubShake(const FCameraShakePatternScrubParams& Params, FCameraShakePatternUpdateResult& OutResult);
 	void DoStopShake(bool bImmediately);
 	bool DoGetIsFinished() const;
 	void DoTeardownShake();
@@ -299,11 +299,11 @@ private:
 
 	// UCameraShakePattern interface
 	virtual void GetShakePatternInfoImpl(FCameraShakeInfo& OutInfo) const override;
-	virtual void StartShakePatternImpl(const FCameraShakeStartParams& Params) override;
-	virtual void UpdateShakePatternImpl(const FCameraShakeUpdateParams& Params, FCameraShakeUpdateResult& OutResult) override;
-	virtual void ScrubShakePatternImpl(const FCameraShakeScrubParams& Params, FCameraShakeUpdateResult& OutResult) override;
+	virtual void StartShakePatternImpl(const FCameraShakePatternStartParams& Params) override;
+	virtual void UpdateShakePatternImpl(const FCameraShakePatternUpdateParams& Params, FCameraShakePatternUpdateResult& OutResult) override;
+	virtual void ScrubShakePatternImpl(const FCameraShakePatternScrubParams& Params, FCameraShakePatternUpdateResult& OutResult) override;
 	virtual bool IsFinishedImpl() const override;
-	virtual void StopShakePatternImpl(const FCameraShakeStopParams& Params) override;
+	virtual void StopShakePatternImpl(const FCameraShakePatternStopParams& Params) override;
 	virtual void TeardownShakePatternImpl() override;
 };
 
diff --git a/Engine/Plugins/Cameras/GameplayCameras/Source/GameplayCameras/Public/PerlinNoiseCameraShakePattern.h b/Engine/Plugins/Cameras/GameplayCameras/Source/GameplayCameras/Public/PerlinNoiseCameraShakePattern.h
--- a/Engine/Plugins/Cameras/GameplayCameras/Source/GameplayCameras/Public/PerlinNoiseCameraShakePattern.h
+++ b/Engine/Plugins/Cameras/GameplayCameras/Source/GameplayCameras/Public/PerlinNoiseCameraShakePattern.h
@@ -5,11 +5,6 @@
 #include "SimpleCameraShakePattern.h"
 #include "PerlinNoiseCameraShakePattern.generated.h"
 
-struct FCameraShakeScrubParams;
-struct FCameraShakeStartParams;
-struct FCameraShakeUpdateParams;
-struct FCameraShakeUpdateResult;
-
 /** A perlin noise shaker for a single number. */
 USTRUCT(BlueprintType)
 struct GAMEPLAYCAMERAS_API FPerlinNoiseShaker
@@ -95,11 +90,11 @@ public:
 private:
 
 	// UCameraShakePattern interface
-	virtual void StartShakePatternImpl(const FCameraShakeStartParams& Params) override;
-	virtual void UpdateShakePatternImpl(const FCameraShakeUpdateParams& Params, FCameraShakeUpdateResult& OutResult) override;
-	virtual void ScrubShakePatternImpl(const FCameraShakeScrubParams& Params, FCameraShakeUpdateResult& OutResult) override;
+	virtual void StartShakePatternImpl(const FCameraShakePatternStartParams& Params) override;
+	virtual void UpdateShakePatternImpl(const FCameraShakePatternUpdateParams& Params, FCameraShakePatternUpdateResult& OutResult) override;
+	virtual void ScrubShakePatternImpl(const FCameraShakePatternScrubParams& Params, FCameraShakePatternUpdateResult& OutResult) override;
 
-	void UpdatePerlinNoise(float DeltaTime, FCameraShakeUpdateResult& OutResult);
+	void UpdatePerlinNoise(float DeltaTime, FCameraShakePatternUpdateResult& OutResult);
 
 private:
 
diff --git a/Engine/Plugins/Cameras/GameplayCameras/Source/GameplayCameras/Public/SimpleCameraShakePattern.h b/Engine/Plugins/Cameras/GameplayCameras/Source/GameplayCameras/Public/SimpleCameraShakePattern.h
--- a/Engine/Plugins/Cameras/GameplayCameras/Source/GameplayCameras/Public/SimpleCameraShakePattern.h
+++ b/Engine/Plugins/Cameras/GameplayCameras/Source/GameplayCameras/Public/SimpleCameraShakePattern.h
@@ -35,9 +35,9 @@ protected:
 
 	// UCameraShakePattern interface
 	virtual void GetShakePatternInfoImpl(FCameraShakeInfo& OutInfo) const override;
-	virtual void StartShakePatternImpl(const FCameraShakeStartParams& Params) override;
+	virtual void StartShakePatternImpl(const FCameraShakePatternStartParams& Params) override;
 	virtual bool IsFinishedImpl() const override;
-	virtual void StopShakePatternImpl(const FCameraShakeStopParams& Params) override;
+	virtual void StopShakePatternImpl(const FCameraShakePatternStopParams& Params) override;
 	virtual void TeardownShakePatternImpl()  override;
 
 protected:
diff --git a/Engine/Plugins/Cameras/GameplayCameras/Source/GameplayCameras/Public/WaveOscillatorCameraShakePattern.h b/Engine/Plugins/Cameras/GameplayCameras/Source/GameplayCameras/Public/WaveOscillatorCameraShakePattern.h
--- a/Engine/Plugins/Cameras/GameplayCameras/Source/GameplayCameras/Public/WaveOscillatorCameraShakePattern.h
+++ b/Engine/Plugins/Cameras/GameplayCameras/Source/GameplayCameras/Public/WaveOscillatorCameraShakePattern.h
@@ -5,11 +5,6 @@
 #include "SimpleCameraShakePattern.h"
 #include "WaveOscillatorCameraShakePattern.generated.h"
 
-struct FCameraShakeScrubParams;
-struct FCameraShakeStartParams;
-struct FCameraShakeUpdateParams;
-struct FCameraShakeUpdateResult;
-
 /** Shake start offset parameter. */
 UENUM()
 enum class EInitialWaveOscillatorOffsetType : uint8
@@ -113,11 +108,11 @@ public:
 private:
 
 	// UCameraShakePattern interface
-	virtual void StartShakePatternImpl(const FCameraShakeStartParams& Params) override;
-	virtual void UpdateShakePatternImpl(const FCameraShakeUpdateParams& Params, FCameraShakeUpdateResult& OutResult) override;
-	virtual void ScrubShakePatternImpl(const FCameraShakeScrubParams& Params, FCameraShakeUpdateResult& OutResult) override;
+	virtual void StartShakePatternImpl(const FCameraShakePatternStartParams& Params) override;
+	virtual void UpdateShakePatternImpl(const FCameraShakePatternUpdateParams& Params, FCameraShakePatternUpdateResult& OutResult) override;
+	virtual void ScrubShakePatternImpl(const FCameraShakePatternScrubParams& Params, FCameraShakePatternUpdateResult& OutResult) override;
 
-	void UpdateOscillators(float DeltaTime, FCameraShakeUpdateResult& OutResult);
+	void UpdateOscillators(float DeltaTime, FCameraShakePatternUpdateResult& OutResult);
 
 private:
 
diff --git a/Engine/Plugins/MovieScene/TemplateSequence/Source/TemplateSequence/Private/SequenceCameraShake.cpp b/Engine/Plugins/MovieScene/TemplateSequence/Source/TemplateSequence/Private/SequenceCameraShake.cpp
--- a/Engine/Plugins/MovieScene/TemplateSequence/Source/TemplateSequence/Private/SequenceCameraShake.cpp
+++ b/Engine/Plugins/MovieScene/TemplateSequence/Source/TemplateSequence/Private/SequenceCameraShake.cpp
@@ -52,7 +52,7 @@ void USequenceCameraShakePattern::GetShakePatternInfoImpl(FCameraShakeInfo& OutI
 	}
 }
 
-void USequenceCameraShakePattern::StartShakePatternImpl(const FCameraShakeStartParams& Params)
+void USequenceCameraShakePattern::StartShakePatternImpl(const FCameraShakePatternStartParams& Params)
 {
 	using namespace UE::MovieScene;
 
@@ -75,7 +75,7 @@ void USequenceCameraShakePattern::StartShakePatternImpl(const FCameraShakeStartP
 	State.Start(this);
 }
 
-void USequenceCameraShakePattern::UpdateShakePatternImpl(const FCameraShakeUpdateParams& Params, FCameraShakeUpdateResult& OutResult)
+void USequenceCameraShakePattern::UpdateShakePatternImpl(const FCameraShakePatternUpdateParams& Params, FCameraShakePatternUpdateResult& OutResult)
 {
 	const float BlendWeight = State.Update(Params.DeltaTime);
 	if (State.IsPlaying())
@@ -98,7 +98,7 @@ void USequenceCameraShakePattern::UpdateShakePatternImpl(const FCameraShakeUpdat
 	}
 }
 
-void USequenceCameraShakePattern::ScrubShakePatternImpl(const FCameraShakeScrubParams& Params, FCameraShakeUpdateResult& OutResult)
+void USequenceCameraShakePattern::ScrubShakePatternImpl(const FCameraShakePatternScrubParams& Params, FCameraShakePatternUpdateResult& OutResult)
 {
 	const float BlendWeight = State.Scrub(Params.AbsoluteTime);
 	if (State.IsPlaying())
@@ -122,7 +122,7 @@ bool USequenceCameraShakePattern::IsFinishedImpl() const
 	return (Player == nullptr || Player->GetPlaybackStatus() == EMovieScenePlayerStatus::Stopped);
 }
 
-void USequenceCameraShakePattern::StopShakePatternImpl(const FCameraShakeStopParams& Params)
+void USequenceCameraShakePattern::StopShakePatternImpl(const FCameraShakePatternStopParams& Params)
 {
 	using namespace UE::MovieScene;
 
@@ -152,7 +152,7 @@ void USequenceCameraShakePattern::TeardownShakePatternImpl()
 	State = FCameraShakeState();
 }
 
-void USequenceCameraShakePattern::UpdateCamera(FFrameTime NewPosition, const FMinimalViewInfo& InPOV, FCameraShakeUpdateResult& OutResult)
+void USequenceCameraShakePattern::UpdateCamera(FFrameTime NewPosition, const FMinimalViewInfo& InPOV, FCameraShakePatternUpdateResult& OutResult)
 {
 	if (!ensure(Sequence))
 	{
diff --git a/Engine/Plugins/MovieScene/TemplateSequence/Source/TemplateSequence/Public/SequenceCameraShake.h b/Engine/Plugins/MovieScene/TemplateSequence/Source/TemplateSequence/Public/SequenceCameraShake.h
--- a/Engine/Plugins/MovieScene/TemplateSequence/Source/TemplateSequence/Public/SequenceCameraShake.h
+++ b/Engine/Plugins/MovieScene/TemplateSequence/Source/TemplateSequence/Public/SequenceCameraShake.h
@@ -63,14 +63,14 @@ private:
 
 	// UCameraShakeBase interface
 	virtual void GetShakePatternInfoImpl(FCameraShakeInfo& OutInfo) const override;
-	virtual void StartShakePatternImpl(const FCameraShakeStartParams& Params) override;
-	virtual void UpdateShakePatternImpl(const FCameraShakeUpdateParams& Params, FCameraShakeUpdateResult& OutResult) override;
-	virtual void ScrubShakePatternImpl(const FCameraShakeScrubParams& Params, FCameraShakeUpdateResult& OutResult) override;
+	virtual void StartShakePatternImpl(const FCameraShakePatternStartParams& Params) override;
+	virtual void UpdateShakePatternImpl(const FCameraShakePatternUpdateParams& Params, FCameraShakePatternUpdateResult& OutResult) override;
+	virtual void ScrubShakePatternImpl(const FCameraShakePatternScrubParams& Params, FCameraShakePatternUpdateResult& OutResult) override;
 	virtual bool IsFinishedImpl() const override;
-	virtual void StopShakePatternImpl(const FCameraShakeStopParams& Params) override;
+	virtual void StopShakePatternImpl(const FCameraShakePatternStopParams& Params) override;
 	virtual void TeardownShakePatternImpl() override;
 
-	void UpdateCamera(FFrameTime NewPosition, const FMinimalViewInfo& InPOV, FCameraShakeUpdateResult& OutResult);
+	void UpdateCamera(FFrameTime NewPosition, const FMinimalViewInfo& InPOV, FCameraShakePatternUpdateResult& OutResult);
 
 private:
 
diff --git a/Engine/Source/Runtime/Engine/Classes/Camera/CameraShakeBase.h b/Engine/Source/Runtime/Engine/Classes/Camera/CameraShakeBase.h
--- a/Engine/Source/Runtime/Engine/Classes/Camera/CameraShakeBase.h
+++ b/Engine/Source/Runtime/Engine/Classes/Camera/CameraShakeBase.h
@@ -16,29 +16,30 @@ class APlayerCameraManager;
 class UCameraShakePattern;
 
 /**
- * Parameters for starting a camera shake.
+ * Parameters for starting a camera shake pattern.
  */
 USTRUCT(BlueprintType)
-struct FCameraShakeStartParams
+struct FCameraShakePatternStartParams
 {
 	GENERATED_BODY()
 	
+	/** Whether the camera shake is restarting while playing */
 	UPROPERTY(EditAnywhere, BlueprintReadWrite, Category=CameraShake)
 	bool bIsRestarting = false;
 };
 
 /**
- * Parameters for updating a camera shake.
+ * Parameters for updating a camera shake pattern.
  */
 USTRUCT(BlueprintType)
-struct FCameraShakeUpdateParams
+struct FCameraShakePatternUpdateParams
 {
 	GENERATED_BODY()
 
-	FCameraShakeUpdateParams()
+	FCameraShakePatternUpdateParams()
 	{}
 
-	FCameraShakeUpdateParams(const FMinimalViewInfo& InPOV)
+	FCameraShakePatternUpdateParams(const FMinimalViewInfo& InPOV)
 		: POV(InPOV)
 	{}
 
@@ -68,19 +69,19 @@ struct FCameraShakeUpdateParams
  * Parameters for scrubbing a camera shake.
  */
 USTRUCT(BlueprintType)
-struct FCameraShakeScrubParams
+struct FCameraShakePatternScrubParams
 {
 	GENERATED_BODY()
 
-	FCameraShakeScrubParams()
+	FCameraShakePatternScrubParams()
 	{}
 
-	FCameraShakeScrubParams(const FMinimalViewInfo& InPOV)
+	FCameraShakePatternScrubParams(const FMinimalViewInfo& InPOV)
 		: POV(InPOV)
 	{}
 
 	/** Convert this to an update parameter struct where the delta time is from 0 to the AbsoluteTime. */
-	ENGINE_API FCameraShakeUpdateParams ToUpdateParams() const;
+	ENGINE_API FCameraShakePatternUpdateParams ToUpdateParams() const;
 
 	/** The time to scrub to */
 	UPROPERTY(EditAnywhere, BlueprintReadWrite, Category=CameraShake)
@@ -105,10 +106,10 @@ struct FCameraShakeScrubParams
 };
 
 /**
- * Flags that camera shakes can return to change base-class behaviour.
+ * Flags that camera shake patterns can return to change base-class behaviour.
  */
 UENUM()
-enum class ECameraShakeUpdateResultFlags : uint8
+enum class ECameraShakePatternUpdateResultFlags : uint8
 {
 	/** Apply the result location, rotation, and field of view as absolute values, instead of additive values. */
 	ApplyAsAbsolute = 1 << 0,
@@ -120,22 +121,22 @@ enum class ECameraShakeUpdateResultFlags : uint8
 	/** Default flags: the sub-class is returning local, additive offsets, and lets the base class take care of the rest. */
 	Default = 0
 };
-ENUM_CLASS_FLAGS(ECameraShakeUpdateResultFlags);
+ENUM_CLASS_FLAGS(ECameraShakePatternUpdateResultFlags);
 
 /**
- * The result of a camera shake update.
+ * The result of a camera shake pattern update.
  */
 USTRUCT(BlueprintType)
-struct FCameraShakeUpdateResult
+struct FCameraShakePatternUpdateResult
 {
 	GENERATED_BODY()
 
-	FCameraShakeUpdateResult()
+	FCameraShakePatternUpdateResult()
 		: Location(FVector::ZeroVector)
 		, Rotation(FRotator::ZeroRotator)
 		, FOV(0.f)
 		, PostProcessBlendWeight(0.f)
-		, Flags(ECameraShakeUpdateResultFlags::Default)
+		, Flags(ECameraShakePatternUpdateResultFlags::Default)
 	{}
 
 	/** Location offset for the view, or new absolute location if ApplyAsAbsolute flag is set */
@@ -151,7 +152,7 @@ struct FCameraShakeUpdateResult
 	float PostProcessBlendWeight;
 
 	/** Flags for how the base class should handle the result */
-	ECameraShakeUpdateResultFlags Flags;
+	ECameraShakePatternUpdateResultFlags Flags;
 
 	/** Apply the given scale to the result (only if it is "relative") */
 	ENGINE_API void ApplyScale(float InScale);
@@ -161,7 +162,7 @@ struct FCameraShakeUpdateResult
  * Parameters for stopping a camera shake.
  */
 USTRUCT(BlueprintType)
-struct FCameraShakeStopParams
+struct FCameraShakePatternStopParams
 {
 	GENERATED_BODY()
 	
@@ -169,6 +170,20 @@ struct FCameraShakeStopParams
 	bool bImmediately = false;
 };
 
+// Redirections from old names.
+#if !defined(UE_LEGACY_CAMERA_SHAKE_PATTERN_TYPES)
+#define UE_LEGACY_CAMERA_SHAKE_PATTERN_TYPES 0
+#endif
+
+#if UE_LEGACY_CAMERA_SHAKE_PATTERN_TYPES
+using ECameraShakeUpdateResultFlags = ECameraShakePatternUpdateResultFlags;
+using FCameraShakeStartParams = FCameraShakePatternStartParams;
+using FCameraShakeUpdateParams = FCameraShakePatternUpdateParams;
+using FCameraShakeScrubParams = FCameraShakePatternScrubParams;
+using FCameraShakeStopParams = FCameraShakePatternStopParams;
+using FCameraShakeUpdateResult = FCameraShakePatternUpdateResult;
+#endif // UE_LEGACY_CAMERA_SHAKE_PATTERN_TYPES
+
 /**
  * Camera shake duration type.
  */
@@ -458,23 +473,23 @@ public:
 	}
 
 	/** Uses the given result parameters to apply the given result to the given input view info */
-	static ENGINE_API void ApplyResult(const FCameraShakeApplyResultParams& ApplyParams, const FCameraShakeUpdateResult& InResult, FMinimalViewInfo& InOutPOV);
+	static ENGINE_API void ApplyResult(const FCameraShakeApplyResultParams& ApplyParams, const FCameraShakePatternUpdateResult& InResult, FMinimalViewInfo& InOutPOV);
 
 	/** Applies all the appropriate auto-scaling to the current shake offset (only if the result is "relative") */
-	static ENGINE_API void ApplyScale(const FCameraShakeUpdateParams& Params, FCameraShakeUpdateResult& InOutResult);
+	static ENGINE_API void ApplyScale(const FCameraShakePatternUpdateParams& Params, FCameraShakePatternUpdateResult& InOutResult);
 
 	/** Applies the given scale to the current shake offset (only if the result is "relative") */
-	static ENGINE_API void ApplyScale(float Scale, FCameraShakeUpdateResult& InOutResult);
+	static ENGINE_API void ApplyScale(float Scale, FCameraShakePatternUpdateResult& InOutResult);
 
 	/** Applies any appropriate system-wide limits */
-	static ENGINE_API void ApplyLimits(const FMinimalViewInfo& InPOV, FCameraShakeUpdateResult& InOutResult);
+	static ENGINE_API void ApplyLimits(const FMinimalViewInfo& InPOV, FCameraShakePatternUpdateResult& InOutResult);
 
 	/**
 	 * Modifies the current shake offset to be oriented in the current shake's play space (only if the result is "relative")
 	 *
 	 * Note that this modifies the result and makes it "absolute".
 	 */
-	static ENGINE_API void ApplyPlaySpace(ECameraShakePlaySpace PlaySpace, FMatrix UserPlaySpaceMatrix, const FMinimalViewInfo& InPOV, FCameraShakeUpdateResult& InOutResult);
+	static ENGINE_API void ApplyPlaySpace(ECameraShakePlaySpace PlaySpace, FMatrix UserPlaySpaceMatrix, const FMinimalViewInfo& InPOV, FCameraShakePatternUpdateResult& InOutResult);
 
 public:
 
@@ -558,7 +573,7 @@ protected:
 	 *
 	 * Note that this modifies the result and makes it "absolute".
 	 */
-	ENGINE_API void ApplyPlaySpace(const FCameraShakeUpdateParams& Params, FCameraShakeUpdateResult& InOutResult) const;
+	ENGINE_API void ApplyPlaySpace(const FCameraShakePatternUpdateParams& Params, FCameraShakePatternUpdateResult& InOutResult) const;
 
 private:
 
@@ -598,15 +613,15 @@ public:
 	/** Gets information about this shake pattern */
 	ENGINE_API void GetShakePatternInfo(FCameraShakeInfo& OutInfo) const;
 	/** Called when the shake pattern starts */
-	ENGINE_API void StartShakePattern(const FCameraShakeStartParams& Params);
+	ENGINE_API void StartShakePattern(const FCameraShakePatternStartParams& Params);
 	/** Updates the shake pattern, which should add its generated offset to the given result */
-	ENGINE_API void UpdateShakePattern(const FCameraShakeUpdateParams& Params, FCameraShakeUpdateResult& OutResult);
+	ENGINE_API void UpdateShakePattern(const FCameraShakePatternUpdateParams& Params, FCameraShakePatternUpdateResult& OutResult);
 	/** Scrubs the shake pattern to the given time, and apply the generated offset to the given result */
-	ENGINE_API void ScrubShakePattern(const FCameraShakeScrubParams& Params, FCameraShakeUpdateResult& OutResult);
+	ENGINE_API void ScrubShakePattern(const FCameraShakePatternScrubParams& Params, FCameraShakePatternUpdateResult& OutResult);
 	/** Returns whether this shake pattern is finished */
 	ENGINE_API bool IsFinished() const;
 	/** Called when the shake pattern is manually stopped */
-	ENGINE_API void StopShakePattern(const FCameraShakeStopParams& Params);
+	ENGINE_API void StopShakePattern(const FCameraShakePatternStopParams& Params);
 	/** Call when the shake pattern is discard, either after naturally finishing or being stopped manually */
 	ENGINE_API void TeardownShakePattern();
 
@@ -623,11 +638,11 @@ private:
 
 	// UCameraShakePattern interface
 	virtual void GetShakePatternInfoImpl(FCameraShakeInfo& OutInfo) const {}
-	virtual void StartShakePatternImpl(const FCameraShakeStartParams& Params) {}
-	virtual void UpdateShakePatternImpl(const FCameraShakeUpdateParams& Params, FCameraShakeUpdateResult& OutResult) {}
-	virtual void ScrubShakePatternImpl(const FCameraShakeScrubParams& Params, FCameraShakeUpdateResult& OutResult) {}
+	virtual void StartShakePatternImpl(const FCameraShakePatternStartParams& Params) {}
+	virtual void UpdateShakePatternImpl(const FCameraShakePatternUpdateParams& Params, FCameraShakePatternUpdateResult& OutResult) {}
+	virtual void ScrubShakePatternImpl(const FCameraShakePatternScrubParams& Params, FCameraShakePatternUpdateResult& OutResult) {}
 	virtual bool IsFinishedImpl() const { return true; }
-	virtual void StopShakePatternImpl(const FCameraShakeStopParams& Params) {}
+	virtual void StopShakePatternImpl(const FCameraShakePatternStopParams& Params) {}
 	virtual void TeardownShakePatternImpl()  {}
 };
 
diff --git a/Engine/Source/Runtime/Engine/Private/Camera/CameraShakeBase.cpp b/Engine/Source/Runtime/Engine/Private/Camera/CameraShakeBase.cpp
--- a/Engine/Source/Runtime/Engine/Private/Camera/CameraShakeBase.cpp
+++ b/Engine/Source/Runtime/Engine/Private/Camera/CameraShakeBase.cpp
@@ -201,18 +201,18 @@ void FCameraShakeState::Stop(bool bImmediately)
 	}
 }
 
-FCameraShakeUpdateParams FCameraShakeScrubParams::ToUpdateParams() const
+FCameraShakePatternUpdateParams FCameraShakePatternScrubParams::ToUpdateParams() const
 {
-	FCameraShakeUpdateParams UpdateParams(POV);
+	FCameraShakePatternUpdateParams UpdateParams(POV);
 	UpdateParams.DeltaTime = AbsoluteTime;
 	UpdateParams.ShakeScale = ShakeScale;
 	UpdateParams.DynamicScale = DynamicScale;
 	return UpdateParams;
 }
 
-void FCameraShakeUpdateResult::ApplyScale(float InScale)
+void FCameraShakePatternUpdateResult::ApplyScale(float InScale)
 {
-	if (ensureMsgf(!EnumHasAnyFlags(Flags, ECameraShakeUpdateResultFlags::ApplyAsAbsolute),
+	if (ensureMsgf(!EnumHasAnyFlags(Flags, ECameraShakePatternUpdateResultFlags::ApplyAsAbsolute),
 			TEXT("Can't scale absolute shake result")))
 	{
 		Location *= InScale;
@@ -268,7 +268,7 @@ void UCameraShakeBase::StartShake(APlayerCameraManager* Camera, float Scale, ECa
 	// Let the root pattern initialize itself.
 	if (RootShakePattern)
 	{
-		FCameraShakeStartParams StartParams;
+		FCameraShakePatternStartParams StartParams;
 		StartParams.bIsRestarting = bIsRestarting;
 		RootShakePattern->StartShakePattern(StartParams);
 	}
@@ -284,14 +284,14 @@ void UCameraShakeBase::UpdateAndApplyCameraShake(float DeltaTime, float Alpha, F
 	if (RootShakePattern)
 	{
 		// Make the sub-class do the actual work.
-		FCameraShakeUpdateParams Params(InOutPOV);
+		FCameraShakePatternUpdateParams Params(InOutPOV);
 		Params.DeltaTime = DeltaTime;
 		Params.ShakeScale = ShakeScale;
 		Params.DynamicScale = Alpha;
 
 		// Result object is initialized with zero values since the default flags make us handle it
 		// as an additive offset.
-		FCameraShakeUpdateResult Result;
+		FCameraShakePatternUpdateResult Result;
 
 		RootShakePattern->UpdateShakePattern(Params, Result);
 
@@ -321,14 +321,14 @@ void UCameraShakeBase::ScrubAndApplyCameraShake(float AbsoluteTime, float Alpha,
 	if (RootShakePattern)
 	{
 		// Make the sub-class do the actual work.
-		FCameraShakeScrubParams Params(InOutPOV);
+		FCameraShakePatternScrubParams Params(InOutPOV);
 		Params.AbsoluteTime = AbsoluteTime;
 		Params.ShakeScale = ShakeScale;
 		Params.DynamicScale = Alpha;
 
 		// Result object is initialized with zero values since the default flags make us handle it
 		// as an additive offset.
-		FCameraShakeUpdateResult Result;
+		FCameraShakePatternUpdateResult Result;
 
 		RootShakePattern->ScrubShakePattern(Params, Result);
 
@@ -369,7 +369,7 @@ void UCameraShakeBase::StopShake(bool bImmediately)
 	// Let the root pattern do any custom logic.
 	if (RootShakePattern)
 	{
-		FCameraShakeStopParams StopParams;
+		FCameraShakePatternStopParams StopParams;
 		StopParams.bImmediately = bImmediately;
 		RootShakePattern->StopShakePattern(StopParams);
 	}
@@ -385,29 +385,29 @@ void UCameraShakeBase::TeardownShake()
 	bIsActive = false;
 }
 
-void UCameraShakeBase::ApplyResult(const FCameraShakeApplyResultParams& ApplyParams, const FCameraShakeUpdateResult& InResult, FMinimalViewInfo& InOutPOV)
+void UCameraShakeBase::ApplyResult(const FCameraShakeApplyResultParams& ApplyParams, const FCameraShakePatternUpdateResult& InResult, FMinimalViewInfo& InOutPOV)
 {
-	FCameraShakeUpdateResult TempResult(InResult);
+	FCameraShakePatternUpdateResult TempResult(InResult);
 
 	// If the sub-class gave us a delta-transform, we can help with some of the basic functionality
 	// of a camera shake... namely: apply shake scaling, system limits, and play space transformation.
-	if (!EnumHasAnyFlags(TempResult.Flags, ECameraShakeUpdateResultFlags::ApplyAsAbsolute))
+	if (!EnumHasAnyFlags(TempResult.Flags, ECameraShakePatternUpdateResultFlags::ApplyAsAbsolute))
 	{
-		if (!EnumHasAnyFlags(TempResult.Flags, ECameraShakeUpdateResultFlags::SkipAutoScale))
+		if (!EnumHasAnyFlags(TempResult.Flags, ECameraShakePatternUpdateResultFlags::SkipAutoScale))
 		{
 			ApplyScale(ApplyParams.Scale, TempResult);
 		}
 
 		ApplyLimits(InOutPOV, TempResult);
 
-		if (!EnumHasAnyFlags(TempResult.Flags, ECameraShakeUpdateResultFlags::SkipAutoPlaySpace))
+		if (!EnumHasAnyFlags(TempResult.Flags, ECameraShakePatternUpdateResultFlags::SkipAutoPlaySpace))
 		{
 			ApplyPlaySpace(ApplyParams.PlaySpace, ApplyParams.UserPlaySpaceMatrix, InOutPOV, TempResult);
 		}
 	}
 
 	// Now we can apply the shake to the camera matrix.
-	if (EnumHasAnyFlags(TempResult.Flags, ECameraShakeUpdateResultFlags::ApplyAsAbsolute))
+	if (EnumHasAnyFlags(TempResult.Flags, ECameraShakePatternUpdateResultFlags::ApplyAsAbsolute))
 	{
 		InOutPOV.Location = TempResult.Location;
 		InOutPOV.Rotation = TempResult.Rotation;
@@ -428,17 +428,17 @@ void UCameraShakeBase::ApplyResult(const FCameraShakeApplyResultParams& ApplyPar
 	}
 }
 
-void UCameraShakeBase::ApplyScale(const FCameraShakeUpdateParams& Params, FCameraShakeUpdateResult& InOutResult)
+void UCameraShakeBase::ApplyScale(const FCameraShakePatternUpdateParams& Params, FCameraShakePatternUpdateResult& InOutResult)
 {
 	InOutResult.ApplyScale(Params.GetTotalScale());
 }
 
-void UCameraShakeBase::ApplyScale(float Scale, FCameraShakeUpdateResult& InOutResult)
+void UCameraShakeBase::ApplyScale(float Scale, FCameraShakePatternUpdateResult& InOutResult)
 {
 	InOutResult.ApplyScale(Scale);
 }
 
-void UCameraShakeBase::ApplyLimits(const FMinimalViewInfo& InPOV, FCameraShakeUpdateResult& InOutResult)
+void UCameraShakeBase::ApplyLimits(const FMinimalViewInfo& InPOV, FCameraShakePatternUpdateResult& InOutResult)
 {
 	// Don't allow shake to flip pitch past vertical, if not using a headset.
 	// If using a headset, we can't limit the camera locked to your head.
@@ -451,12 +451,12 @@ void UCameraShakeBase::ApplyLimits(const FMinimalViewInfo& InPOV, FCameraShakeUp
 	}
 }
 
-void UCameraShakeBase::ApplyPlaySpace(const FCameraShakeUpdateParams& Params, FCameraShakeUpdateResult& InOutResult) const
+void UCameraShakeBase::ApplyPlaySpace(const FCameraShakePatternUpdateParams& Params, FCameraShakePatternUpdateResult& InOutResult) const
 {
 	ApplyPlaySpace(PlaySpace, UserPlaySpaceMatrix, Params.POV, InOutResult);
 }
 
-void UCameraShakeBase::ApplyPlaySpace(ECameraShakePlaySpace PlaySpace, FMatrix UserPlaySpaceMatrix, const FMinimalViewInfo& InPOV, FCameraShakeUpdateResult& InOutResult)
+void UCameraShakeBase::ApplyPlaySpace(ECameraShakePlaySpace PlaySpace, FMatrix UserPlaySpaceMatrix, const FMinimalViewInfo& InPOV, FCameraShakePatternUpdateResult& InOutResult)
 {
 	// Orient the shake according to the play space.
 	const bool bIsCameraLocal = (PlaySpace == ECameraShakePlaySpace::CameraLocal);
@@ -471,7 +471,7 @@ void UCameraShakeBase::ApplyPlaySpace(ECameraShakePlaySpace PlaySpace, FMatrix U
 	}
 
 	// We have a final location/rotation for the camera, so it should be applied verbatim.
-	InOutResult.Flags = (InOutResult.Flags | ECameraShakeUpdateResultFlags::ApplyAsAbsolute);
+	InOutResult.Flags = (InOutResult.Flags | ECameraShakePatternUpdateResultFlags::ApplyAsAbsolute);
 
 	// And since we set that flag, we need to make the FOV absolute too.
 	InOutResult.FOV = InPOV.FOV + InOutResult.FOV;
@@ -492,17 +492,17 @@ void UCameraShakePattern::GetShakePatternInfo(FCameraShakeInfo& OutInfo) const
 	GetShakePatternInfoImpl(OutInfo);
 }
 
-void UCameraShakePattern::StartShakePattern(const FCameraShakeStartParams& Params)
+void UCameraShakePattern::StartShakePattern(const FCameraShakePatternStartParams& Params)
 {
 	StartShakePatternImpl(Params);
 }
 
-void UCameraShakePattern::UpdateShakePattern(const FCameraShakeUpdateParams& Params, FCameraShakeUpdateResult& OutResult)
+void UCameraShakePattern::UpdateShakePattern(const FCameraShakePatternUpdateParams& Params, FCameraShakePatternUpdateResult& OutResult)
 {
 	UpdateShakePatternImpl(Params, OutResult);
 }
 
-void UCameraShakePattern::ScrubShakePattern(const FCameraShakeScrubParams& Params, FCameraShakeUpdateResult& OutResult)
+void UCameraShakePattern::ScrubShakePattern(const FCameraShakePatternScrubParams& Params, FCameraShakePatternUpdateResult& OutResult)
 {
 	ScrubShakePatternImpl(Params, OutResult);
 }
@@ -512,7 +512,7 @@ bool UCameraShakePattern::IsFinished() const
 	return IsFinishedImpl();
 }
 
-void UCameraShakePattern::StopShakePattern(const FCameraShakeStopParams& Params)
+void UCameraShakePattern::StopShakePattern(const FCameraShakePatternStopParams& Params)
 {
 	StopShakePatternImpl(Params);
 }
