Fix subobject creation header testing the wrong Outer pointer for network support

#rnx

diff --git a/Engine/Source/Runtime/Engine/Private/DataChannel.cpp b/Engine/Source/Runtime/Engine/Private/DataChannel.cpp
--- a/Engine/Source/Runtime/Engine/Private/DataChannel.cpp
+++ b/Engine/Source/Runtime/Engine/Private/DataChannel.cpp
@@ -4363,7 +4363,8 @@ void UActorChannel::WriteContentBlockHeader( UObject* Obj, FNetBitWriter &Bunch,
 			UObject* ObjOuter = Obj->GetOuter();
 			// If the subobject's outer is the not the actor (and the outer is supported for networking),
 			// then serialize the object's outer to rebuild the outer chain on the client.
-			const bool bActorIsOuter = (ObjOuter == Actor) || (!Obj->IsSupportedForNetworking());
+			const bool bActorIsOuter = (ObjOuter == Actor) || (!ObjOuter->IsSupportedForNetworking());
+
 			Bunch.WriteBit(bActorIsOuter ? 1 : 0);
 			if (!bActorIsOuter)
 			{
