Configure multiple MemReport profiles with ini sections.
Use "File" command array which outputs each command into separate file.

#tests editor

diff --git a/Engine/Source/Runtime/Engine/Private/UnrealEngine.cpp b/Engine/Source/Runtime/Engine/Private/UnrealEngine.cpp
--- a/Engine/Source/Runtime/Engine/Private/UnrealEngine.cpp
+++ b/Engine/Source/Runtime/Engine/Private/UnrealEngine.cpp
@@ -7418,12 +7418,25 @@ bool UEngine::HandleMemReportDeferredCommand( const TCHAR* Cmd, FOutputDevice& A
 
 	FSlowHeartBeatScope HeartBeatSuspend;
 
-	const bool bPerformSlowCommands = FParse::Param( Cmd, TEXT("FULL") );
 	const bool bLogOutputToFile = !FParse::Param( Cmd, TEXT("LOG") );
 	const bool bCSV = FParse::Param(Cmd, TEXT("CSV"));
+	FString InTypeName;
+	FParse::Value( Cmd, TEXT("TYPE="), InTypeName);
 	FString InFileName;
 	FParse::Value(Cmd, TEXT("NAME="), InFileName);
 
+	if (FParse::Param(Cmd, TEXT("FULL")))
+	{
+		InTypeName = TEXT("FULL");
+	}
+	if (InTypeName.Equals(TEXT("DEFAULT"), ESearchCase::IgnoreCase))
+	{
+		InTypeName.Empty();
+	}
+	const FString SectionToRead = FString::Format(TEXT("MemReport{0}Commands"), { InTypeName });
+
+	UE_LOG(LogEngine, Log, TEXT("MemReportDeferred: Executing \"%s\" profile."), *InTypeName);
+
 	// Turn off as it makes diffing hard
 	TGuardValue<ELogTimes::Type> DisableLogTimes(GPrintLogTimes, ELogTimes::None);
 
@@ -7449,7 +7462,7 @@ bool UEngine::HandleMemReportDeferredCommand( const TCHAR* Cmd, FOutputDevice& A
 		FileArWrapper.Emplace(FileAr.Get());
 		ReportAr = FileArWrapper.GetPtrOrNull();
 
-		UE_LOG(LogEngine, Log, TEXT("MemReportDeferred: saving to %s"), *FilenameFull);		
+		UE_LOG(LogEngine, Log, TEXT("MemReportDeferred: saving to %s"), *FilenameFull);
 	}
 
 
@@ -7499,8 +7512,7 @@ bool UEngine::HandleMemReportDeferredCommand( const TCHAR* Cmd, FOutputDevice& A
 
 	// Run commands from the ini
 	TArray<FString> CommandsToRun;
-
-	if (GConfig->GetArray(TEXT("MemReportCommands"), TEXT("Cmd"), CommandsToRun, GEngineIni))
+	if (GConfig->GetArray(*SectionToRead, TEXT("Cmd"), CommandsToRun, GEngineIni))
 	{
 		for (FString& Command : CommandsToRun)
 		{
@@ -7515,42 +7527,25 @@ bool UEngine::HandleMemReportDeferredCommand( const TCHAR* Cmd, FOutputDevice& A
 		}
 	}
 
-	if (bPerformSlowCommands)
-	{
-		CommandsToRun.Reset();
-
-		if (GConfig->GetArray(TEXT("MemReportFullCommands"), TEXT("Cmd"), CommandsToRun, GEngineIni))
-		{
-			for (FString& Command : CommandsToRun)
-			{
-				if (bCSV)
-				{
-					Command += " -csv";
-				}
-
-				ReportAr->Logf(TEXT("MemReport: Begin command \"%s\""), *Command);
-				Exec(InWorld, *Command, *ReportAr);
-				ReportAr->Logf(TEXT("MemReport: End command \"%s\"") LINE_TERMINATOR, *Command);
-			}
-		}
-	}
-
-	if (bLogOutputToFile && bCSV)
+	if (bLogOutputToFile)
 	{
 		CommandsToRun.Reset();
-
-		if (GConfig->GetArray(TEXT("MemReportCsvCommands"), TEXT("Cmd"), CommandsToRun, GEngineIni))
+		if (GConfig->GetArray(*SectionToRead, TEXT("File"), CommandsToRun, GEngineIni))
 		{
 			const FString DirectoryName = FPaths::GetBaseFilename(FilenameFull, false) + TEXT("/");
 			IFileManager::Get().MakeDirectory( *DirectoryName );
 
 			for (FString& Command : CommandsToRun)
 			{
-				const FString ValidFilePath = DirectoryName + FPaths::MakeValidFileName(Command) + TEXT(".csv");
+				const FString ValidFilePath = DirectoryName + FPaths::MakeValidFileName(Command) + bCSV ? TEXT(".csv") : TEXT(".txt");
 
 				TUniquePtr<FArchive> CommandFile(IFileManager::Get().CreateDebugFileWriter(*ValidFilePath));
 				FOutputDeviceArchiveWrapper CommandFileAr(CommandFile.Get());
-				Command += TEXT(" -csv");
+				if (bCSV)
+				{
+					Command += " -csv";
+				}
+
 				Exec(InWorld, *Command, CommandFileAr);
 				UE_LOG(LogEngine, Log, TEXT("MemReportDeferred: Written \"%s\" output to \"%s\""), *Command, *ValidFilePath);
 			}
