Add GetAllPhysicsObjects() and GetPhysicsObjectById() to ISMCs.

diff --git a/Engine/Source/Runtime/Engine/Classes/Components/InstancedStaticMeshComponent.h b/Engine/Source/Runtime/Engine/Classes/Components/InstancedStaticMeshComponent.h
--- a/Engine/Source/Runtime/Engine/Classes/Components/InstancedStaticMeshComponent.h
+++ b/Engine/Source/Runtime/Engine/Classes/Components/InstancedStaticMeshComponent.h
@@ -412,6 +412,12 @@ protected:
 	ENGINE_API virtual bool ComponentOverlapMultiImpl(TArray<struct FOverlapResult>& OutOverlaps, const class UWorld* InWorld, const FVector& Pos, const FQuat& Rot, ECollisionChannel TestChannel, const struct FComponentQueryParams& Params, const struct FCollisionObjectQueryParams& ObjectQueryParams = FCollisionObjectQueryParams::DefaultObjectQueryParam) const override;
 	//~ End UPrimitiveComponent Interface
 
+	//~ Begin IPhysicsComponent Interface.
+public:
+	ENGINE_API virtual Chaos::FPhysicsObject* GetPhysicsObjectById(Chaos::FPhysicsObjectId Id) const override;
+	ENGINE_API virtual TArray<Chaos::FPhysicsObject*> GetAllPhysicsObjects() const override;
+	//~ End IPhysicsComponent Interface.
+
 	//~ Begin UStaticMeshComponentInterface
 protected:
 	ENGINE_API virtual FPrimitiveSceneProxy* CreateStaticMeshSceneProxy(Nanite::FMaterialAudit& NaniteMaterials, bool bCreateNanite) override;
diff --git a/Engine/Source/Runtime/Engine/Private/InstancedStaticMesh.cpp b/Engine/Source/Runtime/Engine/Private/InstancedStaticMesh.cpp
--- a/Engine/Source/Runtime/Engine/Private/InstancedStaticMesh.cpp
+++ b/Engine/Source/Runtime/Engine/Private/InstancedStaticMesh.cpp
@@ -37,6 +37,8 @@
 #include "Collision.h"
 #include "CollisionDebugDrawingPublic.h"
 #include "DataDrivenShaderPlatformInfo.h"
+#include "Chaos/PhysicsObjectInterface.h"
+#include "PhysicsProxy/SingleParticlePhysicsProxy.h"
 
 #include "Elements/Framework/EngineElementsLibrary.h"
 #include "Elements/Interfaces/TypedElementWorldInterface.h"
@@ -3118,6 +3120,26 @@ void UInstancedStaticMeshComponent::OnDestroyPhysicsState()
 	FNavigationSystem::UpdateComponentData(*this);
 }
 
+Chaos::FPhysicsObject* UInstancedStaticMeshComponent::GetPhysicsObjectById(Chaos::FPhysicsObjectId Id) const
+{
+	if (!InstanceBodies.IsValidIndex(Id) || !InstanceBodies[Id] || !InstanceBodies[Id]->ActorHandle)
+	{
+		return nullptr;
+	}
+	return InstanceBodies[Id]->ActorHandle->GetPhysicsObject();
+}
+
+TArray<Chaos::FPhysicsObject*> UInstancedStaticMeshComponent::GetAllPhysicsObjects() const
+{
+	TArray<Chaos::FPhysicsObject*> Objects;
+	Objects.Reserve(InstanceBodies.Num());
+	for (int32 Index = 0; Index < InstanceBodies.Num(); ++Index)
+	{
+		Objects.Add(GetPhysicsObjectById(Index));
+	}
+	return Objects;
+}
+
 bool UInstancedStaticMeshComponent::CanEditSimulatePhysics()
 {
 	// if instancedstaticmeshcomponent, we will never allow it
