Protect against crashes trying to render debug HUD on a partially created debug canvas

diff --git a/Engine/Source/Runtime/Engine/Private/HUD.cpp b/Engine/Source/Runtime/Engine/Private/HUD.cpp
--- a/Engine/Source/Runtime/Engine/Private/HUD.cpp
+++ b/Engine/Source/Runtime/Engine/Private/HUD.cpp
@@ -257,7 +257,7 @@ void AHUD::DrawSafeZoneOverlay()
 #if ENABLE_DRAW_DEBUG
 	const int32 DebugSafeZoneMode = GSafeZoneVisualizationModeCVar.GetValueOnGameThread();
 
-	if ((DebugSafeZoneMode > 0) && (DebugCanvas != nullptr))
+	if ((DebugSafeZoneMode > 0) && (DebugCanvas != nullptr) && (DebugCanvas->Canvas != nullptr))
 	{
 
 		const float Width = DebugCanvas->SizeX;
@@ -420,7 +420,7 @@ bool AHUD::ShouldDisplayDebug(const FName& DebugType) const
 
 void AHUD::ShowDebugInfo(float& YL, float& YPos)
 {
-	if (DebugCanvas != nullptr)
+	if (DebugCanvas != nullptr && DebugCanvas->Canvas != nullptr)
 	{
 		// Darken background, so we can read text better.
 		FLinearColor BackgroundColor(0.f, 0.f, 0.f, 0.2f);
@@ -670,7 +670,7 @@ void AHUD::OnLostFocusPause(bool bEnable)
 
 void AHUD::DrawDebugTextList()
 {
-	if ((DebugTextList.Num() > 0) && (DebugCanvas != nullptr))
+	if ((DebugTextList.Num() > 0) && (DebugCanvas != nullptr) && (DebugCanvas->Canvas != nullptr))
 	{
 		FRotator CameraRot;
 		FVector CameraLoc;
