Fix bug where stale autogen shader header file for ES3_1 in the intermediate directory can cause shader cache keys to differ from machine-to-machine because the presence of an old autogen shader header will cause the source hash to differ when shader dependencies are computed due to includes.

diff --git a/Engine/Source/Runtime/Engine/Private/ShaderCompiler/ShaderGenerationUtil.cpp b/Engine/Source/Runtime/Engine/Private/ShaderCompiler/ShaderGenerationUtil.cpp
--- a/Engine/Source/Runtime/Engine/Private/ShaderCompiler/ShaderGenerationUtil.cpp
+++ b/Engine/Source/Runtime/Engine/Private/ShaderCompiler/ShaderGenerationUtil.cpp
@@ -2023,15 +2023,43 @@ void FShaderCompileUtilities::WriteGBufferInfoAutogen(EShaderPlatform TargetPlat
 
 	FScopeLock MapLock(&GCriticalSection);
 
+	FString AutoGenDirectory = GetAutoGenDirectory(TargetPlatform);
+	FString AutogenHeaderFilename = AutoGenDirectory / TEXT("AutogenShaderHeaders.ush");
+	FString AutogenHeaderFilenameTemp = AutoGenDirectory / TEXT("AutogenShaderHeaders_temp.ush");
+
+	// auto-generated GBuffer layout is not used by mobile rendering
+	if (FeatureLevel <= ERHIFeatureLevel::ES3_1)
+	{
+		if (GLastGBufferIsValid[TargetPlatform])
+		{
+			// We've already ensured the autogen file is removed.
+			return;
+		}
+
+		IPlatformFile& PlatformFile = FPlatformFileManager::Get().GetPlatformFile();
+		if (PlatformFile.FileExists(*AutogenHeaderFilename))
+		{
+			if (PlatformFile.DeleteFile(*AutogenHeaderFilename))
+			{
+				GLastGBufferIsValid[TargetPlatform] = true;
+			}
+			else
+			{
+				UE_LOG(LogShaderCompilers, Display, TEXT("Failed to delete old and now unused shader autogen file: %s"), *AutogenHeaderFilename);
+			}
+		}
+		else
+		{
+			GLastGBufferIsValid[TargetPlatform] = true;
+		}
+		return;
+	}
+
 	// For now, the logic always calculates the new GBuffer, and if it's the first time, write it, otherwise check it hasn't changed. We are doing this for
 	// debugging, and in the near future it will only calculate the GBuffer on the first time only.
 
 	FGBufferInfo DefaultBufferInfo = FetchFullGBufferInfo(DefaultParams);
 
-	FString AutoGenDirectory = GetAutoGenDirectory(TargetPlatform);
-	FString AutogenHeaderFilename = AutoGenDirectory / TEXT("AutogenShaderHeaders.ush");
-	FString AutogenHeaderFilenameTemp = AutoGenDirectory / TEXT("AutogenShaderHeaders_temp.ush");
-
 	if (GLastGBufferIsValid[TargetPlatform])
 	{
 		const bool bSame = IsGBufferInfoEqual(GLastGBufferInfo[TargetPlatform], DefaultBufferInfo);
@@ -2149,12 +2177,9 @@ void FShaderCompileUtilities::WriteGBufferInfoAutogen(EShaderPlatform TargetPlat
 void FShaderCompileUtilities::GenerateBrdfHeaders(EShaderPlatform TargetPlatform)
 {
 	ERHIFeatureLevel::Type FeatureLevel = GetMaxSupportedFeatureLevel(TargetPlatform);
-	// auto-generated GBuffer layout is not used by mobile rendering
-	if (FeatureLevel > ERHIFeatureLevel::ES3_1)
-	{
-		// Writes the GBuffer format .ush file if it's out of date.
-		WriteGBufferInfoAutogen(TargetPlatform, FeatureLevel);
-	}
+
+	// Writes the GBuffer format .ush file if it's out of date.
+	WriteGBufferInfoAutogen(TargetPlatform, FeatureLevel);
 }
 
 void FShaderCompileUtilities::GenerateBrdfHeaders(const FName& ShaderFormat)
