Add CVar to change the minimum volume for a wave instance to be considered playing + a game feature action to set it in different experiences and platforms
#rnx

diff --git a/Engine/Source/Runtime/Engine/Private/Audio.cpp b/Engine/Source/Runtime/Engine/Private/Audio.cpp
--- a/Engine/Source/Runtime/Engine/Private/Audio.cpp
+++ b/Engine/Source/Runtime/Engine/Private/Audio.cpp
@@ -91,6 +91,14 @@ FAutoConsoleVariableRef CVarBypassPlayWhenSilent(
 	TEXT("0: Honor the Play When Silent flag, 1: stop all silent non-procedural sources."),
 	ECVF_Default);
 
+static float WaveInstanceMinVolumeThresholdCVar = UE_KINDA_SMALL_NUMBER;
+FAutoConsoleVariableRef CVarMinVolumeThreshold(
+	TEXT("au.WaveInstanceMinVolume"),
+	WaveInstanceMinVolumeThresholdCVar,
+	TEXT("Sets the minimum volume for a wave instance to be considered active\n")
+	TEXT("Default is 0.0001 (-80 dB)"),
+	ECVF_Default);
+
 bool IsAudioPluginEnabled(EAudioPlugin PluginType)
 {
 	switch (PluginType)
@@ -893,7 +901,7 @@ bool FWaveInstance::IsPlaying() const
 	}
 
 	const float WaveInstanceVolume = Volume * VolumeMultiplier * GetDistanceAndOcclusionAttenuation() * GetDynamicVolume();
-	if (WaveInstanceVolume > UE_KINDA_SMALL_NUMBER)
+	if (WaveInstanceVolume > WaveInstanceMinVolumeThresholdCVar)
 	{
 		return true;
 	}
