[NetPhysics] Clear friction data on first resim frame, and keep friction turned on when physics prediction is enabled in project settings.
Also implement a runtime correction for actors using resimulation to correct the state while within the limit to not trigger a resimulation.

diff --git a/Engine/Source/Runtime/Engine/Private/PhysicsEngine/PhysicsReplication.cpp b/Engine/Source/Runtime/Engine/Private/PhysicsEngine/PhysicsReplication.cpp
--- a/Engine/Source/Runtime/Engine/Private/PhysicsEngine/PhysicsReplication.cpp
+++ b/Engine/Source/Runtime/Engine/Private/PhysicsEngine/PhysicsReplication.cpp
@@ -87,6 +87,18 @@ namespace PhysicsReplicationCVars
 	static int32 EnableDefaultReplication = 0;
 	static FAutoConsoleVariableRef CVarEnableDefaultReplication(TEXT("np2.EnableDefaultReplication"), EnableDefaultReplication, TEXT("Enable default replication in the networked physics prediction flow."));
 
+	namespace ResimulationCVars
+	{
+		static bool bRuntimeCorrectionEnabled = true;
+		static FAutoConsoleVariableRef CVarResimRuntimeCorrectionEnabled(TEXT("np2.Resim.RuntimeCorrectionEnabled"), bRuntimeCorrectionEnabled, TEXT("Apply runtime corrections while error is smalle enough not to trigger a resim."));
+
+		static float PosStabilityMultiplier = 0.5f;
+		static FAutoConsoleVariableRef CVarResimPosStabilityMultiplier(TEXT("np2.Resim.PosStabilityMultiplier"), PosStabilityMultiplier, TEXT("Recommended range between 0.0-1.0. Lower value means more stable positional corrections."));
+
+		static float RotStabilityMultiplier = 1.0f;
+		static FAutoConsoleVariableRef CVarResimRotStabilityMultiplier(TEXT("np2.Resim.RotStabilityMultiplier"), RotStabilityMultiplier, TEXT("Recommended range between 0.0-1.0. Lower value means more stable rotational corrections."));
+	}
+
 	namespace PredictiveInterpolationCVars
 	{
 		static float PosCorrectionTimeBase = 0.15f;
@@ -1406,62 +1418,95 @@ bool FPhysicsReplicationAsync::ResimulationReplication(Chaos::FPBDRigidParticleH
 	}
 
 	const int32 LocalFrame = Target.ServerFrame - Target.FrameOffset;
-
-	if (LocalFrame <= RewindData->CurrentFrame() && LocalFrame >= RewindData->GetEarliestFrame_Internal())
+	if (LocalFrame > RewindData->CurrentFrame() || LocalFrame < RewindData->GetEarliestFrame_Internal())
 	{
-		static constexpr Chaos::FFrameAndPhase::EParticleHistoryPhase RewindPhase = Chaos::FFrameAndPhase::EParticleHistoryPhase::PostPushData;
+		if (LocalFrame > 0)
+		{
+			UE_LOG(LogPhysics, Warning, TEXT("FPhysicsReplication::ApplyRigidBodyState target frame (%d) out of rewind data bounds (%d,%d)"), LocalFrame,
+				RewindData->GetEarliestFrame_Internal(), RewindData->CurrentFrame());
+		}
+		return true;
+	}
 
-		FAsyncPhysicsTimestamp TimeStamp;
-		TimeStamp.LocalFrame = RewindData->CurrentFrame();
+	bool bClearTarget = true;
 
-		const float ResimErrorThreshold = Chaos::FPhysicsSolverBase::ResimulationErrorThreshold();
+	static constexpr Chaos::FFrameAndPhase::EParticleHistoryPhase RewindPhase = Chaos::FFrameAndPhase::EParticleHistoryPhase::PostPushData;
+	FAsyncPhysicsTimestamp TimeStamp;
+	TimeStamp.LocalFrame = RewindData->CurrentFrame();
 
-		auto PastState = RewindData->GetPastStateAtFrame(*Handle, LocalFrame, RewindPhase);
+	const float ResimErrorThreshold = Chaos::FPhysicsSolverBase::ResimulationErrorThreshold();
+	const Chaos::FGeometryParticleState PastState = RewindData->GetPastStateAtFrame(*Handle, LocalFrame, RewindPhase);
 
-		const FVector ErrorOffset = (PastState.X() - Target.TargetState.Position);
-		const float ErrorDistance = ErrorOffset.Size();
-		const bool ShouldTriggerResim = ErrorDistance >= ResimErrorThreshold;
-		float ColorLerp = ShouldTriggerResim ? 1.0f : 0.0f;
+	const FVector ErrorOffset = (Target.TargetState.Position - PastState.X());
+	const float ErrorDistance = ErrorOffset.Size();
+	const bool ShouldTriggerResim = ErrorDistance >= ResimErrorThreshold;
 
 #if !(UE_BUILD_SHIPPING || UE_BUILD_TEST)
 
-		if (Chaos::FPhysicsSolverBase::CanDebugNetworkPhysicsPrediction())
-		{
-			UE_LOG(LogTemp, Log, TEXT("Apply Rigid body state at local frame %d with offset = %d"), LocalFrame, Target.FrameOffset);
-			UE_LOG(LogTemp, Log, TEXT("Particle Position Error = %f | Should Trigger Resim = %s | Server Frame = %d | Client Frame = %d"), ErrorDistance, (ShouldTriggerResim ? TEXT("True") : TEXT("False")), Target.ServerFrame, LocalFrame);
-			UE_LOG(LogTemp, Log, TEXT("Particle Target Position = %s | Current Position = %s"), *Target.TargetState.Position.ToString(), *PastState.X().ToString());
-			UE_LOG(LogTemp, Log, TEXT("Particle Target Velocity = %s | Current Velocity = %s"), *Target.TargetState.LinVel.ToString(), *PastState.V().ToString());
-			UE_LOG(LogTemp, Log, TEXT("Particle Target Quaternion = %s | Current Quaternion = %s"), *Target.TargetState.Quaternion.ToString(), *PastState.R().ToString());
-			UE_LOG(LogTemp, Log, TEXT("Particle Target Omega = %s | Current Omega= %s"), *Target.TargetState.AngVel.ToString(), *PastState.W().ToString());
+	if (Chaos::FPhysicsSolverBase::CanDebugNetworkPhysicsPrediction())
+	{
+		UE_LOG(LogTemp, Log, TEXT("Apply Rigid body state at local frame %d with offset = %d"), LocalFrame, Target.FrameOffset);
+		UE_LOG(LogTemp, Log, TEXT("Particle Position Error = %f | Should Trigger Resim = %s | Server Frame = %d | Client Frame = %d"), ErrorDistance, (ShouldTriggerResim ? TEXT("True") : TEXT("False")), Target.ServerFrame, LocalFrame);
+		UE_LOG(LogTemp, Log, TEXT("Particle Target Position = %s | Current Position = %s"), *Target.TargetState.Position.ToString(), *PastState.X().ToString());
+		UE_LOG(LogTemp, Log, TEXT("Particle Target Velocity = %s | Current Velocity = %s"), *Target.TargetState.LinVel.ToString(), *PastState.V().ToString());
+		UE_LOG(LogTemp, Log, TEXT("Particle Target Quaternion = %s | Current Quaternion = %s"), *Target.TargetState.Quaternion.ToString(), *PastState.R().ToString());
+		UE_LOG(LogTemp, Log, TEXT("Particle Target Omega = %s | Current Omega= %s"), *Target.TargetState.AngVel.ToString(), *PastState.W().ToString());
+
+		{ // DrawDebug
+			static constexpr float BoxSize = 5.0f;
+			const float ColorLerp = ShouldTriggerResim ? 1.0f : 0.0f;
+			const FColor DebugColor = FLinearColor::LerpUsingHSV(FLinearColor::Green, FLinearColor::Red, ColorLerp).ToFColor(false);
+
+			Chaos::FDebugDrawQueue::GetInstance().DrawDebugBox(Target.TargetState.Position, FVector(BoxSize, BoxSize, BoxSize), Target.TargetState.Quaternion, FColor::Orange, true, CharacterMovementCVars::NetCorrectionLifetime, 0, 1.0f);
+			Chaos::FDebugDrawQueue::GetInstance().DrawDebugBox(PastState.X(), FVector(6, 6, 6), PastState.R(), DebugColor, true, CharacterMovementCVars::NetCorrectionLifetime, 0, 1.0f);
+
+			Chaos::FDebugDrawQueue::GetInstance().DrawDebugDirectionalArrow(PastState.X(), Target.TargetState.Position, 5.0f, FColor::MakeRandomSeededColor(LocalFrame), true, CharacterMovementCVars::NetCorrectionLifetime, 0, 0.5f);
+		}
+	}
+#endif
+
+	if (Target.TickCount == 0 && LocalFrame > RewindData->GetBlockedResimFrame() && ShouldTriggerResim)
+	{
+		// Trigger resimulation
+		RigidsSolver->GetEvolution()->GetIslandManager().SetParticleResimFrame(Handle, LocalFrame);
 
-			{ // DrawDebug
-				static constexpr float BoxSize = 5.0f;
-				const FColor DebugColor = FLinearColor::LerpUsingHSV(FLinearColor::Green, FLinearColor::Red, ColorLerp).ToFColor(false);
+		int32 ResimFrame = RewindData->GetResimFrame();
+		ResimFrame = (ResimFrame == INDEX_NONE) ? LocalFrame : FMath::Min(ResimFrame, LocalFrame);
+		RewindData->SetResimFrame(ResimFrame);
+	}
+	else if (PhysicsReplicationCVars::ResimulationCVars::bRuntimeCorrectionEnabled)
+	{
+		const int32 NumPredictedFrames = RigidsSolver->GetCurrentFrame() - LocalFrame;
+
+		if (NumPredictedFrames > 0)
+		{
+			// Calculate correction to position
+			const float CorrectionAmountX = PhysicsReplicationCVars::ResimulationCVars::PosStabilityMultiplier / NumPredictedFrames; // Same result as (ErrorOffset / NumPredictedFrames) * PosStabilityMultiplier
+			const FVector CorrectedX = Handle->X() + (ErrorOffset * CorrectionAmountX);
 
-				Chaos::FDebugDrawQueue::GetInstance().DrawDebugBox(Target.TargetState.Position, FVector(BoxSize, BoxSize, BoxSize), Target.TargetState.Quaternion, FColor::Orange, true, CharacterMovementCVars::NetCorrectionLifetime, 0, 1.0f);
-				Chaos::FDebugDrawQueue::GetInstance().DrawDebugBox(PastState.X(), FVector(6, 6, 6), PastState.R(), DebugColor, true, CharacterMovementCVars::NetCorrectionLifetime, 0, 1.0f);
+			// Calculate correction to rotation
+			const float CorrectionAmountR = (1.f / NumPredictedFrames) * PhysicsReplicationCVars::ResimulationCVars::RotStabilityMultiplier;
+			const FQuat InvCurrentQuat = PastState.R().Inverse();
+			const FQuat DeltaQuat = Target.TargetState.Quaternion * InvCurrentQuat;
+			const FQuat TargetCorrectionR = Handle->R() * DeltaQuat;
+			const FQuat CorrectedR = FQuat::Slerp(Handle->R(), TargetCorrectionR, CorrectionAmountR);
 
-				Chaos::FDebugDrawQueue::GetInstance().DrawDebugDirectionalArrow(PastState.X(), Target.TargetState.Position, 5.0f, FColor::Green, true, CharacterMovementCVars::NetCorrectionLifetime, 0, 0.5f);
+#if !(UE_BUILD_SHIPPING || UE_BUILD_TEST)
+			if (Chaos::FPhysicsSolverBase::CanDebugNetworkPhysicsPrediction())
+			{
+				Chaos::FDebugDrawQueue::GetInstance().DrawDebugDirectionalArrow(Handle->X(), CorrectedX, 5.0f, FColor::MakeRandomSeededColor(LocalFrame), true, CharacterMovementCVars::NetCorrectionLifetime, 0, 0.5f);
 			}
-		}
 #endif
+			// Apply correction
+			Handle->SetX(CorrectedX);
+			Handle->SetR(CorrectedR);
 
-		if (ShouldTriggerResim)
-		{
-			RigidsSolver->GetEvolution()->GetIslandManager().SetParticleResimFrame(Handle, LocalFrame);
-
-			int32 ResimFrame = RewindData->GetResimFrame();
-			ResimFrame = (ResimFrame == INDEX_NONE) ? LocalFrame : FMath::Min(ResimFrame, LocalFrame);
-			RewindData->SetResimFrame(ResimFrame);
+			// Keep target for NumPredictedFrames time to perform runtime corrections with until a new target is received
+			bClearTarget = Target.TickCount >= NumPredictedFrames;
 		}
 	}
-	else if (LocalFrame > 0)
-	{
-		UE_LOG(LogPhysics, Warning, TEXT("FPhysicsReplication::ApplyRigidBodyState target frame (%d) out of rewind data bounds (%d,%d)"), LocalFrame,
-			RewindData->GetEarliestFrame_Internal(), RewindData->CurrentFrame());
-	}
 
-	return true;
+	return bClearTarget;
 }
 
 FName FPhysicsReplicationAsync::GetFNameForStatId() const
diff --git a/Engine/Source/Runtime/Experimental/Chaos/Private/Chaos/Collision/PBDCollisionConstraint.cpp b/Engine/Source/Runtime/Experimental/Chaos/Private/Chaos/Collision/PBDCollisionConstraint.cpp
--- a/Engine/Source/Runtime/Experimental/Chaos/Private/Chaos/Collision/PBDCollisionConstraint.cpp
+++ b/Engine/Source/Runtime/Experimental/Chaos/Private/Chaos/Collision/PBDCollisionConstraint.cpp
@@ -38,9 +38,6 @@ namespace Chaos
 	FRealSingle Chaos_GBFCharacteristicTimeRatio = 1.0f;
 	FAutoConsoleVariableRef CVarChaos_GBFCharacteristicTimeRatio(TEXT("p.Chaos.Collision.GBFCharacteristicTimeRatio"), Chaos_GBFCharacteristicTimeRatio, TEXT("The ratio between characteristic time and Dt"));
 
-	bool bChaos_Manifold_EnableGjkWarmStart = true;
-	FAutoConsoleVariableRef CVarChaos_Manifold_EnableGjkWarmStart(TEXT("p.Chaos.Collision.Manifold.EnableGjkWarmStart"), bChaos_Manifold_EnableGjkWarmStart, TEXT(""));
-
 	bool bChaos_Manifold_EnableFrictionRestore = true;
 	FAutoConsoleVariableRef CVarChaos_Manifold_EnableFrictionRestore(TEXT("p.Chaos.Collision.Manifold.EnableFrictionRestore"), bChaos_Manifold_EnableFrictionRestore, TEXT(""));
 	
@@ -1054,8 +1051,7 @@ namespace Chaos
 	{
 		int32 MatchIndex = INDEX_NONE;
 
-		const bool bEnableNetworkPhysicsResim = FPhysicsSolverBase::IsNetworkPhysicsPredictionEnabled();
-		if (bChaos_Manifold_EnableFrictionRestore && !bEnableNetworkPhysicsResim)
+		if (bChaos_Manifold_EnableFrictionRestore)
 		{
 			const FManifoldPoint& ManifoldPoint = ManifoldPoints[ManifoldPointIndex];
 			if (!ManifoldPoint.Flags.bDisabled)
diff --git a/Engine/Source/Runtime/Experimental/Chaos/Private/Chaos/CollisionOneShotManifolds.cpp b/Engine/Source/Runtime/Experimental/Chaos/Private/Chaos/CollisionOneShotManifolds.cpp
--- a/Engine/Source/Runtime/Experimental/Chaos/Private/Chaos/CollisionOneShotManifolds.cpp
+++ b/Engine/Source/Runtime/Experimental/Chaos/Private/Chaos/CollisionOneShotManifolds.cpp
@@ -69,6 +69,9 @@ namespace Chaos
 	FAutoConsoleVariableRef CVarChaos_Collision_EnableManifoldReplace(TEXT("p.Chaos.Collision.EnableManifoldGJKReplace"), bChaos_Collision_EnableManifoldGJKReplace, TEXT(""));
 	FAutoConsoleVariableRef CVarChaos_Collision_EnableManifoldInject(TEXT("p.Chaos.Collision.EnableManifoldGJKInject"), bChaos_Collision_EnableManifoldGJKInject, TEXT(""));
 
+	bool bChaos_Manifold_EnableGjkWarmStart = true;
+	FAutoConsoleVariableRef CVarChaos_Manifold_EnableGjkWarmStart(TEXT("p.Chaos.Collision.Manifold.EnableGjkWarmStart"), bChaos_Manifold_EnableGjkWarmStart, TEXT(""));
+
 	// See GJKContactPointMargin for comments on why these matter
 	// LWC_TODO: These needs to be a larger values for float builds (1.e-3f)
 	FRealSingle Chaos_Collision_GJKEpsilon = 1.e-6f;
@@ -794,9 +797,7 @@ namespace Chaos
 			const FReal Margin2 = Constraint.GetCollisionMargin1();
 			const FRigidTransform3 Convex2ToConvex1Transform = Convex2Transform.GetRelativeTransformNoScale(Convex1Transform);
 
-			const bool bEnableNetworkPhysicsResim = FPhysicsSolverBase::IsNetworkPhysicsPredictionEnabled();
-			
-			if (bEnableNetworkPhysicsResim)
+			if (!bChaos_Manifold_EnableGjkWarmStart)
 			{
 				Constraint.GetGJKWarmStartData().Reset();
 			}
diff --git a/Engine/Source/Runtime/Experimental/Chaos/Private/Chaos/PBDRigidsEvolutionGBF.cpp b/Engine/Source/Runtime/Experimental/Chaos/Private/Chaos/PBDRigidsEvolutionGBF.cpp
--- a/Engine/Source/Runtime/Experimental/Chaos/Private/Chaos/PBDRigidsEvolutionGBF.cpp
+++ b/Engine/Source/Runtime/Experimental/Chaos/Private/Chaos/PBDRigidsEvolutionGBF.cpp
@@ -693,6 +693,7 @@ void FPBDRigidsEvolutionGBF::ResetCollisions()
 	{
 		if (Collision != nullptr)
 		{
+			Collision->GetContact().ResetSavedManifoldPoints();
 			Collision->GetContact().GetGJKWarmStartData().Reset();
 		}
 	}
diff --git a/Engine/Source/Runtime/Experimental/Chaos/Private/RewindData.cpp b/Engine/Source/Runtime/Experimental/Chaos/Private/RewindData.cpp
--- a/Engine/Source/Runtime/Experimental/Chaos/Private/RewindData.cpp
+++ b/Engine/Source/Runtime/Experimental/Chaos/Private/RewindData.cpp
@@ -256,6 +256,9 @@ void FRewindData::ApplyTargets(const int32 Frame, const bool bResetSimulation)
 	}
 }
 
+CHAOS_API bool bResimAllowRewindToResimulatedFrames = false;
+FAutoConsoleVariableRef CVarResimAllowRewindToResimulatedFrames(TEXT("p.Resim.AllowRewindToResimulatedFrames"), bResimAllowRewindToResimulatedFrames, TEXT("Allow rewinding back to a frame that was previously part of a resimulation. If a resimulation is performed between frame 100-110, allow a new resim from 105-115 if needed, else next resim will be able to start from frame 111."));
+
 bool FRewindData::RewindToFrame(int32 Frame)
 {
 	QUICK_SCOPE_CYCLE_COUNTER(RewindToFrame);
@@ -284,6 +287,11 @@ bool FRewindData::RewindToFrame(int32 Frame)
 	FFrameAndPhase RewindFrameAndPhase{ Frame, FFrameAndPhase::PostPushData };
 	FFrameAndPhase CurFrameAndPhase{ CurFrame, FFrameAndPhase::PrePushData };
 
+	if (!bResimAllowRewindToResimulatedFrames)
+	{
+		BlockResimFrame = CurFrame;
+	}
+
 	auto RewindHelper = [RewindFrameAndPhase, CurFrameAndPhase, this](auto Obj, bool bResimAsFollower, auto& Property, const auto& RewindFunc) -> bool
 	{
 		if (bResimAsFollower)
@@ -838,10 +846,10 @@ void FRewindData::ExtendHistoryWithFrame(const int32 Frame)
 	FramesSaved = FMath::Max(CurFrame - Frame+1, FramesSaved);
 }
 
-CHAOS_API bool ResimFrameValidationLeniency = true;
-FAutoConsoleVariableRef CVarResimFrameValidationLeniency(TEXT("p.ResimFrameValidationLeniency"), ResimFrameValidationLeniency, TEXT("Lenient resim frame validation finds a frame to resim from where the particle that is triggering the resim has a valid target. Setting this to false will require all replicated particles to have a valid target and both valid component input and state history."));
-CHAOS_API bool ResimIncompleteHistory = true;
-FAutoConsoleVariableRef CVarResimIncompleteHistory(TEXT("p.ResimIncompleteHistory"), ResimIncompleteHistory, TEXT("If a valid resim frame can't be found, use the requested resim frame and perform a resimulation with incomplete data."));
+CHAOS_API bool bResimFrameValidationLeniency = true;
+FAutoConsoleVariableRef CVarResimFrameValidationLeniency(TEXT("p.Resim.FrameValidationLeniency"), bResimFrameValidationLeniency, TEXT("Lenient resim frame validation finds a frame to resim from where the particle that is triggering the resim has a valid target. Setting this to false will require all replicated particles to have a valid target and both valid component input and state history."));
+CHAOS_API bool bResimIncompleteHistory = true;
+FAutoConsoleVariableRef CVarResimIncompleteHistory(TEXT("p.Resim.IncompleteHistory"), bResimIncompleteHistory, TEXT("If a valid resim frame can't be found, use the requested resim frame and perform a resimulation with incomplete data."));
 
 int32 FRewindData::FindValidResimFrame(const int32 RequestedFrame)
 {
@@ -872,7 +880,7 @@ int32 FRewindData::FindValidResimFrame(const int32 RequestedFrame)
 			const bool bResimAsFollower = DirtyParticleInfo.bResimAsFollower;
 
 			// If the particle is not marked for resimulation, don't bother checking for valid target states.
-			if (ResimFrameValidationLeniency && Solver->GetEvolution()->GetIslandManager().GetParticleResimFrame(Handle) == INDEX_NONE)
+			if (bResimFrameValidationLeniency && Solver->GetEvolution()->GetIslandManager().GetParticleResimFrame(Handle) == INDEX_NONE)
 			{
 				continue;
 			}
@@ -895,7 +903,7 @@ int32 FRewindData::FindValidResimFrame(const int32 RequestedFrame)
 			}
 		}
 
-		if (!ResimFrameValidationLeniency)
+		if (!bResimFrameValidationLeniency)
 		{
 			if (bHasTargetHistory)
 			{
@@ -929,7 +937,7 @@ int32 FRewindData::FindValidResimFrame(const int32 RequestedFrame)
 	
 	if (!bHasTargetHistory)
 	{
-		ValidFrame = ResimIncompleteHistory ? RequestedFrame : INDEX_NONE;
+		ValidFrame = bResimIncompleteHistory ? RequestedFrame : INDEX_NONE;
 
 #if DEBUG_REWIND_DATA
 		UE_LOG(LogTemp, Warning, TEXT("COMMON | PT | FindValidResimFrame | No valid resim frame found | RequestedFrame: %d | ValidFrame: %d | EarliestFrame: %d | HasTargetHistory: %d | EarliestHistoryFrame: %d | CurrentFrame: %d | FramesSaved: %d, ResimFrameValidationLeniency:%d"), RequestedFrame, ValidFrame, EarliestFrame, bHasTargetHistory, GetEarliestFrame_Internal(), CurrentFrame(), FramesSaved, ResimFrameValidationLeniency);
diff --git a/Engine/Source/Runtime/Experimental/Chaos/Public/RewindData.h b/Engine/Source/Runtime/Experimental/Chaos/Public/RewindData.h
--- a/Engine/Source/Runtime/Experimental/Chaos/Public/RewindData.h
+++ b/Engine/Source/Runtime/Experimental/Chaos/Public/RewindData.h
@@ -1458,6 +1458,9 @@ public:
 	/** This blocks any future resimulation to rewind back past the frame this is called on */
 	void BlockResim();
 
+	/** Get the latest frame resim has been blocked from rewinding past */
+	int32 GetBlockedResimFrame() { return BlockResimFrame; }
+
 private:
 	friend class FPBDRigidsSolver;
 
