Include non-component replicated subobjects when flushing network dormancy.
Fixes an issue where replicated fast arrays may get redundant change notifications.
Disable with cvar net.EnableFlushDormantSubObjects. Disable checking subobject conditions with cvar net.EnableFlushDormantSubObjectsCheckConditions.
Added test case to EngineTest.

diff --git a/Engine/Source/Runtime/Engine/Classes/Engine/ActorChannel.h b/Engine/Source/Runtime/Engine/Classes/Engine/ActorChannel.h
--- a/Engine/Source/Runtime/Engine/Classes/Engine/ActorChannel.h
+++ b/Engine/Source/Runtime/Engine/Classes/Engine/ActorChannel.h
@@ -28,7 +28,11 @@ class FOutBunch;
 class UNetConnection;
 struct FObjectKey;
 
-namespace UE::Net { struct FSubObjectRegistry;  }
+namespace UE::Net
+{
+	struct FSubObjectRegistry;
+	class FNetConditionGroupManager;
+}
 
 enum class ESetChannelActorFlags : uint32
 {
@@ -367,6 +371,13 @@ public:
 
 	ENGINE_API virtual void AddedToChannelPool() override;
 
+	static ENGINE_API bool CanSubObjectReplicateToClient(
+		APlayerController* PlayerController,
+		ELifetimeCondition NetCondition,
+		FObjectKey SubObjectKey,
+		const TStaticBitArray<COND_Max>& ConditionMap,
+		const UE::Net::FNetConditionGroupManager& ConditionGroupManager);
+
 protected:
 
 	/** Attempts to find a valid, non-dormant replicator for the given object. */
@@ -455,8 +466,6 @@ private:
 	/** Replicate a list of subobjects */
 	ENGINE_API bool WriteSubObjects(UObject* SubObjectOwner, const UE::Net::FSubObjectRegistry& SubObjectList, FOutBunch& Bunch, FReplicationFlags RepFlags, const TStaticBitArray<COND_Max>& ConditionMap);
 
-	ENGINE_API bool CanSubObjectReplicateToClient(ELifetimeCondition NetCondition, FObjectKey SubObjectKey, const TStaticBitArray<COND_Max>& ConditionMap) const;
-
 	ENGINE_API bool ValidateReplicatedSubObjects();
 
 	ENGINE_API void TestLegacyReplicateSubObjects(UActorComponent* ReplicatedComponent, FOutBunch& Bunch, FReplicationFlags RepFlags);
diff --git a/Engine/Source/Runtime/Engine/Private/DataChannel.cpp b/Engine/Source/Runtime/Engine/Private/DataChannel.cpp
--- a/Engine/Source/Runtime/Engine/Private/DataChannel.cpp
+++ b/Engine/Source/Runtime/Engine/Private/DataChannel.cpp
@@ -3810,13 +3810,18 @@ bool UActorChannel::UpdateDeletedSubObjects(FOutBunch& Bunch)
 	return bWroteSomethingImportant;
 }
 
-bool UActorChannel::CanSubObjectReplicateToClient(ELifetimeCondition NetCondition, FObjectKey SubObjectKey, const TStaticBitArray<COND_Max>& ConditionMap) const
+bool UActorChannel::CanSubObjectReplicateToClient(
+	APlayerController* PlayerController,
+	ELifetimeCondition NetCondition,
+	FObjectKey SubObjectKey,
+	const TStaticBitArray<COND_Max>& ConditionMap,
+	const UE::Net::FNetConditionGroupManager& ConditionGroupManager)
 {
 	bool bCanReplicateToClient = ConditionMap[NetCondition];
 
 	if (NetCondition == COND_NetGroup)
 	{
-		TArrayView<const FName> NetGroups = DataChannelInternal::CachedNetworkSubsystem->GetNetConditionGroupManager().GetSubObjectNetConditionGroups(SubObjectKey);
+		TArrayView<const FName> NetGroups = ConditionGroupManager.GetSubObjectNetConditionGroups(SubObjectKey);
 
 		for (int i=0; i < NetGroups.Num() && !bCanReplicateToClient; ++i)
 		{
@@ -3832,7 +3837,7 @@ bool UActorChannel::CanSubObjectReplicateToClient(ELifetimeCondition NetConditio
 			}
 			else
 			{
-				bCanReplicateToClient = Connection->PlayerController->IsMemberOfNetConditionGroup(NetGroup);
+				bCanReplicateToClient = PlayerController->IsMemberOfNetConditionGroup(NetGroup);
 			}
 		}
 	}
@@ -3871,7 +3876,7 @@ bool UActorChannel::ReplicateRegisteredSubObjects(FOutBunch& Bunch, FReplication
 	// Now the replicated actor components
 	for( const FReplicatedComponentInfo& RepComponentInfo : FSubObjectRegistryGetter::GetReplicatedComponents(Actor) )
 	{
-		if (CanSubObjectReplicateToClient(RepComponentInfo.NetCondition, RepComponentInfo.Key, ConditionMap))
+		if (CanSubObjectReplicateToClient(Connection->PlayerController, RepComponentInfo.NetCondition, RepComponentInfo.Key, ConditionMap, DataChannelInternal::CachedNetworkSubsystem->GetNetConditionGroupManager()))
 		{
 			UActorComponent* ReplicatedComponent = RepComponentInfo.Component;
 
@@ -4055,7 +4060,7 @@ bool UActorChannel::ReplicateSubobject(UActorComponent* ReplicatedComponent, FOu
 bool UActorChannel::WriteComponentSubObjects(UActorComponent* ReplicatedComponent, FOutBunch& Bunch, FReplicationFlags RepFlags, const TStaticBitArray<COND_Max>& ConditionMap)
 {
 	using namespace UE::Net;
-	if (const FSubObjectRegistry* SubObjectList = FSubObjectRegistryGetter::GetSubObjectsOfActorCompoment(Actor, ReplicatedComponent))
+	if (const FSubObjectRegistry* SubObjectList = FSubObjectRegistryGetter::GetSubObjectsOfActorComponent(Actor, ReplicatedComponent))
 	{
 		return WriteSubObjects(ReplicatedComponent, *SubObjectList, Bunch, RepFlags, ConditionMap);
 	}
@@ -4095,7 +4100,7 @@ bool UActorChannel::WriteSubObjects(UObject* SubObjectOwner, const UE::Net::FSub
 		{
 			checkf(IsValid(SubObjectToReplicate), TEXT("Found invalid subobject (%s) registered in %s::%s"), *GetNameSafe(SubObjectToReplicate), *Actor->GetName(), *GetNameSafe(SubObjectOwner));
 
-			if (CanSubObjectReplicateToClient(SubObjectInfo.NetCondition, SubObjectInfo.Key, ConditionMap))
+			if (CanSubObjectReplicateToClient(Connection->PlayerController, SubObjectInfo.NetCondition, SubObjectInfo.Key, ConditionMap, DataChannelInternal::CachedNetworkSubsystem->GetNetConditionGroupManager()))
 			{
 #if UE_NET_REPACTOR_NAME_DEBUG
 				if (bSubObjectNameDebug)
diff --git a/Engine/Source/Runtime/Engine/Private/Net/NetSubObjectRegistryGetter.h b/Engine/Source/Runtime/Engine/Private/Net/NetSubObjectRegistryGetter.h
--- a/Engine/Source/Runtime/Engine/Private/Net/NetSubObjectRegistryGetter.h
+++ b/Engine/Source/Runtime/Engine/Private/Net/NetSubObjectRegistryGetter.h
@@ -25,7 +25,7 @@ public:
 		return InActor->ReplicatedSubObjects;
 	}
 
-	static const FSubObjectRegistry* GetSubObjectsOfActorCompoment(AActor* InActor, UActorComponent* InActorComp)
+	static const FSubObjectRegistry* GetSubObjectsOfActorComponent(AActor* InActor, UActorComponent* InActorComp)
 	{
 		UE::Net::FReplicatedComponentInfo* ComponentInfo = InActor->ReplicatedComponentsInfo.FindByKey(InActorComp);
 		return ComponentInfo ? &(ComponentInfo->SubObjects) : nullptr;
diff --git a/Engine/Source/Runtime/Engine/Private/NetConnection.cpp b/Engine/Source/Runtime/Engine/Private/NetConnection.cpp
--- a/Engine/Source/Runtime/Engine/Private/NetConnection.cpp
+++ b/Engine/Source/Runtime/Engine/Private/NetConnection.cpp
@@ -41,6 +41,9 @@
 #include "LevelUtils.h"
 #include "Net/RPCDoSDetection.h"
 #include "Net/NetConnectionFaultRecovery.h"
+#include "Net/NetSubObjectRegistryGetter.h"
+#include "Net/RepLayout.h"
+#include "Net/Subsystems/NetworkSubsystem.h"
 #if UE_WITH_IRIS
 #include "Iris/IrisConfig.h"
 #include "Iris/ReplicationSystem/ReplicationSystem.h"
@@ -203,6 +206,48 @@ namespace UE::Net::Connection::Private
 
 	int32 bTrackFlushedDormantObjects = true;
 	FAutoConsoleVariableRef CVarNetTrackFlushedDormantObjects(TEXT("net.TrackFlushedDormantObjects"), bTrackFlushedDormantObjects, TEXT("If enabled, track dormant subobjects when dormancy is flushed, so they can be properly deleted if destroyed prior to the next ReplicateActor."));
+
+	int32 bEnableFlushDormantSubObjects = true;
+	FAutoConsoleVariableRef CVarNetFlushDormantSubObjects(TEXT("net.EnableFlushDormantSubObjects"), bEnableFlushDormantSubObjects, TEXT("If enabled, FlushNetDormancy will flush replicated subobjects in addition to replicated components. Only applies to objects using the replicated subobject list."));
+
+	int32 bEnableFlushDormantSubObjectsCheckConditions = true;
+	FAutoConsoleVariableRef CVarNetFlushDormantSubObjectsCheckConditions(TEXT("net.EnableFlushDormantSubObjectsCheckConditions"), bEnableFlushDormantSubObjectsCheckConditions, TEXT("If enabled, when net.EnableFlushDormantSubObjects is also true a dormancy flush will also check replicated subobject conditions"));
+
+	// Tracking for dormancy-flushed subobjects for correct deletion, see UE-77163
+	void TrackFlushedSubObject(FDormantObjectMap& InOutFlushedGuids, UObject* FlushedObject, const TSharedPtr<FNetGUIDCache>& GuidCache)
+	{
+		if (Connection::Private::bTrackFlushedDormantObjects)
+		{
+			// Searching for the guid because the value on the replicator built in FlushDormancyForObject can be invalid until the next replication
+			// We can then safely ignore any object that still has no guid, since it won't have ever been replicated
+			FNetworkGUID ObjectNetGUID = GuidCache->GetNetGUID(FlushedObject);
+			if (ObjectNetGUID.IsValid())
+			{
+				InOutFlushedGuids.Add(ObjectNetGUID, FlushedObject);
+			}
+		}
+	}
+
+	// Flushes dormancy for registered subobjects of either an actor or component
+	void FlushDormancyForSubObjects(UNetConnection* Connection, AActor* Actor, const UE::Net::FSubObjectRegistry& SubObjects, FDormantObjectMap& InOutFlushedGuids, const TStaticBitArray<COND_Max>& ConditionMap, const FNetConditionGroupManager* NetConditionGroupManager)
+	{
+		QUICK_SCOPE_CYCLE_COUNTER(STAT_NetConnection_FlushDormancyForSubObjects);
+
+		for (const FSubObjectRegistry::FEntry& SubObjectInfo : SubObjects.GetRegistryList())
+		{
+			UObject* SubObject = SubObjectInfo.GetSubObject();
+			if (ensureMsgf(IsValid(SubObject), TEXT("Found invalid subobject (%s) registered in %s"), *GetNameSafe(SubObject), *Actor->GetName()))
+			{
+				if (!bEnableFlushDormantSubObjectsCheckConditions ||
+					!NetConditionGroupManager ||
+					UActorChannel::CanSubObjectReplicateToClient(Connection->PlayerController, SubObjectInfo.NetCondition, SubObjectInfo.Key, ConditionMap, *NetConditionGroupManager))
+				{
+					Connection->FlushDormancyForObject(Actor, SubObject);
+					TrackFlushedSubObject(InOutFlushedGuids, SubObject, Connection->Driver->GuidCache);
+				}
+			}
+		}
+	}
 }
 
 // ChannelRecord Implementation
@@ -4896,43 +4941,86 @@ void UNetConnection::ClearDormantReplicatorsReference()
 
 void UNetConnection::FlushDormancy(AActor* Actor)
 {
-	UE_LOG( LogNetDormancy, Verbose, TEXT( "FlushDormancy: %s. Connection: %s" ), *Actor->GetName(), *GetName() );
+	using namespace UE::Net;
+	using namespace Connection::Private;
+
+	UE_LOG(LogNetDormancy, Verbose, TEXT( "FlushDormancy: %s. Connection: %s" ), *Actor->GetName(), *GetName());
 	
-	if ( Driver->GetNetworkObjectList().MarkActive( Actor, this, Driver ) )
+	if (Driver->GetNetworkObjectList().MarkActive(Actor, this, Driver))
 	{
-		FlushDormancyForObject( Actor, Actor );
+		FlushDormancyForObject(Actor, Actor);
 
-		const TArray<UActorComponent*>& ReplicatedComponents = Actor->GetReplicatedComponents();
+		const FSubObjectRegistry& ActorSubObjects = FSubObjectRegistryGetter::GetSubObjects(Actor);
+		const TArray<FReplicatedComponentInfo>& ReplicatedComponents = FSubObjectRegistryGetter::GetReplicatedComponents(Actor);
 
-		UE::Net::FDormantObjectMap FlushedGuids;
-		if (UE::Net::Connection::Private::bTrackFlushedDormantObjects)
+		FDormantObjectMap FlushedGuids;
+		if (Connection::Private::bTrackFlushedDormantObjects)
 		{
-			FlushedGuids.Reserve(ReplicatedComponents.Num());
+			// This doesn't reserve space for subobjects of components, but avoids iterating the component list twice.
+			FlushedGuids.Reserve(ReplicatedComponents.Num() + ActorSubObjects.GetRegistryList().Num());
 		}
 
-		// TODO: Is this set of objects sufficient? Should we query the dormancy map from the connection instead?
-		for (UActorComponent* ActorComp : ReplicatedComponents)
+		TStaticBitArray<COND_Max> ConditionMap;
+		const FNetConditionGroupManager* NetConditionGroupManager = nullptr;
+
+		if (bEnableFlushDormantSubObjectsCheckConditions)
 		{
-			if (ActorComp && ActorComp->GetIsReplicated())
-			{
-				FlushDormancyForObject(Actor, ActorComp);
+			// Fill in the flags used by conditional subobjects
+			FReplicationFlags RepFlags;
+		
+			// Since a dormant object won't necessarily have an FObjectReplicator or channel, we can't check for
+			// initial replication in the normal way. So always flush for NetInitial to be safe, it may create
+			// an extra replicator but the condition will be checked again before the object is actually replicated.
+			RepFlags.bNetInitial = true;
+
+			UNetConnection* OwningConnection = Actor->GetNetConnection();
+			RepFlags.bNetOwner = (OwningConnection == this || (OwningConnection != nullptr && OwningConnection->IsA(UChildConnection::StaticClass()) && ((UChildConnection*)OwningConnection)->Parent == this));
+		
+			RepFlags.bNetSimulated = (Actor->GetRemoteRole() == ROLE_SimulatedProxy);
+			RepFlags.bRepPhysics = Actor->GetReplicatedMovement().bRepPhysics;
+			RepFlags.bReplay = bReplay;
+
+			ConditionMap = FSendingRepState::BuildConditionMapFromRepFlags(RepFlags);
+			
+			const UWorld* const World = Actor->GetWorld();
+			const UNetworkSubsystem* const NetworkSubsystem = World ? World->GetSubsystem<UNetworkSubsystem>() : nullptr;
+			NetConditionGroupManager = NetworkSubsystem ? &NetworkSubsystem->GetNetConditionGroupManager() : nullptr;
+			ensureMsgf(NetConditionGroupManager, TEXT("UNetConnection::FlushDormancy: couldn't find a NetConditionGroupManager for %s."), *Actor->GetName());
+		}
 
-				if (UE::Net::Connection::Private::bTrackFlushedDormantObjects)
+		if (bEnableFlushDormantSubObjects)
+		{
+			FlushDormancyForSubObjects(this, Actor, ActorSubObjects, FlushedGuids, ConditionMap, NetConditionGroupManager);
+		}
+
+		for (const FReplicatedComponentInfo& RepComponentInfo : ReplicatedComponents)
+		{
+			UActorComponent* ActorComp = RepComponentInfo.Component;
+
+			if (ensureMsgf(IsValid(ActorComp), TEXT("Found invalid replicated component (%s) registered in %s"), *GetNameSafe(ActorComp), *Actor->GetName()))
+			{
+				if (!bEnableFlushDormantSubObjectsCheckConditions ||
+					!NetConditionGroupManager ||
+					UActorChannel::CanSubObjectReplicateToClient(PlayerController, RepComponentInfo.NetCondition, RepComponentInfo.Key, ConditionMap, *NetConditionGroupManager))
 				{
-					// Searching for the guid because the value on the replicator built in FlushDormancyForObject can be invalid until the next replication
-					// We can then safely ignore any object that still has no guid, since it won't have ever been replicated
-					FNetworkGUID ObjectNetGUID = Driver->GuidCache->GetNetGUID(ActorComp);
-					if (ObjectNetGUID.IsValid())
+					FlushDormancyForObject(Actor, ActorComp);
+					TrackFlushedSubObject(FlushedGuids, ActorComp, Driver->GuidCache);
+
+					if (bEnableFlushDormantSubObjects)
 					{
-						FlushedGuids.Add(ObjectNetGUID, ActorComp);
+						const FSubObjectRegistry* const ComponentSubObjects = FSubObjectRegistryGetter::GetSubObjectsOfActorComponent(Actor, ActorComp);
+						if (ComponentSubObjects)
+						{
+							FlushDormancyForSubObjects(this, Actor, *ComponentSubObjects, FlushedGuids, ConditionMap, NetConditionGroupManager);
+						}
 					}
 				}
 			}
 		}
 
-		if (UE::Net::Connection::Private::bTrackFlushedDormantObjects && !FlushedGuids.IsEmpty())
+		if (bTrackFlushedDormantObjects && !FlushedGuids.IsEmpty())
 		{
-			UE::Net::FDormantObjectMap& DormantObjects = DormantReplicatorSet.FindOrAddFlushedObjectsForActor(Actor);
+			FDormantObjectMap& DormantObjects = DormantReplicatorSet.FindOrAddFlushedObjectsForActor(Actor);
 			DormantObjects.Append(FlushedGuids);
 		}
 	}
diff --git a/Engine/Source/Runtime/Experimental/Iris/Core/Private/Iris/ReplicationSystem/ObjectReplicationBridge.cpp b/Engine/Source/Runtime/Experimental/Iris/Core/Private/Iris/ReplicationSystem/ObjectReplicationBridge.cpp
--- a/Engine/Source/Runtime/Experimental/Iris/Core/Private/Iris/ReplicationSystem/ObjectReplicationBridge.cpp
+++ b/Engine/Source/Runtime/Experimental/Iris/Core/Private/Iris/ReplicationSystem/ObjectReplicationBridge.cpp
@@ -821,6 +821,7 @@ void UObjectReplicationBridge::PreUpdateAndPoll()
 
 		// Update subobjects' owner first and owners' subobjects second. It's the only way to properly mark all groups of objects in two passes.
 		const FNetBitArrayView SubObjects = MakeNetBitArrayView(LocalNetRefHandleManager.GetSubObjectInternalIndices());
+		const FNetBitArrayView ForceNetUpdateObjects = ReplicationSystemInternal->GetDirtyNetObjectTracker().GetForceNetUpdateObjects();
 
 		if (bEnableForceNetUpdate)
 		{
@@ -828,7 +829,6 @@ void UObjectReplicationBridge::PreUpdateAndPoll()
 			FNetBitArray ForceNetUpdateAndRelevantObjects(RelevantObjects.GetNumBits(), FNetBitArray::NoResetNoValidate);
 			FNetBitArrayView ForceNetUpdateAndRelevantObjectsView = MakeNetBitArrayView(ForceNetUpdateAndRelevantObjects, FNetBitArrayView::NoResetNoValidate);
 			
-			const FNetBitArrayView ForceNetUpdateObjects = ReplicationSystemInternal->GetDirtyNetObjectTracker().GetForceNetUpdateObjects();
 			ForceNetUpdateAndRelevantObjectsView.Set(RelevantObjects, FNetBitArray::AndOp, ForceNetUpdateObjects);
 
 			FNetBitArrayView::ForAllSetBits(ForceNetUpdateAndRelevantObjectsView, SubObjects, FNetBitArray::AndOp, PropagateSubObjectNetForceUpdateToOwner);
@@ -836,7 +836,7 @@ void UObjectReplicationBridge::PreUpdateAndPoll()
 		}
 		else
 		{
-			// Make the list of objects which are dirty and are also relevant
+			// Make the list of objects which are dirty or forced an update, and are also relevant
 			FNetBitArray DirtyAndRelevantObjects(RelevantObjects.GetNumBits(), FNetBitArray::NoResetNoValidate);
 			FNetBitArrayView DirtyAndRelevantObjectsView = MakeNetBitArrayView(DirtyAndRelevantObjects, FNetBitArrayView::NoResetNoValidate);
 
@@ -844,7 +844,8 @@ void UObjectReplicationBridge::PreUpdateAndPoll()
 				FDirtyObjectsAccessor DirtyObjectsAccessor(ReplicationSystemInternal->GetDirtyNetObjectTracker());
 				FNetBitArrayView DirtyObjectsThisFrame = DirtyObjectsAccessor.GetDirtyNetObjects();
 
-				DirtyAndRelevantObjectsView.Set(RelevantObjects, FNetBitArray::AndOp, DirtyObjectsThisFrame);
+				DirtyAndRelevantObjectsView.Set(DirtyObjectsThisFrame, FNetBitArray::OrOp, ForceNetUpdateObjects);
+				DirtyAndRelevantObjectsView.Combine(RelevantObjects, FNetBitArray::AndOp);
 			}
 
 			FNetBitArrayView::ForAllSetBits(DirtyAndRelevantObjectsView, SubObjects, FNetBitArray::AndOp, PropagateSubObjectNetForceUpdateToOwner);
diff --git a/Engine/Source/Runtime/Experimental/Iris/Core/Public/Iris/ReplicationSystem/Private/FastArrayReplicationFragmentInternal.h b/Engine/Source/Runtime/Experimental/Iris/Core/Public/Iris/ReplicationSystem/Private/FastArrayReplicationFragmentInternal.h
--- a/Engine/Source/Runtime/Experimental/Iris/Core/Public/Iris/ReplicationSystem/Private/FastArrayReplicationFragmentInternal.h
+++ b/Engine/Source/Runtime/Experimental/Iris/Core/Public/Iris/ReplicationSystem/Private/FastArrayReplicationFragmentInternal.h
@@ -197,12 +197,6 @@ void FFastArrayReplicationFragmentHelper::ApplyReplicatedState(FastArrayType* Ds
 		}
 	}
 	
-	const int32 PreRemoveSize = DstWrappedArray->Num();
-	const int32 FinalSize = PreRemoveSize - RemovedIndices.Num();
-
-	// Remove callback to FastArraySerializer
-	DstArraySerializer->PreReplicatedRemove(MakeArrayView(RemovedIndices), FinalSize);
-
 	// Find new and modified elements in received data, That is elements that do not exist in old map
 	TArray<int32> AddedIndices;
 	TArray<int32> ModifiedIndices;
@@ -223,9 +217,6 @@ void FFastArrayReplicationFragmentHelper::ApplyReplicatedState(FastArrayType* Ds
 
 					// We use per element copy since we do not want to overwrite data that is not replicated
 					InternalCopyArrayElement(ArrayElementDescriptor, &(*DstWrappedArray)[*ExistingIndex], &SrcItems[It]);
-
-					// Change callback
-					(*DstWrappedArray)[*ExistingIndex].PostReplicatedChange(*DstArraySerializer);
 				}
 			}
 			else
@@ -237,18 +228,32 @@ void FFastArrayReplicationFragmentHelper::ApplyReplicatedState(FastArrayType* Ds
 				// We need to propagate the ReplicationID in order to find our object
 				(*DstWrappedArray)[AddedIndex].ReplicationID = SrcItems[It].ReplicationID;
 
-				// Add callback
-				(*DstWrappedArray)[AddedIndex].PostReplicatedAdd(*DstArraySerializer);
-
 				// should we store ids or indices?
 				AddedIndices.Add(AddedIndex);
 			}
-		}	
+		}
+	}
 
-		// Added and changed callbacks to FastArraySerializer
-		DstArraySerializer->PostReplicatedAdd(MakeArrayView(AddedIndices), FinalSize);
-		DstArraySerializer->PostReplicatedChange(MakeArrayView(ModifiedIndices), FinalSize);
+	// Added and changed callbacks to FastArraySerializer
+	const int32 PreRemoveSize = DstWrappedArray->Num();
+	const int32 FinalSize = PreRemoveSize - RemovedIndices.Num();
+
+	// Remove callback to FastArraySerializer - done after adding new elements
+	DstArraySerializer->PreReplicatedRemove(MakeArrayView(RemovedIndices), FinalSize);
+
+	// Add callbacks
+	for (int32 AddedIndex : AddedIndices)
+	{
+		(*DstWrappedArray)[AddedIndex].PostReplicatedAdd(*DstArraySerializer);
+	}
+	DstArraySerializer->PostReplicatedAdd(MakeArrayView(AddedIndices), FinalSize);
+
+	// Change callbacks
+	for (int32 ExistingIndex : ModifiedIndices)
+	{
+		(*DstWrappedArray)[ExistingIndex].PostReplicatedChange(*DstArraySerializer);
 	}
+	DstArraySerializer->PostReplicatedChange(MakeArrayView(ModifiedIndices), FinalSize);
 
 	// Remove indices
 	if (RemovedIndices.Num() > 0)
