Make RVT Actors blueprintable.
Remove unnecessary check() that fires in blueprint view.

diff --git a/Engine/Source/Runtime/Engine/Classes/VT/RuntimeVirtualTextureVolume.h b/Engine/Source/Runtime/Engine/Classes/VT/RuntimeVirtualTextureVolume.h
--- a/Engine/Source/Runtime/Engine/Classes/VT/RuntimeVirtualTextureVolume.h
+++ b/Engine/Source/Runtime/Engine/Classes/VT/RuntimeVirtualTextureVolume.h
@@ -7,7 +7,7 @@
 #include "RuntimeVirtualTextureVolume.generated.h"
 
 /** Actor used to place a URuntimeVirtualTexture in the world. */
-UCLASS(HideCategories = (Actor, Collision, Cooking, Input, LOD, Physics, Replication, Rendering), MinimalAPI)
+UCLASS(Blueprintable, HideCategories = (Actor, Collision, Cooking, Input, LOD, Physics, Replication, Rendering), MinimalAPI)
 class ARuntimeVirtualTextureVolume : public AActor
 {
 	GENERATED_UCLASS_BODY()
diff --git a/Engine/Source/Runtime/Engine/Private/Components/RuntimeVirtualTextureComponent.cpp b/Engine/Source/Runtime/Engine/Private/Components/RuntimeVirtualTextureComponent.cpp
--- a/Engine/Source/Runtime/Engine/Private/Components/RuntimeVirtualTextureComponent.cpp
+++ b/Engine/Source/Runtime/Engine/Private/Components/RuntimeVirtualTextureComponent.cpp
@@ -239,8 +239,6 @@ bool URuntimeVirtualTextureComponent::IsStreamingTextureInvalid(EShadingPath Sha
 
 bool URuntimeVirtualTextureComponent::IsStreamingTextureInvalid() const
 {
-	checkf(IsActiveInWorld(), TEXT("This function should never be called for a world where we're inactive"));
-
 	return 
 		VirtualTexture != nullptr && 
 		StreamingTexture != nullptr && 
