Removing WITH_EDITOR from ILevelInstanceInterface LevelStreaming function definition dependencies so the system can compile in non-editor builds. Code Cleanup in LevelInstanceSubsystem.cpp.
#rnx

diff --git a/Engine/Source/Runtime/Engine/Private/LevelInstance/LevelInstanceSubsystem.cpp b/Engine/Source/Runtime/Engine/Private/LevelInstance/LevelInstanceSubsystem.cpp
--- a/Engine/Source/Runtime/Engine/Private/LevelInstance/LevelInstanceSubsystem.cpp
+++ b/Engine/Source/Runtime/Engine/Private/LevelInstance/LevelInstanceSubsystem.cpp
@@ -559,6 +559,19 @@ void ULevelInstanceSubsystem::ForEachLevelInstanceAncestorsAndSelf(AActor* Actor
 	ForEachLevelInstanceAncestors(Actor, Operation);
 }
 
+void ULevelInstanceSubsystem::ForEachLevelInstanceAncestorsAndSelf(const AActor* Actor, TFunctionRef<bool(const ILevelInstanceInterface*)> Operation) const
+{
+	if (const ILevelInstanceInterface* LevelInstance = Cast<ILevelInstanceInterface>(Actor))
+	{
+		if (!Operation(LevelInstance))
+		{
+			return;
+		}
+	}
+
+	ForEachLevelInstanceAncestors(Actor, Operation);
+}
+
 ULevelStreamingLevelInstance* ULevelInstanceSubsystem::GetLevelInstanceLevelStreaming(const ILevelInstanceInterface* LevelInstance) const
 {
 	if (LevelInstance->HasValidLevelInstanceID())
@@ -583,6 +596,19 @@ void ULevelInstanceSubsystem::ForEachLevelInstanceAncestors(AActor* Actor, TFunc
 	} while (ParentLevelInstance != nullptr && Operation(ParentLevelInstance));
 }
 
+void ULevelInstanceSubsystem::ForEachLevelInstanceAncestors(const AActor* Actor, TFunctionRef<bool(const ILevelInstanceInterface*)> Operation) const
+{
+	const ILevelInstanceInterface* ParentLevelInstance = nullptr;
+	if(Actor)
+	{ 
+		do
+		{
+			ParentLevelInstance = GetOwningLevelInstance(Actor->GetLevel());
+			Actor = Cast<AActor>(ParentLevelInstance);
+		} while (Actor != nullptr && Operation(ParentLevelInstance));
+	}
+}
+
 ILevelInstanceInterface* ULevelInstanceSubsystem::GetOwningLevelInstance(const ULevel* Level) const
 {
 	if (ULevelStreaming* BaseLevelStreaming = FLevelUtils::FindStreamingLevel(Level))
@@ -619,6 +645,112 @@ void ULevelInstanceSubsystem::ForEachActorInLevelInstance(const ILevelInstanceIn
 	}
 }
 
+bool ULevelInstanceSubsystem::CanUseWorldAsset(const ILevelInstanceInterface* LevelInstance, TSoftObjectPtr<UWorld> WorldAsset, FString* OutReason)
+{
+	if (OutReason)
+	{
+		*OutReason = FString(TEXT(""));
+	}
+
+	#if WITH_EDITOR
+	// Do not validate when running convert commandlet as package might not exist yet.
+	if (UWorldPartitionSubsystem::IsRunningConvertWorldPartitionCommandlet())
+	{
+		return true;
+	}
+	#endif
+
+	if (WorldAsset.IsNull())
+	{
+		return true;
+	}
+
+	FString PackageName;
+	if (!FPackageName::DoesPackageExist(WorldAsset.GetLongPackageName()))
+	{
+		if (OutReason)
+		{
+			*OutReason = FString::Format(TEXT("Attempting to set Level Instance to package {0} which does not exist. Ensure the level was saved before attepting to set the level instance world asset."), { WorldAsset.GetLongPackageName() });
+		}
+		return false;
+	}
+
+	TArray<TPair<FText, TSoftObjectPtr<UWorld>>> LoopInfo;
+	const ILevelInstanceInterface* LoopStart = nullptr;
+
+	// Check if the current LevelInstance is being set to a WorldAsset that loads the current LevelInstance or any ancestor that owns the current level's package.
+	if (!CheckForLoop(LevelInstance, WorldAsset, OutReason ? &LoopInfo : nullptr, OutReason ? &LoopStart : nullptr))
+	{
+		if (OutReason)
+		{
+			if (ensure(LoopStart))
+			{
+				const AActor* LoopStartActor = CastChecked<AActor>(LoopStart);
+				TSoftObjectPtr<UWorld> LoopStartAsset(LoopStartActor->GetLevel()->GetTypedOuter<UWorld>());
+				*OutReason = FString::Format(TEXT("Setting LevelInstance to {0} would cause loop {1}:{2}\n"), { WorldAsset.GetLongPackageName(), LoopStartActor->GetName(), LoopStartAsset.GetLongPackageName() });
+				for (int32 i = LoopInfo.Num() - 1; i >= 0; --i)
+				{
+					OutReason->Append(FString::Format(TEXT("{0} {1}\n"), 
+						{ *LoopInfo[i].Key.ToString(), 
+						*LoopInfo[i].Value.GetLongPackageName() }));
+				}
+			}
+		}
+
+		return false;
+	}
+
+	return true;
+}
+
+bool ULevelInstanceSubsystem::CheckForLoop(const ILevelInstanceInterface* LevelInstance, TSoftObjectPtr<UWorld> WorldAsset, TArray<TPair<FText, TSoftObjectPtr<UWorld>>>* LoopInfo, const ILevelInstanceInterface** OutLoopStart)
+{
+	bool bValid = true;
+
+	if (LevelInstance)
+	{
+		if(ULevelInstanceSubsystem* LevelInstanceSubsystem = LevelInstance->GetLevelInstanceSubsystem())
+		{
+			LevelInstanceSubsystem->ForEachLevelInstanceAncestorsAndSelf(CastChecked<AActor>(LevelInstance), [&bValid, &WorldAsset, LevelInstance, LoopInfo, OutLoopStart](const ILevelInstanceInterface* CurrentLevelInstance)
+			{
+				FName PackageToTest(*WorldAsset.GetLongPackageName());
+				// Check to exclude NAME_None since Preview Levels are in the transient package
+				// Check the level we are spawned in to detect the loop (this will handle loops caused by LevelInstances and by regular level streaming)
+				const AActor* CurrentActor = CastChecked<AActor>(CurrentLevelInstance);
+				if (PackageToTest != NAME_None)
+				{
+					if (ULevel* LevelTheActorBelongsTo = CurrentActor->GetLevel())
+					{
+						if (UPackage* PackageTheLevelBelongsTo = LevelTheActorBelongsTo->GetPackage())
+						{
+							if (PackageTheLevelBelongsTo->GetLoadedPath() == FPackagePath::FromPackageNameChecked(PackageToTest))
+							{
+								bValid = false;
+								if (OutLoopStart)
+								{
+									*OutLoopStart = CurrentLevelInstance;
+								}
+							}
+						}
+					}
+				}
+	
+				if (LoopInfo)
+				{
+					TSoftObjectPtr<UWorld> CurrentAsset = (CurrentLevelInstance == LevelInstance)? WorldAsset : (CurrentLevelInstance)? CurrentLevelInstance->GetWorldAsset() : nullptr;
+					FText LevelInstanceName = (CurrentActor)? FText::FromString(CurrentActor->GetPathName()) : FText::FromString("Invalid Actor cannot retrieve LevelInstance name.");
+					FText Description = FText::Format(LOCTEXT("LevelInstanceLoopLink", "-> Actor: {0} loads"), LevelInstanceName);
+					LoopInfo->Emplace(Description, CurrentAsset);
+				}
+	
+				return bValid;
+			});
+		}
+	}
+
+	return bValid;
+}
+
 #if WITH_EDITOR
 
 void ULevelInstanceSubsystem::Tick()
@@ -735,30 +867,6 @@ bool ULevelInstanceSubsystem::GetLevelInstanceBoundsFromPackage(const FTransform
 	return false;
 }
 
-void ULevelInstanceSubsystem::ForEachLevelInstanceAncestorsAndSelf(const AActor* Actor, TFunctionRef<bool(const ILevelInstanceInterface*)> Operation) const
-{
-	if (const ILevelInstanceInterface* LevelInstance = Cast<ILevelInstanceInterface>(Actor))
-	{
-		if (!Operation(LevelInstance))
-		{
-			return;
-		}
-	}
-
-	ForEachLevelInstanceAncestors(Actor, Operation);
-}
-
-void ULevelInstanceSubsystem::ForEachLevelInstanceAncestors(const AActor* Actor, TFunctionRef<bool(const ILevelInstanceInterface*)> Operation) const
-{
-	const ILevelInstanceInterface* ParentLevelInstance = nullptr;
-	do 
-	{
-		ParentLevelInstance = GetOwningLevelInstance(Actor->GetLevel());
-		Actor = Cast<AActor>(ParentLevelInstance);
-	} 
-	while (ParentLevelInstance != nullptr && Operation(ParentLevelInstance));
-}
-
 void ULevelInstanceSubsystem::ForEachLevelInstanceChild(const ILevelInstanceInterface* LevelInstance, bool bRecursive, TFunctionRef<bool(const ILevelInstanceInterface*)> Operation) const
 {
 	ForEachLevelInstanceChildImpl(LevelInstance, bRecursive, Operation);
@@ -2310,90 +2418,6 @@ bool ULevelInstanceSubsystem::CheckForLoop(const ILevelInstanceInterface* LevelI
 	return CheckForLoop(LevelInstance, LevelInstance->GetWorldAsset(), LoopInfo, LoopStart);
 }
 
-bool ULevelInstanceSubsystem::CheckForLoop(const ILevelInstanceInterface* LevelInstance, TSoftObjectPtr<UWorld> WorldAsset, TArray<TPair<FText, TSoftObjectPtr<UWorld>>>* LoopInfo, const ILevelInstanceInterface** LoopStart)
-{
-	bool bValid = true;
-		
-	if (ULevelInstanceSubsystem* LevelInstanceSubsystem = LevelInstance->GetLevelInstanceSubsystem())
-	{
-		LevelInstanceSubsystem->ForEachLevelInstanceAncestorsAndSelf(CastChecked<AActor>(LevelInstance), [&bValid, &WorldAsset, LevelInstance, LoopInfo, LoopStart](const ILevelInstanceInterface* CurrentLevelInstance)
-		{
-			FName PackageToTest(*WorldAsset.GetLongPackageName());
-			// Check to exclude NAME_None since Preview Levels are in the transient package
-			// Check the level we are spawned in to detect the loop (this will handle loops caused by LevelInstances and by regular level streaming)
-			const AActor* CurrentActor = CastChecked<AActor>(CurrentLevelInstance);
-			if (PackageToTest != NAME_None && CurrentActor->GetLevel()->GetPackage()->GetLoadedPath() == FPackagePath::FromPackageNameChecked(PackageToTest))
-			{
-				bValid = false;
-				if (LoopStart)
-				{
-					*LoopStart = CurrentLevelInstance;
-				}
-			}
-
-			if (LoopInfo)
-			{
-				TSoftObjectPtr<UWorld> CurrentAsset = CurrentLevelInstance == LevelInstance ? WorldAsset : CurrentLevelInstance->GetWorldAsset();
-				FText LevelInstanceName = FText::FromString(CurrentActor->GetPathName());
-				FText Description = FText::Format(LOCTEXT("LevelInstanceLoopLink", "-> Actor: {0} loads"), LevelInstanceName);
-				LoopInfo->Emplace(Description, CurrentAsset);
-			}
-			
-			return bValid;
-		});
-	}
-	
-	return bValid;
-}
-
-bool ULevelInstanceSubsystem::CanUseWorldAsset(const ILevelInstanceInterface* LevelInstance, TSoftObjectPtr<UWorld> WorldAsset, FString* OutReason)
-{
-	// Do not validate when running convert commandlet as package might not exist yet.
-	if (UWorldPartitionSubsystem::IsRunningConvertWorldPartitionCommandlet())
-	{
-		return true;
-	}
-
-	if (WorldAsset.IsNull())
-	{
-		return true;
-	}
-
-	FString PackageName;
-	if (!FPackageName::DoesPackageExist(WorldAsset.GetLongPackageName()))
-	{
-		if (OutReason)
-		{
-			*OutReason = FString::Format(TEXT("Attempting to set Level Instance to package {0} which does not exist. Ensure the level was saved before attepting to set the level instance world asset."), { WorldAsset.GetLongPackageName() });
-		}
-		return false;
-	}
-
-	TArray<TPair<FText, TSoftObjectPtr<UWorld>>> LoopInfo;
-	const ILevelInstanceInterface* LoopStart = nullptr;
-
-	if (!CheckForLoop(LevelInstance, WorldAsset, OutReason ? &LoopInfo : nullptr, OutReason ? &LoopStart : nullptr))
-	{
-		if (OutReason)
-		{
-			if (ensure(LoopStart))
-			{
-				const AActor* LoopStartActor = CastChecked<AActor>(LoopStart);
-				TSoftObjectPtr<UWorld> LoopStartAsset(LoopStartActor->GetLevel()->GetTypedOuter<UWorld>());
-				*OutReason = FString::Format(TEXT("Setting LevelInstance to {0} would cause loop {1}:{2}\n"), { WorldAsset.GetLongPackageName(), LoopStartActor->GetName(), LoopStartAsset.GetLongPackageName() });
-				for (int32 i = LoopInfo.Num() - 1; i >= 0; --i)
-				{
-					OutReason->Append(FString::Format(TEXT("{0} {1}\n"), { *LoopInfo[i].Key.ToString(), *LoopInfo[i].Value.GetLongPackageName() }));
-				}
-			}
-		}
-
-		return false;
-	}
-
-	return true;
-}
-
 bool ULevelInstanceSubsystem::PassLevelInstanceFilter(UWorld* World, const FWorldPartitionHandle& Actor) const
 {
 	UWorld* ContainerOuterWorld = Actor->GetContainer()->GetWorldPartition()->GetTypedOuter<UWorld>();
diff --git a/Engine/Source/Runtime/Engine/Public/LevelInstance/LevelInstanceInterface.h b/Engine/Source/Runtime/Engine/Public/LevelInstance/LevelInstanceInterface.h
--- a/Engine/Source/Runtime/Engine/Public/LevelInstance/LevelInstanceInterface.h
+++ b/Engine/Source/Runtime/Engine/Public/LevelInstance/LevelInstanceInterface.h
@@ -47,8 +47,6 @@ class ILevelInstanceInterface
 
 #if WITH_EDITOR
 
-	UFUNCTION(BlueprintCallable, Category = Default)
-	virtual bool SetWorldAsset(TSoftObjectPtr<UWorld> WorldAsset) = 0;
 
 	virtual ULevelInstanceComponent* GetLevelInstanceComponent() const = 0;
 
@@ -89,7 +87,6 @@ class ILevelInstanceInterface
 	UFUNCTION(BlueprintCallable, Category = Default)
 	ENGINE_API virtual ULevel* GetLoadedLevel() const;
 
-	ENGINE_API virtual void UpdateLevelInstanceFromWorldAsset();
 
 	ENGINE_API virtual void OnEdit();
 
@@ -135,4 +132,14 @@ class ILevelInstanceInterface
 	// Return supported filter types when setting filter through details panel
 	ENGINE_API virtual EWorldPartitionActorFilterType GetDetailsFilterTypes() const;
 #endif
+
+	/** Sets the UWorld asset reference when loading a LevelInstance */
+	UFUNCTION(BlueprintCallable, Category = Default)
+	virtual bool SetWorldAsset(TSoftObjectPtr<UWorld> WorldAsset)
+	{
+		return false;
+	}
+
+	ENGINE_API virtual void UpdateLevelInstanceFromWorldAsset();
+
 };
diff --git a/Engine/Source/Runtime/Engine/Public/LevelInstance/LevelInstanceSubsystem.h b/Engine/Source/Runtime/Engine/Public/LevelInstance/LevelInstanceSubsystem.h
--- a/Engine/Source/Runtime/Engine/Public/LevelInstance/LevelInstanceSubsystem.h
+++ b/Engine/Source/Runtime/Engine/Public/LevelInstance/LevelInstanceSubsystem.h
@@ -56,11 +56,21 @@ public:
 	ENGINE_API void RequestUnloadLevelInstance(ILevelInstanceInterface* LevelInstance);
 	ENGINE_API bool IsLoaded(const ILevelInstanceInterface* LevelInstance) const;
 	ENGINE_API bool IsLoading(const ILevelInstanceInterface* LevelInstance) const;
+	ENGINE_API void ForEachLevelInstanceAncestors(const AActor* Actor, TFunctionRef<bool(const ILevelInstanceInterface*)> Operation) const;
 	ENGINE_API void ForEachLevelInstanceAncestorsAndSelf(AActor* Actor, TFunctionRef<bool(ILevelInstanceInterface*)> Operation) const;
+	/** Runs a lambda operation along the ancestors that own the LevelInstance. Primarily for capturing inclusive true/false by using lambda captures */
+	ENGINE_API void ForEachLevelInstanceAncestorsAndSelf(const AActor* Actor, TFunctionRef<bool(const ILevelInstanceInterface*)> Operation) const;
 	ENGINE_API ULevelStreamingLevelInstance* GetLevelInstanceLevelStreaming(const ILevelInstanceInterface* LevelInstance) const;
 	ENGINE_API void ForEachActorInLevelInstance(const ILevelInstanceInterface* LevelInstance, TFunctionRef<bool(AActor* LevelActor)> Operation) const;
 	ENGINE_API ULevel* GetLevelInstanceLevel(const ILevelInstanceInterface* LevelInstance) const;
-
+	/** Checks if a WorldAsset can/should be used in a LevelInstance */
+	static ENGINE_API bool CanUseWorldAsset(const ILevelInstanceInterface* LevelInstance, TSoftObjectPtr<UWorld> WorldAsset, FString* OutReason);
+	/**
+	Lambda expr format that checks ancestor owners of a LevelInstance.
+	Detects if the WorldAsset shares the same package as the current LevelInstance or any of its ancestors owning the LevelInstance.
+	Used in conjunction with ForEachLevelInstanceAncestorsAndSelf to find a loop using a boolean lambda capture.
+	*/
+	static ENGINE_API bool CheckForLoop(const ILevelInstanceInterface* LevelInstance, TSoftObjectPtr<UWorld> WorldAsset, TArray<TPair<FText, TSoftObjectPtr<UWorld>>>* LoopInfo = nullptr, const ILevelInstanceInterface** LoopStart = nullptr);
 
 #if WITH_EDITOR
 	ENGINE_API void Tick();
@@ -84,8 +94,6 @@ public:
 	ENGINE_API bool GetLevelInstanceBounds(const ILevelInstanceInterface* LevelInstance, FBox& OutBounds) const;
 	static ENGINE_API bool GetLevelInstanceBoundsFromPackage(const FTransform& InstanceTransform, FName LevelPackage, FBox& OutBounds);
 	
-	ENGINE_API void ForEachLevelInstanceAncestorsAndSelf(const AActor* Actor, TFunctionRef<bool(const ILevelInstanceInterface*)> Operation) const;
-	ENGINE_API void ForEachLevelInstanceAncestors(const AActor* Actor, TFunctionRef<bool(const ILevelInstanceInterface*)> Operation) const;
 	ENGINE_API void ForEachLevelInstanceChild(const ILevelInstanceInterface* LevelInstance, bool bRecursive, TFunctionRef<bool(const ILevelInstanceInterface*)> Operation) const;
 	ENGINE_API void ForEachLevelInstanceChild(ILevelInstanceInterface* LevelInstance, bool bRecursive, TFunctionRef<bool(ILevelInstanceInterface*)> Operation) const;
 	ENGINE_API bool HasDirtyChildrenLevelInstances(const ILevelInstanceInterface* LevelInstance) const;
@@ -120,8 +128,6 @@ public:
 	ENGINE_API FString PrefixWithParentLevelInstanceActorLabels(const FString& ActorLabel, const ULevel* Level) const;
 
 	static ENGINE_API bool CheckForLoop(const ILevelInstanceInterface* LevelInstance, TArray<TPair<FText, TSoftObjectPtr<UWorld>>>* LoopInfo = nullptr, const ILevelInstanceInterface** LoopStart = nullptr);
-	static ENGINE_API bool CheckForLoop(const ILevelInstanceInterface* LevelInstance, TSoftObjectPtr<UWorld> WorldAsset, TArray<TPair<FText, TSoftObjectPtr<UWorld>>>* LoopInfo = nullptr, const ILevelInstanceInterface** LoopStart = nullptr);
-	static ENGINE_API bool CanUseWorldAsset(const ILevelInstanceInterface* LevelInstance, TSoftObjectPtr<UWorld> WorldAsset, FString* OutReason);
 	
 	UE_DEPRECATED(5.3, "CanUsePackage is deprecated.")
 	static bool CanUsePackage(FName InPackageName) { return true;  }
@@ -133,6 +139,7 @@ public:
 	static ENGINE_API void ResetLoadersForWorldAsset(const FString& WorldAsset);
 
 	ENGINE_API bool PassLevelInstanceFilter(UWorld* World, const FWorldPartitionHandle& Actor) const;
+
 #endif
 
 private:
@@ -142,7 +149,6 @@ private:
 	ENGINE_API void UnloadLevelInstance(const FLevelInstanceID& LevelInstanceID);
 	ENGINE_API void ForEachActorInLevel(ULevel* Level, TFunctionRef<bool(AActor * LevelActor)> Operation) const;
 	ENGINE_API void ForEachLevelInstanceAncestors(AActor* Actor, TFunctionRef<bool(ILevelInstanceInterface*)> Operation) const;
-	
 	ENGINE_API void RegisterLoadedLevelStreamingLevelInstance(ULevelStreamingLevelInstance* LevelStreaming);
 
 #if WITH_EDITOR
