Fixed build break

diff --git a/Engine/Source/Runtime/Engine/Private/Components/ActorComponent.cpp b/Engine/Source/Runtime/Engine/Private/Components/ActorComponent.cpp
--- a/Engine/Source/Runtime/Engine/Private/Components/ActorComponent.cpp
+++ b/Engine/Source/Runtime/Engine/Private/Components/ActorComponent.cpp
@@ -170,6 +170,8 @@ void UpdateAllPrimitiveSceneInfosForSingleComponentInterface (IPrimitiveComponen
 	{
 		if (InScenesToUpdateAllPrimitiveSceneInfosForBatching == nullptr)
 		{
+			UE::RenderCommandPipe::FSyncScope SyncScope;
+
 			// If no batching is available (this ComponentReregisterContext is not created by a FGlobalComponentReregisterContext), issue one update per component
 			ENQUEUE_RENDER_COMMAND(UpdateAllPrimitiveSceneInfosCmd)([Scene](FRHICommandListImmediate& RHICmdList) {
 				Scene->UpdateAllPrimitiveSceneInfos(RHICmdList);
diff --git a/Engine/Source/Runtime/RenderCore/Private/RenderingThread.cpp b/Engine/Source/Runtime/RenderCore/Private/RenderingThread.cpp
--- a/Engine/Source/Runtime/RenderCore/Private/RenderingThread.cpp
+++ b/Engine/Source/Runtime/RenderCore/Private/RenderingThread.cpp
@@ -1899,7 +1899,9 @@ FRenderCommandPipe::FRenderCommandPipe(const TCHAR* InName, ERenderCommandPipeFl
 		SetEnabled(Variable->GetBool());
 	}))
 {
+#if !UE_SERVER
 	GlobalListLink.LinkHead(GRenderCommandPipeRegistry.GetGlobalList());
+#endif
 }
 
 FRenderCommandPipe::~FRenderCommandPipe()
