Anim: Returning playing montage if there is one already playing. Otherwise this was returning nullptr.

This fixes a bug where calling FAnimMontageInstance::PreviewSequencerMontagePosition would not return the montage being played, which resulted in PreviewSetAnimPosition not caching pre animated state for the montage

diff --git a/Engine/Source/Runtime/Engine/Private/Animation/AnimMontage.cpp b/Engine/Source/Runtime/Engine/Private/Animation/AnimMontage.cpp
--- a/Engine/Source/Runtime/Engine/Private/Animation/AnimMontage.cpp
+++ b/Engine/Source/Runtime/Engine/Private/Animation/AnimMontage.cpp
@@ -2848,6 +2848,7 @@ UAnimMontage* FAnimMontageInstance::SetSequencerMontagePosition(FName SlotName,
 
 		if (MontageInstanceToUpdate)
 		{
+			PlayingMontage = MontageInstanceToUpdate->Montage;
 			InOutInstanceId = MontageInstanceToUpdate->GetInstanceID();
 
 			// ensure full weighting to this instance
