Cluster union component : optimization using only one physics object to get the interface
also avoid getting the interface if the map is empty

diff --git a/Engine/Source/Runtime/Engine/Private/PhysicsEngine/ClusterUnionComponent.cpp b/Engine/Source/Runtime/Engine/Private/PhysicsEngine/ClusterUnionComponent.cpp
--- a/Engine/Source/Runtime/Engine/Private/PhysicsEngine/ClusterUnionComponent.cpp
+++ b/Engine/Source/Runtime/Engine/Private/PhysicsEngine/ClusterUnionComponent.cpp
@@ -957,26 +957,29 @@ void UClusterUnionComponent::HandleAddOrModifiedClusteredComponent(UPrimitiveCom
 		}
 	}
 
-	TArray<Chaos::FPhysicsObjectHandle> AllPhysicsObjects = ChangedComponent->GetAllPhysicsObjects();
-	FLockedReadPhysicsObjectExternalInterface Interface = FPhysicsObjectExternalInterface::LockRead(AllPhysicsObjects);
-
-	// One more loop to ensure that our sets of physics objects are valid and up to date.
-	// This needs to happen on both the client and the server.
-	for (const TPair<int32, FTransform>& Kvp : PerBoneChildToParent)
+	if (PerBoneChildToParent.Num())
 	{
-		Chaos::FPhysicsObjectHandle PhysicsObject = ChangedComponent->GetPhysicsObjectById(Kvp.Key);
+		// get the interface from the first element 
+		FLockedReadPhysicsObjectExternalInterface Interface = FPhysicsObjectExternalInterface::LockRead(GetWorld()->GetPhysicsScene());
 
-		if (AccelerationStructure)
+		// One more loop to ensure that our sets of physics objects are valid and up to date.
+		// This needs to happen on both the client and the server.
+		for (const TPair<int32, FTransform>& Kvp : PerBoneChildToParent)
 		{
-			if (PhysicsObject)
+			Chaos::FPhysicsObjectHandle PhysicsObject = ChangedComponent->GetPhysicsObjectById(Kvp.Key);
+
+			if (AccelerationStructure)
 			{
-				const FBox HandleBounds = Interface->GetWorldBounds({ &PhysicsObject, 1 });
-				FExternalSpatialAccelerationPayload Handle;
-				Handle.Initialize(ChangedComponent, Kvp.Key);
-				if (Handle.IsValid())
+				if (PhysicsObject)
 				{
-					AccelerationStructure->UpdateElement(Handle, Chaos::TAABB<Chaos::FReal, 3>{HandleBounds.Min, HandleBounds.Max}, HandleBounds.IsValid != 0);
-					ComponentData.CachedAccelerationPayloads.Add(Handle);
+					const FBox HandleBounds = Interface->GetWorldBounds({ &PhysicsObject, 1 });
+					FExternalSpatialAccelerationPayload Handle;
+					Handle.Initialize(ChangedComponent, Kvp.Key);
+					if (Handle.IsValid())
+					{
+						AccelerationStructure->UpdateElement(Handle, Chaos::TAABB<Chaos::FReal, 3>{HandleBounds.Min, HandleBounds.Max}, HandleBounds.IsValid != 0);
+						ComponentData.CachedAccelerationPayloads.Add(Handle);
+					}
 				}
 			}
 		}
@@ -1104,7 +1107,6 @@ void UClusterUnionComponent::SetSimulatePhysics(bool bSimulate)
 	{
 		return;
 	}
-
 	PhysicsProxy->SetObjectState_External(bSimulate ? Chaos::EObjectStateType::Dynamic : Chaos::EObjectStateType::Kinematic);
 }
 
