Fix subsurface profile material when using SubstrateShadingModels node outside of the automatic conversion.

Subsurface profile were not gathered, causing the material to not be flagged as SSS.

diff --git a/Engine/Source/Runtime/Engine/Private/Materials/MaterialExpressions.cpp b/Engine/Source/Runtime/Engine/Private/Materials/MaterialExpressions.cpp
--- a/Engine/Source/Runtime/Engine/Private/Materials/MaterialExpressions.cpp
+++ b/Engine/Source/Runtime/Engine/Private/Materials/MaterialExpressions.cpp
@@ -23513,6 +23513,11 @@ void UMaterialExpressionSubstrateShadingModels::GatherSubstrateMaterialInfo(FSub
 		if (ShadingModelOverride == MSM_SingleLayerWater)	{ SubstrateMaterialInfo.AddShadingModel(ESubstrateShadingModel::SSM_SingleLayerWater); }
 		if (ShadingModelOverride == MSM_ThinTranslucent)	{ SubstrateMaterialInfo.AddShadingModel(ESubstrateShadingModel::SSM_ThinTranslucent); }
 	}
+
+	if (SubsurfaceProfile)
+	{
+		SubstrateMaterialInfo.AddSubsurfaceProfile(SubsurfaceProfile);
+	}
 }
 
 FSubstrateOperator* UMaterialExpressionSubstrateShadingModels::SubstrateGenerateMaterialTopologyTree(class FMaterialCompiler* Compiler, class UMaterialExpression* Parent, int32 OutputIndex)
