Add Substrate SubstrateConvertMaterialAttributes to ease material attributes setup & manipulation with Subtrate data.

diff --git a/Engine/Source/Editor/MaterialEditor/Private/MaterialEditorDetailCustomization.cpp b/Engine/Source/Editor/MaterialEditor/Private/MaterialEditorDetailCustomization.cpp
--- a/Engine/Source/Editor/MaterialEditor/Private/MaterialEditorDetailCustomization.cpp
+++ b/Engine/Source/Editor/MaterialEditor/Private/MaterialEditorDetailCustomization.cpp
@@ -825,7 +825,14 @@ void FMaterialDetailCustomization::CustomizeDetails( IDetailLayoutBuilder& Detai
 				}
 			}
 
-			if (!Substrate::IsSubstrateEnabled())
+			if (Substrate::IsSubstrateEnabled())
+			{
+				if (PropertyName == GET_MEMBER_NAME_CHECKED(UMaterial, bUseMaterialAttributes))
+				{
+					DetailLayout.HideProperty(PropertyHandle);
+				}
+			}
+			else
 			{
 				if (PropertyName == GET_MEMBER_NAME_CHECKED(UMaterial, bIsThinSurface))
 				{
diff --git a/Engine/Source/Runtime/Engine/Classes/Materials/MaterialExpressionMakeMaterialAttributes.h b/Engine/Source/Runtime/Engine/Classes/Materials/MaterialExpressionMakeMaterialAttributes.h
--- a/Engine/Source/Runtime/Engine/Classes/Materials/MaterialExpressionMakeMaterialAttributes.h
+++ b/Engine/Source/Runtime/Engine/Classes/Materials/MaterialExpressionMakeMaterialAttributes.h
@@ -88,6 +88,8 @@ class UMaterialExpressionMakeMaterialAttributes : public UMaterialExpression
 	virtual bool IsResultMaterialAttributes(int32 OutputIndex) override {return true;}
 	virtual uint32 GetInputType(int32 InputIndex) override;
 	virtual bool GenerateHLSLExpression(FMaterialHLSLGenerator& Generator, UE::HLSLTree::FScope& Scope, int32 OutputIndex, UE::HLSLTree::FExpression const*& OutExpression) const override;
+
+	uint64 GetConnectedInputs() const;
 #endif // WITH_EDITOR
 	//~ End UMaterialExpression Interface
 };
diff --git a/Engine/Source/Runtime/Engine/Classes/Materials/MaterialExpressionSetMaterialAttributes.h b/Engine/Source/Runtime/Engine/Classes/Materials/MaterialExpressionSetMaterialAttributes.h
--- a/Engine/Source/Runtime/Engine/Classes/Materials/MaterialExpressionSetMaterialAttributes.h
+++ b/Engine/Source/Runtime/Engine/Classes/Materials/MaterialExpressionSetMaterialAttributes.h
@@ -45,6 +45,8 @@ class UMaterialExpressionSetMaterialAttributes : public UMaterialExpression
 	virtual void GetExpressionToolTip(TArray<FString>& OutToolTip) override;
 
 	virtual bool GenerateHLSLExpression(FMaterialHLSLGenerator& Generator, UE::HLSLTree::FScope& Scope, int32 OutputIndex, UE::HLSLTree::FExpression const*& OutExpression) const override;
+
+	uint64 GetConnectedInputs() const;
 #endif
 	//~ End UMaterialExpression Interface
 };
diff --git a/Engine/Source/Runtime/Engine/Classes/Materials/MaterialExpressionSubstrate.h b/Engine/Source/Runtime/Engine/Classes/Materials/MaterialExpressionSubstrate.h
--- a/Engine/Source/Runtime/Engine/Classes/Materials/MaterialExpressionSubstrate.h
+++ b/Engine/Source/Runtime/Engine/Classes/Materials/MaterialExpressionSubstrate.h
@@ -821,7 +821,82 @@ class UMaterialExpressionSubstrateConvertToDecal : public UMaterialExpressionSub
 	//~ End UMaterialExpression Interface
 };
 
+UCLASS(collapsecategories, hidecategories=Object, MinimalAPI)
+class UMaterialExpressionSubstrateConvertMaterialAttributes : public UMaterialExpressionSubstrateBSDF
+{
+	GENERATED_UCLASS_BODY()
+
+	UPROPERTY()
+	FMaterialAttributesInput MaterialAttributes;
+
+	/**
+	 * The amount of transmitted light from the back side of the surface to the front side of the surface (type = float3, unit = unitless, defaults to 1)
+	 */
+	UPROPERTY()
+	FExpressionInput TransmittanceColor;
+
+		/**
+	* The single scattering Albedo defining the overall color of the Material (type = float3, unit = unitless, default = 0)
+	 */
+	UPROPERTY()
+	FExpressionInput WaterScatteringCoefficients;
+
+	/**
+	 * The rate at which light is absorbed or out-scattered by the medium. Mean Free Path = 1 / Extinction. (type = float3, unit = 1/cm, default = 0)
+	 */
+	UPROPERTY()
+	FExpressionInput WaterAbsorptionCoefficients;
+
+	/**
+	 * Anisotropy of the volume with values lower than 0 representing back-scattering, equal 0 representing isotropic scattering and greater than 0 representing forward scattering. (type = float, unit = unitless, defaults to 0)
+	 */
+	UPROPERTY()
+	FExpressionInput WaterPhaseG;
 
+	/**
+	 * A scale to apply on the scene color behind the water surface. It can be used to approximate caustics for instance. (type = float3, unit = unitless, defaults to 1)
+	 */
+	UPROPERTY()
+	FExpressionInput ColorScaleBehindWater;
+
+	/**
+	 * Take the bottom clear coat surface normal as input. The normal is considered tangent or world space according to the space properties on the main material node. (type = float3, unit = unitless, defaults to vertex normal)
+	 */
+	UPROPERTY()
+	FExpressionInput ClearCoatNormal;
+
+	/**
+	 * Take the tangent output node as input. The tangent is considered tangent or world space according to the space properties on the main material node. (type = float3, unit = unitless, defaults to vertex tangent)
+	 */
+	UPROPERTY()
+	FExpressionInput CustomTangent;
+
+	/** SubsurfaceProfile, for Screen Space Subsurface Scattering. The profile needs to be set up on both the Substrate diffuse node, and the material node at the moment. */
+	UPROPERTY(EditAnywhere, BlueprintReadOnly, Category = Material, meta = (DisplayName = "Subsurface Profile"))
+	TObjectPtr<class USubsurfaceProfile> SubsurfaceProfile;
+
+	UPROPERTY(EditAnywhere, Category = ShadingModel, meta = (ShowAsInputPin = "Primary", DisplayName = "Single Shading Model"))
+	TEnumAsByte<enum EMaterialShadingModel> ShadingModelOverride = MSM_DefaultLit;
+
+	//~ Begin UMaterialExpression Interface
+#if WITH_EDITOR
+	virtual int32 Compile(class FMaterialCompiler* Compiler, int32 OutputIndex) override;
+	virtual void GetCaption(TArray<FString>& OutCaptions) const override;
+	virtual uint32 GetOutputType(int32 OutputIndex) override;
+	virtual uint32 GetInputType(int32 InputIndex) override;
+	virtual FName GetInputName(int32 InputIndex) const override;
+	virtual bool IsResultSubstrateMaterial(int32 OutputIndex) override;
+	virtual void GatherSubstrateMaterialInfo(FSubstrateMaterialInfo& SubstrateMaterialInfo, int32 OutputIndex) override;
+	virtual FSubstrateOperator* SubstrateGenerateMaterialTopologyTree(class FMaterialCompiler* Compiler, class UMaterialExpression* Parent, int32 OutputIndex) override;
+	virtual void PostEditChangeProperty(FPropertyChangedEvent& PropertyChangedEvent) override;
+	virtual void GetConnectorToolTip(int32 InputIndex, int32 OutputIndex, TArray<FString>& OutToolTip) override;
+	virtual bool IsInputConnectionRequired(int32 InputIndex) const override {return true;}
+	
+	bool HasSSS() const;
+
+#endif // WITH_EDITOR
+	//~ End UMaterialExpression Interface
+};
 
 ///////////////////////////////////////////////////////////////////////////////
 // Operator nodes
@@ -1127,4 +1202,3 @@ class UMaterialExpressionSubstrateThinFilm : public UMaterialExpressionSubstrate
 #endif
 	//~ End UMaterialExpression Interface
 };
-
diff --git a/Engine/Source/Runtime/Engine/Private/Materials/Material.cpp b/Engine/Source/Runtime/Engine/Private/Materials/Material.cpp
--- a/Engine/Source/Runtime/Engine/Private/Materials/Material.cpp
+++ b/Engine/Source/Runtime/Engine/Private/Materials/Material.cpp
@@ -3264,97 +3264,65 @@ void UMaterial::ConvertMaterialToSubstrateMaterial()
 	// Connect all the legacy pin into the conversion node
 	if (bUseMaterialAttributes && EditorOnly->MaterialAttributes.Expression && !EditorOnly->MaterialAttributes.Expression->IsResultSubstrateMaterial(EditorOnly->MaterialAttributes.OutputIndex)) // M_Rifle cause issues there
 	{
-		UMaterialExpressionBreakMaterialAttributes* BreakMatAtt = NewObject<UMaterialExpressionBreakMaterialAttributes>(this);
-		SetPosXAndMoveReferenceToTheRight(BreakMatAtt);
-		MoveConnectionTo(EditorOnly->MaterialAttributes, BreakMatAtt, 0);
+		UMaterialExpressionSubstrateConvertMaterialAttributes* ConvertAttributeNode = NewObject<UMaterialExpressionSubstrateConvertMaterialAttributes>(this);
+		SetPosXAndMoveReferenceToTheRight(ConvertAttributeNode);
+		ConvertAttributeNode->SubsurfaceProfile = bRequireNoSubsurfaceProfile ? nullptr : SubsurfaceProfile;
+		ConvertAttributeNode->Material = this;
 
-		// Check if Anisotropy input is actually used/connected
-		// This is needed/important as if anisotropy is connected a converted material becomes 'complex' instead of 'single'. 
-		// Since this path uses material attributes, there is no direct way to evaluate if anisotropy is connected or not. To find this,
-		// we loop over all the material attribute make/set nodes to figure out if the anisotropy value is ever set.
-		bool bIsAnisotropyConnected = false;
+		MoveConnectionTo(EditorOnly->MaterialAttributes, ConvertAttributeNode, 0);
+
+		// Reconnect custom output to material attribute conversion node
 		{
-			// Gather 'make' attributes nodes
+			check(ConvertAttributeNode);
+			GatherCustomNodes();
+
+			if (ThinTranslucentOutput)
 			{
-				TArray<UMaterialExpressionMakeMaterialAttributes*> MakeAttributeExpressions;
-				GetAllExpressionsInMaterialAndFunctionsOfType<UMaterialExpressionMakeMaterialAttributes>(MakeAttributeExpressions);
-				for (UMaterialExpressionMakeMaterialAttributes* Expression : MakeAttributeExpressions)
-				{
-					bIsAnisotropyConnected = Expression->GetExpressionInput(MP_Anisotropy) != nullptr;
-					if (bIsAnisotropyConnected)
-					{
-						break;
-					}
-				}
+				MoveConnectionTo(*ThinTranslucentOutput->GetInput(0), ConvertAttributeNode, 1);	 // TransmittanceColor
 			}
-
-			// Gather 'set' attributes nodes
-			if (!bIsAnisotropyConnected)
+			if (SingleLayerWaterOutput)
 			{
-				const FGuid AnisotropyGuid = FMaterialAttributeDefinitionMap::GetID(MP_Anisotropy);
-				TArray<UMaterialExpressionSetMaterialAttributes*> SetAttributeExpressions;
-				GetAllExpressionsInMaterialAndFunctionsOfType<UMaterialExpressionSetMaterialAttributes>(SetAttributeExpressions);
-				for (UMaterialExpressionSetMaterialAttributes* Expression : SetAttributeExpressions)
-				{
-					for (const FGuid& AttributeGuid : Expression->AttributeSetTypes)
-					{
-						if (AttributeGuid == AnisotropyGuid)
-						{
-							bIsAnisotropyConnected = true;
-							break;
-						}
-					}
-					if (bIsAnisotropyConnected)
-					{
-						break;
-					}
-				}
+				MoveConnectionTo(*SingleLayerWaterOutput->GetInput(0), ConvertAttributeNode, 2); // WaterScatteringCoefficients
+				MoveConnectionTo(*SingleLayerWaterOutput->GetInput(1), ConvertAttributeNode, 3); // WaterAbsorptionCoefficients
+				MoveConnectionTo(*SingleLayerWaterOutput->GetInput(2), ConvertAttributeNode, 4); // WaterPhaseG
+				MoveConnectionTo(*SingleLayerWaterOutput->GetInput(3), ConvertAttributeNode, 5); // ColorScaleBehindWater
+			}
+			if (ClearCoatBottomNormalOutput)
+			{
+				CopyConnectionTo(*ClearCoatBottomNormalOutput->GetInput(0), ConvertAttributeNode, 6); // ClearCoatNormal
+			}
+			if (TangentOutput)
+			{
+				CopyConnectionTo(*TangentOutput->GetInput(0), ConvertAttributeNode, 7);	// TangentOutput
 			}
 		}
 
-		ConvertNode = NewObject<UMaterialExpressionSubstrateShadingModels>(this);
-		SetPosXAndMoveReferenceToTheRight(ConvertNode);
-		ConvertNode->BaseColor.Connect(0, BreakMatAtt);
-		ConvertNode->Metallic.Connect(1, BreakMatAtt);
-		ConvertNode->Specular.Connect(2, BreakMatAtt);
-		ConvertNode->Roughness.Connect(3, BreakMatAtt);
-		if (bIsAnisotropyConnected)
-		{
-		ConvertNode->Anisotropy.Connect(4, BreakMatAtt);
-		}
-		ConvertNode->EmissiveColor.Connect(5, BreakMatAtt);
-		ConvertNode->Normal.Connect(8, BreakMatAtt);
-		ConvertNode->Tangent.Connect(9, BreakMatAtt);
-		ConvertNode->SubSurfaceColor.Connect(11, BreakMatAtt);
-		ConvertNode->ClearCoat.Connect(12, BreakMatAtt);
-		ConvertNode->ClearCoatRoughness.Connect(13, BreakMatAtt);
-		ConvertNode->Opacity.Connect(6, BreakMatAtt);
-		ConvertNode->SubsurfaceProfile = bRequireNoSubsurfaceProfile ? nullptr : SubsurfaceProfile;
-		bRelinkCustomOutputNodes = true;
-
 		// * Remove support for material attribute
 		// * explicitly connect the Substrate node to the root node
 		// * Forward inputs to the root node (Do not reconnect the Opacity as we handle the opacity by internally within the conversion node)
 		// * Always forward masked opacity because this is required when the blend mode is overriden to masked on a material instance.
 		bUseMaterialAttributes = false;
-		EditorOnly->FrontMaterial.Connect(0, ConvertNode);
-		EditorOnly->Opacity.Connect(6, BreakMatAtt);			// For alpha composite blend mode, Opacity on the root node in case BLEND_AlphaComposite is selected.
-		EditorOnly->OpacityMask.Connect(7, BreakMatAtt);
-		EditorOnly->WorldPositionOffset.Connect(10, BreakMatAtt);
-		EditorOnly->AmbientOcclusion.Connect(14, BreakMatAtt);
-		EditorOnly->PixelDepthOffset.Connect(24, BreakMatAtt);
-		EditorOnly->Refraction.Connect(15, BreakMatAtt);
+		EditorOnly->FrontMaterial.Connect(0, ConvertAttributeNode);
+		EditorOnly->Opacity.Connect(4, ConvertAttributeNode);
+		EditorOnly->OpacityMask.Connect(5, ConvertAttributeNode);
+		EditorOnly->WorldPositionOffset.Connect(3, ConvertAttributeNode);
+		EditorOnly->AmbientOcclusion.Connect(2, ConvertAttributeNode);
+		EditorOnly->PixelDepthOffset.Connect(1, ConvertAttributeNode);
+		EditorOnly->Refraction.Connect(6, ConvertAttributeNode);
+		
 
+		// Shading Model
+		// * either use the shader graph expression 
+		// * or add a constant shading model
 		if (ShadingModel == MSM_FromMaterialExpression)
 		{
-			check(ShadingModels.CountShadingModels() >= 1);
-			ConvertNode->ShadingModel.Connect(25, BreakMatAtt);
+			ConvertAttributeNode->ShadingModelOverride = MSM_FromMaterialExpression;
 		}
 		else
 		{
 			// Store Substrate shading model of the converted material. 
 			check(ShadingModels.CountShadingModels() == 1);
-			ConvertNode->ShadingModelOverride = ShadingModel;
+			ConvertAttributeNode->ShadingModelOverride = ShadingModel;
 		}
 
 		if (MaterialDomain == MD_DeferredDecal)
@@ -3364,8 +3332,8 @@ void UMaterial::ConvertMaterialToSubstrateMaterial()
 
 			// Now pass through the convert to decal node, which flag the material as SSM_Decal, which will set the domain to Decal.
 			UMaterialExpressionSubstrateConvertToDecal* ConvertToDecalNode = NewObject<UMaterialExpressionSubstrateConvertToDecal>(this);
-			ReplaceNodeAndMoveToTheRight(ConvertNode, ConvertToDecalNode);
-			ConvertToDecalNode->DecalMaterial.Connect(0, ConvertNode);
+			ReplaceNodeAndMoveToTheRight(ConvertAttributeNode, ConvertToDecalNode);
+			ConvertToDecalNode->DecalMaterial.Connect(0, ConvertAttributeNode);
 
 			EditorOnly->FrontMaterial.Connect(0, ConvertToDecalNode);
 		}
diff --git a/Engine/Source/Runtime/Engine/Private/Materials/MaterialExpressions.cpp b/Engine/Source/Runtime/Engine/Private/Materials/MaterialExpressions.cpp
--- a/Engine/Source/Runtime/Engine/Private/Materials/MaterialExpressions.cpp
+++ b/Engine/Source/Runtime/Engine/Private/Materials/MaterialExpressions.cpp
@@ -6753,6 +6753,33 @@ void UMaterialExpressionMakeMaterialAttributes::Serialize(FStructuredArchive::FR
 }
 
 #if WITH_EDITOR
+
+// Return a conservative list of connected inputs
+uint64 UMaterialExpressionMakeMaterialAttributes::GetConnectedInputs() const
+{
+	uint64 Out = 0ull;
+	if (BaseColor.Expression != nullptr) 				Out |= (1ull << uint64(MP_BaseColor));
+	if (Metallic.Expression != nullptr) 				Out |= (1ull << uint64(MP_Metallic));
+	if (Specular.Expression != nullptr) 				Out |= (1ull << uint64(MP_Specular));
+	if (Roughness.Expression != nullptr) 				Out |= (1ull << uint64(MP_Roughness));
+	if (Anisotropy.Expression != nullptr) 				Out |= (1ull << uint64(MP_Anisotropy));
+	if (EmissiveColor.Expression != nullptr) 			Out |= (1ull << uint64(MP_EmissiveColor));
+	if (Opacity.Expression != nullptr) 					Out |= (1ull << uint64(MP_Opacity));
+	if (OpacityMask.Expression != nullptr) 				Out |= (1ull << uint64(MP_OpacityMask));
+	if (Normal.Expression != nullptr) 					Out |= (1ull << uint64(MP_Normal));
+	if (Tangent.Expression != nullptr) 					Out |= (1ull << uint64(MP_Tangent));
+	if (WorldPositionOffset.Expression != nullptr) 		Out |= (1ull << uint64(MP_WorldPositionOffset));
+	if (Displacement.Expression != nullptr) 			Out |= (1ull << uint64(MP_Displacement));
+	if (SubsurfaceColor.Expression != nullptr) 			Out |= (1ull << uint64(MP_SubsurfaceColor));
+	if (ClearCoat.Expression != nullptr) 				Out |= (1ull << uint64(MP_CustomData0));
+	if (ClearCoatRoughness.Expression != nullptr) 		Out |= (1ull << uint64(MP_CustomData1));
+	if (AmbientOcclusion.Expression != nullptr) 		Out |= (1ull << uint64(MP_AmbientOcclusion));
+	if (Refraction.Expression != nullptr) 				Out |= (1ull << uint64(MP_Refraction));
+	if (PixelDepthOffset.Expression != nullptr) 		Out |= (1ull << uint64(MP_PixelDepthOffset));
+	if (ShadingModel.Expression != nullptr) 			Out |= (1ull << uint64(MP_ShadingModel));
+	return Out;
+}
+
 int32 UMaterialExpressionMakeMaterialAttributes::Compile(class FMaterialCompiler* Compiler, int32 OutputIndex) 
 {
 	int32 Ret = INDEX_NONE;
@@ -7278,6 +7305,21 @@ UMaterialExpressionSetMaterialAttributes::UMaterialExpressionSetMaterialAttribut
 }
 
 #if WITH_EDITOR
+uint64 UMaterialExpressionSetMaterialAttributes::GetConnectedInputs() const
+{
+	uint64 Out = 0ull;
+	const int32 NumInputPins = AttributeSetTypes.Num();
+	for (int32 i = 0; i < NumInputPins; ++i)
+	{
+		const EMaterialProperty Property = FMaterialAttributeDefinitionMap::GetProperty(AttributeSetTypes[i]);
+		if (Property != MP_MAX)
+		{
+			Out |= 1ull << uint64(Property);
+		}
+	}
+	return Out;
+}
+
 int32 UMaterialExpressionSetMaterialAttributes::Compile(class FMaterialCompiler* Compiler, int32 OutputIndex) 
 {
 	// Verify setup
@@ -23227,7 +23269,6 @@ int32 UMaterialExpressionSubstrateShadingModels::Compile(class FMaterialCompiler
 {
 	int32 RoughnessCodeChunk = CompileWithDefaultFloat1(Compiler, Roughness, 0.5f);
 	int32 AnisotropyCodeChunk = CompileWithDefaultFloat1(Compiler, Anisotropy, 0.0f);
-	// As long as both roughness are potentially different, we must take it into account in our encoding.
 	// We also cannot ignore the tangent when using the default Tangent because GetTangentBasis
 	// used in SubstrateGetBSDFSharedBasis cannot be relied on for smooth tangent used for lighting on any mesh.
 	const bool bHasAnisotropy = HasAnisotropy();
@@ -23304,7 +23345,6 @@ int32 UMaterialExpressionSubstrateShadingModels::Compile(class FMaterialCompiler
 	int32 ShadingModelCodeChunk = ShadingModel.IsConnected() ? CompileWithDefaultFloat1(Compiler, ShadingModel, float(MSM_DefaultLit)) : Compiler->Constant(float(ShadingModelOverride));
 	int32 ShadingModelCount = Compiler->GetMaterialShadingModels().CountShadingModels();
 	const bool bHasDynamicShadingModels = ShadingModelCount > 1;
-	// We probably need to do something along these line as well :::
 	int32 OutputCodeChunk = Compiler->SubstrateConversionFromLegacy(
 		bHasDynamicShadingModels,
 		// Metalness workflow
@@ -23673,8 +23713,6 @@ int32 UMaterialExpressionSubstrateSlabBSDF::Compile(class FMaterialCompiler* Com
 {
 	FSubstrateOperator& SubstrateOperator = Compiler->SubstrateCompilationGetOperator(Compiler->SubstrateTreeStackGetPathUniqueId());
 
-
-	// As long as both roughness are potentially different, we must take it into account in our encoding.
 	// We also cannot ignore the tangent when using the default Tangent because GetTangentBasis
 	// used in SubstrateGetBSDFSharedBasis cannot be relied on for smooth tangent used for lighting on any mesh.
 	const bool bHasAnisotropy		= SubstrateOperator.bBSDFHasAnisotropy > 0;
@@ -26279,6 +26317,474 @@ void UMaterialExpressionSubstrateThinFilm::GetExpressionToolTip(TArray<FString>&
 }
 #endif // WITH_EDITOR
 
+// Return a conservative list of connected material attribute inputs
+#if WITH_EDITOR
+static uint64 GetConnectedMaterialAttributesInputs(const UMaterial* InMaterial)
+{
+	if (!InMaterial) return 0;
+
+	uint64 Out = 0ull;
+	for (const UMaterialExpression* Expression : InMaterial->GetExpressions())
+	{
+		if (Expression)
+		{
+			if (Expression->IsA(UMaterialExpressionSetMaterialAttributes::StaticClass()))
+			{
+				const UMaterialExpressionSetMaterialAttributes* Attr = Cast<UMaterialExpressionSetMaterialAttributes>(Expression);
+				Out |= Attr->GetConnectedInputs();
+			}
+			else if (Expression->IsA(UMaterialExpressionMakeMaterialAttributes::StaticClass()))
+			{
+				const UMaterialExpressionMakeMaterialAttributes* Attr = Cast<UMaterialExpressionMakeMaterialAttributes>(Expression);
+				Out |= Attr->GetConnectedInputs();
+			}
+		}
+	}
+	return Out;
+}
+
+static bool IsMaterialAttributeInputConnected(uint64 InCache, EMaterialProperty InProperty)
+{
+	return !!(InCache & (1ull << uint64(InProperty)));
+}
+#endif
+
+UMaterialExpressionSubstrateConvertMaterialAttributes::UMaterialExpressionSubstrateConvertMaterialAttributes(const FObjectInitializer& ObjectInitializer)
+: Super(ObjectInitializer)
+{
+	struct FConstructorStatics
+	{
+		FText NAME_Strata;
+		FConstructorStatics() : NAME_Strata(LOCTEXT("Substrate Conversion", "Substrate Conversion")) { }
+	};
+	static FConstructorStatics ConstructorStatics;
+#if WITH_EDITORONLY_DATA
+	MenuCategories.Add(ConstructorStatics.NAME_Strata);
+#endif
+	// Cache some material attribute inputs?
+#if WITH_EDITOR
+	CachedInputs.Empty();
+	CachedInputs.Reserve(1);
+	CachedInputs.Add(&MaterialAttributes);
+	CachedInputs.Add(&TransmittanceColor); 
+	CachedInputs.Add(&WaterScatteringCoefficients); 
+	CachedInputs.Add(&WaterAbsorptionCoefficients);
+	CachedInputs.Add(&WaterPhaseG);
+	CachedInputs.Add(&ColorScaleBehindWater); 
+	CachedInputs.Add(&ClearCoatNormal);
+	CachedInputs.Add(&CustomTangent);
+#endif
+
+#if WITH_EDITORONLY_DATA
+	Outputs.Reset();
+	Outputs.Add(FExpressionOutput(TEXT(""))); // Substrate
+	Outputs.Add(FExpressionOutput(TEXT("PixelDepthOffset")));
+	Outputs.Add(FExpressionOutput(TEXT("AmbientOcclusion")));
+	Outputs.Add(FExpressionOutput(TEXT("WorldPositionOffset")));
+	Outputs.Add(FExpressionOutput(TEXT("Opacity")));
+	Outputs.Add(FExpressionOutput(TEXT("OpacityMask")));
+	Outputs.Add(FExpressionOutput(TEXT("Refraction")));
+	bShowOutputNameOnPin = true;
+#endif
+}
+
+#if WITH_EDITOR
+int32 UMaterialExpressionSubstrateConvertMaterialAttributes::Compile(class FMaterialCompiler* Compiler, int32 OutputIndex)
+{
+	if (OutputIndex == 1)
+	{
+		return MaterialAttributes.CompileWithDefault(Compiler, FMaterialAttributeDefinitionMap::GetID(MP_PixelDepthOffset));
+	}
+	else if (OutputIndex == 2)
+	{
+		return MaterialAttributes.CompileWithDefault(Compiler, FMaterialAttributeDefinitionMap::GetID(MP_AmbientOcclusion));
+	}
+	else if (OutputIndex == 3)
+	{
+		return MaterialAttributes.CompileWithDefault(Compiler, FMaterialAttributeDefinitionMap::GetID(MP_WorldPositionOffset));
+	}
+	else if (OutputIndex == 4)
+	{
+		return MaterialAttributes.CompileWithDefault(Compiler, FMaterialAttributeDefinitionMap::GetID(MP_Opacity));
+	}
+	else if (OutputIndex == 5)
+	{
+		return MaterialAttributes.CompileWithDefault(Compiler, FMaterialAttributeDefinitionMap::GetID(MP_OpacityMask));
+	}
+	else if (OutputIndex == 6)
+	{
+		return MaterialAttributes.CompileWithDefault(Compiler, FMaterialAttributeDefinitionMap::GetID(MP_Refraction));
+	}
+	else if (OutputIndex != 0)
+	{
+		return Compiler->Error(TEXT("Output pin index error"));
+	}
+		
+	// We also cannot ignore the tangent when using the default Tangent because GetTangentBasis
+	// used in StrataGetBSDFSharedBasis cannot be relied on for smooth tangent used for lighting on any mesh.
+
+	const uint64 Cached = GetConnectedMaterialAttributesInputs(Material);
+	const bool bHasAnisotropy = IsMaterialAttributeInputConnected(Cached, MP_Anisotropy);
+
+	// Regular normal basis
+	int32 NormalCodeChunk = Compiler->TransformNormalFromRequestedBasisToWorld(MaterialAttributes.CompileWithDefault(Compiler, FMaterialAttributeDefinitionMap::GetID(MP_Normal)));
+
+	// When computing NormalCodeChunk, we invoke TransformNormalFromRequestedBasisToWorld which requires input to be float or float3.
+	// Certain material do not respect this requirement. We handle here a simple recovery when source material doesn't have a valid 
+	// normal (e.g., vec2 normal), and avoid crashing the material compilation. The error will still be reported by the compiler up 
+	// to the user, but the compilation will succeed.
+	if (NormalCodeChunk == INDEX_NONE) { NormalCodeChunk = Compiler->VertexNormal(); } 
+
+	int32 TangentCodeChunk = bHasAnisotropy ? Compiler->TransformNormalFromRequestedBasisToWorld(MaterialAttributes.CompileWithDefault(Compiler, FMaterialAttributeDefinitionMap::GetID(MP_Tangent))) : INDEX_NONE;
+	const FSubstrateRegisteredSharedLocalBasis NewRegisteredSharedLocalBasis = SubstrateCompilationInfoCreateSharedLocalBasis(Compiler, NormalCodeChunk, TangentCodeChunk);
+	const FString BasisIndexMacro = Compiler->GetSubstrateSharedLocalBasisIndexMacro(NewRegisteredSharedLocalBasis);
+
+	const bool bHasCoatNormal = ClearCoatNormal.IsConnected();
+	// Clear coat normal basis
+	int32 ClearCoat_NormalCodeChunk = INDEX_NONE;
+	int32 ClearCoat_TangentCodeChunk = INDEX_NONE;
+	FString ClearCoat_BasisIndexMacro;
+	FSubstrateRegisteredSharedLocalBasis ClearCoat_NewRegisteredSharedLocalBasis;
+	if (bHasCoatNormal)
+	{
+		ClearCoat_NormalCodeChunk = CompileWithDefaultNormalWS(Compiler, ClearCoatNormal);
+		ClearCoat_TangentCodeChunk = TangentCodeChunk;
+		ClearCoat_NewRegisteredSharedLocalBasis = SubstrateCompilationInfoCreateSharedLocalBasis(Compiler, ClearCoat_NormalCodeChunk, ClearCoat_TangentCodeChunk);
+		ClearCoat_BasisIndexMacro = Compiler->GetSubstrateSharedLocalBasisIndexMacro(ClearCoat_NewRegisteredSharedLocalBasis);
+	}
+	else
+	{
+		ClearCoat_NormalCodeChunk = NormalCodeChunk;
+		ClearCoat_TangentCodeChunk = TangentCodeChunk;
+		ClearCoat_NewRegisteredSharedLocalBasis = NewRegisteredSharedLocalBasis;
+		ClearCoat_BasisIndexMacro = BasisIndexMacro;
+	}
+
+	// Custom tangent. No need to register it as a local basis, as it is only used for eye shading internal conversion
+	int32 CustomTangent_TangentCodeChunk = INDEX_NONE;
+
+	const bool bHasCustomTangent = CustomTangent.IsConnected();
+	if (bHasCustomTangent)
+	{
+		// Legacy code doesn't do tangent <-> world basis conversion on tangent output, when provided.
+		CustomTangent_TangentCodeChunk = CompileWithDefaultNormalWS(Compiler, CustomTangent, false /*bConvertToRequestedSpace*/);
+	}
+	else
+	{
+		CustomTangent_TangentCodeChunk = NormalCodeChunk;
+	}
+
+	// Need to handle this by looking at the material instead of the node?
+	int32 SSSProfileCodeChunk = INDEX_NONE;
+	const bool bHasSSS = HasSSS();
+	if (bHasSSS)
+	{
+		SSSProfileCodeChunk = Compiler->ForceCast(Compiler->ScalarParameter(GetSubsurfaceProfileParameterName(), 1.0f), MCT_Float1);
+	}
+	SSSProfileCodeChunk = SSSProfileCodeChunk != INDEX_NONE ? SSSProfileCodeChunk : Compiler->Constant(0.0f);
+
+	FSubstrateOperator& SubstrateOperator = Compiler->SubstrateCompilationGetOperator(Compiler->SubstrateTreeStackGetPathUniqueId());
+	SubstrateOperator.BSDFRegisteredSharedLocalBasis = NewRegisteredSharedLocalBasis;
+
+	int32 OpacityCodeChunk = INDEX_NONE;
+	if (!Compiler->SubstrateSkipsOpacityEvaluation())
+	{
+		// We evaluate opacity only for shading models and blending mode requiring it.
+		// For instance, a translucent shader reading depth for soft fading should no evaluate opacity when an instance forces an opaque mode.
+		OpacityCodeChunk = MaterialAttributes.CompileWithDefault(Compiler, FMaterialAttributeDefinitionMap::GetID(MP_Opacity));
+	}
+	else
+	{
+		OpacityCodeChunk = Compiler->Constant(1.0f);
+	}
+
+	int32 ShadingModelCodeChunk = MaterialAttributes.CompileWithDefault(Compiler, FMaterialAttributeDefinitionMap::GetID(MP_ShadingModel));
+	const bool bHasShadingModelExpression = IsMaterialAttributeInputConnected(Cached, MP_ShadingModel);
+	if (!bHasShadingModelExpression)
+	{
+		ShadingModelCodeChunk = Compiler->Constant(float(ShadingModelOverride));
+	}
+	int32 ShadingModelCount = Compiler->GetMaterialShadingModels().CountShadingModels();
+	const bool bHasDynamicShadingModels = ShadingModelCount > 1;
+	int32 OutputCodeChunk = Compiler->SubstrateConversionFromLegacy(
+		bHasDynamicShadingModels,
+		// Metalness workflow
+		MaterialAttributes.CompileWithDefault(Compiler, FMaterialAttributeDefinitionMap::GetID(MP_BaseColor)),
+		MaterialAttributes.CompileWithDefault(Compiler, FMaterialAttributeDefinitionMap::GetID(MP_Specular)),
+		MaterialAttributes.CompileWithDefault(Compiler, FMaterialAttributeDefinitionMap::GetID(MP_Metallic)),
+		// Roughness
+		MaterialAttributes.CompileWithDefault(Compiler, FMaterialAttributeDefinitionMap::GetID(MP_Roughness)),
+		MaterialAttributes.CompileWithDefault(Compiler, FMaterialAttributeDefinitionMap::GetID(MP_Anisotropy)),
+		// SSS
+		MaterialAttributes.CompileWithDefault(Compiler, FMaterialAttributeDefinitionMap::GetID(MP_SubsurfaceColor)),
+		SSSProfileCodeChunk,
+		// Clear Coat / Custom
+		MaterialAttributes.CompileWithDefault(Compiler, FMaterialAttributeDefinitionMap::GetID(MP_CustomData0)),// Clear coat
+		MaterialAttributes.CompileWithDefault(Compiler, FMaterialAttributeDefinitionMap::GetID(MP_CustomData1)),// Clear coat roughness
+		// Misc
+		MaterialAttributes.CompileWithDefault(Compiler, FMaterialAttributeDefinitionMap::GetID(MP_EmissiveColor)),
+		OpacityCodeChunk,
+		CompileWithDefaultFloat3(Compiler, TransmittanceColor, 0.5f, 0.5f, 0.5f),
+		// Water
+		CompileWithDefaultFloat3(Compiler, WaterScatteringCoefficients, 0.0f, 0.0f, 0.0f),
+		CompileWithDefaultFloat3(Compiler, WaterAbsorptionCoefficients, 0.0f, 0.0f, 0.0f),
+		CompileWithDefaultFloat1(Compiler, WaterPhaseG, 0.0f),
+		CompileWithDefaultFloat3(Compiler, ColorScaleBehindWater, 1.0f, 1.0f, 1.0f),
+		// Shading model
+		ShadingModelCodeChunk,
+		NormalCodeChunk,
+		TangentCodeChunk,
+		BasisIndexMacro,
+		ClearCoat_NormalCodeChunk,
+		ClearCoat_TangentCodeChunk,
+		ClearCoat_BasisIndexMacro,
+		CustomTangent_TangentCodeChunk,
+		!SubstrateOperator.bUseParameterBlending || (SubstrateOperator.bUseParameterBlending && SubstrateOperator.bRootOfParameterBlendingSubTree) ? &SubstrateOperator : nullptr);
+
+	return OutputCodeChunk;
+}
+
+void UMaterialExpressionSubstrateConvertMaterialAttributes::PostEditChangeProperty(FPropertyChangedEvent& PropertyChangedEvent)
+{
+	if (GraphNode)
+	{
+		GraphNode->ReconstructNode();
+	}
+	Super::PostEditChangeProperty(PropertyChangedEvent);
+}
+
+void UMaterialExpressionSubstrateConvertMaterialAttributes::GetCaption(TArray<FString>& OutCaptions) const
+{
+	OutCaptions.Add(TEXT("Substrate Convert Material Attributes"));
+}
+
+uint32 UMaterialExpressionSubstrateConvertMaterialAttributes::GetOutputType(int32 OutputIndex)
+{
+	switch (OutputIndex)
+	{
+		case 0: return MCT_Substrate;	// Substrata data
+		case 1: return MCT_Float1;		// PixelDepthOffset
+		case 2: return MCT_Float1; 		// AmbientOcclusion
+		case 3: return MCT_Float3; 		// WorldPositionOffset
+		case 4: return MCT_Float1; 		// Opacity
+		case 5: return MCT_Float1; 		// OpacityMask
+		case 6: return MCT_Float; 		// Refraction
+	}
+	check(false);
+	return MCT_Float1;
+}
+
+uint32 UMaterialExpressionSubstrateConvertMaterialAttributes::GetInputType(int32 InputIndex)
+{
+	if (InputIndex == 0)	  return MCT_MaterialAttributes; // MaterialAttributes
+	else if (InputIndex == 1) return MCT_Float3; // TransmittanceColor
+	else if (InputIndex == 2) return MCT_Float3; // WaterScatteringCoefficients
+	else if (InputIndex == 3) return MCT_Float3; // WaterAbsorptionCoefficients
+	else if (InputIndex == 4) return MCT_Float1; // WaterPhaseG
+	else if (InputIndex == 5) return MCT_Float3; // ColorScaleBehindWater
+	else if (InputIndex == 6) return MCT_Float3; // ClearCoatNormal
+	else if (InputIndex == 7) return MCT_Float3; // CustomTangent
+	else if (InputIndex == 8) return MCT_ShadingModel; // ShadingModelOverride (as it uses 'ShowAsInputPin' metadata)
+	
+	check(false);
+	return MCT_Float1;
+}
+
+FName UMaterialExpressionSubstrateConvertMaterialAttributes::GetInputName(int32 InputIndex) const
+{
+	if (InputIndex == 0)		return TEXT("Attributes");
+	else if (InputIndex == 1)	return TEXT("TransmittanceColor (ThinTranslucent)");
+	else if (InputIndex == 2)	return TEXT("Water Scattering Coefficients (Water)");
+	else if (InputIndex == 3)	return TEXT("Water Absorption Coefficients (Water)");
+	else if (InputIndex == 4)	return TEXT("Water Phase G (Water)");
+	else if (InputIndex == 5)	return TEXT("Color Scale BehindWater (Water)");
+	else if (InputIndex == 6)	return TEXT("Clear Coat Normal (Coat/Eye)");
+	else if (InputIndex == 7)	return TEXT("Custom Tangent (Eye)");
+	else if (InputIndex == 8)	return TEXT("Shading Model From Expression");
+	return NAME_None;
+}
+
+void UMaterialExpressionSubstrateConvertMaterialAttributes::GetConnectorToolTip(int32 InputIndex, int32 OutputIndex, TArray<FString>& OutToolTip)
+{
+	switch (OutputIndex)
+	{
+		case 0: OutToolTip.Add(TEXT("TT Ouput")); break;
+		case 1: OutToolTip.Add(TEXT("TT PixelDepthOffset")); break;
+		case 2: OutToolTip.Add(TEXT("TT AmbientOcclusion")); break;
+		case 3: OutToolTip.Add(TEXT("TT WorldPositionOffset")); break;
+		case 4: OutToolTip.Add(TEXT("TT Opacity")); break;
+		case 5: OutToolTip.Add(TEXT("TT OpacityMask")); break;
+		case 6: OutToolTip.Add(TEXT("TT Refraction")); break;
+	}
+	Super::GetConnectorToolTip(InputIndex, INDEX_NONE, OutToolTip);
+}
+
+bool UMaterialExpressionSubstrateConvertMaterialAttributes::IsResultSubstrateMaterial(int32 OutputIndex)
+{
+	return OutputIndex == 0;
+}
+
+void UMaterialExpressionSubstrateConvertMaterialAttributes::GatherSubstrateMaterialInfo(FSubstrateMaterialInfo& SubstrateMaterialInfo, int32 OutputIndex)
+{
+	const uint64 Cached = GetConnectedMaterialAttributesInputs(Material);
+
+	if (IsMaterialAttributeInputConnected(Cached, MP_BaseColor)) 		{ SubstrateMaterialInfo.AddPropertyConnected(MP_BaseColor); }
+	if (IsMaterialAttributeInputConnected(Cached, MP_Metallic)) 		{ SubstrateMaterialInfo.AddPropertyConnected(MP_Metallic); }
+	if (IsMaterialAttributeInputConnected(Cached, MP_Specular)) 		{ SubstrateMaterialInfo.AddPropertyConnected(MP_Specular); }
+	if (IsMaterialAttributeInputConnected(Cached, MP_Roughness)) 		{ SubstrateMaterialInfo.AddPropertyConnected(MP_Roughness); }
+	if (IsMaterialAttributeInputConnected(Cached, MP_Anisotropy)) 		{ SubstrateMaterialInfo.AddPropertyConnected(MP_Anisotropy); }
+	if (IsMaterialAttributeInputConnected(Cached, MP_EmissiveColor)) 	{ SubstrateMaterialInfo.AddPropertyConnected(MP_EmissiveColor); }
+	if (IsMaterialAttributeInputConnected(Cached, MP_Normal)) 			{ SubstrateMaterialInfo.AddPropertyConnected(MP_Normal); }
+	if (IsMaterialAttributeInputConnected(Cached, MP_Tangent)) 			{ SubstrateMaterialInfo.AddPropertyConnected(MP_Tangent); }
+	if (IsMaterialAttributeInputConnected(Cached, MP_SubsurfaceColor)) 	{ SubstrateMaterialInfo.AddPropertyConnected(MP_SubsurfaceColor); }
+	if (IsMaterialAttributeInputConnected(Cached, MP_CustomData0)) 		{ SubstrateMaterialInfo.AddPropertyConnected(MP_CustomData0); }
+	if (IsMaterialAttributeInputConnected(Cached, MP_CustomData1)) 		{ SubstrateMaterialInfo.AddPropertyConnected(MP_CustomData1); }
+	if (IsMaterialAttributeInputConnected(Cached, MP_Opacity)) 			{ SubstrateMaterialInfo.AddPropertyConnected(MP_Opacity); }
+
+	if (IsMaterialAttributeInputConnected(Cached, MP_ShadingModel))
+	{
+		SubstrateMaterialInfo.AddPropertyConnected(MP_ShadingModel);
+
+		// If the ShadingModel pin is plugged in, we must use a shading model from expression path.
+		SubstrateMaterialInfo.SetShadingModelFromExpression(true);
+	}
+	else
+	{
+		// If the ShadingModel pin is NOT plugged in, we simply use the shading model selected on the root node drop box.
+		if (ShadingModelOverride == MSM_Unlit)					{ SubstrateMaterialInfo.AddShadingModel(ESubstrateShadingModel::SSM_Unlit); }
+		if (ShadingModelOverride == MSM_DefaultLit)				{ SubstrateMaterialInfo.AddShadingModel(ESubstrateShadingModel::SSM_DefaultLit); }
+		if (ShadingModelOverride == MSM_Subsurface)				{ SubstrateMaterialInfo.AddShadingModel(ESubstrateShadingModel::SSM_SubsurfaceWrap); }
+		if (ShadingModelOverride == MSM_PreintegratedSkin)		{ SubstrateMaterialInfo.AddShadingModel(ESubstrateShadingModel::SSM_SubsurfaceWrap); }
+		if (ShadingModelOverride == MSM_ClearCoat)				{ SubstrateMaterialInfo.AddShadingModel(ESubstrateShadingModel::SSM_ClearCoat); }
+		if (ShadingModelOverride == MSM_SubsurfaceProfile)		{ SubstrateMaterialInfo.AddShadingModel(ESubstrateShadingModel::SSM_SubsurfaceProfile); }
+		if (ShadingModelOverride == MSM_TwoSidedFoliage)		{ SubstrateMaterialInfo.AddShadingModel(ESubstrateShadingModel::SSM_SubsurfaceThinTwoSided); }
+		if (ShadingModelOverride == MSM_Hair)					{ SubstrateMaterialInfo.AddShadingModel(ESubstrateShadingModel::SSM_Hair); }
+		if (ShadingModelOverride == MSM_Cloth)					{ SubstrateMaterialInfo.AddShadingModel(ESubstrateShadingModel::SSM_Cloth); }
+		if (ShadingModelOverride == MSM_Eye)					{ SubstrateMaterialInfo.AddShadingModel(ESubstrateShadingModel::SSM_Eye); }
+		if (ShadingModelOverride == MSM_SingleLayerWater)		{ SubstrateMaterialInfo.AddShadingModel(ESubstrateShadingModel::SSM_SingleLayerWater); }
+		if (ShadingModelOverride == MSM_ThinTranslucent)		{ SubstrateMaterialInfo.AddShadingModel(ESubstrateShadingModel::SSM_ThinTranslucent); }
+	}
+
+	if (SubsurfaceProfile)
+	{
+		SubstrateMaterialInfo.AddSubsurfaceProfile(SubsurfaceProfile);
+	}
+}
+
+FSubstrateOperator* UMaterialExpressionSubstrateConvertMaterialAttributes::SubstrateGenerateMaterialTopologyTree(class FMaterialCompiler* Compiler, class UMaterialExpression* Parent, int32 OutputIndex)
+{
+	const uint64 Cached = GetConnectedMaterialAttributesInputs(Material);
+
+	// Note Thickness has no meaning/usage in the context of StrataLegacyConversionNode
+	int32 ThicknessIndex = Compiler->SubstrateThicknessStackGetThicknessIndex();
+
+	const bool bHasAnisotropy = IsMaterialAttributeInputConnected(Cached, MP_Anisotropy);
+
+	auto AddDefaultWorstCase = [&](bool bSSS, bool bFuzz)
+	{
+
+		FSubstrateOperator& SlabOperator = Compiler->SubstrateCompilationRegisterOperator(SUBSTRATE_OPERATOR_BSDF_LEGACY, Compiler->SubstrateTreeStackGetPathUniqueId(), this, Parent, Compiler->SubstrateTreeStackGetParentPathUniqueId());
+		SlabOperator.BSDFType = SUBSTRATE_BSDF_TYPE_SLAB;
+		SlabOperator.bBSDFHasSSS = bSSS;
+		SlabOperator.bBSDFHasMFPPluggedIn = bSSS;
+		SlabOperator.bBSDFHasFuzz = bFuzz;
+		SlabOperator.bBSDFHasAnisotropy = bHasAnisotropy;
+		SlabOperator.ThicknessIndex = ThicknessIndex;
+
+		return &SlabOperator;
+	};
+
+	// Get the shading models resulting from the UMaterial::RebuildShadingModelField().
+	FMaterialShadingModelField ShadingModels = Compiler->GetMaterialShadingModels();
+
+	// Logic about shading models and complexity should match UMaterialExpressionSubstrateConvertMaterialAttributes::Compile.
+	const bool bHasShadingModelFromExpression = IsMaterialAttributeInputConnected(Cached, MP_ShadingModel); // We keep HasShadingModelFromExpression in case all shading models cannot be safely recovered from material functions.
+	if ((ShadingModels.CountShadingModels() > 1) || bHasShadingModelFromExpression) 
+	{
+		return AddDefaultWorstCase(true, true);
+	}
+	else
+	{
+		check(ShadingModels.CountShadingModels() == 1);
+
+		if (ShadingModels.HasShadingModel(MSM_Unlit))
+		{
+			FSubstrateOperator& Operator = Compiler->SubstrateCompilationRegisterOperator(SUBSTRATE_OPERATOR_BSDF_LEGACY, Compiler->SubstrateTreeStackGetPathUniqueId(), this, Parent, Compiler->SubstrateTreeStackGetParentPathUniqueId());
+			Operator.BSDFType = SUBSTRATE_BSDF_TYPE_UNLIT;
+			Operator.ThicknessIndex = ThicknessIndex;
+			return &Operator;
+		}
+		else if (ShadingModels.HasShadingModel(MSM_DefaultLit))
+		{
+			return AddDefaultWorstCase(false, false);
+		}
+		else if (ShadingModels.HasShadingModel(MSM_ThinTranslucent))
+		{
+			return AddDefaultWorstCase(false, false);
+		}
+		else if (ShadingModels.HasShadingModel(MSM_SubsurfaceProfile))
+		{
+			return AddDefaultWorstCase(true, false);
+		}
+		else if (ShadingModels.HasShadingModel(MSM_Subsurface))
+		{
+			return AddDefaultWorstCase(true, false);
+		}
+		else if (ShadingModels.HasShadingModel(MSM_TwoSidedFoliage))
+		{
+			return AddDefaultWorstCase(true, false);
+		}
+		else if (ShadingModels.HasShadingModel(MSM_PreintegratedSkin))
+		{
+			return AddDefaultWorstCase(true, false);
+		}
+		else if (ShadingModels.HasShadingModel(MSM_Cloth))
+		{
+			return AddDefaultWorstCase(false, true);
+		}
+		else if (ShadingModels.HasShadingModel(MSM_ClearCoat))
+		{
+			FSubstrateOperator& Operator = Compiler->SubstrateCompilationRegisterOperator(SUBSTRATE_OPERATOR_BSDF_LEGACY, Compiler->SubstrateTreeStackGetPathUniqueId(), this, Parent, Compiler->SubstrateTreeStackGetParentPathUniqueId());
+			Operator.BSDFType = SUBSTRATE_BSDF_TYPE_SLAB;
+			Operator.ThicknessIndex = ThicknessIndex;
+			Operator.bBSDFHasSecondRoughnessOrSimpleClearCoat = true;
+			Operator.bBSDFHasAnisotropy = bHasAnisotropy;
+			return &Operator;
+		}
+		else if (ShadingModels.HasShadingModel(MSM_Hair))
+		{
+			FSubstrateOperator& Operator = Compiler->SubstrateCompilationRegisterOperator(SUBSTRATE_OPERATOR_BSDF_LEGACY, Compiler->SubstrateTreeStackGetPathUniqueId(), this, Parent, Compiler->SubstrateTreeStackGetParentPathUniqueId());
+			Operator.BSDFType = SUBSTRATE_BSDF_TYPE_HAIR;
+			Operator.ThicknessIndex = ThicknessIndex;
+			return &Operator;
+		}
+		else if (ShadingModels.HasShadingModel(MSM_Eye))
+		{
+			FSubstrateOperator& Operator = Compiler->SubstrateCompilationRegisterOperator(SUBSTRATE_OPERATOR_BSDF_LEGACY, Compiler->SubstrateTreeStackGetPathUniqueId(), this, Parent, Compiler->SubstrateTreeStackGetParentPathUniqueId());
+			Operator.BSDFType = SUBSTRATE_BSDF_TYPE_EYE;
+			Operator.ThicknessIndex = ThicknessIndex;
+			return &Operator;
+		}
+		else if (ShadingModels.HasShadingModel(MSM_SingleLayerWater))
+		{
+			FSubstrateOperator& Operator = Compiler->SubstrateCompilationRegisterOperator(SUBSTRATE_OPERATOR_BSDF_LEGACY, Compiler->SubstrateTreeStackGetPathUniqueId(), this, Parent, Compiler->SubstrateTreeStackGetParentPathUniqueId());
+			Operator.BSDFType = SUBSTRATE_BSDF_TYPE_SINGLELAYERWATER;
+			Operator.ThicknessIndex = ThicknessIndex;
+			return &Operator;
+		}
+
+		check(false);
+		static FSubstrateOperator DefaultOperatorOnError;
+		return &DefaultOperatorOnError;
+	}
+}
+
+bool UMaterialExpressionSubstrateConvertMaterialAttributes::HasSSS() const
+{
+	return SubsurfaceProfile != nullptr;
+}
+
+#endif // WITH_EDITOR
 
 UMaterialExpressionExecBegin::UMaterialExpressionExecBegin(const FObjectInitializer& ObjectInitializer)
 	: Super(ObjectInitializer)
