Fix crash on undo operation on SubstrateShadingModels and SubstrateConvertMaterialAttributes node.

The undo operation relies on node serialization, and applying graph modification creates incorrect behavior. The undo operation reapplies correctly input renaming.

diff --git a/Engine/Source/Runtime/Engine/Private/Materials/MaterialExpressions.cpp b/Engine/Source/Runtime/Engine/Private/Materials/MaterialExpressions.cpp
--- a/Engine/Source/Runtime/Engine/Private/Materials/MaterialExpressions.cpp
+++ b/Engine/Source/Runtime/Engine/Private/Materials/MaterialExpressions.cpp
@@ -23386,7 +23386,7 @@ int32 UMaterialExpressionSubstrateShadingModels::Compile(class FMaterialCompiler
 
 void UMaterialExpressionSubstrateShadingModels::PostEditChangeProperty(FPropertyChangedEvent& PropertyChangedEvent)
 {
-	if (GraphNode)
+	if (GraphNode && PropertyChangedEvent.Property != nullptr)
 	{
 		GraphNode->ReconstructNode();
 	}
@@ -26546,7 +26546,7 @@ int32 UMaterialExpressionSubstrateConvertMaterialAttributes::Compile(class FMate
 
 void UMaterialExpressionSubstrateConvertMaterialAttributes::PostEditChangeProperty(FPropertyChangedEvent& PropertyChangedEvent)
 {
-	if (GraphNode)
+	if (GraphNode && PropertyChangedEvent.Property != nullptr)
 	{
 		GraphNode->ReconstructNode();
 	}
