Fix invalid Unlit material with Substrate

diff --git a/Engine/Source/Runtime/Engine/Private/Materials/MaterialExpressions.cpp b/Engine/Source/Runtime/Engine/Private/Materials/MaterialExpressions.cpp
--- a/Engine/Source/Runtime/Engine/Private/Materials/MaterialExpressions.cpp
+++ b/Engine/Source/Runtime/Engine/Private/Materials/MaterialExpressions.cpp
@@ -23585,6 +23585,14 @@ FSubstrateOperator* UMaterialExpressionSubstrateShadingModels::SubstrateGenerate
 	const bool bHasShadingModelFromExpression = ShadingModel.IsConnected(); // We keep HasShadingModelFromExpression in case all shading models cannot be safely recovered from material functions.
 	if ((ShadingModels.CountShadingModels() > 1) || bHasShadingModelFromExpression) 
 	{
+		// Special case for unlit only material to get fast path
+		if (ShadingModels.HasOnlyShadingModel(MSM_Unlit))
+		{
+			FSubstrateOperator& Operator = Compiler->SubstrateCompilationRegisterOperator(SUBSTRATE_OPERATOR_BSDF_LEGACY, Compiler->SubstrateTreeStackGetPathUniqueId(), this, Parent, Compiler->SubstrateTreeStackGetParentPathUniqueId());
+			Operator.BSDFType = SUBSTRATE_BSDF_TYPE_UNLIT;
+			Operator.ThicknessIndex = ThicknessIndex;
+			return &Operator;
+		}
 		return AddDefaultWorstCase(true, true);
 	}
 	// else
