Substrate - fixed bad domain (non decal) when using a decal from a SubstrateShadingModel node where the SM is FromExpression.

diff --git a/Engine/Source/Runtime/Engine/Private/Materials/Material.cpp b/Engine/Source/Runtime/Engine/Private/Materials/Material.cpp
--- a/Engine/Source/Runtime/Engine/Private/Materials/Material.cpp
+++ b/Engine/Source/Runtime/Engine/Private/Materials/Material.cpp
@@ -4848,15 +4848,16 @@ void UMaterial::RebuildShadingModelField()
 			UE_LOG(LogMaterial, Error, TEXT("%s: Material information is invalid."), *GetName());
 		}
 
+		if (SubstrateMaterialInfo.HasShadingModel(SSM_Decal))
+		{
+			// Keep the decals information and remove the shading from expression part 
+			// since we are going to bake down the material to a single slab using parameter blending.
+			SubstrateMaterialInfo.SetShadingModelFromExpression(false);
+		}
+
 		if (SubstrateMaterialInfo.CountShadingModels() > 1)
 		{
-			if (SubstrateMaterialInfo.HasShadingModel(SSM_Decal))
-			{
-				// Keep the decals information and remove the shading from expression part 
-				// since we are going to bake down the material to a single slab using parameter blending.
-				SubstrateMaterialInfo.SetShadingModelFromExpression(false);
-			}
-			else if (SubstrateMaterialInfo.HasShadingModelFromExpression())
+			if (SubstrateMaterialInfo.HasShadingModelFromExpression())
 			{
 				if (BlendMode == EBlendMode::BLEND_Opaque || BlendMode == EBlendMode::BLEND_Masked)
 				{
