KismetRenderingLibrary.cpp: Change to CreateStaticTextureEditorOnly, removed some bad code that forced the path to always point to the game content folder, precluding the use of plugin paths.

diff --git a/Engine/Source/Runtime/Engine/Private/KismetRenderingLibrary.cpp b/Engine/Source/Runtime/Engine/Private/KismetRenderingLibrary.cpp
--- a/Engine/Source/Runtime/Engine/Private/KismetRenderingLibrary.cpp
+++ b/Engine/Source/Runtime/Engine/Private/KismetRenderingLibrary.cpp
@@ -580,7 +580,7 @@ UTexture2D* UKismetRenderingLibrary::RenderTargetCreateStaticTexture2DEditorOnly
 		IAssetTools& AssetTools = FModuleManager::Get().LoadModuleChecked<FAssetToolsModule>("AssetTools").Get();
 
 		//Use asset name only if directories are specified, otherwise full path
-		if (!InName.Contains(TEXT("/")))
+		if (!InName.StartsWith(TEXT("/")))
 		{
 			FString AssetName = RenderTarget->GetOutermost()->GetName();
 			const FString SanitizedBasePackageName = UPackageTools::SanitizePackageName(AssetName);
@@ -589,9 +589,6 @@ UTexture2D* UKismetRenderingLibrary::RenderTargetCreateStaticTexture2DEditorOnly
 		}
 		else
 		{
-			InName.RemoveFromStart(TEXT("/"));
-			InName.RemoveFromStart(TEXT("Content/"));
-			InName.StartsWith(TEXT("Game/")) == true ? InName.InsertAt(0, TEXT("/")) : InName.InsertAt(0, TEXT("/Game/"));
 			AssetTools.CreateUniqueAssetName(InName, TEXT(""), PackageName, Name);
 		}
 
