World Partition: make sure to batch actors level removals when destroying loader adapters.

diff --git a/Engine/Source/Runtime/Engine/Private/WorldPartition/WorldPartitionActorLoaderInterface.cpp b/Engine/Source/Runtime/Engine/Private/WorldPartition/WorldPartitionActorLoaderInterface.cpp
--- a/Engine/Source/Runtime/Engine/Private/WorldPartition/WorldPartitionActorLoaderInterface.cpp
+++ b/Engine/Source/Runtime/Engine/Private/WorldPartition/WorldPartitionActorLoaderInterface.cpp
@@ -39,6 +39,9 @@ IWorldPartitionActorLoaderInterface::ILoaderAdapter::ILoaderAdapter(UWorld* InWo
 IWorldPartitionActorLoaderInterface::ILoaderAdapter::~ILoaderAdapter()
 {
 	UnregisterDelegates();
+
+	FWorldPartitionLoadingContext::FDeferred LoadingContext;
+	ContainerActorReferences.Empty();
 }
 
 void IWorldPartitionActorLoaderInterface::ILoaderAdapter::Load()
