Fixed infinite loop when duplicating a level instance in editor

diff --git a/Engine/Source/Runtime/Engine/Private/LevelInstance/LevelInstanceSubsystem.cpp b/Engine/Source/Runtime/Engine/Private/LevelInstance/LevelInstanceSubsystem.cpp
--- a/Engine/Source/Runtime/Engine/Private/LevelInstance/LevelInstanceSubsystem.cpp
+++ b/Engine/Source/Runtime/Engine/Private/LevelInstance/LevelInstanceSubsystem.cpp
@@ -233,7 +233,7 @@ void ULevelInstanceSubsystem::UpdateStreamingState()
 	UpdateStreamingStateInternal();
 
 #if WITH_EDITOR
-	if (!GetWorld()->IsGameWorld())
+	if (!GetWorld()->IsGameWorld() && !GUndo)
 	{
 		// For Editor Worlds make sure UpdateStreamingState completes all recursive loading/unloading
 		while (LevelInstancesToLoadOrUpdate.Num() || LevelInstancesToUnload.Num())
