[Navigation] Use partial update when creating and destroying physics state of ISMC instead of doing a full navigation update.
That prevents huge amount of tile dirtiness.

diff --git a/Engine/Source/Runtime/Engine/Private/InstancedStaticMesh.cpp b/Engine/Source/Runtime/Engine/Private/InstancedStaticMesh.cpp
--- a/Engine/Source/Runtime/Engine/Private/InstancedStaticMesh.cpp
+++ b/Engine/Source/Runtime/Engine/Private/InstancedStaticMesh.cpp
@@ -3050,6 +3050,8 @@ void UInstancedStaticMeshComponent::CreateAllInstanceBodies()
 					Transforms.Add(InstanceTM);
 				}
 			}
+
+	    	PartialNavigationUpdate(i);
 	    }
 
 		if (InstanceBodiesSanitized.Num() > 0 && Mobility != EComponentMobility::Movable)
@@ -3075,6 +3077,8 @@ void UInstancedStaticMeshComponent::ClearAllInstanceBodies()
 	{
 		if (InstanceBodies[i])
 		{
+			PartialNavigationUpdate(i);
+			
 			InstanceBodies[i]->TermBody();
 			delete InstanceBodies[i];
 		}
@@ -3104,7 +3108,6 @@ void UInstancedStaticMeshComponent::OnCreatePhysicsState()
 	// Since StaticMeshComponent was not called
 	// Navigation relevancy needs to be handled here
 	bNavigationRelevant = IsNavigationRelevant();
-	FNavigationSystem::UpdateComponentData(*this);
 }
 
 void UInstancedStaticMeshComponent::OnDestroyPhysicsState()
@@ -3117,7 +3120,6 @@ void UInstancedStaticMeshComponent::OnDestroyPhysicsState()
 	// Since StaticMeshComponent was not called
 	// Navigation relevancy needs to be handled here
 	bNavigationRelevant = IsNavigationRelevant();
-	FNavigationSystem::UpdateComponentData(*this);
 }
 
 Chaos::FPhysicsObject* UInstancedStaticMeshComponent::GetPhysicsObjectById(Chaos::FPhysicsObjectId Id) const
