Fix Visual Logger Rewind Debugger integration

diff --git a/Engine/Source/Runtime/Engine/Private/VisualLogger/VisualLoggerTraceDevice.cpp b/Engine/Source/Runtime/Engine/Private/VisualLogger/VisualLoggerTraceDevice.cpp
--- a/Engine/Source/Runtime/Engine/Private/VisualLogger/VisualLoggerTraceDevice.cpp
+++ b/Engine/Source/Runtime/Engine/Private/VisualLogger/VisualLoggerTraceDevice.cpp
@@ -7,6 +7,7 @@
 #include "ObjectTrace.h"
 #include "Trace/Trace.inl"
 #include "VisualLogger/VisualLoggerCustomVersion.h"
+#include "UObject/UE5MainStreamObjectVersion.h"
 
 UE_TRACE_CHANNEL_DEFINE(VisualLoggerChannel);
 
@@ -55,6 +56,9 @@ void FVisualLoggerTraceDevice::Serialize(const UObject* LogOwner, FName OwnerNam
 	{
 		FBufferArchive Archive;
 		Archive.UsingCustomVersion(EVisualLoggerVersion::GUID);
+		Archive.UsingCustomVersion(FUE5MainStreamObjectVersion::GUID);
+		Archive.SetCustomVersion(FUE5MainStreamObjectVersion::GUID, FUE5MainStreamObjectVersion::LatestVersion, "UE5MainStreamObjectVersion");
+
 		Archive << const_cast<FVisualLogEntry&>(LogEntry);
 
 		UE_TRACE_LOG(VisualLogger, VisualLogEntry, VisualLoggerChannel)
