LFV - apply component transform operations in all cases.

diff --git a/Engine/Source/Runtime/Engine/Private/Components/LocalFogVolumeComponent.cpp b/Engine/Source/Runtime/Engine/Private/Components/LocalFogVolumeComponent.cpp
--- a/Engine/Source/Runtime/Engine/Private/Components/LocalFogVolumeComponent.cpp
+++ b/Engine/Source/Runtime/Engine/Private/Components/LocalFogVolumeComponent.cpp
@@ -41,10 +41,7 @@ void ULocalFogVolumeComponent::SendRenderTransformCommand()
 			[SceneProxy, ComponentTransform, HeightOffset](FRHICommandList& RHICmdList)
 			{
 				// Nothing else is needed so that command could actually go.
-				SceneProxy->FogTransform = ComponentTransform;
-
-				const float MaximumAxisScale = SceneProxy->FogTransform.GetMaximumAxisScale();
-				SceneProxy->FogTransform.SetScale3D(FVector(MaximumAxisScale, MaximumAxisScale, MaximumAxisScale));
+				SceneProxy->UpdateComponentTransform(ComponentTransform);
 			});
 	}
 }
diff --git a/Engine/Source/Runtime/Engine/Private/LocalFogVolumeSceneProxy.cpp b/Engine/Source/Runtime/Engine/Private/LocalFogVolumeSceneProxy.cpp
--- a/Engine/Source/Runtime/Engine/Private/LocalFogVolumeSceneProxy.cpp
+++ b/Engine/Source/Runtime/Engine/Private/LocalFogVolumeSceneProxy.cpp
@@ -10,8 +10,7 @@
 
 
 FLocalFogVolumeSceneProxy::FLocalFogVolumeSceneProxy(const ULocalFogVolumeComponent* InComponent)
-	: FogTransform(InComponent->GetComponentTransform())
-	, FogDensity(InComponent->FogDensity)
+	: FogDensity(InComponent->FogDensity)
 	, FogHeightFalloff(InComponent->FogHeightFalloff)
 	, FogHeightOffset(InComponent->FogHeightOffset)
 	, FogRadialAttenuation(InComponent->FogRadialAttenuation)
@@ -21,10 +20,17 @@ FLocalFogVolumeSceneProxy::FLocalFogVolumeSceneProxy(const ULocalFogVolumeCompon
 	, FogAlbedo(InComponent->FogAlbedo)
 	, FogEmissive(InComponent->FogEmissive)
 {
+	UpdateComponentTransform(InComponent->GetComponentTransform());
 }
 
 FLocalFogVolumeSceneProxy::~FLocalFogVolumeSceneProxy()
 {
 }
 
+void FLocalFogVolumeSceneProxy::UpdateComponentTransform(const FTransform& Transform)
+{
+	FogTransform = Transform;
+	const float MaximumAxisScale = FogTransform.GetMaximumAxisScale();
+	FogTransform.SetScale3D(FVector(MaximumAxisScale, MaximumAxisScale, MaximumAxisScale));
+}
 
diff --git a/Engine/Source/Runtime/Engine/Public/LocalFogVolumeSceneProxy.h b/Engine/Source/Runtime/Engine/Public/LocalFogVolumeSceneProxy.h
--- a/Engine/Source/Runtime/Engine/Public/LocalFogVolumeSceneProxy.h
+++ b/Engine/Source/Runtime/Engine/Public/LocalFogVolumeSceneProxy.h
@@ -16,6 +16,8 @@ public:
 	ENGINE_API FLocalFogVolumeSceneProxy(const ULocalFogVolumeComponent* InComponent);
 	ENGINE_API ~FLocalFogVolumeSceneProxy();
 
+	void UpdateComponentTransform(const FTransform& Transform);
+
 	FTransform FogTransform;
 
 	float FogDensity;
