Fix device id access crash by making modulation update call not require a device look-up internally and keep weak pointers to all data acted upon.
#rnx

diff --git a/Engine/Source/Runtime/Engine/Private/SoundModulationDestination.cpp b/Engine/Source/Runtime/Engine/Private/SoundModulationDestination.cpp
--- a/Engine/Source/Runtime/Engine/Private/SoundModulationDestination.cpp
+++ b/Engine/Source/Runtime/Engine/Private/SoundModulationDestination.cpp
@@ -307,35 +307,44 @@ namespace Audio
 
 	void FModulationDestination::UpdateModulatorsInternal(TArray<TUniquePtr<Audio::IModulatorSettings>>&& ProxySettings)
 	{
-		auto UpdateHandleLambda = [DestinationDataPtr = TSharedPtr<FModulationDestinationData>(DestinationData), ModSettings = MoveTemp(ProxySettings)]() mutable
+		FAudioDeviceManager* AudioDeviceManager = FAudioDeviceManager::Get();
+		if (!AudioDeviceManager)
 		{
-			if (DestinationDataPtr.IsValid())
+			return;
+		}
+
+		const FDeviceId DeviceId = DestinationData->GetDeviceId();
+		FAudioDevice* AudioDevice = AudioDeviceManager->GetAudioDeviceRaw(DeviceId);
+		if (!AudioDevice || !AudioDevice->IsModulationPluginEnabled() || !AudioDevice->ModulationInterface.IsValid())
+		{
+			return;
+		}
+
+		FAudioThread::RunCommandOnAudioThread(
+		[
+			DestinationDataPtr = TWeakPtr<FModulationDestinationData>(DestinationData),
+			ModInterfacePtr = TWeakPtr<IAudioModulationManager>(AudioDevice->ModulationInterface),
+			ModSettings = MoveTemp(ProxySettings)
+		]() mutable
+		{
+			TSharedPtr<FModulationDestinationData> DestDataPtr = DestinationDataPtr.Pin();
+			if (DestDataPtr.IsValid())
 			{
-				if (FAudioDeviceManager* AudioDeviceManager = FAudioDeviceManager::Get())
+				TAudioModulationPtr ModPtr = ModInterfacePtr.Pin();
+				if (ModPtr.IsValid())
 				{
-					if (FAudioDevice* AudioDevice = AudioDeviceManager->GetAudioDeviceRaw(DestinationDataPtr->GetDeviceId()))
+					TSet<FModulatorHandle> NewHandles;
+					for (TUniquePtr<Audio::IModulatorSettings>& ModSetting : ModSettings)
 					{
-						if (AudioDevice->IsModulationPluginEnabled() && AudioDevice->ModulationInterface.IsValid())
-						{
-							if (IAudioModulationManager* Modulation = AudioDevice->ModulationInterface.Get())
-							{
-								TSet<FModulatorHandle> NewHandles;
-								for (TUniquePtr<Audio::IModulatorSettings>& ModSetting : ModSettings)
-								{
-									Audio::FModulationParameter HandleParam = DestinationDataPtr->GetParameter();
-									NewHandles.Add(FModulatorHandle{ *Modulation, *ModSetting.Get(), MoveTemp(HandleParam) });
-								}
-								DestinationDataPtr->SetHandles(MoveTemp(NewHandles));
-							}
-							return;
-						}
+						Audio::FModulationParameter HandleParam = DestDataPtr->GetParameter();
+						NewHandles.Add(FModulatorHandle { *ModPtr.Get(), *ModSetting.Get(), MoveTemp(HandleParam) });
 					}
+					DestDataPtr->SetHandles(MoveTemp(NewHandles));
+					return;
 				}
-				DestinationDataPtr->ResetHandles();
+				DestDataPtr->ResetHandles();
 			}
-		};
-
-		FAudioThread::RunCommandOnAudioThread(MoveTemp(UpdateHandleLambda));
+		});
 	}
 } // namespace Audio
 
