Cluster union : add native multicast events for faster dispatch
Also provide a TSet version of SetPerParticleCollisionProfileName to avoid converting TSet to arrays

diff --git a/Engine/Source/Runtime/Engine/Classes/PhysicsEngine/ClusterUnionComponent.h b/Engine/Source/Runtime/Engine/Classes/PhysicsEngine/ClusterUnionComponent.h
--- a/Engine/Source/Runtime/Engine/Classes/PhysicsEngine/ClusterUnionComponent.h
+++ b/Engine/Source/Runtime/Engine/Classes/PhysicsEngine/ClusterUnionComponent.h
@@ -93,7 +93,11 @@ struct FClusterUnionPendingAddData
 };
 
 DECLARE_DYNAMIC_MULTICAST_DELEGATE_ThreeParams(FOnClusterUnionAddedComponent, UPrimitiveComponent*, Component, const TSet<int32>&, BoneIds, bool, bIsNew);
+DECLARE_MULTICAST_DELEGATE_ThreeParams(FOnClusterUnionAddedComponentNative, UPrimitiveComponent*, const TSet<int32>& /*BoneIds*/, bool /*bIsNew*/);
+
 DECLARE_DYNAMIC_MULTICAST_DELEGATE_OneParam(FOnClusterUnionRemovedComponent, UPrimitiveComponent*, Component);
+DECLARE_MULTICAST_DELEGATE_OneParam(FOnClusterUnionRemovedComponentNative, UPrimitiveComponent*);
+
 DECLARE_DYNAMIC_MULTICAST_DELEGATE_TwoParams(FOnClusterUnionBoundsChanged, UClusterUnionComponent*, Component, const FBoxSphereBounds&, Bounds);
 
 /**
@@ -166,6 +170,10 @@ public:
 	UPROPERTY(BlueprintAssignable, Category = "Events")
 	FOnClusterUnionBoundsChanged OnComponentBoundsChangedEvent;
 
+	// native (fast, low overhead) versions 
+	FOnClusterUnionAddedComponentNative OnComponentAddedNativeEvent;
+	FOnClusterUnionRemovedComponentNative OnComponentRemovedNativeEvent;
+
 	// Lambda returns whether or not iteration should continue;
 	ENGINE_API void VisitAllCurrentChildComponents(const TFunction<bool(UPrimitiveComponent*)>& Lambda) const;
 	ENGINE_API void VisitAllCurrentActors(const TFunction<bool(AActor*)>& Lambda) const;
@@ -188,6 +196,9 @@ protected:
 
 	ENGINE_API TArray<int32> GetAddedBoneIdsForComponent(UPrimitiveComponent* Component) const;
 
+	ENGINE_API void BroadcastComponentAddedEvents(UPrimitiveComponent* ChangedComponent, const TSet<int32>& BoneIds, bool bIsNew);
+	ENGINE_API void BroadcastComponentRemovedEvents(UPrimitiveComponent* ChangedComponent);
+
 	Chaos::FClusterUnionPhysicsProxy* GetPhysicsProxy() const { return PhysicsProxy; }
 	Chaos::FClusterUnionPhysicsProxy* GetPhysicsProxy() { return PhysicsProxy; }
 
diff --git a/Engine/Source/Runtime/Engine/Private/PhysicsEngine/ClusterUnionComponent.cpp b/Engine/Source/Runtime/Engine/Private/PhysicsEngine/ClusterUnionComponent.cpp
--- a/Engine/Source/Runtime/Engine/Private/PhysicsEngine/ClusterUnionComponent.cpp
+++ b/Engine/Source/Runtime/Engine/Private/PhysicsEngine/ClusterUnionComponent.cpp
@@ -944,7 +944,7 @@ void UClusterUnionComponent::HandleAddOrModifiedClusteredComponent(UPrimitiveCom
 
 	const TSet<int32> OldBoneIds{ ComponentData.BoneIds };
 	PerBoneChildToParent.GetKeys(ComponentData.BoneIds);
-	OnComponentAddedEvent.Broadcast(ChangedComponent, ComponentData.BoneIds, bIsNew);
+	BroadcastComponentAddedEvents(ChangedComponent, ComponentData.BoneIds, bIsNew);
 
 	if (IsAuthority() && ComponentData.ReplicatedProxyComponent.IsValid())
 	{
@@ -1047,7 +1047,7 @@ void UClusterUnionComponent::HandleRemovedClusteredComponent(TObjectKey<UPrimiti
 
 		if (UPrimitiveComponent* ChangedComponent = RemovedComponent.ResolveObjectPtr())
 		{
-			OnComponentRemovedEvent.Broadcast(ChangedComponent);
+			BroadcastComponentRemovedEvents(ChangedComponent);
 		}
 		PerComponentData.Remove(RemovedComponent);
 	}
@@ -1055,6 +1055,30 @@ void UClusterUnionComponent::HandleRemovedClusteredComponent(TObjectKey<UPrimiti
 	PendingComponentSync.Remove(RemovedComponent);
 }
 
+void UClusterUnionComponent::BroadcastComponentAddedEvents(UPrimitiveComponent* ChangedComponent, const TSet<int32>& BoneIds, bool bIsNew)
+{
+	if (OnComponentAddedEvent.IsBound())
+	{
+		OnComponentAddedEvent.Broadcast(ChangedComponent, BoneIds, bIsNew);
+	}
+	if (OnComponentAddedNativeEvent.IsBound())
+	{
+		OnComponentAddedNativeEvent.Broadcast(ChangedComponent, BoneIds, bIsNew);
+	}
+}
+
+void UClusterUnionComponent::BroadcastComponentRemovedEvents(UPrimitiveComponent* ChangedComponent)
+{
+	if (OnComponentRemovedEvent.IsBound())
+	{
+		OnComponentRemovedEvent.Broadcast(ChangedComponent);
+	}
+	if (OnComponentRemovedNativeEvent.IsBound())
+	{
+		OnComponentRemovedNativeEvent.Broadcast(ChangedComponent);
+	}
+}
+
 void UClusterUnionComponent::OnRep_RigidState()
 {
 	if (!PhysicsProxy)
diff --git a/Engine/Source/Runtime/Experimental/GeometryCollectionEngine/Private/GeometryCollection/GeometryCollectionComponent.cpp b/Engine/Source/Runtime/Experimental/GeometryCollectionEngine/Private/GeometryCollection/GeometryCollectionComponent.cpp
--- a/Engine/Source/Runtime/Experimental/GeometryCollectionEngine/Private/GeometryCollection/GeometryCollectionComponent.cpp
+++ b/Engine/Source/Runtime/Experimental/GeometryCollectionEngine/Private/GeometryCollection/GeometryCollectionComponent.cpp
@@ -3396,11 +3396,13 @@ void UGeometryCollectionComponent::SetPerLevelCollisionProfileNames(const TArray
 	LoadCollisionProfiles();
 }
 
-void UGeometryCollectionComponent::SetPerParticleCollisionProfileName(const TArray<int32>& BoneIds, FName ProfileName)
+// return true if anything was cahnged
+template <typename TIteratableBoneContainerType>
+static bool SetPerParticleCollisionProfileNameFromIterable(const UGeometryCollection* RestCollection, const TIteratableBoneContainerType& BoneIds, FName ProfileName, TArray<FName>& CollisionProfilePerParticle)
 {
 	if (!RestCollection)
 	{
-		return;
+		return false;
 	}
 
 	bool bHasChanged = false;
@@ -3424,6 +3426,21 @@ void UGeometryCollectionComponent::SetPerParticleCollisionProfileName(const TArr
 		}
 	}
 
+	return bHasChanged;
+}
+
+void UGeometryCollectionComponent::SetPerParticleCollisionProfileName(const TSet<int32>& BoneIds, FName ProfileName)
+{
+	const bool bHasChanged = SetPerParticleCollisionProfileNameFromIterable(RestCollection, BoneIds, ProfileName, CollisionProfilePerParticle);
+	if (bHasChanged)
+	{
+		LoadCollisionProfiles();
+	}
+}
+
+void UGeometryCollectionComponent::SetPerParticleCollisionProfileName(const TArray<int32>& BoneIds, FName ProfileName)
+{
+	const bool bHasChanged = SetPerParticleCollisionProfileNameFromIterable(RestCollection, BoneIds, ProfileName, CollisionProfilePerParticle);
 	if (bHasChanged)
 	{
 		LoadCollisionProfiles();
diff --git a/Engine/Source/Runtime/Experimental/GeometryCollectionEngine/Public/GeometryCollection/GeometryCollectionComponent.h b/Engine/Source/Runtime/Experimental/GeometryCollectionEngine/Public/GeometryCollection/GeometryCollectionComponent.h
--- a/Engine/Source/Runtime/Experimental/GeometryCollectionEngine/Public/GeometryCollection/GeometryCollectionComponent.h
+++ b/Engine/Source/Runtime/Experimental/GeometryCollectionEngine/Public/GeometryCollection/GeometryCollectionComponent.h
@@ -716,6 +716,8 @@ public:
 	UFUNCTION(BlueprintCallable, Category = "ChaosPhysics")
 	GEOMETRYCOLLECTIONENGINE_API void SetPerParticleCollisionProfileName(const TArray<int32>& BoneIds, FName ProfileName);
 
+	GEOMETRYCOLLECTIONENGINE_API void SetPerParticleCollisionProfileName(const TSet<int32>& BoneIds, FName ProfileName);
+
 	/** API for getting at geometry collection data */
 	GEOMETRYCOLLECTIONENGINE_API int32 GetNumElements(FName Group) const;
 
