World Partition Class Descriptors: don't rely on the Asset Registry to filter out classes as it's going to gather all assets for the specified classes first, then filtering for the provided packages names afterward, resulting in execution time being over 100 times slower than filtering for classes after package name.
Work is being done on the Asset Registry side to sort query filters by their minimum result size, to optimize queries like this.

diff --git a/Engine/Source/Runtime/Engine/Private/WorldPartition/WorldPartitionClassDescRegistry.cpp b/Engine/Source/Runtime/Engine/Private/WorldPartition/WorldPartitionClassDescRegistry.cpp
--- a/Engine/Source/Runtime/Engine/Private/WorldPartition/WorldPartitionClassDescRegistry.cpp
+++ b/Engine/Source/Runtime/Engine/Private/WorldPartition/WorldPartitionClassDescRegistry.cpp
@@ -206,17 +206,30 @@ void FWorldPartitionClassDescRegistry::PrefetchClassDescs(const TArray<FTopLevel
 
 	auto GetBlueprintAssets = [&AssetRegistry](const TArray<FString>& FilePaths, TArray<FAssetData>& Assets)
 	{
+		int32 AssetIndex = Assets.Num();
+
 		FARFilter Filter;
-		Filter.ClassPaths.Add(UBlueprint::StaticClass()->GetClassPathName());
-		Filter.ClassPaths.Add(UBlueprintGeneratedClass::StaticClass()->GetClassPathName());
-		Filter.ClassPaths.Add(UObjectRedirector::StaticClass()->GetClassPathName());
-		Filter.bRecursiveClasses = true;
+		// Don't rely on the AR to filter out classes as it's going to gather all assets for the specified classes first, then filtering for the provided packages names afterward,
+		// resulting in execution time being over 100 times slower than filtering for classes after package names.
+		//Filter.ClassPaths.Add(UBlueprint::StaticClass()->GetClassPathName());
+		//Filter.ClassPaths.Add(UBlueprintGeneratedClass::StaticClass()->GetClassPathName());
+		//Filter.ClassPaths.Add(UObjectRedirector::StaticClass()->GetClassPathName());
+		//Filter.bRecursiveClasses = true;
 		Filter.bIncludeOnlyOnDiskAssets = true;
 		Filter.PackageNames.Reserve(FilePaths.Num());
 		Algo::Transform(FilePaths, Filter.PackageNames, [](const FString& ClassPath) { return *ClassPath; });
-	
+
 		AssetRegistry.ScanSynchronous(TArray<FString>(), FilePaths);
 		AssetRegistry.GetAssets(Filter, Assets);
+
+		// Filter out unwanted classes, see above comment for details.
+		for (; AssetIndex < Assets.Num(); AssetIndex++)
+		{
+			if (UClass* AssetClass = Assets[AssetIndex].GetClass(); !AssetClass || (!AssetClass->IsChildOf<UBlueprint>() && !AssetClass->IsChildOf<UBlueprintGeneratedClass>() && !AssetClass->IsChildOf<UObjectRedirector>()))
+			{
+				Assets.RemoveAtSwap(AssetIndex--);
+			}
+		}
 	};
 
 	TArray<FAssetData> Assets;
