Shader Compiler:  Fix for excess job queries generated when per-shader DDC is enabled in cooks.  Per-shader DDC was disabling material map DDC for ODSC, but that doesn't make sense for cooks, so we check if we are cooking when deciding whether to enable material map DDC.  Also added statistics to track duplicate jobs and per-shader DDC cache queries.

#rnx

diff --git a/Engine/Source/Runtime/Engine/Private/Materials/Material.cpp b/Engine/Source/Runtime/Engine/Private/Materials/Material.cpp
--- a/Engine/Source/Runtime/Engine/Private/Materials/Material.cpp
+++ b/Engine/Source/Runtime/Engine/Private/Materials/Material.cpp
@@ -4071,7 +4071,7 @@ void UMaterial::PostLoad()
 				}
 			}
 
-			const bool bSkipCompilationOnPostLoad = IsShaderJobCacheDDCEnabled();
+			const bool bSkipCompilationOnPostLoad = IsMaterialMapDDCEnabled() == false;
 			if (bSkipCompilationOnPostLoad)
 			{
 				CacheResourceShadersForRendering(false, EMaterialShaderPrecompileMode::None);
diff --git a/Engine/Source/Runtime/Engine/Private/Materials/MaterialInstance.cpp b/Engine/Source/Runtime/Engine/Private/Materials/MaterialInstance.cpp
--- a/Engine/Source/Runtime/Engine/Private/Materials/MaterialInstance.cpp
+++ b/Engine/Source/Runtime/Engine/Private/Materials/MaterialInstance.cpp
@@ -3126,7 +3126,7 @@ void UMaterialInstance::PostLoad()
 	{
 		SCOPE_SECONDS_COUNTER(MaterialLoadTime);
 
-		const bool bSkipCompilationOnPostLoad = IsShaderJobCacheDDCEnabled();
+		const bool bSkipCompilationOnPostLoad = IsMaterialMapDDCEnabled() == false;
 
 		// Make sure static parameters are up to date and shaders are cached for the current platform
 		if (bSkipCompilationOnPostLoad)
diff --git a/Engine/Source/Runtime/Engine/Private/Materials/MaterialShared.cpp b/Engine/Source/Runtime/Engine/Private/Materials/MaterialShared.cpp
--- a/Engine/Source/Runtime/Engine/Private/Materials/MaterialShared.cpp
+++ b/Engine/Source/Runtime/Engine/Private/Materials/MaterialShared.cpp
@@ -2697,7 +2697,7 @@ bool FMaterial::CacheShaders(const FMaterialShaderMapId& ShaderMapId, EShaderPla
 			}
 
 			// If we are loading individual shaders from the shader job cache don't attempt to load full maps.
-			const bool bSkipCompilationOnPostLoad = IsShaderJobCacheDDCEnabled();
+			const bool bSkipCompilationOnPostLoad = IsMaterialMapDDCEnabled() == false;
 
 			// Attempt to load from the derived data cache if we are uncooked and don't have any shadermap.
 			// If we have an incomplete shadermap, continue with it to prevent creation of duplicate shadermaps for the same ShaderMapId
@@ -2815,8 +2815,8 @@ bool FMaterial::CacheShaders(const FMaterialShaderMapId& ShaderMapId, EShaderPla
 		}
 		else
 		{
-			const bool bSkipCompilationForODSC = !RequiresSynchronousCompilation() && (GShaderCompilingManager->IsShaderCompilationSkipped() || IsShaderJobCacheDDCEnabled());
-			// If we aren't actually compiling shaders don't print the debug message that we are compiling shaders.
+			const bool bSkipCompilationForODSC = !RequiresSynchronousCompilation() && (GShaderCompilingManager->IsShaderCompilationSkipped() || (IsMaterialMapDDCEnabled() == false));
+			// If we aren't actually generating shadermaps, don't print the debug message that we are generating shadermaps.
 			if (!bSkipCompilationForODSC)
 			{
 				const TCHAR* ShaderMapCondition;
diff --git a/Engine/Source/Runtime/Engine/Private/ShaderCompiler/ShaderCompiler.cpp b/Engine/Source/Runtime/Engine/Private/ShaderCompiler/ShaderCompiler.cpp
--- a/Engine/Source/Runtime/Engine/Private/ShaderCompiler/ShaderCompiler.cpp
+++ b/Engine/Source/Runtime/Engine/Private/ShaderCompiler/ShaderCompiler.cpp
@@ -216,6 +216,12 @@ bool IsShaderJobCacheDDCEnabled()
 	return false;
 }
 
+bool IsMaterialMapDDCEnabled()
+{
+	// If we are loading individual shaders from the shader job cache for ODSC, don't attempt to load full material maps.  Always load/cache material maps in cooks.
+	return (IsShaderJobCacheDDCEnabled() == false) || IsRunningCookCommandlet();
+}
+
 int32 GShaderCompilerAllowDistributedCompilation = 1;
 static FAutoConsoleVariableRef CVarShaderCompilerAllowDistributedCompilation(
 	TEXT("r.ShaderCompiler.AllowDistributedCompilation"),
@@ -909,6 +915,12 @@ private:
 	/** Statistics - total number of times we succeded in Find()ing output for some input hash */
 	uint64 TotalCacheHits = 0;
 
+	/** Statistics - total number of times a duplicate job was added (duplicate jobs are processed when the original finishes compiling) */
+	uint64 TotalCacheDuplicates = 0;
+
+	/** Statistics - total number of times a per-shader DDC query was issued */
+	uint64 TotalCacheDDCQueries = 0;
+
 	/** Statistics - total number of times a per-shader DDC query succeeded for some input hash */
 	uint64 TotalCacheDDCHits = 0;
 
@@ -1609,6 +1621,8 @@ void FShaderJobCache::SubmitJob(FShaderCommonCompileJob* Job)
 				{
 					*WaitListHead = Job;
 				}
+				++TotalCacheDuplicates;
+
 				AddDuplicateJob(Job);
 				JobLock.WriteUnlock();
 				bNewJob = false;
@@ -4426,12 +4440,14 @@ void FShaderCompilerStats::WriteStatSummary()
 	if (Counters.TotalCacheSearchAttempts > 0)
 	{
 		UE_LOG(LogShaderCompilers, Display, TEXT("=== FShaderJobCache stats%s ==="), AggregatedSuffix);
-		UE_LOG(LogShaderCompilers, Display, TEXT("Total job queries %s, among them cache hits %s (%.2f%%), DDC hits %s (%.2f%%)"),
+		UE_LOG(LogShaderCompilers, Display, TEXT("Total job queries %s, among them cache hits %s (%.2f%%), DDC hits %s (%.2f%%), Duplicates %s (%.2f%%)"),
 			*FormatNumber(Counters.TotalCacheSearchAttempts),
 			*FormatNumber(Counters.TotalCacheHits),
 			100.0 * static_cast<double>(Counters.TotalCacheHits) / static_cast<double>(Counters.TotalCacheSearchAttempts),
 			*FormatNumber(Counters.TotalCacheDDCHits),
-			100.0 * static_cast<double>(Counters.TotalCacheDDCHits) / static_cast<double>(Counters.TotalCacheSearchAttempts));
+			100.0 * static_cast<double>(Counters.TotalCacheDDCHits) / static_cast<double>(Counters.TotalCacheSearchAttempts),
+			*FormatNumber(Counters.TotalCacheDuplicates),
+			100.0 * static_cast<double>(Counters.TotalCacheDuplicates) / static_cast<double>(Counters.TotalCacheSearchAttempts));
 
 		UE_LOG(LogShaderCompilers, Display, TEXT("Tracking %s distinct input hashes that result in %s distinct outputs (%.2f%%)"),
 			*FormatNumber(Counters.UniqueCacheInputHashes),
@@ -10183,6 +10199,8 @@ FShaderJobCacheRef FShaderJobCache::FindOrAdd(const FJobInputHash& Hash, EShader
 		{
 			TRACE_COUNTER_INCREMENT(Shaders_JobCacheDDCRequests);
 
+			++TotalCacheDDCQueries;
+
 			UE::DerivedData::EPriority DerivedDataPriority;
 			UE::DerivedData::FRequestOwner* RequestOwner;
 
@@ -10620,6 +10638,8 @@ void FShaderJobCache::GetStats(FShaderCompilerStats& OutStats) const
 	FReadScopeLock Locker(JobLock);
 	OutStats.Counters.TotalCacheSearchAttempts = TotalSearchAttempts;
 	OutStats.Counters.TotalCacheHits = TotalCacheHits;
+	OutStats.Counters.TotalCacheDuplicates = TotalCacheDuplicates;
+	OutStats.Counters.TotalCacheDDCQueries = TotalCacheDDCQueries;
 	OutStats.Counters.TotalCacheDDCHits = TotalCacheDDCHits;
 	OutStats.Counters.UniqueCacheInputHashes = InputHashToJobData.Num();
 	OutStats.Counters.UniqueCacheOutputs = Outputs.Num();
diff --git a/Engine/Source/Runtime/Engine/Public/ShaderCompiler.h b/Engine/Source/Runtime/Engine/Public/ShaderCompiler.h
--- a/Engine/Source/Runtime/Engine/Public/ShaderCompiler.h
+++ b/Engine/Source/Runtime/Engine/Public/ShaderCompiler.h
@@ -46,6 +46,7 @@ struct FAnalyticsEventAttribute;
 #define DEBUG_INFINITESHADERCOMPILE 0
 
 extern ENGINE_API bool IsShaderJobCacheDDCEnabled();
+extern ENGINE_API bool IsMaterialMapDDCEnabled();
 
 struct FShaderJobCacheStoredOutput;
 class FShaderJobCache;
@@ -480,6 +481,12 @@ private:
 		/** Total number of hits in the job cache (i.e. input hashes seen >1 time) */
 		uint64 TotalCacheHits = 0;
 
+		/** Total number of duplicate jobs (input hash matches an in-flight job, processed when in-flight job completes) */
+		uint32 TotalCacheDuplicates = 0;
+
+		/** Total number of DDC queries in the job cache (per-shader DDC). */
+		uint32 TotalCacheDDCQueries = 0;
+
 		/** Total number of DDC hits in the job cache (per shader DDC, as opposed to shader map DDC stats above). */
 		uint32 TotalCacheDDCHits = 0;
 
@@ -526,6 +533,8 @@ private:
 			ShaderMapDDCHits += Other.ShaderMapDDCHits;
 			TotalCacheSearchAttempts += Other.TotalCacheSearchAttempts;
 			TotalCacheHits += Other.TotalCacheHits;
+			TotalCacheDuplicates += Other.TotalCacheDuplicates;
+			TotalCacheDDCQueries += Other.TotalCacheDDCQueries;
 			TotalCacheDDCHits += Other.TotalCacheDDCHits;
 			UniqueCacheInputHashes += Other.UniqueCacheInputHashes;
 			UniqueCacheOutputs += Other.UniqueCacheOutputs;
