Replaced checks with ensures on ExternalStreamingObjects received in UWorldPartitionRuntimeHash::InjectExternalStreamingObject/RemoveExternalStreamingObject and early out if invalid.

diff --git a/Engine/Source/Runtime/Engine/Private/WorldPartition/WorldPartitionRuntimeHash.cpp b/Engine/Source/Runtime/Engine/Private/WorldPartition/WorldPartitionRuntimeHash.cpp
--- a/Engine/Source/Runtime/Engine/Private/WorldPartition/WorldPartitionRuntimeHash.cpp
+++ b/Engine/Source/Runtime/Engine/Private/WorldPartition/WorldPartitionRuntimeHash.cpp
@@ -476,7 +476,11 @@ EWorldPartitionStreamingPerformance UWorldPartitionRuntimeHash::GetStreamingPerf
 
 bool UWorldPartitionRuntimeHash::InjectExternalStreamingObject(URuntimeHashExternalStreamingObjectBase* InExternalStreamingObject)
 {
-	check(InExternalStreamingObject);
+	if (!ensure(InExternalStreamingObject))
+	{
+		return false;
+	}
+	check(IsValid(InExternalStreamingObject));
 	bool bAlreadyInSet = false;
 	InjectedExternalStreamingObjects.Add(InExternalStreamingObject, &bAlreadyInSet);
 	if (bAlreadyInSet)
@@ -489,7 +493,11 @@ bool UWorldPartitionRuntimeHash::InjectExternalStreamingObject(URuntimeHashExter
 
 bool UWorldPartitionRuntimeHash::RemoveExternalStreamingObject(URuntimeHashExternalStreamingObjectBase* InExternalStreamingObject)
 {
-	check(InExternalStreamingObject);
+	if (!ensure(InExternalStreamingObject))
+	{
+		return false;
+	}
+	check(IsValid(InExternalStreamingObject));
 	if (!InjectedExternalStreamingObjects.Remove(InExternalStreamingObject))
 	{
 		UE_LOG(LogWorldPartition, Warning, TEXT("External streaming object %s was not injected."), *InExternalStreamingObject->GetName());
