World Partition Class Descriptors: filter out level script blueprints when saving objects, as they don't need a class descriptor.

diff --git a/Engine/Source/Runtime/Engine/Private/WorldPartition/WorldPartitionClassDescRegistry.cpp b/Engine/Source/Runtime/Engine/Private/WorldPartition/WorldPartitionClassDescRegistry.cpp
--- a/Engine/Source/Runtime/Engine/Private/WorldPartition/WorldPartitionClassDescRegistry.cpp
+++ b/Engine/Source/Runtime/Engine/Private/WorldPartition/WorldPartitionClassDescRegistry.cpp
@@ -498,15 +498,21 @@ void FWorldPartitionClassDescRegistry::RegisterClassDescriptorFromActorClass(con
 
 void FWorldPartitionClassDescRegistry::OnObjectPreSave(UObject* InObject, FObjectPreSaveContext InSaveContext)
 {
+	// Are we are saving a blueprint?
 	if (UBlueprint* Blueprint = Cast<UBlueprint>(InObject))
 	{
-		// We are saving a blueprint
-		if (UBlueprintGeneratedClass* BlueprintGeneratedClass = Cast<UBlueprintGeneratedClass>(Blueprint->GeneratedClass))
+		// Is this blueprint an asset (could be a level script)?
+		if (Cast<UPackage>(Blueprint->GetOuter()))
 		{
-			if (BlueprintGeneratedClass->IsChildOf<AActor>())
+			// Is this blueprint generated class valid?
+			if (UBlueprintGeneratedClass* BlueprintGeneratedClass = Cast<UBlueprintGeneratedClass>(Blueprint->GeneratedClass))
 			{
-				PrefetchClassDesc(BlueprintGeneratedClass);
-				UpdateClassDescriptor(Blueprint, false);
+				// Is it an actor derived blueprint?
+				if (BlueprintGeneratedClass->IsChildOf<AActor>())
+				{
+					PrefetchClassDesc(BlueprintGeneratedClass);
+					UpdateClassDescriptor(Blueprint, false);
+				}
 			}
 		}
 	}
