Add substrate topology visualization for material function.

diff --git a/Engine/Source/Editor/GraphEditor/Private/MaterialNodes/SGraphNodeMaterialBase.cpp b/Engine/Source/Editor/GraphEditor/Private/MaterialNodes/SGraphNodeMaterialBase.cpp
--- a/Engine/Source/Editor/GraphEditor/Private/MaterialNodes/SGraphNodeMaterialBase.cpp
+++ b/Engine/Source/Editor/GraphEditor/Private/MaterialNodes/SGraphNodeMaterialBase.cpp
@@ -483,32 +483,57 @@ void SGraphNodeMaterialBase::CreateBelowPinControls(TSharedPtr<SVerticalBox> Mai
 	}
 
 	// Preview of Substrate nodes topology
-	if (Substrate::IsSubstrateEnabled() && MaterialNode && MaterialNode->MaterialExpression->IsA(UMaterialExpressionSubstrateBSDF::StaticClass()))
+	if (Substrate::IsSubstrateEnabled() && MaterialNode)
 	{
-		if (const UMaterialExpressionSubstrateBSDF* SubstrateExpression = (const UMaterialExpressionSubstrateBSDF*)MaterialNode->MaterialExpression)
-		{		
-			if (UMaterial* MaterialForStats = SubstrateExpression->Material)
+		TArray<FGuid> Guids;
+		if (MaterialNode->MaterialExpression->IsA(UMaterialExpressionSubstrateBSDF::StaticClass()))
+		{
+			const UMaterialExpression* SubstrateExpression = (const UMaterialExpression*)MaterialNode->MaterialExpression;
+			Guids.Add(SubstrateExpression->MaterialExpressionGuid);
+		}
+		else if (MaterialNode->MaterialExpression->IsA(UMaterialExpressionMaterialFunctionCall::StaticClass()))
+		{
+			UMaterialExpressionMaterialFunctionCall* FunctionCall = (UMaterialExpressionMaterialFunctionCall*)MaterialNode->MaterialExpression;
+			const uint32 OutputCount = FunctionCall->FunctionOutputs.Num();
+			for (uint32 OutputIndex = 0; OutputIndex < OutputCount; ++OutputIndex)
 			{
-				if (const FMaterialResource* MaterialResource = MaterialForStats->GetMaterialResource(GMaxRHIFeatureLevel))
+				if (FunctionCall->IsResultSubstrateMaterial(OutputIndex))
 				{
-					if (FMaterialShaderMap* ShaderMap = MaterialResource->GetGameThreadShaderMap())
+					FSubstrateMaterialInfo SubstrateMaterialInfo(true/*bGatherGuids*/);
+					FunctionCall->GatherSubstrateMaterialInfo(SubstrateMaterialInfo, OutputIndex);
+					Guids = SubstrateMaterialInfo.GetGuids();
+					break;
+				}
+			}
+		}
+
+		if (Guids.Num() > 0)
+		{
+			if (const UMaterialExpression* SubstrateExpression = (const UMaterialExpression*)MaterialNode->MaterialExpression)
+			{		
+				if (UMaterial* MaterialForStats = SubstrateExpression->Material)
+				{
+					if (const FMaterialResource* MaterialResource = MaterialForStats->GetMaterialResource(GMaxRHIFeatureLevel))
 					{
-						const FSubstrateMaterialCompilationOutput& CompilationOutput = ShaderMap->GetSubstrateMaterialCompilationOutput();
-						MainBox->AddSlot()
-						.Padding(Settings->GetNonPinNodeBodyPadding())
-						.AutoHeight()
-						[
-							SNew(SHorizontalBox)
-							+SHorizontalBox::Slot()
-							.VAlign(VAlign_Center)
-							.HAlign(HAlign_Center)
-							[							
-								FSubstrateWidget::ProcessOperator(CompilationOutput, SubstrateExpression->MaterialExpressionGuid)
-							]
-						];
+						if (FMaterialShaderMap* ShaderMap = MaterialResource->GetGameThreadShaderMap())
+						{
+							const FSubstrateMaterialCompilationOutput& CompilationOutput = ShaderMap->GetSubstrateMaterialCompilationOutput();
+							MainBox->AddSlot()
+							.Padding(Settings->GetNonPinNodeBodyPadding())
+							.AutoHeight()
+							[
+								SNew(SHorizontalBox)
+								+SHorizontalBox::Slot()
+								.VAlign(VAlign_Center)
+								.HAlign(HAlign_Center)
+								[							
+									FSubstrateWidget::ProcessOperator(CompilationOutput, Guids)
+								]
+							];
+						}
 					}
-				}
 	
+				}
 			}
 		}
 	}
diff --git a/Engine/Source/Editor/GraphEditor/Private/MaterialNodes/SGraphSubstrateMaterial.cpp b/Engine/Source/Editor/GraphEditor/Private/MaterialNodes/SGraphSubstrateMaterial.cpp
--- a/Engine/Source/Editor/GraphEditor/Private/MaterialNodes/SGraphSubstrateMaterial.cpp
+++ b/Engine/Source/Editor/GraphEditor/Private/MaterialNodes/SGraphSubstrateMaterial.cpp
@@ -83,10 +83,10 @@ static const TSharedRef<SWidget> InternalProcessOperator(
 	const FSubstrateMaterialCompilationOutput& CompilationOutput, 
 	const FSubstrateOperator& Op, 
 	ESubstrateWidgetOutputType OutputType,
-	const FGuid& InGuid, 
+	const TArray<FGuid>& InGuid,
 	EStyleColor OverrideColor)
 {
-	const bool bIsCurrent = OutputType == ESubstrateWidgetOutputType::Node ? InGuid == Op.MaterialExpressionGuid : false;
+	const bool bIsCurrent = OutputType == ESubstrateWidgetOutputType::Node ? InGuid.Find(Op.MaterialExpressionGuid) != INDEX_NONE : false;
 	const EStyleColor Color0 = bIsCurrent ? GetSubstrateWidgetColor0() : OverrideColor;
 	const EStyleColor Color1 = bIsCurrent ? GetSubstrateWidgetColor1() : OverrideColor;
 	switch (Op.OperatorType)
@@ -212,10 +212,10 @@ static const TSharedRef<SWidget> InternalProcessOperator(
 
 const TSharedRef<SWidget> FSubstrateWidget::ProcessOperator(const FSubstrateMaterialCompilationOutput& CompilationOutput)
 {
-	return InternalProcessOperator(CompilationOutput, CompilationOutput.Operators[CompilationOutput.RootOperatorIndex], ESubstrateWidgetOutputType::DetailPanel, FGuid(), EStyleColor::MAX);
+	return InternalProcessOperator(CompilationOutput, CompilationOutput.Operators[CompilationOutput.RootOperatorIndex], ESubstrateWidgetOutputType::DetailPanel, TArray<FGuid>(), EStyleColor::MAX);
 }
 
-const TSharedRef<SWidget> FSubstrateWidget::ProcessOperator(const FSubstrateMaterialCompilationOutput& CompilationOutput, const FGuid& InGuid)
+const TSharedRef<SWidget> FSubstrateWidget::ProcessOperator(const FSubstrateMaterialCompilationOutput& CompilationOutput, const TArray<FGuid>& InGuid)
 {
 	return InternalProcessOperator(CompilationOutput, CompilationOutput.Operators[CompilationOutput.RootOperatorIndex], ESubstrateWidgetOutputType::Node, InGuid, EStyleColor::MAX);
 }
diff --git a/Engine/Source/Editor/GraphEditor/Public/SGraphSubstrateMaterial.h b/Engine/Source/Editor/GraphEditor/Public/SGraphSubstrateMaterial.h
--- a/Engine/Source/Editor/GraphEditor/Public/SGraphSubstrateMaterial.h
+++ b/Engine/Source/Editor/GraphEditor/Public/SGraphSubstrateMaterial.h
@@ -12,7 +12,7 @@ struct FSubstrateMaterialCompilationOutput;
 struct GRAPHEDITOR_API FSubstrateWidget
 {
 	static const TSharedRef<SWidget> ProcessOperator(const FSubstrateMaterialCompilationOutput& CompilationOutput);
-	static const TSharedRef<SWidget> ProcessOperator(const FSubstrateMaterialCompilationOutput& CompilationOutput, const FGuid& InGuid);
+	static const TSharedRef<SWidget> ProcessOperator(const FSubstrateMaterialCompilationOutput& CompilationOutput, const TArray<FGuid>& InGuid);
 	static void GetPinColor(TSharedPtr<SGraphPin>& Out, const UMaterialGraphNode* InNode);
 	static FLinearColor GetConnectionColor();
 	static bool HasInputSubstrateType(const UEdGraphPin* InPin);
diff --git a/Engine/Source/Runtime/Engine/Classes/Engine/EngineTypes.h b/Engine/Source/Runtime/Engine/Classes/Engine/EngineTypes.h
--- a/Engine/Source/Runtime/Engine/Classes/Engine/EngineTypes.h
+++ b/Engine/Source/Runtime/Engine/Classes/Engine/EngineTypes.h
@@ -717,7 +717,7 @@ static_assert(SSM_NUM <= 32, "Do not exceed 32 shading models without expanding
 struct FSubstrateMaterialInfo
 {
 public:
-	FSubstrateMaterialInfo() {}
+	FSubstrateMaterialInfo(bool bGatherGuids = false) { bGatherMaterialExpressionGuids = bGatherGuids; }
 	FSubstrateMaterialInfo(ESubstrateShadingModel InShadingModel) { AddShadingModel(InShadingModel); }
 
 	// Shading model
@@ -742,6 +742,10 @@ public:
 	void SetShadingModelFromExpression(bool bIn) { bHasShadingModelFromExpression = bIn ? 1u : 0u; }
 	bool HasShadingModelFromExpression() const { return bHasShadingModelFromExpression > 0u; }
 
+	// Substrate material expression GUIDs
+	void AddGuid(const FGuid& In) { if (bGatherMaterialExpressionGuids) { MaterialExpressionGuids.Add(In); } }
+	const TArray<FGuid>& GetGuids() const { return MaterialExpressionGuids; }
+
 	uint64 GetPropertyConnected() const { return ConnectedPropertyMask; }
 	void AddPropertyConnected(uint64 In) { ConnectedPropertyMask |= (1ull << In); }
 	bool HasPropertyConnected(uint64 In) const { return !!(ConnectedPropertyMask & (1ull << In)); }
@@ -782,6 +786,10 @@ private:
 	TArray<TObjectPtr<USubsurfaceProfile>> SubsurfaceProfiles;
 
 	TArray<TObjectPtr<USpecularProfile>> SpecularProfiles;
+	
+	/* Material expression GUIDs for all the traversed Substrate nodes */
+	bool bGatherMaterialExpressionGuids = false;
+	TArray<FGuid> MaterialExpressionGuids;
 
 #if WITH_EDITOR
 	// A simple way to detect and prevent node re-entry due to cycling graph; stop the compilation and avoid crashing.
diff --git a/Engine/Source/Runtime/Engine/Private/Materials/Material.cpp b/Engine/Source/Runtime/Engine/Private/Materials/Material.cpp
--- a/Engine/Source/Runtime/Engine/Private/Materials/Material.cpp
+++ b/Engine/Source/Runtime/Engine/Private/Materials/Material.cpp
@@ -3265,9 +3265,9 @@ void UMaterial::ConvertMaterialToSubstrateMaterial()
 	if (bUseMaterialAttributes && EditorOnly->MaterialAttributes.Expression && !EditorOnly->MaterialAttributes.Expression->IsResultSubstrateMaterial(EditorOnly->MaterialAttributes.OutputIndex)) // M_Rifle cause issues there
 	{
 		UMaterialExpressionSubstrateConvertMaterialAttributes* ConvertAttributeNode = NewObject<UMaterialExpressionSubstrateConvertMaterialAttributes>(this);
+		ConvertAttributeNode->Material = this;
 		SetPosXAndMoveReferenceToTheRight(ConvertAttributeNode);
 		ConvertAttributeNode->SubsurfaceProfile = bRequireNoSubsurfaceProfile ? nullptr : SubsurfaceProfile;
-		ConvertAttributeNode->Material = this;
 
 		MoveConnectionTo(EditorOnly->MaterialAttributes, ConvertAttributeNode, 0);
 
@@ -3332,6 +3332,7 @@ void UMaterial::ConvertMaterialToSubstrateMaterial()
 
 			// Now pass through the convert to decal node, which flag the material as SSM_Decal, which will set the domain to Decal.
 			UMaterialExpressionSubstrateConvertToDecal* ConvertToDecalNode = NewObject<UMaterialExpressionSubstrateConvertToDecal>(this);
+			ConvertAttributeNode->Material = this;
 			ReplaceNodeAndMoveToTheRight(ConvertAttributeNode, ConvertToDecalNode);
 			ConvertToDecalNode->DecalMaterial.Connect(0, ConvertAttributeNode);
 
@@ -3370,6 +3371,7 @@ void UMaterial::ConvertMaterialToSubstrateMaterial()
 		{
 			// Or if it cannot be found, a slab node
 			UMaterialExpressionSubstrateSlabBSDF* SlabNode = NewObject<UMaterialExpressionSubstrateSlabBSDF>(this);
+			SlabNode->Material = this;
 			SetPosXAndMoveReferenceToTheRight(SlabNode);
 			EditorOnly->FrontMaterial.Connect(0, SlabNode);
 		}
@@ -3397,6 +3399,7 @@ void UMaterial::ConvertMaterialToSubstrateMaterial()
 					
 					// Top slab BSDF as a simple Disney material
 					UMaterialExpressionSubstrateSlabBSDF* BottomSlabBSDF = NewObject<UMaterialExpressionSubstrateSlabBSDF>(this);
+					BottomSlabBSDF->Material = this;
 					SetPosXAndMoveReferenceToTheRight(BottomSlabBSDF);
 					BottomSlabBSDF->GetInput(0)->Connect(0, SubstrateMetalnessToDiffuseAlbedoF0);
 					BottomSlabBSDF->GetInput(1)->Connect(1, SubstrateMetalnessToDiffuseAlbedoF0);
@@ -3410,6 +3413,7 @@ void UMaterial::ConvertMaterialToSubstrateMaterial()
 
 					// Now weight the top base material by opacity.
 					UMaterialExpressionSubstrateSlabBSDF* TopSlabBSDF = NewObject<UMaterialExpressionSubstrateSlabBSDF>(this);
+					TopSlabBSDF->Material = this;
 					TopSlabBSDF->MaterialExpressionEditorX = BottomSlabBSDF->MaterialExpressionEditorX;
 					TopSlabBSDF->MaterialExpressionEditorY = BottomSlabBSDF->MaterialExpressionEditorY + 650;
 					MoveConnectionTo(EditorOnly->EmissiveColor, TopSlabBSDF, 10);				// Emissive
@@ -3460,6 +3464,7 @@ void UMaterial::ConvertMaterialToSubstrateMaterial()
 			if (!bClearCoatConversionDone)
 			{
 				ConvertNode = NewObject<UMaterialExpressionSubstrateShadingModels>(this);
+				ConvertNode->Material = this;
 				SetPosXAndMoveReferenceToTheRight(ConvertNode);
 				ConvertNode->SubsurfaceProfile = bRequireNoSubsurfaceProfile ? nullptr : SubsurfaceProfile;
 				MoveConnectionTo(EditorOnly->BaseColor, ConvertNode, 0);
@@ -3515,6 +3520,7 @@ void UMaterial::ConvertMaterialToSubstrateMaterial()
 		else if (MaterialDomain == MD_Volume)
 		{
 			UMaterialExpressionSubstrateVolumetricFogCloudBSDF* VolBSDF = NewObject<UMaterialExpressionSubstrateVolumetricFogCloudBSDF>(this);
+			VolBSDF->Material = this;
 			SetPosXAndMoveReferenceToTheRight(VolBSDF);
 			MoveConnectionTo(EditorOnly->BaseColor, VolBSDF, 0);		// Albedo
 			MoveConnectionTo(EditorOnly->SubsurfaceColor, VolBSDF, 1);	// Extinction
@@ -3534,6 +3540,7 @@ void UMaterial::ConvertMaterialToSubstrateMaterial()
 
 			// Only Emissive & Opacity are valid input for PostProcess material
 			UMaterialExpressionSubstrateLightFunction* LightFunctionNode = NewObject<UMaterialExpressionSubstrateLightFunction>(this);
+			LightFunctionNode->Material = this;
 			SetPosXAndMoveReferenceToTheRight(LightFunctionNode);
 			MoveConnectionTo(EditorOnly->EmissiveColor, LightFunctionNode, 0);
 
@@ -3553,6 +3560,7 @@ void UMaterial::ConvertMaterialToSubstrateMaterial()
 			}
 
 			UMaterialExpressionSubstratePostProcess* PostProcNode = NewObject<UMaterialExpressionSubstratePostProcess>(this);
+			PostProcNode->Material = this;
 			SetPosXAndMoveReferenceToTheRight(PostProcNode);
 
 			MoveConnectionTo(EditorOnly->EmissiveColor, PostProcNode, 0);
@@ -3569,6 +3577,7 @@ void UMaterial::ConvertMaterialToSubstrateMaterial()
 			ShadingModels.AddShadingModel(MSM_DefaultLit);
 
 			ConvertNode = NewObject<UMaterialExpressionSubstrateShadingModels>(this);
+			ConvertNode->Material = this;
 			SetPosXAndMoveReferenceToTheRight(ConvertNode);
 			MoveConnectionTo(EditorOnly->BaseColor, ConvertNode, 0);
 			MoveConnectionTo(EditorOnly->Metallic, ConvertNode, 1);
@@ -3589,6 +3598,7 @@ void UMaterial::ConvertMaterialToSubstrateMaterial()
 
 			// Now pass through the convert to decal node, which flag the material as SSM_Decal, which will set the domain to Decal.
 			UMaterialExpressionSubstrateConvertToDecal* ConvertToDecalNode= NewObject<UMaterialExpressionSubstrateConvertToDecal>(this);
+			ConvertToDecalNode->Material = this;
 			ReplaceNodeAndMoveToTheRight(ConvertNode, ConvertToDecalNode);
 			ConvertToDecalNode->DecalMaterial.Connect(0, ConvertNode);
 
@@ -3603,6 +3613,7 @@ void UMaterial::ConvertMaterialToSubstrateMaterial()
 			ShadingModels.AddShadingModel(MSM_Unlit);
 
 			UMaterialExpressionSubstrateUI* UINode = NewObject<UMaterialExpressionSubstrateUI>(this);
+			UINode->Material = this;
 			SetPosXAndMoveReferenceToTheRight(UINode);
 			MoveConnectionTo(EditorOnly->EmissiveColor, UINode, 0);
 			CopyConnectionTo(EditorOnly->Opacity, UINode, 1);	// We only copy, to keep Opacity on the root node in case BLEND_AlphaComposite is selected.
diff --git a/Engine/Source/Runtime/Engine/Private/Materials/MaterialExpressions.cpp b/Engine/Source/Runtime/Engine/Private/Materials/MaterialExpressions.cpp
--- a/Engine/Source/Runtime/Engine/Private/Materials/MaterialExpressions.cpp
+++ b/Engine/Source/Runtime/Engine/Private/Materials/MaterialExpressions.cpp
@@ -24097,6 +24097,8 @@ void UMaterialExpressionSubstrateSlabBSDF::GatherSubstrateMaterialInfo(FSubstrat
 	{
 		SubstrateMaterialInfo.AddSpecularProfile(SpecularProfile);
 	}
+
+	SubstrateMaterialInfo.AddGuid(MaterialExpressionGuid);
 }
 
 FSubstrateOperator* UMaterialExpressionSubstrateSlabBSDF::SubstrateGenerateMaterialTopologyTree(class FMaterialCompiler* Compiler, class UMaterialExpression* Parent, int32 OutputIndex)
@@ -24329,6 +24331,7 @@ void UMaterialExpressionSubstrateSimpleClearCoatBSDF::GatherSubstrateMaterialInf
 	if (EmissiveColor.IsConnected()) { SubstrateMaterialInfo.AddPropertyConnected(MP_EmissiveColor); }
 
 	SubstrateMaterialInfo.AddShadingModel(SSM_DefaultLit);
+	SubstrateMaterialInfo.AddGuid(MaterialExpressionGuid);
 }
 
 FSubstrateOperator* UMaterialExpressionSubstrateSimpleClearCoatBSDF::SubstrateGenerateMaterialTopologyTree(class FMaterialCompiler* Compiler, class UMaterialExpression* Parent, int32 OutputIndex)
@@ -24416,6 +24419,7 @@ bool UMaterialExpressionSubstrateVolumetricFogCloudBSDF::IsResultSubstrateMateri
 void UMaterialExpressionSubstrateVolumetricFogCloudBSDF::GatherSubstrateMaterialInfo(FSubstrateMaterialInfo& SubstrateMaterialInfo, int32 OutputIndex)
 {
 	SubstrateMaterialInfo.AddShadingModel(SSM_VolumetricFogCloud);
+	SubstrateMaterialInfo.AddGuid(MaterialExpressionGuid);
 }
 
 FSubstrateOperator* UMaterialExpressionSubstrateVolumetricFogCloudBSDF::SubstrateGenerateMaterialTopologyTree(class FMaterialCompiler* Compiler, class UMaterialExpression* Parent, int32 OutputIndex)
@@ -24486,6 +24490,7 @@ bool UMaterialExpressionSubstrateLightFunction::IsResultSubstrateMaterial(int32
 void UMaterialExpressionSubstrateLightFunction::GatherSubstrateMaterialInfo(FSubstrateMaterialInfo& SubstrateMaterialInfo, int32 OutputIndex)
 {
 	SubstrateMaterialInfo.AddShadingModel(SSM_LightFunction);
+	SubstrateMaterialInfo.AddGuid(MaterialExpressionGuid);
 }
 
 FSubstrateOperator* UMaterialExpressionSubstrateLightFunction::SubstrateGenerateMaterialTopologyTree(class FMaterialCompiler* Compiler, class UMaterialExpression* Parent, int32 OutputIndex)
@@ -24562,6 +24567,7 @@ bool UMaterialExpressionSubstratePostProcess::IsResultSubstrateMaterial(int32 Ou
 void UMaterialExpressionSubstratePostProcess::GatherSubstrateMaterialInfo(FSubstrateMaterialInfo& SubstrateMaterialInfo, int32 OutputIndex)
 {
 	SubstrateMaterialInfo.AddShadingModel(SSM_PostProcess);
+	SubstrateMaterialInfo.AddGuid(MaterialExpressionGuid);
 }
 
 FSubstrateOperator* UMaterialExpressionSubstratePostProcess::SubstrateGenerateMaterialTopologyTree(class FMaterialCompiler* Compiler, class UMaterialExpression* Parent, int32 OutputIndex)
@@ -24638,6 +24644,7 @@ bool UMaterialExpressionSubstrateUI::IsResultSubstrateMaterial(int32 OutputIndex
 void UMaterialExpressionSubstrateUI::GatherSubstrateMaterialInfo(FSubstrateMaterialInfo& SubstrateMaterialInfo, int32 OutputIndex)
 {
 	SubstrateMaterialInfo.AddShadingModel(SSM_UI);
+	SubstrateMaterialInfo.AddGuid(MaterialExpressionGuid);
 }
 
 FSubstrateOperator* UMaterialExpressionSubstrateUI::SubstrateGenerateMaterialTopologyTree(class FMaterialCompiler* Compiler, class UMaterialExpression* Parent, int32 OutputIndex)
@@ -24744,6 +24751,7 @@ void UMaterialExpressionSubstrateConvertToDecal::GatherSubstrateMaterialInfo(FSu
 		TracedInput.Expression->GatherSubstrateMaterialInfo(SubstrateMaterialInfo, TracedInput.OutputIndex);
 	}
 	SubstrateMaterialInfo.AddShadingModel(SSM_Decal);
+	SubstrateMaterialInfo.AddGuid(MaterialExpressionGuid);
 
 	SubstrateMaterialInfo.PopSubstrateTreeStack();
 }
@@ -24839,6 +24847,7 @@ void UMaterialExpressionSubstrateUnlitBSDF::GatherSubstrateMaterialInfo(FSubstra
 {
 	if (EmissiveColor.IsConnected()) { SubstrateMaterialInfo.AddPropertyConnected(MP_EmissiveColor); }
 	SubstrateMaterialInfo.AddShadingModel(SSM_Unlit);
+	SubstrateMaterialInfo.AddGuid(MaterialExpressionGuid);
 }
 
 FSubstrateOperator* UMaterialExpressionSubstrateUnlitBSDF::SubstrateGenerateMaterialTopologyTree(class FMaterialCompiler* Compiler, class UMaterialExpression* Parent, int32 OutputIndex)
@@ -24955,6 +24964,7 @@ void UMaterialExpressionSubstrateHairBSDF::GatherSubstrateMaterialInfo(FSubstrat
 	if (Tangent.IsConnected())			{ SubstrateMaterialInfo.AddPropertyConnected(MP_Tangent); }
 
 	SubstrateMaterialInfo.AddShadingModel(SSM_Hair);
+	SubstrateMaterialInfo.AddGuid(MaterialExpressionGuid);
 }
 
 FSubstrateOperator* UMaterialExpressionSubstrateHairBSDF::SubstrateGenerateMaterialTopologyTree(class FMaterialCompiler* Compiler, class UMaterialExpression* Parent, int32 OutputIndex)
@@ -25069,6 +25079,7 @@ void UMaterialExpressionSubstrateEyeBSDF::GatherSubstrateMaterialInfo(FSubstrate
 		SubstrateMaterialInfo.AddSubsurfaceProfile(SubsurfaceProfile);
 	}
 	SubstrateMaterialInfo.AddShadingModel(SSM_Eye);
+	SubstrateMaterialInfo.AddGuid(MaterialExpressionGuid);
 }
 
 FSubstrateOperator* UMaterialExpressionSubstrateEyeBSDF::SubstrateGenerateMaterialTopologyTree(class FMaterialCompiler* Compiler, class UMaterialExpression* Parent, int32 OutputIndex)
@@ -25200,6 +25211,7 @@ void UMaterialExpressionSubstrateSingleLayerWaterBSDF::GatherSubstrateMaterialIn
 	if (Normal.IsConnected())			{ SubstrateMaterialInfo.AddPropertyConnected(MP_Normal); }
 
 	SubstrateMaterialInfo.AddShadingModel(SSM_SingleLayerWater);
+	SubstrateMaterialInfo.AddGuid(MaterialExpressionGuid);
 }
 
 FSubstrateOperator* UMaterialExpressionSubstrateSingleLayerWaterBSDF::SubstrateGenerateMaterialTopologyTree(class FMaterialCompiler* Compiler, class UMaterialExpression* Parent, int32 OutputIndex)
@@ -25352,6 +25364,7 @@ void UMaterialExpressionSubstrateHorizontalMixing::GatherSubstrateMaterialInfo(F
 	{
 		TracedInputB.Expression->GatherSubstrateMaterialInfo(SubstrateMaterialInfo, TracedInputB.OutputIndex);
 	}
+	SubstrateMaterialInfo.AddGuid(MaterialExpressionGuid);
 
 	SubstrateMaterialInfo.PopSubstrateTreeStack();
 }
@@ -25550,6 +25563,7 @@ void UMaterialExpressionSubstrateVerticalLayering::GatherSubstrateMaterialInfo(F
 	{
 		TracedInputBase.Expression->GatherSubstrateMaterialInfo(SubstrateMaterialInfo, TracedInputBase.OutputIndex);
 	}
+	SubstrateMaterialInfo.AddGuid(MaterialExpressionGuid);
 
 	SubstrateMaterialInfo.PopSubstrateTreeStack();
 }
@@ -25727,6 +25741,7 @@ void UMaterialExpressionSubstrateAdd::GatherSubstrateMaterialInfo(FSubstrateMate
 	{
 		TracedInputB.Expression->GatherSubstrateMaterialInfo(SubstrateMaterialInfo, TracedInputB.OutputIndex);
 	}
+	SubstrateMaterialInfo.AddGuid(MaterialExpressionGuid);
 
 	SubstrateMaterialInfo.PopSubstrateTreeStack();
 }
@@ -25864,6 +25879,7 @@ void UMaterialExpressionSubstrateWeight::GatherSubstrateMaterialInfo(FSubstrateM
 	{
 		TracedInputA.Expression->GatherSubstrateMaterialInfo(SubstrateMaterialInfo, TracedInputA.OutputIndex);
 	}
+	SubstrateMaterialInfo.AddGuid(MaterialExpressionGuid);
 
 	SubstrateMaterialInfo.PopSubstrateTreeStack();
 }
@@ -26677,6 +26693,7 @@ void UMaterialExpressionSubstrateConvertMaterialAttributes::GatherSubstrateMater
 	{
 		SubstrateMaterialInfo.AddSubsurfaceProfile(SubsurfaceProfile);
 	}
+	SubstrateMaterialInfo.AddGuid(MaterialExpressionGuid);
 }
 
 FSubstrateOperator* UMaterialExpressionSubstrateConvertMaterialAttributes::SubstrateGenerateMaterialTopologyTree(class FMaterialCompiler* Compiler, class UMaterialExpression* Parent, int32 OutputIndex)
