Fix for SparseVolumeTexture streaming manager causing an out-of-bounds access crash when playing back certain SVT assets with reverse playback. The wrap around logic for computing the index of the frame to prefetch assumes that the number of frames to prefetch is <= the number of total frames. This crash then happens with very short sequences or when changing the r.SparseVolumeTexture.Streaming.NumPrefetchFrames cvar to a higher value. Negative values were not handled properly either.

#rnx

diff --git a/Engine/Source/Runtime/Engine/Private/SparseVolumeTexture/SparseVolumeTextureStreamingManager.cpp b/Engine/Source/Runtime/Engine/Private/SparseVolumeTexture/SparseVolumeTextureStreamingManager.cpp
--- a/Engine/Source/Runtime/Engine/Private/SparseVolumeTexture/SparseVolumeTextureStreamingManager.cpp
+++ b/Engine/Source/Runtime/Engine/Private/SparseVolumeTexture/SparseVolumeTextureStreamingManager.cpp
@@ -1036,17 +1036,23 @@ void FStreamingManager::Request(UStreamableSparseVolumeTexture* SparseVolumeText
 		}
 
 		check(StreamingWindow);
+
+		// Make sure the number of prefetched frames doesn't exceed the total number of frames.
+		// Not only does this make no sense, it also breaks the wrap around logic in the loop below if there is reverse playback.
+		const int32 NumPrefetchFrames = FMath::Clamp(GSVTStreamingNumPrefetchFrames, 0, NumFrames);
+
 		// No prefetching for blocking requests. Making the prefetches blocking would increase latency even more and making them non-blocking could lead
 		// to situations where lower mips are already streamed in in subsequent frames but can't be used because dependent higher mips haven't finished streaming
 		// due to non-blocking requests.
 		// SVT_TODO: This can still break if blocking and non-blocking requests of the same frames/mips are made to the same SVT. We would need to cancel already scheduled non-blocking requests and reissue them as blocking.
 		// Or alternatively we could just block on all requests if we detect this case. If we had a single DDC request owner per request, we could just selectively wait on already scheduled non-blocking requests.
-		const int32 OffsetMagnitude = !bBlocking ? GSVTStreamingNumPrefetchFrames : 0;
+		const int32 OffsetMagnitude = !bBlocking ? NumPrefetchFrames : 0;
 		const int32 LowerFrameOffset = StreamingWindow->bPlayBackward ? -OffsetMagnitude : 0;
 		const int32 UpperFrameOffset = StreamingWindow->bPlayForward ? OffsetMagnitude : 0;
 
 		for (int32 i = LowerFrameOffset; i <= UpperFrameOffset; ++i)
 		{
+			// Wrap around on both positive and negative numbers, assuming (i + NumFrames) >= 0. See the comment on NumPrefetchFrames.
 			const int32 RequestFrameIndex = (static_cast<int32>(FrameIndex) + i + NumFrames) % NumFrames;
 			const int32 RequestMipLevelOffset = FMath::Abs(i) + GSVTStreamingPrefetchMipLevelBias;
 			FStreamingRequest Request;
