Removed description comparison from FInertializationRequest and switched FInertializationRequest to storing an FString instead of an FText

diff --git a/Engine/Plugins/Experimental/Animation/BlendStack/Source/Runtime/Private/AnimNode_BlendStack.cpp b/Engine/Plugins/Experimental/Animation/BlendStack/Source/Runtime/Private/AnimNode_BlendStack.cpp
--- a/Engine/Plugins/Experimental/Animation/BlendStack/Source/Runtime/Private/AnimNode_BlendStack.cpp
+++ b/Engine/Plugins/Experimental/Animation/BlendStack/Source/Runtime/Private/AnimNode_BlendStack.cpp
@@ -655,7 +655,6 @@ static void RequestInertialBlend(const FAnimationUpdateContext& Context, float B
 			Request.bUseBlendMode = true;
 			Request.BlendMode = BlendOption;
 #if ANIM_TRACE_ENABLED
-			Request.Description = LOCTEXT("InertializationRequestDescription", "Blend Stack");
 			Request.NodeId = Context.GetCurrentNodeId();
 			Request.AnimInstance = Context.AnimInstanceProxy->GetAnimInstanceObject();
 #endif
diff --git a/Engine/Source/Runtime/Engine/Classes/Animation/AnimNode_DeadBlending.h b/Engine/Source/Runtime/Engine/Classes/Animation/AnimNode_DeadBlending.h
--- a/Engine/Source/Runtime/Engine/Classes/Animation/AnimNode_DeadBlending.h
+++ b/Engine/Source/Runtime/Engine/Classes/Animation/AnimNode_DeadBlending.h
@@ -257,16 +257,6 @@ private:
 	// Current inertialization duration (used for curves).
 	float InertializationDuration = 0.0f;
 
-	// Description for the current inertialization request - used for debugging
-	FText InertializationRequestDescription;
-
-	// Node Id for the current inertialization request - used for debugging
-	int32 InertializationRequestNodeId = INDEX_NONE;
-
-	// Anim Instance for the current inertialization request - used for debugging
-	UPROPERTY(Transient)
-	TObjectPtr<UObject> InertializationRequestAnimInstance = nullptr;
-
 	// Current inertialization durations for each bone, indexed by skeleton bone index (used for per-bone blending).
 	TCustomBoneIndexArray<float, FSkeletonPoseBoneIndex> InertializationDurationPerBone;
 
@@ -279,4 +269,19 @@ private:
 	// Custom blend curve being used by the current blend mode.
 	UPROPERTY(Transient)
 	TObjectPtr<UCurveFloat> InertializationCustomBlendCurve = nullptr;
+
+
+// if ANIM_TRACE_ENABLED - these properties are only used for debugging when ANIM_TRACE_ENABLED == 1
+
+	// Description for the current inertialization request
+	FString InertializationRequestDescription;
+
+	// Node Id for the current inertialization request
+	int32 InertializationRequestNodeId = INDEX_NONE;
+
+	// Anim Instance for the current inertialization request
+	UPROPERTY(Transient)
+	TObjectPtr<UObject> InertializationRequestAnimInstance = nullptr;
+
+// endif ANIM_TRACE_ENABLED
 };
diff --git a/Engine/Source/Runtime/Engine/Classes/Animation/AnimNode_Inertialization.h b/Engine/Source/Runtime/Engine/Classes/Animation/AnimNode_Inertialization.h
--- a/Engine/Source/Runtime/Engine/Classes/Animation/AnimNode_Inertialization.h
+++ b/Engine/Source/Runtime/Engine/Classes/Animation/AnimNode_Inertialization.h
@@ -113,26 +113,23 @@ struct FInertializationRequest
 {
 	GENERATED_BODY()
 
-	FInertializationRequest() {}
-
-	FInertializationRequest(float InDuration, const UBlendProfile* InBlendProfile)
-		: Duration(InDuration)
-		, BlendProfile(InBlendProfile)
-	{
-	}
-
-	void Clear()
-	{
-		Duration = -1.0f;
-		BlendProfile = nullptr;
-		bUseBlendMode = false;
-		BlendMode = EAlphaBlendOption::Linear;
-		CustomBlendCurve = nullptr;
-		Description = FText::GetEmpty();
-		NodeId = INDEX_NONE;
-		AnimInstance = nullptr;
-	}
-
+	ENGINE_API FInertializationRequest();
+	ENGINE_API FInertializationRequest(float InDuration, const UBlendProfile* InBlendProfile);
+
+	// Note: We need to explicitly disable warnings on these constructors/operators for clang to be happy with deprecated variables
+	PRAGMA_DISABLE_DEPRECATION_WARNINGS
+	~FInertializationRequest() = default;
+	FInertializationRequest(const FInertializationRequest&) = default;
+	FInertializationRequest(FInertializationRequest&&) = default;
+	FInertializationRequest& operator=(const FInertializationRequest&) = default;
+	FInertializationRequest& operator=(FInertializationRequest&&) = default;
+	PRAGMA_ENABLE_DEPRECATION_WARNINGS
+
+	ENGINE_API void Clear();
+
+	// Comparison operator used to test for equality in the array of animation requests to that
+	// only unique requests are added. This does not take into account the properties that are 
+	// used only for debugging and only used when ANIM_TRACE_ENABLED
 	friend bool operator==(const FInertializationRequest& A, const FInertializationRequest& B)
 	{
 		return
@@ -140,10 +137,7 @@ struct FInertializationRequest
 			(A.BlendProfile == B.BlendProfile) &&
 			(A.bUseBlendMode == B.bUseBlendMode) &&
 			(A.BlendMode == B.BlendMode) &&
-			(A.CustomBlendCurve == B.CustomBlendCurve) &&
-			(A.Description.EqualTo(B.Description) &&
-			(A.NodeId == B.NodeId) &&
-			(A.AnimInstance == B.AnimInstance));
+			(A.CustomBlendCurve == B.CustomBlendCurve);
 	}
 
 	friend bool operator!=(const FInertializationRequest& A, const FInertializationRequest& B)
@@ -171,9 +165,15 @@ struct FInertializationRequest
 	UPROPERTY(Transient)
 	TObjectPtr<UCurveFloat> CustomBlendCurve = nullptr;
 
-	// Description of the request - used for debugging.
+// if ANIM_TRACE_ENABLED - these properties are only used for debugging when ANIM_TRACE_ENABLED == 1
+	
+	UE_DEPRECATED(5.4, "Use DescriptionString instead.")
+	UPROPERTY(Transient, meta = (DeprecatedProperty, DeprecationMessage = "Use DescriptionString instead."))
+	FText Description_DEPRECATED;
+
+	// Description of the request
 	UPROPERTY(Transient)
-	FText Description;
+	FString DescriptionString;
 
 	// Node id from which this request was made.
 	UPROPERTY(Transient)
@@ -182,6 +182,8 @@ struct FInertializationRequest
 	// Anim instance from which this request was made.
 	UPROPERTY(Transient)
 	TObjectPtr<UObject> AnimInstance = nullptr;
+
+// endif ANIM_TRACE_ENABLED
 };
 
 
@@ -445,21 +447,13 @@ private:
 
 	// Inertialization state
 	EInertializationState InertializationState;
+
+	// Amount of time elapsed during the Inertialization
 	float InertializationElapsedTime;
 
 	// Inertialization duration for the main inertialization request (used for curve blending and deficit tracking)
 	float InertializationDuration;
 
-	// Description for the current inertialization request - used for debugging
-	FText InertializationRequestDescription;
-
-	// Node Id for the current inertialization request - used for debugging
-	int32 InertializationRequestNodeId = INDEX_NONE;
-
-	// Anim Instance for the current inertialization request - used for debugging
-	UPROPERTY(Transient)
-	TObjectPtr<UObject> InertializationRequestAnimInstance = nullptr;
-
 	// Inertialization durations indexed by skeleton bone index (used for per-bone blending)
 	TCustomBoneIndexArray<float, FSkeletonPoseBoneIndex> InertializationDurationPerBone;
 
@@ -477,4 +471,20 @@ private:
 	
 	// Reset inertialization timing and state
 	void Deactivate();
+
+
+// if ANIM_TRACE_ENABLED - these properties are only used for debugging when ANIM_TRACE_ENABLED == 1
+
+	// Description for the current inertialization request
+	FString InertializationRequestDescription;
+
+	// Node Id for the current inertialization request
+	int32 InertializationRequestNodeId = INDEX_NONE;
+
+	// Anim Instance for the current inertialization request
+	UPROPERTY(Transient)
+	TObjectPtr<UObject> InertializationRequestAnimInstance = nullptr;
+
+// endif ANIM_TRACE_ENABLED
+
 };
diff --git a/Engine/Source/Runtime/Engine/Private/Animation/AnimNode_DeadBlending.cpp b/Engine/Source/Runtime/Engine/Private/Animation/AnimNode_DeadBlending.cpp
--- a/Engine/Source/Runtime/Engine/Private/Animation/AnimNode_DeadBlending.cpp
+++ b/Engine/Source/Runtime/Engine/Private/Animation/AnimNode_DeadBlending.cpp
@@ -780,7 +780,7 @@ void FAnimNode_DeadBlending::Evaluate_AnyThread(FPoseContext& Output)
 		}
 
 #if ANIM_TRACE_ENABLED
-		InertializationRequestDescription = RequestQueue[ShortestRequestIdx].Description;
+		InertializationRequestDescription = RequestQueue[ShortestRequestIdx].DescriptionString;
 		InertializationRequestNodeId = RequestQueue[ShortestRequestIdx].NodeId;
 		InertializationRequestAnimInstance = RequestQueue[ShortestRequestIdx].AnimInstance;
 #endif
@@ -902,7 +902,7 @@ void FAnimNode_DeadBlending::Evaluate_AnyThread(FPoseContext& Output)
 	TRACE_ANIM_NODE_VALUE_WITH_ID(Output, GetNodeIndex(), TEXT("Max Duration"), InertializationMaxDuration);
 	TRACE_ANIM_NODE_VALUE_WITH_ID(Output, GetNodeIndex(), TEXT("Normalized Time"), InertializationDuration > UE_KINDA_SMALL_NUMBER ? (InertializationTime / InertializationDuration) : 0.0f);
 	TRACE_ANIM_NODE_VALUE_WITH_ID(Output, GetNodeIndex(), TEXT("Inertialization Weight"), InertializationWeight);
-	TRACE_ANIM_NODE_VALUE_WITH_ID(Output, GetNodeIndex(), TEXT("Request Description"), *InertializationRequestDescription.ToString());
+	TRACE_ANIM_NODE_VALUE_WITH_ID(Output, GetNodeIndex(), TEXT("Request Description"), *InertializationRequestDescription);
 	TRACE_ANIM_NODE_VALUE_WITH_ID_ANIM_NODE(Output, GetNodeIndex(), TEXT("Request Node"), InertializationRequestNodeId, InertializationRequestAnimInstance);
 }
 
diff --git a/Engine/Source/Runtime/Engine/Private/Animation/AnimNode_Inertialization.cpp b/Engine/Source/Runtime/Engine/Private/Animation/AnimNode_Inertialization.cpp
--- a/Engine/Source/Runtime/Engine/Private/Animation/AnimNode_Inertialization.cpp
+++ b/Engine/Source/Runtime/Engine/Private/Animation/AnimNode_Inertialization.cpp
@@ -90,6 +90,30 @@ static int32 GetNumSkeletonBones(const FBoneContainer& BoneContainer)
 
 }}	// namespace UE::Anim
 
+PRAGMA_DISABLE_DEPRECATION_WARNINGS
+
+FInertializationRequest::FInertializationRequest() = default;
+
+FInertializationRequest::FInertializationRequest(float InDuration, const UBlendProfile* InBlendProfile)
+	: Duration(InDuration)
+	, BlendProfile(InBlendProfile) {}
+
+PRAGMA_ENABLE_DEPRECATION_WARNINGS
+
+void FInertializationRequest::Clear()
+{
+	Duration = -1.0f;
+	BlendProfile = nullptr;
+	bUseBlendMode = false;
+	BlendMode = EAlphaBlendOption::Linear;
+	CustomBlendCurve = nullptr;
+
+#if ANIM_TRACE_ENABLED 
+	DescriptionString.Empty();
+	NodeId = INDEX_NONE;
+	AnimInstance = nullptr;
+#endif
+}
 
 FAnimNode_Inertialization::FAnimNode_Inertialization()
 	: DeltaTime(0.0f)
@@ -286,10 +310,11 @@ void FAnimNode_Inertialization::Evaluate_AnyThread(FPoseContext& Output)
 		// Handle the first inertialization request in the queue
 		InertializationDuration = FMath::Max(RequestQueue[0].Duration - AppliedDeficit, 0.0f);
 #if ANIM_TRACE_ENABLED
-		InertializationRequestDescription = RequestQueue[0].Description;
+		InertializationRequestDescription = RequestQueue[0].DescriptionString;
 		InertializationRequestNodeId = RequestQueue[0].NodeId;
 		InertializationRequestAnimInstance = RequestQueue[0].AnimInstance;
 #endif
+
 		FillSkeletonBoneDurationsArray(InertializationDurationPerBone, InertializationDuration, RequestQueue[0].BlendProfile);
 
 		// Handle all subsequent inertialization requests (often there will be only a single request)
@@ -306,7 +331,7 @@ void FAnimNode_Inertialization::Evaluate_AnyThread(FPoseContext& Output)
 				{
 					InertializationDuration = RequestDuration;
 #if ANIM_TRACE_ENABLED
-					InertializationRequestDescription = Request.Description;
+					InertializationRequestDescription = Request.DescriptionString;
 					InertializationRequestNodeId = Request.NodeId;
 					InertializationRequestAnimInstance = Request.AnimInstance;
 #endif
@@ -461,7 +486,7 @@ void FAnimNode_Inertialization::Evaluate_AnyThread(FPoseContext& Output)
 	TRACE_ANIM_NODE_VALUE_WITH_ID(Output, GetNodeIndex(), TEXT("Max Duration"), InertializationMaxDuration);
 	TRACE_ANIM_NODE_VALUE_WITH_ID(Output, GetNodeIndex(), TEXT("Normalized Time"), InertializationDuration > UE_KINDA_SMALL_NUMBER ? (InertializationElapsedTime / InertializationDuration) : 0.0f);
 	TRACE_ANIM_NODE_VALUE_WITH_ID(Output, GetNodeIndex(), TEXT("Inertialization Weight"), InertializationWeight);
-	TRACE_ANIM_NODE_VALUE_WITH_ID(Output, GetNodeIndex(), TEXT("Request Description"), *InertializationRequestDescription.ToString());
+	TRACE_ANIM_NODE_VALUE_WITH_ID(Output, GetNodeIndex(), TEXT("Request Description"), *InertializationRequestDescription);
 	TRACE_ANIM_NODE_VALUE_WITH_ID_ANIM_NODE(Output, GetNodeIndex(), TEXT("Request Node"), InertializationRequestNodeId, InertializationRequestAnimInstance);
 }
 
diff --git a/Engine/Source/Runtime/Engine/Private/Animation/AnimNode_LinkedAnimGraph.cpp b/Engine/Source/Runtime/Engine/Private/Animation/AnimNode_LinkedAnimGraph.cpp
--- a/Engine/Source/Runtime/Engine/Private/Animation/AnimNode_LinkedAnimGraph.cpp
+++ b/Engine/Source/Runtime/Engine/Private/Animation/AnimNode_LinkedAnimGraph.cpp
@@ -123,7 +123,7 @@ void FAnimNode_LinkedAnimGraph::Update_AnyThread(const FAnimationUpdateContext&
 				Request.Duration = PendingBlendOutDuration;
 				Request.BlendProfile = PendingBlendOutProfile;
 #if ANIM_TRACE_ENABLED
-				Request.Description = NSLOCTEXT("AnimNode_LinkedAnimGraph", "InertializationRequestDescriptionOut", "Out");
+				Request.DescriptionString = NSLOCTEXT("AnimNode_LinkedAnimGraph", "InertializationRequestDescriptionOut", "Out").ToString();
 				Request.NodeId = InContext.GetCurrentNodeId();
 				Request.AnimInstance = InContext.AnimInstanceProxy->GetAnimInstanceObject();
 #endif
@@ -137,7 +137,7 @@ void FAnimNode_LinkedAnimGraph::Update_AnyThread(const FAnimationUpdateContext&
 				Request.Duration = PendingBlendInDuration;
 				Request.BlendProfile = PendingBlendInProfile;
 #if ANIM_TRACE_ENABLED
-				Request.Description = NSLOCTEXT("AnimNode_LinkedAnimGraph", "InertializationRequestDescriptionIn", "In");
+				Request.DescriptionString = NSLOCTEXT("AnimNode_LinkedAnimGraph", "InertializationRequestDescriptionIn", "In").ToString();
 				Request.NodeId = InContext.GetCurrentNodeId();
 				Request.AnimInstance = InContext.AnimInstanceProxy->GetAnimInstanceObject();
 #endif
diff --git a/Engine/Source/Runtime/Engine/Private/Animation/AnimNode_StateMachine.cpp b/Engine/Source/Runtime/Engine/Private/Animation/AnimNode_StateMachine.cpp
--- a/Engine/Source/Runtime/Engine/Private/Animation/AnimNode_StateMachine.cpp
+++ b/Engine/Source/Runtime/Engine/Private/Animation/AnimNode_StateMachine.cpp
@@ -1257,11 +1257,11 @@ void FAnimNode_StateMachine::TransitionToState(const FAnimationUpdateContext& Co
 				Request.BlendMode = TransitionInfo.BlendMode;
 				Request.CustomBlendCurve = TransitionInfo.CustomCurve;
 #if ANIM_TRACE_ENABLED
-				Request.Description = FText::Format(LOCTEXT("InertializationRequestDescription", 
+				Request.DescriptionString = FText::Format(LOCTEXT("InertializationRequestDescription", 
 					"\"{0}\" Transition from \"{1}\" to \"{2}\""), 
 					FText::FromName(GetMachineDescription()->MachineName),
 					FText::FromName(GetStateInfo(TransitionInfo.PreviousState).StateName),
-					FText::FromName(GetStateInfo(TransitionInfo.NextState).StateName));
+					FText::FromName(GetStateInfo(TransitionInfo.NextState).StateName)).ToString();
 				Request.NodeId = Context.GetCurrentNodeId();
 				Request.AnimInstance = Context.AnimInstanceProxy->GetAnimInstanceObject();
 #endif
