FSoundWaveLoadingBehaviorUtil: don't rely on the Asset Registry to filter out classes as it's going to gather all assets for the specified classes first, then filtering for the provided packages names afterward, resulting in execution time being over 100 times slower than filtering for classes after package name.
Work is being done on the Asset Registry side to sort query filters by their minimum result size, to optimize queries like this.

diff --git a/Engine/Source/Runtime/Engine/Classes/Sound/SoundWaveLoadingBehavior.cpp b/Engine/Source/Runtime/Engine/Classes/Sound/SoundWaveLoadingBehavior.cpp
--- a/Engine/Source/Runtime/Engine/Classes/Sound/SoundWaveLoadingBehavior.cpp
+++ b/Engine/Source/Runtime/Engine/Classes/Sound/SoundWaveLoadingBehavior.cpp
@@ -156,8 +156,10 @@ private:
 
 		// Filter on SoundCues.
 		FARFilter Filter;
-		Filter.ClassPaths.Add(USoundCue::StaticClass()->GetClassPathName());
-		Filter.bRecursiveClasses = true;
+		// Don't rely on the AR to filter out classes as it's going to gather all assets for the specified classes first, then filtering for the provided packages names afterward,
+		// resulting in execution time being over 100 times slower than filtering for classes after package names.
+		//Filter.ClassPaths.Add(USoundCue::StaticClass()->GetClassPathName());
+		//Filter.bRecursiveClasses = true;
 		Filter.PackageNames = SoundWaveReferencerNames;
 		TArray<FAssetData> ReferencingSoundCueAssetDataArray;
 		if (!AssetRegistry.GetAssets(Filter, ReferencingSoundCueAssetDataArray))
@@ -165,6 +167,15 @@ private:
 			return {};
 		}
 
+		// Filter out unwanted classes, see above comment for details.
+		for (int32 AssetIndex = 0; AssetIndex < ReferencingSoundCueAssetDataArray.Num(); AssetIndex++)
+		{
+			if (UClass* AssetClass = ReferencingSoundCueAssetDataArray[AssetIndex].GetClass(); !AssetClass || !AssetClass->IsChildOf<USoundCue>())
+			{
+				ReferencingSoundCueAssetDataArray.RemoveAtSwap(AssetIndex--);
+			}
+		}
+
 		if (ReferencingSoundCueAssetDataArray.IsEmpty())
 		{
 			return {};
