WorldPartition: Experimental Editor Path support (SoftObjectPaths to Level Instance sub actors)
- CVar "EditorPaths.Enabled" disabled by default
- References to sub container actors can now be created/resolved
- WorldPartitionLevelStreamingPolicy now has a FWorldPartitionRuntimerContainerResolver member which allows remapping at Cook/PIE/Runtime
- ALevelInstance supports ResolveSubObject for Editor resolving of sub actors
- EditorPaths can be built from FEditorPathHelper using IEditorPathObjectInterface

#rnx

diff --git a/Engine/Source/Runtime/CoreUObject/Private/Misc/EditorPathHelper.cpp b/Engine/Source/Runtime/CoreUObject/Private/Misc/EditorPathHelper.cpp
new file mode 100644
--- /dev/null
+++ b/Engine/Source/Runtime/CoreUObject/Private/Misc/EditorPathHelper.cpp
@@ -0,0 +1,116 @@
+// Copyright Epic Games, Inc. All Rights Reserved.
+
+#include "Misc/EditorPathHelper.h"
+
+#if WITH_EDITOR
+
+#include "Misc/EditorPathObjectInterface.h"
+#include "UObject/TopLevelAssetPath.h"
+#include "HAL/IConsoleManager.h"
+
+COREUOBJECT_API bool GIsEditorPathFeatureEnabled = false;
+static FAutoConsoleVariableRef CVarIsEditorPathFeatureEnabled(
+	TEXT("Editorpaths.Enabled"),
+	GIsEditorPathFeatureEnabled,
+	TEXT("Enable experimental Editor Path support."));
+
+bool FEditorPathHelper::IsEnabled()
+{
+	return GIsEditorPathFeatureEnabled;
+}
+
+UObject* FEditorPathHelper::GetEditorPathOwner(const UObject* InObject)
+{
+	if (IsEnabled() && InObject)
+	{
+		if (IEditorPathObjectInterface* EditorPathObject = InObject->GetImplementingOuter<IEditorPathObjectInterface>())
+		{
+			return EditorPathObject->GetEditorPathOwner();
+		}
+	}
+
+	return nullptr;
+}
+
+bool FEditorPathHelper::IsInEditorPath(const UObject* InEditorPathOwner, const UObject* InObject)
+{
+	if (InEditorPathOwner)
+	{
+		while (UObject* CurrentEditorPathOwner = FEditorPathHelper::GetEditorPathOwner(InObject))
+		{
+			if (CurrentEditorPathOwner == InEditorPathOwner)
+			{
+				return true;
+			}
+
+			InObject = CurrentEditorPathOwner;
+		}
+	}
+
+	return false;
+}
+
+FSoftObjectPath FEditorPathHelper::GetEditorPath(const UObject* InObject)
+{
+	return GetEditorPathFromEditorPathOwner(InObject, nullptr);
+}
+
+FSoftObjectPath FEditorPathHelper::GetEditorPathFromReferencer(const UObject* InObject, const UObject* InReferencer)
+{
+	return GetEditorPathFromEditorPathOwner(InObject, GetEditorPathOwner(InReferencer));
+}
+
+FSoftObjectPath FEditorPathHelper::GetEditorPathFromEditorPathOwner(const UObject* InObject, const UObject* InEditorPathOwner)
+{
+	check(InObject);
+	FSoftObjectPath SoftObjectPath(InObject);
+	if (!IsEnabled())
+	{
+		return SoftObjectPath;
+	}
+	
+	UObject* EditorPathOwner = GetEditorPathOwner(InObject);
+	
+	if (EditorPathOwner && EditorPathOwner != InEditorPathOwner)
+	{
+		TArray<FString> EditorPathOwnerNames;
+		FString SubObjectString;
+		FString SubObjectContext;
+		check(SoftObjectPath.GetSubPathString().Split(TEXT("."), &SubObjectContext, &SubObjectString));
+		check(SubObjectContext == TEXT("PersistentLevel"));
+
+		EditorPathOwnerNames.Add(EditorPathOwner->GetName());
+
+		while (UObject* ParentEditorPathOwner = GetEditorPathOwner(EditorPathOwner))
+		{
+			// EditorPath needs to be relative to InEditorPathOwner in case we have a valid InEditorPathOwner
+			if (ParentEditorPathOwner == InEditorPathOwner)
+			{
+				break;
+			}
+
+			EditorPathOwner = ParentEditorPathOwner;
+			EditorPathOwnerNames.Add(ParentEditorPathOwner->GetName());
+		}
+
+		TStringBuilder<512> SubPathStringBuilder;
+		SubPathStringBuilder.Append(SubObjectContext);
+		for (int i = EditorPathOwnerNames.Num() - 1; i >= 0; --i)
+		{
+			SubPathStringBuilder.Append(TEXT("."));
+			SubPathStringBuilder.Append(EditorPathOwnerNames[i]);
+		}
+
+		if (!SubObjectString.IsEmpty())
+		{
+			SubPathStringBuilder.Append(TEXT("."));
+			SubPathStringBuilder.Append(SubObjectString);
+		}
+
+		return FSoftObjectPath(FTopLevelAssetPath(EditorPathOwner->GetOutermostObject()), SubPathStringBuilder.ToString());
+	}
+
+	return SoftObjectPath;
+}
+
+#endif
\ No newline at end of file
diff --git a/Engine/Source/Runtime/CoreUObject/Private/Misc/EditorPathObjectInterface.cpp b/Engine/Source/Runtime/CoreUObject/Private/Misc/EditorPathObjectInterface.cpp
new file mode 100644
--- /dev/null
+++ b/Engine/Source/Runtime/CoreUObject/Private/Misc/EditorPathObjectInterface.cpp
@@ -0,0 +1,9 @@
+// Copyright Epic Games, Inc. All Rights Reserved.
+
+#include "Misc/EditorPathObjectInterface.h"
+
+#include UE_INLINE_GENERATED_CPP_BY_NAME(EditorPathObjectInterface)
+
+UEditorPathObjectInterface::UEditorPathObjectInterface(const FObjectInitializer& ObjectInitializer)
+	: Super(ObjectInitializer)
+{}
\ No newline at end of file
diff --git a/Engine/Source/Runtime/CoreUObject/Private/UObject/PropertySoftObjectPtr.cpp b/Engine/Source/Runtime/CoreUObject/Private/UObject/PropertySoftObjectPtr.cpp
--- a/Engine/Source/Runtime/CoreUObject/Private/UObject/PropertySoftObjectPtr.cpp
+++ b/Engine/Source/Runtime/CoreUObject/Private/UObject/PropertySoftObjectPtr.cpp
@@ -6,6 +6,9 @@
 #include "UObject/PropertyPortFlags.h"
 #include "UObject/UnrealType.h"
 #include "UObject/LinkerLoad.h"
+#if WITH_EDITOR
+#include "Misc/EditorPathHelper.h"
+#endif
 
 /*-----------------------------------------------------------------------------
 	FSoftObjectProperty.
@@ -116,8 +119,13 @@ void FSoftObjectProperty::ExportText_Internal( FString& ValueStr, const void* Pr
 
 	if (Object)
 	{
+#if WITH_EDITOR
+		// Use object in case name has changed. Export editor path if feature is enabled.
+		SoftObjectPath = FEditorPathHelper::GetEditorPathFromReferencer(Object, Parent);
+#else
 		// Use object in case name has changed.
 		SoftObjectPath = FSoftObjectPath(Object);
+#endif
 	}
 	else
 	{
@@ -146,6 +154,17 @@ const TCHAR* FSoftObjectProperty::ImportText_Internal( const TCHAR* InBuffer, vo
 
 	if (bImportTextSuccess)
 	{
+#if WITH_EDITOR
+		// If EditorPath feature is enabled. Make sure we import a proper Editor Path if Object has a EditorPathOwner
+		if (FEditorPathHelper::IsEnabled())
+		{
+			if (UObject* Object = SoftObjectPath.ResolveObject())
+			{
+				SoftObjectPath = FEditorPathHelper::GetEditorPathFromReferencer(Object, Parent);
+			}
+		}
+#endif
+
 		if (PropertyPointerType == EPropertyPointerType::Container && HasSetter())
 		{
 			FSoftObjectPtr SoftObjectPtr(SoftObjectPath);
diff --git a/Engine/Source/Runtime/CoreUObject/Public/Misc/EditorPathHelper.h b/Engine/Source/Runtime/CoreUObject/Public/Misc/EditorPathHelper.h
new file mode 100644
--- /dev/null
+++ b/Engine/Source/Runtime/CoreUObject/Public/Misc/EditorPathHelper.h
@@ -0,0 +1,23 @@
+// Copyright Epic Games, Inc. All Rights Reserved.
+
+#pragma once
+
+#include "CoreFwd.h"
+#include "UObject/Object.h"
+#include "UObject/SoftObjectPath.h"
+
+/**
+ * Helper class that can build Editor FSoftObjectPath(s) based on EditorPathOwners in an object's outer chain.
+ */
+class FEditorPathHelper
+{
+#if WITH_EDITOR
+public:
+	static COREUOBJECT_API FSoftObjectPath GetEditorPath(const UObject* InObject);
+	static COREUOBJECT_API FSoftObjectPath GetEditorPathFromReferencer(const UObject* InObject, const UObject* InReferencer);
+	static COREUOBJECT_API FSoftObjectPath GetEditorPathFromEditorPathOwner(const UObject* InObject, const UObject* InEditorPathOwner);
+	static COREUOBJECT_API UObject* GetEditorPathOwner(const UObject* InObject);
+	static COREUOBJECT_API bool IsInEditorPath(const UObject* InEditorPathOwner, const UObject* InObject);
+	static COREUOBJECT_API bool IsEnabled();
+#endif
+};
\ No newline at end of file
diff --git a/Engine/Source/Runtime/CoreUObject/Public/Misc/EditorPathObjectInterface.h b/Engine/Source/Runtime/CoreUObject/Public/Misc/EditorPathObjectInterface.h
new file mode 100644
--- /dev/null
+++ b/Engine/Source/Runtime/CoreUObject/Public/Misc/EditorPathObjectInterface.h
@@ -0,0 +1,32 @@
+// Copyright Epic Games, Inc. All Rights Reserved.
+
+#pragma once
+
+#include "CoreFwd.h"
+#include "UObject/Interface.h"
+#include "UObject/Object.h"
+#include "EditorPathObjectInterface.generated.h"
+
+UINTERFACE(BlueprintType, meta = (CannotImplementInterfaceInBlueprint), MinimalAPI)
+class UEditorPathObjectInterface : public UInterface
+{
+	GENERATED_UINTERFACE_BODY()
+};
+
+/**
+ * Interface to be implemented by Objects that can be part of an Editor Path through a Editor Path owner
+ * 
+ * See FEditorPathHelper::GetEditorPath. Paths returned by FEditorPathHelper::GetEditorPath variants should be resolvable through the EditorPathOwner's class override of ResolveSubObject.
+ * 
+ * Example implementation of this is to allow Objects outside of a Level Instance to reference objects inside a Level Instance.
+ */
+
+class COREUOBJECT_API IEditorPathObjectInterface
+{
+	GENERATED_IINTERFACE_BODY()
+
+#if WITH_EDITOR
+public:
+	virtual UObject* GetEditorPathOwner() const = 0;
+#endif
+};
\ No newline at end of file
diff --git a/Engine/Source/Runtime/Engine/Classes/Engine/Level.h b/Engine/Source/Runtime/Engine/Classes/Engine/Level.h
--- a/Engine/Source/Runtime/Engine/Classes/Engine/Level.h
+++ b/Engine/Source/Runtime/Engine/Classes/Engine/Level.h
@@ -15,6 +15,7 @@
 #include "Misc/WorldCompositionUtility.h"
 #include "Engine/MaterialMerging.h"
 #include "Engine/TextureStreamingTypes.h"
+#include "Misc/EditorPathObjectInterface.h"
 #include <atomic>
 
 #include "Level.generated.h"
@@ -397,7 +398,7 @@ enum class EActorPackagingScheme : uint8
  * @see UActor
  */
 UCLASS(MinimalAPI)
-class ULevel : public UObject, public IInterface_AssetUserData, public ITextureStreamingContainer
+class ULevel : public UObject, public IInterface_AssetUserData, public ITextureStreamingContainer, public IEditorPathObjectInterface
 {
 	GENERATED_BODY()
 
@@ -729,6 +730,9 @@ public:
 
 	ENGINE_API bool GetPromptWhenAddingToLevelOutsideBounds() const;
 	ENGINE_API bool GetPromptWhenAddingToLevelBeforeCheckout() const;
+
+	ENGINE_API void SetEditorPathOwner(UObject* InEditorPathOwner) { EditorPathOwner = InEditorPathOwner; }
+	ENGINE_API virtual UObject* GetEditorPathOwner() const override { return EditorPathOwner.Get(); }
 #endif
 
 private:
@@ -776,6 +780,9 @@ private:
 	/** Temporary array containing actor folder objects manually loaded from their external packages (only used while loading the level). */
 	UPROPERTY(Transient)
 	TArray<TObjectPtr<UActorFolder>> LoadedExternalActorFolders;
+
+	UPROPERTY(Transient)
+	TWeakObjectPtr<UObject> EditorPathOwner;
 #endif // #if WITH_EDITORONLY_DATA
 
 	enum class ERouteActorInitializationState : uint8
diff --git a/Engine/Source/Runtime/Engine/Classes/Engine/LevelStreaming.h b/Engine/Source/Runtime/Engine/Classes/Engine/LevelStreaming.h
--- a/Engine/Source/Runtime/Engine/Classes/Engine/LevelStreaming.h
+++ b/Engine/Source/Runtime/Engine/Classes/Engine/LevelStreaming.h
@@ -225,6 +225,12 @@ public:
 	UPROPERTY()
 	bool bClientOnlyVisible;
 
+#if WITH_EDITORONLY_DATA
+	/** Applied to LoadedLevel */
+	UPROPERTY(Transient)
+	TWeakObjectPtr<UObject> EditorPathOwner;
+#endif
+
 private:
 
 	/** Requested LOD. Non LOD sub-levels have Index = -1  */
diff --git a/Engine/Source/Runtime/Engine/Classes/Engine/LevelStreamingDynamic.h b/Engine/Source/Runtime/Engine/Classes/Engine/LevelStreamingDynamic.h
--- a/Engine/Source/Runtime/Engine/Classes/Engine/LevelStreamingDynamic.h
+++ b/Engine/Source/Runtime/Engine/Classes/Engine/LevelStreamingDynamic.h
@@ -60,6 +60,9 @@ class ULevelStreamingDynamic : public ULevelStreaming
 
 		/** Set whether we allow to reuse an existing level streaming. */
 		bool bAllowReuseExitingLevelStreaming = false;
+
+		/** Set EditorPath Owner */
+		UObject* EditorPathOwner = nullptr;
 	};
 
 	/**  
diff --git a/Engine/Source/Runtime/Engine/Private/Level.cpp b/Engine/Source/Runtime/Engine/Private/Level.cpp
--- a/Engine/Source/Runtime/Engine/Private/Level.cpp
+++ b/Engine/Source/Runtime/Engine/Private/Level.cpp
@@ -4225,19 +4225,31 @@ void ULevel::FixupForPIE(int32 InPIEInstanceID)
 
 bool ULevel::ResolveSubobject(const TCHAR* SubObjectPath, UObject*& OutObject, bool bLoadIfExists)
 {
+	// First check if we can find the sub object through relative path of this level
 	OutObject = StaticFindObject(nullptr, this, SubObjectPath);
+	if (OutObject)
+	{
+		return true;
+	}
 
-	if (!OutObject)
+	// Then check if we can resolve through a top level actor (Editor Path)
+	FString SubObjectName;
+	FString SubObjectContext(SubObjectPath);
+	if (FString(SubObjectPath).Split(TEXT("."), &SubObjectContext, &SubObjectName))
 	{
-		if (UWorldPartition* WorldPartition = GetWorldPartition())
+		if (UObject* SubObject = StaticFindObject(nullptr, this, *SubObjectContext))
 		{
-			return WorldPartition->ResolveSubobject(SubObjectPath, OutObject, bLoadIfExists);
+			return SubObject->ResolveSubobject(*SubObjectName, OutObject, bLoadIfExists);
 		}
-
-		return false;
 	}
-
-	return true;
+	
+	// Lastly forward to world partition
+	if (UWorldPartition* WorldPartition = GetWorldPartition())
+	{
+		return WorldPartition->ResolveSubobject(SubObjectPath, OutObject, bLoadIfExists);
+	}
+	
+	return false;
 }
 
 bool ULevel::IsPersistentLevel() const
diff --git a/Engine/Source/Runtime/Engine/Private/LevelInstance/LevelInstanceActor.cpp b/Engine/Source/Runtime/Engine/Private/LevelInstance/LevelInstanceActor.cpp
--- a/Engine/Source/Runtime/Engine/Private/LevelInstance/LevelInstanceActor.cpp
+++ b/Engine/Source/Runtime/Engine/Private/LevelInstance/LevelInstanceActor.cpp
@@ -224,6 +224,11 @@ bool ALevelInstance::CanEditChange(const FProperty* Property) const
 	return Super::CanEditChange(Property) && LevelInstanceActorImpl.CanEditChange(Property);
 }
 
+bool ALevelInstance::ResolveSubobject(const TCHAR* SubObjectPath, UObject*& OutObject, bool bLoadIfExists)
+{
+	return LevelInstanceActorImpl.ResolveSubobject(SubObjectPath, OutObject, bLoadIfExists);
+}
+
 void ALevelInstance::PostEditImport()
 {
 	LevelInstanceActorImpl.PostEditImport([this]() { Super::PostEditImport(); });
diff --git a/Engine/Source/Runtime/Engine/Private/LevelInstance/LevelInstanceActorImpl.cpp b/Engine/Source/Runtime/Engine/Private/LevelInstance/LevelInstanceActorImpl.cpp
--- a/Engine/Source/Runtime/Engine/Private/LevelInstance/LevelInstanceActorImpl.cpp
+++ b/Engine/Source/Runtime/Engine/Private/LevelInstance/LevelInstanceActorImpl.cpp
@@ -56,6 +56,16 @@ void FLevelInstanceActorImpl::UnregisterLevelInstance()
 	}
 }
 
+bool FLevelInstanceActorImpl::ResolveSubobject(const TCHAR* SubObjectPath, UObject*& OutObject, bool bLoadIfExists)
+{
+	if (ULevel* Level = LevelInstance->GetLoadedLevel())
+	{
+		return Level->ResolveSubobject(SubObjectPath, OutObject, bLoadIfExists);
+	}
+
+	return false;
+}
+
 const FLevelInstanceID& FLevelInstanceActorImpl::GetLevelInstanceID() const
 {
 	check(HasValidLevelInstanceID());
diff --git a/Engine/Source/Runtime/Engine/Private/LevelInstance/LevelInstanceEditorLevelStreaming.cpp b/Engine/Source/Runtime/Engine/Private/LevelInstance/LevelInstanceEditorLevelStreaming.cpp
--- a/Engine/Source/Runtime/Engine/Private/LevelInstance/LevelInstanceEditorLevelStreaming.cpp
+++ b/Engine/Source/Runtime/Engine/Private/LevelInstance/LevelInstanceEditorLevelStreaming.cpp
@@ -161,6 +161,8 @@ void ULevelStreamingLevelInstanceEditor::OnLevelLoadedChanged(ULevel* InLevel)
 		NewLoadedLevel->bPromptWhenAddingToLevelBeforeCheckout = false;
 		NewLoadedLevel->bPromptWhenAddingToLevelOutsideBounds = false;
 
+		NewLoadedLevel->SetEditorPathOwner(Cast<AActor>(GetLevelInstance()));
+
 		check(!NewLoadedLevel->bAlreadyMovedActors);
 		if (AWorldSettings* WorldSettings = NewLoadedLevel->GetWorldSettings())
 		{
@@ -169,13 +171,11 @@ void ULevelStreamingLevelInstanceEditor::OnLevelLoadedChanged(ULevel* InLevel)
 
 		if (ULevelInstanceSubsystem* LevelInstanceSubsystem = GetWorld()->GetSubsystem<ULevelInstanceSubsystem>())
 		{
-#if WITH_EDITOR
 			if (UWorldPartition* OuterWorldPartition = NewLoadedLevel->GetWorldPartition())
 			{
 				check(!OuterWorldPartition->IsInitialized());
 				OuterWorldPartition->bOverrideEnableStreamingInEditor = false;
 			}
-#endif
 
 			LevelInstanceSubsystem->RegisterLoadedLevelStreamingLevelInstanceEditor(this);
 		}
diff --git a/Engine/Source/Runtime/Engine/Private/LevelInstance/LevelInstanceLevelStreaming.cpp b/Engine/Source/Runtime/Engine/Private/LevelInstance/LevelInstanceLevelStreaming.cpp
--- a/Engine/Source/Runtime/Engine/Private/LevelInstance/LevelInstanceLevelStreaming.cpp
+++ b/Engine/Source/Runtime/Engine/Private/LevelInstance/LevelInstanceLevelStreaming.cpp
@@ -245,6 +245,9 @@ ULevelStreamingLevelInstance* ULevelStreamingLevelInstance::LoadInstance(ILevelI
 	Params.OptionalLevelNameOverride = &Suffix;
 	Params.OptionalLevelStreamingClass = LevelInstance->GetLevelStreamingClass();
 	Params.bLoadAsTempPackage = true;
+#if WITH_EDITOR
+	Params.EditorPathOwner = LevelInstanceActor;
+#endif
 	
 	if (World->IsGameWorld())
 	{
diff --git a/Engine/Source/Runtime/Engine/Private/LevelStreaming.cpp b/Engine/Source/Runtime/Engine/Private/LevelStreaming.cpp
--- a/Engine/Source/Runtime/Engine/Private/LevelStreaming.cpp
+++ b/Engine/Source/Runtime/Engine/Private/LevelStreaming.cpp
@@ -1365,6 +1365,9 @@ void ULevelStreaming::PrepareLoadedLevel(ULevel* InLevel, UPackage* InLevelPacka
 	UWorld* LevelOwningWorld = InLevel->OwningWorld;
 
 	InLevel->bGarbageCollectionClusteringEnabled = bGarbageCollectionClusteringEnabled;
+#if WITH_EDITOR
+	InLevel->SetEditorPathOwner(EditorPathOwner.Get());
+#endif
 
 	if (ensure(LevelOwningWorld))
 	{
@@ -2598,7 +2601,9 @@ ULevelStreamingDynamic* ULevelStreamingDynamic::LoadLevelInstance_Internal(const
     StreamingLevel->LevelTransform = Params.LevelTransform;
 	// Map to Load
     StreamingLevel->PackageNameToLoad = FName(*OnDiskPackageName);
-          
+#if WITH_EDITOR
+	StreamingLevel->EditorPathOwner = Params.EditorPathOwner;
+#endif
     // Add the new level to world.
     Params.World->AddStreamingLevel(StreamingLevel);
       
diff --git a/Engine/Source/Runtime/Engine/Private/WorldPartition/WorldPartition.cpp b/Engine/Source/Runtime/Engine/Private/WorldPartition/WorldPartition.cpp
--- a/Engine/Source/Runtime/Engine/Private/WorldPartition/WorldPartition.cpp
+++ b/Engine/Source/Runtime/Engine/Private/WorldPartition/WorldPartition.cpp
@@ -2216,6 +2216,11 @@ void UWorldPartition::RemapSoftObjectPath(FSoftObjectPath& ObjectPath) const
 	}
 }
 
+bool UWorldPartition::ConvertContainerPathToEditorPath(const FActorContainerID& InContainerID, const FSoftObjectPath& InPath, FSoftObjectPath& OutPath) const
+{
+	return StreamingPolicy ? StreamingPolicy->ConvertContainerPathToEditorPath(InContainerID, InPath, OutPath) : false;
+}
+
 FBox UWorldPartition::GetEditorWorldBounds() const
 {
 	check(EditorHash);
diff --git a/Engine/Source/Runtime/Engine/Private/WorldPartition/WorldPartitionLevelHelper.cpp b/Engine/Source/Runtime/Engine/Private/WorldPartition/WorldPartitionLevelHelper.cpp
--- a/Engine/Source/Runtime/Engine/Private/WorldPartition/WorldPartitionLevelHelper.cpp
+++ b/Engine/Source/Runtime/Engine/Private/WorldPartition/WorldPartitionLevelHelper.cpp
@@ -45,6 +45,11 @@ bool FWorldPartitionResolveData::ResolveObject(UWorld* InWorld, const FSoftObjec
 	return false;
 }
 
+FString FWorldPartitionLevelHelper::AddActorContainerID(const FActorContainerID& InContainerID, const FString& InActorName)
+{
+	return InActorName + TEXT("_") + InContainerID.ToShortString();
+}
+
 FString FWorldPartitionLevelHelper::AddActorContainerIDToSubPathString(const FActorContainerID& InContainerID, const FString& InSubPathString)
 {
 	if (!InContainerID.IsMainContainer())
@@ -56,11 +61,11 @@ FString FWorldPartitionLevelHelper::AddActorContainerIDToSubPathString(const FAc
 			const int32 SubObjectPos = InSubPathString.Find(TEXT("."), ESearchCase::IgnoreCase, ESearchDir::FromStart, DotPos);
 			if (SubObjectPos == INDEX_NONE)
 			{
-				return InSubPathString + TEXT("_") + InContainerID.ToShortString();
+				return AddActorContainerID(InContainerID, InSubPathString);
 			}
 			else
 			{
-				return InSubPathString.Mid(0, SubObjectPos) + TEXT("_") + InContainerID.ToShortString() + InSubPathString.Mid(SubObjectPos);
+				return AddActorContainerID(InContainerID, InSubPathString.Mid(0, SubObjectPos)) + InSubPathString.Mid(SubObjectPos);
 			}
 		}
 	}
@@ -551,7 +556,7 @@ bool FWorldPartitionLevelHelper::LoadActors(UWorld* InOuterWorld, ULevel* InDest
 					{
 						if (!Value.IsNull() && Value.GetAssetPathString().Equals(SourceWorldPath, ESearchCase::IgnoreCase))
 						{
-							Value = RemapActorPath(PackageObjectMapping->ContainerID, SourceOuterWorldPath, Value);
+							InOuterWorld->GetWorldPartition()->ConvertContainerPathToEditorPath(PackageObjectMapping->ContainerID, FSoftObjectPath(Value), Value);
 						}
 					});
 					FixupArchive.Fixup(Actor);
diff --git a/Engine/Source/Runtime/Engine/Private/WorldPartition/WorldPartitionLevelStreamingPolicy.cpp b/Engine/Source/Runtime/Engine/Private/WorldPartition/WorldPartitionLevelStreamingPolicy.cpp
--- a/Engine/Source/Runtime/Engine/Private/WorldPartition/WorldPartitionLevelStreamingPolicy.cpp
+++ b/Engine/Source/Runtime/Engine/Private/WorldPartition/WorldPartitionLevelStreamingPolicy.cpp
@@ -124,12 +124,60 @@ bool UWorldPartitionLevelStreamingPolicy::StoreToExternalStreamingObject(URuntim
 	if (Super::StoreToExternalStreamingObject(OutExternalStreamingObject))
 	{
 		OutExternalStreamingObject.SubObjectsToCellRemapping = MoveTemp(SubObjectsToCellRemapping);
+		OutExternalStreamingObject.ContainerResolver = MoveTemp(ContainerResolver);
 		return true;
 	}
 
 	return false;
 }
 
+bool UWorldPartitionLevelStreamingPolicy::ConvertContainerPathToEditorPath(const FActorContainerID& InContainerID, const FSoftObjectPath& InPath, FSoftObjectPath& OutPath) const
+{
+	if (ContainerResolver.IsValid())
+	{
+		FString SubObjectString;
+		FString SubObjectContext;
+		if (InPath.GetSubPathString().Split(TEXT("."), &SubObjectContext, &SubObjectString))
+		{
+			if (SubObjectContext == TEXT("PersistentLevel"))
+			{
+				const FString* FoundEditorPath = ContainerResolver.FindContainerEditorPath(InContainerID);
+				if (!FoundEditorPath)
+				{
+					for (const TWeakObjectPtr<URuntimeHashExternalStreamingObjectBase>& ExternalStreamingObject : ExternalStreamingObjects)
+					{
+						if (URuntimeHashExternalStreamingObjectBase* ExternalStreamingObjectPtr = ExternalStreamingObject.Get())
+						{
+							FoundEditorPath = ExternalStreamingObjectPtr->ContainerResolver.FindContainerEditorPath(InContainerID);
+							if (FoundEditorPath)
+							{
+								break;
+							}
+						}
+					}
+				}
+
+				if (FoundEditorPath)
+				{
+					FString SubPathString = TEXT("PersistentLevel.") + *FoundEditorPath;
+					if (!SubObjectString.IsEmpty())
+					{
+						SubPathString += TEXT(".") + SubObjectString;
+					}
+					OutPath = FSoftObjectPath(SourceWorldAssetPath, SubPathString);
+					return true;
+				}
+			}
+		}
+
+		return false;
+	}
+	
+	// Previous behavior (remap container path to source world path + container id)
+	OutPath = FWorldPartitionLevelHelper::RemapActorPath(InContainerID, SourceWorldAssetPath.ToString(), InPath);
+	return true;
+}
+
 #endif
 
 bool UWorldPartitionLevelStreamingPolicy::ConvertEditorPathToRuntimePath(const FSoftObjectPath& InPath, FSoftObjectPath& OutPath) const
@@ -176,32 +224,30 @@ bool UWorldPartitionLevelStreamingPolicy::ConvertEditorPathToRuntimePath(const F
 		}
 	}
 
-	FString SubAssetName;
-	FString SubAssetContext;
-	if (SrcObjectPath.GetSubPathString().Split(TEXT("."), &SubAssetContext, &SubAssetName))
+	FString SubObjectString;
+	FString SubObjectContext;
+	if (SrcObjectPath.GetSubPathString().Split(TEXT("."), &SubObjectContext, &SubObjectString))
 	{
-		if (SubAssetContext == TEXT("PersistentLevel"))
+		if (SubObjectContext == TEXT("PersistentLevel"))
 		{
-			FString SubObjectName;
-			FString SubObjectContext(SubAssetName);
-			SubAssetName.Split(TEXT("."), &SubObjectContext, &SubObjectName);
+			FString OutSubObjectString;
+			const FName* CellName = FindCellNameForSubObject(SubObjectString, true, OutSubObjectString);
+			FString SubPathString = SubObjectContext + TEXT(".") + OutSubObjectString;
 
-			// Try to find the corresponding streaming cell, if it doesn't exists the actor must be in the persistent level.
-			const FName* CellName = FindCellNameForSubObject(*SubObjectContext);
 			if (!CellName)
 			{
-				OutPath = FSoftObjectPath(WorldAssetPath, InPath.GetSubPathString());
+				OutPath = FSoftObjectPath(WorldAssetPath, SubPathString);
 			}
 #if WITH_EDITOR
 			else if (OuterWorld->IsGameWorld())
 			{
 				const FString PackagePath = UWorldPartitionLevelStreamingPolicy::GetCellPackagePath(*CellName, OuterWorld);
-				OutPath = FString::Printf(TEXT("%s.%s:%s"), *PackagePath, *OuterWorld->GetName(), *InPath.GetSubPathString());
+				OutPath = FString::Printf(TEXT("%s.%s:%s"), *PackagePath, *OuterWorld->GetName(), *SubPathString);
 			}
 #endif
 			else
 			{
-				OutPath = FString::Printf(TEXT("%s%s/_Generated_/%s%s.%s:%s"), *LevelInstancePrefix, *SourceWorldAssetPackageName, *(CellName->ToString()), *LevelInstanceSuffix, *WorldAssetPath.GetAssetName().ToString(), *InPath.GetSubPathString());
+				OutPath = FString::Printf(TEXT("%s%s/_Generated_/%s%s.%s:%s"), *LevelInstancePrefix, *SourceWorldAssetPackageName, *(CellName->ToString()), *LevelInstanceSuffix, *WorldAssetPath.GetAssetName().ToString(), *SubPathString);
 			}
 
 #if WITH_EDITOR
@@ -261,6 +307,13 @@ bool UWorldPartitionLevelStreamingPolicy::RemoveExternalStreamingObject(URuntime
 	return bSuccess;
 }
 
+void UWorldPartitionLevelStreamingPolicy::GetResourceSizeEx(FResourceSizeEx& CumulativeResourceSize)
+{
+	Super::GetResourceSizeEx(CumulativeResourceSize);
+	SIZE_T AllocatedSize = SubObjectsToCellRemapping.GetAllocatedSize();
+	AllocatedSize += ContainerResolver.GetAllocatedSize();
+	CumulativeResourceSize.AddDedicatedSystemMemoryBytes(AllocatedSize);
+}
 
 void UWorldPartitionLevelStreamingPolicy::DrawRuntimeCellsDetails(UCanvas* Canvas, FVector2D& Offset)
 {
@@ -329,9 +382,16 @@ void UWorldPartitionLevelStreamingPolicy::DrawRuntimeCellsDetails(UCanvas* Canva
 	Offset.Y = MaxPosY;
 }
 
-const FName* UWorldPartitionLevelStreamingPolicy::FindCellNameForSubObject(FName SubObjectName) const
+const FName* UWorldPartitionLevelStreamingPolicy::FindCellNameForSubObject(const FString& InSubObjectString, bool bInResolveContainers, FString& OutSubPathString) const
 {
-	if (const FName* CellName = SubObjectsToCellRemapping.Find(SubObjectName))
+	FString SubObjectString;
+	FString SubObjectContext(InSubObjectString);
+	InSubObjectString.Split(TEXT("."), &SubObjectContext, &SubObjectString);
+	
+	// Initialize to received value
+	OutSubPathString = InSubObjectString;
+
+	if (const FName* CellName = SubObjectsToCellRemapping.Find(*SubObjectContext))
 	{
 		return CellName;
 	}
@@ -340,13 +400,33 @@ const FName* UWorldPartitionLevelStreamingPolicy::FindCellNameForSubObject(FName
 	{
 		if (ExternalStreamingObject.IsValid())
 		{
-			if (const FName* CellName = ExternalStreamingObject.Get()->SubObjectsToCellRemapping.Find(SubObjectName))
+			if (const FName* CellName = ExternalStreamingObject.Get()->SubObjectsToCellRemapping.Find(*SubObjectContext))
 			{
 				return CellName;
 			}
 		}
 	}
 	
+	FString OutResolvedSubPathString;
+	if (bInResolveContainers)
+	{
+		if (ContainerResolver.ResolveContainerPath(InSubObjectString, OutResolvedSubPathString))
+		{
+			return FindCellNameForSubObject(OutResolvedSubPathString, false, OutSubPathString);
+		}
+
+		for (const TWeakObjectPtr<URuntimeHashExternalStreamingObjectBase>& ExternalStreamingObject : ExternalStreamingObjects)
+		{
+			if (ExternalStreamingObject.IsValid())
+			{
+				if (ExternalStreamingObject.Get()->ContainerResolver.ResolveContainerPath(InSubObjectString, OutResolvedSubPathString))
+				{
+					return FindCellNameForSubObject(OutResolvedSubPathString, false, OutSubPathString);
+				}
+			}
+		}
+	}
+			
 	return nullptr;
 }
 
diff --git a/Engine/Source/Runtime/Engine/Private/WorldPartition/WorldPartitionRuntimeContainerResolving.cpp b/Engine/Source/Runtime/Engine/Private/WorldPartition/WorldPartitionRuntimeContainerResolving.cpp
new file mode 100644
--- /dev/null
+++ b/Engine/Source/Runtime/Engine/Private/WorldPartition/WorldPartitionRuntimeContainerResolving.cpp
@@ -0,0 +1,90 @@
+// Copyright Epic Games, Inc. All Rights Reserved.
+
+#include "WorldPartition/WorldPartitionRuntimeContainerResolving.h"
+#include "WorldPartition/WorldPartitionLevelHelper.h"
+#include "WorldPartition/WorldPartitionActorContainerID.h"
+
+bool FWorldPartitionRuntimeContainerResolver::ResolveContainerPath(const FString& InSubObjectString, FString& OutSubObjectString) const
+{
+	if (!IsValid())
+	{
+		return false;
+	}
+
+	FString SubObjectString;
+	FString SubObjectContext(InSubObjectString);
+	InSubObjectString.Split(TEXT("."), &SubObjectContext, &SubObjectString);
+	
+	const FWorldPartitionRuntimeContainer* Container = &Containers.FindChecked(MainContainerPackage);
+	FActorContainerID ContainerID;
+	check(ContainerID.IsMainContainer());
+			
+	while (FWorldPartitionRuntimeContainerInstance const* ContainerInstance = Container->ContainerInstances.Find(*SubObjectContext))
+	{
+		Container = &Containers.FindChecked(ContainerInstance->ContainerPackage);
+		ContainerID = FActorContainerID(ContainerID, ContainerInstance->ActorGuid);
+		
+		SubObjectContext = SubObjectString;
+		if (!SubObjectString.Split(TEXT("."), &SubObjectContext, &SubObjectString))
+		{
+			check(SubObjectContext == SubObjectString);
+			SubObjectString.Empty();
+			break;
+		}
+	}
+
+	if (!ContainerID.IsMainContainer())
+	{
+		OutSubObjectString = FWorldPartitionLevelHelper::AddActorContainerID(ContainerID, SubObjectContext);
+		if (!SubObjectString.IsEmpty())
+		{
+			OutSubObjectString += TEXT(".") + SubObjectString;
+		}
+
+		return true;
+	}
+
+	return false;
+}
+
+SIZE_T FWorldPartitionRuntimeContainer::GetAllocatedSize() const
+{
+	return ContainerInstances.GetAllocatedSize();
+}
+
+SIZE_T FWorldPartitionRuntimeContainerResolver::GetAllocatedSize() const
+{
+	SIZE_T AllocatedSize = Containers.GetAllocatedSize();
+	for (auto& [Name,Container] : Containers)
+	{
+		AllocatedSize += Container.GetAllocatedSize();
+	}
+
+	return AllocatedSize;
+}
+
+#if WITH_EDITOR
+
+void FWorldPartitionRuntimeContainerResolver::BuildContainerIDToEditorPathMap()
+{
+	check(!MainContainerPackage.IsNone());
+	check(ContainerIDToEditorPath.IsEmpty());
+
+	TFunction<void(const FActorContainerID&, FName, const FString&)> BuildContainerIDToEditorMapRecursive =
+		[this, &BuildContainerIDToEditorMapRecursive](const FActorContainerID& InParentContainerID, FName InContainerPackage, const FString& InPath)
+	{
+		const FWorldPartitionRuntimeContainer& Container = Containers.FindChecked(InContainerPackage);
+		for (auto [Name, ContainerInstance] : Container.ContainerInstances)
+		{
+			FActorContainerID ContainerInstanceID(InParentContainerID, ContainerInstance.ActorGuid);
+			FString ContainerInstanceEditorPath = InPath.IsEmpty() ? Name.ToString() : InPath + TEXT(".") + Name.ToString();
+			ContainerIDToEditorPath.Add(ContainerInstanceID, ContainerInstanceEditorPath);
+
+			BuildContainerIDToEditorMapRecursive(ContainerInstanceID, ContainerInstance.ContainerPackage, ContainerInstanceEditorPath);
+		}
+	};
+
+	BuildContainerIDToEditorMapRecursive(FActorContainerID(), MainContainerPackage, FString());
+}
+
+#endif
\ No newline at end of file
diff --git a/Engine/Source/Runtime/Engine/Private/WorldPartition/WorldPartitionStreamingGeneration.cpp b/Engine/Source/Runtime/Engine/Private/WorldPartition/WorldPartitionStreamingGeneration.cpp
--- a/Engine/Source/Runtime/Engine/Private/WorldPartition/WorldPartitionStreamingGeneration.cpp
+++ b/Engine/Source/Runtime/Engine/Private/WorldPartition/WorldPartitionStreamingGeneration.cpp
@@ -30,7 +30,9 @@
 #include "WorldPartition/HLOD/HLODActor.h"
 #include "WorldPartition/HLOD/HLODLayer.h"
 #include "WorldPartition/WorldPartitionSubsystem.h"
+#include "WorldPartition/WorldPartitionRuntimeContainerResolving.h"
 #include "HAL/FileManager.h"
+#include "Misc/EditorPathHelper.h"
 
 #define LOCTEXT_NAMESPACE "WorldPartition"
 
@@ -93,6 +95,12 @@ class FWorldPartitionStreamingGenerator
 			// Create the dataset required for IStreamingGenerationContext interface
 			MainWorldActorSetContainerIndex = INDEX_NONE;
 			ActorSetContainers.Empty(StreamingGenerator->ContainerCollectionDescriptorsMap.Num());
+			
+			const bool bBuildContainerResolver = FEditorPathHelper::IsEnabled();
+			if (bBuildContainerResolver)
+			{
+				ContainerResolver.SetMainContainerPackage(TopLevelActorDescCollection.GetMainContainerPackageName());
+			}
 
 			TMap<TWeakPtr<FStreamingGenerationActorDescCollection>, int32> ActorSetContainerMap;
 			for (const auto& [LevelName, ContainerDescriptor] : StreamingGenerator->ContainerCollectionDescriptorsMap)
@@ -116,6 +124,20 @@ class FWorldPartitionStreamingGenerator
 					check(MainWorldActorSetContainerIndex == INDEX_NONE);
 					MainWorldActorSetContainerIndex = ContainerIndex;
 				}
+
+				if (bBuildContainerResolver)
+				{
+					auto& Container = ContainerResolver.AddContainer(ContainerDescriptor.ActorDescCollection->GetMainContainerPackageName());
+					for (const FWorldPartitionActorDescView& ActorDescView : ContainerDescriptor.ContainerCollectionInstanceViews)
+					{
+						Container.AddContainerInstance(ActorDescView.GetActorName(), { ActorDescView.GetGuid(), ActorDescView.GetContainerPackage() });
+					}
+				}
+			}
+
+			if (bBuildContainerResolver)
+			{
+				ContainerResolver.BuildContainerIDToEditorPathMap();
 			}
 
 			ActorSetInstances.Empty();
@@ -203,6 +225,8 @@ class FWorldPartitionStreamingGenerator
 				Func(ActorSetContainer);
 			}
 		}
+
+		virtual const FWorldPartitionRuntimeContainerResolver& GetContainerResolver() const override { return ContainerResolver; }
 		//~End IStreamingGenerationContext interface};
 
 	private:
@@ -210,6 +234,7 @@ class FWorldPartitionStreamingGenerator
 		int32 MainWorldActorSetContainerIndex;
 		TArray<FActorSetContainer> ActorSetContainers;
 		TArray<FActorSetInstance> ActorSetInstances;
+		FWorldPartitionRuntimeContainerResolver ContainerResolver;
 	};
 
 	/** 
@@ -1356,13 +1381,16 @@ bool UWorldPartition::GenerateContainerStreaming(const FGenerateStreamingParams&
 	StreamingPolicy = NewObject<UWorldPartitionStreamingPolicy>(const_cast<UWorldPartition*>(this), WorldPartitionStreamingPolicyClass.Get(), NAME_None, bIsPIE ? RF_Transient : RF_NoFlags);
 
 	check(RuntimeHash);
-	if (RuntimeHash->GenerateStreaming(StreamingPolicy, StreamingGenerator.GetStreamingGenerationContext(InParams.ActorDescCollection), InContext.PackagesToGenerate))
+	const IStreamingGenerationContext* StreamingGenerationContext = StreamingGenerator.GetStreamingGenerationContext(InParams.ActorDescCollection);
+	check(StreamingGenerationContext);
+	if (RuntimeHash->GenerateStreaming(StreamingPolicy, StreamingGenerationContext, InContext.PackagesToGenerate))
 	{
 		if (HierarchicalLogAr.IsValid())
 		{
 			RuntimeHash->DumpStateLog(*HierarchicalLogAr);
 		}
 
+		StreamingPolicy->SetContainerResolver(StreamingGenerationContext->GetContainerResolver());
 		StreamingPolicy->PrepareActorToCellRemapping();
 		return true;
 	}
diff --git a/Engine/Source/Runtime/Engine/Public/LevelInstance/LevelInstanceActor.h b/Engine/Source/Runtime/Engine/Public/LevelInstance/LevelInstanceActor.h
--- a/Engine/Source/Runtime/Engine/Public/LevelInstance/LevelInstanceActor.h
+++ b/Engine/Source/Runtime/Engine/Public/LevelInstance/LevelInstanceActor.h
@@ -82,6 +82,7 @@ public:
 	ENGINE_API virtual void PreEditChange(FProperty* PropertyThatWillChange) override;
 	ENGINE_API virtual void PostEditChangeProperty(FPropertyChangedEvent& PropertyChangedEvent) override;
 	ENGINE_API virtual void PostEditImport() override;
+	ENGINE_API virtual bool ResolveSubobject(const TCHAR* SubObjectPath, UObject*& OutObject, bool bLoadIfExists) override;
 
 	// AActor overrides
 	ENGINE_API virtual void CheckForErrors() override;
diff --git a/Engine/Source/Runtime/Engine/Public/LevelInstance/LevelInstanceActorImpl.h b/Engine/Source/Runtime/Engine/Public/LevelInstance/LevelInstanceActorImpl.h
--- a/Engine/Source/Runtime/Engine/Public/LevelInstance/LevelInstanceActorImpl.h
+++ b/Engine/Source/Runtime/Engine/Public/LevelInstance/LevelInstanceActorImpl.h
@@ -53,6 +53,7 @@ public:
 
 	ENGINE_API virtual void RegisterLevelInstance();
 	ENGINE_API virtual void UnregisterLevelInstance();
+	ENGINE_API virtual bool ResolveSubobject(const TCHAR* SubObjectPath, UObject*& OutObject, bool bLoadIfExists);
 
 	/**
 	 * Begin ILevelInstanceInterface Implementaion 
diff --git a/Engine/Source/Runtime/Engine/Public/WorldPartition/WorldPartition.h b/Engine/Source/Runtime/Engine/Public/WorldPartition/WorldPartition.h
--- a/Engine/Source/Runtime/Engine/Public/WorldPartition/WorldPartition.h
+++ b/Engine/Source/Runtime/Engine/Public/WorldPartition/WorldPartition.h
@@ -219,6 +219,8 @@ public:
 	ENGINE_API bool ConvertEditorPathToRuntimePath(const FSoftObjectPath& InPath, FSoftObjectPath& OutPath) const;
 
 #if WITH_EDITOR
+	ENGINE_API bool ConvertContainerPathToEditorPath(const FActorContainerID& InContainerID, const FSoftObjectPath& InPath, FSoftObjectPath& OutPath) const;
+
 	void SetInstanceTransform(const FTransform& InInstanceTransform) { InstanceTransform = InInstanceTransform; }
 	ENGINE_API FName GetWorldPartitionEditorName() const;
 
diff --git a/Engine/Source/Runtime/Engine/Public/WorldPartition/WorldPartitionLevelHelper.h b/Engine/Source/Runtime/Engine/Public/WorldPartition/WorldPartitionLevelHelper.h
--- a/Engine/Source/Runtime/Engine/Public/WorldPartition/WorldPartitionLevelHelper.h
+++ b/Engine/Source/Runtime/Engine/Public/WorldPartition/WorldPartitionLevelHelper.h
@@ -20,6 +20,7 @@ class FWorldPartitionLevelHelper
 {
 public:
 	static FString AddActorContainerIDToSubPathString(const FActorContainerID& InContainerID, const FString& InSubPathString);
+	static FString AddActorContainerID(const FActorContainerID& InContainerID, const FString& InActorName);
 #if WITH_EDITOR
 public:
 	static FWorldPartitionLevelHelper& Get();
diff --git a/Engine/Source/Runtime/Engine/Public/WorldPartition/WorldPartitionLevelStreamingPolicy.h b/Engine/Source/Runtime/Engine/Public/WorldPartition/WorldPartitionLevelStreamingPolicy.h
--- a/Engine/Source/Runtime/Engine/Public/WorldPartition/WorldPartitionLevelStreamingPolicy.h
+++ b/Engine/Source/Runtime/Engine/Public/WorldPartition/WorldPartitionLevelStreamingPolicy.h
@@ -12,8 +12,10 @@
 
 #include "CoreMinimal.h"
 #include "WorldPartition/WorldPartitionStreamingPolicy.h"
+#include "WorldPartition/WorldPartitionRuntimeContainerResolving.h"
 #include "WorldPartitionLevelStreamingPolicy.generated.h"
 
+struct FResourceSizeEx;
 class UWorldPartitionRuntimeLevelStreamingCell;
 enum class EWorldPartitionRuntimeCellState : uint8;
 
@@ -29,10 +31,12 @@ public:
 #if WITH_EDITOR
 	virtual TSubclassOf<class UWorldPartitionRuntimeCell> GetRuntimeCellClass() const override;
 	virtual void PrepareActorToCellRemapping() override;
+	virtual void SetContainerResolver(const FWorldPartitionRuntimeContainerResolver& InContainerResolver) override { ContainerResolver = InContainerResolver; }
 	virtual void RemapSoftObjectPath(FSoftObjectPath& ObjectPath) const override;
 	static FString GetCellPackagePath(const FName& InCellName, const UWorld* InWorld);
 
 	virtual bool StoreToExternalStreamingObject(URuntimeHashExternalStreamingObjectBase& OutExternalStreamingObject) override;
+	virtual bool ConvertContainerPathToEditorPath(const FActorContainerID& InContainerID, const FSoftObjectPath& InPath, FSoftObjectPath& OutPath) const;
 #endif
 
 	virtual bool ConvertEditorPathToRuntimePath(const FSoftObjectPath& InPath, FSoftObjectPath& OutPath) const override;
@@ -41,11 +45,12 @@ public:
 	virtual bool InjectExternalStreamingObject(URuntimeHashExternalStreamingObjectBase* ExternalStreamingObject) override;
 	virtual bool RemoveExternalStreamingObject(URuntimeHashExternalStreamingObjectBase* ExternalStreamingObject) override;
 
+	virtual void GetResourceSizeEx(FResourceSizeEx& CumulativeResourceSize) override;
 protected:
 	void ForEachActiveRuntimeCell(TFunctionRef<void(const UWorldPartitionRuntimeCell*)> Func) const;
 
 private:
-	const FName* FindCellNameForSubObject(FName SubObjectName) const;
+	const FName* FindCellNameForSubObject(const FString& InSubObjectString, bool bInResolveContainers, FString& OutSubObjectString) const;
 	const UWorldPartitionRuntimeLevelStreamingCell* FindCellForSubObject(FName SubObjectName) const;
 
 	UPROPERTY()
@@ -54,6 +59,9 @@ private:
 	UPROPERTY()
 	TMap<FName, FName> SubObjectsToCellRemapping;
 
+	UPROPERTY()
+	FWorldPartitionRuntimeContainerResolver ContainerResolver;
+
 	UPROPERTY(Transient)
 	TArray<TWeakObjectPtr<URuntimeHashExternalStreamingObjectBase>> ExternalStreamingObjects;
 };
\ No newline at end of file
diff --git a/Engine/Source/Runtime/Engine/Public/WorldPartition/WorldPartitionRuntimeContainerResolving.h b/Engine/Source/Runtime/Engine/Public/WorldPartition/WorldPartitionRuntimeContainerResolving.h
new file mode 100644
--- /dev/null
+++ b/Engine/Source/Runtime/Engine/Public/WorldPartition/WorldPartitionRuntimeContainerResolving.h
@@ -0,0 +1,100 @@
+// Copyright Epic Games, Inc. All Rights Reserved.
+
+#pragma once
+
+#include "CoreFwd.h"
+#include "Misc/Guid.h"
+#include "Containers/Map.h"
+#include "UObject/NameTypes.h"
+#include "UObject/ObjectMacros.h"
+#include "WorldPartition/WorldPartitionActorContainerID.h"
+#include "WorldPartitionRuntimeContainerResolving.generated.h"
+
+USTRUCT()
+struct FWorldPartitionRuntimeContainerInstance
+{
+	GENERATED_USTRUCT_BODY()
+public:
+	FWorldPartitionRuntimeContainerInstance()
+	{}
+
+	FWorldPartitionRuntimeContainerInstance(const FGuid& InActorGuid, FName InContainerPackage)
+		: ActorGuid(InActorGuid), ContainerPackage(InContainerPackage)
+	{}
+
+private:
+	UPROPERTY()
+	FGuid ActorGuid;
+
+	UPROPERTY()
+	FName ContainerPackage;
+
+	friend struct FWorldPartitionRuntimeContainerResolver;
+};
+
+USTRUCT()
+struct FWorldPartitionRuntimeContainer
+{
+	GENERATED_USTRUCT_BODY()
+public:
+#if WITH_EDITOR
+	void AddContainerInstance(FName InActorName, const FWorldPartitionRuntimeContainerInstance& InContainerInstance)
+	{
+		check(!ContainerInstances.Contains(InActorName));
+		ContainerInstances.Add(InActorName, InContainerInstance);
+	}
+#endif
+	SIZE_T GetAllocatedSize() const;
+private:
+	UPROPERTY()
+	TMap<FName, FWorldPartitionRuntimeContainerInstance> ContainerInstances;
+
+	friend struct FWorldPartitionRuntimeContainerResolver;
+};
+
+/**
+ * FWorldPartitionRuntimeContainerResolver
+ *
+ * Helper class that allows resolving a hierarchy of Container actor names (Editor Path) to a resolvable persistent level path (Runtime Path)
+ * 
+ * Editor Path : /Game/Map.Map:PersistentLevel.LevelInstance1.LevelInstance2.StaticMeshActor 
+ * Runtime Path : /Game/Map.Map:PersistentLevel.StaticMeshActor_{ContainerID}
+ */
+USTRUCT()
+struct FWorldPartitionRuntimeContainerResolver
+{
+	GENERATED_USTRUCT_BODY()
+public:
+	bool IsValid() const { return !MainContainerPackage.IsNone(); }
+	bool ResolveContainerPath(const FString& InSubObjectString, FString& OutSubObjectString) const;
+	SIZE_T GetAllocatedSize() const;
+	
+#if WITH_EDITOR
+	void SetMainContainerPackage(FName InContainerPackage)
+	{
+		check(MainContainerPackage.IsNone());
+		MainContainerPackage = InContainerPackage;
+	}
+
+	FWorldPartitionRuntimeContainer& AddContainer(FName InContainerPackage)
+	{
+		check(!Containers.Contains(InContainerPackage));
+		return Containers.Add(InContainerPackage);	
+	}
+
+	void BuildContainerIDToEditorPathMap();
+	const FString* FindContainerEditorPath(const FActorContainerID& InContainerID) const { return ContainerIDToEditorPath.Find(InContainerID); }
+#endif
+private:
+	UPROPERTY()
+	TMap<FName, FWorldPartitionRuntimeContainer> Containers;
+	UPROPERTY()
+	FName MainContainerPackage;
+
+#if WITH_EDITORONLY_DATA
+	// Reverse lookup :
+	// ex: {ContainerID}, "LevelInstance1.LevelInstance2"
+	UPROPERTY()
+	TMap<FActorContainerID, FString> ContainerIDToEditorPath;
+#endif
+};
\ No newline at end of file
diff --git a/Engine/Source/Runtime/Engine/Public/WorldPartition/WorldPartitionRuntimeHash.h b/Engine/Source/Runtime/Engine/Public/WorldPartition/WorldPartitionRuntimeHash.h
--- a/Engine/Source/Runtime/Engine/Public/WorldPartition/WorldPartitionRuntimeHash.h
+++ b/Engine/Source/Runtime/Engine/Public/WorldPartition/WorldPartitionRuntimeHash.h
@@ -12,6 +12,7 @@
 #include "WorldPartition/WorldPartitionRuntimeCell.h"
 #include "WorldPartition/WorldPartitionStreamingGeneration.h"
 #include "WorldPartition/WorldPartitionStreamingGenerationContext.h"
+#include "WorldPartition/WorldPartitionRuntimeContainerResolving.h"
 #if WITH_EDITOR
 #include "CookPackageSplitter.h"
 #include "Misc/HierarchicalLogArchive.h"
@@ -58,6 +59,9 @@ public:
 	UPROPERTY();
 	TMap<FName, FName> SubObjectsToCellRemapping;
 
+	UPROPERTY()
+	FWorldPartitionRuntimeContainerResolver ContainerResolver;
+
 protected:
 	UPROPERTY();
 	TSoftObjectPtr<UWorld> OwningWorld;
diff --git a/Engine/Source/Runtime/Engine/Public/WorldPartition/WorldPartitionStreamingGenerationContext.h b/Engine/Source/Runtime/Engine/Public/WorldPartition/WorldPartitionStreamingGenerationContext.h
--- a/Engine/Source/Runtime/Engine/Public/WorldPartition/WorldPartitionStreamingGenerationContext.h
+++ b/Engine/Source/Runtime/Engine/Public/WorldPartition/WorldPartitionStreamingGenerationContext.h
@@ -10,6 +10,7 @@ class FStreamingGenerationActorDescCollection;
 class FWorldPartitionActorDescView;
 class UActorDescContainer;
 class UDataLayerInstance;
+struct FWorldPartitionRuntimeContainerResolver;
 
 #if WITH_EDITOR
 class IStreamingGenerationContext
@@ -88,5 +89,6 @@ public:
 	virtual const FActorSetContainer* GetMainWorldContainer() const = 0;
 	virtual void ForEachActorSetInstance(TFunctionRef<void(const FActorSetInstance&)> Func) const = 0;
 	virtual void ForEachActorSetContainer(TFunctionRef<void(const FActorSetContainer&)> Func) const = 0;
+	virtual const FWorldPartitionRuntimeContainerResolver& GetContainerResolver() const = 0;
 };
 #endif
\ No newline at end of file
diff --git a/Engine/Source/Runtime/Engine/Public/WorldPartition/WorldPartitionStreamingPolicy.h b/Engine/Source/Runtime/Engine/Public/WorldPartition/WorldPartitionStreamingPolicy.h
--- a/Engine/Source/Runtime/Engine/Public/WorldPartition/WorldPartitionStreamingPolicy.h
+++ b/Engine/Source/Runtime/Engine/Public/WorldPartition/WorldPartitionStreamingPolicy.h
@@ -14,6 +14,7 @@
 #include "WorldPartition/WorldPartitionRuntimeCell.h"
 #include "WorldPartition/WorldPartitionRuntimeHash.h"
 #include "WorldPartition/WorldPartitionStreamingSource.h"
+#include "WorldPartition/WorldPartitionRuntimeContainerResolving.h"
 #include "WorldPartitionStreamingPolicy.generated.h"
 
 class UWorldPartition;
@@ -68,9 +69,11 @@ public:
 
 	// PIE/Game methods
 	virtual void PrepareActorToCellRemapping() {}
+	virtual void SetContainerResolver(const FWorldPartitionRuntimeContainerResolver& InContainerResolver) {}
 	virtual void RemapSoftObjectPath(FSoftObjectPath& ObjectPath) const {}
 
 	virtual bool StoreToExternalStreamingObject(URuntimeHashExternalStreamingObjectBase& OutExternalStreamingObject) { return true; }
+	virtual bool ConvertContainerPathToEditorPath(const FActorContainerID& InContainerID, const FSoftObjectPath& InPath, FSoftObjectPath& OutPath) const { return false; }
 #endif
 
 #if !UE_BUILD_SHIPPING
