NetDriver
* Added  engine crash context "ReplicationDriver"
* Set to Generic, Iris or the class name of the replication driver

diff --git a/Engine/Source/Runtime/Engine/Private/NetDriver.cpp b/Engine/Source/Runtime/Engine/Private/NetDriver.cpp
--- a/Engine/Source/Runtime/Engine/Private/NetDriver.cpp
+++ b/Engine/Source/Runtime/Engine/Private/NetDriver.cpp
@@ -1489,6 +1489,27 @@ bool UNetDriver::InitBase(bool bInitAsClient, FNetworkNotify* InNotify, const FU
 		CreateReplicationSystem(bInitAsClient);
 	}
 #endif // UE_WITH_IRIS
+
+
+	if (NetDriverName == NAME_GameNetDriver)
+	{
+		static FString CrashContext_ReplicationDriver = TEXT("ReplicationDriver");
+
+		if (UReplicationDriver* RepDriver = GetReplicationDriver())
+		{
+			FGenericCrashContext::SetEngineData(CrashContext_ReplicationDriver, RepDriver->GetClass()->GetName());
+		}
+#if UE_WITH_IRIS
+		else if (ReplicationSystem)
+		{
+			FGenericCrashContext::SetEngineData(CrashContext_ReplicationDriver, TEXT("Iris"));
+		}
+#endif
+		else
+		{
+			FGenericCrashContext::SetEngineData(CrashContext_ReplicationDriver, TEXT("Generic"));
+		}
+	}
 	
 	InitNetTraceId();
 
