[NetPhysics] Predictive Interpolation improvements.

- Smooth Predictive Interpolation rotation.
- Extrapolate target by the overshoot extrapolation from previous target.
- Fix bug with rotation check for early out.
- Clean up the return and sleep flow in Predictive Interpolation.

diff --git a/Engine/Source/Runtime/Engine/Private/PhysicsEngine/PhysicsReplication.cpp b/Engine/Source/Runtime/Engine/Private/PhysicsEngine/PhysicsReplication.cpp
--- a/Engine/Source/Runtime/Engine/Private/PhysicsEngine/PhysicsReplication.cpp
+++ b/Engine/Source/Runtime/Engine/Private/PhysicsEngine/PhysicsReplication.cpp
@@ -107,6 +107,12 @@ namespace PhysicsReplicationCVars
 		static float PosCorrectionTimeMultiplier = 1.0f;
 		static FAutoConsoleVariableRef CVarPosCorrectionTimeMultiplier(TEXT("np2.PredictiveInterpolation.PosCorrectionTimeMultiplier"), PosCorrectionTimeMultiplier, TEXT("Multiplier to adjust how much of RTT (network Round Trip Time) to add to positional offset correction."));
 
+		static float RotCorrectionTimeBase = 0.15f;
+		static FAutoConsoleVariableRef CVarRotCorrectionTimeBase(TEXT("np2.PredictiveInterpolation.RotCorrectionTimeBase"), RotCorrectionTimeBase, TEXT("Base time to correct positional offset over. RTT * PosCorrectionTimeMultiplier are added on top of this."));
+
+		static float RotCorrectionTimeMultiplier = 1.0f;
+		static FAutoConsoleVariableRef CVarRotCorrectionTimeMultiplier(TEXT("np2.PredictiveInterpolation.RotCorrectionTimeMultiplier"), RotCorrectionTimeMultiplier, TEXT("Multiplier to adjust how much of RTT (network Round Trip Time) to add to positional offset correction."));
+
 		static float InterpolationTimeMultiplier = 1.1f;
 		static FAutoConsoleVariableRef CVarInterpolationTimeMultiplier(TEXT("np2.PredictiveInterpolation.InterpolationTimeMultiplier"), InterpolationTimeMultiplier, TEXT("Multiplier to adjust the replication interpolation time which is based on the sendrate of replication data from the server."));
 		
@@ -825,6 +831,9 @@ void FPhysicsReplicationAsync::UpdateAsyncTarget(const FPhysicsRepAsyncInputData
 
 	if (Input.ServerFrame >= Target->ServerFrame)
 	{
+		const int32 PrevTickCount = Target->TickCount;
+		const int32 PrevReceiveInterval = Target->ReceiveInterval;
+
 		Target->PrevServerFrame = (Target->ServerFrame == INDEX_NONE) ? (Input.ServerFrame - 1) : Target->ServerFrame;
 		Target->ServerFrame = Input.ServerFrame;
 		Target->PrevReceiveFrame = (Target->ReceiveFrame == INDEX_NONE) ? (RigidsSolver->GetCurrentFrame() - 1) : Target->ReceiveFrame;
@@ -841,6 +850,13 @@ void FPhysicsReplicationAsync::UpdateAsyncTarget(const FPhysicsRepAsyncInputData
 			// Cache the position we received this target at, Predictive Interpolation will alter the target state but use this as the source position for reconciliation.
 			Target->PrevPosTarget = Input.TargetState.Position;
 			Target->PrevRotTarget = Input.TargetState.Quaternion;
+
+			// If we extrapolated the previous target past the receive interval, extrapolate this target by the overshoot
+			const int32 OvershootingExtrapolation = FMath::Clamp((PrevTickCount - PrevReceiveInterval), -1, FMath::CeilToInt(Target->AverageReceiveInterval));
+			if (OvershootingExtrapolation > 0)
+			{
+				ExtrapolateTarget(*Target, OvershootingExtrapolation, GetDeltaTime_Internal());
+			}
 		}
 	}
 
@@ -1217,13 +1233,28 @@ bool FPhysicsReplicationAsync::PredictiveInterpolation(Chaos::FPBDRigidParticleH
 		return true;
 	}
 
+	const bool bCanSimulate = Handle->IsDynamic() || Handle->IsSleeping();
+
+	// Helper for sleep and target clearing at replication end
+	auto EndReplicationHelper = [RigidsSolver, Handle, bCanSimulate](FReplicatedPhysicsTargetAsync& Target, bool bAllowSleep) -> bool
+	{
+		const bool bShouldSleep = (Target.TargetState.Flags & ERigidBodyFlags::Sleeping) != 0;
+
+		if (bAllowSleep && bShouldSleep && bCanSimulate)
+		{
+			RigidsSolver->GetEvolution()->SetParticleObjectState(Handle, Chaos::EObjectStateType::Sleeping);
+		}
+
+		const bool bClearTarget = (!bCanSimulate || bShouldSleep) && !PhysicsReplicationCVars::PredictiveInterpolationCVars::bDontClearTarget;
+		return bClearTarget;
+	};
+
 	// Get the distance from the current position to the source position of our target state
 	const float SourceDistanceSqr = (Target.PrevPosTarget - Handle->X()).SizeSquared();
-	const bool bShouldSleep = (Target.TargetState.Flags & ERigidBodyFlags::Sleeping) != 0;
-	const bool bCanSimulate = Handle->IsDynamic() || Handle->IsSleeping();
+	const bool bXCanEarlyOut = SourceDistanceSqr < PhysicsReplicationCVars::PredictiveInterpolationCVars::EarlyOutDistanceSqr;
 
 	// Early out if we are within range of target, also apply target sleep state
-	if (SourceDistanceSqr < PhysicsReplicationCVars::PredictiveInterpolationCVars::EarlyOutDistanceSqr)
+	if (bXCanEarlyOut)
 	{
 		// Get the rotational offset between the blended rotation target and the current rotation
 		const FQuat TargetRotDelta = Target.TargetState.Quaternion * Handle->R().Inverse();
@@ -1232,25 +1263,18 @@ bool FPhysicsReplicationAsync::PredictiveInterpolation(Chaos::FPBDRigidParticleH
 		float Angle;
 		FVector Axis;
 		TargetRotDelta.ToAxisAndAngle(Axis, Angle);
+		Angle = FMath::Abs(FMath::UnwindRadians(Angle));
 		if (Angle < FMath::DegreesToRadians(PhysicsReplicationCVars::PredictiveInterpolationCVars::EarlyOutAngle))
 		{
-			Handle->SetV(FVector(0, 0, 0));
-			Handle->SetW(FVector(0, 0, 0));
-			Target.PrevLinVel = FVector(FVector(0, 0, 0));
-
-			if (bShouldSleep && bCanSimulate)
-			{
-				RigidsSolver->GetEvolution()->SetParticleObjectState(Handle, Chaos::EObjectStateType::Sleeping);
-			}
-
-			const bool bClearTarget = (!bCanSimulate || bShouldSleep) && !PhysicsReplicationCVars::PredictiveInterpolationCVars::bDontClearTarget;
-			return bClearTarget;
+			// Early Out
+			return EndReplicationHelper(Target, true);
 		}
 	}
 
 	// Calculate position correction time based on current Round Trip Time
 	const float RTT = LatencyOneWay * 2.f;
-	const float PosCorrectionTime = PhysicsReplicationCVars::PredictiveInterpolationCVars::PosCorrectionTimeBase + RTT * PhysicsReplicationCVars::PredictiveInterpolationCVars::PosCorrectionTimeMultiplier;
+	const float PosCorrectionTime = FMath::Max(PhysicsReplicationCVars::PredictiveInterpolationCVars::PosCorrectionTimeBase + RTT * PhysicsReplicationCVars::PredictiveInterpolationCVars::PosCorrectionTimeMultiplier, DeltaSeconds);
+	const float RotCorrectionTime = FMath::Max(PhysicsReplicationCVars::PredictiveInterpolationCVars::RotCorrectionTimeBase + RTT * PhysicsReplicationCVars::PredictiveInterpolationCVars::RotCorrectionTimeMultiplier, DeltaSeconds);
 
 	// Calculate interpolation time based on current average receive rate of targets from the server (receive rate = send-rate from server with network conditions taken into account)
 	Target.AverageReceiveInterval = FMath::Lerp(Target.AverageReceiveInterval, Target.ReceiveInterval, FMath::Clamp((1.0f / (Target.ReceiveInterval * PhysicsReplicationCVars::PredictiveInterpolationCVars::AverageReceiveIntervalSmoothing)), 0.0f, 1.0f));
@@ -1264,10 +1288,10 @@ bool FPhysicsReplicationAsync::PredictiveInterpolation(Chaos::FPBDRigidParticleH
 	CurrentState.AngVel = Handle->W(); // Note: Current angular velocity is in Radians
 
 	// NewState
-	const FVector TargetPos = Target.TargetState.Position;
+	const FVector TargetPos = FVector(Target.TargetState.Position);
 	const FQuat TargetRot = Target.TargetState.Quaternion;
-	const FVector TargetLinVel = Target.TargetState.LinVel;
-	const FVector TargetAngVel = Target.TargetState.AngVel; // Note: Target angular velocity is in Degrees
+	const FVector TargetLinVel = FVector(Target.TargetState.LinVel);
+	const FVector TargetAngVel = FVector(Target.TargetState.AngVel); // Note: Target angular velocity is in Degrees
 
 	/** --- Reconciliation ---
 	* If target velocities are low enough, check the traveled direction and distance from previous frame and compare with replicated linear velocity.
@@ -1323,69 +1347,70 @@ bool FPhysicsReplicationAsync::PredictiveInterpolation(Chaos::FPBDRigidParticleH
 
 		// Cache data for next replication
 		Target.PrevLinVel = FVector(Target.TargetState.LinVel);
+
+		// End replication and go to sleep if that's requested
+		return EndReplicationHelper(Target, true);
 	}
 	else // Velocity-based Replication
 	{
-		// --- Velocity Replication ---
-		// Get PosDiff
-		const FVector PosDiff = TargetPos - CurrentState.Position;
-
-		// Convert PosDiff to a velocity
-		const FVector PosDiffVelocity = PosDiff / PosCorrectionTime;
-
-		// Get LinVelDiff by adding inverted CurrentState.LinVel to TargetLinVel
-		const FVector LinVelDiff = -CurrentState.LinVel + TargetLinVel;
-
-		// Add PosDiffVelocity to LinVelDiff to get BlendedTargetVelocity
-		const FVector BlendedTargetVelocity = LinVelDiff + PosDiffVelocity;
+		if (!bXCanEarlyOut)
+		{	// --- Velocity Replication ---
+			
+			// Get PosDiff
+			const FVector PosDiff = TargetPos - CurrentState.Position;
 
-		// Multiply BlendedTargetVelocity with(deltaTime / interpolationTime), clamp to 1 and add to CurrentState.LinVel to get BlendedTargetVelocityInterpolated
-		const float BlendStepAmount = FMath::Clamp(DeltaSeconds / InterpolationTime, 0.f, 1.f);
-		const FVector RepLinVel = CurrentState.LinVel + (BlendedTargetVelocity * BlendStepAmount);
+			// Convert PosDiff to a velocity
+			const FVector PosDiffVelocity = PosDiff / PosCorrectionTime;
 
+			// Get LinVelDiff by adding inverted CurrentState.LinVel to TargetLinVel
+			const FVector LinVelDiff = -CurrentState.LinVel + TargetLinVel;
 
-		// --- Angular Velocity Replication ---
-		// Extrapolate current rotation along current angular velocity to see where we would end up
-		float CurAngVelSize;
-		FVector CurAngVelAxis;
-		CurrentState.AngVel.FVector::ToDirectionAndLength(CurAngVelAxis, CurAngVelSize);
-		const FQuat CurRotExtrapDelta = FQuat(CurAngVelAxis, CurAngVelSize * DeltaSeconds);
-		const FQuat CurRotExtrap = CurRotExtrapDelta * CurrentState.Quaternion;
+			// Add PosDiffVelocity to LinVelDiff to get BlendedTargetVelocity
+			const FVector BlendedTargetVelocity = LinVelDiff + PosDiffVelocity;
 
-		// Slerp from the extrapolated current rotation towards the target rotation
-		// This takes current angular velocity into account
-		const FQuat TargetRotBlended = FQuat::Slerp(CurRotExtrap, TargetRot, BlendStepAmount);
+			// Multiply BlendedTargetVelocity with(deltaTime / interpolationTime), clamp to 1 and add to CurrentState.LinVel to get BlendedTargetVelocityInterpolated
+			const float BlendStepAmount = FMath::Clamp(DeltaSeconds / InterpolationTime, 0.f, 1.f);
+			const FVector RepLinVel = CurrentState.LinVel + (BlendedTargetVelocity * BlendStepAmount);
+			
+			Handle->SetV(RepLinVel);
 
-		// Get the rotational offset between the blended rotation target and the current rotation
-		const FQuat TargetRotDelta = TargetRotBlended * CurrentState.Quaternion.Inverse();
-
-		// Convert the rotational delta to angular velocity with a magnitude that will make it arrive at the rotation after DeltaTime has passed
-		float WAngle;
-		FVector WAxis;
-		TargetRotDelta.ToAxisAndAngle(WAxis, WAngle);
-		const FVector TargetRotDeltaBlend = FVector(WAxis * (WAngle / DeltaSeconds));
-		const FVector RepAngVel = FMath::DegreesToRadians(Target.TargetState.AngVel) + TargetRotDeltaBlend;
+			// Cache data for next replication
+			Target.PrevLinVel = FVector(RepLinVel);
+		}
 
-		Handle->SetV(RepLinVel);
-		Handle->SetW(RepAngVel);
+		{	// --- Angular Velocity Replication ---
+			
+			// Extrapolate current rotation along current angular velocity to see where we would end up
+			float CurAngVelSize;
+			FVector CurAngVelAxis;
+			CurrentState.AngVel.FVector::ToDirectionAndLength(CurAngVelAxis, CurAngVelSize);
+			const FQuat CurRotExtrapDelta = FQuat(CurAngVelAxis, CurAngVelSize * DeltaSeconds);
+			const FQuat CurRotExtrap = CurRotExtrapDelta * CurrentState.Quaternion;
+
+			// Slerp from the extrapolated current rotation towards the target rotation
+			// This takes current angular velocity into account
+			const FQuat TargetRotBlended = FQuat::Slerp(CurRotExtrap, TargetRot, 1.0f  /*BlendStepAmount*/); // BlendStepAmount is temporarily disabled for rotation
+
+			// Get the rotational offset between the blended rotation target and the current rotation
+			const FQuat TargetRotDelta = TargetRotBlended * CurrentState.Quaternion.Inverse();
+
+			// Convert the rotational delta to angular velocity
+			float WAngle;
+			FVector WAxis;
+			TargetRotDelta.ToAxisAndAngle(WAxis, WAngle);
+			const FVector TargetRotDeltaBlend = FVector(WAxis * (WAngle / RotCorrectionTime));
+			const FVector RepAngVel = FMath::DegreesToRadians(TargetAngVel) + TargetRotDeltaBlend;
+
+			Handle->SetW(RepAngVel);
+		}
 
 		// Cache data for next replication
-		Target.PrevLinVel = FVector(RepLinVel);
 		Target.PrevPos = FVector(CurrentState.Position);
 
 		// --- Target Extrapolation ---
 		if (Target.TickCount <= FMath::CeilToInt(Target.ReceiveInterval * PhysicsReplicationCVars::PredictiveInterpolationCVars::ExtrapolationTimeMultiplier))
 		{
-			// Extrapolate target position
-			Target.TargetState.Position = Target.TargetState.Position + Target.TargetState.LinVel * DeltaSeconds;
-
-			// Extrapolate target rotation
-			float TargetAngVelSize;
-			FVector TargetAngVelAxis;
-			Target.TargetState.AngVel.FVector::ToDirectionAndLength(TargetAngVelAxis, TargetAngVelSize);
-			TargetAngVelSize = FMath::DegreesToRadians(TargetAngVelSize);
-			const FQuat TargetRotExtrapDelta = FQuat(TargetAngVelAxis, TargetAngVelSize * DeltaSeconds);
-			Target.TargetState.Quaternion = TargetRotExtrapDelta * Target.TargetState.Quaternion;
+			ExtrapolateTarget(Target, 1, DeltaSeconds);
 		}
 
 		if (bSoftSnap)
@@ -1399,7 +1424,23 @@ bool FPhysicsReplicationAsync::PredictiveInterpolation(Chaos::FPBDRigidParticleH
 		}
 	}
 
-	return !bCanSimulate && !PhysicsReplicationCVars::PredictiveInterpolationCVars::bDontClearTarget; // If the object can't simulate, clear target
+	return EndReplicationHelper(Target, false);
+}
+
+void FPhysicsReplicationAsync::ExtrapolateTarget(FReplicatedPhysicsTargetAsync& Target, const int32 ExtrapolateFrames, const float DeltaSeconds)
+{
+	const float ExtrapolationTime = DeltaSeconds * static_cast<float>(ExtrapolateFrames);
+
+	// Extrapolate target position
+	Target.TargetState.Position = Target.TargetState.Position + Target.TargetState.LinVel * ExtrapolationTime;
+
+	// Extrapolate target rotation
+	float TargetAngVelSize;
+	FVector TargetAngVelAxis;
+	Target.TargetState.AngVel.FVector::ToDirectionAndLength(TargetAngVelAxis, TargetAngVelSize);
+	TargetAngVelSize = FMath::DegreesToRadians(TargetAngVelSize);
+	const FQuat TargetRotExtrapDelta = FQuat(TargetAngVelAxis, TargetAngVelSize * ExtrapolationTime);
+	Target.TargetState.Quaternion = TargetRotExtrapDelta * Target.TargetState.Quaternion;
 }
 
 /** Compare states and trigger resimulation if needed */
diff --git a/Engine/Source/Runtime/Engine/Public/PhysicsReplication.h b/Engine/Source/Runtime/Engine/Public/PhysicsReplication.h
--- a/Engine/Source/Runtime/Engine/Public/PhysicsReplication.h
+++ b/Engine/Source/Runtime/Engine/Public/PhysicsReplication.h
@@ -92,6 +92,7 @@ struct FReplicatedPhysicsTargetAsync
 {
 	FReplicatedPhysicsTargetAsync()
 		: AccumulatedErrorSeconds(0.0f)
+		, TickCount(0)
 		, ServerFrame(INDEX_NONE)
 		, ReceiveFrame(INDEX_NONE)
 		, ReceiveInterval(5)
@@ -160,6 +161,7 @@ private:
 private:
 	void UpdateAsyncTarget(const FPhysicsRepAsyncInputData& Input, Chaos::FPBDRigidsSolver* RigidsSolver);
 	void UpdateRewindDataTarget(const FPhysicsRepAsyncInputData& Input);
+	void ExtrapolateTarget(FReplicatedPhysicsTargetAsync& Target, const int32 ExtrapolateFrames, const float DeltaSeconds);
 
 public:
 	void Setup(FRigidBodyErrorCorrection ErrorCorrection)
