Iris - Call ClientSetActorTornOff instead of executing similar code.

#rnx

diff --git a/Engine/Source/Runtime/Engine/Private/Net/Iris/ReplicationSystem/ActorReplicationBridge.cpp b/Engine/Source/Runtime/Engine/Private/Net/Iris/ReplicationSystem/ActorReplicationBridge.cpp
--- a/Engine/Source/Runtime/Engine/Private/Net/Iris/ReplicationSystem/ActorReplicationBridge.cpp
+++ b/Engine/Source/Runtime/Engine/Private/Net/Iris/ReplicationSystem/ActorReplicationBridge.cpp
@@ -764,22 +764,7 @@ void UActorReplicationBridge::DestroyInstanceFromRemote(UObject* Instance, ERepl
 	{
 		if ((DestroyReason == EReplicationBridgeDestroyInstanceReason::TearOff) && !NetDriver->ShouldClientDestroyTearOffActors())
 		{
-			if (Actor->GetRemoteRole() == ROLE_Authority)
-			{
-				Actor->SetRole(ROLE_Authority);
-				Actor->SetReplicates(false);
-
-				if (Actor->GetWorld() != nullptr && !IsEngineExitRequested())
-				{
-					Actor->TornOff();
-				}
-
-				NetDriver->NotifyActorTornOff(Actor);
-			}
-			else
-			{
-				UE_LOG_ACTORREPLICATIONBRIDGE(Warning, TEXT("Trying to tear off actor which doesn't support it: %s"), ToCStr(GetNameSafe(Actor)));
-			}
+			NetDriver->ClientSetActorTornOff(Actor);
 		}
 		else if (EnumHasAnyFlags(DestroyFlags, EReplicationBridgeDestroyInstanceFlags::AllowDestroyInstanceFromRemote))
 		{
