Don't replace the default local exposure settings specified in templates when creating a new project.
- for example, virtual production projects want to have local exposure disabled by default.

diff --git a/Engine/Source/Editor/GameProjectGeneration/Private/GameProjectUtils.cpp b/Engine/Source/Editor/GameProjectGeneration/Private/GameProjectUtils.cpp
--- a/Engine/Source/Editor/GameProjectGeneration/Private/GameProjectUtils.cpp
+++ b/Engine/Source/Editor/GameProjectGeneration/Private/GameProjectUtils.cpp
@@ -144,14 +144,14 @@ namespace
 			TEXT("/Script/Engine.RendererSettings"),
 			TEXT("r.DefaultFeature.LocalExposure.HighlightContrastScale"),
 			TEXT("0.8"),
-			true /* ShouldReplaceExistingValue */);
+			false /* ShouldReplaceExistingValue */);
 
 		// Enable Local Exposure by default for new projects
 		ConfigValues.Emplace(TEXT("DefaultEngine.ini"),
 			TEXT("/Script/Engine.RendererSettings"),
 			TEXT("r.DefaultFeature.LocalExposure.ShadowContrastScale"),
 			TEXT("0.8"),
-			true /* ShouldReplaceExistingValue */);
+			false /* ShouldReplaceExistingValue */);
 	}
 
 	/** Get the configuration values for raytracing if enabled. */
diff --git a/Engine/Source/Runtime/Engine/Classes/Engine/RendererSettings.h b/Engine/Source/Runtime/Engine/Classes/Engine/RendererSettings.h
--- a/Engine/Source/Runtime/Engine/Classes/Engine/RendererSettings.h
+++ b/Engine/Source/Runtime/Engine/Classes/Engine/RendererSettings.h
@@ -672,12 +672,12 @@ class URendererSettings : public UDeveloperSettings
 	uint32 bExtendDefaultLuminanceRangeInAutoExposureSettings : 1;
 
 	UPROPERTY(config, EditAnywhere, Category = DefaultSettings, meta = (
-		ConsoleVariable = "r.DefaultFeature.LocalExposure.HighlightContrastScale", DisplayName = "Local Exposure Highlight Contrast Scale",
+		ConsoleVariable = "r.DefaultFeature.LocalExposure.HighlightContrastScale", DisplayName = "Local Exposure Highlight Contrast",
 		ToolTip = "Default Value for Local Exposure Highlight Contrast Scale.", ClampMin = "0.0", ClampMax = "1.0"))
 	float DefaultFeatureLocalExposureHighlightContrastScale;
 
 	UPROPERTY(config, EditAnywhere, Category = DefaultSettings, meta = (
-		ConsoleVariable = "r.DefaultFeature.LocalExposure.ShadowContrastScale", DisplayName = "Local Exposure Shadow Contrast Scale",
+		ConsoleVariable = "r.DefaultFeature.LocalExposure.ShadowContrastScale", DisplayName = "Local Exposure Shadow Contrast",
 		ToolTip = "Default Value for Local Exposure Shadow Contrast Scale.", ClampMin = "0.0", ClampMax = "1.0"))
 	float DefaultFeatureLocalExposureShadowContrastScale;
 
