LevelInstanceSubsystem::UpdateStreamingState: Cleanup early out condition

#rnx

diff --git a/Engine/Source/Runtime/Engine/Private/LevelInstance/LevelInstanceSubsystem.cpp b/Engine/Source/Runtime/Engine/Private/LevelInstance/LevelInstanceSubsystem.cpp
--- a/Engine/Source/Runtime/Engine/Private/LevelInstance/LevelInstanceSubsystem.cpp
+++ b/Engine/Source/Runtime/Engine/Private/LevelInstance/LevelInstanceSubsystem.cpp
@@ -230,10 +230,23 @@ bool ULevelInstanceSubsystem::IsLoading(const ILevelInstanceInterface* LevelInst
 
 void ULevelInstanceSubsystem::UpdateStreamingState()
 {
+	if (!LevelInstancesToUnload.Num() && !LevelInstancesToLoadOrUpdate.Num())
+	{
+		return;
+	}
+
+#if WITH_EDITOR
+	// Do not update during transaction
+	if (GUndo)
+	{
+		return;
+	}
+#endif
+
 	UpdateStreamingStateInternal();
 
 #if WITH_EDITOR
-	if (!GetWorld()->IsGameWorld() && !GUndo)
+	if (!GetWorld()->IsGameWorld())
 	{
 		// For Editor Worlds make sure UpdateStreamingState completes all recursive loading/unloading
 		while (LevelInstancesToLoadOrUpdate.Num() || LevelInstancesToUnload.Num())
@@ -246,18 +259,7 @@ void ULevelInstanceSubsystem::UpdateStreamingState()
 
 void ULevelInstanceSubsystem::UpdateStreamingStateInternal()
 {
-	if (!LevelInstancesToUnload.Num() && !LevelInstancesToLoadOrUpdate.Num())
-	{
-		return;
-	}
-
 #if WITH_EDITOR
-	// Do not update during transaction
-	if (GUndo)
-	{
-		return;
-	}
-
 	FScopedSlowTask SlowTask(LevelInstancesToUnload.Num() + LevelInstancesToLoadOrUpdate.Num() * 2, LOCTEXT("UpdatingLevelInstances", "Updating Level Instances..."), !GetWorld()->IsGameWorld() && !GetWorld()->GetIsInBlockTillLevelStreamingCompleted());
 	SlowTask.MakeDialogDelayed(1.0f);
 
