Add cvar r.VT.RVT.EnableVolumes to globally enable/disable RVT volumes. Defaults to enabled.

diff --git a/Engine/Source/Runtime/Engine/Private/Components/RuntimeVirtualTextureComponent.cpp b/Engine/Source/Runtime/Engine/Private/Components/RuntimeVirtualTextureComponent.cpp
--- a/Engine/Source/Runtime/Engine/Private/Components/RuntimeVirtualTextureComponent.cpp
+++ b/Engine/Source/Runtime/Engine/Private/Components/RuntimeVirtualTextureComponent.cpp
@@ -22,6 +22,12 @@
 
 #define LOCTEXT_NAMESPACE "URuntimeVirtualTextureComponent"
 
+static TAutoConsoleVariable<int32> CVarRVTEnableVolumes(
+	TEXT("r.VT.RVT.EnableVolumes"),
+	1,
+	TEXT("Enable RVT volumes."),
+	ECVF_Default);
+
 URuntimeVirtualTextureComponent::URuntimeVirtualTextureComponent(const FObjectInitializer& ObjectInitializer)
 	: Super(ObjectInitializer)
 	, SceneProxy(nullptr)
@@ -91,7 +97,7 @@ void URuntimeVirtualTextureComponent::GetHidePrimitiveSettings(bool& OutHidePrim
 bool URuntimeVirtualTextureComponent::IsVisible() const
 {
 	// Make sure to have the component do nothing if VT is disabled or if the world is not compatible with RVT
-	return Super::IsVisible() && IsActiveInWorld() && UseVirtualTexturing(GetScene()->GetFeatureLevel());
+	return Super::IsVisible() && IsActiveInWorld() && CVarRVTEnableVolumes.GetValueOnGameThread() != 0 && UseVirtualTexturing(GetScene()->GetFeatureLevel());
 }
 
 void URuntimeVirtualTextureComponent::ApplyWorldOffset(const FVector& InOffset, bool bWorldShift)
