Remove curve compression codec instance GUID since it causes cook determinism issues and isn't otherwise required.

diff --git a/Engine/Source/Runtime/Core/Public/UObject/FortniteMainBranchObjectVersion.h b/Engine/Source/Runtime/Core/Public/UObject/FortniteMainBranchObjectVersion.h
--- a/Engine/Source/Runtime/Core/Public/UObject/FortniteMainBranchObjectVersion.h
+++ b/Engine/Source/Runtime/Core/Public/UObject/FortniteMainBranchObjectVersion.h
@@ -369,6 +369,9 @@ struct FFortniteMainBranchObjectVersion
 		// Started sorting Possessables, Spawnables, and MovieSceneBindings for better search performance.
 		MovieSceneSortedBindings,
 
+		// Remove the UAnimCurveCompressionCodec::InstanceGuid which causes cook determinism issues
+		RemoveAnimCurveCompressionCodecInstanceGuid,
+
 		// -----<new versions can be added above this line>-------------------------------------------------
 		VersionPlusOne,
 		LatestVersion = VersionPlusOne - 1
diff --git a/Engine/Source/Runtime/Engine/Classes/Animation/AnimCurveCompressionCodec.h b/Engine/Source/Runtime/Engine/Classes/Animation/AnimCurveCompressionCodec.h
--- a/Engine/Source/Runtime/Engine/Classes/Animation/AnimCurveCompressionCodec.h
+++ b/Engine/Source/Runtime/Engine/Classes/Animation/AnimCurveCompressionCodec.h
@@ -34,19 +34,14 @@ class UAnimCurveCompressionCodec : public UObject
 {
 	GENERATED_UCLASS_BODY()
 
-	/** A GUID that is unique to this codec instance. After creation, it never changes. */
-	FGuid InstanceGuid;
-
 	/** Allow us to convert DDC serialized path back into codec object */
 	virtual UAnimCurveCompressionCodec* GetCodec(const FString& Path) { return this; }
 
 	//////////////////////////////////////////////////////////////////////////
 
 	ENGINE_API virtual void Serialize(FArchive& Ar) override;
+
 #if WITH_EDITORONLY_DATA
-	// UObject overrides
-	ENGINE_API virtual void PostInitProperties() override;
-	ENGINE_API virtual void PostDuplicate(bool bDuplicateForPIE) override;
 	/** Returns whether or not we can use this codec to compress. */
 	virtual bool IsCodecValid() const { return true; }
 
diff --git a/Engine/Source/Runtime/Engine/Private/Animation/AnimCurveCompressionCodec.cpp b/Engine/Source/Runtime/Engine/Private/Animation/AnimCurveCompressionCodec.cpp
--- a/Engine/Source/Runtime/Engine/Private/Animation/AnimCurveCompressionCodec.cpp
+++ b/Engine/Source/Runtime/Engine/Private/Animation/AnimCurveCompressionCodec.cpp
@@ -1,6 +1,7 @@
 // Copyright Epic Games, Inc. All Rights Reserved.
 
 #include "Animation/AnimCurveCompressionCodec.h"
+#include "UObject/FortniteMainBranchObjectVersion.h"
 #include "UObject/FortniteReleaseBranchCustomObjectVersion.h"
 #include "UObject/Package.h"
 #include "Interfaces/ITargetPlatform.h"
@@ -16,36 +17,31 @@ void UAnimCurveCompressionCodec::Serialize(FArchive& Ar)
 {
 	Super::Serialize(Ar);
 
+	Ar.UsingCustomVersion(FFortniteMainBranchObjectVersion::GUID);
 	Ar.UsingCustomVersion(FFortniteReleaseBranchCustomObjectVersion::GUID);
+
+	if (Ar.CustomVer(FFortniteMainBranchObjectVersion::GUID) < FFortniteMainBranchObjectVersion::RemoveAnimCurveCompressionCodecInstanceGuid)
+	{
 #if !WITH_EDITOR
-	if (Ar.CustomVer(FFortniteReleaseBranchCustomObjectVersion::GUID) >= FFortniteReleaseBranchCustomObjectVersion::SerializeAnimCurveCompressionCodecGuidOnCook)
+		if (Ar.CustomVer(FFortniteReleaseBranchCustomObjectVersion::GUID) >= FFortniteReleaseBranchCustomObjectVersion::SerializeAnimCurveCompressionCodecGuidOnCook)
 #endif // !WITH_EDITOR
-	{
-		Ar << InstanceGuid;
+		{
+			// We serialized a Guid, read it now and discard it
+			check(Ar.IsLoading());
+			FGuid InstanceGuid;
+			Ar << InstanceGuid;
+		}
 	}
 }
 
 #if WITH_EDITORONLY_DATA
-void UAnimCurveCompressionCodec::PostInitProperties()
-{
-	Super::PostInitProperties();
-
-	if (!IsTemplate())
-	{
-		InstanceGuid = FGuid::NewGuid();
-	}
-}
-
-void UAnimCurveCompressionCodec::PostDuplicate(bool bDuplicateForPIE)
-{
-	Super::PostDuplicate(bDuplicateForPIE);
-
-	InstanceGuid = FGuid::NewGuid();
-}
-
 void UAnimCurveCompressionCodec::PopulateDDCKey(FArchive& Ar)
 {
-	Ar << InstanceGuid;
+	// We use the UClass name to compute the DDC key to avoid two codecs with equivalent properties (e.g. none)
+	// from having the same DDC key. Two codecs with the same values and class name can have the same DDC key
+	// since the caller (e.g. anim sequence) factors in raw data and the likes. For a codec class that derives
+	// from this, it is their responsibility to factor in compression settings and other inputs into the DDC key.
+	FString ClassName = GetClass()->GetName();
+	Ar << ClassName;
 }
 #endif
-
