Iris
* Iris creation data will now replicate the true Outer of replicated subobjects.
* Refactored CallWriteNetRefHandleCreationInfo and made CallWriteNetRefHandleDestructionInfo.

#rnx

diff --git a/Engine/Plugins/Runtime/ReplicationSystemTestPlugin/Source/Private/Tests/ReplicationSystem/ReplicatedTestObject.cpp b/Engine/Plugins/Runtime/ReplicationSystemTestPlugin/Source/Private/Tests/ReplicationSystem/ReplicatedTestObject.cpp
--- a/Engine/Plugins/Runtime/ReplicationSystemTestPlugin/Source/Private/Tests/ReplicationSystem/ReplicatedTestObject.cpp
+++ b/Engine/Plugins/Runtime/ReplicationSystemTestPlugin/Source/Private/Tests/ReplicationSystem/ReplicatedTestObject.cpp
@@ -155,7 +155,7 @@ UObjectReplicationBridge::FCreationHeader* UReplicatedTestObjectBridge::ReadCrea
 	return Header.Release();
 }
 
-FObjectReplicationBridgeInstantiateResult UReplicatedTestObjectBridge::BeginInstantiateFromRemote(FNetRefHandle SubObjectOwnerHandle, const UE::Net::FNetObjectResolveContext& ResolveContext, const FCreationHeader* InHeader)
+FObjectReplicationBridgeInstantiateResult UReplicatedTestObjectBridge::BeginInstantiateFromRemote(FNetRefHandle RootObjectOfSubObject, const UE::Net::FNetObjectResolveContext& ResolveContext, const FCreationHeader* InHeader)
 {
 	const FReplicationTestObjectCreationHeader* Header = static_cast<const FReplicationTestObjectCreationHeader*>(InHeader);
 
@@ -202,25 +202,25 @@ void UReplicatedTestObjectBridge::EndInstantiateFromRemote(FNetRefHandle Handle)
 	Instance->NetRefHandle = Handle;
 }
 
-void UReplicatedTestObjectBridge::DestroyInstanceFromRemote(UObject* Instance, EReplicationBridgeDestroyInstanceReason DestroyReason, EReplicationBridgeDestroyInstanceFlags DestroyFlags)
+void UReplicatedTestObjectBridge::DestroyInstanceFromRemote(const FDestroyInstanceParams& Params)
 {
-	if (!Instance)
+	if (!Params.Instance)
 	{
 		return;
 	}
 
-	if (DestroyReason == EReplicationBridgeDestroyInstanceReason::Destroy)
+	if (Params.DestroyReason == EReplicationBridgeDestroyInstanceReason::Destroy)
 	{
-		if (EnumHasAnyFlags(DestroyFlags, EReplicationBridgeDestroyInstanceFlags::AllowDestroyInstanceFromRemote))
+		if (EnumHasAnyFlags(Params.DestroyFlags, EReplicationBridgeDestroyInstanceFlags::AllowDestroyInstanceFromRemote))
 		{
 			// Remove the object from the created objects on the node
 			if (CreatedObjectsOnNode)
 			{
-				CreatedObjectsOnNode->Remove(TStrongObjectPtr<UObject>(Instance));
+				CreatedObjectsOnNode->Remove(TStrongObjectPtr<UObject>(Params.Instance));
 			}
 
-			Instance->PreDestroyFromReplication();
-			Instance->MarkAsGarbage();
+			Params.Instance->PreDestroyFromReplication();
+			Params.Instance->MarkAsGarbage();
 		}
 	}
 }
diff --git a/Engine/Plugins/Runtime/ReplicationSystemTestPlugin/Source/Private/Tests/ReplicationSystem/ReplicatedTestObject.h b/Engine/Plugins/Runtime/ReplicationSystemTestPlugin/Source/Private/Tests/ReplicationSystem/ReplicatedTestObject.h
--- a/Engine/Plugins/Runtime/ReplicationSystemTestPlugin/Source/Private/Tests/ReplicationSystem/ReplicatedTestObject.h
+++ b/Engine/Plugins/Runtime/ReplicationSystemTestPlugin/Source/Private/Tests/ReplicationSystem/ReplicatedTestObject.h
@@ -507,9 +507,9 @@ protected:
 	virtual bool WriteCreationHeader(UE::Net::FNetSerializationContext& Context, FNetRefHandle Handle) override;
 	virtual FCreationHeader* ReadCreationHeader(UE::Net::FNetSerializationContext& Context) override;
 
-	virtual FObjectReplicationBridgeInstantiateResult BeginInstantiateFromRemote(FNetRefHandle SubObjectOwnerHandle, const UE::Net::FNetObjectResolveContext& ResolveContext, const FCreationHeader* InHeader) override;
+	virtual FObjectReplicationBridgeInstantiateResult BeginInstantiateFromRemote(FNetRefHandle RootObjectOfSubObject, const UE::Net::FNetObjectResolveContext& ResolveContext, const FCreationHeader* InHeader) override;
 	virtual void EndInstantiateFromRemote(FNetRefHandle Handle) override;
-	virtual void DestroyInstanceFromRemote(UObject* Instance, EReplicationBridgeDestroyInstanceReason DestroyReason, EReplicationBridgeDestroyInstanceFlags DestroyFlags) override;
+	virtual void DestroyInstanceFromRemote(const FDestroyInstanceParams& Params) override;
 	virtual bool IsAllowedToDestroyInstance(const UObject* Instance) const override;
 
 	TArray<TStrongObjectPtr<UObject>>* CreatedObjectsOnNode;
diff --git a/Engine/Source/Runtime/Engine/Private/Net/Iris/ReplicationSystem/ActorReplicationBridge.cpp b/Engine/Source/Runtime/Engine/Private/Net/Iris/ReplicationSystem/ActorReplicationBridge.cpp
--- a/Engine/Source/Runtime/Engine/Private/Net/Iris/ReplicationSystem/ActorReplicationBridge.cpp
+++ b/Engine/Source/Runtime/Engine/Private/Net/Iris/ReplicationSystem/ActorReplicationBridge.cpp
@@ -481,9 +481,11 @@ bool UActorReplicationBridge::WriteCreationHeader(UE::Net::FNetSerializationCont
 	}
 	else if (Object)
 	{
+		const UObject* RootObject = GetReplicatedObject(InternalGetRootObjectOfSubObject(Handle));
+
 		// Get Header
 		FSubObjectCreationHeader Header;
-		GetSubObjectCreationHeader(Object, Header);
+		GetSubObjectCreationHeader(Object, RootObject, Header);
 
 		// Serialize the data
 		// Indicate that this is a SubObject
@@ -529,7 +531,7 @@ UObjectReplicationBridge::FCreationHeader* UActorReplicationBridge::ReadCreation
 	return nullptr;
 }
 
-FObjectReplicationBridgeInstantiateResult UActorReplicationBridge::BeginInstantiateFromRemote(FNetRefHandle SubObjectOwnerNetHandle, const UE::Net::FNetObjectResolveContext& ResolveContext, const UObjectReplicationBridge::FCreationHeader* InHeader)
+FObjectReplicationBridgeInstantiateResult UActorReplicationBridge::BeginInstantiateFromRemote(FNetRefHandle RootObjectOfSubObject, const UE::Net::FNetObjectResolveContext& ResolveContext, const UObjectReplicationBridge::FCreationHeader* InHeader)
 {
 	using namespace UE::Net::Private;
 
@@ -636,11 +638,9 @@ FObjectReplicationBridgeInstantiateResult UActorReplicationBridge::BeginInstanti
 		// Spawn sub object
 		if (Header->bIsDynamic)
 		{
-			// Resolve owner
-			UObject* ResolvedOwnerObject = GetReplicatedObject(SubObjectOwnerNetHandle);
-			AActor* Owner = Cast<AActor>(ResolvedOwnerObject);
-
-			check(Owner);
+			// Resolve root object
+			UObject* RootObject = GetReplicatedObject(RootObjectOfSubObject);
+			AActor* RootActor = CastChecked<AActor>(RootObject);
 			
 			UObject* SubObj = nullptr;
 
@@ -651,24 +651,45 @@ FObjectReplicationBridgeInstantiateResult UActorReplicationBridge::BeginInstanti
 
 				if (!SubObj)
 				{
-					UE_LOG_ACTORREPLICATIONBRIDGE(Error, TEXT("BeginInstantiateFromRemote Failed to find subobjectReference for dynamic SubObject %s, Owner %s (%s)"), *DescribeObjectReference(Header->ObjectReference, ResolveContext), *SubObjectOwnerNetHandle.ToString(), *GetPathNameSafe(GetReplicatedObject(SubObjectOwnerNetHandle)));
+					UE_LOG_ACTORREPLICATIONBRIDGE(Error, TEXT("BeginInstantiateFromRemote Failed to find subobjectReference for dynamic SubObject %s, Owner %s, RootObject %s"), *DescribeObjectReference(Header->ObjectReference, ResolveContext), *RootObjectOfSubObject.ToString(), *GetPathNameSafe(RootActor));
 				}
 			}
 			else
 			{
-				UObject* ObjOuter = Owner;
+				// Find the proper Outer
+				UObject* OuterObject = nullptr;
+				if (Header->bOuterIsTransientLevel)
+				{
+					OuterObject = GetTransientPackage();
+				}
+				else if (Header->bOuterIsRootObject)
+				{
+					OuterObject = RootActor;
+				}				
+				else
+				{
+					OuterObject = ResolveObjectReference(Header->OuterReference, ResolveContext);
+
+					if (!OuterObject)
+					{
+						UE_LOG_ACTORREPLICATIONBRIDGE(Error, TEXT("BeginInstantiateFromRemote Failed to find Outer %s for dynamic subobject %s"), *DescribeObjectReference(Header->OuterReference, ResolveContext), *DescribeObjectReference(Header->ObjectReference, ResolveContext))
+
+						// Fallback to the rootobject instead
+						OuterObject = RootActor;
+					}
+				}
 
 				// We need to spawn the subobject
 				UObject* SubObjClassObj = ResolveObjectReference(Header->ObjectClassReference, ResolveContext);
 				UClass * SubObjClass = Cast<UClass>(SubObjClassObj);
 
 				// Try to spawn SubObject
- 				SubObj = NewObject<UObject>(Owner, SubObjClass);
+ 				SubObj = NewObject<UObject>(OuterObject, SubObjClass);
 
 				// Sanity check some things
-				checkf(SubObj != nullptr, TEXT("UActorReplicationBridge::BeginInstantiateFromRemote: Subobject is NULL after instantiating. Class: %s, Actor %s"), *GetNameSafe(SubObjClass), *Owner->GetName());
-				checkf(SubObj->IsIn(ObjOuter), TEXT("UActorReplicationBridge::BeginInstantiateFromRemote: Subobject is not in Outer. SubObject: %s, Actor: %s Outer: %s"), *SubObj->GetName(), *Owner->GetName(), *ObjOuter->GetName());
-				checkf(Cast<AActor>(SubObj) == nullptr, TEXT("UActorReplicationBridge::BeginInstantiateFromRemote: Subobject is an Actor. SubObject: %s, Actor: %s"), *SubObj->GetName(), *Owner->GetName());
+				checkf(SubObj != nullptr, TEXT("UActorReplicationBridge::BeginInstantiateFromRemote: Subobject is NULL after instantiating. Class: %s, Outer %s, Actor %s"), *GetNameSafe(SubObjClass), *GetNameSafe(OuterObject), *GetNameSafe(RootActor));
+				checkf(!OuterObject || SubObj->IsIn(OuterObject), TEXT("UActorReplicationBridge::BeginInstantiateFromRemote: Subobject is not in Outer. SubObject: %s, Outer %s, Actor %s"), *SubObj->GetName(), *GetNameSafe(OuterObject), *GetNameSafe(RootActor));
+				checkf(Cast<AActor>(SubObj) == nullptr, TEXT("UActorReplicationBridge::BeginInstantiateFromRemote: Subobject is an Actor. SubObject: %s, Outer %s, Actor %s"), *SubObj->GetName(), *GetNameSafe(OuterObject), *GetNameSafe(RootActor));
 
 				// We must defer call OnSubObjectCreatedFromReplication after the state has been applied to the owning actor in order to behave like old system.
 				InstantiateResult.Flags |= EReplicationBridgeCreateNetRefHandleResultFlags::ShouldCallSubObjectCreatedFromReplication;
@@ -692,7 +713,8 @@ FObjectReplicationBridgeInstantiateResult UActorReplicationBridge::BeginInstanti
 			}
 			else
 			{
-				UE_LOG_ACTORREPLICATIONBRIDGE(Error, TEXT("BeginInstantiateFromRemote Failed to find Resolve SubObjectReference for static SubObject %s, Owner %s (%s)"), *DescribeObjectReference(Header->ObjectReference, ResolveContext), *SubObjectOwnerNetHandle.ToString(), *GetPathNameSafe(GetReplicatedObject(SubObjectOwnerNetHandle)));
+				UE_LOG_ACTORREPLICATIONBRIDGE(Error, TEXT("BeginInstantiateFromRemote Failed to find Resolve SubObjectReference for static SubObject %s, Owner %s (%s)"), 
+				*DescribeObjectReference(Header->ObjectReference, ResolveContext), *RootObjectOfSubObject.ToString(), *GetPathNameSafe(GetReplicatedObject(RootObjectOfSubObject)));
 			}			
 		}
 	}
@@ -745,28 +767,28 @@ void UActorReplicationBridge::EndInstantiateFromRemote(FNetRefHandle Handle)
 
 void UActorReplicationBridge::OnSubObjectCreatedFromReplication(FNetRefHandle SubObjectHandle)
 {
-	AActor* Owner = Cast<AActor>(GetReplicatedObject(InternalGetSubObjectOwner(SubObjectHandle)));
+	AActor* RootObject = Cast<AActor>(GetReplicatedObject(InternalGetRootObjectOfSubObject(SubObjectHandle)));
 	UObject* SubObject = GetReplicatedObject(SubObjectHandle);
-	if (IsValid(Owner) && IsValid(SubObject))
+	if (IsValid(RootObject) && IsValid(SubObject))
 	{
-		Owner->OnSubobjectCreatedFromReplication(SubObject);
+		RootObject->OnSubobjectCreatedFromReplication(SubObject);
 	}
 }
 
-void UActorReplicationBridge::DestroyInstanceFromRemote(UObject* Instance, EReplicationBridgeDestroyInstanceReason DestroyReason, EReplicationBridgeDestroyInstanceFlags DestroyFlags)
+void UActorReplicationBridge::DestroyInstanceFromRemote(const FDestroyInstanceParams& Params)
 {
-	if (DestroyReason == EReplicationBridgeDestroyInstanceReason::DoNotDestroy)
+	if (Params.DestroyReason == EReplicationBridgeDestroyInstanceReason::DoNotDestroy)
 	{
 		return;
 	}
 
-	if (AActor* Actor = Cast<AActor>(Instance))
+	if (AActor* Actor = Cast<AActor>(Params.Instance))
 	{
-		if ((DestroyReason == EReplicationBridgeDestroyInstanceReason::TearOff) && !NetDriver->ShouldClientDestroyTearOffActors())
+		if ((Params.DestroyReason == EReplicationBridgeDestroyInstanceReason::TearOff) && !NetDriver->ShouldClientDestroyTearOffActors())
 		{
 			NetDriver->ClientSetActorTornOff(Actor);
 		}
-		else if (EnumHasAnyFlags(DestroyFlags, EReplicationBridgeDestroyInstanceFlags::AllowDestroyInstanceFromRemote))
+		else if (EnumHasAnyFlags(Params.DestroyFlags, EReplicationBridgeDestroyInstanceFlags::AllowDestroyInstanceFromRemote))
 		{
 			// Any subobjects have already been detached from ReplicationBridge
 			Actor->PreDestroyFromReplication();
@@ -776,24 +798,24 @@ void UActorReplicationBridge::DestroyInstanceFromRemote(UObject* Instance, ERepl
 	else
 	{
 		// If the SubObject is being torn off it is up to owning actor to clean it up properly
-		if (DestroyReason == EReplicationBridgeDestroyInstanceReason::TearOff)
+		if (Params.DestroyReason == EReplicationBridgeDestroyInstanceReason::TearOff)
 		{
 			return;
 		}
 
-		if (!EnumHasAnyFlags(DestroyFlags, EReplicationBridgeDestroyInstanceFlags::AllowDestroyInstanceFromRemote))
+		if (!EnumHasAnyFlags(Params.DestroyFlags, EReplicationBridgeDestroyInstanceFlags::AllowDestroyInstanceFromRemote))
 		{
 			return;
 		}
 
-		AActor* Owner = Cast<AActor>(Instance->GetOuter());
-		if (ensureMsgf(IsValid(Owner) && !Owner->IsUnreachable(), TEXT("UActorReplicationBridge::DestroyInstanceFromRemote Destroyed subobject after owner %s"), *Instance->GetPathName()))
+		AActor* Owner = Cast<AActor>(Params.RootObject);
+		if (ensureMsgf(IsValid(Owner) && !Owner->IsUnreachable(), TEXT("UActorReplicationBridge::DestroyInstanceFromRemote Destroyed subobject: %s has an invalid owner: %s"), *GetNameSafe(Params.Instance), *GetPathNameSafe(Params.RootObject)))
 		{
-			Owner->OnSubobjectDestroyFromReplication(Instance);
+			Owner->OnSubobjectDestroyFromReplication(Params.Instance);
 		}
 
-		Instance->PreDestroyFromReplication();
-		Instance->MarkAsGarbage();
+		Params.Instance->PreDestroyFromReplication();
+		Params.Instance->MarkAsGarbage();
 	}
 }
 
@@ -987,12 +1009,20 @@ void UActorReplicationBridge::GetActorCreationHeader(const AActor* Actor, UE::Ne
 	}
 }
 
-void UActorReplicationBridge::GetSubObjectCreationHeader(const UObject* Object, UE::Net::Private::FSubObjectCreationHeader& Header) const
+void UActorReplicationBridge::GetSubObjectCreationHeader(const UObject* Object, const UObject* RootObject, UE::Net::Private::FSubObjectCreationHeader& Header) const
 {
+	using namespace UE::Net;
+	using namespace UE::Net::Private;
+
 	// SubObject
-	UE::Net::FNetObjectReference ObjectRef = GetOrCreateObjectReference(Object);
+	FNetObjectReference ObjectRef = GetOrCreateObjectReference(Object);
+
+	// It's Outer
+	UObject* OuterObject = Object->GetOuter();
 
 	Header.bIsActor = false;
+	Header.bOuterIsTransientLevel = false;
+	Header.bOuterIsRootObject = false;
 	Header.bIsDynamic = ObjectRef.GetRefHandle().IsDynamic();
 	Header.bUsePersistentLevel = false;
 	Header.bIsNameStableForNetworking = Object->IsNameStableForNetworking();
@@ -1006,6 +1036,25 @@ void UActorReplicationBridge::GetSubObjectCreationHeader(const UObject* Object,
 		if (!Header.bIsNameStableForNetworking)
 		{
 			Header.ObjectClassReference = GetOrCreateObjectReference(Object->GetClass());
+
+			// Find the Outer
+			if (OuterObject == GetTransientPackage())
+			{
+				Header.bOuterIsTransientLevel = true;
+			}
+			else if (OuterObject == RootObject)
+			{
+				Header.bOuterIsRootObject = true;
+			}
+			else
+			{
+				Header.OuterReference = GetOrCreateObjectReference(OuterObject);
+
+				if (!ensure(Header.OuterReference.IsValid()))
+				{
+					UE_LOG_ACTORREPLICATIONBRIDGE(Error, TEXT("Could not create NetReference to Outer %s of subobject %s"), *GetNameSafe(OuterObject), *GetNameSafe(Object));
+				}
+			}
 		}
 	}
 
diff --git a/Engine/Source/Runtime/Engine/Private/Net/Iris/ReplicationSystem/ActorReplicationBridgeInternal.cpp b/Engine/Source/Runtime/Engine/Private/Net/Iris/ReplicationSystem/ActorReplicationBridgeInternal.cpp
--- a/Engine/Source/Runtime/Engine/Private/Net/Iris/ReplicationSystem/ActorReplicationBridgeInternal.cpp
+++ b/Engine/Source/Runtime/Engine/Private/Net/Iris/ReplicationSystem/ActorReplicationBridgeInternal.cpp
@@ -198,6 +198,14 @@ void WriteSubObjectCreationHeader(FNetSerializationContext& Context, const FSubO
 		else
 		{
 			WriteFullNetObjectReference(Context, Header.ObjectClassReference);
+
+			if (!Writer->WriteBool(Header.bOuterIsTransientLevel))
+			{
+				if (!Writer->WriteBool(Header.bOuterIsRootObject))
+				{
+					WriteFullNetObjectReference(Context, Header.OuterReference);
+				}
+			}
 		}
 	}
 	else
@@ -211,6 +219,9 @@ void ReadSubObjectCreationHeader(FNetSerializationContext& Context, FSubObjectCr
 	FNetBitStreamReader* Reader = Context.GetBitStreamReader();
 
 	Header.bIsActor = false;
+	Header.bOuterIsTransientLevel = false;
+	Header.bOuterIsRootObject = false;
+
 	Header.bIsDynamic = Reader->ReadBool();
 	if (Header.bIsDynamic)
 	{
@@ -222,6 +233,16 @@ void ReadSubObjectCreationHeader(FNetSerializationContext& Context, FSubObjectCr
 		else
 		{
 			ReadFullNetObjectReference(Context, Header.ObjectClassReference);
+
+			Header.bOuterIsTransientLevel = Reader->ReadBool();
+			if (!Header.bOuterIsTransientLevel)
+			{
+				Header.bOuterIsRootObject = Reader->ReadBool();
+				if (!Header.bOuterIsRootObject)
+				{
+					ReadFullNetObjectReference(Context, Header.OuterReference);
+				}
+			}
 		}
 	}
 	else
diff --git a/Engine/Source/Runtime/Engine/Private/Net/Iris/ReplicationSystem/ActorReplicationBridgeInternal.h b/Engine/Source/Runtime/Engine/Private/Net/Iris/ReplicationSystem/ActorReplicationBridgeInternal.h
--- a/Engine/Source/Runtime/Engine/Private/Net/Iris/ReplicationSystem/ActorReplicationBridgeInternal.h
+++ b/Engine/Source/Runtime/Engine/Private/Net/Iris/ReplicationSystem/ActorReplicationBridgeInternal.h
@@ -51,6 +51,9 @@ struct FActorCreationHeader : public FActorReplicationBridgeCreationHeader
 struct FSubObjectCreationHeader : public FActorReplicationBridgeCreationHeader
 {
 	FNetObjectReference ObjectClassReference;
+	uint8 bOuterIsTransientLevel : 1;	// When set the OuterReference was not sent because the Outer is the default transient level.
+	uint8 bOuterIsRootObject : 1;		// When set the OuterReference was not sent because the Outer is the known RootObject.
+	FNetObjectReference OuterReference; // Optional: Outer ref sent only for dynamic subobjects.
 };
 
 const FActorReplicationBridgeSpawnInfo& GetDefaultActorReplicationBridgeSpawnInfo();
diff --git a/Engine/Source/Runtime/Engine/Public/Net/Iris/ReplicationSystem/ActorReplicationBridge.h b/Engine/Source/Runtime/Engine/Public/Net/Iris/ReplicationSystem/ActorReplicationBridge.h
--- a/Engine/Source/Runtime/Engine/Public/Net/Iris/ReplicationSystem/ActorReplicationBridge.h
+++ b/Engine/Source/Runtime/Engine/Public/Net/Iris/ReplicationSystem/ActorReplicationBridge.h
@@ -84,11 +84,11 @@ protected:
 	virtual void Initialize(UReplicationSystem* ReplicationSystem) override;
 	virtual bool WriteCreationHeader(UE::Net::FNetSerializationContext& Context, FNetRefHandle Handle) override;
 	virtual FCreationHeader* ReadCreationHeader(UE::Net::FNetSerializationContext& Context) override;
-	virtual FObjectReplicationBridgeInstantiateResult BeginInstantiateFromRemote(FNetRefHandle SubObjectOwnerNetHandle, const UE::Net::FNetObjectResolveContext& ResolveContext, const FCreationHeader* InHeader) override;
+	virtual FObjectReplicationBridgeInstantiateResult BeginInstantiateFromRemote(FNetRefHandle RootObjectOfSubObject, const UE::Net::FNetObjectResolveContext& ResolveContext, const FCreationHeader* InHeader) override;
 	virtual bool OnInstantiatedFromRemote(UObject* Instance, const FCreationHeader* InHeader, uint32 ConnectionId) const override;
 	virtual void OnSubObjectCreatedFromReplication(FNetRefHandle SubObjectHandle) override;
 	virtual void EndInstantiateFromRemote(FNetRefHandle Handle) override;
-	virtual void DestroyInstanceFromRemote(UObject* Instance, EReplicationBridgeDestroyInstanceReason DestroyReason, EReplicationBridgeDestroyInstanceFlags DestroyFlags) override;
+	virtual void DestroyInstanceFromRemote(const FDestroyInstanceParams& Params) override;
 	virtual void GetInitialDependencies(FNetRefHandle Handle, FNetDependencyInfoArray& OutDependencies) const override;
 	virtual bool RemapPathForPIE(uint32 ConnectionId, FString& Path, bool bReading) const override;
 	virtual bool ObjectLevelHasFinishedLoading(UObject* Object) const override;
@@ -98,7 +98,7 @@ protected:
 
 private:
 	void GetActorCreationHeader(const AActor* Actor, UE::Net::Private::FActorCreationHeader& Header) const;
-	void GetSubObjectCreationHeader(const UObject* Object, UE::Net::Private::FSubObjectCreationHeader& Header) const;
+	void GetSubObjectCreationHeader(const UObject* Object, const UObject* RootObject, UE::Net::Private::FSubObjectCreationHeader& Header) const;
 
 	void OnMaxTickRateChanged(UNetDriver* InNetDriver, int32 NewMaxTickRate, int32 OldMaxTickRate);
 
diff --git a/Engine/Source/Runtime/Experimental/Iris/Core/Private/Iris/ReplicationSystem/NetRefHandleManager.h b/Engine/Source/Runtime/Experimental/Iris/Core/Private/Iris/ReplicationSystem/NetRefHandleManager.h
--- a/Engine/Source/Runtime/Experimental/Iris/Core/Private/Iris/ReplicationSystem/NetRefHandleManager.h
+++ b/Engine/Source/Runtime/Experimental/Iris/Core/Private/Iris/ReplicationSystem/NetRefHandleManager.h
@@ -83,8 +83,10 @@ public:
 		const FReplicationProtocol* Protocol;
 		const FReplicationInstanceProtocol* InstanceProtocol;
 		uint8* ReceiveStateBuffer;
-		uint32 SubObjectRootIndex;
-		uint32 SubObjectParentIndex;
+		/** Subobjects only: Internal index of the RootObject of this subobject */
+		FInternalNetRefIndex SubObjectRootIndex;
+		/** Subobjects only: Internal index of the ParentObject of this subobject */
+		FInternalNetRefIndex SubObjectParentIndex;
 
 		union
 		{
diff --git a/Engine/Source/Runtime/Experimental/Iris/Core/Private/Iris/ReplicationSystem/ObjectReplicationBridge.cpp b/Engine/Source/Runtime/Experimental/Iris/Core/Private/Iris/ReplicationSystem/ObjectReplicationBridge.cpp
--- a/Engine/Source/Runtime/Experimental/Iris/Core/Private/Iris/ReplicationSystem/ObjectReplicationBridge.cpp
+++ b/Engine/Source/Runtime/Experimental/Iris/Core/Private/Iris/ReplicationSystem/ObjectReplicationBridge.cpp
@@ -510,15 +510,28 @@ void UObjectReplicationBridge::EndReplication(UObject* Instance, EEndReplication
 void UObjectReplicationBridge::DetachInstanceFromRemote(FNetRefHandle Handle, EReplicationBridgeDestroyInstanceReason DestroyReason, EReplicationBridgeDestroyInstanceFlags DestroyFlags)
 {
 	UE_LOG_OBJECTREPLICATIONBRIDGE(Verbose, TEXT("DetachInstanceFromRemote %s DestroyReason: %s DestroyFlags: %u"), *Handle.ToString(), ToCStr(LexToString(DestroyReason)), unsigned(DestroyFlags));
+
+	using namespace UE::Net::Private;
+
+	FDestroyInstanceParams DestroyInstanceParams;
 	
-	UObject* Instance = GetObjectFromReferenceHandle(Handle);
+	DestroyInstanceParams.Instance = GetObjectFromReferenceHandle(Handle);
+	DestroyInstanceParams.DestroyReason = DestroyReason;
+	DestroyInstanceParams.DestroyFlags = DestroyFlags;
+
+	if (DestroyReason != EReplicationBridgeDestroyInstanceReason::DoNotDestroy)
+	{
+		const FInternalNetRefIndex ObjectNetIndex = NetRefHandleManager->GetInternalIndex(Handle);
+		const FNetRefHandleManager::FReplicatedObjectData& RepObjectData = NetRefHandleManager->GetReplicatedObjectData(ObjectNetIndex);
+		DestroyInstanceParams.RootObject = RepObjectData.SubObjectRootIndex != FNetRefHandleManager::InvalidInternalIndex ? NetRefHandleManager->GetReplicatedObjectInstance(RepObjectData.SubObjectRootIndex) : nullptr;
+	}
 	
 	UnregisterInstance(Handle);
 
 	// Destroy instance if requested
-	if (Instance && DestroyReason != EReplicationBridgeDestroyInstanceReason::DoNotDestroy)
+	if (DestroyInstanceParams.Instance && DestroyReason != EReplicationBridgeDestroyInstanceReason::DoNotDestroy)
 	{
-		DestroyInstanceFromRemote(Instance, DestroyReason, DestroyFlags);
+		DestroyInstanceFromRemote(DestroyInstanceParams);
 	}
 
 	// $IRIS TODO: Cleanup any pending creation data if we have not yet instantiated the instance.
@@ -553,7 +566,7 @@ void UObjectReplicationBridge::RegisterRemoteInstance(FNetRefHandle RefHandle, U
 	UE_LOG_OBJECTREPLICATIONBRIDGE(Verbose, TEXT("RegisterRemoteInstance %s %s with ProtocolId:0x%" UINT64_x_FMT), *RefHandle.ToString(), ToCStr(Instance->GetName()), Protocol->ProtocolIdentifier);
 }
 
-FReplicationBridgeCreateNetRefHandleResult UObjectReplicationBridge::CreateNetRefHandleFromRemote(FNetRefHandle SubObjectOwnerNetHandle, FNetRefHandle WantedNetHandle, FReplicationBridgeSerializationContext& Context)
+FReplicationBridgeCreateNetRefHandleResult UObjectReplicationBridge::CreateNetRefHandleFromRemote(FNetRefHandle RootObjectOfSubObject, FNetRefHandle WantedNetHandle, FReplicationBridgeSerializationContext& Context)
 {
 	LLM_SCOPE_BYTAG(IrisState);
 
@@ -585,7 +598,7 @@ FReplicationBridgeCreateNetRefHandleResult UObjectReplicationBridge::CreateNetRe
 
 	// Currently we need to always instantiate remote objects, moving forward we want to make this optional so that can be deferred until it is time to apply received state data.
 	// https://jira.it.epicgames.com/browse/UE-127369	
-	FObjectReplicationBridgeInstantiateResult InstantiateResult = BeginInstantiateFromRemote(SubObjectOwnerNetHandle, Context.SerializationContext.GetInternalContext()->ResolveContext, Header.Get());
+	FObjectReplicationBridgeInstantiateResult InstantiateResult = BeginInstantiateFromRemote(RootObjectOfSubObject, Context.SerializationContext.GetInternalContext()->ResolveContext, Header.Get());
 	UObject* InstancePtr = InstantiateResult.Object;
 	if (!InstancePtr)
 	{
diff --git a/Engine/Source/Runtime/Experimental/Iris/Core/Private/Iris/ReplicationSystem/ReplicationBridge.cpp b/Engine/Source/Runtime/Experimental/Iris/Core/Private/Iris/ReplicationSystem/ReplicationBridge.cpp
--- a/Engine/Source/Runtime/Experimental/Iris/Core/Private/Iris/ReplicationSystem/ReplicationBridge.cpp
+++ b/Engine/Source/Runtime/Experimental/Iris/Core/Private/Iris/ReplicationSystem/ReplicationBridge.cpp
@@ -47,7 +47,7 @@ bool UReplicationBridge::WriteNetRefHandleCreationInfo(FReplicationBridgeSeriali
 	return true;
 }
 
-FReplicationBridgeCreateNetRefHandleResult UReplicationBridge::CreateNetRefHandleFromRemote(FNetRefHandle SubObjectOwnerNetHandle, FNetRefHandle WantedNetHandle, FReplicationBridgeSerializationContext& Context)
+FReplicationBridgeCreateNetRefHandleResult UReplicationBridge::CreateNetRefHandleFromRemote(FNetRefHandle RootObjectOfSubObject, FNetRefHandle WantedNetHandle, FReplicationBridgeSerializationContext& Context)
 {
 	return FReplicationBridgeCreateNetRefHandleResult();
 };
@@ -197,40 +197,44 @@ bool UReplicationBridge::CallWriteNetRefHandleCreationInfo(FReplicationBridgeSer
 {
 	using namespace UE::Net;
 
-	// Special path for destruction info
-	if (Context.bIsDestructionInfo)
-	{
-		UE_NET_TRACE_SCOPE(DestructionInfo, *Context.SerializationContext.GetBitStreamWriter(), Context.SerializationContext.GetTraceCollector(), ENetTraceVerbosity::Trace);
+	check(!Context.bIsDestructionInfo);
 
-		const FDestructionInfo* Info = StaticObjectsPendingDestroy.Find(Handle);
-		if (ensureAlwaysMsgf(Info, TEXT("Failed to write destructionInfo for %s"), *Handle.ToString()))
-		{
-			// Write destruction info
-			WriteFullNetObjectReference(Context.SerializationContext, Info->StaticRef);
-		}
-		else
-		{
-			// Write invalid reference
-			WriteFullNetObjectReference(Context.SerializationContext, FNetObjectReference());
-		}
-		return !Context.SerializationContext.HasErrorOrOverflow();	
+	return WriteNetRefHandleCreationInfo(Context, Handle);
+}
+
+bool UReplicationBridge::CallWriteNetRefHandleDestructionInfo(FReplicationBridgeSerializationContext& Context, FNetRefHandle Handle)
+{
+	using namespace UE::Net;
+
+	check(Context.bIsDestructionInfo);
+	UE_NET_TRACE_SCOPE(DestructionInfo, *Context.SerializationContext.GetBitStreamWriter(), Context.SerializationContext.GetTraceCollector(), ENetTraceVerbosity::Trace);
+
+	const FDestructionInfo* Info = StaticObjectsPendingDestroy.Find(Handle);
+	if (ensure(Info))
+	{
+		// Write destruction info
+		WriteFullNetObjectReference(Context.SerializationContext, Info->StaticRef);
 	}
 	else
 	{
-		return WriteNetRefHandleCreationInfo(Context, Handle);
+		UE_LOG_REPLICATIONBRIDGE(Error, TEXT("Failed to write destructionInfo for %s"), *Handle.ToString());
+		// Write invalid reference
+		WriteFullNetObjectReference(Context.SerializationContext, FNetObjectReference());
 	}
+
+	return !Context.SerializationContext.HasErrorOrOverflow();
 }
 
-FReplicationBridgeCreateNetRefHandleResult UReplicationBridge::CallCreateNetRefHandleFromRemote(FNetRefHandle SubObjectOwnerHandle, FNetRefHandle WantedNetHandle, FReplicationBridgeSerializationContext& Context)
+FReplicationBridgeCreateNetRefHandleResult UReplicationBridge::CallCreateNetRefHandleFromRemote(FNetRefHandle RootObjectOfSubObject, FNetRefHandle WantedNetHandle, FReplicationBridgeSerializationContext& Context)
 {
 	check(!Context.bIsDestructionInfo);
 
-	FReplicationBridgeCreateNetRefHandleResult CreateResult = CreateNetRefHandleFromRemote(SubObjectOwnerHandle, WantedNetHandle, Context);
+	FReplicationBridgeCreateNetRefHandleResult CreateResult = CreateNetRefHandleFromRemote(RootObjectOfSubObject, WantedNetHandle, Context);
 
 	// Track subobjects on clients
-	if (CreateResult.NetRefHandle.IsValid() && SubObjectOwnerHandle.IsValid())
+	if (CreateResult.NetRefHandle.IsValid() && RootObjectOfSubObject.IsValid())
 	{
-		NetRefHandleManager->AddSubObject(SubObjectOwnerHandle, CreateResult.NetRefHandle);
+		NetRefHandleManager->AddSubObject(RootObjectOfSubObject, CreateResult.NetRefHandle);
 	}
 
 	return CreateResult;
@@ -445,7 +449,7 @@ void UReplicationBridge::DestroyLocalNetHandle(FNetRefHandle Handle, EEndReplica
 	UE_LOG_REPLICATIONBRIDGE(Verbose, TEXT("DestroyLocalNetHandle Local %s"), *Handle.ToString());
 }
 
-UE::Net::FNetRefHandle UReplicationBridge::InternalGetSubObjectOwner(FNetRefHandle SubObjectHandle) const
+UE::Net::FNetRefHandle UReplicationBridge::InternalGetRootObjectOfSubObject(FNetRefHandle SubObjectHandle) const
 {
 	using namespace UE::Net::Private;
 
diff --git a/Engine/Source/Runtime/Experimental/Iris/Core/Private/Iris/ReplicationSystem/ReplicationReader.cpp b/Engine/Source/Runtime/Experimental/Iris/Core/Private/Iris/ReplicationSystem/ReplicationReader.cpp
--- a/Engine/Source/Runtime/Experimental/Iris/Core/Private/Iris/ReplicationSystem/ReplicationReader.cpp
+++ b/Engine/Source/Runtime/Experimental/Iris/Core/Private/Iris/ReplicationSystem/ReplicationReader.cpp
@@ -774,14 +774,14 @@ void FReplicationReader::ReadObjectInBatch(FNetSerializationContext& Context, FN
 		UE_NET_TRACE_SCOPE(CreationInfo, *Context.GetBitStreamReader(), Context.GetTraceCollector(), ENetTraceVerbosity::Trace);
 
 		// SubObject data for initial state
-		FNetRefHandle SubObjectOwnerHandle;
+		FNetRefHandle RootObjectOfSubObject;
 		if (bIsSubObject)
 		{
 			// The owner is the same as the Batch owner
 			const FNetRefHandle IncompleteOwnerHandle = BatchHandle;
 				
-			FInternalNetRefIndex SubObjectOwnerInternalIndex = NetRefHandleManager->GetInternalIndex(IncompleteOwnerHandle);
-			if (Reader.IsOverflown() || SubObjectOwnerInternalIndex == FNetRefHandleManager::InvalidInternalIndex)
+			FInternalNetRefIndex RootObjectInternalIndex = NetRefHandleManager->GetInternalIndex(IncompleteOwnerHandle);
+			if (Reader.IsOverflown() || RootObjectInternalIndex == FNetRefHandleManager::InvalidInternalIndex)
 			{
 				UE_LOG_REPLICATIONREADER_ERROR(TEXT("FReplicationReader::ReadObject Invalid subobjectowner handle. %s"), ToCStr(IncompleteOwnerHandle.ToString()));
 				const FName& NetError = (Reader.IsOverflown() ? GNetError_BitStreamOverflow : GNetError_InvalidNetHandle);
@@ -789,7 +789,7 @@ void FReplicationReader::ReadObjectInBatch(FNetSerializationContext& Context, FN
 				return;			
 			}
 
-			SubObjectOwnerHandle = NetRefHandleManager->GetReplicatedObjectDataNoCheck(SubObjectOwnerInternalIndex).RefHandle;
+			RootObjectOfSubObject = NetRefHandleManager->GetReplicatedObjectDataNoCheck(RootObjectInternalIndex).RefHandle;
 		}
 
 		const bool bIsDeltaCompressed = Reader.ReadBool();
@@ -811,7 +811,7 @@ void FReplicationReader::ReadObjectInBatch(FNetSerializationContext& Context, FN
 		// Get Bridge
 		FReplicationBridgeSerializationContext BridgeContext(Context, Parameters.ConnectionId);
 
-		const FReplicationBridgeCreateNetRefHandleResult CreateResult = ReplicationBridge->CallCreateNetRefHandleFromRemote(SubObjectOwnerHandle, IncompleteHandle, BridgeContext);
+		const FReplicationBridgeCreateNetRefHandleResult CreateResult = ReplicationBridge->CallCreateNetRefHandleFromRemote(RootObjectOfSubObject, IncompleteHandle, BridgeContext);
 		FNetRefHandle NetRefHandle = CreateResult.NetRefHandle;
 		if (!NetRefHandle.IsValid())
 		{	
diff --git a/Engine/Source/Runtime/Experimental/Iris/Core/Private/Iris/ReplicationSystem/ReplicationWriter.cpp b/Engine/Source/Runtime/Experimental/Iris/Core/Private/Iris/ReplicationSystem/ReplicationWriter.cpp
--- a/Engine/Source/Runtime/Experimental/Iris/Core/Private/Iris/ReplicationSystem/ReplicationWriter.cpp
+++ b/Engine/Source/Runtime/Experimental/Iris/Core/Private/Iris/ReplicationSystem/ReplicationWriter.cpp
@@ -2024,10 +2024,21 @@ FReplicationWriter::EWriteObjectStatus FReplicationWriter::WriteObjectAndSubObje
 						Writer.WriteBits(CreatedBaselineIndex, FDeltaCompressionBaselineManager::BaselineIndexBitCount);
 					}
 
-					FReplicationBridgeSerializationContext BridgeContext(Context, Parameters.ConnectionId, Info.IsDestructionInfo == 1U);
+					const bool bIsDestructionInfo = (Info.IsDestructionInfo == 1U);					
+					FReplicationBridgeSerializationContext BridgeContext(Context, Parameters.ConnectionId, bIsDestructionInfo);
 
-					// We need to send creation info, if we fail, we skip this object for now
-					if (!ReplicationBridge->CallWriteNetRefHandleCreationInfo(BridgeContext, NetRefHandle))
+					bool bWriteSuccess = false;
+					if (Info.IsDestructionInfo)
+					{
+						bWriteSuccess = ReplicationBridge->CallWriteNetRefHandleDestructionInfo(BridgeContext, NetRefHandle);
+					}
+					else
+					{
+						bWriteSuccess = ReplicationBridge->CallWriteNetRefHandleCreationInfo(BridgeContext, NetRefHandle);
+					}
+
+					// We need to send creation info, so if we fail we skip this object for now
+					if (!bWriteSuccess)
 					{
 						return BridgeContext.SerializationContext.HasError() ? EWriteObjectStatus::Error : EWriteObjectStatus::BitStreamOverflow;
 					}
@@ -2572,11 +2583,12 @@ int FReplicationWriter::WriteDestructionInfo(FNetSerializationContext& Context,
 	// Special case for static objects that should be destroyed on the client but we have not replicated
 	Writer.WriteBool(true);
 
-	FReplicationBridgeSerializationContext BridgeContext(Context, Parameters.ConnectionId, true);
+	constexpr bool bIsDestructionInfo = true;
+	FReplicationBridgeSerializationContext BridgeContext(Context, Parameters.ConnectionId, bIsDestructionInfo);
 
 	// Push ForceInlineExportScope to inline exports instead of writing exports later.
 	FForceInlineExportScope ForceInlineExportScope(Context.GetInternalContext());
-	if (!ReplicationBridge->CallWriteNetRefHandleCreationInfo(BridgeContext, ObjectData.RefHandle))
+	if (!ReplicationBridge->CallWriteNetRefHandleDestructionInfo(BridgeContext, ObjectData.RefHandle))
 	{
 		// Trigger Rollback
 		Writer.DoOverflow();
diff --git a/Engine/Source/Runtime/Experimental/Iris/Core/Public/Iris/Core/NetObjectReference.h b/Engine/Source/Runtime/Experimental/Iris/Core/Public/Iris/Core/NetObjectReference.h
--- a/Engine/Source/Runtime/Experimental/Iris/Core/Public/Iris/Core/NetObjectReference.h
+++ b/Engine/Source/Runtime/Experimental/Iris/Core/Public/Iris/Core/NetObjectReference.h
@@ -92,7 +92,7 @@ private:
 
 	FNetRefHandle RefHandle;
 	FNetToken PathToken;
-	ENetObjectReferenceTraits Traits;
+	ENetObjectReferenceTraits Traits = ENetObjectReferenceTraits::None;
 };
 
 /** Representation of an object dependency, e.g. when used in UReplicationBridge::GetInitialDependencies(). */
diff --git a/Engine/Source/Runtime/Experimental/Iris/Core/Public/Iris/ReplicationSystem/ObjectReplicationBridge.h b/Engine/Source/Runtime/Experimental/Iris/Core/Public/Iris/ReplicationSystem/ObjectReplicationBridge.h
--- a/Engine/Source/Runtime/Experimental/Iris/Core/Public/Iris/ReplicationSystem/ObjectReplicationBridge.h
+++ b/Engine/Source/Runtime/Experimental/Iris/Core/Public/Iris/ReplicationSystem/ObjectReplicationBridge.h
@@ -163,7 +163,7 @@ protected:
 	IRISCORE_API virtual void UpdateInstancesWorldLocation() override;
 	IRISCORE_API virtual void PruneStaleObjects() override;	
 	IRISCORE_API virtual bool WriteNetRefHandleCreationInfo(FReplicationBridgeSerializationContext& Context, FNetRefHandle Handle) override;
-	IRISCORE_API virtual FReplicationBridgeCreateNetRefHandleResult CreateNetRefHandleFromRemote(FNetRefHandle SubObjectOwnerNetHandle, FNetRefHandle WantedNetHandle, FReplicationBridgeSerializationContext& Context) override;
+	IRISCORE_API virtual FReplicationBridgeCreateNetRefHandleResult CreateNetRefHandleFromRemote(FNetRefHandle RootObjectNetHandle, FNetRefHandle WantedNetHandle, FReplicationBridgeSerializationContext& Context) override;
 	IRISCORE_API virtual void SubObjectCreatedFromReplication(FNetRefHandle SubObjectHandle) override;
 	IRISCORE_API virtual void PostApplyInitialState(FNetRefHandle Handle) override;
 	IRISCORE_API virtual void DetachInstanceFromRemote(FNetRefHandle Handle, EReplicationBridgeDestroyInstanceReason DestroyReason, EReplicationBridgeDestroyInstanceFlags DestroyFlags) override;
@@ -184,7 +184,7 @@ protected:
 	virtual FCreationHeader* ReadCreationHeader(UE::Net::FNetSerializationContext& Context) { return nullptr; };
 
 	/** Called when we instantiate/find object instance requested by remote. */
-	virtual FObjectReplicationBridgeInstantiateResult BeginInstantiateFromRemote(FNetRefHandle SubObjectOwnerHandle, const UE::Net::FNetObjectResolveContext& ResolveContext, const FCreationHeader* Header) { return FObjectReplicationBridgeInstantiateResult(); };
+	virtual FObjectReplicationBridgeInstantiateResult BeginInstantiateFromRemote(FNetRefHandle RootObjectOfSubObject, const UE::Net::FNetObjectResolveContext& ResolveContext, const FCreationHeader* Header) { return FObjectReplicationBridgeInstantiateResult(); };
 
 	/** Invoked before we start applying state data to instance on remote end. */
 	virtual bool OnInstantiatedFromRemote(UObject* Instance, const FCreationHeader* InHeader, uint32 ConnectionId) const { return true; }
@@ -195,8 +195,16 @@ protected:
 	/** Invoked after remote NetHandle has been created and initial state is applied. */
 	virtual void EndInstantiateFromRemote(FNetRefHandle Handle) {};
 
+	struct FDestroyInstanceParams
+	{
+		UObject* Instance = nullptr;
+		UObject* RootObject = nullptr;
+		EReplicationBridgeDestroyInstanceReason DestroyReason = EReplicationBridgeDestroyInstanceReason::DoNotDestroy;
+		EReplicationBridgeDestroyInstanceFlags DestroyFlags = EReplicationBridgeDestroyInstanceFlags::None;
+	};
+
 	/** Destroy or tear-off the game instance on request from remote. */
-	virtual void DestroyInstanceFromRemote(UObject* Instance, EReplicationBridgeDestroyInstanceReason DestroyReason, EReplicationBridgeDestroyInstanceFlags DestroyFlags) {}
+	virtual void DestroyInstanceFromRemote(const FDestroyInstanceParams& Params) {}
 
 protected:
 	/** Lookup the UObject associated with the provided Handle. This function will not try to resolve the reference. */
diff --git a/Engine/Source/Runtime/Experimental/Iris/Core/Public/Iris/ReplicationSystem/ReplicationBridge.h b/Engine/Source/Runtime/Experimental/Iris/Core/Public/Iris/ReplicationSystem/ReplicationBridge.h
--- a/Engine/Source/Runtime/Experimental/Iris/Core/Public/Iris/ReplicationSystem/ReplicationBridge.h
+++ b/Engine/Source/Runtime/Experimental/Iris/Core/Public/Iris/ReplicationSystem/ReplicationBridge.h
@@ -176,11 +176,15 @@ protected:
 
 	// Remote interface, invoked from Replication code during serialization
 	
-	/** Write data required to instantiate NetObject remotely to bitstream. */
+	/**
+	 * Write data required to instantiate NetObject remotely to bitstream.
+	 * @param Context The serialization context parameters.
+	 * @param Handle The handle of the object to write creation data for.
+	 */
 	IRISCORE_API virtual bool WriteNetRefHandleCreationInfo(FReplicationBridgeSerializationContext& Context, FNetRefHandle Handle);
 
 	/** Read data required to instantiate NetObject from bitstream. */
-	IRISCORE_API virtual FReplicationBridgeCreateNetRefHandleResult CreateNetRefHandleFromRemote(FNetRefHandle SubObjectOwnerNetHandle, FNetRefHandle WantedNetHandle, FReplicationBridgeSerializationContext& Context);
+	IRISCORE_API virtual FReplicationBridgeCreateNetRefHandleResult CreateNetRefHandleFromRemote(FNetRefHandle RootObjectOfSubObject, FNetRefHandle WantedNetHandle, FReplicationBridgeSerializationContext& Context);
 
 	/** Invoked right before we apply the state for a new received subobject but after we have applied state for owning/root object in order to behave like old replication system */
 	IRISCORE_API virtual void SubObjectCreatedFromReplication(FNetRefHandle SubObjectRefHandle);
@@ -226,8 +230,8 @@ protected:
 	/** Destroy the handle and all internal book keeping associated with it. */
 	IRISCORE_API void InternalDestroyNetObject(FNetRefHandle Handle);
 	
-	/** Get the owner handle of a subobject handle. */
-	IRISCORE_API FNetRefHandle InternalGetSubObjectOwner(FNetRefHandle SubObjectHandle) const;
+	/** Get the handle of the root object of any given subobject. */
+	IRISCORE_API FNetRefHandle InternalGetRootObjectOfSubObject(FNetRefHandle SubObjectHandle) const;
 
 	/** Add SubObjectHandle as SubObject to OwnerHandle. */
 	IRISCORE_API void InternalAddSubObject(FNetRefHandle OwnerHandle, FNetRefHandle SubObjectHandle, FNetRefHandle InsertRelativeToSubObjectHandle, ESubObjectInsertionOrder InsertionOrder);
@@ -260,11 +264,12 @@ private:
 	// it will be kept around until EndReplication is called.
 	void TearOff(FNetRefHandle Handle, EEndReplicationFlags DestroyFlags, bool bIsImmediate);
 
-	FReplicationBridgeCreateNetRefHandleResult CallCreateNetRefHandleFromRemote(FNetRefHandle SubObjectOwnerHandle, FNetRefHandle WantedNetHandle, FReplicationBridgeSerializationContext& Context);
+	FReplicationBridgeCreateNetRefHandleResult CallCreateNetRefHandleFromRemote(FNetRefHandle RootObjectOfSubObject, FNetRefHandle WantedNetHandle, FReplicationBridgeSerializationContext& Context);
 	void CallPreSendUpdate(float DeltaSeconds);	
 	void CallPreSendUpdateSingleHandle(FNetRefHandle Handle);
 	void CallUpdateInstancesWorldLocation();
 	bool CallWriteNetRefHandleCreationInfo(FReplicationBridgeSerializationContext& Context, FNetRefHandle Handle);
+	bool CallWriteNetRefHandleDestructionInfo(FReplicationBridgeSerializationContext& Context, FNetRefHandle Handle);
 	void CallSubObjectCreatedFromReplication(FNetRefHandle SubObjectHandle);
 	void CallPostApplyInitialState(FNetRefHandle Handle);
 	void CallPruneStaleObjects();
