WorldPartition: Added detection of existing objects before creating them in UWorldPartitionRuntimeHash and UWorldPartitionRuntimeLevelStreamingCell.
#jiira none

diff --git a/Engine/Source/Runtime/Engine/Private/WorldPartition/WorldPartitionRuntimeHash.cpp b/Engine/Source/Runtime/Engine/Private/WorldPartition/WorldPartitionRuntimeHash.cpp
--- a/Engine/Source/Runtime/Engine/Private/WorldPartition/WorldPartitionRuntimeHash.cpp
+++ b/Engine/Source/Runtime/Engine/Private/WorldPartition/WorldPartitionRuntimeHash.cpp
@@ -96,6 +96,11 @@ UWorldPartitionRuntimeCell* UWorldPartitionRuntimeHash::CreateRuntimeCell(UClass
 	const FString CellObjectName = GetCellObjectName(CellName) + CellInstanceSuffix;
 	// Use given outer if provided, else use hash as outer
 	UObject* Outer = InOuter ? InOuter : this;
+	if (FindObject<UWorldPartitionRuntimeCell>(Outer, *CellObjectName))
+	{
+		UE_LOG(LogWorldPartition, Warning, TEXT("UWorldPartitionRuntimeHash::CreateRuntimeCell can't create an already existing UWorldPartitionRuntimeCell object named %s"), *CellObjectName);
+		return nullptr;
+	}
 	UWorldPartitionRuntimeCell* RuntimeCell = NewObject<UWorldPartitionRuntimeCell>(Outer, CellClass, *CellObjectName);
 	RuntimeCell->RuntimeCellData = NewObject<UWorldPartitionRuntimeCellData>(RuntimeCell, CellDataClass);
 	return RuntimeCell;
@@ -103,6 +108,12 @@ UWorldPartitionRuntimeCell* UWorldPartitionRuntimeHash::CreateRuntimeCell(UClass
 
 URuntimeHashExternalStreamingObjectBase* UWorldPartitionRuntimeHash::CreateExternalStreamingObject(TSubclassOf<URuntimeHashExternalStreamingObjectBase> InClass, UObject* InOuter, FName InName, UWorld* InOwningWorld, UWorld* InOuterWorld)
 {
+	if (FindObject<URuntimeHashExternalStreamingObjectBase>(InOuter, *InName.ToString()))
+	{
+		UE_LOG(LogWorldPartition, Warning, TEXT("UWorldPartitionRuntimeHash::CreateExternalStreamingObject can't create an already existing URuntimeHashExternalStreamingObjectBase object named %s"), *InName.ToString());
+		return nullptr;
+	}
+
 	URuntimeHashExternalStreamingObjectBase* StreamingObject = NewObject<URuntimeHashExternalStreamingObjectBase>(InOuter, InClass, InName, RF_Public);
 	StreamingObject->OwningWorld = InOwningWorld;
 	StreamingObject->OuterWorld = InOuterWorld;
diff --git a/Engine/Source/Runtime/Engine/Private/WorldPartition/WorldPartitionRuntimeLevelStreamingCell.cpp b/Engine/Source/Runtime/Engine/Private/WorldPartition/WorldPartitionRuntimeLevelStreamingCell.cpp
--- a/Engine/Source/Runtime/Engine/Private/WorldPartition/WorldPartitionRuntimeLevelStreamingCell.cpp
+++ b/Engine/Source/Runtime/Engine/Private/WorldPartition/WorldPartitionRuntimeLevelStreamingCell.cpp
@@ -104,6 +104,11 @@ bool UWorldPartitionRuntimeLevelStreamingCell::CreateAndSetLevelStreaming(const
 {
 	UWorld* OwningWorld = GetOwningWorld();
 	const FName LevelStreamingName = FName(*FString::Printf(TEXT("WorldPartitionLevelStreaming_%s"), *GetName()));
+	if (FindObject<UWorldPartitionLevelStreamingDynamic>(OwningWorld, *LevelStreamingName.ToString()))
+	{
+		UE_LOG(LogWorldPartition, Warning, TEXT("UWorldPartitionRuntimeLevelStreamingCell::CreateAndSetLevelStreaming can't create an already existing UWorldPartitionLevelStreamingDynamic object named %s"), *LevelStreamingName.ToString());
+		return false;
+	}
 	LevelStreaming = NewObject<UWorldPartitionLevelStreamingDynamic>(OwningWorld, UWorldPartitionLevelStreamingDynamic::StaticClass(), LevelStreamingName, RF_NoFlags, NULL);
 
 	// Generate unique Level Instance name assuming cell has a unique name
