Don't apply virtual texture pool scaling in editor preview modes. In editor we always want to have the full physical pools available.

diff --git a/Engine/Source/Runtime/Engine/Private/VT/VirtualTextureScalability.cpp b/Engine/Source/Runtime/Engine/Private/VT/VirtualTextureScalability.cpp
--- a/Engine/Source/Runtime/Engine/Private/VT/VirtualTextureScalability.cpp
+++ b/Engine/Source/Runtime/Engine/Private/VT/VirtualTextureScalability.cpp
@@ -83,28 +83,28 @@ namespace VirtualTextureScalability
 		GPoolSizeScales[0],
 		TEXT("Scale factor for virtual texture physical pool size.\n")
 		TEXT(" Group 0"),
-		ECVF_Scalability
+		ECVF_Scalability | ECVF_ExcludeFromPreview
 	);
 	static FAutoConsoleVariableRef CVarVTPoolSizeScale0(
 		TEXT("r.VT.PoolSizeScale.Group0"),
 		GPoolSizeScales[0],
 		TEXT("Scale factor for virtual texture physical pool size.\n")
 		TEXT(" Group 0"),
-		ECVF_Scalability
+		ECVF_Scalability | ECVF_ExcludeFromPreview
 	);
 	static FAutoConsoleVariableRef CVarVTPoolSizeScale1(
 		TEXT("r.VT.PoolSizeScale.Group1"),
 		GPoolSizeScales[1],
 		TEXT("Scale factor for virtual texture physical pool sizes.\n")
 		TEXT(" Group 1"),
-		ECVF_Scalability
+		ECVF_Scalability | ECVF_ExcludeFromPreview
 	);
 	static FAutoConsoleVariableRef CVarVTPoolSizeScale2(
 		TEXT("r.VT.PoolSizeScale.Group2"),
 		GPoolSizeScales[2],
 		TEXT("Scale factor for virtual texture physical pool sizes.\n")
 		TEXT(" Group 2"),
-		ECVF_Scalability
+		ECVF_Scalability | ECVF_ExcludeFromPreview
 	);
 
 	static float GTileCountBiases[NumScalabilityGroups] = { 0 };
