In Network Physics also look for Player Controller as owner of the pawn

diff --git a/Engine/Source/Runtime/Engine/Private/PhysicsEngine/NetworkPhysicsComponent.cpp b/Engine/Source/Runtime/Engine/Private/PhysicsEngine/NetworkPhysicsComponent.cpp
--- a/Engine/Source/Runtime/Engine/Private/PhysicsEngine/NetworkPhysicsComponent.cpp
+++ b/Engine/Source/Runtime/Engine/Private/PhysicsEngine/NetworkPhysicsComponent.cpp
@@ -866,7 +866,17 @@ APlayerController* UNetworkPhysicsComponent::GetPlayerController() const
 
 	if (APawn* Pawn = Cast<APawn>(GetOwner()))
 	{
-		return Pawn->GetController<APlayerController>();
+		if (APlayerController * PC = Pawn->GetController<APlayerController>())
+		{
+			return PC;
+		}
+
+		// In this case the APlayerController can be found as the owner of the pawn
+		if (APlayerController* PC = Cast<APlayerController>(Pawn->GetOwner()))
+		{
+			return PC;
+		}
+
 	}
 
 	return nullptr;
