World Partition: when querying for streaming state, make sure to update local streaming sources.

diff --git a/Engine/Source/Runtime/Engine/Private/WorldPartition/WorldPartition.cpp b/Engine/Source/Runtime/Engine/Private/WorldPartition/WorldPartition.cpp
--- a/Engine/Source/Runtime/Engine/Private/WorldPartition/WorldPartition.cpp
+++ b/Engine/Source/Runtime/Engine/Private/WorldPartition/WorldPartition.cpp
@@ -1762,6 +1762,7 @@ bool UWorldPartition::IsStreamingCompleted(const TArray<FWorldPartitionStreaming
 {
 	if (GetWorld()->IsGameWorld() && StreamingPolicy)
 	{
+		++StreamingStateEpoch; // Update streaming state epoch to make sure we reevaluate streaming sources
 		return StreamingPolicy->IsStreamingCompleted(InStreamingSources);
 	}
 	return true;
@@ -1771,6 +1772,7 @@ bool UWorldPartition::IsStreamingCompleted(EWorldPartitionRuntimeCellState Query
 {
 	if (GetWorld()->IsGameWorld() && StreamingPolicy)
 	{
+		++StreamingStateEpoch; // Update streaming state epoch to make sure we reevaluate streaming sources
 		return StreamingPolicy->IsStreamingCompleted(QueryState, QuerySources, bExactState);
 	}
 
diff --git a/Engine/Source/Runtime/Engine/Public/WorldPartition/WorldPartition.h b/Engine/Source/Runtime/Engine/Public/WorldPartition/WorldPartition.h
--- a/Engine/Source/Runtime/Engine/Public/WorldPartition/WorldPartition.h
+++ b/Engine/Source/Runtime/Engine/Public/WorldPartition/WorldPartition.h
@@ -543,7 +543,7 @@ private:
 	static ENGINE_API FAutoConsoleVariableRef CVarDebugDedicatedServerStreaming;
 #endif
 
-	int32 StreamingStateEpoch;
+	mutable int32 StreamingStateEpoch;
 	static ENGINE_API int32 GlobalEnableServerStreaming;
 	static ENGINE_API bool bGlobalEnableServerStreamingOut;
 	static ENGINE_API bool bUseMakingVisibleTransactionRequests;
