World Partition Streaming Generation: when splitting runtime and editor references, make sure to avoid losing the attach parent actor reference due to how we handle attachements for error reporting.

diff --git a/Engine/Source/Runtime/Engine/Private/WorldPartition/WorldPartitionStreamingGeneration.cpp b/Engine/Source/Runtime/Engine/Private/WorldPartition/WorldPartitionStreamingGeneration.cpp
--- a/Engine/Source/Runtime/Engine/Private/WorldPartition/WorldPartitionStreamingGeneration.cpp
+++ b/Engine/Source/Runtime/Engine/Private/WorldPartition/WorldPartitionStreamingGeneration.cpp
@@ -856,15 +856,6 @@ class FWorldPartitionStreamingGenerator
 					}
 				}
 
-				TArray<FGuid> RuntimeReferences;
-				TArray<FGuid> EditorReferences;
-
-				if (PassType == EPassType::Fixup)
-				{
-					RuntimeReferences.Reserve(ActorDescView.GetReferences().Num());
-					EditorReferences.Reserve(ActorDescView.GetReferences().Num());
-				}
-
 				for (FActorReferenceInfo& Info : References)
 				{
 					FWorldPartitionActorDescView* RefererActorDescView = Info.ActorDesc;
@@ -873,16 +864,7 @@ class FWorldPartitionStreamingGenerator
 					if (ReferenceActorDescView)
 					{
 						// The actor reference is not editor-only, but we are referencing it through an editor-only property
-						if (RefererActorDescView->IsEditorOnlyReference(ReferenceActorDescView->GetGuid()))
-						{
-							if (PassType == EPassType::Fixup)
-							{
-								EditorReferences.Add(Info.ReferenceGuid);
-							}
-
-							NbErrorsDetected++;
-						}
-						else
+						if (!RefererActorDescView->IsEditorOnlyReference(ReferenceActorDescView->GetGuid()))
 						{
 							// Validate grid placement
 							if (!IsReferenceGridPlacementValid(*RefererActorDescView, *ReferenceActorDescView))
@@ -929,11 +911,6 @@ class FWorldPartitionStreamingGenerator
 
 								NbErrorsDetected++;
 							}
-
-							if (PassType == EPassType::Fixup)
-							{
-								RuntimeReferences.Add(Info.ReferenceGuid);
-							}
 						}
 					}
 					else
@@ -955,27 +932,50 @@ class FWorldPartitionStreamingGenerator
 								}
 
 								ErrorHandler->OnInvalidReference(*RefererActorDescView, Info.ReferenceGuid, ReferendceActorDescViewPtr);
+
+								NbErrorsDetected++;
 							}
 						}
-						else if (PassType == EPassType::Fixup)
-						{
-							EditorReferences.Add(Info.ReferenceGuid);
-						}
-
-						NbErrorsDetected++;
 					}
 				}
+			});
+		}
+
+		// Split runtime and editor references
+		ContainerCollectionDescriptor.ActorDescViewMap.ForEachActorDescView([this, &ContainerCollectionDescriptor](FWorldPartitionActorDescView& ActorDescView)
+		{
+			TArray<FGuid> RuntimeReferences;
+			TArray<FGuid> EditorReferences;
 
-				if (PassType == EPassType::Fixup)
+			RuntimeReferences.Reserve(ActorDescView.GetReferences().Num());
+			EditorReferences.Reserve(ActorDescView.GetReferences().Num());
+
+			for (const FGuid& ReferenceGuid : ActorDescView.GetReferences())
+			{
+				if (FWorldPartitionActorDescView* ReferenceActorDescView = ContainerCollectionDescriptor.ActorDescViewMap.FindByGuid(ReferenceGuid))
 				{
-					if (RuntimeReferences.Num() != ActorDescView.GetReferences().Num())
+					// The actor reference is not editor-only, but we are referencing it through an editor-only property
+					if (ActorDescView.IsEditorOnlyReference(ReferenceGuid))
+					{
+						EditorReferences.Add(ReferenceGuid);
+					}
+					else
 					{
-						ActorDescView.SetRuntimeReferences(RuntimeReferences);
-						ActorDescView.SetEditorReferences(EditorReferences);
-					}					
+						RuntimeReferences.Add(ReferenceGuid);
+					}
 				}
-			});		
-		}
+				else if (ContainerCollectionDescriptor.EditorOnlyActorDescMap.Contains(ReferenceGuid))
+				{
+					EditorReferences.Add(ReferenceGuid);
+				}
+			}
+
+			if (RuntimeReferences.Num() != ActorDescView.GetReferences().Num())
+			{
+				ActorDescView.SetRuntimeReferences(RuntimeReferences);
+				ActorDescView.SetEditorReferences(EditorReferences);
+			}
+		});
 	}
 
 	/** 
