[MLDeformer] Fix some reimport issue when reimporting skeletal meshes.

#trivial

diff --git a/Engine/Source/Runtime/Engine/Private/SkeletalMesh.cpp b/Engine/Source/Runtime/Engine/Private/SkeletalMesh.cpp
--- a/Engine/Source/Runtime/Engine/Private/SkeletalMesh.cpp
+++ b/Engine/Source/Runtime/Engine/Private/SkeletalMesh.cpp
@@ -2232,6 +2232,16 @@ void USkeletalMesh::SaveLODImportedData(const int32 LODIndex, FSkeletalMeshImpor
 	GetImportedModel()->LODModels[LODIndex].RawSkeletalMeshBulkDataID = RawSkeletalMeshBulkData.GetIdString();
 	GetImportedModel()->LODModels[LODIndex].bIsBuildDataAvailable = RawSkeletalMeshBulkData.IsBuildDataAvailable();
 	GetImportedModel()->LODModels[LODIndex].bIsRawSkeletalMeshBulkDataEmpty = RawSkeletalMeshBulkData.IsEmpty();
+
+	FSkeletalMeshLODModel& LODModel = GetImportedModel()->LODModels[LODIndex];
+	for (const SkeletalMeshImportData::FMeshInfo& MeshInfo : InMesh.MeshInfos)
+	{
+		LODModel.ImportedMeshInfos.AddDefaulted();
+		FSkelMeshImportedMeshInfo& LODMeshInfo = LODModel.ImportedMeshInfos.Last();
+		LODMeshInfo.Name = MeshInfo.Name;
+		LODMeshInfo.NumVertices = MeshInfo.NumVertices;
+		LODMeshInfo.StartImportedVertex = MeshInfo.StartImportedVertex;
+	}
 }
 
 bool USkeletalMesh::IsLODImportedDataBuildAvailable(const int32 LODIndex) const
