Add case in FBodyInstance::UpdateBodyScale to silently handle heightfields when they're skipped instead of giving a warning since Landscape uses a different path to handle its scaling.

diff --git a/Engine/Source/Runtime/Engine/Private/PhysicsEngine/BodyInstance.cpp b/Engine/Source/Runtime/Engine/Private/PhysicsEngine/BodyInstance.cpp
--- a/Engine/Source/Runtime/Engine/Private/PhysicsEngine/BodyInstance.cpp
+++ b/Engine/Source/Runtime/Engine/Private/PhysicsEngine/BodyInstance.cpp
@@ -2178,6 +2178,11 @@ bool FBodyInstance::UpdateBodyScale(const FVector& InScale3D, bool bForceUpdate)
 
 					break;
 				}
+				case ImplicitObjectType::HeightField:
+				{
+					// HeightField is only used by Landscape, which does different code path from other primitives
+					break;
+				}
 				default:
 				{
 					CHAOS_ENSURE(false);
