Fix Substrate material attribute conversion with CustomUVs.

diff --git a/Engine/Source/Runtime/Engine/Private/Materials/Material.cpp b/Engine/Source/Runtime/Engine/Private/Materials/Material.cpp
--- a/Engine/Source/Runtime/Engine/Private/Materials/Material.cpp
+++ b/Engine/Source/Runtime/Engine/Private/Materials/Material.cpp
@@ -3302,6 +3302,7 @@ void UMaterial::ConvertMaterialToSubstrateMaterial()
 		// * Forward inputs to the root node (Do not reconnect the Opacity as we handle the opacity by internally within the conversion node)
 		// * Always forward masked opacity because this is required when the blend mode is overriden to masked on a material instance.
 		bUseMaterialAttributes = false;
+		NumCustomizedUVs = 8;
 		EditorOnly->FrontMaterial.Connect(0, ConvertAttributeNode);
 		EditorOnly->Opacity.Connect(4, ConvertAttributeNode);
 		EditorOnly->OpacityMask.Connect(5, ConvertAttributeNode);
@@ -3309,7 +3310,12 @@ void UMaterial::ConvertMaterialToSubstrateMaterial()
 		EditorOnly->AmbientOcclusion.Connect(2, ConvertAttributeNode);
 		EditorOnly->PixelDepthOffset.Connect(1, ConvertAttributeNode);
 		EditorOnly->Refraction.Connect(6, ConvertAttributeNode);
-		
+
+		// Need to connect all CustomizedUV, as NumCustomizedUVs might be 0, while legacy attribute would connect & use all the customized UVs inputs.
+		for (int32 UVIndex = 0; UVIndex<NumCustomizedUVs; ++UVIndex)
+		{
+			EditorOnly->CustomizedUVs[UVIndex].Connect(7+UVIndex, ConvertAttributeNode);
+		}
 
 		// Shading Model
 		// * either use the shader graph expression 
diff --git a/Engine/Source/Runtime/Engine/Private/Materials/MaterialExpressions.cpp b/Engine/Source/Runtime/Engine/Private/Materials/MaterialExpressions.cpp
--- a/Engine/Source/Runtime/Engine/Private/Materials/MaterialExpressions.cpp
+++ b/Engine/Source/Runtime/Engine/Private/Materials/MaterialExpressions.cpp
@@ -26408,6 +26408,12 @@ UMaterialExpressionSubstrateConvertMaterialAttributes::UMaterialExpressionSubstr
 	Outputs.Add(FExpressionOutput(TEXT("Opacity")));
 	Outputs.Add(FExpressionOutput(TEXT("OpacityMask")));
 	Outputs.Add(FExpressionOutput(TEXT("Refraction")));
+
+	for (int32 UVIndex = 0; UVIndex <= MP_CustomizedUVs7 - MP_CustomizedUVs0; UVIndex++)
+	{
+		Outputs.Add(FExpressionOutput(*FString::Printf(TEXT("CustomizedUV%u"), UVIndex), 1, 1, 1, 0, 0));
+	}
+
 	bShowOutputNameOnPin = true;
 #endif
 }
@@ -26439,6 +26445,12 @@ int32 UMaterialExpressionSubstrateConvertMaterialAttributes::Compile(class FMate
 	{
 		return MaterialAttributes.CompileWithDefault(Compiler, FMaterialAttributeDefinitionMap::GetID(MP_Refraction));
 	}
+	
+	else if (OutputIndex >= 7 && OutputIndex <= 7 + (MP_CustomizedUVs7 - MP_CustomizedUVs0))
+	{
+		const uint32 UVIndex = OutputIndex-7u;
+		return MaterialAttributes.CompileWithDefault(Compiler, FMaterialAttributeDefinitionMap::GetID(EMaterialProperty(MP_CustomizedUVs0 + UVIndex)));
+	}
 	else if (OutputIndex != 0)
 	{
 		return Compiler->Error(TEXT("Output pin index error"));
@@ -26586,13 +26598,21 @@ uint32 UMaterialExpressionSubstrateConvertMaterialAttributes::GetOutputType(int3
 {
 	switch (OutputIndex)
 	{
-		case 0: return MCT_Substrate;	// Substrata data
-		case 1: return MCT_Float1;		// PixelDepthOffset
-		case 2: return MCT_Float1; 		// AmbientOcclusion
-		case 3: return MCT_Float3; 		// WorldPositionOffset
-		case 4: return MCT_Float1; 		// Opacity
-		case 5: return MCT_Float1; 		// OpacityMask
-		case 6: return MCT_Float; 		// Refraction
+		case 0 : return MCT_Substrate;		// Substrata data
+		case 1 : return MCT_Float1;			// PixelDepthOffset
+		case 2 : return MCT_Float1; 		// AmbientOcclusion
+		case 3 : return MCT_Float3; 		// WorldPositionOffset
+		case 4 : return MCT_Float1; 		// Opacity
+		case 5 : return MCT_Float1; 		// OpacityMask
+		case 6 : return MCT_Float; 			// Refraction
+		case 7 : return MCT_Float2;			// MP_CustomizedUVs0
+		case 8 : return MCT_Float2;			// MP_CustomizedUVs1
+		case 9 : return MCT_Float2;			// MP_CustomizedUVs2
+		case 10: return MCT_Float2;			// MP_CustomizedUVs3
+		case 11: return MCT_Float2;			// MP_CustomizedUVs4
+		case 12: return MCT_Float2;			// MP_CustomizedUVs5
+		case 13: return MCT_Float2;			// MP_CustomizedUVs6
+		case 14: return MCT_Float2; 		// MP_CustomizedUVs7
 	}
 	check(false);
 	return MCT_Float1;
