Substrate - Added closure count limit for performance control.
Also added in editor UI for that.

see also https://horde.devtools.epicgames.com/job/65004b61563726dac8c16b5c

diff --git a/Engine/Source/Editor/MaterialEditor/Private/SMaterialEditorSubstrateWidget.cpp b/Engine/Source/Editor/MaterialEditor/Private/SMaterialEditorSubstrateWidget.cpp
--- a/Engine/Source/Editor/MaterialEditor/Private/SMaterialEditorSubstrateWidget.cpp
+++ b/Engine/Source/Editor/MaterialEditor/Private/SMaterialEditorSubstrateWidget.cpp
@@ -62,6 +62,26 @@ void SMaterialEditorSubstrateWidget::Construct(const FArguments& InArgs, TWeakPt
 		.ToolTipText(LOCTEXT("CheckBoxBytesPerPixelOverride", "This will force the byte per pixel count for the preview material. It cannot go higher than the current project setting."))
 		.OnCheckStateChanged(this, &SMaterialEditorSubstrateWidget::OnCheckBoxBytesPerPixelChanged);
 
+	ClosuresPerPixelOverride = Substrate::GetClosurePerPixel(SP_PCD3D_SM5);
+	ClosuresPerPixelOverrideInput = SNew(SNumericEntryBox<uint32>)
+		.MinDesiredValueWidth(150.0f)
+		.MinValue(1)
+		.MaxValue(ClosuresPerPixelOverride)
+		.MinSliderValue(1)
+		.MaxSliderValue(ClosuresPerPixelOverride)
+		.OnBeginSliderMovement(this, &SMaterialEditorSubstrateWidget::OnBeginClosuresPerPixelSliderMovement)
+		.OnEndSliderMovement(this, &SMaterialEditorSubstrateWidget::OnEndClosuresPerPixelSliderMovement)
+		.AllowSpin(true)
+		.OnValueChanged(this, &SMaterialEditorSubstrateWidget::OnClosuresPerPixelChanged)
+		.OnValueCommitted(this, &SMaterialEditorSubstrateWidget::OnClosuresPerPixelCommitted)
+		.Value(this, &SMaterialEditorSubstrateWidget::GetClosuresPerPixelValue)
+		.IsEnabled(false);
+
+	CheckBoxClosuresPerPixelOverride = SNew(SCheckBox)
+		.Padding(5.0f)
+		.ToolTipText(LOCTEXT("CheckBoxClosuresPerPixelOverride", "This will force the byte per pixel count for the preview material. It cannot go higher than the current project setting."))
+		.OnCheckStateChanged(this, &SMaterialEditorSubstrateWidget::OnCheckBoxClosuresPerPixelChanged);
+
 	if (Substrate::IsSubstrateEnabled())
 	{
 		this->ChildSlot
@@ -176,6 +196,49 @@ void SMaterialEditorSubstrateWidget::Construct(const FArguments& InArgs, TWeakPt
 							]
 							+SHorizontalBox::Slot()
 							.AutoWidth()
+							.Padding(20.0, 0.0, 0.0, 0.0)
+							.HAlign(HAlign_Left)
+							.VAlign(VAlign_Center)
+							[
+								SNew(SWrapBox)
+								.UseAllottedSize(true)
+								+SWrapBox::Slot()
+								.Padding(5.0f)
+								.HAlign(HAlign_Left)
+								.VAlign(VAlign_Center)
+								[
+									CheckBoxClosuresPerPixelOverride->AsShared()
+								]
+							]
+							+SHorizontalBox::Slot()
+							.AutoWidth()
+							.Padding(0.0f)
+							.HAlign(HAlign_Left)
+							.VAlign(VAlign_Center)
+							[
+								SNew(STextBlock)
+								.ColorAndOpacity(FLinearColor::White)
+								.ShadowColorAndOpacity(FLinearColor::Black)
+								.ShadowOffset(FVector2D::UnitVector)
+								.Text(LOCTEXT("OverrideClosuresPerPixel", "Override closures per pixel"))
+							]
+							+SHorizontalBox::Slot()
+							.MaxWidth(200)
+							.HAlign(HAlign_Left)
+							.VAlign(VAlign_Center)
+							[
+								SNew(SWrapBox)
+								.UseAllottedSize(true)
+								+SWrapBox::Slot()
+								.Padding(0.0, 0.0, 0.0, 0.0)
+								.HAlign(HAlign_Left)
+								.VAlign(VAlign_Center)
+								[
+									ClosuresPerPixelOverrideInput->AsShared()
+								]
+							]
+							+SHorizontalBox::Slot()
+							.AutoWidth()
 							.Padding(0.0f)
 							.HAlign(HAlign_Left)
 							.VAlign(VAlign_Center)
@@ -344,8 +407,9 @@ void SMaterialEditorSubstrateWidget::Tick(const FGeometry& AllottedGeometry, con
 
 				if (CompilationOutput.bMaterialOutOfBudgetHasBeenSimplified)
 				{
-					MaterialDescription += FString::Printf(TEXT("The material was OUT-OF-BUDGET so it has been fully simplified: request bytes = %i / budget = %i\r\n"),
-						CompilationOutput.RequestedBytePixePixel, CompilationOutput.PlatformBytePixePixel);
+					MaterialDescription += FString::Printf(TEXT("The material was OUT-OF-BUDGET so it has been fully simplified: Request bytes = %i / budget = %i  -  Request Closures = %i / budget = %i\r\n"),
+						CompilationOutput.RequestedBytePixePixel, CompilationOutput.PlatformBytePixePixel,
+						CompilationOutput.RequestedClosurePerPixel, CompilationOutput.PlatformClosurePixel);
 					MaterialDescription += FString::Printf(TEXT("Final per pixel byte count   = %i\r\n"),
 						FinalPixelByteCount);
 				}
@@ -514,6 +578,47 @@ void SMaterialEditorSubstrateWidget::OnCheckBoxBytesPerPixelChanged(ECheckBoxSta
 	BytesPerPixelOverrideInput->SetEnabled(InCheckBoxState == ECheckBoxState::Checked);
 }
 
+void SMaterialEditorSubstrateWidget::OnClosuresPerPixelChanged(uint32 NewValue)
+{
+	ClosuresPerPixelOverride = NewValue;
+}
+
+void SMaterialEditorSubstrateWidget::OnClosuresPerPixelCommitted(uint32 NewValue, ETextCommit::Type InCommitType)
+{
+	if (InCommitType == ETextCommit::OnEnter)
+	{
+		ClosuresPerPixelOverride = NewValue;
+	}
+}
+
+void SMaterialEditorSubstrateWidget::OnBeginClosuresPerPixelSliderMovement()
+{
+	if (bClosuresPerPixelStartedTransaction == false)
+	{
+		bClosuresPerPixelStartedTransaction = true;
+		GEditor->BeginTransaction(LOCTEXT("PastePoseTransation", "Paste Pose"));
+	}
+}
+void SMaterialEditorSubstrateWidget::OnEndClosuresPerPixelSliderMovement(uint32 NewValue)
+{
+	if (bClosuresPerPixelStartedTransaction)
+	{
+		GEditor->EndTransaction();
+		bClosuresPerPixelStartedTransaction = false;
+		ClosuresPerPixelOverride = NewValue;
+	}
+}
+
+TOptional<uint32> SMaterialEditorSubstrateWidget::GetClosuresPerPixelValue() const
+{
+	return ClosuresPerPixelOverride;
+}
+
+void SMaterialEditorSubstrateWidget::OnCheckBoxClosuresPerPixelChanged(ECheckBoxState InCheckBoxState)
+{
+	ClosuresPerPixelOverrideInput->SetEnabled(InCheckBoxState == ECheckBoxState::Checked);
+}
+
 FReply SMaterialEditorSubstrateWidget::OnButtonApplyToPreview()
 {
 	if (MaterialEditorPtr.IsValid())
@@ -523,8 +628,9 @@ FReply SMaterialEditorSubstrateWidget::OnButtonApplyToPreview()
 		FSubstrateCompilationConfig SubstrateCompilationConfig;
 		SubstrateCompilationConfig.bFullSimplify = CheckBoxForceFullSimplification->IsChecked();
 
-		// Have a look at MaterialStats.cpp when we want to visualise a specific platform.
-		SubstrateCompilationConfig.BytesPerPixelOverride = CheckBoxBytesPerPixelOverride->IsChecked() ? FMath::Clamp(BytesPerPixelOverride, 12, Substrate::GetBytePerPixel(GMaxRHIShaderPlatform)) : -1;
+		// Have a look at MaterialStats.cpp when we want to visualize a specific platform.
+		SubstrateCompilationConfig.BytesPerPixelOverride	= CheckBoxBytesPerPixelOverride->IsChecked()	? FMath::Clamp(BytesPerPixelOverride,		12,	Substrate::GetBytePerPixel(GMaxRHIShaderPlatform))		: -1;
+		SubstrateCompilationConfig.ClosuresPerPixelOverride = CheckBoxClosuresPerPixelOverride->IsChecked() ? FMath::Clamp(ClosuresPerPixelOverride,	1,	Substrate::GetClosurePerPixel(GMaxRHIShaderPlatform))	: -1;
 
 		MaterialInterface->SetSubstrateCompilationConfig(SubstrateCompilationConfig);
 
diff --git a/Engine/Source/Editor/MaterialEditor/Private/SMaterialEditorSubstrateWidget.h b/Engine/Source/Editor/MaterialEditor/Private/SMaterialEditorSubstrateWidget.h
--- a/Engine/Source/Editor/MaterialEditor/Private/SMaterialEditorSubstrateWidget.h
+++ b/Engine/Source/Editor/MaterialEditor/Private/SMaterialEditorSubstrateWidget.h
@@ -40,8 +40,10 @@ private:
 
 	TSharedPtr<class SCheckBox> CheckBoxForceFullSimplification;
 	TSharedPtr<class SCheckBox> CheckBoxBytesPerPixelOverride;
+	TSharedPtr<class SCheckBox> CheckBoxClosuresPerPixelOverride;
 
 	TSharedPtr<class SNumericEntryBox<uint32>> BytesPerPixelOverrideInput;
+	TSharedPtr<class SNumericEntryBox<uint32>> ClosuresPerPixelOverrideInput;
 
 	TSharedPtr<class SButton> ButtonApplyToPreview;
 
@@ -65,5 +67,15 @@ private:
 	TOptional<uint32> GetBytesPerPixelValue() const;
 
 	void OnCheckBoxBytesPerPixelChanged(ECheckBoxState InCheckBoxState);
+
+	bool bClosuresPerPixelStartedTransaction;
+	uint32 ClosuresPerPixelOverride;
+	void OnClosuresPerPixelChanged(uint32 NewValue);
+	void OnClosuresPerPixelCommitted(uint32 NewValue, ETextCommit::Type InCommitType);
+	void OnBeginClosuresPerPixelSliderMovement();
+	void OnEndClosuresPerPixelSliderMovement(uint32 NewValue);
+	TOptional<uint32> GetClosuresPerPixelValue() const;
+
+	void OnCheckBoxClosuresPerPixelChanged(ECheckBoxState InCheckBoxState);
 };
 
diff --git a/Engine/Source/Runtime/Engine/Private/Materials/HLSLMaterialTranslator.cpp b/Engine/Source/Runtime/Engine/Private/Materials/HLSLMaterialTranslator.cpp
--- a/Engine/Source/Runtime/Engine/Private/Materials/HLSLMaterialTranslator.cpp
+++ b/Engine/Source/Runtime/Engine/Private/Materials/HLSLMaterialTranslator.cpp
@@ -2331,6 +2331,7 @@ void FHLSLMaterialTranslator::GetMaterialEnvironment(EShaderPlatform InPlatform,
 			{
 				// Output some debug info as comment in code and in the material stat window
 				const uint32 SubstrateBytePerPixel = Substrate::GetBytePerPixel(InPlatform);
+				const uint32 SubstrateClosurePerPixel = Substrate::GetClosurePerPixel(InPlatform);
 				FString SubstrateMaterialContextDescription;
 
 				auto GetSubstrateCompilationContextName = [](uint32 Index)
@@ -2350,11 +2351,14 @@ void FHLSLMaterialTranslator::GetMaterialEnvironment(EShaderPlatform InPlatform,
 
 				SubstrateMaterialContextDescription += FString::Printf(TEXT("----- SUBSTRATE - %s -----\r\n"), *SubstrateCompilationContextName);
 				SubstrateMaterialContextDescription += FString::Printf(TEXT("SubstrateCompilationInfo -\r\n"));
-				SubstrateMaterialContextDescription += FString::Printf(TEXT(" - Byte Per Pixel Budget                      %u\r\n"), SubstrateBytePerPixel);
-				SubstrateMaterialContextDescription += FString::Printf(TEXT(" - Requested Byte Size before simplification  %u (%d UINT32)\r\n"), SubstrateSimplificationStatus.OriginalRequestedByteSize, SubstrateSimplificationStatus.OriginalRequestedByteSize / 4);
-				SubstrateMaterialContextDescription += FString::Printf(TEXT(" - Requested Byte Size after simplification   %u (%d UINT32)\r\n"), SubstrateCtx.SubstrateMaterialRequestedSizeByte, SubstrateCtx.SubstrateMaterialRequestedSizeByte / 4);
-				SubstrateMaterialContextDescription += FString::Printf(TEXT(" - Material complexity                        %s\r\n"), SubstrateCtx.bSubstrateMaterialIsSingle ? TEXT("SINGLE") : (SubstrateCtx.bSubstrateMaterialIsSimple ? TEXT("SIMPLE") : TEXT("COMPLEX")));
-				SubstrateMaterialContextDescription += FString::Printf(TEXT(" - TotalBSDFCount                             %i\r\n"), SubstrateCtx.SubstrateMaterialBSDFCount);
+				SubstrateMaterialContextDescription += FString::Printf(TEXT(" - Byte Per Pixel Budget                           %u\r\n"), SubstrateBytePerPixel);
+				SubstrateMaterialContextDescription += FString::Printf(TEXT(" - Closure Per Pixel Budget                        %u\r\n"), SubstrateClosurePerPixel);
+				SubstrateMaterialContextDescription += FString::Printf(TEXT(" - Requested Byte Size before simplification       %u (%d UINT32)\r\n"), SubstrateSimplificationStatus.OriginalRequestedByteSize, SubstrateSimplificationStatus.OriginalRequestedByteSize / 4);
+				SubstrateMaterialContextDescription += FString::Printf(TEXT(" - Requested Byte Size after simplification        %u (%d UINT32)\r\n"), SubstrateCtx.SubstrateMaterialRequestedSizeByte, SubstrateCtx.SubstrateMaterialRequestedSizeByte / 4);
+				SubstrateMaterialContextDescription += FString::Printf(TEXT(" - Requested Closure Count before simplification   %u\r\n"), SubstrateSimplificationStatus.OriginalRequestedClosureCount);
+				SubstrateMaterialContextDescription += FString::Printf(TEXT(" - Requested Closure Count after simplification    %u\r\n"), SubstrateCtx.SubstrateMaterialClosureCount);
+				SubstrateMaterialContextDescription += FString::Printf(TEXT(" - Material complexity                             %s\r\n"), SubstrateCtx.bSubstrateMaterialIsSingle ? TEXT("SINGLE") : (SubstrateCtx.bSubstrateMaterialIsSimple ? TEXT("SIMPLE") : TEXT("COMPLEX")));
+				SubstrateMaterialContextDescription += FString::Printf(TEXT(" - BSDF Count                                      %i\r\n"), SubstrateCtx.SubstrateMaterialBSDFCount); // REMOVE?
 				if (RequestedSharedLocalBasesCount > SUBSTRATE_MAX_SHAREDLOCALBASES_REGISTERS)
 				{
 					SubstrateMaterialDescription += FString::Printf(TEXT(" - SharedLocalBasesCount                      %i (Requested:%i)\r\n"), SubstrateCtx.FinalUsedSharedLocalBasesCount, RequestedSharedLocalBasesCount);
@@ -11627,8 +11631,9 @@ bool FHLSLMaterialTranslator::FSubstrateCompilationContext::SubstrateGenerateDer
 		}
 	}
 
-
-	const uint32 SubstrateBytePerPixel = Compiler->SubstrateCompilationConfig.BytesPerPixelOverride > 0 ? Compiler->SubstrateCompilationConfig.BytesPerPixelOverride : Substrate::GetBytePerPixel(Compiler->GetShaderPlatform());
+	EShaderPlatform ShaderPlatform = Compiler->GetShaderPlatform();
+	const uint32 SubstrateBytePerPixel = Compiler->SubstrateCompilationConfig.BytesPerPixelOverride > 0 ? Compiler->SubstrateCompilationConfig.BytesPerPixelOverride : Substrate::GetBytePerPixel(ShaderPlatform);
+	const uint32 SubstrateClosurePerPixel = Compiler->SubstrateCompilationConfig.ClosuresPerPixelOverride > 0 ? Compiler->SubstrateCompilationConfig.ClosuresPerPixelOverride : Substrate::GetClosurePerPixel(ShaderPlatform);
 
 	bool bFirstLoop = true;
 	do 
@@ -12027,6 +12032,7 @@ bool FHLSLMaterialTranslator::FSubstrateCompilationContext::SubstrateGenerateDer
 			static_assert(SUBSTRATE_MAX_BSDF_COUNT					<= (1u << SUBSTRATE_BSDF_OFFSET_BIT_COUNT));
 			const uint32 BSDFMaxByteCountForOffset = uint32(1u << SUBSTRATE_BSDF_OFFSET_BIT_COUNT) * sizeof(uint32);
 			uint32 OperatorCount = 0;
+			SubstrateMaterialClosureCount = 0;
 			for (auto& It : SubstrateMaterialExpressionRegisteredOperators)
 			{
 				if (It.IsDiscarded())
@@ -12043,6 +12049,9 @@ bool FHLSLMaterialTranslator::FSubstrateCompilationContext::SubstrateGenerateDer
 				{
 				case SUBSTRATE_OPERATOR_BSDF:
 				{
+					// we have encountered a new BSDF which directly link to a single closure evaluation
+					SubstrateMaterialClosureCount++;
+
 					// From the compiler side, we can only assume the top layer has gray scale luminance weight.
 					const bool bMayHaveColoredWeight = !It.bIsTop;
 
@@ -12172,7 +12181,7 @@ bool FHLSLMaterialTranslator::FSubstrateCompilationContext::SubstrateGenerateDer
 				return false;
 			}
 
-			SubstrateSimplificationStatus.bMaterialFitsInMemoryBudget = SubstrateMaterialRequestedSizeByte <= SubstrateBytePerPixel;
+			SubstrateSimplificationStatus.bMaterialFitsInMemoryBudget = (SubstrateMaterialRequestedSizeByte <= SubstrateBytePerPixel) && (SubstrateMaterialClosureCount <= SubstrateClosurePerPixel);
 			if (!SubstrateSimplificationStatus.bMaterialFitsInMemoryBudget && SubstrateSimplificationStatus.bFullSimplificationStepHasBeenRun && SubstrateSimplificationStatus.bSlabSimplificationStepHasBeenRun)
 			{
 				// If we have already run the full simplification but the material still does not fit in memory, we must fail the material compilation.
@@ -12183,6 +12192,7 @@ bool FHLSLMaterialTranslator::FSubstrateCompilationContext::SubstrateGenerateDer
 			{
 				// Record the original requested byte size before simplification, only for the first pass.
 				SubstrateSimplificationStatus.OriginalRequestedByteSize = SubstrateMaterialRequestedSizeByte;
+				SubstrateSimplificationStatus.OriginalRequestedClosureCount = SubstrateMaterialClosureCount;
 			}
 			SubstrateSimplificationStatus.bFullSimplificationStepHasBeenRun |= SubstrateSimplificationStatus.bRunFullSimplification;
 
@@ -12202,6 +12212,9 @@ bool FHLSLMaterialTranslator::FSubstrateCompilationContext::SubstrateGenerateDer
 				Compiler->MaterialCompilationOutput.SubstrateMaterialCompilationOutput.RequestedBytePixePixel = SubstrateMaterialRequestedSizeByte;
 				Compiler->MaterialCompilationOutput.SubstrateMaterialCompilationOutput.PlatformBytePixePixel = SubstrateBytePerPixel;
 
+				Compiler->MaterialCompilationOutput.SubstrateMaterialCompilationOutput.RequestedClosurePerPixel = SubstrateSimplificationStatus.OriginalRequestedClosureCount;
+				Compiler->MaterialCompilationOutput.SubstrateMaterialCompilationOutput.PlatformClosurePixel = SubstrateMaterialClosureCount;
+
 				Compiler->MaterialCompilationOutput.SubstrateMaterialCompilationOutput.bIsThin = CompilerMaterial->IsThinSurface() ? 1 : 0;
 
 				// The order of ifs here is important.
diff --git a/Engine/Source/Runtime/Engine/Private/Materials/HLSLMaterialTranslator.h b/Engine/Source/Runtime/Engine/Private/Materials/HLSLMaterialTranslator.h
--- a/Engine/Source/Runtime/Engine/Private/Materials/HLSLMaterialTranslator.h
+++ b/Engine/Source/Runtime/Engine/Private/Materials/HLSLMaterialTranslator.h
@@ -432,6 +432,7 @@ protected:
 		bool bMaterialFitsInMemoryBudget = false;	// Track whether or not the material fits.
 
 		uint32 OriginalRequestedByteSize = 0;
+		uint32 OriginalRequestedClosureCount = 0;
 		bool bRunFullSimplification = false;	// Simple implementation for now: if the material does not fit, we simply everything.
 		bool bFullSimplificationStepHasBeenRun = false;
 
@@ -491,8 +492,9 @@ protected:
 		TMap<FGuid, int32> SubstrateMaterialExpressionToOperatorIndex;
 		TArray<FSubstrateOperator> SubstrateMaterialExpressionRegisteredOperators;
 		FSubstrateOperator* SubstrateMaterialRootOperator;
-		uint32 SubstrateMaterialBSDFCount;
+		uint32 SubstrateMaterialBSDFCount; // Also acts as requested 
 		uint32 SubstrateMaterialRequestedSizeByte;
+		uint32 SubstrateMaterialClosureCount;
 		bool bSubstrateMaterialIsSimple;
 		bool bSubstrateMaterialIsSingle;
 		bool bSubstrateMaterialIsUnlitNode;
diff --git a/Engine/Source/Runtime/Engine/Private/Materials/MaterialShader.cpp b/Engine/Source/Runtime/Engine/Private/Materials/MaterialShader.cpp
--- a/Engine/Source/Runtime/Engine/Private/Materials/MaterialShader.cpp
+++ b/Engine/Source/Runtime/Engine/Private/Materials/MaterialShader.cpp
@@ -576,23 +576,25 @@ void FStaticParameterSet::SetStaticComponentMaskParameterValue(const FMaterialPa
 #if WITH_EDITOR
 FString FSubstrateCompilationConfig::GetShaderMapKeyString() const
 {
-	if (bFullSimplify)
-	{
-		return FString::Printf(TEXT("_STRTFS%i"), BytesPerPixelOverride < 0 ? 0 : BytesPerPixelOverride);
-	}
-	return TEXT("");
+	FString SubStrateCompStr;
+	if (bFullSimplify)					SubStrateCompStr += TEXT("_SBSTRFS");
+	if (BytesPerPixelOverride >= 0)		SubStrateCompStr += TEXT("_SBSTRBS");
+	if (ClosuresPerPixelOverride >= 0)	SubStrateCompStr += TEXT("_SBSTRCS");
+	return SubStrateCompStr;
 }
 
 void FSubstrateCompilationConfig::UpdateHash(FSHA1& Hasher) const
 {
 	Hasher.Update((const uint8*)(&bFullSimplify), sizeof(bFullSimplify));
 	Hasher.Update((const uint8*)(&BytesPerPixelOverride), sizeof(BytesPerPixelOverride));
+	Hasher.Update((const uint8*)(&ClosuresPerPixelOverride), sizeof(ClosuresPerPixelOverride));
 }
 
 void FSubstrateCompilationConfig::Serialize(FArchive& Ar)
 {
 	Ar << bFullSimplify;
 	Ar << BytesPerPixelOverride;
+	Ar << ClosuresPerPixelOverride;
 }
 #endif
 
diff --git a/Engine/Source/Runtime/Engine/Private/ShaderCompiler/ShaderCompiler.cpp b/Engine/Source/Runtime/Engine/Private/ShaderCompiler/ShaderCompiler.cpp
--- a/Engine/Source/Runtime/Engine/Private/ShaderCompiler/ShaderCompiler.cpp
+++ b/Engine/Source/Runtime/Engine/Private/ShaderCompiler/ShaderCompiler.cpp
@@ -7982,6 +7982,9 @@ void GlobalBeginCompileShader(
 			const uint32 SubstrateUintPerPixel = Substrate::GetBytePerPixel(Target.GetPlatform()) / 4u;
 			SET_SHADER_DEFINE(Input.Environment, SUBSTRATE_MATERIAL_NUM_UINTS, SubstrateUintPerPixel);
 
+			const uint32 SubstrateClosurePerPixel = Substrate::GetClosurePerPixel(Target.GetPlatform());
+			SET_SHADER_DEFINE(Input.Environment, SUBSTRATE_MATERIAL_CLOSURE_COUNT, SubstrateClosurePerPixel);
+
 			const bool bTileCoord8Bits = Substrate::Is8bitTileCoordEnabled();
 			SET_SHADER_DEFINE(Input.Environment, USE_8BIT_TILE_COORD, bTileCoord8Bits ? 1 : 0);
 
diff --git a/Engine/Source/Runtime/Engine/Public/MaterialTypes.h b/Engine/Source/Runtime/Engine/Public/MaterialTypes.h
--- a/Engine/Source/Runtime/Engine/Public/MaterialTypes.h
+++ b/Engine/Source/Runtime/Engine/Public/MaterialTypes.h
@@ -466,6 +466,7 @@ struct FSubstrateCompilationConfig
 {
 	bool bFullSimplify = false;
 	int16 BytesPerPixelOverride = -1;
+	int16 ClosuresPerPixelOverride = -1;
 
 	FString GetShaderMapKeyString() const;
 
@@ -475,7 +476,7 @@ struct FSubstrateCompilationConfig
 
 	friend inline bool operator==(const FSubstrateCompilationConfig& Lhs, const FSubstrateCompilationConfig& Rhs)
 	{
-		return Lhs.bFullSimplify == Rhs.bFullSimplify && Lhs.BytesPerPixelOverride == Rhs.BytesPerPixelOverride;
+		return Lhs.bFullSimplify == Rhs.bFullSimplify && Lhs.BytesPerPixelOverride == Rhs.BytesPerPixelOverride && Lhs.ClosuresPerPixelOverride == Rhs.ClosuresPerPixelOverride;
 	}
 
 	friend inline bool operator!=(const FSubstrateCompilationConfig& Lhs, const FSubstrateCompilationConfig& Rhs)
diff --git a/Engine/Source/Runtime/Engine/Public/Rendering/SubstrateMaterialShared.cpp b/Engine/Source/Runtime/Engine/Public/Rendering/SubstrateMaterialShared.cpp
--- a/Engine/Source/Runtime/Engine/Public/Rendering/SubstrateMaterialShared.cpp
+++ b/Engine/Source/Runtime/Engine/Public/Rendering/SubstrateMaterialShared.cpp
@@ -15,6 +15,8 @@ FSubstrateMaterialCompilationOutput::FSubstrateMaterialCompilationOutput()
 	, SharedLocalBasesCount(0)
 	, RequestedBytePixePixel(0)
 	, PlatformBytePixePixel(0)
+	, RequestedClosurePerPixel(0)
+	, PlatformClosurePixel(0)
 	, bIsThin(0)
 	, MaterialType(0)
 	, bMaterialOutOfBudgetHasBeenSimplified(0)
diff --git a/Engine/Source/Runtime/Engine/Public/Rendering/SubstrateMaterialShared.h b/Engine/Source/Runtime/Engine/Public/Rendering/SubstrateMaterialShared.h
--- a/Engine/Source/Runtime/Engine/Public/Rendering/SubstrateMaterialShared.h
+++ b/Engine/Source/Runtime/Engine/Public/Rendering/SubstrateMaterialShared.h
@@ -129,6 +129,11 @@ public:
 	/** The byte count per pixel supported by the platform the material has been compiled against */
 	LAYOUT_FIELD_EDITORONLY(uint8, PlatformBytePixePixel);
 
+	/** Material requested BSDF/closure count per pixel */
+	LAYOUT_FIELD_EDITORONLY(uint8, RequestedClosurePerPixel);
+	/** The closure count per pixel supported by the platform the material has been compiled against */
+	LAYOUT_FIELD_EDITORONLY(uint8, PlatformClosurePixel);
+
 	/** Indicate that the material is considered a thin surface instead of a volume filled up with matter */
 	LAYOUT_FIELD_EDITORONLY(uint8, bIsThin);
 	/** Indicate the final material type */
diff --git a/Engine/Source/Runtime/RenderCore/Private/RenderUtils.cpp b/Engine/Source/Runtime/RenderCore/Private/RenderUtils.cpp
--- a/Engine/Source/Runtime/RenderCore/Private/RenderUtils.cpp
+++ b/Engine/Source/Runtime/RenderCore/Private/RenderUtils.cpp
@@ -1606,12 +1606,18 @@ static TAutoConsoleVariable<int32> CVarSubstrate(
 	TEXT("Enable Substrate materials (Beta)."),
 	ECVF_ReadOnly | ECVF_RenderThreadSafe);
 
-static TAutoConsoleVariable<int32> CVarSubstrateBytePerPixel(
+static TAutoConsoleVariable<int32> CVarSubstrateBytesPerPixel(
 	TEXT("r.Substrate.BytesPerPixel"),
 	80,
 	TEXT("Substrate allocated byte per pixel to store materials data. Higher value means more complex material can be represented."),
 	ECVF_ReadOnly | ECVF_RenderThreadSafe);
 
+static TAutoConsoleVariable<int32> CVarSubstrateClosuresPerPixel(
+	TEXT("r.Substrate.ClosuresPerPixel"),
+	8,	// Similar to SUBSTRATE_MAX_BSDF_COUNT
+	TEXT("Substrate closure count per pixel can be constrained. That is usefull to reduce the number of byte written durting the base pass, but also to limit the number of closures in Forward. Higher value means more complex material can be represented."),
+	ECVF_ReadOnly | ECVF_RenderThreadSafe);
+
 static TAutoConsoleVariable<int32> CVarSubstrateBackCompatibility(
 	TEXT("r.SubstrateBackCompatibility"),
 	0,
@@ -1724,14 +1730,21 @@ namespace Substrate
 
 	uint32 GetBytePerPixel()
 	{
-		return InternalGetBytePerPixel(CVarSubstrateBytePerPixel.GetValueOnAnyThread());
+		return InternalGetBytePerPixel(CVarSubstrateBytesPerPixel.GetValueOnAnyThread());
 	}
 
 	uint32 GetBytePerPixel(EShaderPlatform InPlatform)
 	{
 		// Variant for shader compilation per platform
-		static FShaderPlatformCachedIniValue<int32> CVarBudget(TEXT("r.Substrate.BytesPerPixel"));
-		return InternalGetBytePerPixel(CVarBudget.Get(InPlatform));
+		static FShaderPlatformCachedIniValue<int32> CVarByteBudget(TEXT("r.Substrate.BytesPerPixel"));
+		return InternalGetBytePerPixel(CVarByteBudget.Get(InPlatform));
+	}
+
+	uint32 GetClosurePerPixel(EShaderPlatform InPlatform)
+	{
+		// Variant for shader compilation per platform
+		static FShaderPlatformCachedIniValue<int32> CVarClosureBudget(TEXT("r.Substrate.ClosuresPerPixel"));
+		return uint32(FMath::Max(0, CVarClosureBudget.Get(InPlatform)));
 	}
 
 	uint32 GetNormalQuality()
diff --git a/Engine/Source/Runtime/RenderCore/Private/Shader.cpp b/Engine/Source/Runtime/RenderCore/Private/Shader.cpp
--- a/Engine/Source/Runtime/RenderCore/Private/Shader.cpp
+++ b/Engine/Source/Runtime/RenderCore/Private/Shader.cpp
@@ -1998,6 +1998,10 @@ void ShaderMapAppendKeyString(EShaderPlatform Platform, FString& KeyString)
 			KeyString += FString::Printf(TEXT("_BUDGET%u"), Substrate::GetBytePerPixel(Platform));
 		}
 
+		{
+			KeyString += FString::Printf(TEXT("_CLOSURE%u"), Substrate::GetClosurePerPixel(Platform));
+		}
+
 		if (Substrate::IsDBufferPassEnabled(Platform))
 		{
 			KeyString += FString::Printf(TEXT("_DBUFFERPASS"));
diff --git a/Engine/Source/Runtime/RenderCore/Public/RenderUtils.h b/Engine/Source/Runtime/RenderCore/Public/RenderUtils.h
--- a/Engine/Source/Runtime/RenderCore/Public/RenderUtils.h
+++ b/Engine/Source/Runtime/RenderCore/Public/RenderUtils.h
@@ -503,6 +503,7 @@ namespace Substrate
 
 	RENDERCORE_API uint32 GetBytePerPixel();
 	RENDERCORE_API uint32 GetBytePerPixel(EShaderPlatform InPlatform);
+	RENDERCORE_API uint32 GetClosurePerPixel(EShaderPlatform InPlatform);
 
 	RENDERCORE_API uint32 GetNormalQuality();
 
