Fix SSM_UI overflowing FSubstrateMaterialInfo Shading model causing UI material to be flagged as invalid.

diff --git a/Engine/Source/Runtime/Engine/Classes/Engine/EngineTypes.h b/Engine/Source/Runtime/Engine/Classes/Engine/EngineTypes.h
--- a/Engine/Source/Runtime/Engine/Classes/Engine/EngineTypes.h
+++ b/Engine/Source/Runtime/Engine/Classes/Engine/EngineTypes.h
@@ -721,10 +721,10 @@ public:
 	FSubstrateMaterialInfo(ESubstrateShadingModel InShadingModel) { AddShadingModel(InShadingModel); }
 
 	// Shading model
-	void AddShadingModel(ESubstrateShadingModel InShadingModel) { check(InShadingModel < SSM_NUM); ShadingModelField |= (uint16)(1 << (uint16)InShadingModel); }
-	void SetSingleShadingModel(ESubstrateShadingModel InShadingModel) { check(InShadingModel < SSM_NUM); ShadingModelField = (uint16)(1 << (uint16)InShadingModel); }
-	bool HasShadingModel(ESubstrateShadingModel InShadingModel) const { return (ShadingModelField & (1 << (uint16)InShadingModel)) != 0; }
-	bool HasOnlyShadingModel(ESubstrateShadingModel InShadingModel) const { return ShadingModelField == (1 << (uint16)InShadingModel); }
+	void AddShadingModel(ESubstrateShadingModel InShadingModel) { check(InShadingModel < SSM_NUM); ShadingModelField |= (1u << (uint32)InShadingModel); }
+	void SetSingleShadingModel(ESubstrateShadingModel InShadingModel) { check(InShadingModel < SSM_NUM); ShadingModelField = (1u << (uint32)InShadingModel); }
+	bool HasShadingModel(ESubstrateShadingModel InShadingModel) const { return (ShadingModelField & (1u << (uint32)InShadingModel)) != 0; }
+	bool HasOnlyShadingModel(ESubstrateShadingModel InShadingModel) const { return ShadingModelField == (1u << (uint32)InShadingModel); }
 	uint32 GetShadingModelField() const { return ShadingModelField; }
 	int32 CountShadingModels() const { return FMath::CountBits(ShadingModelField); }
 
