Fix invalid warning for default deferred decal.

diff --git a/Engine/Source/Runtime/Engine/Private/Materials/Material.cpp b/Engine/Source/Runtime/Engine/Private/Materials/Material.cpp
--- a/Engine/Source/Runtime/Engine/Private/Materials/Material.cpp
+++ b/Engine/Source/Runtime/Engine/Private/Materials/Material.cpp
@@ -4940,6 +4940,14 @@ void UMaterial::RebuildShadingModelField()
 					bSanitizeMaterial = true;
 				}
 			}
+			else if (SubstrateMaterialInfo.CountShadingModels() == 2 && SubstrateMaterialInfo.HasShadingModel(ESubstrateShadingModel::SSM_Decal))
+			{
+				// If material has SSM_Decal it has to have 'decal' domain and DefaultLit shading model
+				if (MaterialDomain != MD_DeferredDecal || !SubstrateMaterialInfo.HasShadingModel(SSM_DefaultLit))
+				{
+					bSanitizeMaterial = true;
+				}
+			}
 			else
 			{
 				// Clear the material to default Lit
