LevelInstance: Fix regression with recursive pinning (CL25426119)

#rnx

diff --git a/Engine/Source/Runtime/Engine/Private/WorldPartition/LoaderAdapter/LoaderAdapterPinnedActors.cpp b/Engine/Source/Runtime/Engine/Private/WorldPartition/LoaderAdapter/LoaderAdapterPinnedActors.cpp
--- a/Engine/Source/Runtime/Engine/Private/WorldPartition/LoaderAdapter/LoaderAdapterPinnedActors.cpp
+++ b/Engine/Source/Runtime/Engine/Private/WorldPartition/LoaderAdapter/LoaderAdapterPinnedActors.cpp
@@ -10,38 +10,52 @@
 
 #define LOCTEXT_NAMESPACE "FLoaderAdapterPinnedActors"
 
-bool FLoaderAdapterPinnedActors::PassActorDescFilter(const FWorldPartitionHandle& ActorHandle) const
+namespace LoaderAdapterPinnedActorsUtils
 {
-	// We want to be able to pin any type of actors (HLODs, etc).
-	return ActorHandle.IsValid() && !ActorsToRemove.Contains(ActorHandle) && SupportsPinning(ActorHandle.Get());
-}
-
-bool FLoaderAdapterPinnedActors::SupportsPinning(FWorldPartitionActorDesc* InActorDesc)
-{
-	if (!InActorDesc)
+	static bool SupportsPinning(FWorldPartitionActorDesc* InActorDesc, bool bCheckIsMainWorldPartition)
 	{
-		return false;
-	}
+		if (!InActorDesc)
+		{
+			return false;
+		}
 
-	// Only Spatially loaded actors can be pinned with the exception of non spatially loaded, runtime only actors (ex: HLODs)
-	if (!InActorDesc->GetIsSpatiallyLoaded() && !InActorDesc->GetActorIsRuntimeOnly())
-	{
-		return false;
-	}
+		// Only Spatially loaded actors can be pinned with the exception of non spatially loaded runtime only actors (ex: HLODs)
+		if (!InActorDesc->GetIsSpatiallyLoaded() && !InActorDesc->GetActorIsRuntimeOnly())
+		{
+			return false;
+		}
+
+		// This allows skipping actors that can never be loaded in current context: ex: bActorShouldSkipLevelInstance
+		if (!InActorDesc->IsEditorRelevant() && !InActorDesc->GetActorIsRuntimeOnly())
+		{
+			return false;
+		}
+
+		if (UActorDescContainer* Container = InActorDesc->GetContainer())
+		{
+			const UWorldPartition* ContainerWorldPartition = Container->GetWorldPartition();
+			return ContainerWorldPartition && (ContainerWorldPartition->IsMainWorldPartition() || !bCheckIsMainWorldPartition);
+		}
 
-	// This allows skipping actors that can never be loaded in current context: ex: bActorShouldSkipLevelInstance
-	if (!InActorDesc->IsEditorRelevant() && !InActorDesc->GetActorIsRuntimeOnly())
-	{
 		return false;
 	}
+}
 
-	if (UActorDescContainer* Container = InActorDesc->GetContainer())
-	{
-		const UWorldPartition* ContainerWorldPartition = Container->GetWorldPartition();
-		return ContainerWorldPartition && ContainerWorldPartition->IsMainWorldPartition();
-	}
+bool FLoaderAdapterPinnedActors::PassActorDescFilter(const FWorldPartitionHandle& ActorHandle) const
+{
+	// We want to be able to pin any type of actors (HLODs, etc).
+	// Allow recursive pinning by setting bCheckIsMainWorldPartition = false
+	const bool bCheckIsMainWorldPartition = false;
+	return ActorHandle.IsValid() && !ActorsToRemove.Contains(ActorHandle) && LoaderAdapterPinnedActorsUtils::SupportsPinning(ActorHandle.Get(), bCheckIsMainWorldPartition);
+}
 
-	return false;
+
+
+bool FLoaderAdapterPinnedActors::SupportsPinning(FWorldPartitionActorDesc* InActorDesc)
+{
+	// Public api doesn't allow pinning of non main world partition actors
+	const bool bCheckIsMainWorldPartition = true;
+	return LoaderAdapterPinnedActorsUtils::SupportsPinning(InActorDesc, bCheckIsMainWorldPartition);
 }
 
 bool FLoaderAdapterPinnedActors::SupportsPinning(AActor* InActor)
