Moved deferred uniform expression caching off the render thread to overlap with scene update.

diff --git a/Engine/Source/Runtime/Engine/Private/Materials/MaterialRenderProxy.cpp b/Engine/Source/Runtime/Engine/Private/Materials/MaterialRenderProxy.cpp
--- a/Engine/Source/Runtime/Engine/Private/Materials/MaterialRenderProxy.cpp
+++ b/Engine/Source/Runtime/Engine/Private/Materials/MaterialRenderProxy.cpp
@@ -60,24 +60,21 @@ public:
 		return ReferenceCount > 0 && GUniformExpressionCacheAsyncUpdates > 0 && !GRHICommandList.Bypass();
 	}
 
-	void SetTask(const FGraphEventRef& InTask)
+	void SetTask(const UE::Tasks::FTask& InTask)
 	{
-		check(!Task);
 		check(IsEnabled());
 		Task = InTask;
 	}
 
+	const UE::Tasks::FTask& GetTask() { return Task; }
+
 	void Wait()
 	{
-		if (Task)
-		{
-			Task->Wait();
-			Task = nullptr;
-		}
+		Task.Wait();
 	}
 
 private:
-	FGraphEventRef Task;
+	UE::Tasks::FTask Task;
 	int32 ReferenceCount = 0;
 
 } GUniformExpressionCacheAsyncUpdateTask;
@@ -113,24 +110,23 @@ public:
 		Items.Emplace(UniformExpressionCache, UniformExpressionSet, UniformBufferLayout, Context);
 	}
 
-	void Update(FRHICommandListImmediate& RHICmdListImmediate)
+	void Update(FRHICommandListBase& RHICmdList)
 	{
+		check(!RHICmdList.IsImmediate());
+
 		if (Items.IsEmpty())
 		{
+			RHICmdList.FinishRecording();
 			return;
 		}
 
-		GUniformExpressionCacheAsyncUpdateTask.Wait();
-
-		FRHICommandList* RHICmdList = new FRHICommandList(FRHIGPUMask::All());
-
-		FGraphEventRef Event = FFunctionGraphTask::CreateAndDispatchWhenReady(
-			[Items = MoveTemp(Items), RHICmdList]
+		UE::Tasks::FTask Task = UE::Tasks::Launch(
+			UE_SOURCE_LOCATION,
+			[Items = MoveTemp(Items), &RHICmdList]
 			{
 				TRACE_CPUPROFILER_EVENT_SCOPE(FUniformExpressionCacheAsyncUpdater::Update);
 				FTaskTagScope Scope(ETaskTag::EParallelRenderingThread);
 				FMemMark Mark(FMemStack::Get());
-				RHICmdList->SwitchPipeline(ERHIPipeline::Graphics);
 
 				for (const FItem& Item : Items)
 				{
@@ -140,16 +136,14 @@ public:
 
 					Item.UniformExpressionSet->FillUniformBuffer(Context, Item.AllocatedVTs, Item.UniformBufferLayout, TempBuffer, Item.UniformBufferLayout->ConstantBufferSize);
 
-					RHICmdList->UpdateUniformBuffer(Item.UniformBuffer, TempBuffer);
+					RHICmdList.UpdateUniformBuffer(Item.UniformBuffer, TempBuffer);
 				}
 
-				RHICmdList->FinishRecording();
-
-			}, TStatId(), nullptr, ENamedThreads::AnyHiPriThreadHiPriTask);
+				RHICmdList.FinishRecording();
 
-		RHICmdListImmediate.QueueAsyncCommandListSubmit(RHICmdList);
+			}, GUniformExpressionCacheAsyncUpdateTask.GetTask());
 
-		GUniformExpressionCacheAsyncUpdateTask.SetTask(Event);
+		GUniformExpressionCacheAsyncUpdateTask.SetTask(Task);
 	}
 
 private:
@@ -273,7 +267,7 @@ IAllocatedVirtualTexture* FMaterialRenderProxy::GetPreallocatedVTStack(const FMa
 	return Texture->GetAllocatedVirtualTexture();
 }
 
-IAllocatedVirtualTexture* FMaterialRenderProxy::AllocateVTStack(const FMaterialRenderContext& Context, const FUniformExpressionSet& UniformExpressionSet, const FMaterialVirtualTextureStack& VTStack) const
+IAllocatedVirtualTexture* FMaterialRenderProxy::AllocateVTStack(FRHICommandListBase& RHICmdList, const FMaterialRenderContext& Context, const FUniformExpressionSet& UniformExpressionSet, const FMaterialVirtualTextureStack& VTStack) const
 {
 	check(!VTStack.IsPreallocatedStack());
 	const uint32 NumLayers = VTStack.GetNumLayers();
@@ -346,13 +340,13 @@ IAllocatedVirtualTexture* FMaterialRenderProxy::AllocateVTStack(const FMaterialR
 	if (bFoundValidLayer)
 	{
 		HasVirtualTextureCallbacks = -1;
-		return GetRendererModule().AllocateVirtualTexture(VTDesc);
+		return GetRendererModule().AllocateVirtualTexture(RHICmdList, VTDesc);
 	}
 	return nullptr;
 }
 
 
-void FMaterialRenderProxy::EvaluateUniformExpressions(FUniformExpressionCache& OutUniformExpressionCache, const FMaterialRenderContext& Context, FUniformExpressionCacheAsyncUpdater* Updater, FRHICommandListBase* RHICmdList) const
+void FMaterialRenderProxy::EvaluateUniformExpressions(FRHICommandListBase& RHICmdList, FUniformExpressionCache& OutUniformExpressionCache, const FMaterialRenderContext& Context, FUniformExpressionCacheAsyncUpdater* Updater) const
 {
 	SCOPE_CYCLE_COUNTER(STAT_CacheUniformExpressions);
 
@@ -387,7 +381,7 @@ void FMaterialRenderProxy::EvaluateUniformExpressions(FUniformExpressionCache& O
 		}
 		else
 		{
-			AllocatedVT = AllocateVTStack(Context, UniformExpressionSet, VTStack);
+			AllocatedVT = AllocateVTStack(RHICmdList, Context, UniformExpressionSet, VTStack);
 			if (AllocatedVT != nullptr)
 			{
 				OutUniformExpressionCache.OwnedAllocatedVTs.Add(AllocatedVT);
@@ -427,8 +421,7 @@ void FMaterialRenderProxy::EvaluateUniformExpressions(FUniformExpressionCache& O
 
 		if (IsValidRef(OutUniformExpressionCache.UniformBuffer))
 		{
-			check(RHICmdList);
-			RHICmdList->UpdateUniformBuffer(OutUniformExpressionCache.UniformBuffer, TempBuffer);
+			RHICmdList.UpdateUniformBuffer(OutUniformExpressionCache.UniformBuffer, TempBuffer);
 		}
 		else
 		{
@@ -444,6 +437,11 @@ void FMaterialRenderProxy::EvaluateUniformExpressions(FUniformExpressionCache& O
 	++UniformExpressionCacheSerialNumber;
 }
 
+void FMaterialRenderProxy::EvaluateUniformExpressions(FUniformExpressionCache& OutUniformExpressionCache, const FMaterialRenderContext& Context, FUniformExpressionCacheAsyncUpdater* Updater) const
+{
+	EvaluateUniformExpressions(FRHICommandListImmediate::Get(), OutUniformExpressionCache, Context, Updater);
+}
+
 void FMaterialRenderProxy::CacheUniformExpressions(FRHICommandListBase& RHICmdList, bool bRecreateUniformBuffer)
 {
 	// Register the render proxy's as a render resource so it can receive notifications to free the uniform buffer.
@@ -533,7 +531,7 @@ void FMaterialRenderProxy::UpdateUniformExpressionCacheIfNeeded(FRHICommandListB
 	{
 		FMaterialRenderContext MaterialRenderContext(this, *Material, nullptr);
 		MaterialRenderContext.bShowSelection = GIsEditor;
-		EvaluateUniformExpressions(UniformExpressionCache[InFeatureLevel], MaterialRenderContext, nullptr, &RHICmdList);
+		EvaluateUniformExpressions(RHICmdList, UniformExpressionCache[InFeatureLevel], MaterialRenderContext, nullptr);
 	}
 }
 
@@ -670,48 +668,89 @@ const FMaterial& FMaterialRenderProxy::GetIncompleteMaterialWithFallback(ERHIFea
 	return *Material;
 }
 
-void FMaterialRenderProxy::UpdateDeferredCachedUniformExpressions(FRHICommandListBase& RHICmdList)
+void FMaterialRenderProxy::UpdateDeferredCachedUniformExpressions(FRHICommandListBase& RHICmdList, UE::Tasks::FTask* TaskIfAsync)
 {
 	LLM_SCOPE(ELLMTag::Materials);
 
-	TRACE_CPUPROFILER_EVENT_SCOPE(FMaterialRenderProxy::UpdateDeferredCachedUniformExpressions);
-	CSV_SCOPED_TIMING_STAT_EXCLUSIVE(Material_UpdateDeferredCachedUniformExpressions);
-	QUICK_SCOPE_CYCLE_COUNTER(STAT_UpdateDeferredCachedUniformExpressions);
+	const bool bAllowAsyncUpdate = GUniformExpressionCacheAsyncUpdateTask.IsEnabled();
 
-	FUniformExpressionCacheAsyncUpdater Updater;
-	FUniformExpressionCacheAsyncUpdater* UpdaterIfEnabled = RHICmdList.IsImmediate() && GUniformExpressionCacheAsyncUpdateTask.IsEnabled() ? &Updater : nullptr;
+	if (!bAllowAsyncUpdate || !RHICmdList.IsImmediate())
+	{
+		TaskIfAsync = nullptr;
+	}
 
-	UE::TScopeLock Lock(DeferredUniformExpressionCacheRequestsMutex);
+	FRHICommandListBase* RHICmdListTask = nullptr;
 
-	for (TSet<FMaterialRenderProxy*>::TConstIterator It(DeferredUniformExpressionCacheRequests); It; ++It)
+	if (TaskIfAsync)
 	{
-		FMaterialRenderProxy* MaterialProxy = *It;
-		if (MaterialProxy->IsDeleted())
-		{
-			UE_LOG(LogMaterial, Fatal, TEXT("FMaterialRenderProxy deleted and GC mark was: %i"), MaterialProxy->IsMarkedForGarbageCollection());
-		}
+		RHICmdListTask = new FRHICommandList(FRHIGPUMask::All());
+		RHICmdListTask->SwitchPipeline(ERHIPipeline::Graphics);
+
+		FRHICommandListImmediate::Get(RHICmdList).QueueAsyncCommandListSubmit(RHICmdListTask);
+	}
+	else
+	{
+		RHICmdListTask = &RHICmdList;
+	}
+
+	auto EvaluateUniformExpressionsLambda = [RHICmdListTask, bAllowAsyncUpdate]
+	{
+		TRACE_CPUPROFILER_EVENT_SCOPE(FMaterialRenderProxy::UpdateDeferredCachedUniformExpressions);
+		CSV_SCOPED_TIMING_STAT_EXCLUSIVE(Material_UpdateDeferredCachedUniformExpressions);
+		QUICK_SCOPE_CYCLE_COUNTER(STAT_UpdateDeferredCachedUniformExpressions);
+
+		FOptionalTaskTagScope Scope(ETaskTag::EParallelRenderingThread);
+		FUniformExpressionCacheAsyncUpdater Updater;
+		FUniformExpressionCacheAsyncUpdater* UpdaterIfEnabled = bAllowAsyncUpdate ? &Updater : nullptr;
 
-		UMaterialInterface::IterateOverActiveFeatureLevels([&](ERHIFeatureLevel::Type InFeatureLevel)
 		{
-			// Don't bother caching if we'll be falling back to a different FMaterialRenderProxy for rendering anyway
-			const FMaterial* Material = MaterialProxy->GetMaterialNoFallback(InFeatureLevel);
-			if (Material && Material->GetRenderingThreadShaderMap())
+			UE::TScopeLock Lock(DeferredUniformExpressionCacheRequestsMutex);
+
+			for (TSet<FMaterialRenderProxy*>::TConstIterator It(DeferredUniformExpressionCacheRequests); It; ++It)
 			{
-				FMaterialRenderContext MaterialRenderContext(MaterialProxy, *Material, nullptr);
-				MaterialRenderContext.bShowSelection = GIsEditor;
-				MaterialProxy->EvaluateUniformExpressions(MaterialProxy->UniformExpressionCache[(int32)InFeatureLevel], MaterialRenderContext, UpdaterIfEnabled, &RHICmdList);
+				FMaterialRenderProxy* MaterialProxy = *It;
+				if (MaterialProxy->IsDeleted())
+				{
+					UE_LOG(LogMaterial, Fatal, TEXT("FMaterialRenderProxy deleted and GC mark was: %i"), MaterialProxy->IsMarkedForGarbageCollection());
+				}
+
+				UMaterialInterface::IterateOverActiveFeatureLevels([&](ERHIFeatureLevel::Type InFeatureLevel)
+				{
+					// Don't bother caching if we'll be falling back to a different FMaterialRenderProxy for rendering anyway
+					const FMaterial* Material = MaterialProxy->GetMaterialNoFallback(InFeatureLevel);
+					if (Material && Material->GetRenderingThreadShaderMap())
+					{
+						FMaterialRenderContext MaterialRenderContext(MaterialProxy, *Material, nullptr);
+						MaterialRenderContext.bShowSelection = GIsEditor;
+						MaterialProxy->EvaluateUniformExpressions(*RHICmdListTask, MaterialProxy->UniformExpressionCache[(int32)InFeatureLevel], MaterialRenderContext, UpdaterIfEnabled);
+					}
+				});
+
+				MaterialProxy->FinishCacheUniformExpressions();
 			}
-		});
+			DeferredUniformExpressionCacheRequests.Reset();
+		}
 
-		MaterialProxy->FinishCacheUniformExpressions();
-	}
+		if (bAllowAsyncUpdate)
+		{
+			Updater.Update(*RHICmdListTask);
+		}
+		else if (!RHICmdListTask->IsImmediate())
+		{
+			RHICmdListTask->FinishRecording();
+		}
+	};
 
-	if (UpdaterIfEnabled)
+	UE::Tasks::FTask Task;
+
+	if (TaskIfAsync)
 	{
-		Updater.Update(FRHICommandListImmediate::Get(RHICmdList));
+		*TaskIfAsync = UE::Tasks::Launch(UE_SOURCE_LOCATION, MoveTemp(EvaluateUniformExpressionsLambda));
+	}
+	else
+	{
+		EvaluateUniformExpressionsLambda();
 	}
-
-	DeferredUniformExpressionCacheRequests.Reset();
 }
 
 void FMaterialRenderProxy::UpdateDeferredCachedUniformExpressions()
diff --git a/Engine/Source/Runtime/Engine/Public/Materials/MaterialRenderProxy.h b/Engine/Source/Runtime/Engine/Public/Materials/MaterialRenderProxy.h
--- a/Engine/Source/Runtime/Engine/Public/Materials/MaterialRenderProxy.h
+++ b/Engine/Source/Runtime/Engine/Public/Materials/MaterialRenderProxy.h
@@ -8,6 +8,7 @@
 #include "RenderResource.h"
 #include "RHIImmutableSamplerState.h"
 #include "Async/Mutex.h"
+#include "Tasks/Task.h"
 
 enum class EMaterialParameterType : uint8;
 
@@ -110,7 +111,10 @@ public:
 	 * @param OutUniformExpressionCache - The uniform expression cache to build.
 	 * @param MaterialRenderContext - The context for which to cache expressions.
 	 */
-	ENGINE_API void EvaluateUniformExpressions(FUniformExpressionCache& OutUniformExpressionCache, const FMaterialRenderContext& Context, FUniformExpressionCacheAsyncUpdater* Updater = nullptr, FRHICommandListBase* RHICmdList = nullptr) const;
+	ENGINE_API void EvaluateUniformExpressions(FRHICommandListBase& RHICmdList, FUniformExpressionCache& OutUniformExpressionCache, const FMaterialRenderContext& Context, FUniformExpressionCacheAsyncUpdater* Updater = nullptr) const;
+
+	UE_DEPRECATED(5.4, "EvaluateUniformExpressions requires an RHI command list.")
+	ENGINE_API void EvaluateUniformExpressions(FUniformExpressionCache& OutUniformExpressionCache, const FMaterialRenderContext& Context, FUniformExpressionCacheAsyncUpdater* Updater = nullptr) const;
 
 	/**
 	 * Caches uniform expressions for efficient runtime evaluation.
@@ -210,7 +214,7 @@ public:
 	const uint32 NumSpecularProfileRT() const { return SpecularProfilesRT.Num(); }
 
 	static ENGINE_API void UpdateDeferredCachedUniformExpressions();
-	static ENGINE_API void UpdateDeferredCachedUniformExpressions(FRHICommandListBase& RHICmdList);
+	static ENGINE_API void UpdateDeferredCachedUniformExpressions(FRHICommandListBase& RHICmdList, UE::Tasks::FTask* TaskIfAsync = nullptr);
 
 	static ENGINE_API bool HasDeferredUniformExpressionCacheRequests();
 
@@ -220,7 +224,7 @@ public:
 
 private:
 	ENGINE_API IAllocatedVirtualTexture* GetPreallocatedVTStack(const FMaterialRenderContext& Context, const FUniformExpressionSet& UniformExpressionSet, const FMaterialVirtualTextureStack& VTStack) const;
-	ENGINE_API IAllocatedVirtualTexture* AllocateVTStack(const FMaterialRenderContext& Context, const FUniformExpressionSet& UniformExpressionSet, const FMaterialVirtualTextureStack& VTStack) const;
+	ENGINE_API IAllocatedVirtualTexture* AllocateVTStack(FRHICommandListBase& RHICmdList, const FMaterialRenderContext& Context, const FUniformExpressionSet& UniformExpressionSet, const FMaterialVirtualTextureStack& VTStack) const;
 
 	virtual void StartCacheUniformExpressions() const {}
 	virtual void FinishCacheUniformExpressions() const {}
diff --git a/Engine/Source/Runtime/Engine/Public/SceneInterface.h b/Engine/Source/Runtime/Engine/Public/SceneInterface.h
--- a/Engine/Source/Runtime/Engine/Public/SceneInterface.h
+++ b/Engine/Source/Runtime/Engine/Public/SceneInterface.h
@@ -91,7 +91,10 @@ enum class EUpdateAllPrimitiveSceneInfosAsyncOps
 	// Light primitive interactions are created asynchronously.
 	CreateLightPrimitiveInteractions = 1 << 1,
 
-	All = CacheMeshDrawCommands | CreateLightPrimitiveInteractions
+	// Material uniform expressions are cached asynchronously.
+	CacheMaterialUniformExpressions = 1 << 2,
+
+	All = CacheMeshDrawCommands | CreateLightPrimitiveInteractions | CacheMaterialUniformExpressions
 };
 ENUM_CLASS_FLAGS(EUpdateAllPrimitiveSceneInfosAsyncOps);
 
diff --git a/Engine/Source/Runtime/Renderer/Private/RendererScene.cpp b/Engine/Source/Runtime/Renderer/Private/RendererScene.cpp
--- a/Engine/Source/Runtime/Renderer/Private/RendererScene.cpp
+++ b/Engine/Source/Runtime/Renderer/Private/RendererScene.cpp
@@ -5428,7 +5428,9 @@ void FScene::UpdateAllPrimitiveSceneInfos(FRDGBuilder& GraphBuilder, EUpdateAllP
 	}
 
 	FSceneRenderer::WaitForCleanUpTasks(GraphBuilder.RHICmdList);
-	FMaterialRenderProxy::UpdateDeferredCachedUniformExpressions();
+
+	UE::Tasks::FTask UpdateUniformExpressionsTask;
+	FMaterialRenderProxy::UpdateDeferredCachedUniformExpressions(GraphBuilder.RHICmdList, EnumHasAnyFlags(AsyncOps, EUpdateAllPrimitiveSceneInfosAsyncOps::CacheMaterialUniformExpressions) ? &UpdateUniformExpressionsTask : nullptr);
 
 	RDG_EVENT_SCOPE(GraphBuilder, "UpdateAllPrimitiveSceneInfos");
 
@@ -6384,6 +6386,7 @@ void FScene::UpdateAllPrimitiveSceneInfos(FRDGBuilder& GraphBuilder, EUpdateAllP
 	ShadowScene->PostSceneUpdate(SceneUpdateChangeSetStorage.GetPreUpdateSet(), SceneUpdateChangeSetStorage.GetPostUpdateSet());
 
 	GPUSkinCacheTask.Wait();
+	UpdateUniformExpressionsTask.Wait();
 
 	if (SceneInfosWithStaticDrawListUpdate.Num() > 0)
 	{
diff --git a/Engine/Source/Runtime/Renderer/Private/SceneRendering.cpp b/Engine/Source/Runtime/Renderer/Private/SceneRendering.cpp
--- a/Engine/Source/Runtime/Renderer/Private/SceneRendering.cpp
+++ b/Engine/Source/Runtime/Renderer/Private/SceneRendering.cpp
@@ -123,6 +123,13 @@ static FAutoConsoleVariableRef CVarAsyncMeshDrawCommands(
 	TEXT("Mesh draw command caching is offloaded to an async task."),
 	ECVF_RenderThreadSafe);
 
+static int32 GAsyncCacheMaterialUniformExpressions = 1;
+static FAutoConsoleVariableRef CVarAsyncMaterialUniformExpressions(
+	TEXT("r.AsyncCacheMaterialUniformExpressions"),
+	GAsyncCacheMaterialUniformExpressions,
+	TEXT("Material uniform expression caching is offloaded to an async task."),
+	ECVF_RenderThreadSafe);
+
 static TAutoConsoleVariable<int32> CVarCachedMeshDrawCommands(
 	TEXT("r.MeshDrawCommands.UseCachedCommands"),
 	1,
@@ -3414,6 +3421,11 @@ IVisibilityTaskData* FSceneRenderer::OnRenderBegin(FRDGBuilder& GraphBuilder)
 		AsyncOps |= EUpdateAllPrimitiveSceneInfosAsyncOps::CacheMeshDrawCommands;
 	}
 
+	if (GAsyncCacheMaterialUniformExpressions > 0)
+	{
+		AsyncOps |= EUpdateAllPrimitiveSceneInfosAsyncOps::CacheMaterialUniformExpressions;
+	}
+
 	Scene->UpdateAllPrimitiveSceneInfos(GraphBuilder, AsyncOps);
 
 	if (!ViewFamily.ViewExtensions.IsEmpty())
diff --git a/Engine/Source/Runtime/Renderer/Private/ShaderBaseClasses.cpp b/Engine/Source/Runtime/Renderer/Private/ShaderBaseClasses.cpp
--- a/Engine/Source/Runtime/Renderer/Private/ShaderBaseClasses.cpp
+++ b/Engine/Source/Runtime/Renderer/Private/ShaderBaseClasses.cpp
@@ -297,7 +297,7 @@ void FMaterialShader::SetParameters(
 		FMaterialRenderContext MaterialRenderContext(MaterialRenderProxy, Material, &View);
 		bUniformExpressionCacheNeedsDelete = true;
 		UniformExpressionCache = new FUniformExpressionCache();
-		MaterialRenderProxy->EvaluateUniformExpressions(*UniformExpressionCache, MaterialRenderContext);
+		MaterialRenderProxy->EvaluateUniformExpressions(FRHICommandListExecutor::GetImmediateCommandList(), *UniformExpressionCache, MaterialRenderContext);
 	}
 
 	SetUniformBufferParameter(BatchedParameters, MaterialUniformBuffer, UniformExpressionCache->UniformBuffer);
