Material Parameter Collection: Non fatal if world is no longer valid

diff --git a/Engine/Source/Runtime/Engine/Private/Materials/ParameterCollection.cpp b/Engine/Source/Runtime/Engine/Private/Materials/ParameterCollection.cpp
--- a/Engine/Source/Runtime/Engine/Private/Materials/ParameterCollection.cpp
+++ b/Engine/Source/Runtime/Engine/Private/Materials/ParameterCollection.cpp
@@ -631,7 +631,12 @@ void UMaterialParameterCollectionInstance::SetCollection(UMaterialParameterColle
 
 bool UMaterialParameterCollectionInstance::SetScalarParameterValue(FName ParameterName, float ParameterValue)
 {
-	check(World.IsValid() && Collection.IsValid());
+	if (!World.IsValid())
+	{
+		return false;
+	}
+
+	check(Collection.IsValid());
 
 	if (Collection->GetScalarParameterByName(ParameterName))
 	{
@@ -665,7 +670,12 @@ bool UMaterialParameterCollectionInstance::SetScalarParameterValue(FName Paramet
 
 bool UMaterialParameterCollectionInstance::SetVectorParameterValue(FName ParameterName, const FLinearColor& ParameterValue)
 {
-	check(World.IsValid() && Collection.IsValid());
+	if (!World.IsValid())
+	{
+		return false;
+	}
+
+	check(Collection.IsValid());
 
 	if (Collection->GetVectorParameterByName(ParameterName))
 	{
