Fix typo in FContentBundleContainer::Initialize.

diff --git a/Engine/Source/Runtime/Engine/Private/WorldPartition/ContentBundle/ContentBundleContainer.cpp b/Engine/Source/Runtime/Engine/Private/WorldPartition/ContentBundle/ContentBundleContainer.cpp
--- a/Engine/Source/Runtime/Engine/Private/WorldPartition/ContentBundle/ContentBundleContainer.cpp
+++ b/Engine/Source/Runtime/Engine/Private/WorldPartition/ContentBundle/ContentBundleContainer.cpp
@@ -35,7 +35,7 @@ UWorld* FContentBundleContainer::GetInjectedWorld() const
 void FContentBundleContainer::Initialize()
 {
 	TRACE_CPUPROFILER_EVENT_SCOPE(FContentBundleContainer::Initialize);
-	UE_LOG(LogContentBundle, Log, TEXT("%s Creating new contrainer."), *ContentBundle::Log::MakeDebugInfoString(*this));
+	UE_LOG(LogContentBundle, Log, TEXT("%s Creating new container."), *ContentBundle::Log::MakeDebugInfoString(*this));
 
 #if WITH_EDITOR
 	if (UseEditorContentBundle())
