[Navigation] ISMC support for obstacles and nav areas: fix ISMC not exporting as obstacles

diff --git a/Engine/Source/Runtime/Engine/Classes/Components/StaticMeshComponent.h b/Engine/Source/Runtime/Engine/Classes/Components/StaticMeshComponent.h
--- a/Engine/Source/Runtime/Engine/Classes/Components/StaticMeshComponent.h
+++ b/Engine/Source/Runtime/Engine/Classes/Components/StaticMeshComponent.h
@@ -30,6 +30,7 @@ class FStaticMeshStaticLightingMesh;
 class ULightComponent;
 class UStaticMesh;
 class UStaticMeshComponent;
+class UNavCollisionBase;
 struct FConvexVolume;
 struct FEngineShowFlags;
 struct FNavigableGeometryExport;
@@ -727,6 +728,8 @@ protected:
 	// Overload this in child implementations that wish to extend Static Mesh or Nanite scene proxy implementations
 	ENGINE_API virtual FPrimitiveSceneProxy* CreateStaticMeshSceneProxy(Nanite::FMaterialAudit& NaniteMaterials, bool bCreateNanite);
 
+	ENGINE_API bool ShouldExportAsObstacle(const UNavCollisionBase& InNavCollision) const;
+
 public:
 
 	ENGINE_API void ReleaseResources();
diff --git a/Engine/Source/Runtime/Engine/Private/Components/StaticMeshComponent.cpp b/Engine/Source/Runtime/Engine/Private/Components/StaticMeshComponent.cpp
--- a/Engine/Source/Runtime/Engine/Private/Components/StaticMeshComponent.cpp
+++ b/Engine/Source/Runtime/Engine/Private/Components/StaticMeshComponent.cpp
@@ -3014,9 +3014,7 @@ bool UStaticMeshComponent::DoCustomNavigableGeometryExport(FNavigableGeometryExp
 	if (!Scale3D.IsZero() && GetStaticMesh() && !GetStaticMesh()->IsCompiling() && GetStaticMesh()->GetNavCollision())
 	{
 		const UNavCollisionBase* NavCollision = GetStaticMesh()->GetNavCollision();
-		const bool bExportAsObstacle = bOverrideNavigationExport ? bForceNavigationObstacle : NavCollision->IsDynamicObstacle();
-
-		if (bExportAsObstacle)
+		if (ShouldExportAsObstacle(*NavCollision))
 		{
 			// skip default export
 			return false;
@@ -3108,15 +3106,18 @@ void UStaticMeshComponent::GetNavigationData(FNavigationRelevantData& Data) cons
 	if (!Scale3D.IsZero() && GetStaticMesh() && !GetStaticMesh()->IsCompiling() && GetStaticMesh()->GetNavCollision())
 	{
 		UNavCollisionBase* NavCollision = GetStaticMesh()->GetNavCollision();
-		const bool bExportAsObstacle = bOverrideNavigationExport ? bForceNavigationObstacle : NavCollision->IsDynamicObstacle();
-
-		if (bExportAsObstacle)
+		if (ShouldExportAsObstacle(*NavCollision))
 		{
 			NavCollision->GetNavigationModifier(Data.Modifiers, GetComponentTransform());
 		}
 	}
 }
 
+bool UStaticMeshComponent::ShouldExportAsObstacle(const UNavCollisionBase& InNavCollision) const
+{
+	return bOverrideNavigationExport ? bForceNavigationObstacle : InNavCollision.IsDynamicObstacle();
+}
+
 bool UStaticMeshComponent::IsShown(const FEngineShowFlags& ShowFlags) const
 {
 	return ShowFlags.StaticMeshes;
diff --git a/Engine/Source/Runtime/Engine/Private/InstancedStaticMesh.cpp b/Engine/Source/Runtime/Engine/Private/InstancedStaticMesh.cpp
--- a/Engine/Source/Runtime/Engine/Private/InstancedStaticMesh.cpp
+++ b/Engine/Source/Runtime/Engine/Private/InstancedStaticMesh.cpp
@@ -5253,7 +5253,7 @@ bool UInstancedStaticMeshComponent::DoCustomNavigableGeometryExport(FNavigableGe
 	if (GetStaticMesh() && GetStaticMesh()->GetNavCollision())
 	{
 		UNavCollisionBase* NavCollision = GetStaticMesh()->GetNavCollision();
-		if (NavCollision->IsDynamicObstacle())
+		if (ShouldExportAsObstacle(*NavCollision))
 		{
 			return false;
 		}
@@ -5427,7 +5427,7 @@ void UInstancedStaticMeshComponent::GetNavigationData(FNavigationRelevantData& D
 	if (GetStaticMesh() && !GetStaticMesh()->IsCompiling() && GetStaticMesh()->GetNavCollision())
 	{
 		UNavCollisionBase* NavCollision = GetStaticMesh()->GetNavCollision();
-		if (NavCollision->IsDynamicObstacle())
+		if (ShouldExportAsObstacle(*NavCollision))
 		{
 			Data.Modifiers.MarkAsPerInstanceModifier();
 			NavCollision->GetNavigationModifier(Data.Modifiers, FTransform::Identity);
