Fix to give smoother interpolation when evaluation frames are skipped. This is similar to the calculation used in the budget allocator.

diff --git a/Engine/Source/Runtime/Engine/Private/Components/SkinnedMeshComponent.cpp b/Engine/Source/Runtime/Engine/Private/Components/SkinnedMeshComponent.cpp
--- a/Engine/Source/Runtime/Engine/Private/Components/SkinnedMeshComponent.cpp
+++ b/Engine/Source/Runtime/Engine/Private/Components/SkinnedMeshComponent.cpp
@@ -4910,7 +4910,9 @@ float FAnimUpdateRateParameters::GetInterpolationAlpha() const
 {
 	if (OptimizeMode == TrailMode)
 	{
-		return 0.25f + (1.f / float(FMath::Max(EvaluationRate, 2) * 2));
+		const float FramesToNextEval = EvaluationRate - SkippedEvalFrames;
+		return (EvaluationRate > 1.f) ? FMath::Clamp(1.0f / FramesToNextEval, 0.0f, 1.0f) : 1.f;
+
 	}
 	else if (OptimizeMode == LookAheadMode)
 	{
