Work-in-progress changes for the new HLSL translator.

#tests loaded various materials in editor, loaded level

diff --git a/Engine/Source/Runtime/Engine/Classes/Materials/MaterialExpressionDecalColor.h b/Engine/Source/Runtime/Engine/Classes/Materials/MaterialExpressionDecalColor.h
--- a/Engine/Source/Runtime/Engine/Classes/Materials/MaterialExpressionDecalColor.h
+++ b/Engine/Source/Runtime/Engine/Classes/Materials/MaterialExpressionDecalColor.h
@@ -17,6 +17,7 @@ class UMaterialExpressionDecalColor : public UMaterialExpression
 #if WITH_EDITOR
 	virtual int32 Compile(class FMaterialCompiler* Compiler, int32 OutputIndex) override;
 	virtual void GetCaption(TArray<FString>& OutCaptions) const override;
+	virtual bool GenerateHLSLExpression(FMaterialHLSLGenerator& Generator, UE::HLSLTree::FScope& Scope, int32 OutputIndex, UE::HLSLTree::FExpression const*& OutExpression) const override;
 #endif
 	//~ End UMaterialExpression Interface
 };
diff --git a/Engine/Source/Runtime/Engine/Classes/Materials/MaterialExpressionDecalLifetimeOpacity.h b/Engine/Source/Runtime/Engine/Classes/Materials/MaterialExpressionDecalLifetimeOpacity.h
--- a/Engine/Source/Runtime/Engine/Classes/Materials/MaterialExpressionDecalLifetimeOpacity.h
+++ b/Engine/Source/Runtime/Engine/Classes/Materials/MaterialExpressionDecalLifetimeOpacity.h
@@ -17,6 +17,7 @@ class UMaterialExpressionDecalLifetimeOpacity : public UMaterialExpression
 #if WITH_EDITOR
 	virtual int32 Compile(class FMaterialCompiler* Compiler, int32 OutputIndex) override;
 	virtual void GetCaption(TArray<FString>& OutCaptions) const override;
+	virtual bool GenerateHLSLExpression(FMaterialHLSLGenerator& Generator, UE::HLSLTree::FScope& Scope, int32 OutputIndex, UE::HLSLTree::FExpression const*& OutExpression) const override;
 #endif
 	//~ End UMaterialExpression Interface
 };
diff --git a/Engine/Source/Runtime/Engine/Classes/Materials/MaterialExpressionDepthOfFieldFunction.h b/Engine/Source/Runtime/Engine/Classes/Materials/MaterialExpressionDepthOfFieldFunction.h
--- a/Engine/Source/Runtime/Engine/Classes/Materials/MaterialExpressionDepthOfFieldFunction.h
+++ b/Engine/Source/Runtime/Engine/Classes/Materials/MaterialExpressionDepthOfFieldFunction.h
@@ -42,6 +42,7 @@ class UMaterialExpressionDepthOfFieldFunction : public UMaterialExpression
 #if WITH_EDITOR
 	virtual int32 Compile(class FMaterialCompiler* Compiler, int32 OutputIndex) override;
 	virtual void GetCaption(TArray<FString>& OutCaptions) const override;
+	virtual bool GenerateHLSLExpression(FMaterialHLSLGenerator& Generator, UE::HLSLTree::FScope& Scope, int32 OutputIndex, UE::HLSLTree::FExpression const*& OutExpression) const override;
 #endif
 	//~ End UMaterialExpression Interface
 };
diff --git a/Engine/Source/Runtime/Engine/Classes/Materials/MaterialExpressionParticleSubUVProperties.h b/Engine/Source/Runtime/Engine/Classes/Materials/MaterialExpressionParticleSubUVProperties.h
--- a/Engine/Source/Runtime/Engine/Classes/Materials/MaterialExpressionParticleSubUVProperties.h
+++ b/Engine/Source/Runtime/Engine/Classes/Materials/MaterialExpressionParticleSubUVProperties.h
@@ -18,6 +18,7 @@ class UMaterialExpressionParticleSubUVProperties : public UMaterialExpression
 	virtual int32 Compile(class FMaterialCompiler* Compiler, int32 OutputIndex) override;
 	virtual void GetCaption(TArray<FString>& OutCaptions) const override;
 	virtual void GetExpressionToolTip(TArray<FString>& OutToolTip) override;
+	virtual bool GenerateHLSLExpression(FMaterialHLSLGenerator& Generator, UE::HLSLTree::FScope& Scope, int32 OutputIndex, UE::HLSLTree::FExpression const*& OutExpression) const override;
 #endif
 	//~ End UMaterialExpression Interface
 };
diff --git a/Engine/Source/Runtime/Engine/Classes/Materials/MaterialExpressionPathTracingQualitySwitch.h b/Engine/Source/Runtime/Engine/Classes/Materials/MaterialExpressionPathTracingQualitySwitch.h
--- a/Engine/Source/Runtime/Engine/Classes/Materials/MaterialExpressionPathTracingQualitySwitch.h
+++ b/Engine/Source/Runtime/Engine/Classes/Materials/MaterialExpressionPathTracingQualitySwitch.h
@@ -30,6 +30,8 @@ class UMaterialExpressionPathTracingQualitySwitch : public UMaterialExpression
 	virtual uint32 GetInputType(int32 InputIndex) override;
 	virtual uint32 GetOutputType(int32 OutputIndex) override { return MCT_Unknown; }
 
+	virtual bool GenerateHLSLExpression(FMaterialHLSLGenerator& Generator, UE::HLSLTree::FScope& Scope, int32 OutputIndex, UE::HLSLTree::FExpression const*& OutExpression) const override;
+
 	virtual bool IsResultSubstrateMaterial(int32 OutputIndex) override;
 	virtual void GatherSubstrateMaterialInfo(FSubstrateMaterialInfo& SubstrateMaterialInfo, int32 OutputIndex) override;
 	virtual FSubstrateOperator* SubstrateGenerateMaterialTopologyTree(class FMaterialCompiler* Compiler, class UMaterialExpression* Parent, int32 OutputIndex) override;
diff --git a/Engine/Source/Runtime/Engine/Classes/Materials/MaterialExpressionThinTranslucentMaterialOutput.h b/Engine/Source/Runtime/Engine/Classes/Materials/MaterialExpressionThinTranslucentMaterialOutput.h
--- a/Engine/Source/Runtime/Engine/Classes/Materials/MaterialExpressionThinTranslucentMaterialOutput.h
+++ b/Engine/Source/Runtime/Engine/Classes/Materials/MaterialExpressionThinTranslucentMaterialOutput.h
@@ -23,6 +23,7 @@ public:
 	virtual int32 Compile(class FMaterialCompiler* Compiler, int32 OutputIndex) override;
 	virtual void GetCaption(TArray<FString>& OutCaptions) const override;
 	virtual uint32 GetInputType(int32 InputIndex) override;
+	virtual bool GenerateHLSLExpression(FMaterialHLSLGenerator& Generator, UE::HLSLTree::FScope& Scope, int32 OutputIndex, UE::HLSLTree::FExpression const*& OutExpression) const override;
 	//~ End UMaterialExpression Interface
 #endif
 
diff --git a/Engine/Source/Runtime/Engine/Private/Materials/MaterialExpressionHLSL.cpp b/Engine/Source/Runtime/Engine/Private/Materials/MaterialExpressionHLSL.cpp
--- a/Engine/Source/Runtime/Engine/Private/Materials/MaterialExpressionHLSL.cpp
+++ b/Engine/Source/Runtime/Engine/Private/Materials/MaterialExpressionHLSL.cpp
@@ -21,6 +21,12 @@
 #include "Materials/MaterialExpressionChannelMaskParameter.h"
 #include "Materials/MaterialExpressionCollectionParameter.h"
 #include "Materials/MaterialParameterCollection.h"
+#include "Materials/MaterialExpressionDecalColor.h"
+#include "Materials/MaterialExpressionDecalLifetimeOpacity.h"
+#include "Materials/MaterialExpressionDepthOfFieldFunction.h"
+#include "Materials/MaterialExpressionParticleSubUVProperties.h"
+#include "Materials/MaterialExpressionPathTracingQualitySwitch.h"
+#include "Materials/MaterialExpressionThinTranslucentMaterialOutput.h"
 #include "Materials/MaterialExpressionDynamicParameter.h"
 #include "Materials/MaterialExpressionCurveAtlasRowParameter.h"
 #include "Materials/MaterialExpressionStaticSwitchParameter.h"
@@ -743,6 +749,31 @@ bool UMaterialExpressionCameraVectorWS::GenerateHLSLExpression(FMaterialHLSLGene
 	return true;
 }
 
+bool UMaterialExpressionParticleSubUVProperties::GenerateHLSLExpression(FMaterialHLSLGenerator& Generator, UE::HLSLTree::FScope& Scope, int32 OutputIndex, UE::HLSLTree::FExpression const*& OutExpression) const
+{
+	using namespace UE::HLSLTree::Material;
+	EExternalInput Input = EExternalInput::None;
+	switch (OutputIndex)
+	{
+	case 0:
+		Input = EExternalInput::ParticleSubUVCoords0; break;
+	case 1:
+		Input = EExternalInput::ParticleSubUVCoords1; break;
+	case 2:
+		Input = EExternalInput::ParticleSubUVLerp; break;
+	default:
+		checkNoEntry();
+		break;
+
+	}
+	
+	if (Input != EExternalInput::None)
+	{
+		OutExpression = Generator.GetTree().NewExpression<FExpressionExternalInput>(Input);
+	}
+	return OutExpression != nullptr;
+}
+
 bool UMaterialExpressionViewProperty::GenerateHLSLExpression(FMaterialHLSLGenerator& Generator, UE::HLSLTree::FScope& Scope, int32 OutputIndex, UE::HLSLTree::FExpression const*& OutExpression) const
 {
 	using namespace UE::HLSLTree::Material;
@@ -1902,6 +1933,50 @@ bool UMaterialExpressionSign::GenerateHLSLExpression(FMaterialHLSLGenerator& Gen
 	return true;
 }
 
+bool UMaterialExpressionDecalLifetimeOpacity::GenerateHLSLExpression(FMaterialHLSLGenerator& Generator, UE::HLSLTree::FScope& Scope, int32 OutputIndex, UE::HLSLTree::FExpression const*& OutExpression) const
+{
+	using namespace UE::HLSLTree;
+	OutExpression = Generator.GetTree().NewExpression<UE::HLSLTree::FExpressionInlineCustomHLSL>(UE::Shader::EValueType::Float1, TEXT("DecalLifetimeOpacity()"));
+	return true;
+}
+
+bool UMaterialExpressionDecalColor::GenerateHLSLExpression(FMaterialHLSLGenerator& Generator, UE::HLSLTree::FScope& Scope, int32 OutputIndex, UE::HLSLTree::FExpression const*& OutExpression) const
+{
+	using namespace UE::HLSLTree;
+	OutExpression = Generator.GetTree().NewExpression<UE::HLSLTree::FExpressionInlineCustomHLSL>(UE::Shader::EValueType::Float4, TEXT("DecalColor()"));
+	return true;
+}
+
+bool UMaterialExpressionPathTracingQualitySwitch::GenerateHLSLExpression(FMaterialHLSLGenerator& Generator, UE::HLSLTree::FScope& Scope, int32 OutputIndex, UE::HLSLTree::FExpression const*& OutExpression) const
+{
+	if (!Normal.GetTracedInput().Expression)
+	{
+		return Generator.Error(TEXT("Missing PathTracingQualitySwitch input 'Normal'"));
+	}
+	else if (!PathTraced.GetTracedInput().Expression)
+	{
+		return Generator.Error(TEXT("Missing PathTracingQualitySwitch input 'PathTraced'"));
+	}
+	else
+	{
+		using namespace UE::HLSLTree;
+		const FExpression* NormalExpression = Normal.AcquireHLSLExpression(Generator, Scope);
+		const FExpression* PathTracedExpression = PathTraced.AcquireHLSLExpression(Generator, Scope);
+		const FExpression* ConditionExpression = Generator.GetTree().NewExpression<FExpressionInlineCustomHLSL>(UE::Shader::EValueType::Bool1, TEXT("GetPathTracingQualitySwitch()"));
+
+		OutExpression = Generator.GetTree().NewExpression<FExpressionSelect>(ConditionExpression, PathTracedExpression, NormalExpression);
+		return true;
+	}
+}
+
+bool UMaterialExpressionDepthOfFieldFunction::GenerateHLSLExpression(FMaterialHLSLGenerator& Generator, UE::HLSLTree::FScope& Scope, int32 OutputIndex, UE::HLSLTree::FExpression const*& OutExpression) const
+{
+	//const FExpression* ExpressionDepth = Depth.AcquireHLSLExpression(Generator, Scope);
+	//FunctionValue
+	//MaterialExpressionDepthOfFieldFunction(%s, %d)
+	return Generator.Error(TEXT("UMaterialExpressionDepthOfFieldFunction::GenerateHLSLExpression unimplemented"));
+}
+
 static bool GenerateHLSLExpressionTrig(FMaterialHLSLGenerator& Generator,
 	UE::HLSLTree::EOperation Op,
 	const UE::HLSLTree::FExpression* InputExpression,
@@ -3136,11 +3211,22 @@ bool UMaterialExpressionVolumetricAdvancedMaterialOutput::GenerateHLSLExpression
 	case 4: OutExpression = MultiScatteringOcclusion.AcquireHLSLExpressionOrConstant(Generator, Scope, ConstMultiScatteringOcclusion); break;
 	case 5: OutExpression = MultiScatteringEccentricity.AcquireHLSLExpressionOrConstant(Generator, Scope, ConstMultiScatteringEccentricity); break;
 	case 6: OutExpression = ConservativeDensity.AcquireHLSLExpressionOrConstant(Generator, Scope, FVector4f(1.0f, 1.0f, 1.0f, 1.0f)); break;
-	default: return Generator.Error(TEXT("Invlid output"));
+	default: return Generator.Error(TEXT("Invalid output"));
 	}
 	return OutExpression != nullptr;
 }
 
+bool UMaterialExpressionThinTranslucentMaterialOutput::GenerateHLSLExpression(FMaterialHLSLGenerator& Generator, UE::HLSLTree::FScope& Scope, int32 OutputIndex, UE::HLSLTree::FExpression const*& OutExpression) const
+{
+	if (OutputIndex == 0)
+	{
+		OutExpression = TransmittanceColor.AcquireHLSLExpressionOrConstant(Generator, Scope, FVector3f(0.5f, 0.5f, 0.5f));
+		return true;
+	}
+
+	return Generator.Error(TEXT("Invalid output"));
+}
+
 UE::Shader::EValueType UMaterialExpressionVolumetricAdvancedMaterialOutput::GetCustomOutputType(int32 OutputIndex) const
 {
 	if (OutputIndex >= 0 && OutputIndex < 6) return UE::Shader::EValueType::Float1;
@@ -3150,14 +3236,12 @@ UE::Shader::EValueType UMaterialExpressionVolumetricAdvancedMaterialOutput::GetC
 
 bool UMaterialExpressionVolumetricCloudEmptySpaceSkippingInput::GenerateHLSLExpression(FMaterialHLSLGenerator& Generator, UE::HLSLTree::FScope& Scope, int32 OutputIndex, UE::HLSLTree::FExpression const*& OutExpression) const
 {
-	check(false);	// TODO implement when the new compiler is enabled
-	return OutExpression != nullptr;
+	return Generator.Error(TEXT("UMaterialExpressionVolumetricCloudEmptySpaceSkippingInput::GenerateHLSLExpression unimplemented"));
 }
 
 bool UMaterialExpressionVolumetricCloudEmptySpaceSkippingOutput::GenerateHLSLExpression(FMaterialHLSLGenerator& Generator, UE::HLSLTree::FScope& Scope, int32 OutputIndex, UE::HLSLTree::FExpression const*& OutExpression) const
 {
-	check(false);	// TODO implement when the new compiler is enabled
-	return OutExpression != nullptr;
+	return Generator.Error(TEXT("UMaterialExpressionVolumetricCloudEmptySpaceSkippingOutput::GenerateHLSLExpression unimplemented"));
 }
 
 UE::Shader::EValueType UMaterialExpressionVolumetricCloudEmptySpaceSkippingOutput::GetCustomOutputType(int32 OutputIndex) const
diff --git a/Engine/Source/Runtime/Engine/Private/Materials/MaterialHLSLEmitter.cpp b/Engine/Source/Runtime/Engine/Private/Materials/MaterialHLSLEmitter.cpp
--- a/Engine/Source/Runtime/Engine/Private/Materials/MaterialHLSLEmitter.cpp
+++ b/Engine/Source/Runtime/Engine/Private/Materials/MaterialHLSLEmitter.cpp
@@ -419,6 +419,16 @@ static void GetMaterialEnvironment(EShaderPlatform InPlatform,
 		OutEnvironment.SetDefine(TEXT("NEEDS_PARTICLE_SIZE"), TEXT("1"));
 	}
 
+	const bool bNeedsParticleSubUVs =
+		EmitMaterialData.IsExternalInputUsed(Material::EExternalInput::ParticleSubUVCoords0) ||
+		EmitMaterialData.IsExternalInputUsed(Material::EExternalInput::ParticleSubUVCoords1) ||
+		EmitMaterialData.IsExternalInputUsed(Material::EExternalInput::ParticleSubUVLerp);
+
+	if (bNeedsParticleSubUVs)
+	{
+		OutEnvironment.SetDefine(TEXT("USE_PARTICLE_SUBUVS"), TEXT("1"));
+	}
+
 	if (MaterialCompilationOutput.bNeedsSceneTextures)
 	{
 		OutEnvironment.SetDefine(TEXT("NEEDS_SCENE_TEXTURES"), TEXT("1"));
diff --git a/Engine/Source/Runtime/Engine/Private/Materials/MaterialHLSLTree.cpp b/Engine/Source/Runtime/Engine/Private/Materials/MaterialHLSLTree.cpp
--- a/Engine/Source/Runtime/Engine/Private/Materials/MaterialHLSLTree.cpp
+++ b/Engine/Source/Runtime/Engine/Private/Materials/MaterialHLSLTree.cpp
@@ -149,6 +149,9 @@ FExternalInputDescription GetExternalInputDescription(EExternalInput Input)
 	case EExternalInput::ParticleRelativeTime: return FExternalInputDescription(TEXT("ParticleRelativeTime"), Shader::EValueType::Float1);
 	case EExternalInput::ParticleRandom: return FExternalInputDescription(TEXT("ParticleRandom"), Shader::EValueType::Float1);
 	case EExternalInput::ParticleSize: return FExternalInputDescription(TEXT("ParticleSize"), Shader::EValueType::Float2);
+	case EExternalInput::ParticleSubUVCoords0: return FExternalInputDescription(TEXT("ParticleSubUVCoords0"), Shader::EValueType::Float2);
+	case EExternalInput::ParticleSubUVCoords1: return FExternalInputDescription(TEXT("ParticleSubUVCoords1"), Shader::EValueType::Float2);
+	case EExternalInput::ParticleSubUVLerp: return FExternalInputDescription(TEXT("ParticleSubUVLerp"), Shader::EValueType::Float1);
 
 	case EExternalInput::IsOrthographic: return FExternalInputDescription(TEXT("IsOrthographic"), Shader::EValueType::Float1);
 
@@ -344,6 +347,9 @@ void FExpressionExternalInput::EmitValueShader(FEmitContext& Context, FEmitScope
 		case EExternalInput::ParticleRelativeTime: Code = TEXT("Parameters.Particle.RelativeTime"); break;
 		case EExternalInput::ParticleRandom: Code = TEXT("Parameters.Particle.Random"); break;
 		case EExternalInput::ParticleSize: Code = TEXT("Parameters.Particle.Size"); break;
+		case EExternalInput::ParticleSubUVCoords0: Code = TEXT("Parameters.Particle.SubUVCoords[0].xy"); break;
+		case EExternalInput::ParticleSubUVCoords1: Code = TEXT("Parameters.Particle.SubUVCoords[1].xy"); break;
+		case EExternalInput::ParticleSubUVLerp: Code = TEXT("Parameters.Particle.SubUVLerp"); break;
 
 		case EExternalInput::IsOrthographic: Code = TEXT("((View.ViewToClip[3][3] < 1.0f) ? 0.0f : 1.0f)"); break;
 
diff --git a/Engine/Source/Runtime/Engine/Public/MaterialHLSLTree.h b/Engine/Source/Runtime/Engine/Public/MaterialHLSLTree.h
--- a/Engine/Source/Runtime/Engine/Public/MaterialHLSLTree.h
+++ b/Engine/Source/Runtime/Engine/Public/MaterialHLSLTree.h
@@ -187,6 +187,9 @@ enum class EExternalInput : uint8
 	ParticleRelativeTime,
 	ParticleRandom,
 	ParticleSize,
+	ParticleSubUVCoords0,
+	ParticleSubUVCoords1,
+	ParticleSubUVLerp,
 
 	IsOrthographic,
 
