Changes for the WIP FMaterialHLSLGenerator. This change doesn't affect current or future projects unless the new translator is enabled.
- Fixed an issue where the execution goes into an infinite loop when an input of a UMaterialExpressionMaterialFunctionCall is directly or indirectly connected to an output of the same expression.
- Added support mixing MaterialAttributes with numeric vectors in the same operation. For example, lerping a MA with a float3. This is allowed in the current translator because each MA is processed separately one at a time. Though not ideal, the new translator needs to support this to avoid breaking existing materials.
- Implemented GenerateHLSLExpression for several material expressions

diff --git a/Engine/Source/Editor/UnrealEd/Private/MaterialGraph.cpp b/Engine/Source/Editor/UnrealEd/Private/MaterialGraph.cpp
--- a/Engine/Source/Editor/UnrealEd/Private/MaterialGraph.cpp
+++ b/Engine/Source/Editor/UnrealEd/Private/MaterialGraph.cpp
@@ -183,7 +183,7 @@ void UMaterialGraph::UpdatePinTypes()
 	{
 		const FMaterialConnectionKey& Key = It.Key;
 		const FExpression* OutputExpression = It.Value;
-		const UE::Shader::FType OutputType = EmitContext.GetType(OutputExpression);
+		const UE::Shader::FType OutputType = EmitContext.GetTypeForPinColoring(OutputExpression);
 		if (!OutputType.IsVoid())
 		{
 			const FConnectionKey InputKey(Key.InputObject, OutputExpression);
diff --git a/Engine/Source/Runtime/Engine/Classes/Materials/MaterialExpressionParticleMacroUV.h b/Engine/Source/Runtime/Engine/Classes/Materials/MaterialExpressionParticleMacroUV.h
--- a/Engine/Source/Runtime/Engine/Classes/Materials/MaterialExpressionParticleMacroUV.h
+++ b/Engine/Source/Runtime/Engine/Classes/Materials/MaterialExpressionParticleMacroUV.h
@@ -22,6 +22,8 @@ class UMaterialExpressionParticleMacroUV : public UMaterialExpression
 #if WITH_EDITOR
 	virtual int32 Compile(class FMaterialCompiler* Compiler, int32 OutputIndex) override;
 	virtual void GetCaption(TArray<FString>& OutCaptions) const override;
+
+	virtual bool GenerateHLSLExpression(FMaterialHLSLGenerator& Generator, UE::HLSLTree::FScope& Scope, int32 OutputIndex, UE::HLSLTree::FExpression const*& OutExpression) const;
 #endif
 	//~ End UMaterialExpression Interface
 };
diff --git a/Engine/Source/Runtime/Engine/Classes/Materials/MaterialExpressionTextureProperty.h b/Engine/Source/Runtime/Engine/Classes/Materials/MaterialExpressionTextureProperty.h
--- a/Engine/Source/Runtime/Engine/Classes/Materials/MaterialExpressionTextureProperty.h
+++ b/Engine/Source/Runtime/Engine/Classes/Materials/MaterialExpressionTextureProperty.h
@@ -43,6 +43,8 @@ class UMaterialExpressionTextureProperty : public UMaterialExpression
 	ENGINE_API virtual void GetCaption(TArray<FString>& OutCaptions) const override;
 	ENGINE_API virtual uint32 GetInputType(int32 InputIndex) override;
 	ENGINE_API virtual bool MatchesSearchQuery(const TCHAR* SearchQuery) override;
+
+	virtual bool GenerateHLSLExpression(FMaterialHLSLGenerator& Generator, UE::HLSLTree::FScope& Scope, int32 OutputIndex, UE::HLSLTree::FExpression const*& OutExpression) const override;
 #endif //WITH_EDITOR
 	//~ End UMaterialExpression Interface
 };
diff --git a/Engine/Source/Runtime/Engine/Private/HLSLTree/HLSLTree.cpp b/Engine/Source/Runtime/Engine/Private/HLSLTree/HLSLTree.cpp
--- a/Engine/Source/Runtime/Engine/Private/HLSLTree/HLSLTree.cpp
+++ b/Engine/Source/Runtime/Engine/Private/HLSLTree/HLSLTree.cpp
@@ -371,7 +371,7 @@ void FExpressionLocalPHI::EmitValueShader(FEmitContext& Context, FEmitScope& Sco
 
 		if (LiveScopes.bCanForwardValue)
 		{
-			EmitExpression = LiveScopes.LiveValues[0]->GetValueShader(Context, Scope);
+			EmitExpression = LiveScopes.LiveValues[0]->GetValueShader(Context, Scope, RequestedType);
 			Context.EmitLocalPHIMap.Add(this, EmitExpression);
 		}
 		else
@@ -379,7 +379,7 @@ void FExpressionLocalPHI::EmitValueShader(FEmitContext& Context, FEmitScope& Sco
 			// This is the first time we've emitted shader code for this PHI
 			// Create an expression and add it to the map first, so if this is called recursively this path will only be taken the first time
 			const int32 LocalPHIIndex = Context.NumExpressionLocalPHIs++;
-			const Shader::FType LocalType = Context.GetType(this);
+			const Shader::FType LocalType = Context.GetResultType(this, RequestedType);
 
 			EmitExpression = OutResult.Code = Context.EmitInlineExpression(Scope,
 				LocalType,
@@ -443,7 +443,7 @@ void FExpressionLocalPHI::EmitValuePreshader(FEmitContext& Context, FEmitScope&
 		}
 	}
 
-	OutResult.Type = Context.GetType(this);
+	OutResult.Type = Context.GetResultType(this, RequestedType);
 	if (ValueStackPosition == INDEX_NONE)
 	{
 		Private::FLocalPHILiveScopes LiveScopes;
@@ -1195,13 +1195,13 @@ FEmitShaderExpression* FExpression::GetValueShader(FEmitContext& Context, FEmitS
 
 FEmitShaderExpression* FExpression::GetValueShader(FEmitContext& Context, FEmitScope& Scope, const FRequestedType& RequestedType, const Shader::FType& ResultType) const
 {
-	const FPreparedType& PreparedType = Context.GetPreparedType(this);
+	const FPreparedType& PreparedType = Context.GetPreparedType(this, RequestedType);
 	return GetValueShader(Context, Scope, RequestedType, PreparedType, ResultType);
 }
 
 FEmitShaderExpression* FExpression::GetValueShader(FEmitContext& Context, FEmitScope& Scope, const FRequestedType& RequestedType) const
 {
-	const FPreparedType& PreparedType = Context.GetPreparedType(this);
+	const FPreparedType& PreparedType = Context.GetPreparedType(this, RequestedType);
 	return GetValueShader(Context, Scope, RequestedType, PreparedType, PreparedType.GetResultType());
 }
 
@@ -1215,12 +1215,6 @@ FEmitShaderExpression* FExpression::GetValueShader(FEmitContext& Context, FEmitS
 	return GetValueShader(Context, Scope, ResultType, ResultType);
 }
 
-FEmitShaderExpression* FExpression::GetValueShader(FEmitContext& Context, FEmitScope& Scope) const
-{
-	const FPreparedType& PreparedType = Context.GetPreparedType(this);
-	return GetValueShader(Context, Scope, PreparedType.GetRequestedType(), PreparedType, PreparedType.GetResultType());
-}
-
 Shader::FType FExpression::GetValuePreshader(FEmitContext& Context, FEmitScope& Scope, const FRequestedType& RequestedType, const FPreparedType& PreparedType, const Shader::FType& ResultType, Shader::FPreshaderData& OutPreshader) const
 {
 	FEmitOwnerScope OwnerScope(Context, this);
@@ -1255,13 +1249,13 @@ Shader::FType FExpression::GetValuePreshader(FEmitContext& Context, FEmitScope&
 
 Shader::FType FExpression::GetValuePreshader(FEmitContext& Context, FEmitScope& Scope, const FRequestedType& RequestedType, const Shader::FType& ResultType, Shader::FPreshaderData& OutPreshader) const
 {
-	const FPreparedType& PreparedType = Context.GetPreparedType(this);
+	const FPreparedType& PreparedType = Context.GetPreparedType(this, RequestedType);
 	return GetValuePreshader(Context, Scope, RequestedType, PreparedType, ResultType, OutPreshader);
 }
 
 Shader::FType FExpression::GetValuePreshader(FEmitContext& Context, FEmitScope& Scope, const FRequestedType& RequestedType, Shader::FPreshaderData& OutPreshader) const
 {
-	const FPreparedType& PreparedType = Context.GetPreparedType(this);
+	const FPreparedType& PreparedType = Context.GetPreparedType(this, RequestedType);
 	return GetValuePreshader(Context, Scope, RequestedType, PreparedType, PreparedType.GetResultType(), OutPreshader);
 }
 
@@ -1311,13 +1305,13 @@ Shader::FValue FExpression::GetValueConstant(FEmitContext& Context, FEmitScope&
 
 Shader::FValue FExpression::GetValueConstant(FEmitContext& Context, FEmitScope& Scope, const FRequestedType& RequestedType, const Shader::FType& ResultType) const
 {
-	const FPreparedType PreparedType = Context.GetPreparedType(this);
+	const FPreparedType PreparedType = Context.GetPreparedType(this, RequestedType);
 	return GetValueConstant(Context, Scope, RequestedType, PreparedType, ResultType);
 }
 
 Shader::FValue FExpression::GetValueConstant(FEmitContext& Context, FEmitScope& Scope, const FRequestedType& RequestedType) const
 {
-	const FPreparedType PreparedType = Context.GetPreparedType(this);
+	const FPreparedType PreparedType = Context.GetPreparedType(this, RequestedType);
 	return GetValueConstant(Context, Scope, RequestedType, PreparedType, PreparedType.GetResultType());
 }
 
@@ -1333,14 +1327,14 @@ Shader::FValue FExpression::GetValueConstant(FEmitContext& Context, FEmitScope&
 
 bool FExpression::GetValueObject(FEmitContext& Context, FEmitScope& Scope, const FName& ObjectTypeName, void* OutObjectBase) const
 {
-	const FPreparedType PreparedType = Context.GetPreparedType(this);
+	const FPreparedType PreparedType = Context.GetPreparedType(this, FRequestedType(ObjectTypeName));
 	check(PreparedType.Type.ObjectType == ObjectTypeName);
 	return EmitValueObject(Context, Scope, ObjectTypeName, OutObjectBase);
 }
 
 bool FExpression::CheckObjectSupportsCustomHLSL(FEmitContext& Context, FEmitScope& Scope, const FName& ObjectTypeName) const
 {
-	const FPreparedType PreparedType = Context.GetPreparedType(this);
+	const FPreparedType PreparedType = Context.GetPreparedType(this, FRequestedType(ObjectTypeName));
 	check(PreparedType.Type.ObjectType == ObjectTypeName);
 
 	FEmitCustomHLSLParameterResult UnusedResult;
diff --git a/Engine/Source/Runtime/Engine/Private/HLSLTree/HLSLTreeCommon.cpp b/Engine/Source/Runtime/Engine/Private/HLSLTree/HLSLTreeCommon.cpp
--- a/Engine/Source/Runtime/Engine/Private/HLSLTree/HLSLTreeCommon.cpp
+++ b/Engine/Source/Runtime/Engine/Private/HLSLTree/HLSLTreeCommon.cpp
@@ -84,7 +84,7 @@ bool FExpressionError::PrepareValue(FEmitContext& Context, FEmitScope&, const FR
 
 void FExpressionForward::ComputeAnalyticDerivatives(FTree& Tree, FExpressionDerivatives& OutResult) const
 {
-	return Expression->ComputeAnalyticDerivatives(Tree, OutResult);
+	Expression->ComputeAnalyticDerivatives(Tree, OutResult);
 }
 
 const FExpression* FExpressionForward::ComputePreviousFrame(FTree& Tree, const FRequestedType& RequestedType) const
@@ -100,12 +100,17 @@ bool FExpressionForward::PrepareValue(FEmitContext& Context, FEmitScope& Scope,
 
 void FExpressionForward::EmitValueShader(FEmitContext& Context, FEmitScope& Scope, const FRequestedType& RequestedType, FEmitValueShaderResult& OutResult) const
 {
-	return Expression->EmitValueShader(Context, Scope, RequestedType, OutResult);
+	Expression->EmitValueShader(Context, Scope, RequestedType, OutResult);
 }
 
 void FExpressionForward::EmitValuePreshader(FEmitContext& Context, FEmitScope& Scope, const FRequestedType& RequestedType, FEmitValuePreshaderResult& OutResult) const
 {
-	return Expression->EmitValuePreshader(Context, Scope, RequestedType, OutResult);
+	Expression->EmitValuePreshader(Context, Scope, RequestedType, OutResult);
+}
+
+bool FExpressionForward::EmitValueObject(FEmitContext& Context, FEmitScope& Scope, const FName& ObjectTypeName, void* OutObjectBase) const
+{
+	return Expression->EmitValueObject(Context, Scope, ObjectTypeName, OutObjectBase);
 }
 
 const FExpression* FExpressionPreviousFrameSwitch::ComputePreviousFrame(FTree& Tree, const FRequestedType& RequestedType) const
@@ -188,7 +193,7 @@ bool FExpressionDefaultValue::PrepareValue(FEmitContext& Context, FEmitScope& Sc
 
 void FExpressionDefaultValue::EmitValuePreshader(FEmitContext& Context, FEmitScope& Scope, const FRequestedType& RequestedType, FEmitValuePreshaderResult& OutResult) const
 {
-	const FPreparedType& PreparedType = Context.GetPreparedType(Expression);
+	const FPreparedType& PreparedType = Context.GetPreparedType(Expression, RequestedType);
 	if (!PreparedType.IsVoid())
 	{
 		// If 'Exrpession' is valid, use that
@@ -278,20 +283,22 @@ void FExpressionSetStructField::ComputeAnalyticDerivatives(FTree& Tree, FExpress
 
 const FExpression* FExpressionSetStructField::ComputePreviousFrame(FTree& Tree, const FRequestedType& RequestedType) const
 {
-	FRequestedType RequestedStructType(RequestedType);
-	RequestedStructType.ClearFieldRequested(Field);
+	FRequestedType RequestedStructType = MakeRequestedStructType(Tree.ActiveStructFieldStack, RequestedType);
 	const FExpression* PrevStructExpression = Tree.GetPreviousFrame(StructExpression, RequestedStructType);
 
-	const FRequestedType RequestedFieldType = RequestedType.GetField(Field);
-	const FExpression* PrevFieldExpression = Tree.GetPreviousFrame(FieldExpression, RequestedFieldType);
+	const FRequestedType RequestedFieldType = MakeRequestedFieldType(Tree.ActiveStructFieldStack, RequestedType);
+	const FExpression* PrevFieldExpression;
+	{
+		FScopedActiveStructField ScopedActiveField(Tree.ActiveStructFieldStack, Field);
+		PrevFieldExpression = Tree.GetPreviousFrame(FieldExpression, RequestedFieldType);
+	}
 
 	return Tree.NewExpression<FExpressionSetStructField>(StructType, Field, PrevStructExpression, PrevFieldExpression);
 }
 
 bool FExpressionSetStructField::PrepareValue(FEmitContext& Context, FEmitScope& Scope, const FRequestedType& RequestedType, FPrepareValueResult& OutResult) const
 {
-	FRequestedType RequestedStructType(RequestedType);
-	RequestedStructType.ClearFieldRequested(Field);
+	FRequestedType RequestedStructType = MakeRequestedStructType(Context.ActiveStructFieldStack, RequestedType);
 
 	const FPreparedType StructPreparedType = Context.PrepareExpression(StructExpression, Scope, RequestedStructType);
 	if (!StructPreparedType.IsVoid() && StructPreparedType.Type.StructType != StructType)
@@ -305,24 +312,31 @@ bool FExpressionSetStructField::PrepareValue(FEmitContext& Context, FEmitScope&
 		ResultType = StructType;
 	}
 	
-	const FRequestedType RequestedFieldType = RequestedType.GetField(Field);
+	const FRequestedType RequestedFieldType = MakeRequestedFieldType(Context.ActiveStructFieldStack, RequestedType);
 	if (!RequestedFieldType.IsEmpty())
 	{
+		FScopedActiveStructField ActiveFieldScope(Context.ActiveStructFieldStack, Field);
+
 		const FPreparedType FieldPreparedType = Context.PrepareExpression(FieldExpression, Scope, RequestedFieldType);
 		ResultType.SetField(Field, FieldPreparedType);
 	}
 
+	if (!RequestedType.IsStruct())
+	{
+		const Shader::FStructField* ActiveField = Context.ActiveStructFieldStack.Last();
+		ResultType = ResultType.GetFieldType(ActiveField);
+	}
+
 	return OutResult.SetType(Context, RequestedType, ResultType);
 }
 
 void FExpressionSetStructField::EmitValueShader(FEmitContext& Context, FEmitScope& Scope, const FRequestedType& RequestedType, FEmitValueShaderResult& OutResult) const
 {
-	FRequestedType RequestedStructType(RequestedType);
-	RequestedStructType.ClearFieldRequested(Field);
+	FRequestedType RequestedStructType = MakeRequestedStructType(Context.ActiveStructFieldStack, RequestedType);
 	const EExpressionEvaluation StructEvaluation = Context.GetEvaluation(StructExpression, Scope, RequestedStructType);
 	check(StructEvaluation != EExpressionEvaluation::None);
 
-	const FRequestedType RequestedFieldType = RequestedType.GetField(Field);
+	const FRequestedType RequestedFieldType = MakeRequestedFieldType(Context.ActiveStructFieldStack, RequestedType);
 	const EExpressionEvaluation FieldEvaluation = Context.GetEvaluation(FieldExpression, Scope, RequestedFieldType);
 	check(FieldEvaluation != EExpressionEvaluation::None);
 
@@ -338,25 +352,33 @@ void FExpressionSetStructField::EmitValueShader(FEmitContext& Context, FEmitScop
 		if (RequestedFieldType.IsEmpty())
 		{
 			OutResult.Code = StructValue;
-			return;
 		}
+		else
+		{
+			FScopedActiveStructField ActiveFieldScope(Context.ActiveStructFieldStack, Field);
+
+			FEmitShaderExpression* FieldValue = FieldExpression->GetValueShader(Context, Scope, RequestedFieldType, Field->Type);
+			OutResult.Code = Context.EmitExpression(Scope, StructType, TEXT("%_Set%(%, %)"),
+				StructType->Name,
+				Field->Name,
+				StructValue,
+				FieldValue);
+		}
+	}
 
-		FEmitShaderExpression* FieldValue = FieldExpression->GetValueShader(Context, Scope, RequestedFieldType, Field->Type);
-		OutResult.Code = Context.EmitExpression(Scope, StructType, TEXT("%_Set%(%, %)"),
-			StructType->Name,
-			Field->Name,
-			StructValue,
-			FieldValue);
+	if (!RequestedType.IsStruct())
+	{
+		const Shader::FStructField* ActiveField = Context.ActiveStructFieldStack.Last();
+		OutResult.Code = Context.EmitInlineExpression(Scope, ActiveField->Type, TEXT("%.%"), OutResult.Code, ActiveField->Name);
 	}
 }
 
 void FExpressionSetStructField::EmitValuePreshader(FEmitContext& Context, FEmitScope& Scope, const FRequestedType& RequestedType, FEmitValuePreshaderResult& OutResult) const
 {
-	FRequestedType RequestedStructType(RequestedType);
-	RequestedStructType.ClearFieldRequested(Field);
+	FRequestedType RequestedStructType = MakeRequestedStructType(Context.ActiveStructFieldStack, RequestedType);
 	const EExpressionEvaluation StructEvaluation = Context.GetEvaluation(StructExpression, Scope, RequestedStructType);
 
-	const FRequestedType RequestedFieldType = RequestedType.GetField(Field);
+	const FRequestedType RequestedFieldType = MakeRequestedFieldType(Context.ActiveStructFieldStack, RequestedType);
 	const EExpressionEvaluation FieldEvaluation = Context.GetEvaluation(FieldExpression, Scope, RequestedFieldType);
 
 	OutResult.Type = StructType;
@@ -368,6 +390,9 @@ void FExpressionSetStructField::EmitValuePreshader(FEmitContext& Context, FEmitS
 	else
 	{
 		StructExpression->GetValuePreshader(Context, Scope, RequestedStructType, StructType, OutResult.Preshader);
+
+		FScopedActiveStructField ActiveFieldScope(Context.ActiveStructFieldStack, Field);
+
 		FieldExpression->GetValuePreshader(Context, Scope, RequestedFieldType, Field->Type, OutResult.Preshader);
 
 		check(Context.PreshaderStackPosition > 0);
@@ -375,6 +400,47 @@ void FExpressionSetStructField::EmitValuePreshader(FEmitContext& Context, FEmitS
 
 		OutResult.Preshader.WriteOpcode(Shader::EPreshaderOpcode::SetField).Write(Field->ComponentIndex).Write(Field->GetNumComponents());
 	}
+
+	if (!RequestedType.IsStruct())
+	{
+		const Shader::FStructField* ActiveField = Context.ActiveStructFieldStack.Last();
+		OutResult.Type = ActiveField->Type;
+		OutResult.Preshader.WriteOpcode(Shader::EPreshaderOpcode::GetField).Write(ActiveField->Type).Write(ActiveField->ComponentIndex);
+	}
+}
+
+FRequestedType FExpressionSetStructField::MakeRequestedStructType(const FActiveStructFieldStack& ActiveFieldStack, const FRequestedType& RequestedType) const
+{
+	FRequestedType Result;
+	if (RequestedType.IsStruct())
+	{
+		Result = RequestedType;
+	}
+	else
+	{
+		check(!RequestedType.IsObject() && RequestedType.Type.IsNumericVector());
+		// Due to legacy, it is possible for RequestedType to be numeric instead of a struct type. It implies requesting the material attribute being processed
+		const Shader::FStructField* ActiveField = ActiveFieldStack.Last();
+		Result = FRequestedType(StructType, false);
+		Result.SetFieldRequested(ActiveField, RequestedType);
+	}
+	Result.ClearFieldRequested(Field);
+	return Result;
+}
+
+FRequestedType FExpressionSetStructField::MakeRequestedFieldType(const FActiveStructFieldStack& ActiveFieldStack, const FRequestedType& RequestedType) const
+{
+	if (RequestedType.IsStruct())
+	{
+		return RequestedType.GetField(Field);
+	}
+	else
+	{
+		check(!RequestedType.IsObject() && RequestedType.Type.IsNumericVector());
+		// Due to legacy, it is possible for RequestedType to be numeric instead of a struct type. It implies requesting the material attribute being processed
+		const Shader::FStructField* ActiveField = ActiveFieldStack.Last();
+		return ActiveField == Field ? RequestedType : FRequestedType(RequestedType.Type, false);
+	}
 }
 
 void FExpressionSelect::ComputeAnalyticDerivatives(FTree& Tree, FExpressionDerivatives& OutResult) const
@@ -434,7 +500,7 @@ bool FExpressionSelect::PrepareValue(FEmitContext& Context, FEmitScope& Scope, c
 
 void FExpressionSelect::EmitValueShader(FEmitContext& Context, FEmitScope& Scope, const FRequestedType& RequestedType, FEmitValueShaderResult& OutResult) const
 {
-	const FPreparedType& ConditionType = Context.GetPreparedType(ConditionExpression);
+	const FPreparedType& ConditionType = Context.GetPreparedType(ConditionExpression, Shader::EValueType::Bool1);
 	const EExpressionEvaluation ConditionEvaluation = ConditionType.GetEvaluation(Scope, Shader::EValueType::Bool1);
 	if (IsConstantEvaluation(ConditionEvaluation))
 	{
@@ -444,7 +510,7 @@ void FExpressionSelect::EmitValueShader(FEmitContext& Context, FEmitScope& Scope
 	}
 	else
 	{
-		const Shader::FType LocalType = Context.GetType(this);
+		const Shader::FType LocalType = Context.GetResultType(this, RequestedType);
 		FEmitShaderExpression* TrueValue = TrueExpression->GetValueShader(Context, Scope, RequestedType, LocalType);
 		FEmitShaderExpression* FalseValue = FalseExpression->GetValueShader(Context, Scope, RequestedType, LocalType);
 
@@ -457,7 +523,7 @@ void FExpressionSelect::EmitValueShader(FEmitContext& Context, FEmitScope& Scope
 
 void FExpressionSelect::EmitValuePreshader(FEmitContext& Context, FEmitScope& Scope, const FRequestedType& RequestedType, FEmitValuePreshaderResult& OutResult) const
 {
-	const FPreparedType& ConditionType = Context.GetPreparedType(ConditionExpression);
+	const FPreparedType& ConditionType = Context.GetPreparedType(ConditionExpression, Shader::EValueType::Bool1);
 	const EExpressionEvaluation ConditionEvaluation = ConditionType.GetEvaluation(Scope, Shader::EValueType::Bool1);
 	if (IsConstantEvaluation(ConditionEvaluation))
 	{
@@ -467,7 +533,7 @@ void FExpressionSelect::EmitValuePreshader(FEmitContext& Context, FEmitScope& Sc
 	}
 	else
 	{
-		const Shader::FType ResultType = Context.GetType(this);
+		const Shader::FType ResultType = Context.GetResultType(this, RequestedType);
 		auto ConvertToResultType = [&ResultType](const Shader::FType& Type, Shader::FPreshaderData& Preshader)
 		{
 			if (Type == ResultType)
@@ -557,7 +623,7 @@ void FExpressionDerivative::EmitValuePreshader(FEmitContext& Context, FEmitScope
 {
 	// Derivative of a constant is 0
 	Context.PreshaderStackPosition++;
-	OutResult.Type = Context.GetType(this);
+	OutResult.Type = Context.GetResultType(this, RequestedType);
 	OutResult.Preshader.WriteOpcode(Shader::EPreshaderOpcode::ConstantZero).Write(OutResult.Type);
 }
 
@@ -914,7 +980,7 @@ bool FExpressionAppend::PrepareValue(FEmitContext& Context, FEmitScope& Scope, c
 
 void FExpressionAppend::EmitValueShader(FEmitContext& Context, FEmitScope& Scope, const FRequestedType& RequestedType, FEmitValueShaderResult& OutResult) const
 {
-	const Private::FAppendTypes Types = Private::GetAppendTypes(RequestedType, Context.GetPreparedType(Lhs), Context.GetPreparedType(Rhs));
+	const Private::FAppendTypes Types = Private::GetAppendTypes(RequestedType, Context.GetPreparedType(Lhs, RequestedType), Context.GetPreparedType(Rhs, RequestedType));
 	const Shader::FType ResultType = Types.ResultType.GetResultType();
 	FEmitShaderExpression* LhsValue = Lhs->GetValueShader(Context, Scope, Types.LhsRequestedType, Types.LhsType);
 
@@ -943,7 +1009,7 @@ void FExpressionAppend::EmitValueShader(FEmitContext& Context, FEmitScope& Scope
 
 void FExpressionAppend::EmitValuePreshader(FEmitContext& Context, FEmitScope& Scope, const FRequestedType& RequestedType, FEmitValuePreshaderResult& OutResult) const
 {
-	const Private::FAppendTypes Types = Private::GetAppendTypes(RequestedType, Context.GetPreparedType(Lhs), Context.GetPreparedType(Rhs));
+	const Private::FAppendTypes Types = Private::GetAppendTypes(RequestedType, Context.GetPreparedType(Lhs, RequestedType), Context.GetPreparedType(Rhs, RequestedType));
 	Lhs->GetValuePreshader(Context, Scope, Types.LhsRequestedType, OutResult.Preshader);
 	if (Types.RhsType != Shader::EValueType::Void)
 	{
@@ -1272,7 +1338,7 @@ bool FStatementIf::Prepare(FEmitContext& Context, FEmitScope& Scope) const
 void FStatementIf::EmitShader(FEmitContext& Context, FEmitScope& Scope) const
 {
 	FEmitShaderNode* Dependency = nullptr;
-	const FPreparedType& ConditionType = Context.GetPreparedType(ConditionExpression);
+	const FPreparedType& ConditionType = Context.GetPreparedType(ConditionExpression, Shader::EValueType::Bool1);
 	const EExpressionEvaluation ConditionEvaluation = ConditionType.GetEvaluation(Scope, Shader::EValueType::Bool1);
 	if (IsConstantEvaluation(ConditionEvaluation))
 	{
diff --git a/Engine/Source/Runtime/Engine/Private/HLSLTree/HLSLTreeEmit.cpp b/Engine/Source/Runtime/Engine/Private/HLSLTree/HLSLTreeEmit.cpp
--- a/Engine/Source/Runtime/Engine/Private/HLSLTree/HLSLTreeEmit.cpp
+++ b/Engine/Source/Runtime/Engine/Private/HLSLTree/HLSLTreeEmit.cpp
@@ -208,7 +208,7 @@ FEmitContext::~FEmitContext()
 	{
 		delete It.Value;
 	}
-	for (TMap<const FExpression*, FPrepareValueResult*>::TIterator It(PrepareValueMap); It; ++It)
+	for (TMap<FXxHash64, FPrepareValueResult*>::TIterator It(PrepareValueMap); It; ++It)
 	{
 		FPrepareValueResult* Value = It.Value();
 		if (Value)
@@ -341,6 +341,19 @@ void EmitCustomHLSL(const FEmitCustomHLSL& EmitCustomHLSL, const TCHAR* Paramete
 	OutCode.Append(TEXT("\treturn Result;\n"));
 	OutCode.Append(TEXT("}\n"));
 }
+
+FXxHash64 GetPrepareValueHash(const FExpression* Expression, bool bRequestedStruct)
+{
+	FXxHash64Builder Hasher;
+	Hasher.Update(&Expression, sizeof(Expression));
+	Hasher.Update(&bRequestedStruct, sizeof(bRequestedStruct));
+	return Hasher.Finalize();
+}
+
+FXxHash64 GetPrepareValueHash(const FExpression* Expression, const FRequestedType& RequestedType)
+{
+	return GetPrepareValueHash(Expression, RequestedType.IsStruct());
+}
 } // namespace Private
 
 void FEmitContext::EmitDeclarationsCode(FStringBuilderBase& OutCode)
@@ -363,31 +376,33 @@ void FEmitContext::EmitDeclarationsCode(FStringBuilderBase& OutCode)
 	}
 }
 
-const FPreparedType& FEmitContext::GetPreparedType(const FExpression* Expression) const
+FPreparedType FEmitContext::GetPreparedType(const FExpression* Expression, const FRequestedType& RequestedType) const
 {
-	static const FPreparedType VoidType;
-	FPrepareValueResult const* const* PrevResult = PrepareValueMap.Find(Expression);
-	const FPrepareValueResult* Result = PrevResult ? *PrevResult : nullptr;
-	if (Result)
-	{
-		return Result->PreparedType;
-	}
-	return VoidType;
+	const FXxHash64 Hash = Private::GetPrepareValueHash(Expression, RequestedType);
+	FPrepareValueResult const* const* PrevResult = PrepareValueMap.Find(Hash);
+	return PrevResult ? (*PrevResult)->PreparedType : FPreparedType();
 }
 
-FRequestedType FEmitContext::GetRequestedType(const FExpression* Expression) const
+Shader::FType FEmitContext::GetResultType(const FExpression* Expression, const FRequestedType& RequestedType) const
 {
-	return GetPreparedType(Expression).GetRequestedType();
+	return GetPreparedType(Expression, RequestedType).GetResultType();
 }
 
-Shader::FType FEmitContext::GetType(const FExpression* Expression) const
+Shader::FType FEmitContext::GetTypeForPinColoring(const FExpression* Expression) const
 {
-	return GetPreparedType(Expression).GetResultType();
+	FXxHash64 Hash = Private::GetPrepareValueHash(Expression, true);
+	FPrepareValueResult const* const* Found = PrepareValueMap.Find(Hash);
+	if (!Found)
+	{
+		Hash = Private::GetPrepareValueHash(Expression, false);
+		Found = PrepareValueMap.Find(Hash);
+	}
+	return Found ? (*Found)->PreparedType.GetResultType() : Shader::FType();
 }
 
 EExpressionEvaluation FEmitContext::GetEvaluation(const FExpression* Expression, const FEmitScope& Scope, const FRequestedType& RequestedType) const
 {
-	return GetPreparedType(Expression).GetEvaluation(Scope, RequestedType);
+	return GetPreparedType(Expression, RequestedType).GetEvaluation(Scope, RequestedType);
 }
 
 bool AllComponentsRequestedInPast(const FRequestedType& RequestedType, const FRequestedType& PastRequestedType)
@@ -421,13 +436,14 @@ FPreparedType FEmitContext::PrepareExpression(const FExpression* InExpression, F
 		return FPreparedType();
 	}
 
-	FPrepareValueResult** PrevResult = PrepareValueMap.Find(InExpression);
+	FXxHash64 Hash = Private::GetPrepareValueHash(InExpression, RequestedType);
+	FPrepareValueResult** PrevResult = PrepareValueMap.Find(Hash);
 	FPrepareValueResult* Result = PrevResult ? *PrevResult : nullptr;
 	if (!Result)
 	{
 		check(!bMarkLiveValues); // value should already be prepared at this point
 		Result = new(*Allocator) FPrepareValueResult();
-		PrepareValueMap.Add(InExpression, Result);
+		PrepareValueMap.Add(Hash, Result);
 	}
 
 	if (RequestedType.IsEmpty() && !Result->PreparedType.IsVoid())
@@ -450,10 +466,14 @@ FPreparedType FEmitContext::PrepareExpression(const FExpression* InExpression, F
 		return Result->PreparedType;
 	}
 
-	FXxHash64Builder Hasher;
-	Hasher.Update(&InExpression, sizeof(InExpression));
-	Hasher.Update(&bMarkLiveValues, sizeof(bMarkLiveValues));
-	const FXxHash64 Hash = Hasher.Finalize();
+	{
+		FXxHash64Builder Hasher;
+		Hasher.Update(&InExpression, sizeof(InExpression));
+		Hasher.Update(&bMarkLiveValues, sizeof(bMarkLiveValues));
+		const bool bRequestedStruct = RequestedType.IsStruct();
+		Hasher.Update(&bRequestedStruct, sizeof(bRequestedStruct));
+		Hash = Hasher.Finalize();
+	}
 
 	FRequestedType** PastRequestedTypePtr = RequestedTypeTracker.Find(Hash);
 	if (PastRequestedTypePtr)
@@ -488,7 +508,6 @@ FPreparedType FEmitContext::PrepareExpression(const FExpression* InExpression, F
 	}
 	else
 	{
-		check(Result->PreparedType.IsVoid() || bMarkLiveValues);
 		FRequestedType* StackRequestedType = new (*Allocator) FRequestedType(RequestedType);
 		// Turn Bool into Int
 		StackRequestedType->Type = Shader::CombineTypes(StackRequestedType->Type, StackRequestedType->Type);
@@ -1016,7 +1035,7 @@ FEmitShaderExpression* FEmitContext::EmitPreshaderOrConstant(FEmitScope& Scope,
 		return ShaderValue;
 	}
 
-	const FPreparedType& PreparedType = GetPreparedType(Expression);
+	const FPreparedType& PreparedType = GetPreparedType(Expression, RequestedType);
 
 	Shader::FPreshaderStack Stack;
 	const Shader::FPreshaderValue PreshaderConstantValue = LocalPreshader.EvaluateConstant(*Material, Stack);
@@ -1378,7 +1397,7 @@ FEmitShaderExpression* FEmitContext::EmitCustomHLSL(FEmitScope& Scope, FStringVi
 	FHasher Hasher;
 	for (const FCustomHLSLInput& Input : Inputs)
 	{
-		const Shader::FType InputType = GetType(Input.Expression);
+		const Shader::FType InputType = GetResultType(Input.Expression, Shader::EValueType::Any);
 		AppendHash(Hasher, Input.Name);
 		AppendHash(Hasher, InputType);
 
@@ -1423,9 +1442,10 @@ FEmitShaderExpression* FEmitContext::EmitCustomHLSL(FEmitScope& Scope, FStringVi
 
 	TStringBuilder<1024> FormattedCode;
 	FormattedCode.Appendf(TEXT("CustomExpression%d(Parameters"), EmitCustomHLSL->Index);
-	for (const FCustomHLSLInput& Input : Inputs)
+	for (int32 InputIndex = 0; InputIndex < Inputs.Num(); ++InputIndex)
 	{
-		FEmitShaderExpression* EmitInputExpression = Input.Expression->GetValueShader(*this, Scope);
+		const FCustomHLSLInput& Input = Inputs[InputIndex];
+		FEmitShaderExpression* EmitInputExpression = Input.Expression->GetValueShader(*this, Scope, EmitInputs[InputIndex].Type);
 		FormattedCode.Appendf(TEXT(", %s"), EmitInputExpression->Reference);
 		Dependencies.Add(EmitInputExpression);
 	}
@@ -1445,7 +1465,7 @@ void FEmitContext::Finalize()
 	
 	// Don't reset Expression/Preshader maps, allow future passes to share matching preshaders/expressions
 
-	for (TMap<const FExpression*, FPrepareValueResult*>::TIterator It(PrepareValueMap); It; ++It)
+	for (TMap<FXxHash64, FPrepareValueResult*>::TIterator It(PrepareValueMap); It; ++It)
 	{
 		FPrepareValueResult* Value = It.Value();
 		if (Value)
diff --git a/Engine/Source/Runtime/Engine/Private/HLSLTree/HLSLTreeOperation.cpp b/Engine/Source/Runtime/Engine/Private/HLSLTree/HLSLTreeOperation.cpp
--- a/Engine/Source/Runtime/Engine/Private/HLSLTree/HLSLTreeOperation.cpp
+++ b/Engine/Source/Runtime/Engine/Private/HLSLTree/HLSLTreeOperation.cpp
@@ -659,7 +659,7 @@ bool FExpressionOperation::PrepareValue(FEmitContext& Context, FEmitScope& Scope
 	{
 		if (Context.bMarkLiveValues)
 		{
-			InputPreparedType[Index] = Context.GetPreparedType(Inputs[Index]);
+			InputPreparedType[Index] = Context.GetPreparedType(Inputs[Index], RequestedTypes.InputType[Index]);
 			bMarkLiveInput[Index] = true;
 		}
 		else
@@ -737,13 +737,14 @@ bool FExpressionOperation::PrepareValue(FEmitContext& Context, FEmitScope& Scope
 void FExpressionOperation::EmitValueShader(FEmitContext& Context, FEmitScope& Scope, const FRequestedType& RequestedType, FEmitValueShaderResult& OutResult) const
 {
 	const FOperationDescription OpDesc = GetOperationDescription(Op);
+	const Private::FOperationRequestedTypes RequestedTypes = Private::GetOperationRequestedTypes(Op, RequestedType);
+
 	FPreparedType InputPreparedType[MaxInputs];
 	for (int32 Index = 0; Index < OpDesc.NumInputs; ++Index)
 	{
-		InputPreparedType[Index] = Context.GetPreparedType(Inputs[Index]);
+		InputPreparedType[Index] = Context.GetPreparedType(Inputs[Index], RequestedTypes.InputType[Index]);
 	}
 
-	const Private::FOperationRequestedTypes RequestedTypes = Private::GetOperationRequestedTypes(Op, RequestedType);
 	const Private::FOperationTypes Types = Private::GetOperationTypes(Op, MakeArrayView(InputPreparedType, OpDesc.NumInputs));
 	FEmitShaderExpression* InputValue[MaxInputs] = { nullptr };
 	for (int32 Index = 0; Index < OpDesc.NumInputs; ++Index)
@@ -850,13 +851,14 @@ void FExpressionOperation::EmitValueShader(FEmitContext& Context, FEmitScope& Sc
 void FExpressionOperation::EmitValuePreshader(FEmitContext& Context, FEmitScope& Scope, const FRequestedType& RequestedType, FEmitValuePreshaderResult& OutResult) const
 {
 	const FOperationDescription OpDesc = GetOperationDescription(Op);
+	const Private::FOperationRequestedTypes RequestedTypes = Private::GetOperationRequestedTypes(Op, RequestedType);
+
 	FPreparedType InputPreparedType[MaxInputs];
 	for (int32 Index = 0; Index < OpDesc.NumInputs; ++Index)
 	{
-		InputPreparedType[Index] = Context.GetPreparedType(Inputs[Index]);
+		InputPreparedType[Index] = Context.GetPreparedType(Inputs[Index], RequestedTypes.InputType[Index]);
 	}
 
-	const Private::FOperationRequestedTypes RequestedTypes = Private::GetOperationRequestedTypes(Op, RequestedType);
 	const Private::FOperationTypes Types = Private::GetOperationTypes(Op, MakeArrayView(InputPreparedType, OpDesc.NumInputs));
 	check(OpDesc.PreshaderOpcode != Shader::EPreshaderOpcode::Nop);
 
diff --git a/Engine/Source/Runtime/Engine/Private/Materials/MaterialExpressionHLSL.cpp b/Engine/Source/Runtime/Engine/Private/Materials/MaterialExpressionHLSL.cpp
--- a/Engine/Source/Runtime/Engine/Private/Materials/MaterialExpressionHLSL.cpp
+++ b/Engine/Source/Runtime/Engine/Private/Materials/MaterialExpressionHLSL.cpp
@@ -53,6 +53,7 @@
 #include "Materials/MaterialExpressionVertexNormalWS.h"
 #include "Materials/MaterialExpressionVertexTangentWS.h"
 #include "Materials/MaterialExpressionTextureCoordinate.h"
+#include "Materials/MaterialExpressionTextureProperty.h"
 #include "Materials/MaterialExpressionLightmapUVs.h"
 #include "Materials/MaterialExpressionEyeAdaptation.h"
 #include "Materials/MaterialExpressionEyeAdaptationInverse.h"
@@ -64,6 +65,7 @@
 #include "Materials/MaterialExpressionParticleRelativeTime.h"
 #include "Materials/MaterialExpressionParticleRandom.h"
 #include "Materials/MaterialExpressionParticleSize.h"
+#include "Materials/MaterialExpressionParticleMacroUV.h"
 #include "Materials/MaterialExpressionSkyLightEnvMapSample.h"
 #include "Materials/MaterialExpressionTextureSample.h"
 #include "Materials/MaterialExpressionTextureSampleParameter.h"
@@ -656,7 +658,7 @@ bool UMaterialExpressionCurveAtlasRowParameter::GenerateHLSLExpression(FMaterial
 	}
 
 	FTree& Tree = Generator.GetTree();
-	const FExpression* ExpressionTextureSize = Tree.NewExpression<Material::FExpressionTextureSize>(ExpressionTexture);
+	const FExpression* ExpressionTextureSize = Tree.NewExpression<Material::FExpressionTextureProperty>(ExpressionTexture, TMTM_TextureSize);
 	const FExpression* ExpressionTextureHeight = Tree.NewExpression<FExpressionSwizzle>(MakeSwizzleMask(false, true, false, false), ExpressionTextureSize);
 	const FExpression* ExpressionCoordU = InputTime.AcquireHLSLExpressionOrConstant(Generator, Scope, 0.0f);
 	const FExpression* ExpressionCoordV = Tree.NewMul(Tree.NewAdd(ExpressionSlot, Tree.NewConstant(0.5f)), Tree.NewRcp(ExpressionTextureHeight));
@@ -1004,6 +1006,25 @@ bool UMaterialExpressionTextureCoordinate::GenerateHLSLExpression(FMaterialHLSLG
 	return true;
 }
 
+bool UMaterialExpressionTextureProperty::GenerateHLSLExpression(FMaterialHLSLGenerator& Generator, UE::HLSLTree::FScope& Scope, int32 OutputIndex, UE::HLSLTree::FExpression const*& OutExpression) const
+{
+	using namespace UE::HLSLTree;
+
+	if (Property < 0 || Property >= TMTM_MAX)
+	{
+		return Generator.Errorf(TEXT("Invalid texture property %d"), Property);
+	}
+
+	const FExpression* TextureExpression = TextureObject.AcquireHLSLExpression(Generator, Scope);
+	if (!TextureExpression)
+	{
+		return false;
+	}
+
+	OutExpression = Generator.GetTree().NewExpression<Material::FExpressionTextureProperty>(TextureExpression, (EMaterialExposedTextureProperty)Property);
+	return true;
+}
+
 bool UMaterialExpressionLightmapUVs::GenerateHLSLExpression(FMaterialHLSLGenerator& Generator, UE::HLSLTree::FScope& Scope, int32 OutputIndex, UE::HLSLTree::FExpression const*& OutExpression) const
 {
 	using namespace UE::HLSLTree::Material;
@@ -1095,6 +1116,13 @@ bool UMaterialExpressionParticleSize::GenerateHLSLExpression(FMaterialHLSLGenera
 	return true;
 }
 
+bool UMaterialExpressionParticleMacroUV::GenerateHLSLExpression(FMaterialHLSLGenerator& Generator, UE::HLSLTree::FScope& Scope, int32 OutputIndex, UE::HLSLTree::FExpression const*& OutExpression) const
+{
+	using namespace UE::HLSLTree;
+	OutExpression = Generator.GetTree().NewExpression<FExpressionInlineCustomHLSL>(UE::Shader::EValueType::Float2, TEXT("GetParticleMacroUV(Parameters)"));
+	return true;
+}
+
 bool UMaterialExpressionTextureObject::GenerateHLSLExpression(FMaterialHLSLGenerator& Generator, UE::HLSLTree::FScope& Scope, int32 OutputIndex, UE::HLSLTree::FExpression const*& OutExpression) const
 {
 	using namespace UE::HLSLTree;
@@ -2525,14 +2553,16 @@ bool UMaterialExpressionFunctionInput::GenerateHLSLExpression(FMaterialHLSLGener
 bool UMaterialExpressionMaterialFunctionCall::GenerateHLSLExpression(FMaterialHLSLGenerator& Generator, UE::HLSLTree::FScope& Scope, int32 OutputIndex, UE::HLSLTree::FExpression const*& OutExpression) const
 {
 	using namespace UE::HLSLTree;
-	TArray<const FExpression*, TInlineAllocator<16>> ConnectedInputs;
-	ConnectedInputs.Reserve(FunctionInputs.Num());
+
+	TArray<FMaterialHLSLGenerator::FConnectedInput> ConnectedInputs;
+	ConnectedInputs.Empty(FunctionInputs.Num());
+
 	for (int32 InputIndex = 0; InputIndex < FunctionInputs.Num(); ++InputIndex)
 	{
-		// ConnectedInputs are the inputs from the UMaterialFunctionCall object
-		// We want to connect the UMaterialExpressionFunctionInput from the UMaterialFunction to whatever UMaterialExpression is passed to the UMaterialFunctionCall
-		const FExpression* ConnectedInput = FunctionInputs[InputIndex].Input.TryAcquireHLSLExpression(Generator, Scope);
-		ConnectedInputs.Add(ConnectedInput);
+		// We cannot call AcquireExpression on the FExpressionInput now because it will form an infinite loop if at least one input
+		// is directly or indirectly connected to another output of this function call. Instead, we store a reference to the FExpressionInput
+		// and process it later when we GenerateHLSLExpression for the corresponding UMaterialExpressionFunctionInput inside the material function.
+		ConnectedInputs.Emplace(&FunctionInputs[InputIndex].Input, &Scope);
 	}
 
 	OutExpression = Generator.GenerateFunctionCall(Scope, MaterialFunction, GlobalParameter, INDEX_NONE, ConnectedInputs, OutputIndex);
@@ -2554,7 +2584,7 @@ bool UMaterialExpressionMaterialAttributeLayers::GenerateHLSLExpression(FMateria
 	TArray<FFunctionExpressionOutput> FunctionOutputs;
 
 	const FExpression* ExpressionLayerInput = Input.AcquireHLSLExpressionOrConstant(Generator, Scope, Generator.GetMaterialAttributesDefaultValue());
-	TArray<const FExpression*, TInlineAllocator<1>> LayerInputExpressions;
+	TArray<FMaterialHLSLGenerator::FConnectedInput, TInlineAllocator<1>> LayerInputExpressions;
 
 	TArray<const FExpression*, TInlineAllocator<16>> LayerExpressions;
 	LayerExpressions.Reserve(MaterialLayers.Layers.Num());
@@ -2583,7 +2613,7 @@ bool UMaterialExpressionMaterialAttributeLayers::GenerateHLSLExpression(FMateria
 			LayerInputExpressions.Empty(1);
 			if (FunctionInputs.Num() == 1)
 			{
-				LayerInputExpressions.Add(ExpressionLayerInput);
+				LayerInputExpressions.Emplace(ExpressionLayerInput);
 			}
 			LayerExpression = Generator.GenerateFunctionCall(Scope, LayerFunction, LayerParameter, LayerIndex, LayerInputExpressions, 0);
 		}
@@ -2596,7 +2626,7 @@ bool UMaterialExpressionMaterialAttributeLayers::GenerateHLSLExpression(FMateria
 		return Generator.Error(TEXT("No layers"));
 	}
 
-	TArray<const FExpression*, TInlineAllocator<2>> BlendInputExpressions;
+	TArray<FMaterialHLSLGenerator::FConnectedInput, TInlineAllocator<2>> BlendInputExpressions;
 	for (int32 BlendIndex = 0; BlendIndex < MaterialLayers.Blends.Num(); ++BlendIndex)
 	{
 		const int32 LayerIndex = BlendIndex + 1;
@@ -2633,8 +2663,8 @@ bool UMaterialExpressionMaterialAttributeLayers::GenerateHLSLExpression(FMateria
 				}
 
 				BlendInputExpressions.Empty(2);
-				BlendInputExpressions.Add(BottomLayerExpression);
-				BlendInputExpressions.Add(LayerExpression);
+				BlendInputExpressions.Emplace(BottomLayerExpression);
+				BlendInputExpressions.Emplace(LayerExpression);
 				BottomLayerExpression = Generator.GenerateFunctionCall(Scope, BlendFunction, BlendParameter, BlendIndex, BlendInputExpressions, 0);
 			}
 			else
diff --git a/Engine/Source/Runtime/Engine/Private/Materials/MaterialHLSLGenerator.cpp b/Engine/Source/Runtime/Engine/Private/Materials/MaterialHLSLGenerator.cpp
--- a/Engine/Source/Runtime/Engine/Private/Materials/MaterialHLSLGenerator.cpp
+++ b/Engine/Source/Runtime/Engine/Private/Materials/MaterialHLSLGenerator.cpp
@@ -455,8 +455,15 @@ const UE::HLSLTree::FExpression* FMaterialHLSLGenerator::AcquireExpression(UE::H
 
 			for (int32 FunctionInputIndex = 0; FunctionInputIndex < FunctionCall->ConnectedInputs.Num(); ++FunctionInputIndex)
 			{
-				const FExpression* InputExpression = FunctionCall->ConnectedInputs[FunctionInputIndex];
-				Hasher.AppendData(&InputExpression, sizeof(InputExpression));
+				const FConnectedInput& ConnectedInput = FunctionCall->ConnectedInputs[FunctionInputIndex];
+				if (ConnectedInput.Input)
+				{
+					Hasher.AppendData(&ConnectedInput.Input, sizeof(ConnectedInput.Input));
+				}
+				else
+				{
+					Hasher.AppendData(&ConnectedInput.Expression, sizeof(ConnectedInput.Expression));
+				}
 			}
 		}
 
@@ -496,17 +503,25 @@ const UE::HLSLTree::FExpression* FMaterialHLSLGenerator::AcquireExpression(UE::H
 const UE::HLSLTree::FExpression* FMaterialHLSLGenerator::AcquireFunctionInputExpression(UE::HLSLTree::FScope& Scope, const UMaterialExpressionFunctionInput* MaterialExpression)
 {
 	using namespace UE::HLSLTree;
-	const FFunctionCallEntry* FunctionEntry = FunctionCallStack.Last();
+	
+	// Need to pop because we are going out of the function when processing expressions connected to the inputs
+	FFunctionCallEntry* FunctionEntry = FunctionCallStack.Pop(false);
 	const FExpression* InputExpression = nullptr;
 	int32 InputIndex = INDEX_NONE;
+
 	if (FunctionEntry->MaterialFunction)
 	{
 		for (int32 Index = 0; Index < FunctionEntry->FunctionInputs.Num(); ++Index)
 		{
 			if (FunctionEntry->FunctionInputs[Index] == MaterialExpression)
 			{
-				InputIndex = Index;;
-				InputExpression = FunctionEntry->ConnectedInputs[Index];
+				InputIndex = Index;
+				FConnectedInput& ConnectedInput = FunctionEntry->ConnectedInputs[Index];
+				if (!ConnectedInput.Expression && ConnectedInput.Input)
+				{
+					ConnectedInput.Expression = ConnectedInput.Input->TryAcquireHLSLExpression(*this, *ConnectedInput.Scope);
+				}
+				InputExpression = ConnectedInput.Expression;
 				break;
 			}
 		}
@@ -519,6 +534,8 @@ const UE::HLSLTree::FExpression* FMaterialHLSLGenerator::AcquireFunctionInputExp
 		}
 	}
 
+	FunctionCallStack.Push(FunctionEntry);
+
 	if (!InputExpression && (MaterialExpression->bUsePreviewValueAsDefault || !FunctionEntry->MaterialFunction))
 	{
 		// Either we're previewing the material function, or the input isn't connected and we're using preview as default value
@@ -627,7 +644,7 @@ const UE::HLSLTree::FExpression* FMaterialHLSLGenerator::GenerateFunctionCall(UE
 	UMaterialFunctionInterface* MaterialFunction,
 	EMaterialParameterAssociation InParameterAssociation,
 	int32 InParameterIndex,
-	TArrayView<const UE::HLSLTree::FExpression*> ConnectedInputs,
+	TArrayView<FConnectedInput> ConnectedInputs,
 	int32 OutputIndex)
 {
 	using namespace UE::HLSLTree;
@@ -677,13 +694,19 @@ const UE::HLSLTree::FExpression* FMaterialHLSLGenerator::GenerateFunctionCall(UE
 		{
 			// FunctionInputs are the inputs from the UMaterialFunction object
 			const FFunctionExpressionInput& FunctionInput = FunctionInputs[InputIndex];
+			LocalFunctionInputs.Add(FunctionInput.ExpressionInput);
 
 			// ConnectedInputs are the inputs from the UMaterialFunctionCall object
 			// We want to connect the UMaterialExpressionFunctionInput from the UMaterialFunction to whatever UMaterialExpression is passed to the UMaterialFunctionCall
-			const FExpression* ConnectedInput = ConnectedInputs[InputIndex];
-
-			LocalFunctionInputs.Add(FunctionInput.ExpressionInput);
-			Hasher.Update((uint8*)&ConnectedInput, sizeof(ConnectedInput));
+			const FConnectedInput& ConnectedInput = ConnectedInputs[InputIndex];
+			if (ConnectedInput.Input)
+			{
+				Hasher.Update(&ConnectedInput.Input, sizeof(ConnectedInput.Input));
+			}
+			else
+			{
+				Hasher.Update(&ConnectedInput.Expression, sizeof(ConnectedInput.Expression));
+			}
 		}
 		Hash = Hasher.Finalize();
 	}
diff --git a/Engine/Source/Runtime/Engine/Private/Materials/MaterialHLSLTree.cpp b/Engine/Source/Runtime/Engine/Private/Materials/MaterialHLSLTree.cpp
--- a/Engine/Source/Runtime/Engine/Private/Materials/MaterialHLSLTree.cpp
+++ b/Engine/Source/Runtime/Engine/Private/Materials/MaterialHLSLTree.cpp
@@ -578,17 +578,25 @@ void FExpressionParameter::EmitValuePreshader(FEmitContext& Context, FEmitScope&
 
 bool FExpressionParameter::EmitValueObject(FEmitContext& Context, FEmitScope& Scope, const FName& ObjectTypeName, void* OutObjectBase) const
 {
-	UTexture* Texture = Cast<UTexture>(ParameterMeta.Value.AsTextureObject());
-	if (Texture && ObjectTypeName == FMaterialTextureValue::GetTypeName())
+	if (ObjectTypeName != FMaterialTextureValue::GetTypeName())
 	{
-		FMaterialTextureValue& OutObject = *static_cast<FMaterialTextureValue*>(OutObjectBase);
-		OutObject.ParameterInfo = ParameterInfo;
-		OutObject.Texture = Texture;
-		OutObject.SamplerType = TextureSamplerType;
-		OutObject.ExternalTextureGuid = ExternalTextureGuid;
-		return true;
+		return false;
 	}
-	return false;
+
+	UObject* TextureObject = ParameterMeta.Value.AsTextureObject();
+	FMaterialTextureValue& OutObject = *static_cast<FMaterialTextureValue*>(OutObjectBase);
+
+	OutObject.Texture = Cast<UTexture>(TextureObject);
+	OutObject.RuntimeVirtualTexture = Cast<URuntimeVirtualTexture>(TextureObject);
+	if (!OutObject.Texture && !OutObject.RuntimeVirtualTexture)
+	{
+		return false;
+	}
+
+	OutObject.SamplerType = TextureSamplerType;
+	OutObject.ExternalTextureGuid = ExternalTextureGuid;
+	OutObject.ParameterInfo = ParameterInfo;
+	return true;
 }
 
 bool FExpressionParameter::EmitCustomHLSLParameter(FEmitContext& Context, FEmitScope& Scope, const FName& ObjectTypeName, const TCHAR* ParameterName, FEmitCustomHLSLParameterResult& OutResult) const
@@ -726,7 +734,15 @@ bool FExpressionTextureSample::PrepareValue(FEmitContext& Context, FEmitScope& S
 		return Context.Error(TEXT("Expected texture"));
 	}
 
-	const EMaterialValueType TextureMaterialType = TextureValue.Texture->GetMaterialType();
+	EMaterialValueType TextureMaterialType = MCT_Unknown;
+	if (TextureValue.Texture)
+	{
+		TextureMaterialType = TextureValue.Texture->GetMaterialType();
+	}
+	else if (TextureValue.RuntimeVirtualTexture)
+	{
+		TextureMaterialType = MCT_TextureVirtual;
+	}
 	const FRequestedType RequestedTexCoordType = Private::GetTexCoordType(TextureMaterialType);
 	const FPreparedType& TexCoordType = Context.PrepareExpression(TexCoordExpression, Scope, RequestedTexCoordType);
 	if (TexCoordType.IsVoid())
@@ -918,15 +934,27 @@ void FExpressionTextureSample::EmitValueShader(FEmitContext& Context, FEmitScope
 	FMaterialTextureValue TextureValue;
 	verify(TextureExpression->GetValueObject(Context, Scope, TextureValue));
 	UTexture* Texture = TextureValue.Texture;
-	check(Texture);
+	URuntimeVirtualTexture* RuntimeVirtualTexture = TextureValue.RuntimeVirtualTexture;
+
+	UObject* TextureObject = nullptr;
+	EMaterialValueType TextureType = MCT_Unknown;
+	if (Texture)
+	{
+		TextureObject = Texture;
+		TextureType = Texture->GetMaterialType();
+	}
+	else if (RuntimeVirtualTexture)
+	{
+		TextureObject = RuntimeVirtualTexture;
+		TextureType = MCT_TextureVirtual;
+	}
 
-	const EMaterialValueType TextureType = Texture->GetMaterialType();
 	const Shader::EValueType TexCoordType = Private::GetTexCoordType(TextureType);
 
 	TextureAddress StaticAddressX = TA_Wrap;
 	TextureAddress StaticAddressY = TA_Wrap;
 	TextureAddress StaticAddressZ = TA_Wrap;
-	if (Texture->Source.GetNumBlocks() > 1)
+	if (Texture && Texture->Source.GetNumBlocks() > 1)
 	{
 		// UDIM (multi-block) texture are forced to use wrap address mode
 		// This is important for supporting VT stacks made from UDIMs with differing number of blocks, as this requires wrapping vAddress for certain layers
@@ -939,9 +967,12 @@ void FExpressionTextureSample::EmitValueShader(FEmitContext& Context, FEmitScope
 		switch (SamplerSource)
 		{
 		case SSM_FromTextureAsset:
-			StaticAddressX = Texture->GetTextureAddressX();
-			StaticAddressY = Texture->GetTextureAddressY();
-			StaticAddressZ = Texture->GetTextureAddressZ();
+			if (Texture)
+			{
+				StaticAddressX = Texture->GetTextureAddressX();
+				StaticAddressY = Texture->GetTextureAddressY();
+				StaticAddressZ = Texture->GetTextureAddressZ();
+			}
 			break;
 		case SSM_Wrap_WorldGroupSettings:
 			StaticAddressX = TA_Wrap;
@@ -989,7 +1020,7 @@ void FExpressionTextureSample::EmitValueShader(FEmitContext& Context, FEmitScope
 
 		FMaterialTextureParameterInfo TextureParameterInfo;
 		TextureParameterInfo.ParameterInfo = TextureValue.ParameterInfo;
-		TextureParameterInfo.TextureIndex = Context.Material->GetReferencedTextures().Find(TextureValue.Texture);
+		TextureParameterInfo.TextureIndex = Context.Material->GetReferencedTextures().Find(TextureObject);
 		TextureParameterInfo.SamplerSource = SamplerSource;
 		check(TextureParameterInfo.TextureIndex != INDEX_NONE);
 		const int32 TextureParameterIndex = Context.MaterialCompilationOutput->UniformExpressionSet.FindOrAddTextureParameter(EMaterialTextureParameterType::Virtual, TextureParameterInfo);
@@ -1021,6 +1052,7 @@ void FExpressionTextureSample::EmitValueShader(FEmitContext& Context, FEmitScope
 		}
 		else
 		{
+			check(Texture);
 			// Using Source size because we care about the aspect ratio of each block (each block of multi-block texture must have same aspect ratio)
 			// We can still combine multi-block textures of different block aspect ratios, as long as each block has the same ratio
 			// This is because we only need to overlay VT pages from within a given block
@@ -1174,7 +1206,7 @@ void FExpressionTextureSample::EmitValueShader(FEmitContext& Context, FEmitScope
 	OutResult.Code = Context.EmitExpression(Scope, Shader::EValueType::Float4, TEXT("ApplyMaterialSamplerType(%, %)"), EmitTextureResult, TextureValue.SamplerType);
 }
 
-bool FExpressionTextureSize::PrepareValue(FEmitContext& Context, FEmitScope& Scope, const FRequestedType& RequestedType, FPrepareValueResult& OutResult) const
+bool FExpressionTextureProperty::PrepareValue(FEmitContext& Context, FEmitScope& Scope, const FRequestedType& RequestedType, FPrepareValueResult& OutResult) const
 {
 	const FPreparedType& TextureType = Context.PrepareExpression(TextureExpression, Scope, FMaterialTextureValue::GetTypeName());
 	if (TextureType.Type.ObjectType != FMaterialTextureValue::GetTypeName())
@@ -1182,20 +1214,32 @@ bool FExpressionTextureSize::PrepareValue(FEmitContext& Context, FEmitScope& Sco
 		return Context.Error(TEXT("Expected texture"));
 	}
 
-	return OutResult.SetType(Context, RequestedType, EExpressionEvaluation::Preshader, Shader::EValueType::Float2);
+	return OutResult.SetType(Context, RequestedType, EExpressionEvaluation::Preshader, Shader::EValueType::Float3);
 }
 
-void FExpressionTextureSize::EmitValuePreshader(FEmitContext& Context, FEmitScope& Scope, const FRequestedType& RequestedType, FEmitValuePreshaderResult& OutResult) const
+void FExpressionTextureProperty::EmitValuePreshader(FEmitContext& Context, FEmitScope& Scope, const FRequestedType& RequestedType, FEmitValuePreshaderResult& OutResult) const
 {
 	FMaterialTextureValue TextureValue;
 	TextureExpression->GetValueObject(Context, Scope, TextureValue);
 
-	const int32 TextureIndex = Context.Material->GetReferencedTextures().Find(TextureValue.Texture);
+	UObject* TextureObejct = nullptr;
+	if (TextureValue.Texture)
+	{
+		TextureObejct = TextureValue.Texture;
+	}
+	else if (TextureValue.RuntimeVirtualTexture)
+	{
+		TextureObejct = TextureValue.RuntimeVirtualTexture;
+	}
+
+	const int32 TextureIndex = Context.Material->GetReferencedTextures().Find(TextureObejct);
 	check(TextureIndex != INDEX_NONE);
 	
+	const Shader::EPreshaderOpcode Op = (TextureProperty == TMTM_TextureSize ? Shader::EPreshaderOpcode::TextureSize : Shader::EPreshaderOpcode::TexelSize);
+
 	Context.PreshaderStackPosition++;
-	OutResult.Type = Shader::EValueType::Float2;
-	OutResult.Preshader.WriteOpcode(Shader::EPreshaderOpcode::TextureSize).Write<FMemoryImageMaterialParameterInfo>(TextureValue.ParameterInfo).Write(TextureIndex);
+	OutResult.Type = Shader::EValueType::Float3;
+	OutResult.Preshader.WriteOpcode(Op).Write<FMemoryImageMaterialParameterInfo>(TextureValue.ParameterInfo).Write(TextureIndex);
 }
 
 bool FExpressionRuntimeVirtualTextureUniform::PrepareValue(FEmitContext& Context, FEmitScope& Scope, const FRequestedType& RequestedType, FPrepareValueResult& OutResult) const
@@ -1214,7 +1258,17 @@ void FExpressionRuntimeVirtualTextureUniform::EmitValuePreshader(FEmitContext& C
 	FMaterialTextureValue TextureValue;
 	TextureExpression->GetValueObject(Context, Scope, TextureValue);
 
-	const int32 TextureIndex = Context.Material->GetReferencedTextures().Find(TextureValue.Texture);
+	UObject* TextureObejct = nullptr;
+	if (TextureValue.Texture)
+	{
+		TextureObejct = TextureValue.Texture;
+	}
+	else if (TextureValue.RuntimeVirtualTexture)
+	{
+		TextureObejct = TextureValue.RuntimeVirtualTexture;
+	}
+
+	const int32 TextureIndex = Context.Material->GetReferencedTextures().Find(TextureObejct);
 	check(TextureIndex != INDEX_NONE);
 	const int32 VectorIndex = (int32)UniformType;
 
@@ -1575,7 +1629,7 @@ bool FExpressionNoise::PrepareValue(FEmitContext& Context, FEmitScope& Scope, co
 
 void FExpressionNoise::EmitValueShader(FEmitContext& Context, FEmitScope& Scope, const FRequestedType& RequestedType, FEmitValueShaderResult& OutResult) const
 {
-	const FPreparedType& PreparedType = Context.GetPreparedType(PositionExpression);
+	const FPreparedType& PreparedType = Context.GetPreparedType(PositionExpression, Shader::EValueType::Float3);
 	bool bIsLWC = Shader::IsLWCType(PreparedType.Type);
 	FEmitShaderExpression* EmitPosition = PositionExpression->GetValueShader(Context, Scope, bIsLWC ? Shader::EValueType::Double3 : Shader::EValueType::Float3);
 	FEmitShaderExpression* EmitFilterWidth = FilterWidthExpression->GetValueShader(Context, Scope, Shader::EValueType::Float1);
diff --git a/Engine/Source/Runtime/Engine/Public/HLSLTree/HLSLTree.h b/Engine/Source/Runtime/Engine/Public/HLSLTree/HLSLTree.h
--- a/Engine/Source/Runtime/Engine/Public/HLSLTree/HLSLTree.h
+++ b/Engine/Source/Runtime/Engine/Public/HLSLTree/HLSLTree.h
@@ -291,7 +291,7 @@ public:
 
 	/** Converts to a requested type, based on IsRequestedEvaluation() */
 	FRequestedType GetRequestedType() const;
-	
+
 	Shader::EValueComponentType GetValueComponentType() const;
 	const TCHAR* GetName() const { return Type.GetName(); }
 	bool IsVoid() const { return Type.IsVoid(); }
@@ -423,7 +423,6 @@ public:
 	FEmitShaderExpression* GetValueShader(FEmitContext& Context, FEmitScope& Scope, const FRequestedType& RequestedType) const;
 	FEmitShaderExpression* GetValueShader(FEmitContext& Context, FEmitScope& Scope, const Shader::FType& ResultType) const;
 	FEmitShaderExpression* GetValueShader(FEmitContext& Context, FEmitScope& Scope, Shader::EValueType ResultType) const;
-	FEmitShaderExpression* GetValueShader(FEmitContext& Context, FEmitScope& Scope) const;
 
 	Shader::FType GetValuePreshader(FEmitContext& Context, FEmitScope& Scope, const FRequestedType& RequestedType, const Shader::FType& ResultType, Shader::FPreshaderData& OutPreshader) const;
 	Shader::FType GetValuePreshader(FEmitContext& Context, FEmitScope& Scope, const FRequestedType& RequestedType, Shader::FPreshaderData& OutPreshader) const;
@@ -467,10 +466,10 @@ public:
 
 protected:
 	/** Create new expressions representing DDX/DDY of this expression */
-	virtual void ComputeAnalyticDerivatives(FTree& Tree, FExpressionDerivatives& OutResult) const;
+	ENGINE_API virtual void ComputeAnalyticDerivatives(FTree& Tree, FExpressionDerivatives& OutResult) const;
 
 	/** Creates a new expression representing this expression on the previous frame.  By default returns nullptr, which means the previous frame is the same as the current frame */
-	virtual const FExpression* ComputePreviousFrame(FTree& Tree, const FRequestedType& RequestedType) const;
+	ENGINE_API virtual const FExpression* ComputePreviousFrame(FTree& Tree, const FRequestedType& RequestedType) const;
 
 	/**
 	 * Computes a FPreparedType for this expression, given a FRequestedType.  Will be called multiple times, with potentially different requested types, if the FExpression is used multiple times.
@@ -487,19 +486,19 @@ protected:
 	virtual bool PrepareValue(FEmitContext& Context, FEmitScope& Scope, const FRequestedType& RequestedType, FPrepareValueResult& OutResult) const = 0;
 
 	/** Emit HLSL shader code representing this expression */
-	virtual void EmitValueShader(FEmitContext& Context, FEmitScope& Scope, const FRequestedType& RequestedType, FEmitValueShaderResult& OutResult) const;
+	ENGINE_API virtual void EmitValueShader(FEmitContext& Context, FEmitScope& Scope, const FRequestedType& RequestedType, FEmitValueShaderResult& OutResult) const;
 
 	/** Emit Preshader code representing this expression */
-	virtual void EmitValuePreshader(FEmitContext& Context, FEmitScope& Scope, const FRequestedType& RequestedType, FEmitValuePreshaderResult& OutResult) const;
+	ENGINE_API virtual void EmitValuePreshader(FEmitContext& Context, FEmitScope& Scope, const FRequestedType& RequestedType, FEmitValuePreshaderResult& OutResult) const;
 
 	/** Emit an object.  The given 'ObjectTypeName' determines the C++ type pointed to by OutObjectBase */
-	virtual bool EmitValueObject(FEmitContext& Context, FEmitScope& Scope, const FName& ObjectTypeName, void* OutObjectBase) const;
+	ENGINE_API virtual bool EmitValueObject(FEmitContext& Context, FEmitScope& Scope, const FName& ObjectTypeName, void* OutObjectBase) const;
 
 	/**
 	 * Allows custom objects to be passed to custom HLSL functions.  This needs to initialize some HLSL code that facilitates this interface.
 	 * If this returns 'true', then EmitValueShader() will be called to generate the actual HLSL code for the FExpression
 	 */
-	virtual bool EmitCustomHLSLParameter(FEmitContext& Context, FEmitScope& Scope, const FName& ObjectTypeName, const TCHAR* ParameterName, FEmitCustomHLSLParameterResult& OutResult) const;
+	ENGINE_API virtual bool EmitCustomHLSLParameter(FEmitContext& Context, FEmitScope& Scope, const FName& ObjectTypeName, const TCHAR* ParameterName, FEmitCustomHLSLParameterResult& OutResult) const;
 
 private:
 	TArray<UObject*, TInlineAllocator<2>> Owners;
@@ -635,7 +634,7 @@ public:
 	/** Shortcuts to create various common expression types */
 	const FExpression* NewConstant(const Shader::FValue& Value);
 	const FExpression* NewSwizzle(const FSwizzleParameters& Params, const FExpression* Input);
-	const FExpression* NewUnaryOp(EOperation Op, const FExpression* Input);
+	ENGINE_API const FExpression* NewUnaryOp(EOperation Op, const FExpression* Input);
 	ENGINE_API const FExpression* NewBinaryOp(EOperation Op, const FExpression* Lhs, const FExpression* Rhs);
 	const FExpression* NewTernaryOp(EOperation Op, const FExpression* Input0, const FExpression* Input1, const FExpression* Input2);
 
@@ -679,6 +678,8 @@ public:
 
 	const FExpression* NewTruncateLWC(const FExpression* Input) { return NewUnaryOp(EOperation::TruncateLWC, Input); }
 
+	FActiveStructFieldStack ActiveStructFieldStack;
+
 private:
 	template<typename T, typename... ArgTypes>
 	inline T* NewNode(ArgTypes&&... Args)
@@ -688,12 +689,12 @@ private:
 		return Node;
 	}
 
-	void RegisterNode(FNode* Node);
-	void RegisterExpression(FExpression* Expression, FXxHash64 Hash);
+	ENGINE_API void RegisterNode(FNode* Node);
+	ENGINE_API void RegisterExpression(FExpression* Expression, FXxHash64 Hash);
 	void RegisterExpression(FExpressionLocalPHI* Expression, FXxHash64 Hash);
-	void AddCurrentOwner(FExpression* Expression);
+	ENGINE_API void AddCurrentOwner(FExpression* Expression);
 	void RegisterStatement(FScope& Scope, FStatement* Statement);
-	FExpression* FindExpression(FXxHash64 Hash);
+	ENGINE_API FExpression* FindExpression(FXxHash64 Hash);
 
 	FMemStackBase* Allocator = nullptr;
 	FNode* Nodes = nullptr;
diff --git a/Engine/Source/Runtime/Engine/Public/HLSLTree/HLSLTreeCommon.h b/Engine/Source/Runtime/Engine/Public/HLSLTree/HLSLTreeCommon.h
--- a/Engine/Source/Runtime/Engine/Public/HLSLTree/HLSLTreeCommon.h
+++ b/Engine/Source/Runtime/Engine/Public/HLSLTree/HLSLTreeCommon.h
@@ -39,6 +39,7 @@ public:
 	virtual bool PrepareValue(FEmitContext& Context, FEmitScope& Scope, const FRequestedType& RequestedType, FPrepareValueResult& OutResult) const override;
 	virtual void EmitValueShader(FEmitContext& Context, FEmitScope& Scope, const FRequestedType& RequestedType, FEmitValueShaderResult& OutResult) const override;
 	virtual void EmitValuePreshader(FEmitContext& Context, FEmitScope& Scope, const FRequestedType& RequestedType, FEmitValuePreshaderResult& OutResult) const override;
+	virtual bool EmitValueObject(FEmitContext& Context, FEmitScope& Scope, const FName& ObjectTypeName, void* OutObjectBase) const override;
 };
 
 /**
@@ -132,6 +133,10 @@ public:
 	virtual bool PrepareValue(FEmitContext& Context, FEmitScope& Scope, const FRequestedType& RequestedType, FPrepareValueResult& OutResult) const override;
 	virtual void EmitValueShader(FEmitContext& Context, FEmitScope& Scope, const FRequestedType& RequestedType, FEmitValueShaderResult& OutResult) const override;
 	virtual void EmitValuePreshader(FEmitContext& Context, FEmitScope& Scope, const FRequestedType& RequestedType, FEmitValuePreshaderResult& OutResult) const override;
+
+private:
+	FRequestedType MakeRequestedStructType(const FActiveStructFieldStack& ActiveFieldStack, const FRequestedType& RequestedType) const;
+	FRequestedType MakeRequestedFieldType(const FActiveStructFieldStack& ActiveFieldStack, const FRequestedType& RequestedType) const;
 };
 
 class FExpressionSelect : public FExpression
@@ -225,11 +230,11 @@ public:
 	const FExpression* Lhs;
 	const FExpression* Rhs;
 
-	virtual void ComputeAnalyticDerivatives(FTree& Tree, FExpressionDerivatives& OutResult) const override;
-	virtual const FExpression* ComputePreviousFrame(FTree& Tree, const FRequestedType& RequestedType) const override;
-	virtual bool PrepareValue(FEmitContext& Context, FEmitScope& Scope, const FRequestedType& RequestedType, FPrepareValueResult& OutResult) const override;
-	virtual void EmitValueShader(FEmitContext& Context, FEmitScope& Scope, const FRequestedType& RequestedType, FEmitValueShaderResult& OutResult) const override;
-	virtual void EmitValuePreshader(FEmitContext& Context, FEmitScope& Scope, const FRequestedType& RequestedType, FEmitValuePreshaderResult& OutResult) const override;
+	ENGINE_API virtual void ComputeAnalyticDerivatives(FTree& Tree, FExpressionDerivatives& OutResult) const override;
+	ENGINE_API virtual const FExpression* ComputePreviousFrame(FTree& Tree, const FRequestedType& RequestedType) const override;
+	ENGINE_API virtual bool PrepareValue(FEmitContext& Context, FEmitScope& Scope, const FRequestedType& RequestedType, FPrepareValueResult& OutResult) const override;
+	ENGINE_API virtual void EmitValueShader(FEmitContext& Context, FEmitScope& Scope, const FRequestedType& RequestedType, FEmitValueShaderResult& OutResult) const override;
+	ENGINE_API virtual void EmitValuePreshader(FEmitContext& Context, FEmitScope& Scope, const FRequestedType& RequestedType, FEmitValuePreshaderResult& OutResult) const override;
 };
 
 class FExpressionSwitchBase : public FExpression
diff --git a/Engine/Source/Runtime/Engine/Public/HLSLTree/HLSLTreeEmit.h b/Engine/Source/Runtime/Engine/Public/HLSLTree/HLSLTreeEmit.h
--- a/Engine/Source/Runtime/Engine/Public/HLSLTree/HLSLTreeEmit.h
+++ b/Engine/Source/Runtime/Engine/Public/HLSLTree/HLSLTreeEmit.h
@@ -370,9 +370,9 @@ public:
 
 	ENGINE_API void EmitDeclarationsCode(FStringBuilderBase& OutCode);
 
-	ENGINE_API const FPreparedType& GetPreparedType(const FExpression* Expression) const;
-	ENGINE_API FRequestedType GetRequestedType(const FExpression* Expression) const;
-	ENGINE_API Shader::FType GetType(const FExpression* Expression) const;
+	ENGINE_API FPreparedType GetPreparedType(const FExpression* Expression, const FRequestedType& RequestedType) const;
+	ENGINE_API Shader::FType GetResultType(const FExpression* Expression, const FRequestedType& RequestedType) const;
+	ENGINE_API Shader::FType GetTypeForPinColoring(const FExpression* Expression) const;
 	ENGINE_API EExpressionEvaluation GetEvaluation(const FExpression* Expression, const FEmitScope& Scope, const FRequestedType& RequestedType) const;
 
 	ENGINE_API FPreparedType PrepareExpression(const FExpression* InExpression, FEmitScope& Scope, const FRequestedType& RequestedType);
@@ -583,10 +583,12 @@ public:
 
 	uint32 DynamicParticleParameterMask = 0u;
 
+	FActiveStructFieldStack ActiveStructFieldStack;
+
 	TArray<const FOwnedNode*, TInlineAllocator<32>> OwnerStack;
 	TArray<FEmitShaderNode*> EmitNodes;
 	TMap<const FScope*, FEmitScope*> EmitScopeMap;
-	TMap<const FExpression*, FPrepareValueResult*> PrepareValueMap;
+	TMap<FXxHash64, FPrepareValueResult*> PrepareValueMap;
 	TMap<FXxHash64, FRequestedType*> RequestedTypeTracker;
 	TMap<const FExpression*, FEmitScope*> PrepareLocalPHIMap;
 	TMap<const FExpression*, FEmitShaderExpression*> EmitLocalPHIMap;
diff --git a/Engine/Source/Runtime/Engine/Public/HLSLTree/HLSLTreeTypes.h b/Engine/Source/Runtime/Engine/Public/HLSLTree/HLSLTreeTypes.h
--- a/Engine/Source/Runtime/Engine/Public/HLSLTree/HLSLTreeTypes.h
+++ b/Engine/Source/Runtime/Engine/Public/HLSLTree/HLSLTreeTypes.h
@@ -180,6 +180,25 @@ struct FCustomHLSLInput
 	const FExpression* Expression = nullptr;
 };
 
+typedef TArray<const Shader::FStructField*> FActiveStructFieldStack;
+
+class FScopedActiveStructField
+{
+	FActiveStructFieldStack& Stack;
+
+public:
+	FScopedActiveStructField(FActiveStructFieldStack& InStack, const Shader::FStructField* Field)
+		: Stack(InStack)
+	{
+		Stack.Push(Field);
+	}
+
+	~FScopedActiveStructField()
+	{
+		Stack.Pop(false);
+	}
+};
+
 } // namespace UE::HLSLTree
 
 #endif // WITH_EDITOR
diff --git a/Engine/Source/Runtime/Engine/Public/MaterialHLSLGenerator.h b/Engine/Source/Runtime/Engine/Public/MaterialHLSLGenerator.h
--- a/Engine/Source/Runtime/Engine/Public/MaterialHLSLGenerator.h
+++ b/Engine/Source/Runtime/Engine/Public/MaterialHLSLGenerator.h
@@ -105,8 +105,8 @@ public:
 
 	UE::HLSLTree::FScope* NewJoinedScope(UE::HLSLTree::FScope& Scope);
 
-	const UE::HLSLTree::FExpression* NewConstant(const UE::Shader::FValue& Value);
-	const UE::HLSLTree::FExpression* NewTexCoord(int32 Index);
+	ENGINE_API const UE::HLSLTree::FExpression* NewConstant(const UE::Shader::FValue& Value);
+	ENGINE_API const UE::HLSLTree::FExpression* NewTexCoord(int32 Index);
 	const UE::HLSLTree::FExpression* NewExternalInput(UE::HLSLTree::Material::EExternalInput Input);
 	const UE::HLSLTree::FExpression* NewSwizzle(const UE::HLSLTree::FSwizzleParameters& Params, const UE::HLSLTree::FExpression* Input);
 
@@ -149,11 +149,30 @@ public:
 		EMaterialSamplerType InSamplerType = SAMPLERTYPE_Color,
 		const FGuid& InExternalTextureGuid = FGuid());
 
+	struct FConnectedInput
+	{
+		const UE::HLSLTree::FExpression* Expression;
+		const FExpressionInput* Input;
+		UE::HLSLTree::FScope* Scope;
+
+		explicit FConnectedInput(const UE::HLSLTree::FExpression* InExpression)
+			: Expression(InExpression)
+			, Input(nullptr)
+			, Scope(nullptr)
+		{}
+
+		FConnectedInput(const FExpressionInput* InInput, UE::HLSLTree::FScope* InScope)
+			: Expression(nullptr)
+			, Input(InInput)
+			, Scope(InScope)
+		{}
+	};
+
 	const UE::HLSLTree::FExpression* GenerateFunctionCall(UE::HLSLTree::FScope& Scope,
 		UMaterialFunctionInterface* Function,
 		EMaterialParameterAssociation ParameterAssociation,
 		int32 ParameterIndex,
-		TArrayView<const UE::HLSLTree::FExpression*> ConnectedInputs,
+		TArrayView<FConnectedInput> ConnectedInputs,
 		int32 OutputIndex);
 
 	const UE::HLSLTree::FExpression* GenerateBranch(UE::HLSLTree::FScope& Scope,
@@ -217,7 +236,7 @@ private:
 
 	using FFunctionInputArray = TArray<const UMaterialExpressionFunctionInput*, TInlineAllocator<4>>;
 	using FFunctionOutputArray = TArray<UMaterialExpressionFunctionOutput*, TInlineAllocator<4>>;
-	using FConnectedInputArray = TArray<const UE::HLSLTree::FExpression*, TInlineAllocator<4>>;
+	using FConnectedInputArray = TArray<FConnectedInput, TInlineAllocator<4>>;
 
 	struct FFunctionCallEntry
 	{
diff --git a/Engine/Source/Runtime/Engine/Public/MaterialHLSLTree.h b/Engine/Source/Runtime/Engine/Public/MaterialHLSLTree.h
--- a/Engine/Source/Runtime/Engine/Public/MaterialHLSLTree.h
+++ b/Engine/Source/Runtime/Engine/Public/MaterialHLSLTree.h
@@ -373,14 +373,16 @@ public:
 	virtual void EmitValueShader(FEmitContext& Context, FEmitScope& Scope, const FRequestedType& RequestedType, FEmitValueShaderResult& OutResult) const override;
 };
 
-class FExpressionTextureSize : public FExpression
+class FExpressionTextureProperty : public FExpression
 {
 public:
-	explicit FExpressionTextureSize(const FExpression* InTextureExpression)
+	explicit FExpressionTextureProperty(const FExpression* InTextureExpression, EMaterialExposedTextureProperty InTextureProperty)
 		: TextureExpression(InTextureExpression)
+		, TextureProperty(InTextureProperty)
 	{}
 
 	const FExpression* TextureExpression;
+	EMaterialExposedTextureProperty TextureProperty;
 
 	virtual bool PrepareValue(FEmitContext& Context, FEmitScope& Scope, const FRequestedType& RequestedType, FPrepareValueResult& OutResult) const override;
 	virtual void EmitValuePreshader(FEmitContext& Context, FEmitScope& Scope, const FRequestedType& RequestedType, FEmitValuePreshaderResult& OutResult) const override;
diff --git a/Engine/Source/Runtime/Engine/Public/MaterialTypes.h b/Engine/Source/Runtime/Engine/Public/MaterialTypes.h
--- a/Engine/Source/Runtime/Engine/Public/MaterialTypes.h
+++ b/Engine/Source/Runtime/Engine/Public/MaterialTypes.h
@@ -310,6 +310,7 @@ struct FMaterialTextureValue
 	static FName GetTypeName() { return TEXT("FMaterialTextureValue"); }
 
 	UTexture* Texture = nullptr;
+	URuntimeVirtualTexture* RuntimeVirtualTexture = nullptr;
 	FGuid ExternalTextureGuid;
 	EMaterialSamplerType SamplerType = SAMPLERTYPE_Color;
 	FMaterialParameterInfo ParameterInfo;
diff --git a/Engine/Source/Runtime/Landscape/Classes/Materials/MaterialExpressionLandscapeLayerCoords.h b/Engine/Source/Runtime/Landscape/Classes/Materials/MaterialExpressionLandscapeLayerCoords.h
--- a/Engine/Source/Runtime/Landscape/Classes/Materials/MaterialExpressionLandscapeLayerCoords.h
+++ b/Engine/Source/Runtime/Landscape/Classes/Materials/MaterialExpressionLandscapeLayerCoords.h
@@ -64,6 +64,8 @@ class UMaterialExpressionLandscapeLayerCoords : public UMaterialExpression
 #if WITH_EDITOR
 	LANDSCAPE_API virtual int32 Compile(class FMaterialCompiler* Compiler, int32 OutputIndex) override;
 	LANDSCAPE_API virtual void GetCaption(TArray<FString>& OutCaptions) const override;
+
+	virtual bool GenerateHLSLExpression(FMaterialHLSLGenerator& Generator, UE::HLSLTree::FScope& Scope, int32 OutputIndex, UE::HLSLTree::FExpression const*& OutExpression) const override;
 #endif
 	//~ End UMaterialExpression Interface
 };
diff --git a/Engine/Source/Runtime/Landscape/Private/Materials/MaterialExpressionLandscapeLayerCoords.cpp b/Engine/Source/Runtime/Landscape/Private/Materials/MaterialExpressionLandscapeLayerCoords.cpp
--- a/Engine/Source/Runtime/Landscape/Private/Materials/MaterialExpressionLandscapeLayerCoords.cpp
+++ b/Engine/Source/Runtime/Landscape/Private/Materials/MaterialExpressionLandscapeLayerCoords.cpp
@@ -3,6 +3,8 @@
 #include "Materials/MaterialExpressionLandscapeLayerCoords.h"
 #include "LandscapePrivate.h"
 #include "MaterialCompiler.h"
+#include "MaterialHLSLGenerator.h"
+#include "MaterialHLSLTree.h"
 
 
 #define LOCTEXT_NAMESPACE "Landscape"
@@ -84,6 +86,64 @@ void UMaterialExpressionLandscapeLayerCoords::GetCaption(TArray<FString>& OutCap
 {
 	OutCaptions.Add(FString(TEXT("Landscape Coords")));
 }
+
+bool UMaterialExpressionLandscapeLayerCoords::GenerateHLSLExpression(FMaterialHLSLGenerator& Generator, UE::HLSLTree::FScope& Scope, int32 OutputIndex, UE::HLSLTree::FExpression const*& OutExpression) const
+{
+	using namespace UE::HLSLTree;
+
+	switch (CustomUVType)
+	{
+	case LCCT_CustomUV0:
+		OutExpression = Generator.NewTexCoord(0);
+		return true;
+	case LCCT_CustomUV1:
+		OutExpression = Generator.NewTexCoord(1);
+		return true;
+	case LCCT_CustomUV2:
+		OutExpression = Generator.NewTexCoord(2);
+		return true;
+	case LCCT_WeightMapUV:
+		OutExpression = Generator.NewTexCoord(3);
+		return true;
+	default:
+		break;
+	}
+
+	const FExpression* BaseUVExpression;
+
+	switch (MappingType)
+	{
+	case TCMT_Auto:
+	case TCMT_XY:
+		BaseUVExpression = Generator.NewTexCoord(0);
+		break;
+	case TCMT_XZ:
+		BaseUVExpression = Generator.NewTexCoord(1);
+		break;
+	case TCMT_YZ:
+		BaseUVExpression = Generator.NewTexCoord(2);
+		break;
+	default:
+		return Generator.Errorf(TEXT("Invalid mapping type %u"), (uint8)MappingType);
+	}
+
+	const float Scale = (MappingScale == 0.0f) ? 1.0f : 1.0f / MappingScale;
+	const FExpression* RealScaleExpression = Generator.NewConstant(Scale);
+
+	const float Cos = FMath::Cos(MappingRotation * PI / 180.0f);
+	const float Sin = FMath::Sin(MappingRotation * PI / 180.0f);
+
+	OutExpression = Generator.GetTree().NewAdd(
+		Generator.GetTree().NewMul(
+			RealScaleExpression,
+			Generator.GetTree().NewExpression<FExpressionAppend>(
+				Generator.GetTree().NewDot(BaseUVExpression, Generator.NewConstant(FVector2f(Cos, Sin))),
+				Generator.GetTree().NewDot(BaseUVExpression, Generator.NewConstant(FVector2f(-Sin, Cos))))),
+		Generator.NewConstant(FVector2f(MappingPanU, MappingPanV)));
+
+	return true;
+}
+
 #endif // WITH_EDITOR
 
 #undef LOCTEXT_NAMESPACE
