Simplified material expression async logic to only handle fully async or fully serial.

diff --git a/Engine/Source/Runtime/Engine/Private/Materials/MaterialRenderProxy.cpp b/Engine/Source/Runtime/Engine/Private/Materials/MaterialRenderProxy.cpp
--- a/Engine/Source/Runtime/Engine/Private/Materials/MaterialRenderProxy.cpp
+++ b/Engine/Source/Runtime/Engine/Private/Materials/MaterialRenderProxy.cpp
@@ -672,16 +672,12 @@ void FMaterialRenderProxy::UpdateDeferredCachedUniformExpressions(FRHICommandLis
 {
 	LLM_SCOPE(ELLMTag::Materials);
 
-	const bool bAllowAsyncUpdate = GUniformExpressionCacheAsyncUpdateTask.IsEnabled();
+	const bool bAllowAsyncUpdate = RHICmdList.IsImmediate() && GUniformExpressionCacheAsyncUpdateTask.IsEnabled();
 
-	if (!bAllowAsyncUpdate || !RHICmdList.IsImmediate())
-	{
-		TaskIfAsync = nullptr;
-	}
+	FRHICommandListBase* RHICmdListTask = &RHICmdList;
 
-	FRHICommandListBase* RHICmdListTask = nullptr;
-
-	if (TaskIfAsync)
+	// Create an async command list when immediate command list is supplied and async update is allowed.
+	if (bAllowAsyncUpdate)
 	{
 		RHICmdListTask = new FRHICommandList(FRHIGPUMask::All());
 		RHICmdListTask->SwitchPipeline(ERHIPipeline::Graphics);
@@ -690,7 +686,7 @@ void FMaterialRenderProxy::UpdateDeferredCachedUniformExpressions(FRHICommandLis
 	}
 	else
 	{
-		RHICmdListTask = &RHICmdList;
+		GUniformExpressionCacheAsyncUpdateTask.Wait();
 	}
 
 	auto EvaluateUniformExpressionsLambda = [RHICmdListTask, bAllowAsyncUpdate]
@@ -731,19 +727,13 @@ void FMaterialRenderProxy::UpdateDeferredCachedUniformExpressions(FRHICommandLis
 			DeferredUniformExpressionCacheRequests.Reset();
 		}
 
-		if (bAllowAsyncUpdate)
+		if (UpdaterIfEnabled)
 		{
-			Updater.Update(*RHICmdListTask);
-		}
-		else if (!RHICmdListTask->IsImmediate())
-		{
-			RHICmdListTask->FinishRecording();
+			UpdaterIfEnabled->Update(*RHICmdListTask);
 		}
 	};
 
-	UE::Tasks::FTask Task;
-
-	if (TaskIfAsync)
+	if (bAllowAsyncUpdate)
 	{
 		*TaskIfAsync = UE::Tasks::Launch(UE_SOURCE_LOCATION, MoveTemp(EvaluateUniformExpressionsLambda));
 	}
