Prevent any form of streaming for Verse script devices
* Disallow modifying the spatially loaded property
* Disallow inclusion of those actors in data layers
* Existing script devices that have issues won't be corrected silently and will make the validation fail. Users will need to fix them manually, so that they understand the impact of the change.

diff --git a/Engine/Source/Runtime/Engine/Classes/GameFramework/Actor.h b/Engine/Source/Runtime/Engine/Classes/GameFramework/Actor.h
--- a/Engine/Source/Runtime/Engine/Classes/GameFramework/Actor.h
+++ b/Engine/Source/Runtime/Engine/Classes/GameFramework/Actor.h
@@ -1302,8 +1302,9 @@ public:
 	//~=============================================================================
 	// DataLayers functions.
 #if WITH_EDITOR
-private:
+protected:
 	virtual bool IsDataLayerTypeSupported(TSubclassOf<UDataLayerInstance> DataLayerType) const { return true; }
+private:
 	virtual bool ActorTypeIsMainWorldOnly() const { return false; }
 	ENGINE_API TArray<const UDataLayerAsset*> ResolveDataLayerAssets(const TArray<TSoftObjectPtr<UDataLayerAsset>>& InDataLayerAssets) const;
 public:
