Update FRetargetingScope to ensure raw transforms have compatible skeletons applied

diff --git a/Engine/Source/Runtime/Engine/Private/Animation/AnimSequenceHelpers.cpp b/Engine/Source/Runtime/Engine/Private/Animation/AnimSequenceHelpers.cpp
--- a/Engine/Source/Runtime/Engine/Private/Animation/AnimSequenceHelpers.cpp
+++ b/Engine/Source/Runtime/Engine/Private/Animation/AnimSequenceHelpers.cpp
@@ -197,7 +197,7 @@ Retargeting::FRetargetingScope::FRetargetingScope(const USkeleton* InSourceSkele
 	RetargetPose(ToRetargetPose),
 	EvaluationContext(InEvaluationContext),
 	RetargetTracking(FBuildRawPoseScratchArea::Get().RetargetTracking),
-	bShouldRetarget(!ToRetargetPose.GetBoneContainer().GetDisableRetargeting() && EvaluationContext.RetargetTransforms.Num())
+	bShouldRetarget((!ToRetargetPose.GetBoneContainer().GetDisableRetargeting() && EvaluationContext.RetargetTransforms.Num()) || (SourceSkeleton != ToRetargetPose.GetBoneContainer().GetSkeletonAsset()))
 {
 	check(SourceSkeleton);
 	RetargetTracking.Reset();
diff --git a/Engine/Source/Runtime/Engine/Private/Animation/AnimationRuntime.cpp b/Engine/Source/Runtime/Engine/Private/Animation/AnimationRuntime.cpp
--- a/Engine/Source/Runtime/Engine/Private/Animation/AnimationRuntime.cpp
+++ b/Engine/Source/Runtime/Engine/Private/Animation/AnimationRuntime.cpp
@@ -2587,8 +2587,15 @@ void FAnimationRuntime::RetargetBoneTransform(const USkeleton* SourceSkeleton, c
 		const FSkeletonRemapping& SkeletonRemapping = UE::Anim::FSkeletonRemappingRegistry::Get().GetRemapping(SourceSkeleton, TargetSkeleton);
 
 		const int32 TargetSkeletonBoneIndex = RequiredBones.GetSkeletonIndex(BoneIndex);
-		const int32 SourceSkeletonBoneIndex = SkeletonRemapping.IsValid() ? SkeletonRemapping.GetSourceSkeletonBoneIndex(TargetSkeletonBoneIndex) : SkeletonBoneIndex;
+		int32 SourceSkeletonBoneIndex = SkeletonBoneIndex;
 
+		// Apply compatible skeleton remapping if required
+		if (SkeletonRemapping.IsValid() && SkeletonRemapping.RequiresReferencePoseRetarget())
+		{
+			SourceSkeletonBoneIndex = SkeletonRemapping.GetSourceSkeletonBoneIndex(TargetSkeletonBoneIndex);
+			BoneTransform = SkeletonRemapping.RetargetBoneTransformToTargetSkeleton(TargetSkeletonBoneIndex, BoneTransform);
+		}
+		
 		switch (TargetSkeleton->GetBoneTranslationRetargetingMode(TargetSkeletonBoneIndex, RequiredBones.GetDisableRetargeting()))
 		{
 			case EBoneTranslationRetargetingMode::AnimationScaled:
