Remove EAccelerationStructureBuildMode parameter from FRayTracingSkinnedGeometryUpdateQueue::Add(...)
- instead check InRayTracingGeometry->GetRequiresBuild() directly

diff --git a/Engine/Source/Runtime/Engine/Private/GPUSkinCache.cpp b/Engine/Source/Runtime/Engine/Private/GPUSkinCache.cpp
--- a/Engine/Source/Runtime/Engine/Private/GPUSkinCache.cpp
+++ b/Engine/Source/Runtime/Engine/Private/GPUSkinCache.cpp
@@ -1891,15 +1891,15 @@ void FGPUSkinCache::ProcessRayTracingGeometryToUpdate(FRHICommandList& RHICmdLis
 {
 	if (IsGPUSkinCacheRayTracingSupported() && SkinCacheEntry && SkinCacheEntry->GPUSkin && SkinCacheEntry->GPUSkin->bSupportRayTracing)
 	{
- 		TArray<FBufferRHIRef> VertexBufffers;
- 		SkinCacheEntry->GetRayTracingSegmentVertexBuffers(VertexBufffers);
+ 		TArray<FBufferRHIRef> VertexBuffers;
+ 		SkinCacheEntry->GetRayTracingSegmentVertexBuffers(VertexBuffers);
 
 		const int32 LODIndex = SkinCacheEntry->LOD;
 		FSkeletalMeshRenderData& SkelMeshRenderData = SkinCacheEntry->GPUSkin->GetSkeletalMeshRenderData();
 		check(LODIndex < SkelMeshRenderData.LODRenderData.Num());
 		FSkeletalMeshLODRenderData& LODModel = SkelMeshRenderData.LODRenderData[LODIndex];
 
- 		SkinCacheEntry->GPUSkin->UpdateRayTracingGeometry(RHICmdList, LODModel, LODIndex, VertexBufffers);
+ 		SkinCacheEntry->GPUSkin->UpdateRayTracingGeometry(RHICmdList, LODModel, LODIndex, VertexBuffers);
 	}
 }
 
diff --git a/Engine/Source/Runtime/Engine/Private/RayTracingSkinnedGeometry.cpp b/Engine/Source/Runtime/Engine/Private/RayTracingSkinnedGeometry.cpp
--- a/Engine/Source/Runtime/Engine/Private/RayTracingSkinnedGeometry.cpp
+++ b/Engine/Source/Runtime/Engine/Private/RayTracingSkinnedGeometry.cpp
@@ -34,23 +34,25 @@ FAutoConsoleVariableRef CVarSkinnedGeometryMaxRayTracingPrimitivesPerCmdList(
 	ECVF_RenderThreadSafe
 );
 
-void FRayTracingSkinnedGeometryUpdateQueue::Add(FRayTracingGeometry* InRayTracingGeometry, const FRayTracingAccelerationStructureSize& StructureSize, EAccelerationStructureBuildMode InBuildMode)
+void FRayTracingSkinnedGeometryUpdateQueue::Add(FRayTracingGeometry* InRayTracingGeometry, const FRayTracingAccelerationStructureSize& StructureSize)
 {
 	FScopeLock Lock(&CS);
 	FRayTracingUpdateInfo* CurrentUpdateInfo = ToUpdate.Find(InRayTracingGeometry);
 	if (CurrentUpdateInfo == nullptr)
 	{
 		FRayTracingUpdateInfo UpdateInfo;
-		UpdateInfo.BuildMode = InBuildMode;
-		UpdateInfo.ScratchSize = InBuildMode == EAccelerationStructureBuildMode::Build ? StructureSize.BuildScratchSize : StructureSize.UpdateScratchSize;
+		UpdateInfo.BuildMode = InRayTracingGeometry->GetRequiresBuild() ? EAccelerationStructureBuildMode::Build : EAccelerationStructureBuildMode::Update;
+		UpdateInfo.ScratchSize = InRayTracingGeometry->GetRequiresBuild() ? StructureSize.BuildScratchSize : StructureSize.UpdateScratchSize;
 		ToUpdate.Add(InRayTracingGeometry, UpdateInfo);
 	}
 	// If currently updating but need full rebuild then update the stored build mode
-	else if (CurrentUpdateInfo->BuildMode == EAccelerationStructureBuildMode::Update && InBuildMode == EAccelerationStructureBuildMode::Build)
+	else if (CurrentUpdateInfo->BuildMode == EAccelerationStructureBuildMode::Update && InRayTracingGeometry->GetRequiresBuild())
 	{
-		CurrentUpdateInfo->BuildMode = InBuildMode;
+		CurrentUpdateInfo->BuildMode = EAccelerationStructureBuildMode::Build;
 		CurrentUpdateInfo->ScratchSize = StructureSize.BuildScratchSize;
 	}
+
+	InRayTracingGeometry->SetRequiresBuild(false);
 }
 
 void FRayTracingSkinnedGeometryUpdateQueue::Remove(FRayTracingGeometry* RayTracingGeometry, uint32 EstimatedMemory)
diff --git a/Engine/Source/Runtime/Engine/Private/SkeletalRenderGPUSkin.cpp b/Engine/Source/Runtime/Engine/Private/SkeletalRenderGPUSkin.cpp
--- a/Engine/Source/Runtime/Engine/Private/SkeletalRenderGPUSkin.cpp
+++ b/Engine/Source/Runtime/Engine/Private/SkeletalRenderGPUSkin.cpp
@@ -772,7 +772,7 @@ void FSkeletalMeshObjectGPUSkin::ProcessUpdatedDynamicData(EGPUSkinCacheEntryMod
 
 #if RHI_RAYTRACING
 
-void FSkeletalMeshObjectGPUSkin::UpdateRayTracingGeometry(FRHICommandList& RHICmdList, FSkeletalMeshLODRenderData& LODModel, uint32 LODIndex, TArray<FBufferRHIRef>& VertexBufffers)
+void FSkeletalMeshObjectGPUSkin::UpdateRayTracingGeometry(FRHICommandList& RHICmdList, FSkeletalMeshLODRenderData& LODModel, uint32 LODIndex, TArray<FBufferRHIRef>& VertexBuffers)
 {
 	if (IsRayTracingEnabled() && bSupportRayTracing)
 	{
@@ -842,7 +842,7 @@ void FSkeletalMeshObjectGPUSkin::UpdateRayTracingGeometry(FRHICommandList& RHICm
 				const FSkelMeshRenderSection& Section = LODModel.RenderSections[SectionIndex];
 
 				FRayTracingGeometrySegment Segment;
-				Segment.VertexBuffer = VertexBufffers[SectionIndex];
+				Segment.VertexBuffer = VertexBuffers[SectionIndex];
 				Segment.VertexBufferElementType = VET_Float3;
 				Segment.VertexBufferStride = VertexBufferStride;
 				Segment.VertexBufferOffset = 0;
@@ -869,13 +869,13 @@ void FSkeletalMeshObjectGPUSkin::UpdateRayTracingGeometry(FRHICommandList& RHICm
 
 			Initializer.SourceGeometry = LODModel.SourceRayTracingGeometry.RayTracingGeometryRHI;
 
+			// Get the scratch sizes used for build & update
+			RayTracingGeometryStructureSize = RHICmdList.CalcRayTracingGeometrySize(Initializer);
+
 			RayTracingGeometry.LODIndex = LODIndex;
 
 			// Update the new init data
-			RayTracingGeometry.SetInitializer(Initializer);
-
-			// Get the scratch sizes used for build & update
-			RayTracingGeometryStructureSize = RHICmdList.CalcRayTracingGeometrySize(Initializer);
+			RayTracingGeometry.SetInitializer(MoveTemp(Initializer));
 
 			// Only create RHI object but enqueue actual BLAS creation so they can be accumulated
 			RayTracingGeometry.CreateRayTracingGeometry(RHICmdList, ERTAccelerationStructureBuildPriority::Skip);
@@ -888,7 +888,7 @@ void FSkeletalMeshObjectGPUSkin::UpdateRayTracingGeometry(FRHICommandList& RHICm
 			for (int32 SectionIndex = 0; SectionIndex < LODModel.RenderSections.Num(); ++SectionIndex)
 			{
 				FRayTracingGeometrySegment& Segment = RayTracingGeometry.Initializer.Segments[SectionIndex];
-				Segment.VertexBuffer = VertexBufffers[SectionIndex];
+				Segment.VertexBuffer = VertexBuffers[SectionIndex];
 				Segment.VertexBufferOffset = 0;
 			}
 		}
@@ -896,9 +896,8 @@ void FSkeletalMeshObjectGPUSkin::UpdateRayTracingGeometry(FRHICommandList& RHICm
 		// If we are not using world position offset in material, handle BLAS build/refit here
 		if (!bAnySegmentUsesWorldPositionOffset)
 		{
-			EAccelerationStructureBuildMode BuildMode = bRequireRecreatingRayTracingGeometry ? EAccelerationStructureBuildMode::Build : EAccelerationStructureBuildMode::Update;
-			RayTracingUpdateQueue->Add(&RayTracingGeometry, RayTracingGeometryStructureSize, BuildMode);
-			RayTracingGeometry.SetRequiresBuild(false);
+			check(bRequireRecreatingRayTracingGeometry == RayTracingGeometry.GetRequiresBuild());
+			RayTracingUpdateQueue->Add(&RayTracingGeometry, RayTracingGeometryStructureSize);
 		}
 		else
 		{
diff --git a/Engine/Source/Runtime/Engine/Private/SkeletalRenderGPUSkin.h b/Engine/Source/Runtime/Engine/Private/SkeletalRenderGPUSkin.h
--- a/Engine/Source/Runtime/Engine/Private/SkeletalRenderGPUSkin.h
+++ b/Engine/Source/Runtime/Engine/Private/SkeletalRenderGPUSkin.h
@@ -395,7 +395,7 @@ public:
 	 * VSinCS path is still required for world position offset materials but this can still use 
 	 * the updated vertex buffers from here with a passthrough vertex factory.
 	 */
-	ENGINE_API void UpdateRayTracingGeometry(FRHICommandList& RHICmdList, FSkeletalMeshLODRenderData& LODModel, uint32 LODIndex, TArray<FBufferRHIRef>& VertexBufffers);
+	ENGINE_API void UpdateRayTracingGeometry(FRHICommandList& RHICmdList, FSkeletalMeshLODRenderData& LODModel, uint32 LODIndex, TArray<FBufferRHIRef>& VertexBuffers);
 
 #endif // RHI_RAYTRACING
 
diff --git a/Engine/Source/Runtime/Engine/Public/RayTracingSkinnedGeometry.h b/Engine/Source/Runtime/Engine/Public/RayTracingSkinnedGeometry.h
--- a/Engine/Source/Runtime/Engine/Public/RayTracingSkinnedGeometry.h
+++ b/Engine/Source/Runtime/Engine/Public/RayTracingSkinnedGeometry.h
@@ -22,7 +22,7 @@ class FRayTracingSkinnedGeometryUpdateQueue
 {
 public:
 	/** Add a pending update. This is guaranteed to be processed on the next renderer update. */
-	ENGINE_API void Add(FRayTracingGeometry* InRayTracingGeometry, const FRayTracingAccelerationStructureSize& StructureSize, EAccelerationStructureBuildMode InBuildMode);
+	ENGINE_API void Add(FRayTracingGeometry* InRayTracingGeometry, const FRayTracingAccelerationStructureSize& StructureSize);
 
 	/** Remove a pending update. Passing a memory estimation will allow us to keep track of memory overhead from pending object release. */
 	ENGINE_API void Remove(FRayTracingGeometry* RayTracingGeometry, uint32 EstimatedMemory = 0);
