Sort includes in MaterialExpressionHLSL

diff --git a/Engine/Source/Runtime/Engine/Private/Materials/MaterialExpressionHLSL.cpp b/Engine/Source/Runtime/Engine/Private/Materials/MaterialExpressionHLSL.cpp
--- a/Engine/Source/Runtime/Engine/Private/Materials/MaterialExpressionHLSL.cpp
+++ b/Engine/Source/Runtime/Engine/Private/Materials/MaterialExpressionHLSL.cpp
@@ -2,194 +2,194 @@
 
 #if WITH_EDITOR
 
+#include "Curves/CurveLinearColorAtlas.h"
+#include "Engine/Engine.h"
+#include "Engine/Font.h"
+#include "MaterialDomain.h"
 #include "MaterialExpressionIO.h"
+#include "MaterialHLSLGenerator.h"
+#include "MaterialHLSLTree.h"
+#include "Materials/Material.h"
 #include "Materials/MaterialAttributeDefinitionMap.h"
 #include "Materials/MaterialExpression.h"
-#include "Materials/MaterialExpressionExecBegin.h"
-#include "Materials/MaterialExpressionExecEnd.h"
-#include "Materials/MaterialExpressionReroute.h"
-#include "Materials/MaterialExpressionNamedReroute.h"
-#include "Materials/MaterialExpressionGenericConstant.h"
+#include "Materials/MaterialExpressionAbs.h"
+#include "Materials/MaterialExpressionActorPositionWS.h"
+#include "Materials/MaterialExpressionAdd.h"
+#include "Materials/MaterialExpressionAppendVector.h"
+#include "Materials/MaterialExpressionArccosine.h"
+#include "Materials/MaterialExpressionArccosineFast.h"
+#include "Materials/MaterialExpressionArcsine.h"
+#include "Materials/MaterialExpressionArcsineFast.h"
+#include "Materials/MaterialExpressionArctangent.h"
+#include "Materials/MaterialExpressionArctangent2.h"
+#include "Materials/MaterialExpressionArctangent2Fast.h"
+#include "Materials/MaterialExpressionArctangentFast.h"
+#include "Materials/MaterialExpressionAtmosphericLightVector.h"
+#include "Materials/MaterialExpressionBinaryOp.h"
+#include "Materials/MaterialExpressionBlendMaterialAttributes.h"
+#include "Materials/MaterialExpressionBreakMaterialAttributes.h"
+#include "Materials/MaterialExpressionBumpOffset.h"
+#include "Materials/MaterialExpressionCameraPositionWS.h"
+#include "Materials/MaterialExpressionCameraVectorWS.h"
+#include "Materials/MaterialExpressionCeil.h"
+#include "Materials/MaterialExpressionChannelMaskParameter.h"
+#include "Materials/MaterialExpressionClamp.h"
+#include "Materials/MaterialExpressionClearCoatNormalCustomOutput.h"
+#include "Materials/MaterialExpressionCloudLayer.h"
+#include "Materials/MaterialExpressionCollectionParameter.h"
+#include "Materials/MaterialExpressionComponentMask.h"
 #include "Materials/MaterialExpressionConstant.h"
 #include "Materials/MaterialExpressionConstant2Vector.h"
 #include "Materials/MaterialExpressionConstant3Vector.h"
 #include "Materials/MaterialExpressionConstant4Vector.h"
-#include "Materials/MaterialExpressionStaticBool.h"
 #include "Materials/MaterialExpressionConstantBiasScale.h"
-#include "Materials/MaterialExpressionShadingModel.h"
-#include "Materials/MaterialExpressionParameter.h"
-#include "Materials/MaterialExpressionChannelMaskParameter.h"
-#include "Materials/MaterialExpressionCollectionParameter.h"
-#include "Materials/MaterialParameterCollection.h"
+#include "Materials/MaterialExpressionCosine.h"
+#include "Materials/MaterialExpressionCrossProduct.h"
+#include "Materials/MaterialExpressionCurveAtlasRowParameter.h"
+#include "Materials/MaterialExpressionCustom.h"
+#include "Materials/MaterialExpressionDDX.h"
+#include "Materials/MaterialExpressionDDY.h"
 #include "Materials/MaterialExpressionDecalColor.h"
 #include "Materials/MaterialExpressionDecalLifetimeOpacity.h"
+#include "Materials/MaterialExpressionDeltaTime.h"
+#include "Materials/MaterialExpressionDepthFade.h"
 #include "Materials/MaterialExpressionDepthOfFieldFunction.h"
-#include "Materials/MaterialExpressionParticleSubUVProperties.h"
-#include "Materials/MaterialExpressionPathTracingQualitySwitch.h"
-#include "Materials/MaterialExpressionThinTranslucentMaterialOutput.h"
+#include "Materials/MaterialExpressionDeriveNormalZ.h"
+#include "Materials/MaterialExpressionDesaturation.h"
+#include "Materials/MaterialExpressionDistance.h"
+#include "Materials/MaterialExpressionDistanceFieldsRenderingSwitch.h"
+#include "Materials/MaterialExpressionDivide.h"
+#include "Materials/MaterialExpressionDotProduct.h"
 #include "Materials/MaterialExpressionDynamicParameter.h"
-#include "Materials/MaterialExpressionCurveAtlasRowParameter.h"
-#include "Materials/MaterialExpressionStaticSwitchParameter.h"
-#include "Materials/MaterialExpressionStaticComponentMaskParameter.h"
-#include "Materials/MaterialExpressionPixelDepth.h"
-#include "Materials/MaterialExpressionWorldPosition.h"
-#include "Materials/MaterialExpressionCameraPositionWS.h"
-#include "Materials/MaterialExpressionCameraVectorWS.h"
-#include "Materials/MaterialExpressionViewProperty.h"
-#include "Materials/MaterialExpressionIsOrthographic.h"
-#include "Materials/MaterialExpressionTime.h"
-#include "Materials/MaterialExpressionDeltaTime.h"
-#include "Materials/MaterialExpressionScreenPosition.h"
-#include "Materials/MaterialExpressionSceneTexelSize.h"
-#include "Materials/MaterialExpressionViewSize.h"
-#include "Materials/MaterialExpressionTwoSidedSign.h"
-#include "Materials/MaterialExpressionPixelNormalWS.h"
-#include "Materials/MaterialExpressionReflectionVectorWS.h"
-#include "Materials/MaterialExpressionActorPositionWS.h"
-#include "Materials/MaterialExpressionPreSkinnedPosition.h"
-#include "Materials/MaterialExpressionPreSkinnedNormal.h"
-#include "Materials/MaterialExpressionPreSkinnedLocalBounds.h"
-#include "Materials/MaterialExpressionVertexColor.h"
-#include "Materials/MaterialExpressionVertexNormalWS.h"
-#include "Materials/MaterialExpressionVertexTangentWS.h"
-#include "Materials/MaterialExpressionTextureCoordinate.h"
-#include "Materials/MaterialExpressionTextureProperty.h"
-#include "Materials/MaterialExpressionLightmapUVs.h"
+#include "Materials/MaterialExpressionExecBegin.h"
+#include "Materials/MaterialExpressionExecEnd.h"
+#include "Materials/MaterialExpressionExponential.h"
+#include "Materials/MaterialExpressionExponential2.h"
 #include "Materials/MaterialExpressionEyeAdaptation.h"
 #include "Materials/MaterialExpressionEyeAdaptationInverse.h"
+#include "Materials/MaterialExpressionFeatureLevelSwitch.h"
+#include "Materials/MaterialExpressionFloor.h"
+#include "Materials/MaterialExpressionFmod.h"
+#include "Materials/MaterialExpressionFontSample.h"
+#include "Materials/MaterialExpressionForLoop.h"
+#include "Materials/MaterialExpressionFrac.h"
+#include "Materials/MaterialExpressionFresnel.h"
+#include "Materials/MaterialExpressionFunctionInput.h"
+#include "Materials/MaterialExpressionFunctionOutput.h"
+#include "Materials/MaterialExpressionGenericConstant.h"
+#include "Materials/MaterialExpressionGetLocal.h"
+#include "Materials/MaterialExpressionGetMaterialAttributes.h"
+#include "Materials/MaterialExpressionHairAttributes.h"
+#include "Materials/MaterialExpressionIf.h"
+#include "Materials/MaterialExpressionIfThenElse.h"
+#include "Materials/MaterialExpressionInverseLinearInterpolate.h"
+#include "Materials/MaterialExpressionIsOrthographic.h"
+#include "Materials/MaterialExpressionLength.h"
+#include "Materials/MaterialExpressionLightmapUVs.h"
+#include "Materials/MaterialExpressionLinearInterpolate.h"
+#include "Materials/MaterialExpressionLogarithm.h"
+#include "Materials/MaterialExpressionLogarithm10.h"
+#include "Materials/MaterialExpressionLogarithm2.h"
+#include "Materials/MaterialExpressionMakeMaterialAttributes.h"
+#include "Materials/MaterialExpressionMaterialAttributeLayers.h"
+#include "Materials/MaterialExpressionMaterialFunctionCall.h"
+#include "Materials/MaterialExpressionMaterialProxyReplace.h"
+#include "Materials/MaterialExpressionMax.h"
+#include "Materials/MaterialExpressionMin.h"
+#include "Materials/MaterialExpressionMultiply.h"
+#include "Materials/MaterialExpressionNamedReroute.h"
+#include "Materials/MaterialExpressionNoise.h"
+#include "Materials/MaterialExpressionNormalize.h"
+#include "Materials/MaterialExpressionObjectBounds.h"
+#include "Materials/MaterialExpressionObjectOrientation.h"
+#include "Materials/MaterialExpressionObjectPositionWS.h"
+#include "Materials/MaterialExpressionObjectRadius.h"
+#include "Materials/MaterialExpressionOneMinus.h"
+#include "Materials/MaterialExpressionPanner.h"
+#include "Materials/MaterialExpressionParameter.h"
 #include "Materials/MaterialExpressionParticleColor.h"
+#include "Materials/MaterialExpressionParticleDirection.h"
+#include "Materials/MaterialExpressionParticleMacroUV.h"
 #include "Materials/MaterialExpressionParticlePositionWS.h"
 #include "Materials/MaterialExpressionParticleRadius.h"
-#include "Materials/MaterialExpressionParticleDirection.h"
-#include "Materials/MaterialExpressionParticleSpeed.h"
-#include "Materials/MaterialExpressionParticleRelativeTime.h"
 #include "Materials/MaterialExpressionParticleRandom.h"
+#include "Materials/MaterialExpressionParticleRelativeTime.h"
 #include "Materials/MaterialExpressionParticleSize.h"
-#include "Materials/MaterialExpressionParticleMacroUV.h"
-#include "Materials/MaterialExpressionSkyLightEnvMapSample.h"
-#include "Materials/MaterialExpressionTextureSample.h"
-#include "Materials/MaterialExpressionTextureSampleParameter.h"
+#include "Materials/MaterialExpressionParticleSpeed.h"
+#include "Materials/MaterialExpressionParticleSubUVProperties.h"
+#include "Materials/MaterialExpressionPathTracingQualitySwitch.h"
+#include "Materials/MaterialExpressionPerInstanceFadeAmount.h"
+#include "Materials/MaterialExpressionPixelDepth.h"
+#include "Materials/MaterialExpressionPixelNormalWS.h"
+#include "Materials/MaterialExpressionPower.h"
+#include "Materials/MaterialExpressionPreSkinnedLocalBounds.h"
+#include "Materials/MaterialExpressionPreSkinnedNormal.h"
+#include "Materials/MaterialExpressionPreSkinnedPosition.h"
+#include "Materials/MaterialExpressionPreviousFrameSwitch.h"
+#include "Materials/MaterialExpressionQualitySwitch.h"
+#include "Materials/MaterialExpressionRayTracingQualitySwitch.h"
+#include "Materials/MaterialExpressionReflectionVectorWS.h"
+#include "Materials/MaterialExpressionReroute.h"
+#include "Materials/MaterialExpressionRotateAboutAxis.h"
+#include "Materials/MaterialExpressionRotator.h"
+#include "Materials/MaterialExpressionRound.h"
 #include "Materials/MaterialExpressionRuntimeVirtualTextureSample.h"
 #include "Materials/MaterialExpressionRuntimeVirtualTextureSampleParameter.h"
-#include "Materials/MaterialExpressionFontSample.h"
-#include "Materials/MaterialExpressionTextureObject.h"
-#include "Materials/MaterialExpressionTextureObjectParameter.h"
-#include "Materials/MaterialExpressionSceneTexture.h"
+#include "Materials/MaterialExpressionSaturate.h"
 #include "Materials/MaterialExpressionSceneColor.h"
 #include "Materials/MaterialExpressionSceneDepth.h"
-#include "Materials/MaterialExpressionDepthFade.h"
-#include "Materials/MaterialExpressionNoise.h"
-#include "Materials/MaterialExpressionVectorNoise.h"
-#include "Materials/MaterialExpressionVertexInterpolator.h"
-#include "Materials/MaterialExpressionFunctionInput.h"
-#include "Materials/MaterialExpressionFunctionOutput.h"
-#include "Materials/MaterialExpressionMaterialFunctionCall.h"
-#include "Materials/MaterialExpressionMaterialAttributeLayers.h"
-#include "Materials/MaterialExpressionBlendMaterialAttributes.h"
-#include "Materials/MaterialExpressionPreviousFrameSwitch.h"
-#include "Materials/MaterialExpressionStaticSwitch.h"
-#include "Materials/MaterialExpressionFeatureLevelSwitch.h"
-#include "Materials/MaterialExpressionShadingPathSwitch.h"
-#include "Materials/MaterialExpressionQualitySwitch.h"
+#include "Materials/MaterialExpressionSceneTexelSize.h"
+#include "Materials/MaterialExpressionSceneTexture.h"
+#include "Materials/MaterialExpressionScreenPosition.h"
+#include "Materials/MaterialExpressionSetLocal.h"
+#include "Materials/MaterialExpressionSetMaterialAttributes.h"
 #include "Materials/MaterialExpressionShaderStageSwitch.h"
-#include "Materials/MaterialExpressionRayTracingQualitySwitch.h"
-#include "Materials/MaterialExpressionVirtualTextureFeatureSwitch.h"
-#include "Materials/MaterialExpressionDistanceFieldsRenderingSwitch.h"
+#include "Materials/MaterialExpressionShadingModel.h"
+#include "Materials/MaterialExpressionShadingPathSwitch.h"
 #include "Materials/MaterialExpressionShadowReplace.h"
-#include "Materials/MaterialExpressionMaterialProxyReplace.h"
-#include "Materials/MaterialExpressionGetLocal.h"
-#include "Materials/MaterialExpressionOneMinus.h"
-#include "Materials/MaterialExpressionAbs.h"
-#include "Materials/MaterialExpressionSquareRoot.h"
-#include "Materials/MaterialExpressionExponential.h"
-#include "Materials/MaterialExpressionExponential2.h"
-#include "Materials/MaterialExpressionLength.h"
-#include "Materials/MaterialExpressionLogarithm.h"
-#include "Materials/MaterialExpressionLogarithm2.h"
-#include "Materials/MaterialExpressionLogarithm10.h"
-#include "Materials/MaterialExpressionFrac.h"
-#include "Materials/MaterialExpressionFloor.h"
-#include "Materials/MaterialExpressionCeil.h"
-#include "Materials/MaterialExpressionRound.h"
-#include "Materials/MaterialExpressionTruncate.h"
-#include "Materials/MaterialExpressionTruncateLWC.h"
-#include "Materials/MaterialExpressionSaturate.h"
 #include "Materials/MaterialExpressionSign.h"
 #include "Materials/MaterialExpressionSine.h"
-#include "Materials/MaterialExpressionCosine.h"
-#include "Materials/MaterialExpressionTangent.h"
-#include "Materials/MaterialExpressionArcsine.h"
-#include "Materials/MaterialExpressionArccosine.h"
-#include "Materials/MaterialExpressionArctangent.h"
-#include "Materials/MaterialExpressionArctangent2.h"
-#include "Materials/MaterialExpressionArcsineFast.h"
-#include "Materials/MaterialExpressionArccosineFast.h"
-#include "Materials/MaterialExpressionArctangentFast.h"
-#include "Materials/MaterialExpressionArctangent2Fast.h"
-#include "Materials/MaterialExpressionDDX.h"
-#include "Materials/MaterialExpressionDDY.h"
-#include "Materials/MaterialExpressionAdd.h"
-#include "Materials/MaterialExpressionSubtract.h"
-#include "Materials/MaterialExpressionMultiply.h"
-#include "Materials/MaterialExpressionDivide.h"
-#include "Materials/MaterialExpressionFmod.h"
-#include "Materials/MaterialExpressionPower.h"
-#include "Materials/MaterialExpressionDotProduct.h"
-#include "Materials/MaterialExpressionCrossProduct.h"
-#include "Materials/MaterialExpressionMin.h"
-#include "Materials/MaterialExpressionMax.h"
-#include "Materials/MaterialExpressionClamp.h"
-#include "Materials/MaterialExpressionLinearInterpolate.h"
-#include "Materials/MaterialExpressionInverseLinearInterpolate.h"
-#include "Materials/MaterialExpressionStep.h"
+#include "Materials/MaterialExpressionSkyAtmosphereLightDirection.h"
+#include "Materials/MaterialExpressionSkyLightEnvMapSample.h"
 #include "Materials/MaterialExpressionSmoothStep.h"
-#include "Materials/MaterialExpressionDistance.h"
-#include "Materials/MaterialExpressionNormalize.h"
-#include "Materials/MaterialExpressionBinaryOp.h"
-#include "Materials/MaterialExpressionAppendVector.h"
-#include "Materials/MaterialExpressionComponentMask.h"
-#include "Materials/MaterialExpressionGetMaterialAttributes.h"
-#include "Materials/MaterialExpressionSetMaterialAttributes.h"
-#include "Materials/MaterialExpressionMakeMaterialAttributes.h"
-#include "Materials/MaterialExpressionBreakMaterialAttributes.h"
+#include "Materials/MaterialExpressionSphereMask.h"
+#include "Materials/MaterialExpressionSquareRoot.h"
+#include "Materials/MaterialExpressionStaticBool.h"
+#include "Materials/MaterialExpressionStaticComponentMaskParameter.h"
+#include "Materials/MaterialExpressionStaticSwitch.h"
+#include "Materials/MaterialExpressionStaticSwitchParameter.h"
+#include "Materials/MaterialExpressionStep.h"
+#include "Materials/MaterialExpressionSubtract.h"
+#include "Materials/MaterialExpressionTangent.h"
+#include "Materials/MaterialExpressionTextureCoordinate.h"
+#include "Materials/MaterialExpressionTextureObject.h"
+#include "Materials/MaterialExpressionTextureObjectParameter.h"
+#include "Materials/MaterialExpressionTextureProperty.h"
+#include "Materials/MaterialExpressionTextureSample.h"
+#include "Materials/MaterialExpressionTextureSampleParameter.h"
+#include "Materials/MaterialExpressionThinTranslucentMaterialOutput.h"
+#include "Materials/MaterialExpressionTime.h"
 #include "Materials/MaterialExpressionTransform.h"
 #include "Materials/MaterialExpressionTransformPosition.h"
-#include "Materials/MaterialExpressionIf.h"
-#include "Materials/MaterialExpressionFresnel.h"
-#include "Materials/MaterialExpressionDesaturation.h"
-#include "Materials/MaterialExpressionSphereMask.h"
-#include "Materials/MaterialExpressionRotateAboutAxis.h"
-#include "Materials/MaterialExpressionBumpOffset.h"
-#include "Materials/MaterialExpressionPanner.h"
-#include "Materials/MaterialExpressionRotator.h"
-#include "Materials/MaterialExpressionAtmosphericLightVector.h"
-#include "Materials/MaterialExpressionSkyAtmosphereLightDirection.h"
+#include "Materials/MaterialExpressionTruncate.h"
+#include "Materials/MaterialExpressionTruncateLWC.h"
+#include "Materials/MaterialExpressionTwoSidedSign.h"
+#include "Materials/MaterialExpressionVectorNoise.h"
+#include "Materials/MaterialExpressionVertexColor.h"
+#include "Materials/MaterialExpressionVertexInterpolator.h"
+#include "Materials/MaterialExpressionVertexNormalWS.h"
+#include "Materials/MaterialExpressionVertexTangentWS.h"
+#include "Materials/MaterialExpressionViewProperty.h"
+#include "Materials/MaterialExpressionViewSize.h"
+#include "Materials/MaterialExpressionVirtualTextureFeatureSwitch.h"
 #include "Materials/MaterialExpressionVolumetricAdvancedMaterialInput.h"
 #include "Materials/MaterialExpressionVolumetricAdvancedMaterialOutput.h"
-#include "Materials/MaterialExpressionHairAttributes.h"
-#include "Materials/MaterialExpressionCloudLayer.h"
-#include "Materials/MaterialExpressionPerInstanceFadeAmount.h"
-#include "Materials/MaterialExpressionObjectBounds.h"
-#include "Materials/MaterialExpressionObjectOrientation.h"
-#include "Materials/MaterialExpressionObjectPositionWS.h"
-#include "Materials/MaterialExpressionObjectRadius.h"
-#include "Materials/MaterialExpressionCustom.h"
-#include "Materials/MaterialExpressionClearCoatNormalCustomOutput.h"
-#include "Materials/MaterialExpressionSetLocal.h"
-#include "Materials/MaterialExpressionIfThenElse.h"
-#include "Materials/MaterialExpressionForLoop.h"
 #include "Materials/MaterialExpressionWhileLoop.h"
-#include "Materials/MaterialExpressionDeriveNormalZ.h"
+#include "Materials/MaterialExpressionWorldPosition.h"
 #include "Materials/MaterialFunctionInterface.h"
-#include "Materials/Material.h"
-#include "MaterialHLSLTree.h"
+#include "Materials/MaterialParameterCollection.h"
 #include "MaterialShared.h"
-#include "MaterialDomain.h"
-#include "Engine/Engine.h"
-#include "Engine/Font.h"
-#include "Curves/CurveLinearColorAtlas.h"
 #include "Misc/MemStackUtility.h"
-#include "MaterialHLSLGenerator.h"
 
 bool UMaterialExpression::GenerateHLSLExpression(FMaterialHLSLGenerator& Generator, UE::HLSLTree::FScope& Scope, int32 OutputIndex, UE::HLSLTree::FExpression const*& OutExpression) const
 {
