Add a bReplaceHitWithSubComponents option to collision query params that will let us determine whether or not we hit a cluster union by subtracing into the cluster union but only return the cluster union if it hit

diff --git a/Engine/Source/Runtime/Engine/Private/Collision/Collision.cpp b/Engine/Source/Runtime/Engine/Private/Collision/Collision.cpp
--- a/Engine/Source/Runtime/Engine/Private/Collision/Collision.cpp
+++ b/Engine/Source/Runtime/Engine/Private/Collision/Collision.cpp
@@ -56,6 +56,7 @@ FCollisionQueryParams::FCollisionQueryParams(FName InTraceTag, const TStatId& In
 		bDebugQuery = false;
 #endif
 	bTraceIntoSubComponents = true;
+	bReplaceHitWithSubComponents = true;
 }
 
 
diff --git a/Engine/Source/Runtime/Engine/Private/Collision/SceneQuery.cpp b/Engine/Source/Runtime/Engine/Private/Collision/SceneQuery.cpp
--- a/Engine/Source/Runtime/Engine/Private/Collision/SceneQuery.cpp
+++ b/Engine/Source/Runtime/Engine/Private/Collision/SceneQuery.cpp
@@ -539,7 +539,10 @@ bool TSceneCastCommonImp(const UWorld* World, typename Traits::TOutHits& OutHits
 							FClusterUnionHit ClusterUnionHit = DoClusterUnionTrace(OutHits[Index], NewHit);
 							if (ClusterUnionHit.bIsClusterUnion)
 							{
-								ClusterUnionIndices.Add(Index);
+								if (Params.bReplaceHitWithSubComponents || !ClusterUnionHit.bHit)
+								{
+									ClusterUnionIndices.Add(Index);
+								}
 
 								if (ClusterUnionHit.bHit)
 								{
@@ -554,7 +557,11 @@ bool TSceneCastCommonImp(const UWorld* World, typename Traits::TOutHits& OutHits
 							OutHits.RemoveAtSwap(ClusterUnionIndices[Index], 1, false);
 						}
 
-						OutHits.Append(AllNewHits);
+						if (Params.bReplaceHitWithSubComponents)
+						{
+							OutHits.Append(AllNewHits);
+						}
+						
 						bBlockingHit &= !OutHits.IsEmpty();
 					}
 					else
@@ -566,7 +573,10 @@ bool TSceneCastCommonImp(const UWorld* World, typename Traits::TOutHits& OutHits
 							bBlockingHit = ClusterUnionHit.bHit;
 							if (ClusterUnionHit.bHit)
 							{
-								OutHits = NewHit;
+								if (Params.bReplaceHitWithSubComponents)
+								{
+									OutHits = NewHit;
+								}
 							}
 							else if (AActor* ClusterUnionActor = OutHits.GetActor())
 							{
@@ -929,7 +939,10 @@ bool GeomOverlapMultiImp(const UWorld* World, const FPhysicsGeometry& Geom, cons
 								FClusterUnionHit ClusterUnionHit = DoClusterUnionOverlap(OutOverlaps[Index], NewOverlaps);
 								if (ClusterUnionHit.bIsClusterUnion)
 								{
-									ClusterUnionIndices.Add(Index);
+									if (Params.bReplaceHitWithSubComponents || !ClusterUnionHit.bHit)
+									{
+										ClusterUnionIndices.Add(Index);
+									}
 
 									if (ClusterUnionHit.bHit)
 									{
@@ -944,7 +957,10 @@ bool GeomOverlapMultiImp(const UWorld* World, const FPhysicsGeometry& Geom, cons
 								OutOverlaps.RemoveAtSwap(ClusterUnionIndices[Index], 1, false);
 							}
 
-							OutOverlaps.Append(AllNewOverlaps);
+							if (Params.bReplaceHitWithSubComponents)
+							{
+								OutOverlaps.Append(AllNewOverlaps);
+							}
 							bHaveBlockingHit &= !OutOverlaps.IsEmpty();
 						}
 					}
diff --git a/Engine/Source/Runtime/Engine/Private/Components/CharacterMovementComponent.cpp b/Engine/Source/Runtime/Engine/Private/Components/CharacterMovementComponent.cpp
--- a/Engine/Source/Runtime/Engine/Private/Components/CharacterMovementComponent.cpp
+++ b/Engine/Source/Runtime/Engine/Private/Components/CharacterMovementComponent.cpp
@@ -6598,7 +6598,8 @@ void UCharacterMovementComponent::ComputeFloorDist(const FVector& CapsuleLocatio
 	// Having a character base on a component within a cluster union will cause replication problems.
 	// The issue is that ACharacter::SetBase() gets a GeometryCollectionComponent passed to it when standing on the DynamicPlatform
 	// and that GC is never simulating, and since it's not simulating it's stopping the based movement flow there for simulated proxies.
-	QueryParams.bTraceIntoSubComponents = false;
+	QueryParams.bTraceIntoSubComponents = true;
+	QueryParams.bReplaceHitWithSubComponents = false;
 
 	FCollisionResponseParams ResponseParam;
 	InitCollisionParams(QueryParams, ResponseParam);
diff --git a/Engine/Source/Runtime/Engine/Public/CollisionQueryParams.h b/Engine/Source/Runtime/Engine/Public/CollisionQueryParams.h
--- a/Engine/Source/Runtime/Engine/Public/CollisionQueryParams.h
+++ b/Engine/Source/Runtime/Engine/Public/CollisionQueryParams.h
@@ -70,6 +70,9 @@ struct FCollisionQueryParams
 	/** Whether to ignore traces to the cluster union and trace against its children instead. */
 	bool bTraceIntoSubComponents;
 
+	/** If bTraceIntoSubComponents is true, whether to replace the hit of the cluster union with its children instead. */
+	bool bReplaceHitWithSubComponents;
+
 	/** Filters query by mobility types (static vs stationary/movable)*/
 	EQueryMobilityType MobilityType;
 
