[Quartz] Removing unnecessary ensure on Delegate.IsBound() w/ some extra gaurds around ref-counting going below zero.
#p4v-preflight-copy 27858980

diff --git a/Engine/Source/Runtime/Engine/Classes/Sound/QuartzSubscription.cpp b/Engine/Source/Runtime/Engine/Classes/Sound/QuartzSubscription.cpp
--- a/Engine/Source/Runtime/Engine/Classes/Sound/QuartzSubscription.cpp
+++ b/Engine/Source/Runtime/Engine/Classes/Sound/QuartzSubscription.cpp
@@ -95,11 +95,11 @@ void FQuartzTickableObject::ExecCommand(const Audio::FQuartzQuantizedCommandDele
 
 	// Broadcast to the BP delegate if we have one bound
 	if (Data.DelegateID >= 0 && Data.DelegateID < QuantizedCommandDelegates.Num()
-		&& ensure(QuantizedCommandDelegates[Data.DelegateID].MulticastDelegate.IsBound()))
+		&& QuantizedCommandDelegates[Data.DelegateID].MulticastDelegate.IsBound())
 	{
 		FCommandDelegateGameThreadData& GameThreadEntry = QuantizedCommandDelegates[Data.DelegateID];
 
-		GameThreadEntry.MulticastDelegate.Broadcast(Data.DelegateSubType, "Quartz Evenct");
+		GameThreadEntry.MulticastDelegate.Broadcast(Data.DelegateSubType, "Quartz Event");
 
 		// track the number of active QuantizedCommands that may be sending info back to us.
 		// this is a bit of a hack because sound cues can play multiple wave instances
@@ -116,10 +116,11 @@ void FQuartzTickableObject::ExecCommand(const Audio::FQuartzQuantizedCommandDele
 		bShouldDecrement |= (DecrementSlotIndexOnStartedCvar && Data.DelegateSubType == EQuartzCommandDelegateSubType::CommandOnStarted);
 			
 		// are all the commands for this delegate done?
-		if (bShouldDecrement && (GameThreadEntry.RefCount.Decrement() == 0))
+		if (bShouldDecrement && (GameThreadEntry.RefCount.Decrement() <= 0))
 		{
 			// free up the slot for new subscriptions on this clock handle
 			GameThreadEntry.MulticastDelegate.Clear();
+			GameThreadEntry.RefCount.Reset();
 		}
 	}
 
