Change Substrate material attribute conversion node to only output Substrate data and change root node to allow material attribute to be plugged to transfert all other geometric/meta data.

diff --git a/Engine/Source/Editor/MaterialEditor/Private/MaterialEditorDetailCustomization.cpp b/Engine/Source/Editor/MaterialEditor/Private/MaterialEditorDetailCustomization.cpp
--- a/Engine/Source/Editor/MaterialEditor/Private/MaterialEditorDetailCustomization.cpp
+++ b/Engine/Source/Editor/MaterialEditor/Private/MaterialEditorDetailCustomization.cpp
@@ -825,14 +825,7 @@ void FMaterialDetailCustomization::CustomizeDetails( IDetailLayoutBuilder& Detai
 				}
 			}
 
-			if (Substrate::IsSubstrateEnabled())
-			{
-				if (PropertyName == GET_MEMBER_NAME_CHECKED(UMaterial, bUseMaterialAttributes))
-				{
-					DetailLayout.HideProperty(PropertyHandle);
-				}
-			}
-			else
+			if (!Substrate::IsSubstrateEnabled())
 			{
 				if (PropertyName == GET_MEMBER_NAME_CHECKED(UMaterial, bIsThinSurface))
 				{
diff --git a/Engine/Source/Editor/UnrealEd/Classes/MaterialGraph/MaterialGraph.h b/Engine/Source/Editor/UnrealEd/Classes/MaterialGraph/MaterialGraph.h
--- a/Engine/Source/Editor/UnrealEd/Classes/MaterialGraph/MaterialGraph.h
+++ b/Engine/Source/Editor/UnrealEd/Classes/MaterialGraph/MaterialGraph.h
@@ -61,7 +61,9 @@ struct FMaterialInputInfo
 
 		if (Material->bUseMaterialAttributes && !bIgnoreMaterialAttributes)
 		{
-			return Property == MP_MaterialAttributes;
+			// On the root node both MaterialAttributes and FrontMaterial are visible when Substrate is enabled. 
+			// Otherwise only MaterialAttributes is visible.
+			return Property == MP_MaterialAttributes || (Property == MP_FrontMaterial && Substrate::IsSubstrateEnabled());
 		}
 		else if( Material->IsUIMaterial() )
 		{
diff --git a/Engine/Source/Runtime/Engine/Classes/Materials/MaterialExpressionSubstrate.h b/Engine/Source/Runtime/Engine/Classes/Materials/MaterialExpressionSubstrate.h
--- a/Engine/Source/Runtime/Engine/Classes/Materials/MaterialExpressionSubstrate.h
+++ b/Engine/Source/Runtime/Engine/Classes/Materials/MaterialExpressionSubstrate.h
@@ -886,6 +886,7 @@ class UMaterialExpressionSubstrateConvertMaterialAttributes : public UMaterialEx
 	virtual uint32 GetInputType(int32 InputIndex) override;
 	virtual FName GetInputName(int32 InputIndex) const override;
 	virtual bool IsResultSubstrateMaterial(int32 OutputIndex) override;
+	virtual bool IsResultMaterialAttributes(int32 OutputIndex) override;
 	virtual void GatherSubstrateMaterialInfo(FSubstrateMaterialInfo& SubstrateMaterialInfo, int32 OutputIndex) override;
 	virtual FSubstrateOperator* SubstrateGenerateMaterialTopologyTree(class FMaterialCompiler* Compiler, class UMaterialExpression* Parent, int32 OutputIndex) override;
 	virtual void PostEditChangeProperty(FPropertyChangedEvent& PropertyChangedEvent) override;
diff --git a/Engine/Source/Runtime/Engine/Private/Materials/Material.cpp b/Engine/Source/Runtime/Engine/Private/Materials/Material.cpp
--- a/Engine/Source/Runtime/Engine/Private/Materials/Material.cpp
+++ b/Engine/Source/Runtime/Engine/Private/Materials/Material.cpp
@@ -3262,14 +3262,18 @@ void UMaterial::ConvertMaterialToSubstrateMaterial()
 	bool bRelinkCustomOutputNodes = false;
 	UMaterialExpressionSubstrateShadingModels* ConvertNode = nullptr;
 	// Connect all the legacy pin into the conversion node
-	if (bUseMaterialAttributes && EditorOnly->MaterialAttributes.Expression && !EditorOnly->MaterialAttributes.Expression->IsResultSubstrateMaterial(EditorOnly->MaterialAttributes.OutputIndex)) // M_Rifle cause issues there
+	if (bUseMaterialAttributes && EditorOnly->MaterialAttributes.Expression && !EditorOnly->FrontMaterial.IsConnected() && !EditorOnly->MaterialAttributes.Expression->IsResultSubstrateMaterial(EditorOnly->MaterialAttributes.OutputIndex)) // M_Rifle cause issues there
 	{
 		UMaterialExpressionSubstrateConvertMaterialAttributes* ConvertAttributeNode = NewObject<UMaterialExpressionSubstrateConvertMaterialAttributes>(this);
 		ConvertAttributeNode->Material = this;
 		SetPosXAndMoveReferenceToTheRight(ConvertAttributeNode);
 		ConvertAttributeNode->SubsurfaceProfile = bRequireNoSubsurfaceProfile ? nullptr : SubsurfaceProfile;
 
-		MoveConnectionTo(EditorOnly->MaterialAttributes, ConvertAttributeNode, 0);
+		// * Copy the material attribute connection to the conversion node.
+		// * Leave the material attribute existing connection plugged to the root node, 
+		//   so that other input (PixelDepthOffset, WorldPositionOffset, ...) get pull 
+		//   from the material attributes node
+		CopyConnectionTo(EditorOnly->MaterialAttributes, ConvertAttributeNode, 0);
 
 		// Reconnect custom output to material attribute conversion node
 		{
@@ -3297,25 +3301,8 @@ void UMaterial::ConvertMaterialToSubstrateMaterial()
 			}
 		}
 
-		// * Remove support for material attribute
-		// * explicitly connect the Substrate node to the root node
-		// * Forward inputs to the root node (Do not reconnect the Opacity as we handle the opacity by internally within the conversion node)
-		// * Always forward masked opacity because this is required when the blend mode is overriden to masked on a material instance.
-		bUseMaterialAttributes = false;
-		NumCustomizedUVs = 8;
+		// Connect converted Substrate data to root node
 		EditorOnly->FrontMaterial.Connect(0, ConvertAttributeNode);
-		EditorOnly->Opacity.Connect(4, ConvertAttributeNode);
-		EditorOnly->OpacityMask.Connect(5, ConvertAttributeNode);
-		EditorOnly->WorldPositionOffset.Connect(3, ConvertAttributeNode);
-		EditorOnly->AmbientOcclusion.Connect(2, ConvertAttributeNode);
-		EditorOnly->PixelDepthOffset.Connect(1, ConvertAttributeNode);
-		EditorOnly->Refraction.Connect(6, ConvertAttributeNode);
-
-		// Need to connect all CustomizedUV, as NumCustomizedUVs might be 0, while legacy attribute would connect & use all the customized UVs inputs.
-		for (int32 UVIndex = 0; UVIndex<NumCustomizedUVs; ++UVIndex)
-		{
-			EditorOnly->CustomizedUVs[UVIndex].Connect(7+UVIndex, ConvertAttributeNode);
-		}
 
 		// Shading Model
 		// * either use the shader graph expression 
@@ -4596,8 +4583,8 @@ void UMaterial::PostEditChangePropertyInternal(FPropertyChangedEvent& PropertyCh
 
 	// If BLEND_TranslucentColoredTransmittance is selected while Substrate is not enabled, force BLEND_Translucent blend mode
 	if (!Substrate::IsSubstrateEnabled() && BlendMode == BLEND_TranslucentColoredTransmittance)
-		{
-			BlendMode = BLEND_Translucent;
+	{
+		BlendMode = BLEND_Translucent;
 	}
 
 	bool bRequiresCompilation = true;
@@ -6777,11 +6764,15 @@ bool UMaterial::IsPropertySupported(EMaterialProperty InProperty) const
 		case MP_FrontMaterial:
 			bSupported = true;
 			break;
+		case MP_MaterialAttributes:
+			bSupported = bUseMaterialAttributes;
+			break;
 		}
 
 		if (InProperty >= MP_CustomizedUVs0 && InProperty <= MP_CustomizedUVs7)
 		{
-			bSupported = (InProperty - MP_CustomizedUVs0) < NumCustomizedUVs;
+			// When bUseMaterialAttributes is enabled all MP_CustomizedUVs are valid to match legacy behavior
+			bSupported = bUseMaterialAttributes || (InProperty - MP_CustomizedUVs0) < NumCustomizedUVs ;
 		}
 	}
 	return bSupported;
@@ -6975,6 +6966,9 @@ static bool IsPropertyActive_Internal(EMaterialProperty InProperty,
 			case MP_FrontMaterial:
 				Active = true;
 				break;
+			case MP_MaterialAttributes:
+				Active = true;
+				break;
 			}
 
 			if (InProperty >= MP_CustomizedUVs0 && InProperty <= MP_CustomizedUVs7)
@@ -7054,7 +7048,7 @@ static bool IsPropertyActive_Internal(EMaterialProperty InProperty,
 		case MP_SurfaceThickness:
 		case MP_FrontMaterial:
 			{
-				Active = bSubstrateEnabled;
+				Active = false;
 				break;
 			}
 		case MP_MaterialAttributes:
diff --git a/Engine/Source/Runtime/Engine/Private/Materials/MaterialExpressions.cpp b/Engine/Source/Runtime/Engine/Private/Materials/MaterialExpressions.cpp
--- a/Engine/Source/Runtime/Engine/Private/Materials/MaterialExpressions.cpp
+++ b/Engine/Source/Runtime/Engine/Private/Materials/MaterialExpressions.cpp
@@ -26470,18 +26470,6 @@ UMaterialExpressionSubstrateConvertMaterialAttributes::UMaterialExpressionSubstr
 #if WITH_EDITORONLY_DATA
 	Outputs.Reset();
 	Outputs.Add(FExpressionOutput(TEXT(""))); // Substrate
-	Outputs.Add(FExpressionOutput(TEXT("PixelDepthOffset")));
-	Outputs.Add(FExpressionOutput(TEXT("AmbientOcclusion")));
-	Outputs.Add(FExpressionOutput(TEXT("WorldPositionOffset")));
-	Outputs.Add(FExpressionOutput(TEXT("Opacity")));
-	Outputs.Add(FExpressionOutput(TEXT("OpacityMask")));
-	Outputs.Add(FExpressionOutput(TEXT("Refraction")));
-
-	for (int32 UVIndex = 0; UVIndex <= MP_CustomizedUVs7 - MP_CustomizedUVs0; UVIndex++)
-	{
-		Outputs.Add(FExpressionOutput(*FString::Printf(TEXT("CustomizedUV%u"), UVIndex), 1, 1, 1, 0, 0));
-	}
-
 	bShowOutputNameOnPin = true;
 #endif
 }
@@ -26489,37 +26477,7 @@ UMaterialExpressionSubstrateConvertMaterialAttributes::UMaterialExpressionSubstr
 #if WITH_EDITOR
 int32 UMaterialExpressionSubstrateConvertMaterialAttributes::Compile(class FMaterialCompiler* Compiler, int32 OutputIndex)
 {
-	if (OutputIndex == 1)
-	{
-		return MaterialAttributes.CompileWithDefault(Compiler, FMaterialAttributeDefinitionMap::GetID(MP_PixelDepthOffset));
-	}
-	else if (OutputIndex == 2)
-	{
-		return MaterialAttributes.CompileWithDefault(Compiler, FMaterialAttributeDefinitionMap::GetID(MP_AmbientOcclusion));
-	}
-	else if (OutputIndex == 3)
-	{
-		return MaterialAttributes.CompileWithDefault(Compiler, FMaterialAttributeDefinitionMap::GetID(MP_WorldPositionOffset));
-	}
-	else if (OutputIndex == 4)
-	{
-		return MaterialAttributes.CompileWithDefault(Compiler, FMaterialAttributeDefinitionMap::GetID(MP_Opacity));
-	}
-	else if (OutputIndex == 5)
-	{
-		return MaterialAttributes.CompileWithDefault(Compiler, FMaterialAttributeDefinitionMap::GetID(MP_OpacityMask));
-	}
-	else if (OutputIndex == 6)
-	{
-		return MaterialAttributes.CompileWithDefault(Compiler, FMaterialAttributeDefinitionMap::GetID(MP_Refraction));
-	}
-	
-	else if (OutputIndex >= 7 && OutputIndex <= 7 + (MP_CustomizedUVs7 - MP_CustomizedUVs0))
-	{
-		const uint32 UVIndex = OutputIndex-7u;
-		return MaterialAttributes.CompileWithDefault(Compiler, FMaterialAttributeDefinitionMap::GetID(EMaterialProperty(MP_CustomizedUVs0 + UVIndex)));
-	}
-	else if (OutputIndex != 0)
+	if (OutputIndex != 0)
 	{
 		return Compiler->Error(TEXT("Output pin index error"));
 	}
@@ -26666,21 +26624,7 @@ uint32 UMaterialExpressionSubstrateConvertMaterialAttributes::GetOutputType(int3
 {
 	switch (OutputIndex)
 	{
-		case 0 : return MCT_Substrate;		// Substrata data
-		case 1 : return MCT_Float1;			// PixelDepthOffset
-		case 2 : return MCT_Float1; 		// AmbientOcclusion
-		case 3 : return MCT_Float3; 		// WorldPositionOffset
-		case 4 : return MCT_Float1; 		// Opacity
-		case 5 : return MCT_Float1; 		// OpacityMask
-		case 6 : return MCT_Float; 			// Refraction
-		case 7 : return MCT_Float2;			// MP_CustomizedUVs0
-		case 8 : return MCT_Float2;			// MP_CustomizedUVs1
-		case 9 : return MCT_Float2;			// MP_CustomizedUVs2
-		case 10: return MCT_Float2;			// MP_CustomizedUVs3
-		case 11: return MCT_Float2;			// MP_CustomizedUVs4
-		case 12: return MCT_Float2;			// MP_CustomizedUVs5
-		case 13: return MCT_Float2;			// MP_CustomizedUVs6
-		case 14: return MCT_Float2; 		// MP_CustomizedUVs7
+		case 0 : return MCT_Substrate;
 	}
 	check(false);
 	return MCT_Float1;
@@ -26720,13 +26664,7 @@ void UMaterialExpressionSubstrateConvertMaterialAttributes::GetConnectorToolTip(
 {
 	switch (OutputIndex)
 	{
-		case 0: OutToolTip.Add(TEXT("TT Ouput")); break;
-		case 1: OutToolTip.Add(TEXT("TT PixelDepthOffset")); break;
-		case 2: OutToolTip.Add(TEXT("TT AmbientOcclusion")); break;
-		case 3: OutToolTip.Add(TEXT("TT WorldPositionOffset")); break;
-		case 4: OutToolTip.Add(TEXT("TT Opacity")); break;
-		case 5: OutToolTip.Add(TEXT("TT OpacityMask")); break;
-		case 6: OutToolTip.Add(TEXT("TT Refraction")); break;
+		case 0: OutToolTip.Add(TEXT("TT Out Substrate Data")); break;
 	}
 	Super::GetConnectorToolTip(InputIndex, INDEX_NONE, OutToolTip);
 }
@@ -26736,6 +26674,11 @@ bool UMaterialExpressionSubstrateConvertMaterialAttributes::IsResultSubstrateMat
 	return OutputIndex == 0;
 }
 
+bool UMaterialExpressionSubstrateConvertMaterialAttributes::IsResultMaterialAttributes(int32 OutputIndex)
+{
+	return false;
+}
+
 void UMaterialExpressionSubstrateConvertMaterialAttributes::GatherSubstrateMaterialInfo(FSubstrateMaterialInfo& SubstrateMaterialInfo, int32 OutputIndex)
 {
 	const uint64 Cached = GetConnectedMaterialAttributesInputs(Material);
