Simplify Substrate materia attribute conversion.

Now only the Water inputs remain and other are pulled from material attribute directly.

diff --git a/Engine/Source/Runtime/Engine/Classes/Materials/MaterialExpressionSubstrate.h b/Engine/Source/Runtime/Engine/Classes/Materials/MaterialExpressionSubstrate.h
--- a/Engine/Source/Runtime/Engine/Classes/Materials/MaterialExpressionSubstrate.h
+++ b/Engine/Source/Runtime/Engine/Classes/Materials/MaterialExpressionSubstrate.h
@@ -829,12 +829,6 @@ class UMaterialExpressionSubstrateConvertMaterialAttributes : public UMaterialEx
 	UPROPERTY()
 	FMaterialAttributesInput MaterialAttributes;
 
-	/**
-	 * The amount of transmitted light from the back side of the surface to the front side of the surface (type = float3, unit = unitless, defaults to 1)
-	 */
-	UPROPERTY()
-	FExpressionInput TransmittanceColor;
-
 		/**
 	* The single scattering Albedo defining the overall color of the Material (type = float3, unit = unitless, default = 0)
 	 */
@@ -859,18 +853,6 @@ class UMaterialExpressionSubstrateConvertMaterialAttributes : public UMaterialEx
 	UPROPERTY()
 	FExpressionInput ColorScaleBehindWater;
 
-	/**
-	 * Take the bottom clear coat surface normal as input. The normal is considered tangent or world space according to the space properties on the main material node. (type = float3, unit = unitless, defaults to vertex normal)
-	 */
-	UPROPERTY()
-	FExpressionInput ClearCoatNormal;
-
-	/**
-	 * Take the tangent output node as input. The tangent is considered tangent or world space according to the space properties on the main material node. (type = float3, unit = unitless, defaults to vertex tangent)
-	 */
-	UPROPERTY()
-	FExpressionInput CustomTangent;
-
 	/** SubsurfaceProfile, for Screen Space Subsurface Scattering. The profile needs to be set up on both the Substrate diffuse node, and the material node at the moment. */
 	UPROPERTY(EditAnywhere, BlueprintReadOnly, Category = Material, meta = (DisplayName = "Subsurface Profile"))
 	TObjectPtr<class USubsurfaceProfile> SubsurfaceProfile;
diff --git a/Engine/Source/Runtime/Engine/Private/Materials/MaterialAttributeDefinitionMap.cpp b/Engine/Source/Runtime/Engine/Private/Materials/MaterialAttributeDefinitionMap.cpp
--- a/Engine/Source/Runtime/Engine/Private/Materials/MaterialAttributeDefinitionMap.cpp
+++ b/Engine/Source/Runtime/Engine/Private/Materials/MaterialAttributeDefinitionMap.cpp
@@ -562,6 +562,19 @@ const FMaterialCustomOutputAttributeDefintion* FMaterialAttributeDefinitionMap::
 	return nullptr;
 }
 
+const FMaterialCustomOutputAttributeDefintion* FMaterialAttributeDefinitionMap::GetCustomAttribute(const FGuid& AttributeID)
+{
+	for (auto& Attribute : GMaterialPropertyAttributesMap.CustomAttributes)
+	{
+		if (Attribute.AttributeID == AttributeID)
+		{
+			return &Attribute;
+		}
+	}
+
+	return nullptr;
+}
+
 void FMaterialAttributeDefinitionMap::GetCustomAttributeList(TArray<FMaterialCustomOutputAttributeDefintion>& CustomAttributeList)
 {
 	CustomAttributeList.Empty(GMaterialPropertyAttributesMap.CustomAttributes.Num());
diff --git a/Engine/Source/Runtime/Engine/Private/Materials/MaterialExpressions.cpp b/Engine/Source/Runtime/Engine/Private/Materials/MaterialExpressions.cpp
--- a/Engine/Source/Runtime/Engine/Private/Materials/MaterialExpressions.cpp
+++ b/Engine/Source/Runtime/Engine/Private/Materials/MaterialExpressions.cpp
@@ -7308,12 +7308,29 @@ UMaterialExpressionSetMaterialAttributes::UMaterialExpressionSetMaterialAttribut
 #if WITH_EDITOR
 uint64 UMaterialExpressionSetMaterialAttributes::GetConnectedInputs() const
 {
+	static TArray<FMaterialCustomOutputAttributeDefintion> CustomAttributeList;
+	if (CustomAttributeList.Num() == 0)
+	{
+		FMaterialAttributeDefinitionMap::GetCustomAttributeList(CustomAttributeList);
+	}
+
 	uint64 Out = 0ull;
 	const int32 NumInputPins = AttributeSetTypes.Num();
 	for (int32 i = 0; i < NumInputPins; ++i)
 	{
 		const EMaterialProperty Property = FMaterialAttributeDefinitionMap::GetProperty(AttributeSetTypes[i]);
-		if (Property != MP_MAX)
+		if (Property == MP_CustomOutput)
+		{
+			if (const FMaterialCustomOutputAttributeDefintion* Attribute = FMaterialAttributeDefinitionMap::GetCustomAttribute(AttributeSetTypes[i]))
+			{
+				const uint32 CustomIndex = CustomAttributeList.IndexOfByPredicate([Attribute](const FMaterialCustomOutputAttributeDefintion& A) { return A.AttributeID == Attribute->AttributeID; });
+				if (CustomIndex != INDEX_NONE)
+				{
+					Out |= 1ull << uint64(MP_MAX + CustomIndex);
+				}
+			}
+		}
+		else if (Property != MP_MAX)
 		{
 			Out |= 1ull << uint64(Property);
 		}
@@ -26435,6 +26452,29 @@ static uint64 GetConnectedMaterialAttributesInputs(const UMaterial* InMaterial)
 	return Out;
 }
 
+static bool IsCustomMaterialAttributeInputConnected(uint64 InCache, FGuid InProperty)
+{
+	static TArray<FMaterialCustomOutputAttributeDefintion> CustomAttributeList;
+	if (CustomAttributeList.Num() == 0)
+	{
+		FMaterialAttributeDefinitionMap::GetCustomAttributeList(CustomAttributeList);
+	}
+
+	const EMaterialProperty Property = FMaterialAttributeDefinitionMap::GetProperty(InProperty);
+	if (Property == MP_CustomOutput)
+	{
+		if (const FMaterialCustomOutputAttributeDefintion* Attribute = FMaterialAttributeDefinitionMap::GetCustomAttribute(InProperty))
+		{
+			const uint32 CustomIndex = CustomAttributeList.IndexOfByPredicate([Attribute](const FMaterialCustomOutputAttributeDefintion& A) { return A.AttributeID == Attribute->AttributeID; });
+			if (CustomIndex != INDEX_NONE)
+			{
+				return !!(InCache & (1ull << uint64(MP_MAX + CustomIndex)));
+			}
+		}
+	}
+	return false;
+}
+
 static bool IsMaterialAttributeInputConnected(uint64 InCache, EMaterialProperty InProperty)
 {
 	return !!(InCache & (1ull << uint64(InProperty)));
@@ -26458,13 +26498,10 @@ UMaterialExpressionSubstrateConvertMaterialAttributes::UMaterialExpressionSubstr
 	CachedInputs.Empty();
 	CachedInputs.Reserve(1);
 	CachedInputs.Add(&MaterialAttributes);
-	CachedInputs.Add(&TransmittanceColor); 
 	CachedInputs.Add(&WaterScatteringCoefficients); 
 	CachedInputs.Add(&WaterAbsorptionCoefficients);
 	CachedInputs.Add(&WaterPhaseG);
 	CachedInputs.Add(&ColorScaleBehindWater); 
-	CachedInputs.Add(&ClearCoatNormal);
-	CachedInputs.Add(&CustomTangent);
 #endif
 
 #if WITH_EDITORONLY_DATA
@@ -26477,6 +26514,10 @@ UMaterialExpressionSubstrateConvertMaterialAttributes::UMaterialExpressionSubstr
 #if WITH_EDITOR
 int32 UMaterialExpressionSubstrateConvertMaterialAttributes::Compile(class FMaterialCompiler* Compiler, int32 OutputIndex)
 {
+	static const FGuid ClearCoatBottomNormalGuid 	= FMaterialAttributeDefinitionMap::GetCustomAttributeID(TEXT("ClearCoatBottomNormal"));;
+	static const FGuid CustomEyeTangentGuid 		= FMaterialAttributeDefinitionMap::GetCustomAttributeID(TEXT("CustomEyeTangent"));
+	static const FGuid TransmittanceColorGuid 		= FMaterialAttributeDefinitionMap::GetCustomAttributeID(TEXT("TransmittanceColor"));
+
 	if (OutputIndex != 0)
 	{
 		return Compiler->Error(TEXT("Output pin index error"));
@@ -26501,7 +26542,7 @@ int32 UMaterialExpressionSubstrateConvertMaterialAttributes::Compile(class FMate
 	const FSubstrateRegisteredSharedLocalBasis NewRegisteredSharedLocalBasis = SubstrateCompilationInfoCreateSharedLocalBasis(Compiler, NormalCodeChunk, TangentCodeChunk);
 	const FString BasisIndexMacro = Compiler->GetSubstrateSharedLocalBasisIndexMacro(NewRegisteredSharedLocalBasis);
 
-	const bool bHasCoatNormal = ClearCoatNormal.IsConnected();
+	const bool bHasCoatNormal = IsCustomMaterialAttributeInputConnected(Cached, ClearCoatBottomNormalGuid);
 	// Clear coat normal basis
 	int32 ClearCoat_NormalCodeChunk = INDEX_NONE;
 	int32 ClearCoat_TangentCodeChunk = INDEX_NONE;
@@ -26509,7 +26550,7 @@ int32 UMaterialExpressionSubstrateConvertMaterialAttributes::Compile(class FMate
 	FSubstrateRegisteredSharedLocalBasis ClearCoat_NewRegisteredSharedLocalBasis;
 	if (bHasCoatNormal)
 	{
-		ClearCoat_NormalCodeChunk = CompileWithDefaultNormalWS(Compiler, ClearCoatNormal);
+		ClearCoat_NormalCodeChunk = MaterialAttributes.CompileWithDefault(Compiler, ClearCoatBottomNormalGuid);
 		ClearCoat_TangentCodeChunk = TangentCodeChunk;
 		ClearCoat_NewRegisteredSharedLocalBasis = SubstrateCompilationInfoCreateSharedLocalBasis(Compiler, ClearCoat_NormalCodeChunk, ClearCoat_TangentCodeChunk);
 		ClearCoat_BasisIndexMacro = Compiler->GetSubstrateSharedLocalBasisIndexMacro(ClearCoat_NewRegisteredSharedLocalBasis);
@@ -26524,18 +26565,23 @@ int32 UMaterialExpressionSubstrateConvertMaterialAttributes::Compile(class FMate
 
 	// Custom tangent. No need to register it as a local basis, as it is only used for eye shading internal conversion
 	int32 CustomTangent_TangentCodeChunk = INDEX_NONE;
-
-	const bool bHasCustomTangent = CustomTangent.IsConnected();
+	const bool bHasCustomTangent = IsCustomMaterialAttributeInputConnected(Cached, CustomEyeTangentGuid);
 	if (bHasCustomTangent)
 	{
 		// Legacy code doesn't do tangent <-> world basis conversion on tangent output, when provided.
-		CustomTangent_TangentCodeChunk = CompileWithDefaultNormalWS(Compiler, CustomTangent, false /*bConvertToRequestedSpace*/);
+		CustomTangent_TangentCodeChunk = MaterialAttributes.CompileWithDefault(Compiler, CustomEyeTangentGuid); // CompileWithDefaultNormalWS(Compiler, CustomTangent, false /*bConvertToRequestedSpace*/);
+		if (CustomTangent_TangentCodeChunk == INDEX_NONE)
+		{
+			// Nothing is plug in from the linked input, so specify world space normal the BSDF node expects.
+			CustomTangent_TangentCodeChunk = Compiler->VertexNormal();
+		}
 	}
 	else
 	{
 		CustomTangent_TangentCodeChunk = NormalCodeChunk;
 	}
 
+	// SSS profile
 	// Need to handle this by looking at the material instead of the node?
 	int32 SSSProfileCodeChunk = INDEX_NONE;
 	const bool bHasSSS = HasSSS();
@@ -26548,6 +26594,7 @@ int32 UMaterialExpressionSubstrateConvertMaterialAttributes::Compile(class FMate
 	FSubstrateOperator& SubstrateOperator = Compiler->SubstrateCompilationGetOperator(Compiler->SubstrateTreeStackGetPathUniqueId());
 	SubstrateOperator.BSDFRegisteredSharedLocalBasis = NewRegisteredSharedLocalBasis;
 
+	// Opacity
 	int32 OpacityCodeChunk = INDEX_NONE;
 	if (!Compiler->SubstrateSkipsOpacityEvaluation())
 	{
@@ -26560,6 +26607,18 @@ int32 UMaterialExpressionSubstrateConvertMaterialAttributes::Compile(class FMate
 		OpacityCodeChunk = Compiler->Constant(1.0f);
 	}
 
+	// Transmittance Color
+	const bool bHasTransmittanceColor = IsCustomMaterialAttributeInputConnected(Cached, TransmittanceColorGuid);
+	int32 TransmittanceColorChunk = INDEX_NONE;
+	if (bHasTransmittanceColor)
+	{
+		TransmittanceColorChunk = MaterialAttributes.CompileWithDefault(Compiler, TransmittanceColorGuid);
+	}
+	else
+	{
+		TransmittanceColorChunk = Compiler->Constant3(0.5f, 0.5f, 0.5f);
+	}
+
 	int32 ShadingModelCodeChunk = MaterialAttributes.CompileWithDefault(Compiler, FMaterialAttributeDefinitionMap::GetID(MP_ShadingModel));
 	const bool bHasShadingModelExpression = IsMaterialAttributeInputConnected(Cached, MP_ShadingModel);
 	if (!bHasShadingModelExpression)
@@ -26586,7 +26645,7 @@ int32 UMaterialExpressionSubstrateConvertMaterialAttributes::Compile(class FMate
 		// Misc
 		MaterialAttributes.CompileWithDefault(Compiler, FMaterialAttributeDefinitionMap::GetID(MP_EmissiveColor)),
 		OpacityCodeChunk,
-		CompileWithDefaultFloat3(Compiler, TransmittanceColor, 0.5f, 0.5f, 0.5f),
+		TransmittanceColorChunk,
 		// Water
 		CompileWithDefaultFloat3(Compiler, WaterScatteringCoefficients, 0.0f, 0.0f, 0.0f),
 		CompileWithDefaultFloat3(Compiler, WaterAbsorptionCoefficients, 0.0f, 0.0f, 0.0f),
@@ -26633,14 +26692,11 @@ uint32 UMaterialExpressionSubstrateConvertMaterialAttributes::GetOutputType(int3
 uint32 UMaterialExpressionSubstrateConvertMaterialAttributes::GetInputType(int32 InputIndex)
 {
 	if (InputIndex == 0)	  return MCT_MaterialAttributes; // MaterialAttributes
-	else if (InputIndex == 1) return MCT_Float3; // TransmittanceColor
-	else if (InputIndex == 2) return MCT_Float3; // WaterScatteringCoefficients
-	else if (InputIndex == 3) return MCT_Float3; // WaterAbsorptionCoefficients
-	else if (InputIndex == 4) return MCT_Float1; // WaterPhaseG
-	else if (InputIndex == 5) return MCT_Float3; // ColorScaleBehindWater
-	else if (InputIndex == 6) return MCT_Float3; // ClearCoatNormal
-	else if (InputIndex == 7) return MCT_Float3; // CustomTangent
-	else if (InputIndex == 8) return MCT_ShadingModel; // ShadingModelOverride (as it uses 'ShowAsInputPin' metadata)
+	else if (InputIndex == 1) return MCT_Float3; // WaterScatteringCoefficients
+	else if (InputIndex == 2) return MCT_Float3; // WaterAbsorptionCoefficients
+	else if (InputIndex == 3) return MCT_Float1; // WaterPhaseG
+	else if (InputIndex == 4) return MCT_Float3; // ColorScaleBehindWater
+	else if (InputIndex == 5) return MCT_ShadingModel; // ShadingModelOverride (as it uses 'ShowAsInputPin' metadata)
 	
 	check(false);
 	return MCT_Float1;
@@ -26649,14 +26705,11 @@ uint32 UMaterialExpressionSubstrateConvertMaterialAttributes::GetInputType(int32
 FName UMaterialExpressionSubstrateConvertMaterialAttributes::GetInputName(int32 InputIndex) const
 {
 	if (InputIndex == 0)		return TEXT("Attributes");
-	else if (InputIndex == 1)	return TEXT("TransmittanceColor (ThinTranslucent)");
-	else if (InputIndex == 2)	return TEXT("Water Scattering Coefficients (Water)");
-	else if (InputIndex == 3)	return TEXT("Water Absorption Coefficients (Water)");
-	else if (InputIndex == 4)	return TEXT("Water Phase G (Water)");
-	else if (InputIndex == 5)	return TEXT("Color Scale BehindWater (Water)");
-	else if (InputIndex == 6)	return TEXT("Clear Coat Normal (Coat/Eye)");
-	else if (InputIndex == 7)	return TEXT("Custom Tangent (Eye)");
-	else if (InputIndex == 8)	return TEXT("Shading Model From Expression");
+	else if (InputIndex == 1)	return TEXT("Water Scattering Coefficients (Water)");
+	else if (InputIndex == 2)	return TEXT("Water Absorption Coefficients (Water)");
+	else if (InputIndex == 3)	return TEXT("Water Phase G (Water)");
+	else if (InputIndex == 4)	return TEXT("Color Scale BehindWater (Water)");
+	else if (InputIndex == 5)	return TEXT("Shading Model From Expression");
 	return NAME_None;
 }
 
diff --git a/Engine/Source/Runtime/Engine/Public/Materials/MaterialAttributeDefinitionMap.h b/Engine/Source/Runtime/Engine/Public/Materials/MaterialAttributeDefinitionMap.h
--- a/Engine/Source/Runtime/Engine/Public/Materials/MaterialAttributeDefinitionMap.h
+++ b/Engine/Source/Runtime/Engine/Public/Materials/MaterialAttributeDefinitionMap.h
@@ -140,6 +140,9 @@ public:
 	/** Returns the first custom attribute definition that has the specificed attribute name */
 	static ENGINE_API const FMaterialCustomOutputAttributeDefintion* GetCustomAttribute(const FString& AttributeName);
 
+	/** Returns the first custom attribute definition that has the specificed attribute AttributeID */
+	static ENGINE_API const FMaterialCustomOutputAttributeDefintion* GetCustomAttribute(const FGuid& AttributeID);
+
 	/** Returns a list of registered custom attributes */
 	static ENGINE_API void GetCustomAttributeList(TArray<FMaterialCustomOutputAttributeDefintion>& CustomAttributeList);
 
