Fixed PIE crash when client tries to connect to the server: loaded world package PIEInstanceID was not set and package name had no preflix UEDPIE.
This change is basically a revert of an old CL 1810531, which doesn't really go into the details on the change outside of mentioning that it's was a temp fix.
Loading a PIE World without its package PIEInstanceID nor properly fixing up for PIE is not a supported code path.

diff --git a/Engine/Source/Runtime/Engine/Private/UnrealEngine.cpp b/Engine/Source/Runtime/Engine/Private/UnrealEngine.cpp
--- a/Engine/Source/Runtime/Engine/Private/UnrealEngine.cpp
+++ b/Engine/Source/Runtime/Engine/Private/UnrealEngine.cpp
@@ -15166,7 +15166,7 @@ bool UEngine::LoadMap( FWorldContext& WorldContext, FURL URL, class UPendingNetG
 				GIsPlayInEditorWorld = true;
 			}
 			// Otherwise we are probably loading new map while in PIE, so we need to rename world package and all streaming levels
-			else if (WorldContext.PIEInstance != -1 && ((Pending == nullptr) || (Pending->GetDemoNetDriver() != nullptr)))
+			else if (WorldContext.PIEInstance != -1)
 			{
 				NewWorld->RenameToPIEWorld(WorldContext.PIEInstance);
 			}
