Implement GenerateHLSLExpression for DistanceFieldApproxAO, DistanceFieldGradient, and DistanceToNearestSurface material nodes

diff --git a/Engine/Source/Runtime/Engine/Classes/Materials/MaterialExpressionDistanceFieldApproxAO.h b/Engine/Source/Runtime/Engine/Classes/Materials/MaterialExpressionDistanceFieldApproxAO.h
--- a/Engine/Source/Runtime/Engine/Classes/Materials/MaterialExpressionDistanceFieldApproxAO.h
+++ b/Engine/Source/Runtime/Engine/Classes/Materials/MaterialExpressionDistanceFieldApproxAO.h
@@ -44,6 +44,8 @@ class UMaterialExpressionDistanceFieldApproxAO : public UMaterialExpression
 #if WITH_EDITOR
 	virtual int32 Compile(class FMaterialCompiler* Compiler, int32 OutputIndex) override;
 	virtual void GetCaption(TArray<FString>& OutCaptions) const override;
+
+	virtual bool GenerateHLSLExpression(FMaterialHLSLGenerator& Generator, UE::HLSLTree::FScope& Scope, int32 OutputIndex, UE::HLSLTree::FExpression const*& OutExpression) const override;
 #endif
 	//~ End UMaterialExpression Interface
 };
diff --git a/Engine/Source/Runtime/Engine/Classes/Materials/MaterialExpressionDistanceFieldGradient.h b/Engine/Source/Runtime/Engine/Classes/Materials/MaterialExpressionDistanceFieldGradient.h
--- a/Engine/Source/Runtime/Engine/Classes/Materials/MaterialExpressionDistanceFieldGradient.h
+++ b/Engine/Source/Runtime/Engine/Classes/Materials/MaterialExpressionDistanceFieldGradient.h
@@ -21,6 +21,8 @@ class UMaterialExpressionDistanceFieldGradient : public UMaterialExpression
 #if WITH_EDITOR
 	virtual int32 Compile(class FMaterialCompiler* Compiler, int32 OutputIndex) override;
 	virtual void GetCaption(TArray<FString>& OutCaptions) const override;
+
+	virtual bool GenerateHLSLExpression(FMaterialHLSLGenerator& Generator, UE::HLSLTree::FScope& Scope, int32 OutputIndex, UE::HLSLTree::FExpression const*& OutExpression) const override;
 #endif
 	//~ End UMaterialExpression Interface
 };
diff --git a/Engine/Source/Runtime/Engine/Classes/Materials/MaterialExpressionDistanceToNearestSurface.h b/Engine/Source/Runtime/Engine/Classes/Materials/MaterialExpressionDistanceToNearestSurface.h
--- a/Engine/Source/Runtime/Engine/Classes/Materials/MaterialExpressionDistanceToNearestSurface.h
+++ b/Engine/Source/Runtime/Engine/Classes/Materials/MaterialExpressionDistanceToNearestSurface.h
@@ -21,6 +21,8 @@ class UMaterialExpressionDistanceToNearestSurface : public UMaterialExpression
 #if WITH_EDITOR
 	virtual int32 Compile(class FMaterialCompiler* Compiler, int32 OutputIndex) override;
 	virtual void GetCaption(TArray<FString>& OutCaptions) const override;
+
+	virtual bool GenerateHLSLExpression(FMaterialHLSLGenerator& Generator, UE::HLSLTree::FScope& Scope, int32 OutputIndex, UE::HLSLTree::FExpression const*& OutExpression) const override;
 #endif
 	//~ End UMaterialExpression Interface
 };
diff --git a/Engine/Source/Runtime/Engine/Private/Materials/MaterialExpressionHLSL.cpp b/Engine/Source/Runtime/Engine/Private/Materials/MaterialExpressionHLSL.cpp
--- a/Engine/Source/Runtime/Engine/Private/Materials/MaterialExpressionHLSL.cpp
+++ b/Engine/Source/Runtime/Engine/Private/Materials/MaterialExpressionHLSL.cpp
@@ -57,7 +57,10 @@
 #include "Materials/MaterialExpressionDeriveNormalZ.h"
 #include "Materials/MaterialExpressionDesaturation.h"
 #include "Materials/MaterialExpressionDistance.h"
+#include "Materials/MaterialExpressionDistanceFieldApproxAO.h"
+#include "Materials/MaterialExpressionDistanceFieldGradient.h"
 #include "Materials/MaterialExpressionDistanceFieldsRenderingSwitch.h"
+#include "Materials/MaterialExpressionDistanceToNearestSurface.h"
 #include "Materials/MaterialExpressionDivide.h"
 #include "Materials/MaterialExpressionDotProduct.h"
 #include "Materials/MaterialExpressionDynamicParameter.h"
@@ -3646,4 +3649,108 @@ bool UMaterialExpressionDeriveNormalZ::GenerateHLSLExpression(FMaterialHLSLGener
 	return true;
 }
 
+bool UMaterialExpressionDistanceToNearestSurface::GenerateHLSLExpression(FMaterialHLSLGenerator& Generator, UE::HLSLTree::FScope& Scope, int32 OutputIndex, UE::HLSLTree::FExpression const*& OutExpression) const
+{
+	using namespace UE::HLSLTree;
+
+	const FExpression* PositionExpression;
+
+	if (Position.GetTracedInput().Expression)
+	{
+		PositionExpression = Position.AcquireHLSLExpression(Generator, Scope);
+	}
+	else
+	{
+		PositionExpression = Generator.NewExternalInput(Material::EExternalInput::WorldPosition);
+	}
+
+	OutExpression = Generator.GetTree().NewExpression<Material::FExpressionDistanceToNearestSurface>(PositionExpression);
+	return true;
+}
+
+bool UMaterialExpressionDistanceFieldGradient::GenerateHLSLExpression(FMaterialHLSLGenerator& Generator, UE::HLSLTree::FScope& Scope, int32 OutputIndex, UE::HLSLTree::FExpression const*& OutExpression) const
+{
+	using namespace UE::HLSLTree;
+
+	const FExpression* PositionExpression;
+
+	if (Position.GetTracedInput().Expression)
+	{
+		PositionExpression = Position.AcquireHLSLExpression(Generator, Scope);
+	}
+	else
+	{
+		PositionExpression = Generator.NewExternalInput(Material::EExternalInput::WorldPosition);
+	}
+
+	OutExpression = Generator.GetTree().NewExpression<Material::FExpressionDistanceFieldGradient>(PositionExpression);
+	return true;
+}
+
+bool UMaterialExpressionDistanceFieldApproxAO::GenerateHLSLExpression(FMaterialHLSLGenerator& Generator, UE::HLSLTree::FScope& Scope, int32 OutputIndex, UE::HLSLTree::FExpression const*& OutExpression) const
+{
+	using namespace UE::HLSLTree;
+
+	const FExpression* PositionExpression;;
+
+	if (Position.GetTracedInput().Expression)
+	{
+		PositionExpression = Position.AcquireHLSLExpression(Generator, Scope);
+	}
+	else
+	{
+		PositionExpression = Generator.NewExternalInput(Material::EExternalInput::WorldPosition);
+	}
+
+	const FExpression* NormalExpression;
+
+	if (Normal.GetTracedInput().Expression)
+	{
+		NormalExpression = Normal.AcquireHLSLExpression(Generator, Scope);
+	}
+	else
+	{
+		NormalExpression = Generator.NewExternalInput(Material::EExternalInput::WorldVertexNormal);
+	}
+
+	const FExpression* BaseDistanceExpression = BaseDistance.AcquireHLSLExpressionOrConstant(Generator, Scope, BaseDistanceDefault);
+	const FExpression* RadiusExpression = Radius.AcquireHLSLExpressionOrConstant(Generator, Scope, RadiusDefault);
+
+	const int32 LocalNumSteps = FMath::Clamp((int32)NumSteps, 1, 4);
+	const float LocalStepScale = FMath::Max(StepScaleDefault, 1.0f);
+
+	const FExpression* NumStepsMinusOneExpression = Generator.NewConstant(LocalNumSteps - 1);
+	const FExpression* StepScaleExpression = Generator.NewConstant(LocalStepScale);
+
+	const FExpression* StepDistanceExpression;
+	const FExpression* DistanceBiasExpression;
+	const FExpression* MaxDistanceExpression;
+	FTree& Tree = Generator.GetTree();
+
+	if (LocalNumSteps == 1)
+	{
+		StepDistanceExpression = Generator.NewConstant(0.f);
+		DistanceBiasExpression = BaseDistanceExpression;
+		MaxDistanceExpression = BaseDistanceExpression;
+	}
+	else
+	{
+		StepDistanceExpression = Tree.NewDiv(
+			Tree.NewSub(RadiusExpression, BaseDistanceExpression),
+			Tree.NewSub(Tree.NewPowClamped(StepScaleExpression, NumStepsMinusOneExpression), Tree.NewConstant(1.f)));
+		DistanceBiasExpression = Tree.NewSub(BaseDistanceExpression, StepDistanceExpression);
+		MaxDistanceExpression = RadiusExpression;
+	}
+
+	OutExpression = Tree.NewExpression<Material::FExpressionDistanceFieldApproxAO>(
+		PositionExpression,
+		NormalExpression,
+		StepDistanceExpression,
+		DistanceBiasExpression,
+		MaxDistanceExpression,
+		LocalNumSteps,
+		LocalStepScale);
+	return true;
+}
+
 #endif // WITH_EDITOR
diff --git a/Engine/Source/Runtime/Engine/Private/Materials/MaterialHLSLTree.cpp b/Engine/Source/Runtime/Engine/Private/Materials/MaterialHLSLTree.cpp
--- a/Engine/Source/Runtime/Engine/Private/Materials/MaterialHLSLTree.cpp
+++ b/Engine/Source/Runtime/Engine/Private/Materials/MaterialHLSLTree.cpp
@@ -12,6 +12,7 @@
 #include "Engine/BlendableInterface.h" // BL_AfterTonemapping
 #include "VT/VirtualTextureScalability.h"
 #include "VT/RuntimeVirtualTexture.h"
+#include "DataDrivenShaderPlatformInfo.h"
 
 namespace UE::HLSLTree::Material
 {
@@ -1800,6 +1801,139 @@ void FExpressionSkyAtmosphereLightDirection::EmitValueShader(FEmitContext& Conte
 	OutResult.Code = Context.EmitInlineExpression(Scope, Shader::EValueType::Float3, TEXT("MaterialExpressionSkyAtmosphereLightDirection(Parameters, %)"), LightIndex);
 }
 
+namespace Private
+{
+	bool PlatformSupportDistanceFields(FEmitContext& Context)
+	{
+		if (!Context.TargetParameters.IsGenericTarget() && !FDataDrivenShaderPlatformInfo::GetSupportsDistanceFields(Context.TargetParameters.ShaderPlatform))
+		{
+			const FString ShaderPlatformName = FDataDrivenShaderPlatformInfo::GetName(Context.TargetParameters.ShaderPlatform).ToString();
+			return Context.Errorf(TEXT("Node not supported in shader platform %s. The node requires DistanceField support."), *ShaderPlatformName);
+		}
+		return true;
+	}
+}
+
+bool FExpressionDistanceToNearestSurface::PrepareValue(FEmitContext& Context, FEmitScope& Scope, const FRequestedType& RequestedType, FPrepareValueResult& OutResult) const
+{
+	if (!Private::PlatformSupportDistanceFields(Context))
+	{
+		return false;
+	}
+
+	const FPreparedType& PositionType = Context.PrepareExpression(PositionExpression, Scope, Shader::EValueType::Double3);
+	if (PositionType.IsVoid())
+	{
+		return false;
+	}
+
+	if (Context.bMarkLiveValues && Context.MaterialCompilationOutput)
+	{
+		Context.MaterialCompilationOutput->bUsesGlobalDistanceField = true;
+	}
+
+	return OutResult.SetType(Context, RequestedType, EExpressionEvaluation::Shader, Shader::EValueType::Float1);
+}
+
+void FExpressionDistanceToNearestSurface::EmitValueShader(FEmitContext& Context, FEmitScope& Scope, const FRequestedType& RequestedType, FEmitValueShaderResult& OutResult) const
+{
+	const FPreparedType& PositionType = Context.GetPreparedType(PositionExpression, Shader::EValueType::Double3);
+	FEmitShaderExpression* EmitPosition = PositionExpression->GetValueShader(Context, Scope, IsLWCType(PositionType.Type.ValueType) ? Shader::EValueType::Double3 : Shader::EValueType::Float3);
+	OutResult.Code = Context.EmitInlineExpression(Scope, Shader::EValueType::Float1, TEXT("GetDistanceToNearestSurfaceGlobal(%)"), EmitPosition);
+}
+
+bool FExpressionDistanceFieldGradient::PrepareValue(FEmitContext& Context, FEmitScope& Scope, const FRequestedType& RequestedType, FPrepareValueResult& OutResult) const
+{
+	if (!Private::PlatformSupportDistanceFields(Context))
+	{
+		return false;
+	}
+
+	const FPreparedType& PositionType = Context.PrepareExpression(PositionExpression, Scope, Shader::EValueType::Double3);
+	if (PositionType.IsVoid())
+	{
+		return false;
+	}
+
+	if (Context.bMarkLiveValues && Context.MaterialCompilationOutput)
+	{
+		Context.MaterialCompilationOutput->bUsesGlobalDistanceField = true;
+	}
+
+	return OutResult.SetType(Context, RequestedType, EExpressionEvaluation::Shader, Shader::EValueType::Float3);
+}
+
+void FExpressionDistanceFieldGradient::EmitValueShader(FEmitContext& Context, FEmitScope& Scope, const FRequestedType& RequestedType, FEmitValueShaderResult& OutResult) const
+{
+	const FPreparedType& PositionType = Context.GetPreparedType(PositionExpression, Shader::EValueType::Double3);
+	FEmitShaderExpression* EmitPosition = PositionExpression->GetValueShader(Context, Scope, IsLWCType(PositionType.Type.ValueType) ? Shader::EValueType::Double3 : Shader::EValueType::Float3);
+	OutResult.Code = Context.EmitInlineExpression(Scope, Shader::EValueType::Float3, TEXT("GetDistanceFieldGradientGlobal(%)"), EmitPosition);
+}
+
+bool FExpressionDistanceFieldApproxAO::PrepareValue(FEmitContext& Context, FEmitScope& Scope, const FRequestedType& RequestedType, FPrepareValueResult& OutResult) const
+{
+	if (!Private::PlatformSupportDistanceFields(Context))
+	{
+		return false;
+	}
+
+	const FPreparedType& PositionType = Context.PrepareExpression(PositionExpression, Scope, Shader::EValueType::Double3);
+	if (PositionType.IsVoid())
+	{
+		return false;
+	}
+
+	const FPreparedType& NormalType = Context.PrepareExpression(NormalExpression, Scope, Shader::EValueType::Float3);
+	if (NormalType.IsVoid())
+	{
+		return false;
+	}
+
+	const FPreparedType& StepDistanceType = Context.PrepareExpression(StepDistanceExpression, Scope, Shader::EValueType::Float1);
+	if (StepDistanceType.IsVoid())
+	{
+		return false;
+	}
+
+	const FPreparedType& DistanceBiasType = Context.PrepareExpression(DistanceBiasExpression, Scope, Shader::EValueType::Float1);
+	if (DistanceBiasType.IsVoid())
+	{
+		return false;
+	}
+
+	const FPreparedType& MaxDistanceType = Context.PrepareExpression(MaxDistanceExpression, Scope, Shader::EValueType::Float1);
+	if (MaxDistanceType.IsVoid())
+	{
+		return false;
+	}
+
+	if (Context.bMarkLiveValues && Context.MaterialCompilationOutput)
+	{
+		Context.MaterialCompilationOutput->bUsesGlobalDistanceField = true;
+	}
+
+	return OutResult.SetType(Context, RequestedType, EExpressionEvaluation::Shader, Shader::EValueType::Float1);
+}
+
+void FExpressionDistanceFieldApproxAO::EmitValueShader(FEmitContext& Context, FEmitScope& Scope, const FRequestedType& RequestedType, FEmitValueShaderResult& OutResult) const
+{
+	FEmitShaderExpression* EmitPosition = PositionExpression->GetValueShader(Context, Scope, Shader::EValueType::Double3);
+	FEmitShaderExpression* EmitNormal = NormalExpression->GetValueShader(Context, Scope, Shader::EValueType::Float3);
+	FEmitShaderExpression* EmitStepDistance = StepDistanceExpression->GetValueShader(Context, Scope, Shader::EValueType::Float1);
+	FEmitShaderExpression* EmitDistanceBias = DistanceBiasExpression->GetValueShader(Context, Scope, Shader::EValueType::Float1);
+	FEmitShaderExpression* EmitMaxDistance = MaxDistanceExpression->GetValueShader(Context, Scope, Shader::EValueType::Float1);
+
+	OutResult.Code = Context.EmitInlineExpression(Scope, Shader::EValueType::Float1,
+		TEXT("CalculateDistanceFieldApproxAO(%, %, %, %, %, %, %)"),
+		EmitPosition,
+		EmitNormal,
+		NumSteps,
+		EmitStepDistance,
+		StepScale,
+		EmitDistanceBias,
+		EmitMaxDistance);
+}
+
 int32 FEmitData::FindInterpolatorIndex(const FExpression* Expression) const
 {
 	for (int32 Index = 0; Index < VertexInterpolators.Num(); ++Index)
diff --git a/Engine/Source/Runtime/Engine/Public/MaterialHLSLTree.h b/Engine/Source/Runtime/Engine/Public/MaterialHLSLTree.h
--- a/Engine/Source/Runtime/Engine/Public/MaterialHLSLTree.h
+++ b/Engine/Source/Runtime/Engine/Public/MaterialHLSLTree.h
@@ -607,6 +607,64 @@ public:
 	virtual void EmitValueShader(FEmitContext& Context, FEmitScope& Scope, const FRequestedType& RequestedType, FEmitValueShaderResult& OutResult) const override;
 };
 
+class FExpressionDistanceToNearestSurface : public FExpression
+{
+public:
+	const FExpression* PositionExpression;
+
+	FExpressionDistanceToNearestSurface(const FExpression* InPositionExpression)
+		: PositionExpression(InPositionExpression)
+	{}
+
+	virtual bool PrepareValue(FEmitContext& Context, FEmitScope& Scope, const FRequestedType& RequestedType, FPrepareValueResult& OutResult) const override;
+	virtual void EmitValueShader(FEmitContext& Context, FEmitScope& Scope, const FRequestedType& RequestedType, FEmitValueShaderResult& OutResult) const override;
+};
+
+class FExpressionDistanceFieldGradient : public FExpression
+{
+public:
+	const FExpression* PositionExpression;
+
+	FExpressionDistanceFieldGradient(const FExpression* InPositionExpression)
+		: PositionExpression(InPositionExpression)
+	{}
+
+	virtual bool PrepareValue(FEmitContext& Context, FEmitScope& Scope, const FRequestedType& RequestedType, FPrepareValueResult& OutResult) const override;
+	virtual void EmitValueShader(FEmitContext& Context, FEmitScope& Scope, const FRequestedType& RequestedType, FEmitValueShaderResult& OutResult) const override;
+};
+
+class FExpressionDistanceFieldApproxAO : public FExpression
+{
+public:
+	const FExpression* PositionExpression;
+	const FExpression* NormalExpression;
+	const FExpression* StepDistanceExpression;
+	const FExpression* DistanceBiasExpression;
+	const FExpression* MaxDistanceExpression;
+	int32 NumSteps;
+	float StepScale;
+
+	FExpressionDistanceFieldApproxAO(
+		const FExpression* InPositionExpression,
+		const FExpression* InNormalExpression,
+		const FExpression* InStepDistanceExpression,
+		const FExpression* InDistanceBiasExpression,
+		const FExpression* InMaxDistanceExpression,
+		int32 InNumSteps,
+		float InStepScale)
+		: PositionExpression(InPositionExpression)
+		, NormalExpression(InNormalExpression)
+		, StepDistanceExpression(InStepDistanceExpression)
+		, DistanceBiasExpression(InDistanceBiasExpression)
+		, MaxDistanceExpression(InMaxDistanceExpression)
+		, NumSteps(InNumSteps)
+		, StepScale(InStepScale)
+	{}
+
+	virtual bool PrepareValue(FEmitContext& Context, FEmitScope& Scope, const FRequestedType& RequestedType, FPrepareValueResult& OutResult) const override;
+	virtual void EmitValueShader(FEmitContext& Context, FEmitScope& Scope, const FRequestedType& RequestedType, FEmitValueShaderResult& OutResult) const override;
+};
+
 struct FVertexInterpolator
 {
 	FVertexInterpolator() = default;
