Fixed null dereference in UpdateDeferredCachedUniformExpressions.

diff --git a/Engine/Source/Runtime/Engine/Private/Materials/MaterialRenderProxy.cpp b/Engine/Source/Runtime/Engine/Private/Materials/MaterialRenderProxy.cpp
--- a/Engine/Source/Runtime/Engine/Private/Materials/MaterialRenderProxy.cpp
+++ b/Engine/Source/Runtime/Engine/Private/Materials/MaterialRenderProxy.cpp
@@ -733,7 +733,7 @@ void FMaterialRenderProxy::UpdateDeferredCachedUniformExpressions(FRHICommandLis
 		}
 	};
 
-	if (bAllowAsyncUpdate)
+	if (TaskIfAsync && bAllowAsyncUpdate)
 	{
 		*TaskIfAsync = UE::Tasks::Launch(UE_SOURCE_LOCATION, MoveTemp(EvaluateUniformExpressionsLambda));
 	}
