Iris - RepMovementNetSerializer will now store the quantization levels in dequantize so that they get applied when copying this data to the target instance. A better solution will be implemented in the future.

#rnx

diff --git a/Engine/Source/Runtime/Engine/Private/RepMovementNetSerializer.cpp b/Engine/Source/Runtime/Engine/Private/RepMovementNetSerializer.cpp
--- a/Engine/Source/Runtime/Engine/Private/RepMovementNetSerializer.cpp
+++ b/Engine/Source/Runtime/Engine/Private/RepMovementNetSerializer.cpp
@@ -563,6 +563,9 @@ void FRepMovementNetSerializer::Dequantize(FNetSerializationContext& Context, co
 	Target.bRepPhysics = (Source.Flags & Flag_RepPhysics) ? 1 : 0;
 	Target.ServerFrame = (Source.Flags & Flag_ServerFrameIsPresent ? Source.ServerFrame : 0);
 	Target.ServerPhysicsHandle = (Source.Flags & Flag_ServerPhysicsHandleIsPresent ? Source.ServerPhysicsHandle : INDEX_NONE);
+	Target.LocationQuantizationLevel = EVectorQuantization(Source.LocationQuantizationLevel);
+	Target.VelocityQuantizationLevel = EVectorQuantization(Source.VelocityQuantizationLevel);
+	Target.RotationQuantizationLevel = ERotatorQuantization(Source.RotationQuantizationLevel);
 
 	// We do not overwrite AngularVelocity unless we're replicating it. This is consistent with the FRepMovement serialization method.
 	if (Source.Flags & Flag_RepPhysics)
