Added conditional logging in UWorldPartitionLevelStreamingPolicy::ConvertEditorPathToRuntimePath to better understand an unexpected SoftObjectPath containing a different PIEInstanceID.

diff --git a/Engine/Source/Runtime/Engine/Private/WorldPartition/WorldPartitionLevelStreamingPolicy.cpp b/Engine/Source/Runtime/Engine/Private/WorldPartition/WorldPartitionLevelStreamingPolicy.cpp
--- a/Engine/Source/Runtime/Engine/Private/WorldPartition/WorldPartitionLevelStreamingPolicy.cpp
+++ b/Engine/Source/Runtime/Engine/Private/WorldPartition/WorldPartitionLevelStreamingPolicy.cpp
@@ -196,7 +196,7 @@ bool UWorldPartitionLevelStreamingPolicy::ConvertEditorPathToRuntimePath(const F
 	check(PathPIEInstanceID == PIEInstanceID);
 
 	FString SrcPath = UWorld::RemovePIEPrefix(InPath.ToString(), &PathPIEInstanceID);
-	check(PathPIEInstanceID == INDEX_NONE || PathPIEInstanceID == PIEInstanceID);
+	checkf(PathPIEInstanceID == INDEX_NONE || PathPIEInstanceID == PIEInstanceID, TEXT("Unexpected PIEInstanceID %d while converting editor to runtime path %s for world %s with PIEInstanceID %d "), PathPIEInstanceID, *InPath.ToString(), *OuterWorld->GetFullName(), PIEInstanceID);
 	const FSoftObjectPath SrcObjectPath(SrcPath);
 #else
 	const FSoftObjectPath SrcObjectPath(InPath);
