Fix incorrect implementation of cross product and ConstantBiasScale in the new material HLSL translator

#tests launch editor with/without and new translator

diff --git a/Engine/Source/Runtime/Engine/Private/HLSLTree/HLSLTreeOperation.cpp b/Engine/Source/Runtime/Engine/Private/HLSLTree/HLSLTreeOperation.cpp
--- a/Engine/Source/Runtime/Engine/Private/HLSLTree/HLSLTreeOperation.cpp
+++ b/Engine/Source/Runtime/Engine/Private/HLSLTree/HLSLTreeOperation.cpp
@@ -127,11 +127,11 @@ const FExpression* FTree::NewCross(const FExpression* Lhs, const FExpression* Rh
 	//c_P[0] = v_A[1] * v_B[2] - v_A[2] * v_B[1];
 	//c_P[1] = -(v_A[0] * v_B[2] - v_A[2] * v_B[0]);
 	//c_P[2] = v_A[0] * v_B[1] - v_A[1] * v_B[0];
-	const FExpression* Lhs0 = NewSwizzle(FSwizzleParameters(1, 0, 0), Lhs);
-	const FExpression* Lhs1 = NewSwizzle(FSwizzleParameters(2, 2, 1), Lhs);
-	const FExpression* Rhs0 = NewSwizzle(FSwizzleParameters(2, 2, 1), Rhs);
-	const FExpression* Rhs1 = NewSwizzle(FSwizzleParameters(1, 0, 0), Rhs);
-	return NewSub(NewMul(NewMul(Lhs0, Rhs0), NewConstant(FVector3f(1.0f, -1.0f, 1.0f))), NewMul(Lhs1, Rhs1));
+	const FExpression* Lhs0 = NewSwizzle(FSwizzleParameters(1, 2, 0), Lhs);
+	const FExpression* Lhs1 = NewSwizzle(FSwizzleParameters(2, 0, 1), Lhs);
+	const FExpression* Rhs0 = NewSwizzle(FSwizzleParameters(2, 0, 1), Rhs);
+	const FExpression* Rhs1 = NewSwizzle(FSwizzleParameters(1, 2, 0), Rhs);
+	return NewSub(NewMul(Lhs0, Rhs0), NewMul(Lhs1, Rhs1));
 }
 
 FExpressionOperation::FExpressionOperation(EOperation InOp, TConstArrayView<const FExpression*> InInputs) : Op(InOp)
diff --git a/Engine/Source/Runtime/Engine/Private/Materials/MaterialExpressionHLSL.cpp b/Engine/Source/Runtime/Engine/Private/Materials/MaterialExpressionHLSL.cpp
--- a/Engine/Source/Runtime/Engine/Private/Materials/MaterialExpressionHLSL.cpp
+++ b/Engine/Source/Runtime/Engine/Private/Materials/MaterialExpressionHLSL.cpp
@@ -266,7 +266,7 @@ bool UMaterialExpressionConstantBiasScale::GenerateHLSLExpression(FMaterialHLSLG
 	{
 		return false;
 	}
-	OutExpression = Generator.GetTree().NewAdd(Generator.GetTree().NewMul(Generator.NewConstant(Scale), ExpressionInput), Generator.NewConstant(Bias));
+	OutExpression = Generator.GetTree().NewMul(Generator.GetTree().NewAdd(Generator.NewConstant(Bias), ExpressionInput), Generator.NewConstant(Scale));
 	return true;
 }
 
