Partial backout of 27805890 for UWorldPartitionRuntimeHash::CreateRuntimeCell as streaming-out and streaming-in again the same instanced WorldPartition in PIE will recreate the streaming cells with the exact same name (thus recycle the object if it exists).

diff --git a/Engine/Source/Runtime/Engine/Private/WorldPartition/WorldPartitionRuntimeHash.cpp b/Engine/Source/Runtime/Engine/Private/WorldPartition/WorldPartitionRuntimeHash.cpp
--- a/Engine/Source/Runtime/Engine/Private/WorldPartition/WorldPartitionRuntimeHash.cpp
+++ b/Engine/Source/Runtime/Engine/Private/WorldPartition/WorldPartitionRuntimeHash.cpp
@@ -96,11 +96,6 @@ UWorldPartitionRuntimeCell* UWorldPartitionRuntimeHash::CreateRuntimeCell(UClass
 	const FString CellObjectName = GetCellObjectName(CellName) + CellInstanceSuffix;
 	// Use given outer if provided, else use hash as outer
 	UObject* Outer = InOuter ? InOuter : this;
-	if (FindObject<UWorldPartitionRuntimeCell>(Outer, *CellObjectName))
-	{
-		UE_LOG(LogWorldPartition, Warning, TEXT("UWorldPartitionRuntimeHash::CreateRuntimeCell can't create an already existing UWorldPartitionRuntimeCell object named %s"), *CellObjectName);
-		return nullptr;
-	}
 	UWorldPartitionRuntimeCell* RuntimeCell = NewObject<UWorldPartitionRuntimeCell>(Outer, CellClass, *CellObjectName);
 	RuntimeCell->RuntimeCellData = NewObject<UWorldPartitionRuntimeCellData>(RuntimeCell, CellDataClass);
 	return RuntimeCell;
