Fixed race condition between DrawTileMesh and volumetric fog volumes added to scene.

diff --git a/Engine/Source/Runtime/Engine/Private/UserInterface/Canvas.cpp b/Engine/Source/Runtime/Engine/Private/UserInterface/Canvas.cpp
--- a/Engine/Source/Runtime/Engine/Private/UserInterface/Canvas.cpp
+++ b/Engine/Source/Runtime/Engine/Private/UserInterface/Canvas.cpp
@@ -29,6 +29,7 @@
 #include "CanvasRendererItem.h"
 #include "RenderGraphUtils.h"
 #include "TextureResource.h"
+#include "Rendering/RenderCommandPipes.h"
 
 #include UE_INLINE_GENERATED_CPP_BY_NAME(Canvas)
 
@@ -559,6 +560,9 @@ FCanvasRenderThreadScope::~FCanvasRenderThreadScope()
 {
 	RenderCommandFunctionArray* RenderCommandArray = RenderCommands;
 
+	// DrawTileMesh can reference scene data that is mutated by the scene pipe.
+	UE::RenderCommandPipe::FSyncScope SyncScope({ &UE::RenderCommandPipe::Scene });
+
 	ENQUEUE_RENDER_COMMAND(DispatchCanvasRenderCommands)(
 		[RenderCommandArray, RenderTarget = Canvas.RenderTarget, ViewRect = Canvas.ViewRect, ScissorRect = Canvas.ScissorRect, bScaledToRenderTarget = Canvas.bScaledToRenderTarget](FRHICommandListImmediate& RHICmdList) mutable
 	{
