Fixed race condition in material and material instance BeginDestroy. Release references before assigning atomic state that allows destruction.

diff --git a/Engine/Source/Runtime/Engine/Private/Materials/Material.cpp b/Engine/Source/Runtime/Engine/Private/Materials/Material.cpp
--- a/Engine/Source/Runtime/Engine/Private/Materials/Material.cpp
+++ b/Engine/Source/Runtime/Engine/Private/Materials/Material.cpp
@@ -5224,6 +5224,9 @@ void UMaterial::BeginDestroy()
 				Resource->PrepareDestroy_RenderThread();
 			}
 
+			// Release the references before assigning the bool below.
+			ResourcesToDestroy.Empty();
+
 			*Released = true;
 		});
 	}
diff --git a/Engine/Source/Runtime/Engine/Private/Materials/MaterialInstance.cpp b/Engine/Source/Runtime/Engine/Private/Materials/MaterialInstance.cpp
--- a/Engine/Source/Runtime/Engine/Private/Materials/MaterialInstance.cpp
+++ b/Engine/Source/Runtime/Engine/Private/Materials/MaterialInstance.cpp
@@ -3261,7 +3261,7 @@ void UMaterialInstance::BeginDestroy()
 		FMaterialRenderProxy* LocalResource = Resource;
 		std::atomic<uint32>* Used = &UsedByRT;
 		ENQUEUE_RENDER_COMMAND(BeginDestroyCommand)(
-		[ResourcesToDestroy = MoveTemp(ResourcesToDestroy), LocalResource, Used](FRHICommandListImmediate& RHICmdList)
+		[ResourcesToDestroy = MoveTemp(ResourcesToDestroy), LocalResource, Used](FRHICommandListImmediate& RHICmdList) mutable
 		{
 			if (LocalResource)
 			{
@@ -3274,6 +3274,9 @@ void UMaterialInstance::BeginDestroy()
 				CurrentResource->PrepareDestroy_RenderThread();
 			}
 
+			// Clear all references before assigning the atomic state below.
+			ResourcesToDestroy.Empty();
+
 			// Clear flag set when Resource was created
 			*Used &= ~(uint32)EMaterialInstanceUsedByRTFlag::ResourceCreate;
 
