Fixed ISM render bounds calculations

diff --git a/Engine/Source/Runtime/Engine/Private/InstancedStaticMesh.cpp b/Engine/Source/Runtime/Engine/Private/InstancedStaticMesh.cpp
--- a/Engine/Source/Runtime/Engine/Private/InstancedStaticMesh.cpp
+++ b/Engine/Source/Runtime/Engine/Private/InstancedStaticMesh.cpp
@@ -3151,24 +3151,25 @@ FBoxSphereBounds UInstancedStaticMeshComponent::CalcBounds(const FTransform& Bou
 	return CalcBoundsImpl(BoundTransform, /*bForNavigation*/false);
 }
 
-FBoxSphereBounds UInstancedStaticMeshComponent::CalcBoundsImpl(const FTransform& BoundTransform, bool bForNavigation) const
+FBoxSphereBounds UInstancedStaticMeshComponent::CalcBoundsImpl(const FTransform& BoundTransform, const bool bForNavigation) const
 {
-	const FBox InstanceBounds = GetInstanceNavigationBounds();
-	if (InstanceBounds.IsValid && PerInstanceSMData.Num() > 0)
+	if (GetStaticMesh() && PerInstanceSMData.Num() > 0)
 	{
-		const FMatrix BoundTransformMatrix = BoundTransform.ToMatrixWithScale();
-		FBoxSphereBounds::Builder BoundsBuilder;
-		for (int32 InstanceIndex = 0; InstanceIndex < PerInstanceSMData.Num(); InstanceIndex++)
+		const FBox InstanceBounds = bForNavigation ? GetInstanceNavigationBounds() : GetStaticMesh()->GetBounds().GetBox();
+		if (InstanceBounds.IsValid)
 		{
-			BoundsBuilder += InstanceBounds.TransformBy(PerInstanceSMData[InstanceIndex].Transform * BoundTransformMatrix);
-		}
+			const FMatrix BoundTransformMatrix = BoundTransform.ToMatrixWithScale();
+			FBoxSphereBounds::Builder BoundsBuilder;
+			for (int32 InstanceIndex = 0; InstanceIndex < PerInstanceSMData.Num(); InstanceIndex++)
+			{
+				BoundsBuilder += InstanceBounds.TransformBy(PerInstanceSMData[InstanceIndex].Transform * BoundTransformMatrix);
+			}
 
-		return BoundsBuilder;
-	}
-	else
-	{
-		return FBoxSphereBounds(BoundTransform.GetLocation(), FVector::ZeroVector, 0.f);
+			return BoundsBuilder;
+		}
 	}
+	
+	return FBoxSphereBounds(BoundTransform.GetLocation(), FVector::ZeroVector, 0.f);
 }
 
 #if WITH_EDITOR
