Fix static code analysis error

diff --git a/Engine/Source/Runtime/Engine/Classes/Camera/CameraShakeBase.h b/Engine/Source/Runtime/Engine/Classes/Camera/CameraShakeBase.h
--- a/Engine/Source/Runtime/Engine/Classes/Camera/CameraShakeBase.h
+++ b/Engine/Source/Runtime/Engine/Classes/Camera/CameraShakeBase.h
@@ -33,7 +33,7 @@ struct FCameraShakePatternStartParams
 
 	/** An optional override for the camera shake's duration */
 	UPROPERTY(EditAnywhere, BlueprintReadWrite, Category=CameraShake)
-	float DurationOverride;
+	float DurationOverride = 0.f;
 };
 
 /**
