add initialization of FCollisionQueryParams::bReplaceHitWithSubComponents

diff --git a/Engine/Source/Runtime/Engine/Public/CollisionQueryParams.h b/Engine/Source/Runtime/Engine/Public/CollisionQueryParams.h
--- a/Engine/Source/Runtime/Engine/Public/CollisionQueryParams.h
+++ b/Engine/Source/Runtime/Engine/Public/CollisionQueryParams.h
@@ -161,6 +161,7 @@ public:
 		bDebugQuery = false;
 #endif
 		bTraceIntoSubComponents = true;
+		bReplaceHitWithSubComponents = true;
 	}
 
 	FCollisionQueryParams(FName InTraceTag, bool bInTraceComplex=false, const AActor* InIgnoreActor=NULL)
