Fix for Bink and Opus audio skipping in Sequencer.

diff --git a/Engine/Source/Runtime/BinkAudioDecoder/Module/Private/BinkAudioInfo.cpp b/Engine/Source/Runtime/BinkAudioDecoder/Module/Private/BinkAudioInfo.cpp
--- a/Engine/Source/Runtime/BinkAudioDecoder/Module/Private/BinkAudioInfo.cpp
+++ b/Engine/Source/Runtime/BinkAudioDecoder/Module/Private/BinkAudioInfo.cpp
@@ -121,22 +121,25 @@ void FBinkAudioInfo::SeekToTime(const float SeekTimeSeconds)
 		return;
 	}
 
+	uint32 SeekTimeFrames = 0;
 	uint32 SeekTimeSamples = 0;
 	if (SeekTimeSeconds > 0)
 	{
-		SeekTimeSamples = (uint32)(SeekTimeSeconds * SampleRate);
+		SeekTimeFrames = (uint32)(SeekTimeSeconds * SampleRate);
+		SeekTimeSamples = SeekTimeFrames * NumChannels;
 	}
 
 	uint32 SamplesInFrame = GetMaxFrameSizeSamples();
 	if (SeekTimeSamples > this->TrueSampleCount)
 	{
 		SeekTimeSamples = this->TrueSampleCount - 1;
+		SeekTimeFrames = (this->TrueSampleCount / NumChannels) - 1;
 	}
 	this->CurrentSampleCount = SeekTimeSamples;
 
 	uint32 SamplesPerBlock = SamplesInFrame * Decoder->FramesPerSeekTableEntry;
-	uint32 SeekTableIndex = SeekTimeSamples / SamplesPerBlock;
-	uint32 SeekTableOffset = SeekTimeSamples % SamplesPerBlock;
+	uint32 SeekTableIndex = SeekTimeFrames / SamplesPerBlock;
+	uint32 SeekTableOffset = SeekTimeFrames % SamplesPerBlock;
 
 	uint32 OffsetToBlock = Decoder->SeekTable()[SeekTableIndex] + sizeof(BinkAudioFileHeader) + Decoder->SeekTableCount * sizeof(uint16);
 
@@ -222,7 +225,8 @@ bool FBinkAudioInfo::ParseHeader(const uint8* InSrcBufferData, uint32 InSrcBuffe
 	}
 
 	SampleRate = Header->rate;
-	TrueSampleCount = Header->sample_count;
+	// Bink sample_count is per-channel so we multiply by num channels here
+	TrueSampleCount = Header->sample_count * Header->channels;
 	NumChannels = Header->channels;
 	MaxCompSpaceNeeded = Header->max_comp_space_needed;
   
diff --git a/Engine/Source/Runtime/Engine/Private/AudioDecompress.cpp b/Engine/Source/Runtime/Engine/Private/AudioDecompress.cpp
--- a/Engine/Source/Runtime/Engine/Private/AudioDecompress.cpp
+++ b/Engine/Source/Runtime/Engine/Private/AudioDecompress.cpp
@@ -163,15 +163,16 @@ void IStreamedCompressedInfo::ExpandFile(uint8* DstBuffer, struct FSoundQualityI
 
 	while (RawPCMOffset < QualityInfo->SampleDataSize)
 	{
-		int32 DecodedSamples = DecompressToPCMBuffer( /*Unused*/ 0);
+		int32 DecodedFrames = DecompressToPCMBuffer( /*Unused*/ 0);
 
-		if (DecodedSamples < 0)
+		if (DecodedFrames < 0)
 		{
 			RawPCMOffset += ZeroBuffer(DstBuffer + RawPCMOffset, QualityInfo->SampleDataSize - RawPCMOffset);
 		}
 		else
 		{
-			LastPCMByteSize = IncrementCurrentSampleCount(DecodedSamples) * SampleStride;
+			IncrementCurrentSampleCount(DecodedFrames * NumChannels);
+			LastPCMByteSize = DecodedFrames * SampleStride;
 			RawPCMOffset += WriteFromDecodedPCM(DstBuffer + RawPCMOffset, QualityInfo->SampleDataSize - RawPCMOffset);
 		}
 	}
@@ -231,14 +232,16 @@ bool IStreamedCompressedInfo::StreamCompressedData(uint8* Destination, bool bLoo
 
 	SCOPE_CYCLE_COUNTER(STAT_AudioStreamedDecompressTime);
 	
-	UE_LOG(LogAudio, Log, TEXT("Streaming compressed data from SoundWave'%s' - Chunk=%d\tNumChunks=%d\tOffset=%d\tChunkSize=%d\tLooping=%s\tLastPCMOffset=%d\tContainsEOF=%s" ), 
+	UE_LOG(LogAudio, Log, TEXT("Streaming compressed data from SoundWave'%s' - Chunk=%d\tCurrentSampleCount=%d\tTrueSampleCount=%d\tNumChunks=%d\tOffset=%d\tChunkSize=%d\tLooping=%s\tLastPCMOffset=%d\tContainsEOF=%s" ), 
 		*StreamingSoundWave->GetFName().ToString(), 	
 		CurrentChunkIndex, 
+		CurrentSampleCount, 
+		TrueSampleCount, 
 		StreamingSoundWave->GetNumChunks(),
 		SrcBufferOffset, 
 		SrcBufferDataSize,
 		bLooping ? TEXT("YES") : TEXT("NO"),
-		LastPCMOffset,		
+		LastPCMOffset,	
 		bStoringEndOfFile ? TEXT("YES") : TEXT("NO")
 	);
 
@@ -401,9 +404,9 @@ bool IStreamedCompressedInfo::StreamCompressedData(uint8* Destination, bool bLoo
 	while (RawPCMOffset < BufferSize)
 	{
 		// Decompress the next compression frame of audio (many samples) into the PCM buffer
-		int32 DecodedSamples = DecompressToPCMBuffer(/*Unused*/ 0);
+		int32 DecodedFrames = DecompressToPCMBuffer(/*Unused*/ 0);
 
-		if (DecodedSamples < 0)
+		if (DecodedFrames < 0)
 		{
 			UE_LOG(LogAudioStreamCaching, Warning, TEXT("Zero pad buffer Chunk=%d, Wave=%s, Reason=Decoder returned negative samples."),
 				CurrentChunkIndex, *StreamingSoundWave->GetFName().ToString());
@@ -414,7 +417,8 @@ bool IStreamedCompressedInfo::StreamCompressedData(uint8* Destination, bool bLoo
 		}
 		else
 		{
-			LastPCMByteSize = IncrementCurrentSampleCount(DecodedSamples) * SampleStride;
+			IncrementCurrentSampleCount(DecodedFrames * NumChannels);
+			LastPCMByteSize = DecodedFrames * SampleStride;
 
 			// update OutNumBytesStreamed as we write out data
 			RawPCMOffset += WriteFromDecodedPCM(Destination + RawPCMOffset, BufferSize - RawPCMOffset);
diff --git a/Engine/Source/Runtime/Engine/Private/OpusAudioInfo.cpp b/Engine/Source/Runtime/Engine/Private/OpusAudioInfo.cpp
--- a/Engine/Source/Runtime/Engine/Private/OpusAudioInfo.cpp
+++ b/Engine/Source/Runtime/Engine/Private/OpusAudioInfo.cpp
@@ -144,8 +144,8 @@ bool FOpusAudioInfo::ParseHeader(const uint8* InSrcBufferData, uint32 InSrcBuffe
 
 	// Store the offset to where the audio data begins
 	AudioDataOffset = SrcBufferOffset;
-	// Set members from header
-	TrueSampleCount = Header.ActiveSampleCount;
+	// Sample counts in the Opus API are always per-channel so we multiply by NumChannels here
+	TrueSampleCount = (uint32)Header.ActiveSampleCount * Header.NumChannels;
 	NumChannels = Header.NumChannels;
 
 	// Write out the the header info
@@ -333,7 +333,7 @@ void FOpusAudioInfo::SeekToTime(const float InSeekTime)
 		uint32 SeekSampleNum = 0;
 		if (InSeekTime > 0.0f)
 		{
-			SeekSampleNum = (uint32)(InSeekTime * Header.SampleRate);
+			SeekSampleNum = (uint32)(InSeekTime * Header.SampleRate * NumChannels);
 		}
 
 		const uint8* ChunkPtr = SrcBufferData + AudioDataOffset;
@@ -344,7 +344,7 @@ void FOpusAudioInfo::SeekToTime(const float InSeekTime)
 			uint32 ChunkSize = (uint32)ChunkPtr[0] + ((uint32)ChunkPtr[1] << 8);
 			int32 ExpectedFrames = opus_packet_get_nb_frames(ChunkPtr + 2, ChunkSize);
 			int32 ExpectedFrameSize = opus_packet_get_samples_per_frame(ChunkPtr + 2, Header.EncodedSampleRate);
-			int32 NumExpectedTotal = ExpectedFrames * ExpectedFrameSize;
+			int32 NumExpectedTotal = ExpectedFrames * ExpectedFrameSize * NumChannels;
 
 			if (CurrentChunkSampleNum >= SeekSampleNum && SeekSampleNum < CurrentChunkSampleNum + NumExpectedTotal)
 			{
diff --git a/Engine/Source/Runtime/Engine/Public/AudioDecompress.h b/Engine/Source/Runtime/Engine/Public/AudioDecompress.h
--- a/Engine/Source/Runtime/Engine/Public/AudioDecompress.h
+++ b/Engine/Source/Runtime/Engine/Public/AudioDecompress.h
@@ -241,7 +241,7 @@ protected:
 	* Decompresses a frame of data to PCM buffer
 	*
 	* @param FrameSize Size of the frame in bytes
-	* @return The amount of samples that were decompressed (< 0 indicates error)
+	* @return The number of audio frames that were decompressed (< 0 indicates error)
 	*/
 	int32 DecompressToPCMBuffer(uint16 FrameSize);
 
