Content Bundler Cooking: implement FWorldPartitionCookPackageContext::UnregisterPackageCookPackageGenerator and route it to FContentBundleContainer through WorldPartition::OnEndCook so content bundles properly unregister their cook generator when uninitializing.
Also added assert to track destruction of a content bundle without going through the normal tear down code paths.

diff --git a/Engine/Source/Runtime/Engine/Private/WorldPartition/ContentBundle/ContentBundleContainer.cpp b/Engine/Source/Runtime/Engine/Private/WorldPartition/ContentBundle/ContentBundleContainer.cpp
--- a/Engine/Source/Runtime/Engine/Private/WorldPartition/ContentBundle/ContentBundleContainer.cpp
+++ b/Engine/Source/Runtime/Engine/Private/WorldPartition/ContentBundle/ContentBundleContainer.cpp
@@ -57,6 +57,7 @@ void FContentBundleContainer::Initialize()
 		UWorldPartition* WorldPartition = GetInjectedWorld()->GetWorldPartition();
 		WorldPartition->OnPreGenerateStreaming.AddRaw(this, &FContentBundleContainer::OnPreGenerateStreaming);
 		WorldPartition->OnBeginCook.AddRaw(this, &FContentBundleContainer::OnBeginCook);
+		WorldPartition->OnEndCook.AddRaw(this, &FContentBundleContainer::OnEndCook);
 	}
 #endif
 }
@@ -75,6 +76,7 @@ void FContentBundleContainer::Deinitialize()
 			UWorldPartition* WorldPartition = GetInjectedWorld()->GetWorldPartition();
 			WorldPartition->OnPreGenerateStreaming.RemoveAll(this);
 			WorldPartition->OnBeginCook.RemoveAll(this);
+			WorldPartition->OnEndCook.RemoveAll(this);
 		}
 #endif
 	}
@@ -437,4 +439,11 @@ void FContentBundleContainer::OnBeginCook(IWorldPartitionCookPackageContext& Coo
 	}
 }
 
+void FContentBundleContainer::OnEndCook(IWorldPartitionCookPackageContext& CookContext)
+{
+	for (TSharedPtr<FContentBundleEditor>& ContentBundle : GetEditorContentBundles())
+	{
+		ContentBundle->OnEndCook(CookContext);
+	}
+}
 #endif
diff --git a/Engine/Source/Runtime/Engine/Private/WorldPartition/ContentBundle/ContentBundleEditor.cpp b/Engine/Source/Runtime/Engine/Private/WorldPartition/ContentBundle/ContentBundleEditor.cpp
--- a/Engine/Source/Runtime/Engine/Private/WorldPartition/ContentBundle/ContentBundleEditor.cpp
+++ b/Engine/Source/Runtime/Engine/Private/WorldPartition/ContentBundle/ContentBundleEditor.cpp
@@ -26,8 +26,14 @@ FContentBundleEditor::FContentBundleEditor(TSharedPtr<FContentBundleClient>& InC
 	, ExternalStreamingObject(nullptr)
 	, TreeItemID(FGuid::NewGuid())
 	, bIsBeingEdited(false)
+	, bIsInCook(false)
 {}
 
+FContentBundleEditor::~FContentBundleEditor()
+{
+	check(!bIsInCook);
+}
+
 void FContentBundleEditor::DoInitialize()
 {
 	SetStatus(EContentBundleStatus::Registered);
@@ -365,7 +371,16 @@ void FContentBundleEditor::GenerateStreaming(TArray<FString>* OutPackageToGenera
 
 void FContentBundleEditor::OnBeginCook(IWorldPartitionCookPackageContext& CookContext)
 {
+	check(!bIsInCook);
 	CookContext.RegisterPackageCookPackageGenerator(this);
+	bIsInCook = true;
+}
+
+void FContentBundleEditor::OnEndCook(IWorldPartitionCookPackageContext& CookContext)
+{
+	check(bIsInCook);
+	CookContext.UnregisterPackageCookPackageGenerator(this);
+	bIsInCook = false;
 }
 
 bool FContentBundleEditor::GatherPackagesToCook(class IWorldPartitionCookPackageContext& CookContext)
diff --git a/Engine/Source/Runtime/Engine/Private/WorldPartition/Cook/WorldPartitionCookPackageContext.cpp b/Engine/Source/Runtime/Engine/Private/WorldPartition/Cook/WorldPartitionCookPackageContext.cpp
--- a/Engine/Source/Runtime/Engine/Private/WorldPartition/Cook/WorldPartitionCookPackageContext.cpp
+++ b/Engine/Source/Runtime/Engine/Private/WorldPartition/Cook/WorldPartitionCookPackageContext.cpp
@@ -18,6 +18,11 @@ void FWorldPartitionCookPackageContext::RegisterPackageCookPackageGenerator(IWor
 	check(!PackagesToCookByGenerator.Contains(CookPackageGenerator));
 }
 
+void FWorldPartitionCookPackageContext::UnregisterPackageCookPackageGenerator(IWorldPartitionCookPackageGenerator* CookPackageGenerator)
+{
+	verify(CookPackageGenerators.Remove(CookPackageGenerator));
+}
+
 const FWorldPartitionCookPackage* FWorldPartitionCookPackageContext::AddLevelStreamingPackageToGenerate(IWorldPartitionCookPackageGenerator* CookPackageGenerator, const FString& Root, const FString& RelativePath)
 {
 	return AddPackageToGenerateInternal(CookPackageGenerator, Root, RelativePath, FWorldPartitionCookPackage::EType::Level);
diff --git a/Engine/Source/Runtime/Engine/Private/WorldPartition/Cook/WorldPartitionCookPackageContext.h b/Engine/Source/Runtime/Engine/Private/WorldPartition/Cook/WorldPartitionCookPackageContext.h
--- a/Engine/Source/Runtime/Engine/Private/WorldPartition/Cook/WorldPartitionCookPackageContext.h
+++ b/Engine/Source/Runtime/Engine/Private/WorldPartition/Cook/WorldPartitionCookPackageContext.h
@@ -17,7 +17,7 @@ public:
 
 	//~ Begin IWorldPartitionCookPackageContext Interface 
 	virtual void RegisterPackageCookPackageGenerator(IWorldPartitionCookPackageGenerator* CookPackageGenerator) override;
-	virtual void UnregisterPackageCookPackageGenerator(IWorldPartitionCookPackageGenerator* CookPackageGenerator) override { check(0); /*No use case*/ }
+	virtual void UnregisterPackageCookPackageGenerator(IWorldPartitionCookPackageGenerator* CookPackageGenerator) override;
 
 	virtual const FWorldPartitionCookPackage* AddLevelStreamingPackageToGenerate(IWorldPartitionCookPackageGenerator* CookPackageGenerator, const FString& Root, const FString& RelativePath) override;
 	virtual const FWorldPartitionCookPackage* AddGenericPackageToGenerate(IWorldPartitionCookPackageGenerator* CookPackageGenerator, const FString& Root, const FString& RelativePath) override;
diff --git a/Engine/Source/Runtime/Engine/Private/WorldPartition/WorldPartition.cpp b/Engine/Source/Runtime/Engine/Private/WorldPartition/WorldPartition.cpp
--- a/Engine/Source/Runtime/Engine/Private/WorldPartition/WorldPartition.cpp
+++ b/Engine/Source/Runtime/Engine/Private/WorldPartition/WorldPartition.cpp
@@ -1866,6 +1866,13 @@ void UWorldPartition::BeginCook(IWorldPartitionCookPackageContext& CookContext)
 	CookContext.RegisterPackageCookPackageGenerator(this);
 }
 
+void UWorldPartition::EndCook(IWorldPartitionCookPackageContext& CookContext)
+{
+	OnEndCook.Broadcast(CookContext);
+
+	CookContext.UnregisterPackageCookPackageGenerator(this);
+}
+
 bool UWorldPartition::GatherPackagesToCook(IWorldPartitionCookPackageContext& CookContext)
 {
 	FGenerateStreamingParams Params = FGenerateStreamingParams()
diff --git a/Engine/Source/Runtime/Engine/Private/WorldPartition/WorldPartitionCookPackageSplitter.cpp b/Engine/Source/Runtime/Engine/Private/WorldPartition/WorldPartitionCookPackageSplitter.cpp
--- a/Engine/Source/Runtime/Engine/Private/WorldPartition/WorldPartitionCookPackageSplitter.cpp
+++ b/Engine/Source/Runtime/Engine/Private/WorldPartition/WorldPartitionCookPackageSplitter.cpp
@@ -38,6 +38,7 @@ void FWorldPartitionCookPackageSplitter::Teardown(ETeardown Status)
 			UWorldPartition* WorldPartition = LocalWorld->PersistentLevel->GetWorldPartition();
 			if (WorldPartition)
 			{
+				WorldPartition->EndCook(CookContext);
 				WorldPartition->Uninitialize();
 			}
 		}
diff --git a/Engine/Source/Runtime/Engine/Public/WorldPartition/ContentBundle/ContentBundleContainer.h b/Engine/Source/Runtime/Engine/Public/WorldPartition/ContentBundle/ContentBundleContainer.h
--- a/Engine/Source/Runtime/Engine/Public/WorldPartition/ContentBundle/ContentBundleContainer.h
+++ b/Engine/Source/Runtime/Engine/Public/WorldPartition/ContentBundle/ContentBundleContainer.h
@@ -87,6 +87,7 @@ private:
 #if WITH_EDITOR
 	void OnPreGenerateStreaming(TArray<FString>* OutPackageToGenerate);
 	void OnBeginCook(IWorldPartitionCookPackageContext& CookContext);
+	void OnEndCook(IWorldPartitionCookPackageContext& CookContext);
 #endif
 
 	using ContentBundleGameArray = TArray<TUniquePtr<FContentBundle>>;
diff --git a/Engine/Source/Runtime/Engine/Public/WorldPartition/ContentBundle/ContentBundleEditor.h b/Engine/Source/Runtime/Engine/Public/WorldPartition/ContentBundle/ContentBundleEditor.h
--- a/Engine/Source/Runtime/Engine/Public/WorldPartition/ContentBundle/ContentBundleEditor.h
+++ b/Engine/Source/Runtime/Engine/Public/WorldPartition/ContentBundle/ContentBundleEditor.h
@@ -27,6 +27,7 @@ class FContentBundleEditor : public FContentBundleBase, IWorldPartitionCookPacka
 
 public:
 	ENGINE_API FContentBundleEditor(TSharedPtr<FContentBundleClient>& InClient, UWorld* InWorld);
+	ENGINE_API ~FContentBundleEditor();
 
 	//~ Begin IContentBundle Interface
 	ENGINE_API virtual void AddReferencedObjects(FReferenceCollector& Collector) override;
@@ -59,6 +60,7 @@ public:
 
 	// Cooking
 	ENGINE_API void OnBeginCook(IWorldPartitionCookPackageContext& CookContext);
+	ENGINE_API void OnEndCook(IWorldPartitionCookPackageContext& CookContext);
 	bool HasCookedContent() const { return ExternalStreamingObject != nullptr; }
 	//~Begin IWorldPartitionCookPackageGenerator
 	ENGINE_API virtual bool GatherPackagesToCook(class IWorldPartitionCookPackageContext& CookContext) override;
@@ -102,6 +104,7 @@ private:
 	FGuid TreeItemID;
 
 	bool bIsBeingEdited;
+	bool bIsInCook;
 };
 
 #endif
diff --git a/Engine/Source/Runtime/Engine/Public/WorldPartition/WorldPartition.h b/Engine/Source/Runtime/Engine/Public/WorldPartition/WorldPartition.h
--- a/Engine/Source/Runtime/Engine/Public/WorldPartition/WorldPartition.h
+++ b/Engine/Source/Runtime/Engine/Public/WorldPartition/WorldPartition.h
@@ -276,10 +276,14 @@ public:
 	ENGINE_API void RemapSoftObjectPath(FSoftObjectPath& ObjectPath) const;
 	ENGINE_API bool IsValidPackageName(const FString& InPackageName);
 
-	// Begin Cooking
+	// Cooking events
+	DECLARE_MULTICAST_DELEGATE_OneParam(FWorldPartitionCookEventDelegate, IWorldPartitionCookPackageContext&);
+
 	ENGINE_API void BeginCook(IWorldPartitionCookPackageContext& CookContext);
-	DECLARE_MULTICAST_DELEGATE_OneParam(FWorldPartitionBeginCookDelegate, IWorldPartitionCookPackageContext&);
-	FWorldPartitionBeginCookDelegate OnBeginCook;
+	ENGINE_API void EndCook(IWorldPartitionCookPackageContext& CookContext);
+
+	FWorldPartitionCookEventDelegate OnBeginCook;	
+	FWorldPartitionCookEventDelegate OnEndCook;
 
 	//~ Begin IWorldPartitionCookPackageGenerator Interface 
 	ENGINE_API virtual bool GatherPackagesToCook(IWorldPartitionCookPackageContext& CookContext) override;
