Iris
* Fix error when filling splitscreen views due to using wrong connection pointer.
* This caused all splitscreen views to have an identical view matrix for Iris

#rnx

diff --git a/Engine/Source/Runtime/Engine/Classes/Engine/NetConnection.h b/Engine/Source/Runtime/Engine/Classes/Engine/NetConnection.h
--- a/Engine/Source/Runtime/Engine/Classes/Engine/NetConnection.h
+++ b/Engine/Source/Runtime/Engine/Classes/Engine/NetConnection.h
@@ -706,6 +706,13 @@ public:
 		}
 	}
 
+	/** Returns the view target for this connection. Controlled by the player controller when one is assigned. The view target is the owning actor when no PC's are assigned or the PC's view target is invalid */
+	AActor* GetConnectionViewTarget() const
+	{
+		AActor* TempViewTarget = PlayerController ? PlayerController->GetViewTarget() : nullptr;
+		return (TempViewTarget && TempViewTarget->GetWorld()) ? TempViewTarget : ToRawPtr(OwningActor);
+	}
+
 private:
 	/** @todo document */
 	FActorChannelMap ActorChannels;
diff --git a/Engine/Source/Runtime/Engine/Classes/Engine/NetDriver.h b/Engine/Source/Runtime/Engine/Classes/Engine/NetDriver.h
--- a/Engine/Source/Runtime/Engine/Classes/Engine/NetDriver.h
+++ b/Engine/Source/Runtime/Engine/Classes/Engine/NetDriver.h
@@ -2125,7 +2125,7 @@ private:
 	void InitIrisSettings(FName NewDriverName);
 	void SetReplicationSystem(UReplicationSystem* ReplicationSystem);
 	void CreateReplicationSystem(bool bInitAsClient);
-	void UpdateReplicationViews() const;
+	void UpdateIrisReplicationViews() const;
 	void SendClientMoveAdjustments();
 #endif
 
diff --git a/Engine/Source/Runtime/Engine/Private/NetDriver.cpp b/Engine/Source/Runtime/Engine/Private/NetDriver.cpp
--- a/Engine/Source/Runtime/Engine/Private/NetDriver.cpp
+++ b/Engine/Source/Runtime/Engine/Private/NetDriver.cpp
@@ -946,7 +946,7 @@ void UNetDriver::TickFlush(float DeltaSeconds)
 #if UE_WITH_IRIS
 		if (ReplicationSystem)
 		{
-			UpdateReplicationViews();
+			UpdateIrisReplicationViews();
 			SendClientMoveAdjustments();
 			ReplicationSystem->PreSendUpdate(DeltaSeconds);
 		}
@@ -970,7 +970,7 @@ void UNetDriver::TickFlush(float DeltaSeconds)
 	{
 		if (ReplicationSystem)
 		{
-			UpdateReplicationViews();
+			UpdateIrisReplicationViews();
 			ReplicationSystem->PreSendUpdate(DeltaSeconds);
 		}
 	}
@@ -6380,58 +6380,38 @@ void UNetDriver::CreateReplicationSystem(bool bInitAsClient)
 	}
 }
 
-void UNetDriver::UpdateReplicationViews() const
+void UNetDriver::UpdateIrisReplicationViews() const
 {
 	using namespace UE::Net;
 
-	FReplicationView ReplicationView;
-
 	FReplicationView::FView DefaultView;
 	{
 		const APlayerCameraManager* CameraManager = GetDefault<APlayerCameraManager>();
 		DefaultView.FoVRadians = FMath::DegreesToRadians(CameraManager->DefaultFOV);
 	}
 
-	TObjectPtr<UNetConnection> const* Connections = (ServerConnection != nullptr ? &ServerConnection : ClientConnections.GetData());
-	const bool bUpdateConnectionViewTarget = ServerConnection == nullptr;
+	const bool bUpdateConnectionViewTarget = IsServer();
 
-	const uint32 ConnectionCount = (ServerConnection != nullptr ? 1U : uint32(ClientConnections.Num()));
-	TArray<UNetConnection*, TInlineAllocator<4>> SubConnections;
-	for (uint32 ConnIt = 0, ConnEndIt = ConnectionCount; ConnIt != ConnEndIt; ++ConnIt)
+	auto FillIrisReplicationViews = [&DefaultView, bUpdateConnectionViewTarget](TArrayView<UNetConnection*> AllConnections, FReplicationView& OutReplicationView)
 	{
-		ReplicationView.Views.Reset();
-
-		UNetConnection* Conn = Connections[ConnIt];
-		SubConnections.Reset();
-		SubConnections.Add(Conn);
-		if (Conn->Children.Num() > 0)
-		{
-			SubConnections.Insert(reinterpret_cast<UNetConnection*const*>(Conn->Children.GetData()), Conn->Children.Num(), 1);
-		}
-
-		// As we no longer call ServerReplicateActors we must make sure to update the ViewTarget
-		if (bUpdateConnectionViewTarget)
-		{
-			Conn->ViewTarget = Conn->PlayerController ? Conn->PlayerController->GetViewTarget() : ToRawPtr(Conn->OwningActor);
-		}
-
-		for (UNetConnection* SubConn : SubConnections)
+		for (UNetConnection* AnyConnection : AllConnections)
 		{
-			// See comment above
+			// As we no longer call ServerReplicateActors we must make sure to update the ViewTarget
 			if (bUpdateConnectionViewTarget)
 			{
-				SubConn->ViewTarget = SubConn->PlayerController ? SubConn->PlayerController->GetViewTarget() : ToRawPtr(SubConn->OwningActor);
+				AnyConnection->ViewTarget = AnyConnection->GetConnectionViewTarget();
 			}
 
-			const AActor* ViewTarget = Conn->ViewTarget;
-			const APlayerController* ViewingController = Conn->PlayerController;
+			const AActor* ViewTarget = AnyConnection->ViewTarget;
+			const APlayerController* ViewingController = AnyConnection->PlayerController;
 			if (ViewTarget == nullptr && ViewingController == nullptr)
 			{
 				continue;
 			}
 
-			FReplicationView::FView View;
+			FReplicationView::FView& View = OutReplicationView.Views.AddZeroed_GetRef();
 			View.FoVRadians = DefaultView.FoVRadians;
+
 			if (ViewTarget)
 			{
 				View.Pos = ViewTarget->GetActorLocation();
@@ -6448,11 +6428,44 @@ void UNetDriver::UpdateReplicationViews() const
 					View.FoVRadians = FMath::DegreesToRadians(CameraManager->GetFOVAngle());
 				}
 			}
+		}
+	};
+
+	if (IsServer())
+	{
+		FReplicationView ReplicationView;
+		TArray<UNetConnection*, TInlineAllocator<UE_IRIS_INLINE_VIEWS_PER_CONNECTION>> AllConnections;
 
-			ReplicationView.Views.Add(View);
+		for (UNetConnection* ClientConnection : ClientConnections)
+		{
+			AllConnections.Add(ClientConnection);
+			for (UNetConnection* Children : ClientConnection->Children)
+			{
+				AllConnections.Add(Children);
+			}
+
+			FillIrisReplicationViews(AllConnections, ReplicationView);
+
+			ReplicationSystem->SetReplicationView(ClientConnection->GetConnectionId(), ReplicationView);
+
+			ReplicationView.Views.Reset();
+			AllConnections.Reset();
 		}
+	}
+	else
+	{
+		FReplicationView ReplicationView;
+		TArray<UNetConnection*, TInlineAllocator<UE_IRIS_INLINE_VIEWS_PER_CONNECTION>> AllConnections;
+
+		AllConnections.Add(ServerConnection);
+		for (UNetConnection* Children : ServerConnection->Children)
+		{
+			AllConnections.Add(Children);
+		}
+
+		FillIrisReplicationViews(AllConnections, ReplicationView);
 
-		ReplicationSystem->SetReplicationView(Conn->GetConnectionId(), ReplicationView);
+		ReplicationSystem->SetReplicationView(ServerConnection->GetConnectionId(), ReplicationView);
 	}
 }
 
