Fix monolithic build

#rnx

diff --git a/Engine/Source/Runtime/Engine/Classes/Engine/NetConnection.h b/Engine/Source/Runtime/Engine/Classes/Engine/NetConnection.h
--- a/Engine/Source/Runtime/Engine/Classes/Engine/NetConnection.h
+++ b/Engine/Source/Runtime/Engine/Classes/Engine/NetConnection.h
@@ -635,6 +635,9 @@ public:
 	FNetTraceCollector* GetInTraceCollector() const;
 	FNetTraceCollector* GetOutTraceCollector() const;
 
+	/** Returns the view target for this connection. Controlled by the player controller when one is assigned. The view target is the owning actor when no PC's are assigned or the PC's view target is invalid */
+	AActor* GetConnectionViewTarget() const;
+
 	// ----------------------------------------------
 	// Actor Channel Accessors
 	// ----------------------------------------------
@@ -706,13 +709,6 @@ public:
 		}
 	}
 
-	/** Returns the view target for this connection. Controlled by the player controller when one is assigned. The view target is the owning actor when no PC's are assigned or the PC's view target is invalid */
-	AActor* GetConnectionViewTarget() const
-	{
-		AActor* TempViewTarget = PlayerController ? PlayerController->GetViewTarget() : nullptr;
-		return (TempViewTarget && TempViewTarget->GetWorld()) ? TempViewTarget : ToRawPtr(OwningActor);
-	}
-
 private:
 	/** @todo document */
 	FActorChannelMap ActorChannels;
diff --git a/Engine/Source/Runtime/Engine/Private/NetConnection.cpp b/Engine/Source/Runtime/Engine/Private/NetConnection.cpp
--- a/Engine/Source/Runtime/Engine/Private/NetConnection.cpp
+++ b/Engine/Source/Runtime/Engine/Private/NetConnection.cpp
@@ -795,6 +795,12 @@ int32 UNetConnection::GetLastNotifiedPacketId() const
 	return LastNotifiedPacketId;
 }
 
+AActor* UNetConnection::GetConnectionViewTarget() const;
+{
+	AActor* TempViewTarget = PlayerController ? PlayerController->GetViewTarget() : nullptr;
+	return (TempViewTarget && TempViewTarget->GetWorld()) ? TempViewTarget : ToRawPtr(OwningActor);
+}
+
 void UNetConnection::EnableEncryption(const FEncryptionData& EncryptionData)
 {
 	if (Handler.IsValid())
