Should save file properly in VS before submitting

#rnx

diff --git a/Engine/Source/Runtime/Engine/Private/NetConnection.cpp b/Engine/Source/Runtime/Engine/Private/NetConnection.cpp
--- a/Engine/Source/Runtime/Engine/Private/NetConnection.cpp
+++ b/Engine/Source/Runtime/Engine/Private/NetConnection.cpp
@@ -795,7 +795,7 @@ int32 UNetConnection::GetLastNotifiedPacketId() const
 	return LastNotifiedPacketId;
 }
 
-AActor* UNetConnection::GetConnectionViewTarget() const;
+AActor* UNetConnection::GetConnectionViewTarget() const
 {
 	AActor* TempViewTarget = PlayerController ? PlayerController->GetViewTarget() : nullptr;
 	return (TempViewTarget && TempViewTarget->GetWorld()) ? TempViewTarget : ToRawPtr(OwningActor);
