Set proper NetDriver crashcontext on clients by looking at NetDriverDefinition instead of NetDriverName

#rnx

diff --git a/Engine/Source/Runtime/Engine/Private/NetDriver.cpp b/Engine/Source/Runtime/Engine/Private/NetDriver.cpp
--- a/Engine/Source/Runtime/Engine/Private/NetDriver.cpp
+++ b/Engine/Source/Runtime/Engine/Private/NetDriver.cpp
@@ -1492,7 +1492,7 @@ bool UNetDriver::InitBase(bool bInitAsClient, FNetworkNotify* InNotify, const FU
 #endif // UE_WITH_IRIS
 
 
-	if (NetDriverName == NAME_GameNetDriver)
+	if (NetDriverDefinition == NAME_GameNetDriver)
 	{
 		static FString CrashContext_ReplicationDriver = TEXT("ReplicationDriver");
 
