Implement AudioProperties into USoundBase
#rnx

diff --git a/Engine/Source/Editor/DetailCustomizations/Private/SoundBaseDetails.cpp b/Engine/Source/Editor/DetailCustomizations/Private/SoundBaseDetails.cpp
--- a/Engine/Source/Editor/DetailCustomizations/Private/SoundBaseDetails.cpp
+++ b/Engine/Source/Editor/DetailCustomizations/Private/SoundBaseDetails.cpp
@@ -3,6 +3,7 @@
 #include "SoundBaseDetails.h"
 
 #include "DetailLayoutBuilder.h"
+#include "Modules/ModuleManager.h"
 #include "PropertyHandle.h"
 #include "Sound/AudioSettings.h"
 #include "Sound/SoundBase.h"
@@ -16,5 +17,8 @@ TSharedRef<IDetailCustomization> FSoundBaseDetails::MakeInstance()
 
 void FSoundBaseDetails::CustomizeDetails(IDetailLayoutBuilder& DetailBuilder)
 {
-	// Does nothing currently
+	if (!FModuleManager::Get().IsModuleLoaded("AudioProperties"))
+	{
+		DetailBuilder.HideCategory("AudioProperties");
+	}
 }
diff --git a/Engine/Source/Runtime/Engine/Classes/Sound/SoundBase.h b/Engine/Source/Runtime/Engine/Classes/Sound/SoundBase.h
--- a/Engine/Source/Runtime/Engine/Classes/Sound/SoundBase.h
+++ b/Engine/Source/Runtime/Engine/Classes/Sound/SoundBase.h
@@ -36,15 +36,17 @@ namespace Audio
 } // namespace Audio
 
 class ISoundGenerator;
+class UAssetUserData;
+class UAudioPropertiesSheetAssetBase;
+class UAudioPropertiesBindings;
 class USoundAttenuation;
 class USoundClass;
 class USoundEffectSourcePreset;
+class USoundEffectSourcePresetChain;
 class USoundSourceBus;
 class USoundSubmix;
 class USoundSubmixBase;
 class USoundWave;
-class USoundEffectSourcePresetChain;
-class UAssetUserData;
 enum class EBusSendType : uint8;
 namespace EMaxConcurrentResolutionRule { enum Type : int; }
 struct FActiveSound;
@@ -215,8 +217,14 @@ public:
 	/** Array of user data stored with the asset */
 	UPROPERTY(EditAnywhere, AdvancedDisplay, Instanced, Category = Advanced)
 	TArray<TObjectPtr<UAssetUserData>> AssetUserData;
-
 #if WITH_EDITORONLY_DATA	
+	UPROPERTY(EditAnywhere, Category = AudioProperties)
+	TObjectPtr<UAudioPropertiesSheetAssetBase> AudioPropertiesSheet;
+
+	UPROPERTY(EditAnywhere, Category = AudioProperties)
+	TObjectPtr<UAudioPropertiesBindings> AudioPropertiesBindings;
+
+
 private:
 	UPROPERTY()
 	FSoundTimecodeOffset TimecodeOffset;
@@ -231,6 +239,10 @@ public:
 	ENGINE_API virtual bool CanBeInCluster() const override;
 	ENGINE_API virtual void Serialize(FArchive& Ar) override;
 
+#if WITH_EDITOR
+	ENGINE_API virtual void PostEditChangeProperty(struct FPropertyChangedEvent& PropertyChangedEvent) override;
+#endif
+
 	//~ End UObject interface.
 
 	/** Returns whether the sound base is set up in a playable manner */
@@ -352,5 +364,7 @@ public:
 #if WITH_EDITORONLY_DATA
 	ENGINE_API void SetTimecodeOffset(const FSoundTimecodeOffset& InTimecodeOffset);
 	ENGINE_API TOptional<FSoundTimecodeOffset> GetTimecodeOffset() const;
+
+	void InjectPropertySheet();
 #endif //WITH_EDITORONLY_DATA
 };
diff --git a/Engine/Source/Runtime/Engine/Private/SoundBase.cpp b/Engine/Source/Runtime/Engine/Private/SoundBase.cpp
--- a/Engine/Source/Runtime/Engine/Private/SoundBase.cpp
+++ b/Engine/Source/Runtime/Engine/Private/SoundBase.cpp
@@ -2,6 +2,7 @@
 #include "Sound/SoundBase.h"
 
 #include "AudioDevice.h"
+#include "AudioPropertiesSheetAssetBase.h"
 #include "Engine/AssetUserData.h"
 #include "IAudioParameterTransmitter.h"
 #include "Sound/SoundSubmix.h"
@@ -188,7 +189,28 @@ bool USoundBase::GetSoundWavesWithCookedAnalysisData(TArray<USoundWave*>& OutSou
 	return false;
 }
 
+#if WITH_EDITOR
+void USoundBase::PostEditChangeProperty(struct FPropertyChangedEvent& PropertyChangedEvent)
+{
+	Super::PostEditChangeProperty(PropertyChangedEvent);
+
+	const FName AudioPropertiesSheetFName = GET_MEMBER_NAME_CHECKED(USoundBase, AudioPropertiesSheet);
+	const FName AudioPropertiesBindingsFName = GET_MEMBER_NAME_CHECKED(USoundBase, AudioPropertiesBindings);
+
+	if (FProperty* PropertyThatChanged = PropertyChangedEvent.Property)
+	{
+		const FName& Name = PropertyThatChanged->GetFName();
+		if (Name == AudioPropertiesSheetFName || Name == AudioPropertiesBindingsFName)
+		{
+			InjectPropertySheet();
+		}
+	}
+}
+#endif
+
 #if WITH_EDITORONLY_DATA
+
+
 void USoundBase::PostLoad()
 {
 	Super::PostLoad();
@@ -236,6 +258,8 @@ void USoundBase::Serialize(FArchive& Ar)
 			SoundConcurrencySettings_DEPRECATED = nullptr;
 		}
 	}
+
+	InjectPropertySheet();
 #endif // WITH_EDITORONLY_DATA
 }
 
@@ -326,4 +350,12 @@ TOptional<FSoundTimecodeOffset> USoundBase::GetTimecodeOffset() const
 	return TimecodeOffset;
 }
 
+void USoundBase::InjectPropertySheet()
+{
+	if (AudioPropertiesSheet && AudioPropertiesBindings)
+	{
+		AudioPropertiesSheet->CopyToObjectProperties(this, AudioPropertiesBindings);
+	}
+}
+
 #endif //WITH_EDITORONLY_DATA
