Temporary workaround for invalid dereference of static mesh after a GC.

diff --git a/Engine/Source/Runtime/Engine/Private/StaticMeshRender.cpp b/Engine/Source/Runtime/Engine/Private/StaticMeshRender.cpp
--- a/Engine/Source/Runtime/Engine/Private/StaticMeshRender.cpp
+++ b/Engine/Source/Runtime/Engine/Private/StaticMeshRender.cpp
@@ -442,6 +442,8 @@ FStaticMeshSceneProxy::FStaticMeshSceneProxy(const FStaticMeshSceneProxyDesc& In
 	{
 		UpdateVisibleInLumenScene();
 	}
+
+	bConstrainToRenderThread = true;
 }
 
 void FStaticMeshSceneProxy::SetEvaluateWorldPositionOffsetInRayTracing(FRHICommandListBase& RHICmdList, bool NewValue)
