PR #10880: Mark `WorldPartitionBlueprintLibrary` functions as `DevelopmentOnly`

diff --git a/Engine/Source/Runtime/Engine/Public/WorldPartition/WorldPartitionBlueprintLibrary.h b/Engine/Source/Runtime/Engine/Public/WorldPartition/WorldPartitionBlueprintLibrary.h
--- a/Engine/Source/Runtime/Engine/Public/WorldPartition/WorldPartitionBlueprintLibrary.h
+++ b/Engine/Source/Runtime/Engine/Public/WorldPartition/WorldPartitionBlueprintLibrary.h
@@ -101,57 +101,57 @@ public:
 	 * Gets the editor world bounds, which includes all actor descriptors.
 	 * @return The editor world bounds.
 	 */
-	UFUNCTION(BlueprintCallable, Category="World Partition")
+	UFUNCTION(BlueprintCallable, Category="World Partition", meta=(DevelopmentOnly))
 	static FBox GetEditorWorldBounds();
 
 	/**
 	 * Gets the runtime world bounds, which only includes actor descriptors that aren't editor only.
 	 * @return The runtime world bounds.
 	 */
-	UFUNCTION(BlueprintCallable, Category="World Partition")
+	UFUNCTION(BlueprintCallable, Category="World Partition", meta=(DevelopmentOnly))
 	static FBox GetRuntimeWorldBounds();
 
 	/**
 	 * Load actors
 	 */
-	UFUNCTION(BlueprintCallable, Category="World Partition")
+	UFUNCTION(BlueprintCallable, Category="World Partition", meta=(DevelopmentOnly))
 	static void LoadActors(const TArray<FGuid>& InActorsToLoad);
 
 	/**
 	 * Unload actors
 	 */
-	UFUNCTION(BlueprintCallable, Category="World Partition")
+	UFUNCTION(BlueprintCallable, Category="World Partition", meta=(DevelopmentOnly))
 	static void UnloadActors(const TArray<FGuid>& InActorsToUnload);
 
 	/**
 	 * Pin actors
 	 */
-	UFUNCTION(BlueprintCallable, Category="World Partition")
+	UFUNCTION(BlueprintCallable, Category="World Partition", meta=(DevelopmentOnly))
 	static void PinActors(const TArray<FGuid>& InActorsToPin);
 
 	/**
 	 * Unpin actors
 	 */
-	UFUNCTION(BlueprintCallable, Category="World Partition")
+	UFUNCTION(BlueprintCallable, Category="World Partition", meta=(DevelopmentOnly))
 	static void UnpinActors(const TArray<FGuid>& InActorsToUnpin);
 
 	/**
 	 * Gets all the actor descriptors into the provided array, recursing into actor containers.
 	 * @return True if the operation was successful.
 	 */
-	UFUNCTION(BlueprintCallable, Category="World Partition")
+	UFUNCTION(BlueprintCallable, Category="World Partition", meta=(DevelopmentOnly))
 	static bool GetActorDescs(TArray<FActorDesc>& OutActorDescs);
 
 	/**
 	 * Gets all the actor descriptors intersecting the provided box into the provided array, recursing into actor containers.
 	 * @return True if the operation was successful.
 	 */
-	UFUNCTION(BlueprintCallable, Category="World Partition")
+	UFUNCTION(BlueprintCallable, Category="World Partition", meta=(DevelopmentOnly))
 	static bool GetIntersectingActorDescs(const FBox& InBox, TArray<FActorDesc>& OutActorDescs);
 
 	/**
 	 * Returns the Data Layer Manager for this object.
 	 */
-	UFUNCTION(BlueprintCallable, Category = "World Partition", meta = (WorldContext = "WorldContextObject"))
+	UFUNCTION(BlueprintCallable, Category = "World Partition", meta = (DevelopmentOnly, WorldContext = "WorldContextObject"))
 	static UDataLayerManager* GetDataLayerManager(class UObject* WorldContextObject);
-};
\ No newline at end of file
+};
