Fixed crash in skylight component due to deleting the scene proxy on the wrong pipe.

diff --git a/Engine/Source/Runtime/Engine/Private/Components/SkyLightComponent.cpp b/Engine/Source/Runtime/Engine/Private/Components/SkyLightComponent.cpp
--- a/Engine/Source/Runtime/Engine/Private/Components/SkyLightComponent.cpp
+++ b/Engine/Source/Runtime/Engine/Private/Components/SkyLightComponent.cpp
@@ -25,6 +25,7 @@
 #include "UObject/ReleaseObjectVersion.h"
 #include "UObject/UE5MainStreamObjectVersion.h"
 #include "Engine/TextureCube.h"
+#include "Rendering/RenderCommandPipes.h"
 
 #if RHI_RAYTRACING
 #endif
@@ -503,8 +504,8 @@ void USkyLightComponent::DestroyRenderState_Concurrent()
 		GetWorld()->Scene->DisableSkyLight(SceneProxy);
 
 		FSkyLightSceneProxy* LightSceneProxy = SceneProxy;
-		ENQUEUE_RENDER_COMMAND(FDestroySkyLightCommand)(
-			[LightSceneProxy](FRHICommandList& RHICmdList)
+		ENQUEUE_RENDER_COMMAND(FDestroySkyLightCommand)(&UE::RenderCommandPipe::Scene,
+			[LightSceneProxy]
 			{
 				delete LightSceneProxy;
 			});
