Stop the "Waiting on {AssetType} preparation" slow task window opening during PIE and stealing focus.

diff --git a/Engine/Source/Runtime/Engine/Private/AsyncCompilationHelpers.cpp b/Engine/Source/Runtime/Engine/Private/AsyncCompilationHelpers.cpp
--- a/Engine/Source/Runtime/Engine/Private/AsyncCompilationHelpers.cpp
+++ b/Engine/Source/Runtime/Engine/Private/AsyncCompilationHelpers.cpp
@@ -114,7 +114,7 @@ namespace AsyncCompilationHelpers
 			Args.Add(TEXT("AssetType"), AssetType.ToLower());
 
 			SlowTask.Emplace((float)Num, FText::Format(LOCTEXT("WaitingOnFinishCompilation", "Waiting on {AssetType} preparation"), Args), true);
-			SlowTask->MakeDialogDelayed(1.0f, false /*bShowCancelButton*/, true /*bAllowInPIE*/);
+			SlowTask->MakeDialogDelayed(1.0f, false /*bShowCancelButton*/, false /*bAllowInPIE*/);
 		}
 
 		// Reschedule everything to be executed at blocking priority, it bypasses the pause mechanism to ensure forward progress since we're waiting.
