Fix FTextRenderSceneProxy not releasing RayTracingGeometry when Ray Tracing is toggled off at runtime.

diff --git a/Engine/Source/Runtime/Engine/Private/Components/TextRenderComponent.cpp b/Engine/Source/Runtime/Engine/Private/Components/TextRenderComponent.cpp
--- a/Engine/Source/Runtime/Engine/Private/Components/TextRenderComponent.cpp
+++ b/Engine/Source/Runtime/Engine/Private/Components/TextRenderComponent.cpp
@@ -760,11 +760,7 @@ void FTextRenderSceneProxy::ReleaseRenderThreadResources()
 	VertexFactory.ReleaseResource();
 
 #if RHI_RAYTRACING
-	if (RayTracingGeometry.IsValid())
-	{
-		check(IsRayTracingAllowed());
-		RayTracingGeometry.ReleaseResource();
-	}
+	RayTracingGeometry.ReleaseResource();
 #endif
 }
 
