Iris
 - Improved iris support for Net.DebugDraw to draw NetRefHandle and state for currently debugged NetRefHandle
 - Renamed InternalGetRootOBjectOfSubObject to GetRootOBjectOfSubObject and made it public

diff --git a/Engine/Source/Runtime/Engine/Private/Net/Iris/ReplicationSystem/ActorReplicationBridge.cpp b/Engine/Source/Runtime/Engine/Private/Net/Iris/ReplicationSystem/ActorReplicationBridge.cpp
--- a/Engine/Source/Runtime/Engine/Private/Net/Iris/ReplicationSystem/ActorReplicationBridge.cpp
+++ b/Engine/Source/Runtime/Engine/Private/Net/Iris/ReplicationSystem/ActorReplicationBridge.cpp
@@ -481,7 +481,7 @@ bool UActorReplicationBridge::WriteCreationHeader(UE::Net::FNetSerializationCont
 	}
 	else if (Object)
 	{
-		const UObject* RootObject = GetReplicatedObject(InternalGetRootObjectOfSubObject(Handle));
+		const UObject* RootObject = GetReplicatedObject(GetRootObjectOfSubObject(Handle));
 
 		// Get Header
 		FSubObjectCreationHeader Header;
@@ -767,7 +767,7 @@ void UActorReplicationBridge::EndInstantiateFromRemote(FNetRefHandle Handle)
 
 void UActorReplicationBridge::OnSubObjectCreatedFromReplication(FNetRefHandle SubObjectHandle)
 {
-	AActor* RootObject = Cast<AActor>(GetReplicatedObject(InternalGetRootObjectOfSubObject(SubObjectHandle)));
+	AActor* RootObject = Cast<AActor>(GetReplicatedObject(GetRootObjectOfSubObject(SubObjectHandle)));
 	UObject* SubObject = GetReplicatedObject(SubObjectHandle);
 	if (IsValid(RootObject) && IsValid(SubObject))
 	{
diff --git a/Engine/Source/Runtime/Engine/Private/NetDriver.cpp b/Engine/Source/Runtime/Engine/Private/NetDriver.cpp
--- a/Engine/Source/Runtime/Engine/Private/NetDriver.cpp
+++ b/Engine/Source/Runtime/Engine/Private/NetDriver.cpp
@@ -46,6 +46,7 @@
 #include "UObject/Stack.h"
 #if UE_WITH_IRIS
 #include "Iris/IrisConfig.h"
+#include "Iris/Core/IrisDebugging.h"
 #include "Iris/Core/IrisProfiler.h"
 #include "Iris/Core/IrisMemoryTracker.h"
 #include "Iris/ReplicationSystem/ReplicationSystem.h"
@@ -5630,6 +5631,11 @@ void UNetDriver::DrawNetDriverDebug()
 		return;
 	}
 
+	if (!IsInGameThread())
+	{
+		return;
+	}
+
 	ULocalPlayer*	LocalPlayer = NULL;
 	for(FLocalPlayerIterator It(GEngine, LocalWorld);It;++It)
 	{
@@ -5647,6 +5653,10 @@ void UNetDriver::DrawNetDriverDebug()
 	// Used to draw additional boxes to show more state
 	const FBox BoxExpansion = FBox(&Extent, 1);
 
+#if UE_WITH_IRIS
+		const UObjectReplicationBridge* Bridge = ReplicationSystem ? ReplicationSystem->GetReplicationBridgeAs<UObjectReplicationBridge>() : nullptr;
+#endif
+
 	for (FActorIterator It(LocalWorld); It; ++It)
 	{
 		// Skip drawing non-replicated actors and controllers since their bounds are
@@ -5692,7 +5702,7 @@ void UNetDriver::DrawNetDriverDebug()
 		const FNetworkObjectInfo* NetworkObjectInfo = Connection->Driver->FindNetworkObjectInfo( *It );
 		PRAGMA_ENABLE_DEPRECATION_WARNINGS
 
-		FColor DrawColor;
+		FColor DrawColor = FColor::White;
 		if ( NetworkObjectInfo && NetworkObjectInfo->DormantConnections.Contains( Connection ) )
 		{
 			DrawColor = FColor::Blue;
@@ -5707,7 +5717,35 @@ void UNetDriver::DrawNetDriverDebug()
 			DrawColor = FColor::Orange;
 		}
 
-		DrawDebugBox( LocalWorld, Box.GetCenter(), Box.GetExtent(), FQuat::Identity, DrawColor, false );
+		// Draw NetGUID or NetHandle
+		FVector AboveActor(0.f, 0.f, 32.f);
+#if UE_WITH_IRIS
+		if (Bridge)
+		{
+			AActor* Actor = *It;
+			const UE::Net::FNetRefHandle ActorRefHandle = Bridge->GetReplicatedRefHandle(Actor);
+			if (ActorRefHandle.IsValid() && UE::Net::IrisDebugHelper::FilterDebuggedObject(Actor))
+			{
+				DrawColor = FColor::Green;
+				DrawDebugString(Actor->GetWorld(), AboveActor, ActorRefHandle.ToString(), Actor, FColor::White, 0.f);
+			}
+			else
+			{
+				continue;
+			}
+		}
+		else
+#else
+		{
+			FNetworkGUID NetGUID = Connection->PackageMap->GetNetGUIDFromObject(*It);
+			if (NetGUID.IsValid())
+			{
+				DrawDebugString(LocalWorld, AboveActor, NetGUID.ToString(), *It, FColor::White, 0.f);
+			}
+		}
+#endif
+
+		DrawDebugBox(LocalWorld, Box.GetCenter(), Box.GetExtent(), FQuat::Identity, DrawColor, false);
 		const bool bDrawSecondBox = bWasCulled || bIsAlwaysRelevant;
 		// Draw a second box around the object if it was net culled
 		if (bDrawSecondBox)
@@ -5716,6 +5754,25 @@ void UNetDriver::DrawNetDriverDebug()
 			DrawDebugBox(LocalWorld, Box.GetCenter(), Box.GetExtent(), FQuat::Identity, ExtraStateDrawColor, false);
 		}
 	}
+
+	// Draw state
+#if UE_WITH_IRIS
+	if (Bridge)
+	{
+		UE::Net::FNetRefHandle DebugRefHandle = UE::Net::IrisDebugHelper::GetDebugNetRefHandle();
+		if (DebugRefHandle.IsValid() && Bridge)
+		{
+			UE::Net::FNetRefHandle RootObjectHandle = Bridge->GetRootObjectOfSubObject(DebugRefHandle);
+			UObject* Object = Bridge->GetReplicatedObject(RootObjectHandle.IsValid() ? RootObjectHandle : DebugRefHandle);
+			if (AActor* Actor = Cast<AActor>(Object))
+			{
+				DrawDebugString(LocalWorld, FVector(0.f, 0.f, 0.f), UE::Net::IrisDebugHelper::DebugNetObjectStateToString(DebugRefHandle.GetId(), ReplicationSystem->GetId()), Actor, FColor::White, 0.f);
+			}
+		}
+	}
+#endif
+
+
 #endif
 }
 
diff --git a/Engine/Source/Runtime/Experimental/Iris/Core/Private/Iris/ReplicationSystem/NetRefHandleManager.cpp b/Engine/Source/Runtime/Experimental/Iris/Core/Private/Iris/ReplicationSystem/NetRefHandleManager.cpp
--- a/Engine/Source/Runtime/Experimental/Iris/Core/Private/Iris/ReplicationSystem/NetRefHandleManager.cpp
+++ b/Engine/Source/Runtime/Experimental/Iris/Core/Private/Iris/ReplicationSystem/NetRefHandleManager.cpp
@@ -568,7 +568,7 @@ bool FNetRefHandleManager::SetSubObjectNetCondition(FInternalNetRefIndex SubObje
 	return false;
 }
 
-FNetRefHandle FNetRefHandleManager::GetSubObjectOwner(FNetRefHandle SubObjectRefHandle) const
+FNetRefHandle FNetRefHandleManager::GetRootObjectOfSubObject(FNetRefHandle SubObjectRefHandle) const
 {
 	const FInternalNetRefIndex SubObjectInternalIndex = GetInternalIndex(SubObjectRefHandle);
 	const FInternalNetRefIndex OwnerInternalIndex = SubObjectInternalIndex != InvalidInternalIndex ? ReplicatedObjectData[SubObjectInternalIndex].SubObjectRootIndex : InvalidInternalIndex;
diff --git a/Engine/Source/Runtime/Experimental/Iris/Core/Private/Iris/ReplicationSystem/NetRefHandleManager.h b/Engine/Source/Runtime/Experimental/Iris/Core/Private/Iris/ReplicationSystem/NetRefHandleManager.h
--- a/Engine/Source/Runtime/Experimental/Iris/Core/Private/Iris/ReplicationSystem/NetRefHandleManager.h
+++ b/Engine/Source/Runtime/Experimental/Iris/Core/Private/Iris/ReplicationSystem/NetRefHandleManager.h
@@ -212,7 +212,7 @@ public:
 	bool AddSubObject(FNetRefHandle OwnerHandle, FNetRefHandle SubObjectHandle, FNetRefHandle RelativeOtherSubObjectHandle, EAddSubObjectFlags Flags = EAddSubObjectFlags::Default);
 	bool AddSubObject(FNetRefHandle OwnerHandle, FNetRefHandle SubObjectHandle, EAddSubObjectFlags Flags = EAddSubObjectFlags::Default);
 	void RemoveSubObject(FNetRefHandle SubObjectHandle);
-	FNetRefHandle GetSubObjectOwner(FNetRefHandle SubObjectHandle) const;
+	FNetRefHandle GetRootObjectOfSubObject(FNetRefHandle SubObjectHandle) const;
 	bool SetSubObjectNetCondition(FInternalNetRefIndex SubObjectInternalIndex, FLifeTimeConditionStorage SubObjectCondition);
 
 	// DependentObjects
diff --git a/Engine/Source/Runtime/Experimental/Iris/Core/Private/Iris/ReplicationSystem/ObjectReplicationBridge.cpp b/Engine/Source/Runtime/Experimental/Iris/Core/Private/Iris/ReplicationSystem/ObjectReplicationBridge.cpp
--- a/Engine/Source/Runtime/Experimental/Iris/Core/Private/Iris/ReplicationSystem/ObjectReplicationBridge.cpp
+++ b/Engine/Source/Runtime/Experimental/Iris/Core/Private/Iris/ReplicationSystem/ObjectReplicationBridge.cpp
@@ -398,7 +398,7 @@ UE::Net::FNetRefHandle UObjectReplicationBridge::BeginReplication(FNetRefHandle
 	if (SubObjectRefHandle.IsValid())
 	{
 		// Verify that the existing object is a subobject of the owner
-		check(OwnerRefHandle == LocalNetRefHandleManager.GetSubObjectOwner(SubObjectRefHandle));
+		check(OwnerRefHandle == LocalNetRefHandleManager.GetRootObjectOfSubObject(SubObjectRefHandle));
 		return SubObjectRefHandle;
 	}
 	else
@@ -450,6 +450,17 @@ void UObjectReplicationBridge::SetSubObjectNetCondition(FNetRefHandle SubObjectR
 	}
 }
 
+UE::Net::FNetRefHandle UObjectReplicationBridge::GetRootObjectOfSubObject(FNetRefHandle SubObjectHandle) const
+{
+	using namespace UE::Net::Private;
+
+	FReplicationSystemInternal* ReplicationSystemInternal = GetReplicationSystem()->GetReplicationSystemInternal();
+	FNetRefHandleManager& LocalNetRefHandleManager = ReplicationSystemInternal->GetNetRefHandleManager();
+
+	return LocalNetRefHandleManager.GetRootObjectOfSubObject(SubObjectHandle);
+}
+
+
 void UObjectReplicationBridge::AddDependentObject(FNetRefHandle ParentHandle, FNetRefHandle DependentHandle, UE::Net::EDependentObjectSchedulingHint SchedulingHint)
 {
 	using namespace UE::Net::Private;
diff --git a/Engine/Source/Runtime/Experimental/Iris/Core/Private/Iris/ReplicationSystem/ReplicationBridge.cpp b/Engine/Source/Runtime/Experimental/Iris/Core/Private/Iris/ReplicationSystem/ReplicationBridge.cpp
--- a/Engine/Source/Runtime/Experimental/Iris/Core/Private/Iris/ReplicationSystem/ReplicationBridge.cpp
+++ b/Engine/Source/Runtime/Experimental/Iris/Core/Private/Iris/ReplicationSystem/ReplicationBridge.cpp
@@ -449,25 +449,6 @@ void UReplicationBridge::DestroyLocalNetHandle(FNetRefHandle Handle, EEndReplica
 	UE_LOG_REPLICATIONBRIDGE(Verbose, TEXT("DestroyLocalNetHandle Local %s"), *Handle.ToString());
 }
 
-UE::Net::FNetRefHandle UReplicationBridge::InternalGetRootObjectOfSubObject(FNetRefHandle SubObjectHandle) const
-{
-	using namespace UE::Net::Private;
-
-	const FInternalNetRefIndex InternalNetRefIndex = NetRefHandleManager->GetInternalIndex(SubObjectHandle);
-	if (InternalNetRefIndex != FNetRefHandleManager::InvalidInternalIndex)
-	{
-		const FNetRefHandleManager::FReplicatedObjectData& ObjectData = NetRefHandleManager->GetReplicatedObjectDataNoCheck(InternalNetRefIndex);
-		const FInternalNetRefIndex OwnerInternalNetRefIndex = ObjectData.SubObjectRootIndex;
-
-		if (OwnerInternalNetRefIndex != FNetRefHandleManager::InvalidInternalIndex)
-		{
-			return NetRefHandleManager->GetReplicatedObjectDataNoCheck(OwnerInternalNetRefIndex).RefHandle;
-		}
-	}
-
-	return FNetRefHandle();
-}
-
 void UReplicationBridge::InternalAddSubObject(FNetRefHandle OwnerHandle, FNetRefHandle SubObjectHandle, FNetRefHandle InsertRelativeToSubObjectHandle, ESubObjectInsertionOrder InsertionOrder)
 {
 	using namespace UE::Net::Private;
diff --git a/Engine/Source/Runtime/Experimental/Iris/Core/Public/Iris/ReplicationSystem/ObjectReplicationBridge.h b/Engine/Source/Runtime/Experimental/Iris/Core/Public/Iris/ReplicationSystem/ObjectReplicationBridge.h
--- a/Engine/Source/Runtime/Experimental/Iris/Core/Public/Iris/ReplicationSystem/ObjectReplicationBridge.h
+++ b/Engine/Source/Runtime/Experimental/Iris/Core/Public/Iris/ReplicationSystem/ObjectReplicationBridge.h
@@ -101,7 +101,7 @@ public:
 
 	/** Create handle and start replicating the Instance as a SubObject of the OwnerHandle. */
 	/** Begin replicating the Instance as a subobject of the OwnerHandle and return a valid NetRefHandle for the Instance if successful. */
-	FNetRefHandle BeginReplication(FNetRefHandle OwnerHandle, UObject* Instance, const FCreateNetRefHandleParams& Params = DefaultCreateNetRefHandleParams);
+	IRISCORE_API FNetRefHandle BeginReplication(FNetRefHandle OwnerHandle, UObject* Instance, const FCreateNetRefHandleParams& Params = DefaultCreateNetRefHandleParams);
 
 	/** 
 	 * Set NetCondition for a subobject, the condition is used to determine if the SubObject should replicate or not.
@@ -110,6 +110,9 @@ public:
 	 */
 	IRISCORE_API void SetSubObjectNetCondition(FNetRefHandle SubObjectHandle, ELifetimeCondition Condition);
 
+	/** Get the handle of the root object of any replicated subobject. */
+	IRISCORE_API FNetRefHandle GetRootObjectOfSubObject(FNetRefHandle SubObjectHandle) const;
+
 	/** Stop replicating the object. */
 	IRISCORE_API void EndReplication(UObject* Instance, EEndReplicationFlags EndReplicationFlags = EEndReplicationFlags::Destroy, FEndReplicationParameters* Parameters = nullptr);
 
diff --git a/Engine/Source/Runtime/Experimental/Iris/Core/Public/Iris/ReplicationSystem/ReplicationBridge.h b/Engine/Source/Runtime/Experimental/Iris/Core/Public/Iris/ReplicationSystem/ReplicationBridge.h
--- a/Engine/Source/Runtime/Experimental/Iris/Core/Public/Iris/ReplicationSystem/ReplicationBridge.h
+++ b/Engine/Source/Runtime/Experimental/Iris/Core/Public/Iris/ReplicationSystem/ReplicationBridge.h
@@ -230,9 +230,6 @@ protected:
 	/** Destroy the handle and all internal book keeping associated with it. */
 	IRISCORE_API void InternalDestroyNetObject(FNetRefHandle Handle);
 	
-	/** Get the handle of the root object of any given subobject. */
-	IRISCORE_API FNetRefHandle InternalGetRootObjectOfSubObject(FNetRefHandle SubObjectHandle) const;
-
 	/** Add SubObjectHandle as SubObject to OwnerHandle. */
 	IRISCORE_API void InternalAddSubObject(FNetRefHandle OwnerHandle, FNetRefHandle SubObjectHandle, FNetRefHandle InsertRelativeToSubObjectHandle, ESubObjectInsertionOrder InsertionOrder);
 
