Removed LightWeightInstance-related friend class declarations from UInstancedStaticMeshComponent. To make it possible a temporary bridge function has been. The function has been marked deprecated and the call sites wrapped in `PRAGMA_DISABLE_DEPRECATION_WARNINGS` blocks

diff --git a/Engine/Source/Runtime/Engine/Classes/Components/InstancedStaticMeshComponent.h b/Engine/Source/Runtime/Engine/Classes/Components/InstancedStaticMeshComponent.h
--- a/Engine/Source/Runtime/Engine/Classes/Components/InstancedStaticMeshComponent.h
+++ b/Engine/Source/Runtime/Engine/Classes/Components/InstancedStaticMeshComponent.h
@@ -88,9 +88,6 @@ class UInstancedStaticMeshComponent : public UStaticMeshComponent, public ISMIns
 {
 	GENERATED_UCLASS_BODY()
 
-	friend class ALightWeightInstanceStaticMeshManager;
-	friend class ALightWeightInstanceManager;
-	
 	/** Needs implementation in InstancedStaticMesh.cpp to compile UniquePtr for forward declared class */
 	ENGINE_API UInstancedStaticMeshComponent(FVTableHelper& Helper);
 	ENGINE_API virtual ~UInstancedStaticMeshComponent();
@@ -477,6 +474,9 @@ public:
 
 	ENGINE_API void FlushInstanceUpdateCommands(bool bFlushInstanceUpdateCmdBuffer);
 
+	UE_DEPRECATED(5.4, "This function has been added only for the purposes of moving LWI code outside of the engine. Don't use it, it will be removed soon.")
+	ENGINE_API void OnPostPopulatePerInstanceData() { OnPostLoadPerInstanceData(); }
+
 	TArray<int32> PerInstanceIds;
 
 	/** Used to cache a unique identifier for each instance.  These are provided
diff --git a/Engine/Source/Runtime/Engine/Private/LightWeightInstanceStaticMeshManager.cpp b/Engine/Source/Runtime/Engine/Private/LightWeightInstanceStaticMeshManager.cpp
--- a/Engine/Source/Runtime/Engine/Private/LightWeightInstanceStaticMeshManager.cpp
+++ b/Engine/Source/Runtime/Engine/Private/LightWeightInstanceStaticMeshManager.cpp
@@ -71,7 +71,7 @@ PRAGMA_DISABLE_DEPRECATION_WARNINGS
 			ISMComponent->AddInstance(InstanceTransform, /*bWorldSpace*/false);
 		}
 
-		ISMComponent->OnPostLoadPerInstanceData();
+		ISMComponent->OnPostPopulatePerInstanceData();
 
 		// release all instances in the HISM
 		InstancedStaticMeshComponent_DEPRECATED->ClearInstances();
@@ -110,7 +110,9 @@ void ALightWeightInstanceStaticMeshManager::SetRepresentedClass(UClass* ActorCla
 
 	if (ISMComponent)
 	{
-		ISMComponent->OnPostLoadPerInstanceData();
+PRAGMA_DISABLE_DEPRECATION_WARNINGS
+		ISMComponent->OnPostPopulatePerInstanceData();
+PRAGMA_ENABLE_DEPRECATION_WARNINGS
 	}
 }
 
