Implemented GenerateHLSLExpression for PerInstanceRandom, PerInstanceCustomData(3Vector), and SamplePhysics(Vector/Scalar/Integer)Field. Added a hash map to speed up GetValueShader on complex trees

diff --git a/Engine/Source/Runtime/Engine/Classes/Materials/MaterialExpressionPerInstanceCustomData.h b/Engine/Source/Runtime/Engine/Classes/Materials/MaterialExpressionPerInstanceCustomData.h
--- a/Engine/Source/Runtime/Engine/Classes/Materials/MaterialExpressionPerInstanceCustomData.h
+++ b/Engine/Source/Runtime/Engine/Classes/Materials/MaterialExpressionPerInstanceCustomData.h
@@ -25,6 +25,8 @@ class UMaterialExpressionPerInstanceCustomData : public UMaterialExpression
 #if WITH_EDITOR
 	virtual int32 Compile(class FMaterialCompiler* Compiler, int32 OutputIndex) override;
 	virtual void GetCaption(TArray<FString>& OutCaptions) const override;
+
+	virtual bool GenerateHLSLExpression(FMaterialHLSLGenerator& Generator, UE::HLSLTree::FScope& Scope, int32 OutputIndex, UE::HLSLTree::FExpression const*& OutExpression) const override;
 #endif
 	//~ End UMaterialExpression Interface
 };
@@ -47,6 +49,8 @@ class UMaterialExpressionPerInstanceCustomData3Vector : public UMaterialExpressi
 #if WITH_EDITOR
 	virtual int32 Compile(class FMaterialCompiler* Compiler, int32 OutputIndex) override;
 	virtual void GetCaption(TArray<FString>& OutCaptions) const override;
+
+	virtual bool GenerateHLSLExpression(FMaterialHLSLGenerator& Generator, UE::HLSLTree::FScope& Scope, int32 OutputIndex, UE::HLSLTree::FExpression const*& OutExpression) const override;
 #endif
 	//~ End UMaterialExpression Interface
 };
\ No newline at end of file
diff --git a/Engine/Source/Runtime/Engine/Classes/Materials/MaterialExpressionPerInstanceRandom.h b/Engine/Source/Runtime/Engine/Classes/Materials/MaterialExpressionPerInstanceRandom.h
--- a/Engine/Source/Runtime/Engine/Classes/Materials/MaterialExpressionPerInstanceRandom.h
+++ b/Engine/Source/Runtime/Engine/Classes/Materials/MaterialExpressionPerInstanceRandom.h
@@ -18,6 +18,8 @@ class UMaterialExpressionPerInstanceRandom : public UMaterialExpression
 #if WITH_EDITOR
 	virtual int32 Compile(class FMaterialCompiler* Compiler, int32 OutputIndex) override;
 	virtual void GetCaption(TArray<FString>& OutCaptions) const override;
+
+	virtual bool GenerateHLSLExpression(FMaterialHLSLGenerator& Generator, UE::HLSLTree::FScope& Scope, int32 OutputIndex, UE::HLSLTree::FExpression const*& OutExpression) const override;
 #endif
 	//~ End UMaterialExpression Interface
 };
diff --git a/Engine/Source/Runtime/Engine/Classes/Materials/MaterialExpressionSamplePhysicsField.h b/Engine/Source/Runtime/Engine/Classes/Materials/MaterialExpressionSamplePhysicsField.h
--- a/Engine/Source/Runtime/Engine/Classes/Materials/MaterialExpressionSamplePhysicsField.h
+++ b/Engine/Source/Runtime/Engine/Classes/Materials/MaterialExpressionSamplePhysicsField.h
@@ -29,6 +29,8 @@ class UMaterialExpressionSamplePhysicsVectorField : public UMaterialExpression
 #if WITH_EDITOR
 	virtual int32 Compile(class FMaterialCompiler* Compiler, int32 OutputIndex) override;
 	virtual void GetCaption(TArray<FString>& OutCaptions) const override;
+
+	virtual bool GenerateHLSLExpression(FMaterialHLSLGenerator& Generator, UE::HLSLTree::FScope& Scope, int32 OutputIndex, UE::HLSLTree::FExpression const*& OutExpression) const override;
 #endif
 	//~ End UMaterialExpression Interface
 };
@@ -49,6 +51,8 @@ class UMaterialExpressionSamplePhysicsScalarField : public UMaterialExpression
 #if WITH_EDITOR
 	virtual int32 Compile(class FMaterialCompiler* Compiler, int32 OutputIndex) override;
 	virtual void GetCaption(TArray<FString>& OutCaptions) const override;
+
+	virtual bool GenerateHLSLExpression(FMaterialHLSLGenerator& Generator, UE::HLSLTree::FScope& Scope, int32 OutputIndex, UE::HLSLTree::FExpression const*& OutExpression) const override;
 #endif
 	//~ End UMaterialExpression Interface
 };
@@ -69,6 +73,8 @@ class UMaterialExpressionSamplePhysicsIntegerField : public UMaterialExpression
 #if WITH_EDITOR
 	virtual int32 Compile(class FMaterialCompiler* Compiler, int32 OutputIndex) override;
 	virtual void GetCaption(TArray<FString>& OutCaptions) const override;
+
+	virtual bool GenerateHLSLExpression(FMaterialHLSLGenerator& Generator, UE::HLSLTree::FScope& Scope, int32 OutputIndex, UE::HLSLTree::FExpression const*& OutExpression) const override;
 #endif
 	//~ End UMaterialExpression Interface
 };
diff --git a/Engine/Source/Runtime/Engine/Private/HLSLTree/HLSLTree.cpp b/Engine/Source/Runtime/Engine/Private/HLSLTree/HLSLTree.cpp
--- a/Engine/Source/Runtime/Engine/Private/HLSLTree/HLSLTree.cpp
+++ b/Engine/Source/Runtime/Engine/Private/HLSLTree/HLSLTree.cpp
@@ -1169,6 +1169,22 @@ bool FExpression::EmitCustomHLSLParameter(FEmitContext& Context, FEmitScope& Sco
 
 FEmitShaderExpression* FExpression::GetValueShader(FEmitContext& Context, FEmitScope& Scope, const FRequestedType& RequestedType, const FPreparedType& PreparedType, const Shader::FType& ResultType) const
 {
+	FXxHash64 Hash;
+	{
+		FHasher Hasher;
+		AppendHash(Hasher, this);
+		AppendHash(Hasher, &Scope);
+		AppendHash(Hasher, RequestedType);
+		AppendHash(Hasher, ResultType);
+		Hash = Hasher.Finalize();
+	}
+
+	FEmitShaderExpression** Found = Context.EmitValueMap.Find(Hash);
+	if (Found)
+	{
+		return *Found;
+	}
+
 	FEmitOwnerScope OwnerScope(Context, this);
 
 	const EExpressionEvaluation Evaluation = PreparedType.GetEvaluation(Scope, RequestedType);
@@ -1190,7 +1206,10 @@ FEmitShaderExpression* FExpression::GetValueShader(FEmitContext& Context, FEmitS
 		check(Result.Code);
 		Value = Result.Code;
 	}
-	return Context.EmitCast(Scope, Value, ResultType);
+	Value = Context.EmitCast(Scope, Value, ResultType);
+
+	Context.EmitValueMap.Add(Hash, Value);
+	return Value;
 }
 
 FEmitShaderExpression* FExpression::GetValueShader(FEmitContext& Context, FEmitScope& Scope, const FRequestedType& RequestedType, const Shader::FType& ResultType) const
diff --git a/Engine/Source/Runtime/Engine/Private/Materials/MaterialExpressionHLSL.cpp b/Engine/Source/Runtime/Engine/Private/Materials/MaterialExpressionHLSL.cpp
--- a/Engine/Source/Runtime/Engine/Private/Materials/MaterialExpressionHLSL.cpp
+++ b/Engine/Source/Runtime/Engine/Private/Materials/MaterialExpressionHLSL.cpp
@@ -121,7 +121,9 @@
 #include "Materials/MaterialExpressionParticleSpeed.h"
 #include "Materials/MaterialExpressionParticleSubUVProperties.h"
 #include "Materials/MaterialExpressionPathTracingQualitySwitch.h"
+#include "Materials/MaterialExpressionPerInstanceCustomData.h"
 #include "Materials/MaterialExpressionPerInstanceFadeAmount.h"
+#include "Materials/MaterialExpressionPerInstanceRandom.h"
 #include "Materials/MaterialExpressionPixelDepth.h"
 #include "Materials/MaterialExpressionPixelNormalWS.h"
 #include "Materials/MaterialExpressionPower.h"
@@ -129,6 +131,7 @@
 #include "Materials/MaterialExpressionPreSkinnedNormal.h"
 #include "Materials/MaterialExpressionPreSkinnedPosition.h"
 #include "Materials/MaterialExpressionPreviousFrameSwitch.h"
+#include "Materials/MaterialExpressionSamplePhysicsField.h"
 #include "Materials/MaterialExpressionQualitySwitch.h"
 #include "Materials/MaterialExpressionRayTracingQualitySwitch.h"
 #include "Materials/MaterialExpressionReflectionVectorWS.h"
@@ -3337,8 +3340,30 @@ bool UMaterialExpressionCloudSampleAttribute::GenerateHLSLExpression(FMaterialHL
 bool UMaterialExpressionPerInstanceFadeAmount::GenerateHLSLExpression(FMaterialHLSLGenerator& Generator, UE::HLSLTree::FScope& Scope, int32 OutputIndex, UE::HLSLTree::FExpression const*& OutExpression) const
 {
 	using namespace UE::HLSLTree;
-	using namespace UE::Shader;
-	OutExpression = Generator.GetTree().NewExpression<FExpressionInlineCustomHLSL>(EValueType::Float1, TEXT("GetPerInstanceFadeAmount(Parameters)"));
+	OutExpression = Generator.NewExternalInput(Material::EExternalInput::PerInstanceFadeAmount);
+	return true;
+}
+
+bool UMaterialExpressionPerInstanceRandom::GenerateHLSLExpression(FMaterialHLSLGenerator& Generator, UE::HLSLTree::FScope& Scope, int32 OutputIndex, UE::HLSLTree::FExpression const*& OutExpression) const
+{
+	using namespace UE::HLSLTree;
+	OutExpression = Generator.NewExternalInput(Material::EExternalInput::PerInstanceRandom);
+	return true;
+}
+
+bool UMaterialExpressionPerInstanceCustomData::GenerateHLSLExpression(FMaterialHLSLGenerator& Generator, UE::HLSLTree::FScope& Scope, int32 OutputIndex, UE::HLSLTree::FExpression const*& OutExpression) const
+{
+	using namespace UE::HLSLTree;
+	const FExpression* DefaultValueExpression = DefaultValue.AcquireHLSLExpressionOrConstant(Generator, Scope, ConstDefaultValue);
+	OutExpression = Generator.GetTree().NewExpression<Material::FExpressionPerInstanceCustomData>(DefaultValueExpression, (int32)DataIndex, false);
+	return true;
+}
+
+bool UMaterialExpressionPerInstanceCustomData3Vector::GenerateHLSLExpression(FMaterialHLSLGenerator& Generator, UE::HLSLTree::FScope& Scope, int32 OutputIndex, UE::HLSLTree::FExpression const*& OutExpression) const
+{
+	using namespace UE::HLSLTree;
+	const FExpression* DefaultValueExpression = DefaultValue.AcquireHLSLExpressionOrConstant(Generator, Scope, ConstDefaultValue);
+	OutExpression = Generator.GetTree().NewExpression<Material::FExpressionPerInstanceCustomData>(DefaultValueExpression, (int32)DataIndex, true);
 	return true;
 }
 
@@ -3652,18 +3677,7 @@ bool UMaterialExpressionDeriveNormalZ::GenerateHLSLExpression(FMaterialHLSLGener
 bool UMaterialExpressionDistanceToNearestSurface::GenerateHLSLExpression(FMaterialHLSLGenerator& Generator, UE::HLSLTree::FScope& Scope, int32 OutputIndex, UE::HLSLTree::FExpression const*& OutExpression) const
 {
 	using namespace UE::HLSLTree;
-
-	const FExpression* PositionExpression;
-
-	if (Position.GetTracedInput().Expression)
-	{
-		PositionExpression = Position.AcquireHLSLExpression(Generator, Scope);
-	}
-	else
-	{
-		PositionExpression = Generator.NewExternalInput(Material::EExternalInput::WorldPosition);
-	}
-
+	const FExpression* PositionExpression = Position.AcquireHLSLExpressionOrExternalInput(Generator, Scope, Material::EExternalInput::WorldPosition);
 	OutExpression = Generator.GetTree().NewExpression<Material::FExpressionDistanceToNearestSurface>(PositionExpression);
 	return true;
 }
@@ -3671,18 +3685,7 @@ bool UMaterialExpressionDistanceToNearestSurface::GenerateHLSLExpression(FMateri
 bool UMaterialExpressionDistanceFieldGradient::GenerateHLSLExpression(FMaterialHLSLGenerator& Generator, UE::HLSLTree::FScope& Scope, int32 OutputIndex, UE::HLSLTree::FExpression const*& OutExpression) const
 {
 	using namespace UE::HLSLTree;
-
-	const FExpression* PositionExpression;
-
-	if (Position.GetTracedInput().Expression)
-	{
-		PositionExpression = Position.AcquireHLSLExpression(Generator, Scope);
-	}
-	else
-	{
-		PositionExpression = Generator.NewExternalInput(Material::EExternalInput::WorldPosition);
-	}
-
+	const FExpression* PositionExpression = Position.AcquireHLSLExpressionOrExternalInput(Generator, Scope, Material::EExternalInput::WorldPosition);
 	OutExpression = Generator.GetTree().NewExpression<Material::FExpressionDistanceFieldGradient>(PositionExpression);
 	return true;
 }
@@ -3691,27 +3694,8 @@ bool UMaterialExpressionDistanceFieldApproxAO::GenerateHLSLExpression(FMaterialH
 {
 	using namespace UE::HLSLTree;
 
-	const FExpression* PositionExpression;;
-
-	if (Position.GetTracedInput().Expression)
-	{
-		PositionExpression = Position.AcquireHLSLExpression(Generator, Scope);
-	}
-	else
-	{
-		PositionExpression = Generator.NewExternalInput(Material::EExternalInput::WorldPosition);
-	}
-
-	const FExpression* NormalExpression;
-
-	if (Normal.GetTracedInput().Expression)
-	{
-		NormalExpression = Normal.AcquireHLSLExpression(Generator, Scope);
-	}
-	else
-	{
-		NormalExpression = Generator.NewExternalInput(Material::EExternalInput::WorldVertexNormal);
-	}
+	const FExpression* PositionExpression = Position.AcquireHLSLExpressionOrExternalInput(Generator, Scope, Material::EExternalInput::WorldPosition);
+	const FExpression* NormalExpression = Normal.AcquireHLSLExpressionOrExternalInput(Generator, Scope, Material::EExternalInput::WorldVertexNormal);
 
 	const FExpression* BaseDistanceExpression = BaseDistance.AcquireHLSLExpressionOrConstant(Generator, Scope, BaseDistanceDefault);
 	const FExpression* RadiusExpression = Radius.AcquireHLSLExpressionOrConstant(Generator, Scope, RadiusDefault);
@@ -3753,4 +3737,49 @@ bool UMaterialExpressionDistanceFieldApproxAO::GenerateHLSLExpression(FMaterialH
 	return true;
 }
 
+bool UMaterialExpressionSamplePhysicsVectorField::GenerateHLSLExpression(FMaterialHLSLGenerator& Generator, UE::HLSLTree::FScope& Scope, int32 OutputIndex, UE::HLSLTree::FExpression const*& OutExpression) const
+{
+	using namespace UE::HLSLTree;
+
+	const int32 TargetIndex = (int32)FieldTarget;
+	if (TargetIndex < 0 || TargetIndex >= Vector_TargetMax)
+	{
+		return Generator.Errorf(TEXT("Invalid physics field target %d"), TargetIndex);
+	}
+
+	const FExpression* PositionExpression = WorldPosition.AcquireHLSLExpressionOrExternalInput(Generator, Scope, Material::EExternalInput::WorldPosition);
+	OutExpression = Generator.GetTree().NewExpression<Material::FExpressionSamplePhysicsField>(PositionExpression, Field_Output_Vector, TargetIndex);
+	return true;
+}
+
+bool UMaterialExpressionSamplePhysicsScalarField::GenerateHLSLExpression(FMaterialHLSLGenerator& Generator, UE::HLSLTree::FScope& Scope, int32 OutputIndex, UE::HLSLTree::FExpression const*& OutExpression) const
+{
+	using namespace UE::HLSLTree;
+
+	const int32 TargetIndex = (int32)FieldTarget;
+	if (TargetIndex < 0 || TargetIndex >= Vector_TargetMax)
+	{
+		return Generator.Errorf(TEXT("Invalid physics field target %d"), TargetIndex);
+	}
+
+	const FExpression* PositionExpression = WorldPosition.AcquireHLSLExpressionOrExternalInput(Generator, Scope, Material::EExternalInput::WorldPosition);
+	OutExpression = Generator.GetTree().NewExpression<Material::FExpressionSamplePhysicsField>(PositionExpression, Field_Output_Scalar, TargetIndex);
+	return true;
+}
+
+bool UMaterialExpressionSamplePhysicsIntegerField::GenerateHLSLExpression(FMaterialHLSLGenerator& Generator, UE::HLSLTree::FScope& Scope, int32 OutputIndex, UE::HLSLTree::FExpression const*& OutExpression) const
+{
+	using namespace UE::HLSLTree;
+
+	const int32 TargetIndex = (int32)FieldTarget;
+	if (TargetIndex < 0 || TargetIndex >= Vector_TargetMax)
+	{
+		return Generator.Errorf(TEXT("Invalid physics field target %d"), TargetIndex);
+	}
+
+	const FExpression* PositionExpression = WorldPosition.AcquireHLSLExpressionOrExternalInput(Generator, Scope, Material::EExternalInput::WorldPosition);
+	OutExpression = Generator.GetTree().NewExpression<Material::FExpressionSamplePhysicsField>(PositionExpression, Field_Output_Integer, TargetIndex);
+	return true;
+}
+
 #endif // WITH_EDITOR
diff --git a/Engine/Source/Runtime/Engine/Private/Materials/MaterialHLSLEmitter.cpp b/Engine/Source/Runtime/Engine/Private/Materials/MaterialHLSLEmitter.cpp
--- a/Engine/Source/Runtime/Engine/Private/Materials/MaterialHLSLEmitter.cpp
+++ b/Engine/Source/Runtime/Engine/Private/Materials/MaterialHLSLEmitter.cpp
@@ -443,8 +443,11 @@ static void GetMaterialEnvironment(EShaderPlatform InPlatform,
 		OutEnvironment.SetDefine(TEXT("VIRTUAL_TEXTURE_OUTPUT"), 1);
 	}
 
+	const bool bUsesPerInstanceFadeAmount = EmitMaterialData.IsExternalInputUsed(Material::EExternalInput::PerInstanceFadeAmount);
+
 	OutEnvironment.SetDefine(TEXT("USES_PER_INSTANCE_CUSTOM_DATA"), MaterialCompilationOutput.bUsesPerInstanceCustomData && InMaterial.IsUsedWithInstancedStaticMeshes());
 	OutEnvironment.SetDefine(TEXT("USES_PER_INSTANCE_RANDOM"), MaterialCompilationOutput.bUsesPerInstanceRandom && InMaterial.IsUsedWithInstancedStaticMeshes());
+	OutEnvironment.SetDefine(TEXT("USES_PER_INSTANCE_FADE_AMOUNT"), bUsesPerInstanceFadeAmount&& InMaterial.IsUsedWithInstancedStaticMeshes());
 	OutEnvironment.SetDefine(TEXT("USES_VERTEX_INTERPOLATOR"), MaterialCompilationOutput.bUsesVertexInterpolator);
 
 	const bool bUsesVertexColor = EmitMaterialData.IsExternalInputUsed(SF_Vertex, Material::EExternalInput::VertexColor) ||
diff --git a/Engine/Source/Runtime/Engine/Private/Materials/MaterialHLSLTree.cpp b/Engine/Source/Runtime/Engine/Private/Materials/MaterialHLSLTree.cpp
--- a/Engine/Source/Runtime/Engine/Private/Materials/MaterialHLSLTree.cpp
+++ b/Engine/Source/Runtime/Engine/Private/Materials/MaterialHLSLTree.cpp
@@ -154,6 +154,9 @@ FExternalInputDescription GetExternalInputDescription(EExternalInput Input)
 	case EExternalInput::ParticleSubUVCoords1: return FExternalInputDescription(TEXT("ParticleSubUVCoords1"), Shader::EValueType::Float2);
 	case EExternalInput::ParticleSubUVLerp: return FExternalInputDescription(TEXT("ParticleSubUVLerp"), Shader::EValueType::Float1);
 
+	case EExternalInput::PerInstanceFadeAmount: return FExternalInputDescription(TEXT("PerInstanceFadeAmount"), Shader::EValueType::Float1);
+	case EExternalInput::PerInstanceRandom: return FExternalInputDescription(TEXT("PerInstanceRandom"), Shader::EValueType::Float1);
+
 	case EExternalInput::IsOrthographic: return FExternalInputDescription(TEXT("IsOrthographic"), Shader::EValueType::Float1);
 
 	default: checkNoEntry(); return FExternalInputDescription(TEXT("Invalid"), Shader::EValueType::Void);
@@ -221,6 +224,18 @@ bool FExpressionExternalInput::PrepareValue(FEmitContext& Context, FEmitScope& S
 		FEmitData& EmitMaterialData = Context.FindData<FEmitData>();
 		const int32 TypeIndex = (int32)InputType;
 		EmitMaterialData.ExternalInputMask[Context.ShaderFrequency][TypeIndex] = true;
+
+		if (Context.MaterialCompilationOutput)
+		{
+			switch (InputType)
+			{
+			case EExternalInput::PerInstanceRandom:
+				Context.MaterialCompilationOutput->bUsesPerInstanceRandom = true;
+				break;
+			default:
+				break;
+			}
+		}
 	}
 
 	return OutResult.SetType(Context, RequestedType, EExpressionEvaluation::Shader, InputDesc.Type);
@@ -352,6 +367,9 @@ void FExpressionExternalInput::EmitValueShader(FEmitContext& Context, FEmitScope
 		case EExternalInput::ParticleSubUVCoords1: Code = TEXT("Parameters.Particle.SubUVCoords[1].xy"); break;
 		case EExternalInput::ParticleSubUVLerp: Code = TEXT("Parameters.Particle.SubUVLerp"); break;
 
+		case EExternalInput::PerInstanceFadeAmount: Code = TEXT("GetPerInstanceFadeAmount(Parameters)"); break;
+		case EExternalInput::PerInstanceRandom: Code = TEXT("GetPerInstanceRandom(Parameters)"); break;
+
 		case EExternalInput::IsOrthographic: Code = TEXT("((View.ViewToClip[3][3] < 1.0f) ? 0.0f : 1.0f)"); break;
 
 		default:
@@ -1870,6 +1888,83 @@ void FExpressionDistanceFieldGradient::EmitValueShader(FEmitContext& Context, FE
 	OutResult.Code = Context.EmitInlineExpression(Scope, Shader::EValueType::Float3, TEXT("GetDistanceFieldGradientGlobal(%)"), EmitPosition);
 }
 
+bool FExpressionPerInstanceCustomData::PrepareValue(FEmitContext& Context, FEmitScope& Scope, const FRequestedType& RequestedType, FPrepareValueResult& OutResult) const
+{
+	const FPreparedType& DefaultValueType = Context.PrepareExpression(DefaultValueExpression, Scope, GetCustomDataType());
+	if (DefaultValueType.IsVoid())
+	{
+		return false;
+	}
+
+	if (Context.bMarkLiveValues && Context.MaterialCompilationOutput)
+	{
+		Context.MaterialCompilationOutput->bUsesPerInstanceCustomData = true;
+	}
+
+	return OutResult.SetType(Context, RequestedType, EExpressionEvaluation::Shader, GetCustomDataType());
+}
+
+void FExpressionPerInstanceCustomData::EmitValueShader(FEmitContext& Context, FEmitScope& Scope, const FRequestedType& RequestedType, FEmitValueShaderResult& OutResult) const
+{
+	FEmitShaderExpression* EmitDefaultValue = DefaultValueExpression->GetValueShader(Context, Scope, GetCustomDataType());
+	OutResult.Code = Context.EmitInlineExpression(Scope, GetCustomDataType(), TEXT("GetPerInstanceCustomData%(Parameters, %, %)"), b3Vector ? TEXT("3Vector") : TEXT(""), DataIndex, EmitDefaultValue);
+}
+
+Shader::EValueType FExpressionPerInstanceCustomData::GetCustomDataType() const
+{
+	return b3Vector ? Shader::EValueType::Float3 : Shader::EValueType::Float1;
+}
+
+bool FExpressionSamplePhysicsField::PrepareValue(FEmitContext& Context, FEmitScope& Scope, const FRequestedType& RequestedType, FPrepareValueResult& OutResult) const
+{
+	// LWC_TODO: LWC aware physics field
+	const FPreparedType& PositionType = Context.PrepareExpression(PositionExpression, Scope, Shader::EValueType::Float3);
+	if (PositionType.IsVoid())
+	{
+		return false;
+	}
+
+	return OutResult.SetType(Context, RequestedType, EExpressionEvaluation::Shader, GetOutputType());
+}
+
+void FExpressionSamplePhysicsField::EmitValueShader(FEmitContext& Context, FEmitScope& Scope, const FRequestedType& RequestedType, FEmitValueShaderResult& OutResult) const
+{
+	FEmitShaderExpression* EmitPosition = PositionExpression->GetValueShader(Context, Scope, Shader::EValueType::Float3);
+	OutResult.Code = Context.EmitInlineExpression(Scope, GetOutputType(), GetEmitExpressionFormat(), EmitPosition, TargetIndex);
+}
+
+Shader::EValueType FExpressionSamplePhysicsField::GetOutputType() const
+{
+	switch (FieldOutputType)
+	{
+	case Field_Output_Vector:
+		return Shader::EValueType::Float3;
+	case Field_Output_Scalar:
+		return Shader::EValueType::Float1;
+	case Field_Output_Integer:
+		return Shader::EValueType::Int1;
+	default:
+		checkNoEntry();
+		return Shader::EValueType::Void;
+	}
+}
+
+const TCHAR* FExpressionSamplePhysicsField::GetEmitExpressionFormat() const
+{
+	switch (FieldOutputType)
+	{
+	case Field_Output_Vector:
+		return TEXT("MatPhysicsField_SamplePhysicsVectorField(%, %)");
+	case Field_Output_Scalar:
+		return TEXT("MatPhysicsField_SamplePhysicsScalarField(%, %)");
+	case Field_Output_Integer:
+		return TEXT("MatPhysicsField_SamplePhysicsIntegerField(%, %)");
+	default:
+		checkNoEntry();
+		return nullptr;
+	}
+}
+
 bool FExpressionDistanceFieldApproxAO::PrepareValue(FEmitContext& Context, FEmitScope& Scope, const FRequestedType& RequestedType, FPrepareValueResult& OutResult) const
 {
 	if (!Private::PlatformSupportDistanceFields(Context))
diff --git a/Engine/Source/Runtime/Engine/Public/HLSLTree/HLSLTreeEmit.h b/Engine/Source/Runtime/Engine/Public/HLSLTree/HLSLTreeEmit.h
--- a/Engine/Source/Runtime/Engine/Public/HLSLTree/HLSLTreeEmit.h
+++ b/Engine/Source/Runtime/Engine/Public/HLSLTree/HLSLTreeEmit.h
@@ -592,8 +592,10 @@ public:
 	TMap<FXxHash64, FRequestedType*> RequestedTypeTracker;
 	TMap<const FExpression*, FEmitScope*> PrepareLocalPHIMap;
 	TMap<const FExpression*, FEmitShaderExpression*> EmitLocalPHIMap;
+	// TODO: investigate EmitExpressionMap and EmitPreshaderMap can be removed and only keep EmitValueMap
 	TMap<FXxHash64, FEmitShaderExpression*> EmitExpressionMap;
 	TMap<FXxHash64, FEmitShaderExpression*> EmitPreshaderMap;
+	TMap<FXxHash64, FEmitShaderExpression*> EmitValueMap;
 	TMap<const FFunction*, FEmitShaderNode*> EmitFunctionMap;
 	TMap<FXxHash64, FEmitCustomHLSL> EmitCustomHLSLMap;
 	TArray<struct FPreshaderLoopScope*> PreshaderLoopScopes;
diff --git a/Engine/Source/Runtime/Engine/Public/HLSLTree/HLSLTreeHash.h b/Engine/Source/Runtime/Engine/Public/HLSLTree/HLSLTreeHash.h
--- a/Engine/Source/Runtime/Engine/Public/HLSLTree/HLSLTreeHash.h
+++ b/Engine/Source/Runtime/Engine/Public/HLSLTree/HLSLTreeHash.h
@@ -78,6 +78,7 @@ inline void AppendHash(FHasher& Hasher, const FCustomHLSLInput& Value)
 inline void AppendHash(FHasher& Hasher, const Shader::FType& Type)
 {
 	if (Type.IsStruct()) AppendHash(Hasher, Type.StructType);
+	else if (Type.IsObject()) AppendHash(Hasher, Type.ObjectType);
 	else AppendHash(Hasher, Type.ValueType);
 }
 
diff --git a/Engine/Source/Runtime/Engine/Public/MaterialHLSLTree.h b/Engine/Source/Runtime/Engine/Public/MaterialHLSLTree.h
--- a/Engine/Source/Runtime/Engine/Public/MaterialHLSLTree.h
+++ b/Engine/Source/Runtime/Engine/Public/MaterialHLSLTree.h
@@ -9,6 +9,7 @@
 #include "RHIDefinitions.h"
 #include "Engine/Texture.h" // enum TextureAddress for VTStackEntry
 #include "VT/RuntimeVirtualTextureEnum.h"
+#include "Field/FieldSystemTypes.h"
 #include "MaterialCompiler.h"
 
 class UTexture;
@@ -191,6 +192,9 @@ enum class EExternalInput : uint8
 	ParticleSubUVCoords1,
 	ParticleSubUVLerp,
 
+	PerInstanceFadeAmount,
+	PerInstanceRandom,
+
 	IsOrthographic,
 
 	Num,
@@ -633,6 +637,47 @@ public:
 	virtual void EmitValueShader(FEmitContext& Context, FEmitScope& Scope, const FRequestedType& RequestedType, FEmitValueShaderResult& OutResult) const override;
 };
 
+class FExpressionPerInstanceCustomData : public FExpression
+{
+public:
+	const FExpression* DefaultValueExpression;
+	int32 DataIndex;
+	bool b3Vector;
+
+	FExpressionPerInstanceCustomData(const FExpression* InDefaultValueExpression, int32 InDataIndex, bool bIn3Vector)
+		: DefaultValueExpression(InDefaultValueExpression)
+		, DataIndex(InDataIndex)
+		, b3Vector(bIn3Vector)
+	{}
+
+	virtual bool PrepareValue(FEmitContext& Context, FEmitScope& Scope, const FRequestedType& RequestedType, FPrepareValueResult& OutResult) const override;
+	virtual void EmitValueShader(FEmitContext& Context, FEmitScope& Scope, const FRequestedType& RequestedType, FEmitValueShaderResult& OutResult) const override;
+
+private:
+	Shader::EValueType GetCustomDataType() const;
+};
+
+class FExpressionSamplePhysicsField : public FExpression
+{
+public:
+	const FExpression* PositionExpression;
+	EFieldOutputType FieldOutputType;
+	int32 TargetIndex;
+
+	FExpressionSamplePhysicsField(const FExpression* InPositionExpression, EFieldOutputType InFieldOutputType, int32 InTargetIndex)
+		: PositionExpression(InPositionExpression)
+		, FieldOutputType(InFieldOutputType)
+		, TargetIndex(InTargetIndex)
+	{}
+
+	virtual bool PrepareValue(FEmitContext& Context, FEmitScope& Scope, const FRequestedType& RequestedType, FPrepareValueResult& OutResult) const override;
+	virtual void EmitValueShader(FEmitContext& Context, FEmitScope& Scope, const FRequestedType& RequestedType, FEmitValueShaderResult& OutResult) const override;
+
+private:
+	Shader::EValueType GetOutputType() const;
+	const TCHAR* GetEmitExpressionFormat() const;
+};
+
 class FExpressionDistanceFieldApproxAO : public FExpression
 {
 public:
