Fix harmless stack-use-after-scope warning.

diff --git a/Engine/Source/Runtime/Engine/Private/Level.cpp b/Engine/Source/Runtime/Engine/Private/Level.cpp
--- a/Engine/Source/Runtime/Engine/Private/Level.cpp
+++ b/Engine/Source/Runtime/Engine/Private/Level.cpp
@@ -1008,7 +1008,8 @@ void ULevel::RemoveLoadedActors(const TArray<AActor*>& ActorList, const FTransfo
 
 			if (TransformToRemove)
 			{
-				FLevelUtils::FApplyLevelTransformParams TransformParams(this, TransformToRemove->Inverse());
+				const FTransform TransformToRemoveInverse = TransformToRemove->Inverse();
+				FLevelUtils::FApplyLevelTransformParams TransformParams(this, TransformToRemoveInverse);
 				TransformParams.Actor = Actor;
 				TransformParams.bDoPostEditMove = true;
 				FLevelUtils::ApplyLevelTransform(TransformParams);
