World Partition: fix UWorldPartitionSubsystem::IsStreamingCompleted functions not converting query sources in local space.
Adatped from PR #10865

diff --git a/Engine/Source/Runtime/Engine/Private/WorldPartition/WorldPartitionSubsystem.cpp b/Engine/Source/Runtime/Engine/Private/WorldPartition/WorldPartitionSubsystem.cpp
--- a/Engine/Source/Runtime/Engine/Private/WorldPartition/WorldPartitionSubsystem.cpp
+++ b/Engine/Source/Runtime/Engine/Private/WorldPartition/WorldPartitionSubsystem.cpp
@@ -814,12 +814,11 @@ bool UWorldPartitionSubsystem::IsStreamingCompleted(const IWorldPartitionStreami
 {
 	// Convert specified/optional streaming source provider to a world partition 
 	// streaming source and pass it along to each registered world partition
-	TArray<FWorldPartitionStreamingSource> LocalStreamingSources;
-	TArray<FWorldPartitionStreamingSource>* StreamingSourcesPtr = nullptr;
+	TArray<FWorldPartitionStreamingSource> WorldStreamingSources;
+
 	if (InStreamingSourceProvider)
 	{
-		StreamingSourcesPtr = &LocalStreamingSources;
-		if (!InStreamingSourceProvider->GetStreamingSources(LocalStreamingSources))
+		if (!InStreamingSourceProvider->GetStreamingSources(WorldStreamingSources))
 		{
 			return true;
 		}
@@ -827,7 +826,21 @@ bool UWorldPartitionSubsystem::IsStreamingCompleted(const IWorldPartitionStreami
 
 	for (UWorldPartition* RegisteredWorldPartition : RegisteredWorldPartitions)
 	{
-		if (!RegisteredWorldPartition->IsStreamingCompleted(StreamingSourcesPtr))
+		TArray<FWorldPartitionStreamingSource> LocalStreamingSources;
+
+		if (InStreamingSourceProvider)
+		{
+			const FTransform WorldToLocal = RegisteredWorldPartition->GetInstanceTransform().Inverse();
+			
+			LocalStreamingSources = WorldStreamingSources;
+			for (FWorldPartitionStreamingSource& StreamingSource : LocalStreamingSources)
+			{
+				StreamingSource.Location = WorldToLocal.TransformPosition(StreamingSource.Location);
+				StreamingSource.Rotation = WorldToLocal.TransformRotation(StreamingSource.Rotation.Quaternion()).Rotator();
+			}
+		}
+
+		if (!RegisteredWorldPartition->IsStreamingCompleted(InStreamingSourceProvider ? &LocalStreamingSources : nullptr))
 		{
 			return false;
 		}
@@ -839,7 +852,16 @@ bool UWorldPartitionSubsystem::IsStreamingCompleted(EWorldPartitionRuntimeCellSt
 {
 	for (UWorldPartition* RegisteredWorldPartition : RegisteredWorldPartitions)
 	{
-		if (!RegisteredWorldPartition->IsStreamingCompleted(QueryState, QuerySources, bExactState))
+		const FTransform WorldToLocal = RegisteredWorldPartition->GetInstanceTransform().Inverse();
+
+		TArray<FWorldPartitionStreamingQuerySource> LocalQuerySources = QuerySources;
+		for (FWorldPartitionStreamingQuerySource& QuerySource : LocalQuerySources)
+		{
+			QuerySource.Location = WorldToLocal.TransformPosition(QuerySource.Location);
+			QuerySource.Rotation = WorldToLocal.TransformRotation(QuerySource.Rotation.Quaternion()).Rotator();
+		}
+
+		if (!RegisteredWorldPartition->IsStreamingCompleted(QueryState, LocalQuerySources, bExactState))
 		{
 			return false;
 		}
