Implement IsDataValid for World and remove early out for Actors.

In some validation scenarios, the actors may not be loaded for a world so may not have IsDataValid called.

diff --git a/Engine/Source/Runtime/Engine/Classes/Engine/World.h b/Engine/Source/Runtime/Engine/Classes/Engine/World.h
--- a/Engine/Source/Runtime/Engine/Classes/Engine/World.h
+++ b/Engine/Source/Runtime/Engine/Classes/Engine/World.h
@@ -2887,6 +2887,7 @@ public:
 	virtual FPrimaryAssetId GetPrimaryAssetId() const override;
 	static void AddReferencedObjects(UObject* InThis, FReferenceCollector& Collector);
 #if WITH_EDITOR
+	virtual EDataValidationResult IsDataValid(FDataValidationContext& Context) const override;
 	virtual bool Rename(const TCHAR* NewName = NULL, UObject* NewOuter = NULL, ERenameFlags Flags = REN_None) override;
 	virtual void GetAssetRegistryTags(TArray<FAssetRegistryTag>& OutTags) const override;
 	virtual void PostLoadAssetRegistryTags(const FAssetData& InAssetData, TArray<FAssetRegistryTag>& OutTagsAndValuesToUpdate) const;
diff --git a/Engine/Source/Runtime/Engine/Private/ActorEditor.cpp b/Engine/Source/Runtime/Engine/Private/ActorEditor.cpp
--- a/Engine/Source/Runtime/Engine/Private/ActorEditor.cpp
+++ b/Engine/Source/Runtime/Engine/Private/ActorEditor.cpp
@@ -1590,12 +1590,6 @@ bool AActor::GetSoftReferencedContentObjects(TArray<FSoftObjectPath>& SoftObject
 
 EDataValidationResult AActor::IsDataValid(FDataValidationContext& Context) const
 {
-	// Do not run asset validation on external actors, validation will be caught through map check
-	if (IsPackageExternal())
-	{
-		return EDataValidationResult::NotValidated;
-	}
-
 	bool bSuccess = CheckDefaultSubobjects();
 	if (!bSuccess)
 	{
diff --git a/Engine/Source/Runtime/Engine/Private/World.cpp b/Engine/Source/Runtime/Engine/Private/World.cpp
--- a/Engine/Source/Runtime/Engine/Private/World.cpp
+++ b/Engine/Source/Runtime/Engine/Private/World.cpp
@@ -787,6 +787,20 @@ void UWorld::AddReferencedObjects(UObject* InThis, FReferenceCollector& Collecto
 }
 
 #if WITH_EDITOR
+EDataValidationResult UWorld::IsDataValid(FDataValidationContext& Context) const 
+{
+	EDataValidationResult Result = EDataValidationResult::NotValidated;
+	for (FActorIterator It(this); It; ++It)
+	{
+		AActor* Actor = *It;
+		if (Actor)
+		{
+			Result = CombineDataValidationResults(Result, Actor->IsDataValid(Context));
+		}
+	}
+	return CombineDataValidationResults(Result, Super::IsDataValid(Context));
+}
+
 bool UWorld::Rename(const TCHAR* InName, UObject* NewOuter, ERenameFlags Flags)
 {
 	check(PersistentLevel);
