LFV - renamed an enum to match its actual meaning.

#prelfight shader

diff --git a/Engine/Source/Runtime/Engine/Classes/Components/LocalFogVolumeComponent.h b/Engine/Source/Runtime/Engine/Classes/Components/LocalFogVolumeComponent.h
--- a/Engine/Source/Runtime/Engine/Classes/Components/LocalFogVolumeComponent.h
+++ b/Engine/Source/Runtime/Engine/Classes/Components/LocalFogVolumeComponent.h
@@ -12,7 +12,7 @@ class FLocalFogVolumeSceneProxy;
 UENUM()
 enum class ELocalFogMode : uint8
 {
-	LocalFogVolume = 0,
+	LocalHeightFog = 0,
 	LocalSphereFog = 1,
 };
 
@@ -25,7 +25,7 @@ class ULocalFogVolumeComponent : public USceneComponent
 
 	/** Controls the softness of the transition region when the volume is fading out. */
 	UPROPERTY(EditAnywhere, BlueprintReadOnly, Category = "Fog Mode")
-	ELocalFogMode FogMode = ELocalFogMode::LocalFogVolume;
+	ELocalFogMode FogMode = ELocalFogMode::LocalHeightFog;
 
 	/** Global density factor for this fog. */
 	UPROPERTY(EditAnywhere, BlueprintReadOnly, interp, Category = "Fog Distribution", meta = (UIMin = "0", UIMax = "10.0", SliderExponent = 2.0, ClampMin = 0.0))
diff --git a/Engine/Source/Runtime/Engine/Private/Components/LocalFogVolumeComponent.cpp b/Engine/Source/Runtime/Engine/Private/Components/LocalFogVolumeComponent.cpp
--- a/Engine/Source/Runtime/Engine/Private/Components/LocalFogVolumeComponent.cpp
+++ b/Engine/Source/Runtime/Engine/Private/Components/LocalFogVolumeComponent.cpp
@@ -99,7 +99,7 @@ bool ULocalFogVolumeComponent::CanEditChange(const FProperty* InProperty) const
 	{
 		FString PropertyName = InProperty->GetName();
 
-		if (FogMode != ELocalFogMode::LocalFogVolume)
+		if (FogMode != ELocalFogMode::LocalHeightFog)
 		{
 			if (PropertyName == GET_MEMBER_NAME_STRING_CHECKED(ULocalFogVolumeComponent, FogHeightFalloff)
 				|| PropertyName == GET_MEMBER_NAME_STRING_CHECKED(ULocalFogVolumeComponent, FogHeightOffset))
