Moved some render setting from shader permutation to optimisation because that is actually what they do.

#prelfight shader

diff --git a/Engine/Source/Runtime/Engine/Classes/Engine/RendererSettings.h b/Engine/Source/Runtime/Engine/Classes/Engine/RendererSettings.h
--- a/Engine/Source/Runtime/Engine/Classes/Engine/RendererSettings.h
+++ b/Engine/Source/Runtime/Engine/Classes/Engine/RendererSettings.h
@@ -836,6 +836,34 @@ class URendererSettings : public UDeveloperSettings
 		ToolTip = "Whether to use original CPU method (loop per morph then by vertex) or use a GPU-based method on Shader Model 5 hardware."))
 	uint32 bUseGPUMorphTargets : 1;
 
+	/**
+	"The sky atmosphere component requires extra samplers/textures to be bound to apply aerial perspective on transparent surfaces (and all surfaces on mobile via per vertex evaluation)."
+	*/
+	UPROPERTY(config, EditAnywhere, Category = Optimizations, meta = (
+		ConsoleVariable = "r.SupportSkyAtmosphere", DisplayName = "Support Sky Atmosphere",
+		ToolTip = "The sky atmosphere component requires extra samplers/textures to be bound to apply aerial perspective on transparent surfaces (and all surfaces on mobile via per vertex evaluation).",
+		ConfigRestartRequired = true))
+		uint32 bSupportSkyAtmosphere : 1;
+
+	/**
+	"The sky atmosphere component can light up the height fog but it requires extra samplers/textures to be bound to apply aerial perspective on transparent surfaces (and all surfaces on mobile via per vertex evaluation)."
+	"It requires r.SupportSkyAtmosphere to be true."
+	*/
+	UPROPERTY(config, EditAnywhere, Category = Optimizations, meta = (
+		ConsoleVariable = "r.SupportSkyAtmosphereAffectsHeightFog", DisplayName = "Support Sky Atmosphere Affecting Height Fog",
+		ToolTip = "The sky atmosphere component can light up the height fog but it requires extra samplers/textures to be bound to apply aerial perspective on transparent surfaces (and all surfaces on mobile via per vertex evaluation). It requires r.SupportSkyAtmosphere to be true.",
+		ConfigRestartRequired = true))
+		uint32 bSupportSkyAtmosphereAffectsHeightFog : 1;
+
+	/**
+	"Enable cloud shadow on translucent surface. This is evaluated per vertex to reduce GPU cost. The cloud system requires extra samplers/textures to be bound to vertex shaders."
+	*/
+	UPROPERTY(config, EditAnywhere, Category = Optimizations, meta = (
+		ConsoleVariable = "r.SupportCloudShadowOnForwardLitTranslucent", DisplayName = "Support Cloud Shadow On Forward Lit Translucent",
+		ToolTip = "Enable cloud shadow on translucent surface not relying on the translucenct lighting volume, e.g. using Forward lighting. This is evaluated per vertex to reduce GPU cost and requires extra samplers/textures to be bound to vertex shaders. This is not implemented on mobile as VolumetricClouds are not available on these platforms.",
+		ConfigRestartRequired = true))
+		uint32 bSupportCloudShadowOnForwardLitTranslucent : 1;
+
 	UPROPERTY(config, EditAnywhere, Category = Debugging, meta = (
 		ConsoleVariable = "r.GPUCrashDebugging", DisplayName = "Enable vendor specific GPU crash analysis tools",
 		ToolTip = "Enables vendor specific GPU crash analysis tools.",
@@ -914,34 +942,6 @@ class URendererSettings : public UDeveloperSettings
 		ConfigRestartRequired = true))
 		uint32 bSupportPointLightWholeSceneShadows : 1;
 
-	/**
-	"The sky atmosphere component requires extra samplers/textures to be bound to apply aerial perspective on transparent surfaces (and all surfaces on mobile via per vertex evaluation)."
-	*/
-	UPROPERTY(config, EditAnywhere, Category = ShaderPermutationReduction, meta = (
-		ConsoleVariable = "r.SupportSkyAtmosphere", DisplayName = "Support Sky Atmosphere",
-		ToolTip = "The sky atmosphere component requires extra samplers/textures to be bound to apply aerial perspective on transparent surfaces (and all surfaces on mobile via per vertex evaluation).",
-		ConfigRestartRequired = true))
-		uint32 bSupportSkyAtmosphere : 1;
-
-	/**
-	"The sky atmosphere component can light up the height fog but it requires extra samplers/textures to be bound to apply aerial perspective on transparent surfaces (and all surfaces on mobile via per vertex evaluation)."
-	"It requires r.SupportSkyAtmosphere to be true."
-	*/
-	UPROPERTY(config, EditAnywhere, Category = ShaderPermutationReduction, meta = (
-		ConsoleVariable = "r.SupportSkyAtmosphereAffectsHeightFog", DisplayName = "Support Sky Atmosphere Affecting Height Fog",
-		ToolTip = "The sky atmosphere component can light up the height fog but it requires extra samplers/textures to be bound to apply aerial perspective on transparent surfaces (and all surfaces on mobile via per vertex evaluation). It requires r.SupportSkyAtmosphere to be true.",
-		ConfigRestartRequired = true))
-		uint32 bSupportSkyAtmosphereAffectsHeightFog : 1;
-
-	/**
-	"Enable cloud shadow on translucent surface. This is evaluated per vertex to reduce GPU cost. The cloud system requires extra samplers/textures to be bound to vertex shaders."
-	*/
-	UPROPERTY(config, EditAnywhere, Category = ShaderPermutationReduction, meta = (
-		ConsoleVariable = "r.SupportCloudShadowOnForwardLitTranslucent", DisplayName = "Support Cloud Shadow On Forward Lit Translucent",
-		ToolTip = "Enable cloud shadow on translucent surface not relying on the translucenct lighting volume, e.g. using Forward lighting. This is evaluated per vertex to reduce GPU cost and requires extra samplers/textures to be bound to vertex shaders. This is not implemented on mobile as VolumetricClouds are not available on these platforms.",
-		ConfigRestartRequired = true))
-		uint32 bSupportCloudShadowOnForwardLitTranslucent : 1;
-
 	/**
 	""Enable translucent volumetric self-shadow, requires vertex and pixel shader permutations for all tranlucent materials even if not used by any light."
 	*/
