Moved remaining skylight render commands to the scene pipe. Since the proxy is deleted from the scene pipe all other commands must run on it.

diff --git a/Engine/Source/Runtime/Engine/Private/Components/SkyLightComponent.cpp b/Engine/Source/Runtime/Engine/Private/Components/SkyLightComponent.cpp
--- a/Engine/Source/Runtime/Engine/Private/Components/SkyLightComponent.cpp
+++ b/Engine/Source/Runtime/Engine/Private/Components/SkyLightComponent.cpp
@@ -255,9 +255,9 @@ FSkyLightSceneProxy::FSkyLightSceneProxy(const USkyLightComponent* InLightCompon
 	float InBlendFraction = InLightComponent->BlendFraction;
 	const FLinearColor* InSpecifiedCubemapColorScale = &InLightComponent->SpecifiedCubemapColorScale;
 	FSkyLightSceneProxy* LightSceneProxy = this;
-	ENQUEUE_RENDER_COMMAND(FInitSkyProxy)(
+	ENQUEUE_RENDER_COMMAND(FInitSkyProxy)(&UE::RenderCommandPipe::Scene,
 		[InIrradianceEnvironmentMap, BlendDestinationIrradianceEnvironmentMap, InAverageBrightness, 
-		BlendDestinationAverageBrightness, InBlendFraction, LightSceneProxy, InSpecifiedCubemapColorScale](FRHICommandList& RHICmdList)
+		BlendDestinationAverageBrightness, InBlendFraction, LightSceneProxy, InSpecifiedCubemapColorScale]
 		{
 			// Only access the irradiance maps on the RT, even though they belong to the USkyLightComponent, 
 			// Because FScene::UpdateSkyCaptureContents does not block the RT so the writes could still be in flight
@@ -464,8 +464,8 @@ void USkyLightComponent::UpdateLimitedRenderingStateFast()
 		FLinearColor InLightColor = FLinearColor(LightColor) * Intensity;
 		float InIndirectLightingIntensity = IndirectLightingIntensity;
 		float InVolumetricScatteringIntensity = VolumetricScatteringIntensity;
-		ENQUEUE_RENDER_COMMAND(FFastUpdateSkyLightCommand)(
-			[LightSceneProxy, InLightColor, InIndirectLightingIntensity, InVolumetricScatteringIntensity](FRHICommandList& RHICmdList)
+		ENQUEUE_RENDER_COMMAND(FFastUpdateSkyLightCommand)(&UE::RenderCommandPipe::Scene,
+			[LightSceneProxy, InLightColor, InIndirectLightingIntensity, InVolumetricScatteringIntensity]
 			{
 				LightSceneProxy->SetLightColor(InLightColor);
 				LightSceneProxy->IndirectLightingIntensity = InIndirectLightingIntensity;
@@ -483,8 +483,8 @@ void USkyLightComponent::UpdateOcclusionRenderingStateFast()
 		float InOcclusionExponent = OcclusionExponent;
 		float InMinOcclusion = MinOcclusion;
 		FColor InOcclusionTint = OcclusionTint;
-		ENQUEUE_RENDER_COMMAND(FFastUpdateSkyLightOcclusionCommand)(
-			[InLightSceneProxy, InContrast, InOcclusionExponent, InMinOcclusion, InOcclusionTint](FRHICommandList& RHICmdList)
+		ENQUEUE_RENDER_COMMAND(FFastUpdateSkyLightOcclusionCommand)(&UE::RenderCommandPipe::Scene,
+			[InLightSceneProxy, InContrast, InOcclusionExponent, InMinOcclusion, InOcclusionTint]
 			{
 				InLightSceneProxy->Contrast = InContrast;
 				InLightSceneProxy->OcclusionExponent = InOcclusionExponent;
@@ -521,8 +521,8 @@ void USkyLightComponent::SendRenderTransform_Concurrent()
 		FSkyLightSceneProxy* InLightSceneProxy = SceneProxy;
 		FVector Position = GetComponentTransform().GetLocation();
 
-		ENQUEUE_RENDER_COMMAND(UpdateSkyLightCapturePosition)(
-			[InLightSceneProxy, Position](FRHICommandListImmediate& RHICmdList)
+		ENQUEUE_RENDER_COMMAND(UpdateSkyLightCapturePosition)(&UE::RenderCommandPipe::Scene,
+			[InLightSceneProxy, Position]
 			{
 				InLightSceneProxy->CapturePosition = Position;
 			});
@@ -860,8 +860,8 @@ void USkyLightComponent::UpdateSkyCaptureContentsArray(UWorld* WorldToUpdate, TA
 #if WITH_EDITOR
 		CaptureComponent->SecondsSinceLastCapture += CaptureComponent->CaptureStatus == ESkyLightCaptureStatus::SLCS_CapturedButIncomplete ? WorldToUpdate->DeltaTimeSeconds : 0.0f;
 
-		ENQUEUE_RENDER_COMMAND(FUpdateSkyLightProxyStatusForcedCapture)(
-			[CaptureComponent, SecondsBetweenIncompleteCaptures, bCubemapSkyLightWaitingForCubemapAsset, bCaptureSkyLightWaitingCompiledShader, bCaptureSkyLightWaitingForMeshOrTexAssets](FRHICommandList& RHICmdList)
+		ENQUEUE_RENDER_COMMAND(FUpdateSkyLightProxyStatusForcedCapture)(&UE::RenderCommandPipe::Scene,
+			[CaptureComponent, SecondsBetweenIncompleteCaptures, bCubemapSkyLightWaitingForCubemapAsset, bCaptureSkyLightWaitingCompiledShader, bCaptureSkyLightWaitingForMeshOrTexAssets]
 			{
 				FSkyLightSceneProxy* SkyLightSceneProxy = CaptureComponent->SceneProxy;
 				if (SkyLightSceneProxy)
@@ -1028,8 +1028,8 @@ void USkyLightComponent::SetCubemapBlend(UTextureCube* SourceCubemap, UTextureCu
 				const float* InBlendDestinationAverageBrightness = &BlendDestinationAverageBrightness;
 				FSkyLightSceneProxy* LightSceneProxy = SceneProxy;
 				const FLinearColor* InSpecifiedCubemapColorScale = &SpecifiedCubemapColorScale;
-				ENQUEUE_RENDER_COMMAND(FUpdateSkyProxy)(
-					[InIrradianceEnvironmentMap, InBlendDestinationIrradianceEnvironmentMap, InAverageBrightness, InBlendDestinationAverageBrightness, InBlendFraction, LightSceneProxy, InSpecifiedCubemapColorScale](FRHICommandList& RHICmdList)
+				ENQUEUE_RENDER_COMMAND(FUpdateSkyProxy)(&UE::RenderCommandPipe::Scene,
+					[InIrradianceEnvironmentMap, InBlendDestinationIrradianceEnvironmentMap, InAverageBrightness, InBlendDestinationAverageBrightness, InBlendFraction, LightSceneProxy, InSpecifiedCubemapColorScale]
 					{
 						// Only access the irradiance maps on the RT, even though they belong to the USkyLightComponent, 
 						// Because FScene::UpdateSkyCaptureContents does not block the RT so the writes could still be in flight
