Remove MapCheck call from IsDataValid as it surfaces many less strict/usually-ignored issues.

diff --git a/Engine/Source/Runtime/Engine/Private/ActorEditor.cpp b/Engine/Source/Runtime/Engine/Private/ActorEditor.cpp
--- a/Engine/Source/Runtime/Engine/Private/ActorEditor.cpp
+++ b/Engine/Source/Runtime/Engine/Private/ActorEditor.cpp
@@ -1597,18 +1597,6 @@ EDataValidationResult AActor::IsDataValid(FDataValidationContext& Context) const
 		Context.AddError(ErrorMsg);
 	}
 
-	int32 OldNumMapWarningsAndErrors = FMessageLog("MapCheck").NumMessages(EMessageSeverity::Warning);
-	const_cast<AActor*>(this)->CheckForErrors();
-	int32 NewNumMapWarningsAndErrors = FMessageLog("MapCheck").NumMessages(EMessageSeverity::Warning);
-	if (NewNumMapWarningsAndErrors != OldNumMapWarningsAndErrors)
-	{
-		FFormatNamedArguments Arguments;
-		Arguments.Add(TEXT("ActorName"), FText::FromString(GetName()));
-		FText ErrorMsg = FText::Format(LOCTEXT("IsDataValid_Failed_CheckForErrors", "{ActorName} is not valid. See the MapCheck log messages for details."), Arguments);
-		Context.AddError(ErrorMsg);
-		bSuccess = false;
-	}
-
 	EDataValidationResult Result = bSuccess ? EDataValidationResult::Valid : EDataValidationResult::Invalid;
 
 	// check the components
