Validate detail level and editor-only status of attachments for scene components.

diff --git a/Engine/Source/Runtime/Engine/Classes/Components/SceneComponent.h b/Engine/Source/Runtime/Engine/Classes/Components/SceneComponent.h
--- a/Engine/Source/Runtime/Engine/Classes/Components/SceneComponent.h
+++ b/Engine/Source/Runtime/Engine/Classes/Components/SceneComponent.h
@@ -939,6 +939,7 @@ public:
 #endif
 
 #if WITH_EDITOR
+	ENGINE_API virtual EDataValidationResult IsDataValid(FDataValidationContext& Context) const override;
 	ENGINE_API virtual bool NeedsLoadForTargetPlatform(const ITargetPlatform* TargetPlatform) const;
 	ENGINE_API virtual void PostEditChangeProperty(FPropertyChangedEvent& PropertyChangedEvent) override;
 	ENGINE_API virtual void PostEditChangeChainProperty(FPropertyChangedChainEvent& PropertyChangedEvent) override;
diff --git a/Engine/Source/Runtime/Engine/Private/Components/SceneComponent.cpp b/Engine/Source/Runtime/Engine/Private/Components/SceneComponent.cpp
--- a/Engine/Source/Runtime/Engine/Private/Components/SceneComponent.cpp
+++ b/Engine/Source/Runtime/Engine/Private/Components/SceneComponent.cpp
@@ -30,7 +30,9 @@
 #include "Interfaces/ITargetPlatform.h"
 #include "DeviceProfiles/DeviceProfile.h"
 #include "Net/Core/PushModel/PushModel.h"
-#include  "UObject/ICookInfo.h"
+#include "UObject/Class.h"
+#include "UObject/ICookInfo.h"
+#include "Misc/DataValidation.h"
 
 #define LOCTEXT_NAMESPACE "SceneComponent"
 
@@ -491,7 +493,7 @@ bool USceneComponent::NeedsLoadForTargetPlatform(const ITargetPlatform* TargetPl
 		// Child not culled, so warn
 		if(!bDescendantsCulled)
 		{
-			UE_LOG(LogSceneComponent, Warning, TEXT("Component %s not cooked out for client because descendants were not also cooked out."), *GetPathName());
+			UE_LOG(LogSceneComponent, Warning, TEXT("Component %s not removed from client data because descendants were not also not removed."), *GetPathName());
 			return true;
 		}
 
@@ -500,6 +502,65 @@ bool USceneComponent::NeedsLoadForTargetPlatform(const ITargetPlatform* TargetPl
 
 	return true;
 }
+bool GValidateSceneComponentAttachmentEditorOnlySettings = true;
+static FAutoConsoleVariableRef CVarValidateSceneComponentAttachmentEditorOnlySettings (
+	TEXT("p.ValidateSceneComponentAttachmentEditorOnlySettings"),
+	GValidateSceneComponentAttachmentEditorOnlySettings,
+	TEXT("If enabled, checks that components which are editor only don't have attached components which are not editor only"),
+	ECVF_Default
+);
+
+bool GValidateSceneComponentAttachmentDetailLevel = true;
+static FAutoConsoleVariableRef CVarValidateSceneComponentAttachmentDetailLevel (
+	TEXT("p.ValidateSceneComponentAttachmentDetailLevel"),
+	GValidateSceneComponentAttachmentDetailLevel,
+	TEXT("If enabled, checks that components don't have attached components of higher detail level during data validation/"),
+	ECVF_Default
+);
+
+EDataValidationResult USceneComponent::IsDataValid(FDataValidationContext& Context) const
+{
+	EDataValidationResult Result = EDataValidationResult::Valid;
+	if (GValidateSceneComponentAttachmentEditorOnlySettings && IsEditorOnly())
+	{
+		for (USceneComponent* ChildSceneComponent : GetAttachChildren())
+		{
+			if (!ChildSceneComponent->IsEditorOnly())
+			{
+				Context.AddError(FText::Format(LOCTEXT("SceneComponent_AttachmentEditorOnlyMismatch",
+					"Component {0} is editor-only but it has an attached child {1} that is not"),
+					FText::FromString(GetPathName()),
+					FText::FromString(ChildSceneComponent->GetPathName())
+				));
+				Result = EDataValidationResult::Invalid;
+			}	
+		}
+	}
+	if (GValidateSceneComponentAttachmentDetailLevel)
+	{
+		for (USceneComponent* ChildSceneComponent : GetAttachChildren())
+		{
+			if (ChildSceneComponent->IsEditorOnly())
+			{
+				continue;
+			}
+
+			if (ChildSceneComponent->DetailMode < DetailMode)
+			{
+				Context.AddError(FText::Format(LOCTEXT("SceneComponent_AttachmentDetailLevelMismatch",
+					"Component {0} of detail level {1} cannot be removed because it has an attached child {2} of detail level {3}"),
+					FText::FromString(GetPathName()),
+					UEnum::GetDisplayValueAsText(DetailMode),
+					FText::FromString(ChildSceneComponent->GetPathName()),
+					UEnum::GetDisplayValueAsText(ChildSceneComponent->DetailMode)
+				));
+				Result = EDataValidationResult::Invalid;
+			}	
+		}
+	}
+	
+	return CombineDataValidationResults(Result, Super::IsDataValid(Context));
+}
 
 void USceneComponent::PostEditChangeProperty(FPropertyChangedEvent& PropertyChangedEvent)
 {
