Defer sending render dynamic data when the register component context is used. SendRenderDynamicData_Concurrent expects the scene proxy to be created, but when the context defers creation this is no longer the case. Proxies like cable component that require an initial set of dynamic data can render the first frame with bogus contents.

diff --git a/Engine/Plugins/FX/Niagara/Source/Niagara/Private/NiagaraComponent.cpp b/Engine/Plugins/FX/Niagara/Source/Niagara/Private/NiagaraComponent.cpp
--- a/Engine/Plugins/FX/Niagara/Source/Niagara/Private/NiagaraComponent.cpp
+++ b/Engine/Plugins/FX/Niagara/Source/Niagara/Private/NiagaraComponent.cpp
@@ -2151,9 +2151,10 @@ void UNiagaraComponent::OnEndOfFrameUpdateDuringTick()
 void UNiagaraComponent::CreateRenderState_Concurrent(FRegisterComponentContext* Context)
 {
 	Super::CreateRenderState_Concurrent(Context);
+
 	// The emitter instance may not tick again next frame so we send the dynamic data here so that the current state
 	// renders.  This can happen when while editing, or any time the age update mode is set to desired age.
-	SendRenderDynamicData_Concurrent();
+	FRegisterComponentContext::SendRenderDynamicData(Context, this);
 }
 
 void UNiagaraComponent::DestroyRenderState_Concurrent()
diff --git a/Engine/Plugins/Runtime/CableComponent/Source/CableComponent/Private/CableComponent.cpp b/Engine/Plugins/Runtime/CableComponent/Source/CableComponent/Private/CableComponent.cpp
--- a/Engine/Plugins/Runtime/CableComponent/Source/CableComponent/Private/CableComponent.cpp
+++ b/Engine/Plugins/Runtime/CableComponent/Source/CableComponent/Private/CableComponent.cpp
@@ -974,7 +974,7 @@ void UCableComponent::CreateRenderState_Concurrent(FRegisterComponentContext* Co
 {
 	Super::CreateRenderState_Concurrent(Context);
 
-	SendRenderDynamicData_Concurrent();
+	FRegisterComponentContext::SendRenderDynamicData(Context, this);
 }
 
 void UCableComponent::ApplyWorldOffset(const FVector & InOffset, bool bWorldShift)
diff --git a/Engine/Source/Runtime/Engine/Classes/Components/ActorComponent.h b/Engine/Source/Runtime/Engine/Classes/Components/ActorComponent.h
--- a/Engine/Source/Runtime/Engine/Classes/Components/ActorComponent.h
+++ b/Engine/Source/Runtime/Engine/Classes/Components/ActorComponent.h
@@ -19,6 +19,7 @@
 #include "UObject/StructOnScope.h"
 #include "PropertyPairsMap.h"
 #include "ComponentInstanceDataCache.h"
+#include "Experimental/ConcurrentLinearAllocator.h"
 #include "ActorComponent.generated.h"
 
 struct FTypedElementHandle;
@@ -48,12 +49,21 @@ public:
 		AddPrimitiveBatches.Add(PrimitiveComponent);
 	}
 
+	void AddSendRenderDynamicData(UPrimitiveComponent* PrimitiveComponent)
+	{
+		checkSlow(!SendRenderDynamicDataPrimitives.Contains(PrimitiveComponent));
+		SendRenderDynamicDataPrimitives.Add(PrimitiveComponent);
+	}
+
+	ENGINE_API static void SendRenderDynamicData(FRegisterComponentContext* Context, UPrimitiveComponent* PrimitiveComponent);
+
 	int32 Count() const { return AddPrimitiveBatches.Num(); }
 	void Process();
 
 private:
 	UWorld* World;
-	TArray<UPrimitiveComponent*> AddPrimitiveBatches;
+	TArray<UPrimitiveComponent*, FConcurrentLinearArrayAllocator> AddPrimitiveBatches;
+	TArray<UPrimitiveComponent*, FConcurrentLinearArrayAllocator> SendRenderDynamicDataPrimitives;
 };
 
 #if WITH_EDITOR
diff --git a/Engine/Source/Runtime/Engine/Classes/Components/PrimitiveComponent.h b/Engine/Source/Runtime/Engine/Classes/Components/PrimitiveComponent.h
--- a/Engine/Source/Runtime/Engine/Classes/Components/PrimitiveComponent.h
+++ b/Engine/Source/Runtime/Engine/Classes/Components/PrimitiveComponent.h
@@ -2344,6 +2344,7 @@ public:
 #endif
 
 	//~ Begin UActorComponent Interface
+	using Super::SendRenderDynamicData_Concurrent;
 	ENGINE_API virtual void CreateRenderState_Concurrent(FRegisterComponentContext* Context) override;
 	ENGINE_API virtual void SendRenderTransform_Concurrent() override;
 	ENGINE_API virtual void OnRegister()  override;
diff --git a/Engine/Source/Runtime/Engine/Private/Components/ActorComponent.cpp b/Engine/Source/Runtime/Engine/Private/Components/ActorComponent.cpp
--- a/Engine/Source/Runtime/Engine/Private/Components/ActorComponent.cpp
+++ b/Engine/Source/Runtime/Engine/Private/Components/ActorComponent.cpp
@@ -141,6 +141,24 @@ void FRegisterComponentContext::Process()
 		bSingleThreaded
 	);
 	AddPrimitiveBatches.Empty();
+
+	for (UPrimitiveComponent* Primitive : SendRenderDynamicDataPrimitives)
+	{
+		Primitive->SendRenderDynamicData_Concurrent();
+	}
+	SendRenderDynamicDataPrimitives.Empty();
+}
+
+void FRegisterComponentContext::SendRenderDynamicData(FRegisterComponentContext* Context, UPrimitiveComponent* PrimitiveComponent)
+{
+	if (Context)
+	{
+		Context->AddSendRenderDynamicData(PrimitiveComponent);
+	}
+	else
+	{
+		PrimitiveComponent->SendRenderDynamicData_Concurrent();
+	}
 }
 
 void UpdateAllPrimitiveSceneInfosForSingleComponent(UActorComponent* InComponent, TSet<FSceneInterface*>* InScenesToUpdateAllPrimitiveSceneInfosForBatching /* = nullptr*/)
