[Navigation] Use primitive bounds when NavCollision is empty

diff --git a/Engine/Source/Runtime/Engine/Private/Components/StaticMeshComponent.cpp b/Engine/Source/Runtime/Engine/Private/Components/StaticMeshComponent.cpp
--- a/Engine/Source/Runtime/Engine/Private/Components/StaticMeshComponent.cpp
+++ b/Engine/Source/Runtime/Engine/Private/Components/StaticMeshComponent.cpp
@@ -3089,7 +3089,12 @@ FBox UStaticMeshComponent::GetNavigationBounds() const
 
 		if (const UNavCollisionBase* NavCollision = GetStaticMesh()->GetNavCollision())
 		{
-			return NavCollision->GetBounds().TransformBy(GetComponentTransform());
+			FBox NavBounds = NavCollision->GetBounds();
+			if (!NavBounds.IsValid)
+			{
+				NavBounds = GetStaticMesh()->GetBounds().GetBox();
+			}
+			return NavBounds.TransformBy(GetComponentTransform());
 		}
 	}
 
diff --git a/Engine/Source/Runtime/Engine/Private/InstancedStaticMesh.cpp b/Engine/Source/Runtime/Engine/Private/InstancedStaticMesh.cpp
--- a/Engine/Source/Runtime/Engine/Private/InstancedStaticMesh.cpp
+++ b/Engine/Source/Runtime/Engine/Private/InstancedStaticMesh.cpp
@@ -5454,7 +5454,8 @@ FBox UInstancedStaticMeshComponent::GetInstanceNavigationBounds() const
 	if (const UStaticMesh* Mesh = GetStaticMesh())
 	{
 		const UNavCollisionBase* NavCollision = Mesh->GetNavCollision();
-		return NavCollision ? NavCollision->GetBounds() : Mesh->GetBounds().GetBox();
+		const FBox NavBounds = NavCollision ? NavCollision->GetBounds() : FBox();
+		return NavBounds.IsValid ? NavBounds : Mesh->GetBounds().GetBox();
 	}
 
 	return FBox();
diff --git a/Engine/Source/Runtime/NavigationSystem/Private/NavigationOctree.cpp b/Engine/Source/Runtime/NavigationSystem/Private/NavigationOctree.cpp
--- a/Engine/Source/Runtime/NavigationSystem/Private/NavigationOctree.cpp
+++ b/Engine/Source/Runtime/NavigationSystem/Private/NavigationOctree.cpp
@@ -145,7 +145,7 @@ void FNavigationOctree::AddNode(UObject* ElementOb, INavRelevantInterface* NavEl
 
 	if (!Bounds.IsValid)
 	{
-		UE_LOG(LogNavigation, Warning, TEXT("%hs: Ignoring %s because bounds are invalid."), __FUNCTION__, *GetNameSafe(ElementOb));
+		UE_LOG(LogNavigation, Warning, TEXT("%hs: Ignoring %s because bounds are invalid."), __FUNCTION__, *GetFullNameSafe(ElementOb));
 		return;
 	}
 	
