By default, do not create physics state for a BodyInstance that is supplied with a transform containing NaNs. The original behaviour (always create, regardless of NaNs) can be chosen by setting the console variable p.PreventInvalidBodyInstanceTransforms to 0, but this will almost certainly lead to issues like stalls in collision detection, disappearing objects, etc. An error will be logged regardless of the value of p.PreventInvalidBodyInstanceTransforms. FBodyInstance::SetBodyTransform behavior has not changed, so a transform containing NaNs will be rejected.

[RN] minor Simulation/Physics

diff --git a/Engine/Source/Runtime/Engine/Private/PhysicsEngine/BodyInstance.cpp b/Engine/Source/Runtime/Engine/Private/PhysicsEngine/BodyInstance.cpp
--- a/Engine/Source/Runtime/Engine/Private/PhysicsEngine/BodyInstance.cpp
+++ b/Engine/Source/Runtime/Engine/Private/PhysicsEngine/BodyInstance.cpp
@@ -78,6 +78,12 @@ TAutoConsoleVariable<int32> CVarIgnoreAnalyticCollisionsOverride(
 	ECVF_ReadOnly
 );
 
+bool bPreventInvalidBodyInstanceTransforms = true;
+FAutoConsoleVariableRef CVarbPreventInvalidBodyInstanceTransforms(
+	TEXT("p.PreventInvalidBodyInstanceTransforms"), 
+	bPreventInvalidBodyInstanceTransforms, 
+	TEXT("If true, an attempt to create a BodyInstance with an invalid transform will fail with a warning"));
+
 
 using namespace PhysicsInterfaceTypes;
 
@@ -1246,6 +1252,63 @@ UBodySetup* FBodyInstance::GetBodySetup() const
 	return nullptr;
 }
 
+const FString& GetBodyInstanceDebugName(FInitBodiesHelperBase& InitHelper)
+{
+	static FString NullName = TEXT("<NoName>");
+
+#if USE_BODYINSTANCE_DEBUG_NAMES
+	if (InitHelper.DebugName.IsValid())
+	{
+		return *InitHelper.DebugName.Get();
+	}
+#endif
+	return NullName;
+}
+
+bool ValidateTransformScale(const FTransform& Transform, const FString& DebugName)
+{
+	if (Transform.GetScale3D().IsNearlyZero())
+	{
+		UE_LOG(LogPhysics, Warning, TEXT("Initialising Body : Scale3D is (nearly) zero: %s"), *DebugName);
+		return false;
+	}
+
+	return true;
+}
+
+bool ValidateTransformMirror(const FTransform& Transform, const FString& DebugName, bool bGenerateMirroredCollision, bool bGenerateNonMirroredCollision)
+{
+	// Check we support mirroring/non-mirroring
+	const float TransformDet = Transform.GetDeterminant();
+	if (TransformDet < 0.f && !bGenerateMirroredCollision)
+	{
+		UE_LOG(LogPhysics, Warning, TEXT("Initialising Body : Body is mirrored but bGenerateMirroredCollision == false: %s"), *DebugName);
+		return false;
+	}
+
+	if (TransformDet > 0.f && !bGenerateNonMirroredCollision)
+	{
+		UE_LOG(LogPhysics, Warning, TEXT("Initialising Body : Body is not mirrored but bGenerateNonMirroredCollision == false: %s"), *DebugName);
+		return false;
+	}
+
+	return true;
+}
+
+bool ValidateTransformNaN(const FTransform& Transform, const FString& DebugName, const FName& BoneName)
+{
+#if !(UE_BUILD_SHIPPING)
+	if (Transform.ContainsNaN())
+	{
+		UE_LOG(LogPhysics, Warning, TEXT("Initialising Body : Bad transform - %s %s\n%s"), *DebugName, *BoneName.ToString(), *Transform.ToString());
+		return false;
+	}
+#endif
+
+	return true;
+}
+
+
 // Takes actor ref arrays.
 // #PHYS2 this used to return arrays of low-level physics bodies, which would be added to scene in InitBodies. Should it still do that, rather then later iterate over BodyInstances to get phys actor refs?
 bool FInitBodiesHelperBase::CreateShapesAndActors()
@@ -1262,12 +1325,25 @@ bool FInitBodiesHelperBase::CreateShapesAndActors()
 		FBodyInstance* Instance = Bodies[BodyIdx];
 		const FTransform& Transform = Transforms[BodyIdx];
 
-#if !USE_BODYINSTANCE_DEBUG_NAMES
-		FString DebugName;
-		FBodyInstance::ValidateTransform(Transform, DebugName, BodySetup);
-#else
-		FBodyInstance::ValidateTransform(Transform, *DebugName, BodySetup);
-#endif
+		// Log some warnings for unexpected transforms, but treat NaNs as errors
+		const FString& SafeDebugName = GetBodyInstanceDebugName(*this);
+		ValidateTransformScale(Transform, SafeDebugName);
+		ValidateTransformMirror(Transform, SafeDebugName, BodySetup->bGenerateMirroredCollision, BodySetup->bGenerateNonMirroredCollision);
+		const bool bValidTransform = ValidateTransformNaN(Transform, SafeDebugName, BodySetup->BoneName);
+		if (!bValidTransform)
+		{
+			if (bPreventInvalidBodyInstanceTransforms)
+			{
+				// NaNs are errors and we don't create the physics state
+				UE_LOG(LogPhysics, Error, TEXT("Rejecting BodyInstance %d on %s with an invalid transform"), BodyIdx , *SafeDebugName);
+				return false;
+			}
+			else
+			{
+				// NaNs are errors but we still create the physics state. This will almost certainly cause problems later
+				UE_LOG(LogPhysics, Error, TEXT("Creating a BodyInstance %d on %s with an invalid transform which will likely lead to severe performance and behavioural problems in Physics"), BodyIdx , *SafeDebugName);
+			}
+		}
 
 		Instance->OwnerComponent = PrimitiveComp;
 		Instance->BodySetup = BodySetup;
@@ -3968,35 +4044,9 @@ void FBodyInstance::ApplyMaterialToInstanceShapes_AssumesLocked(UPhysicalMateria
 
 bool FBodyInstance::ValidateTransform(const FTransform &Transform, const FString& DebugName, const UBodySetup* Setup)
 {
-	if(Transform.GetScale3D().IsNearlyZero())
-	{
-		UE_LOG(LogPhysics, Warning, TEXT("Initialising Body : Scale3D is (nearly) zero: %s"), *DebugName);
-		return false;
-	}
-
-	// Check we support mirroring/non-mirroring
-	const float TransformDet = Transform.GetDeterminant();
-	if(TransformDet < 0.f && !Setup->bGenerateMirroredCollision)
-	{
-		UE_LOG(LogPhysics, Warning, TEXT("Initialising Body : Body is mirrored but bGenerateMirroredCollision == false: %s"), *DebugName);
-		return false;
-	}
-
-	if(TransformDet > 0.f && !Setup->bGenerateNonMirroredCollision)
-	{
-		UE_LOG(LogPhysics, Warning, TEXT("Initialising Body : Body is not mirrored but bGenerateNonMirroredCollision == false: %s"), *DebugName);
-		return false;
-	}
-
-#if !(UE_BUILD_SHIPPING)
-	if(Transform.ContainsNaN())
-	{
-		UE_LOG(LogPhysics, Warning, TEXT("Initialising Body : Bad transform - %s %s\n%s"), *DebugName, *Setup->BoneName.ToString(), *Transform.ToString());
-		return false;
-	}
-#endif
-
-	return true;
+	return ValidateTransformScale(Transform, DebugName) 
+		&& ValidateTransformMirror(Transform, DebugName, Setup->bGenerateMirroredCollision, Setup->bGenerateNonMirroredCollision)
+		&& ValidateTransformNaN(Transform, DebugName, Setup->BoneName);
 }
 
 void FBodyInstance::InitDynamicProperties_AssumesLocked()
