Fix edge case when reading to end of data when decoding streamed audio.

diff --git a/Engine/Source/Runtime/Engine/Private/AudioDecompress.cpp b/Engine/Source/Runtime/Engine/Private/AudioDecompress.cpp
--- a/Engine/Source/Runtime/Engine/Private/AudioDecompress.cpp
+++ b/Engine/Source/Runtime/Engine/Private/AudioDecompress.cpp
@@ -62,7 +62,7 @@ bool IStreamedCompressedInfo::ReadCompressedInfo(const uint8* InSrcBufferData, u
 	check(SrcBufferData != nullptr);
 
 	// Sample Stride is 
-	SampleStride = NumChannels * sizeof(int16);
+	SampleStride = NumChannels * MONO_PCM_SAMPLE_SIZE;
 
 	MaxFrameSizeSamples = GetMaxFrameSizeSamples();
 		
@@ -171,8 +171,7 @@ void IStreamedCompressedInfo::ExpandFile(uint8* DstBuffer, struct FSoundQualityI
 		}
 		else
 		{
-			IncrementCurrentSampleCount(DecodedFrames * NumChannels);
-			LastPCMByteSize = DecodedFrames * SampleStride;
+			LastPCMByteSize = IncrementCurrentSampleCount(DecodedFrames * NumChannels) * MONO_PCM_SAMPLE_SIZE;
 			RawPCMOffset += WriteFromDecodedPCM(DstBuffer + RawPCMOffset, QualityInfo->SampleDataSize - RawPCMOffset);
 		}
 	}
@@ -417,8 +416,7 @@ bool IStreamedCompressedInfo::StreamCompressedData(uint8* Destination, bool bLoo
 		}
 		else
 		{
-			IncrementCurrentSampleCount(DecodedFrames * NumChannels);
-			LastPCMByteSize = DecodedFrames * SampleStride;
+			LastPCMByteSize = IncrementCurrentSampleCount(DecodedFrames * NumChannels) * MONO_PCM_SAMPLE_SIZE;
 
 			// update OutNumBytesStreamed as we write out data
 			RawPCMOffset += WriteFromDecodedPCM(Destination + RawPCMOffset, BufferSize - RawPCMOffset);
@@ -929,7 +927,7 @@ void FAsyncAudioDecompressWorker::DoWork()
 			{
 				LLM_SCOPE(ELLMTag::AudioRealtimePrecache);
 #if PLATFORM_NUM_AUDIODECOMPRESSION_PRECACHE_BUFFERS > 0
-				const uint32 PCMBufferSize = NumPrecacheFrames * sizeof(int16) * Wave->NumChannels * PLATFORM_NUM_AUDIODECOMPRESSION_PRECACHE_BUFFERS;
+				const uint32 PCMBufferSize = NumPrecacheFrames * MONO_PCM_SAMPLE_SIZE * Wave->NumChannels * PLATFORM_NUM_AUDIODECOMPRESSION_PRECACHE_BUFFERS;
 				Wave->NumPrecacheFrames = NumPrecacheFrames;
 				if (Wave->CachedRealtimeFirstBuffer == nullptr)
 				{
diff --git a/Engine/Source/Runtime/Engine/Public/AudioDecompress.h b/Engine/Source/Runtime/Engine/Public/AudioDecompress.h
--- a/Engine/Source/Runtime/Engine/Public/AudioDecompress.h
+++ b/Engine/Source/Runtime/Engine/Public/AudioDecompress.h
@@ -17,8 +17,9 @@
 // 100ms of 48KHz data
 // 108ms of 44.1KHz data
 // 218ms of 22KHz data
-#define MONO_PCM_BUFFER_SAMPLES		(4800)
-#define MONO_PCM_BUFFER_SIZE		( MONO_PCM_BUFFER_SAMPLES * sizeof( int16 ) )
+constexpr int32 MONO_PCM_BUFFER_SAMPLES = 4800;
+constexpr uint32 MONO_PCM_SAMPLE_SIZE = sizeof(int16);
+constexpr uint32 MONO_PCM_BUFFER_SIZE = MONO_PCM_BUFFER_SAMPLES * MONO_PCM_SAMPLE_SIZE;
 
 struct FSoundQualityInfo;
 class FStreamedAudioChunkSeekTable;
