#UE Fix shipping editor

diff --git a/Engine/Source/Runtime/Engine/Private/NetDriver.cpp b/Engine/Source/Runtime/Engine/Private/NetDriver.cpp
--- a/Engine/Source/Runtime/Engine/Private/NetDriver.cpp
+++ b/Engine/Source/Runtime/Engine/Private/NetDriver.cpp
@@ -5719,7 +5719,7 @@ void UNetDriver::DrawNetDriverDebug()
 
 		// Draw NetGUID or NetHandle
 		FVector AboveActor(0.f, 0.f, 32.f);
-#if UE_WITH_IRIS
+#if UE_WITH_IRIS && !UE_BUILD_SHIPPING
 		if (Bridge)
 		{
 			AActor* Actor = *It;
@@ -5756,7 +5756,7 @@ void UNetDriver::DrawNetDriverDebug()
 	}
 
 	// Draw state
-#if UE_WITH_IRIS
+#if UE_WITH_IRIS && !UE_BUILD_SHIPPING
 	if (Bridge)
 	{
 		UE::Net::FNetRefHandle DebugRefHandle = UE::Net::IrisDebugHelper::GetDebugNetRefHandle();
