Fix Substrate SSS MFP not working due to incorrect shading mode conversion.

diff --git a/Engine/Source/Runtime/Engine/Private/Materials/Material.cpp b/Engine/Source/Runtime/Engine/Private/Materials/Material.cpp
--- a/Engine/Source/Runtime/Engine/Private/Materials/Material.cpp
+++ b/Engine/Source/Runtime/Engine/Private/Materials/Material.cpp
@@ -5026,7 +5026,7 @@ void UMaterial::RebuildShadingModelField()
 				MaterialDomain = EMaterialDomain::MD_Surface;
 				ShadingModel = MSM_DefaultLit;
 			}
-			else if (SubstrateMaterialInfo.HasOnlyShadingModel(SSM_SubsurfaceProfile))
+			else if (SubstrateMaterialInfo.HasOnlyShadingModel(SSM_SubsurfaceProfile) || SubstrateMaterialInfo.HasOnlyShadingModel(SSM_SubsurfaceMFP))
 			{
 				MaterialDomain = EMaterialDomain::MD_Surface;
 				ShadingModel = MSM_SubsurfaceProfile;
