Fix -rhivalidation assert "RHI Resource Transition Validation Error" on SkinCachePositions & SkinCacheTangents resources. It happens when the same sk mesh section is inserted into skin cache batch dispatch queue more than once.

diff --git a/Engine/Source/Runtime/Engine/Private/GPUSkinCache.cpp b/Engine/Source/Runtime/Engine/Private/GPUSkinCache.cpp
--- a/Engine/Source/Runtime/Engine/Private/GPUSkinCache.cpp
+++ b/Engine/Source/Runtime/Engine/Private/GPUSkinCache.cpp
@@ -1856,7 +1856,22 @@ bool FGPUSkinCache::ProcessEntry(
 	if (bShouldBatchDispatches)
 	{
 		InOutEntry->bQueuedForDispatch = true;
-		BatchDispatches.Add({ InOutEntry, RevisionNumber, uint32(Section) });
+
+		bool bFoundEntry = false;
+		for (FDispatchEntry& Entry : BatchDispatches)
+		{
+			// Check if the combo of skin cache entry and section index already exists, if so use the entry and update to latest revision number.
+			if (Entry.SkinCacheEntry == InOutEntry && Entry.Section == Section && Entry.RevisionNumber < RevisionNumber)
+			{
+				Entry.RevisionNumber = RevisionNumber;
+				bFoundEntry = true;
+				break;
+			}
+		}
+		if (!bFoundEntry)
+		{
+			BatchDispatches.Add({ InOutEntry, RevisionNumber, uint32(Section) });
+		}
 	}
 	else
 	{
