Landscape cleanup :
* Removed deprecated code from <=UE5.2
* Moved UHierarchicalInstancedStaticMeshComponent::AcceptPrebuiltTree and UHierarchicalInstancedStaticMeshComponent::BuildTreeAnyThread to UGrassInstancedStaticMeshComponent

diff --git a/Engine/Source/Editor/LandscapeEditor/Private/NewLandscapeUtils.cpp b/Engine/Source/Editor/LandscapeEditor/Private/NewLandscapeUtils.cpp
deleted file mode 100644
diff --git a/Engine/Source/Editor/LandscapeEditor/Public/LandscapeEditorObject.h b/Engine/Source/Editor/LandscapeEditor/Public/LandscapeEditorObject.h
--- a/Engine/Source/Editor/LandscapeEditor/Public/LandscapeEditorObject.h
+++ b/Engine/Source/Editor/LandscapeEditor/Public/LandscapeEditorObject.h
@@ -127,17 +127,6 @@ enum class ELandscapeConvertMode : int8
 	Resample = 2,
 };
 
-namespace EColorChannel
-{
-	enum UE_DEPRECATED(5.2, "Use ELandscapeTextureColorChannel") Type : int
-	{
-		Red,
-		Green,
-		Blue,
-		Alpha,
-	};
-}
-
 UENUM()
 enum class ELandscapeTextureColorChannel : int32
 {
diff --git a/Engine/Source/Editor/LandscapeEditor/Public/NewLandscapeUtils.h b/Engine/Source/Editor/LandscapeEditor/Public/NewLandscapeUtils.h
deleted file mode 100644
diff --git a/Engine/Source/Runtime/Engine/Classes/Components/HierarchicalInstancedStaticMeshComponent.h b/Engine/Source/Runtime/Engine/Classes/Components/HierarchicalInstancedStaticMeshComponent.h
--- a/Engine/Source/Runtime/Engine/Classes/Components/HierarchicalInstancedStaticMeshComponent.h
+++ b/Engine/Source/Runtime/Engine/Classes/Components/HierarchicalInstancedStaticMeshComponent.h
@@ -11,7 +11,6 @@
 
 #include "HierarchicalInstancedStaticMeshComponent.generated.h"
 
-class FClusterBuilder;
 class FStaticLightingTextureMapping_InstancedStaticMesh;
 
 
@@ -205,11 +204,89 @@ class UHierarchicalInstancedStaticMeshComponent : public UInstancedStaticMeshCom
 	bool bCanEnableDensityScaling : 1;
 #endif
 
+public:
+	struct FClusterTree
+	{
+		TArray<FClusterNode> Nodes;
+		TArray<int32> SortedInstances;
+		TArray<int32> InstanceReorderTable;
+		int32 OutOcclusionLayerNum = 0;
+
+		bool PrintLevel(int32 NodeIndex, int32 Level, int32 CurrentLevel, int32 Parent);
+	};
+
+	class FClusterBuilder
+	{
+	public:
+		ENGINE_API FClusterBuilder(TArray<FMatrix> InTransforms, TArray<float> InCustomDataFloats, int32 InNumCustomDataFloats, const FBox& InInstBox, int32 InMaxInstancesPerLeaf, float InDensityScaling, int32 InInstancingRandomSeed, bool InGenerateInstanceScalingRange);
+		ENGINE_API void BuildTreeAndBufferAsync(ENamedThreads::Type CurrentThread, const FGraphEventRef& MyCompletionGraphEvent);
+		ENGINE_API void BuildTreeAndBuffer();
+		ENGINE_API void BuildTree();
+
+	protected:
+		void Split(int32 InNum);
+		void Split(int32 Start, int32 End);
+		void BuildInstanceBuffer();
+		void Init();
+
+	public:
+		TUniquePtr<FClusterTree> Result;
+		TUniquePtr<FStaticMeshInstanceData> BuiltInstanceData;
+
+	protected:
+		int32 OriginalNum;
+		int32 Num;
+		FBox InstBox;
+		int32 BranchingFactor;
+		int32 InternalNodeBranchingFactor;
+		int32 OcclusionLayerTarget;
+		int32 MaxInstancesPerLeaf;
+		int32 NumRoots;
+
+		int32 InstancingRandomSeed;
+		float DensityScaling;
+		bool GenerateInstanceScalingRange;
+
+		TArray<int32> SortIndex;
+		TArray<FVector> SortPoints;
+		TArray<FMatrix> Transforms;
+		TArray<float> CustomDataFloats;
+		int32 NumCustomDataFloats;
+
+		struct FRunPair
+		{
+			int32 Start;
+			int32 Num;
+
+			FRunPair(int32 InStart, int32 InNum)
+				: Start(InStart)
+				, Num(InNum)
+			{
+			}
+
+			bool operator< (const FRunPair& Other) const
+			{
+				return Start < Other.Start;
+			}
+		};
+		TArray<FRunPair> Clusters;
+
+		struct FSortPair
+		{
+			float d;
+			int32 Index;
+
+			bool operator< (const FSortPair& Other) const
+			{
+				return d < Other.d;
+			}
+		};
+		TArray<FSortPair> SortPairs;
+	};
+
 	// Apply the results of the async build
 	ENGINE_API void ApplyBuildTreeAsync(ENamedThreads::Type CurrentThread, const FGraphEventRef& MyCompletionGraphEvent, TSharedRef<FClusterBuilder, ESPMode::ThreadSafe> Builder, double StartTime);
 
-public:
-
 	//Begin UObject Interface
 	ENGINE_API virtual void Serialize(FArchive& Ar) override;
 	ENGINE_API virtual void GetResourceSizeEx(FResourceSizeEx& CumulativeResourceSize) override;
@@ -253,10 +330,6 @@ public:
 	ENGINE_API void GetOverlappingBoxTransforms(const FBox& Box, TArray<FTransform>& OutTransforms) const;
 
 	ENGINE_API bool BuildTreeIfOutdated(bool Async, bool ForceUpdate);
-	UE_DEPRECATED(5.1, "The BuildTreeAnyThread method is moving to UGrassInstancedStaticMeshComponent. Please update your project to use the new component and method or your project may not compile in the next udpate.")
-	static ENGINE_API void BuildTreeAnyThread(TArray<FMatrix>& InstanceTransforms, TArray<float>& InstanceCustomDataFloats, int32 NumCustomDataFloats, const FBox& MeshBox, TArray<FClusterNode>& OutClusterTree, TArray<int32>& OutSortedInstances, TArray<int32>& OutInstanceReorderTable, int32& OutOcclusionLayerNum, int32 MaxInstancesPerLeaf, bool InGenerateInstanceScalingRange);
-	UE_DEPRECATED(5.1, "The AcceptPrebuiltTree method is moving to UGrassInstancedStaticMeshComponent. Please update your project to use the new component and method or your project may not compile in the next udpate.")
-	ENGINE_API void AcceptPrebuiltTree(TArray<FInstancedStaticMeshInstanceData>& InInstanceData, TArray<FClusterNode>& InClusterTree, int32 InOcclusionLayerNumNodes, int32 InNumBuiltRenderInstances);
 	bool IsAsyncBuilding() const { return bIsAsyncBuilding; }
 	bool IsTreeFullyBuilt() const { return !bIsOutOfDate; }
 	ENGINE_API void GetTree(TArray<FClusterNode>& OutClusterTree) const;
@@ -304,6 +377,8 @@ protected:
 
 	virtual FVector GetTranslatedInstanceSpaceOrigin() const override { return TranslatedInstanceSpaceOrigin; }
 
+	ENGINE_API static FBox GetClusterTreeBounds(TArray<FClusterNode> const& InClusterTree, const FVector& InOffset);
+
 	ENGINE_API virtual void GetNavigationPerInstanceTransforms(const FBox& AreaBox, TArray<FTransform>& InstanceData) const override;
 	ENGINE_API virtual void PartialNavigationUpdate(int32 InstanceIdx) override;
 	virtual bool SupportsPartialNavigationUpdate() const override { return true; }
@@ -321,6 +396,6 @@ protected:
 	friend FStaticLightingTextureMapping_InstancedStaticMesh;
 	friend FInstancedLightMap2D;
 	friend FInstancedShadowMap2D;
-	friend class FClusterBuilder;
+	friend FClusterBuilder;
 };
 
diff --git a/Engine/Source/Runtime/Engine/Private/HierarchicalInstancedStaticMesh.cpp b/Engine/Source/Runtime/Engine/Private/HierarchicalInstancedStaticMesh.cpp
--- a/Engine/Source/Runtime/Engine/Private/HierarchicalInstancedStaticMesh.cpp
+++ b/Engine/Source/Runtime/Engine/Private/HierarchicalInstancedStaticMesh.cpp
@@ -190,523 +190,460 @@ static void FoliageCVarSinkFunction()
 
 static FAutoConsoleVariableSink CVarFoliageSink(FConsoleCommandDelegate::CreateStatic(&FoliageCVarSinkFunction));
 
-struct FClusterTree
-{
-	TArray<FClusterNode> Nodes;
-	TArray<int32> SortedInstances;
-	TArray<int32> InstanceReorderTable;
-	int32 OutOcclusionLayerNum = 0;
-};
 
-class FClusterBuilder
-{
-protected:
-	int32 OriginalNum;
-	int32 Num;
-	FBox InstBox;
-	int32 BranchingFactor;
-	int32 InternalNodeBranchingFactor;
-	int32 OcclusionLayerTarget;
-	int32 MaxInstancesPerLeaf;
-	int32 NumRoots;
-	
-	int32 InstancingRandomSeed;
-	float DensityScaling;
-	bool GenerateInstanceScalingRange;
+// ----------------------------------------------------------------------------------
 
-	TArray<int32> SortIndex;
-	TArray<FVector> SortPoints;
-	TArray<FMatrix> Transforms;
-	TArray<float> CustomDataFloats;
-	int32 NumCustomDataFloats;
+UHierarchicalInstancedStaticMeshComponent::FClusterBuilder::FClusterBuilder(TArray<FMatrix> InTransforms, TArray<float> InCustomDataFloats, int32 InNumCustomDataFloats, const FBox& InInstBox, int32 InMaxInstancesPerLeaf, float InDensityScaling, int32 InInstancingRandomSeed, bool InGenerateInstanceScalingRange)
+	: OriginalNum(InTransforms.Num())
+	, InstBox(InInstBox)
+	, MaxInstancesPerLeaf(InMaxInstancesPerLeaf)
+	, InstancingRandomSeed(InInstancingRandomSeed)
+	, DensityScaling(InDensityScaling)
+	, GenerateInstanceScalingRange(InGenerateInstanceScalingRange)
+	, Transforms(MoveTemp(InTransforms))
+	, CustomDataFloats(MoveTemp(InCustomDataFloats))
+	, NumCustomDataFloats(InNumCustomDataFloats)
+{
+}
 
-	struct FRunPair
+void UHierarchicalInstancedStaticMeshComponent::FClusterBuilder::Split(int32 InNum)
+{
+	checkSlow(InNum);
+	Clusters.Reset();
+	Split(0, InNum - 1);
+	Clusters.Sort();
+	checkSlow(Clusters.Num() > 0);
+	int32 At = 0;
+	for (auto& Cluster : Clusters)
 	{
-		int32 Start;
-		int32 Num;
-
-		FRunPair(int32 InStart, int32 InNum)
-			: Start(InStart)
-			, Num(InNum)
-		{
-		}
-
-		bool operator< (const FRunPair& Other) const
-		{
-			return Start < Other.Start;
-		}
-	};
-	TArray<FRunPair> Clusters;
+		checkSlow(At == Cluster.Start);
+		At += Cluster.Num;
+	}
+	checkSlow(At == InNum);
+}
 
-	struct FSortPair
+void UHierarchicalInstancedStaticMeshComponent::FClusterBuilder::Split(int32 Start, int32 End)
+{
+	int32 NumRange = 1 + End - Start;
+	FBox ClusterBounds(ForceInit);
+	for (int32 Index = Start; Index <= End; Index++)
 	{
-		float d;
-		int32 Index;
-
-		bool operator< (const FSortPair& Other) const
-		{
-			return d < Other.d;
-		}
-	};
-	TArray<FSortPair> SortPairs;
-
-	void Split(int32 InNum)
+		ClusterBounds += SortPoints[SortIndex[Index]];
+	}
+	if (NumRange <= BranchingFactor)
 	{
-		checkSlow(InNum);
-		Clusters.Reset();
-		Split(0, InNum - 1);
-		Clusters.Sort();
-		checkSlow(Clusters.Num() > 0);
-		int32 At = 0;
-		for (auto& Cluster : Clusters)
-		{
-			checkSlow(At == Cluster.Start);
-			At += Cluster.Num;
-		}
-		checkSlow(At == InNum);
+		Clusters.Add(FRunPair(Start, NumRange));
+		return;
 	}
-
-	void Split(int32 Start, int32 End)
+	checkSlow(NumRange >= 2);
+	SortPairs.Reset();
+	int32 BestAxis = -1;
+	float BestAxisValue = -1.0f;
+	for (int32 Axis = 0; Axis < 3; Axis++)
 	{
-		int32 NumRange = 1 + End - Start;
-		FBox ClusterBounds(ForceInit);
-		for (int32 Index = Start; Index <= End; Index++)
-		{
-			ClusterBounds += SortPoints[SortIndex[Index]];
-		}
-		if (NumRange <= BranchingFactor)
+		float ThisAxisValue = ClusterBounds.Max[Axis] - ClusterBounds.Min[Axis];
+		if (!Axis || ThisAxisValue > BestAxisValue)
 		{
-			Clusters.Add(FRunPair(Start, NumRange));
-			return;
+			BestAxis = Axis;
+			BestAxisValue = ThisAxisValue;
 		}
-		checkSlow(NumRange >= 2);
-		SortPairs.Reset();
-		int32 BestAxis = -1;
-		float BestAxisValue = -1.0f;
-		for (int32 Axis = 0; Axis < 3; Axis++)
-		{
-			float ThisAxisValue = ClusterBounds.Max[Axis] - ClusterBounds.Min[Axis];
-			if (!Axis || ThisAxisValue > BestAxisValue)
-			{
-				BestAxis = Axis;
-				BestAxisValue = ThisAxisValue;
-			}
-		}
-		for (int32 Index = Start; Index <= End; Index++)
-		{
-			FSortPair Pair;
+	}
+	for (int32 Index = Start; Index <= End; Index++)
+	{
+		FSortPair Pair;
 
-			Pair.Index = SortIndex[Index];
-			Pair.d = SortPoints[Pair.Index][BestAxis];
-			SortPairs.Add(Pair);
-		}
-		SortPairs.Sort();
-		for (int32 Index = Start; Index <= End; Index++)
-		{
-			SortIndex[Index] = SortPairs[Index - Start].Index;
-		}
+		Pair.Index = SortIndex[Index];
+		Pair.d = SortPoints[Pair.Index][BestAxis];
+		SortPairs.Add(Pair);
+	}
+	SortPairs.Sort();
+	for (int32 Index = Start; Index <= End; Index++)
+	{
+		SortIndex[Index] = SortPairs[Index - Start].Index;
+	}
 
-		int32 Half = NumRange / 2;
+	int32 Half = NumRange / 2;
 
-		int32 EndLeft = Start + Half - 1;
-		int32 StartRight = 1 + End - Half;
+	int32 EndLeft = Start + Half - 1;
+	int32 StartRight = 1 + End - Half;
 
-		if (NumRange & 1)
+	if (NumRange & 1)
+	{
+		if (SortPairs[Half].d - SortPairs[Half - 1].d < SortPairs[Half + 1].d - SortPairs[Half].d)
 		{
-			if (SortPairs[Half].d - SortPairs[Half - 1].d < SortPairs[Half + 1].d - SortPairs[Half].d)
-			{
-				EndLeft++;
-			}
-			else
-			{
-				StartRight--;
-			}
+			EndLeft++;
+		}
+		else
+		{
+			StartRight--;
 		}
-		checkSlow(EndLeft + 1 == StartRight);
-		checkSlow(EndLeft >= Start);
-		checkSlow(End >= StartRight);
-
-		Split(Start, EndLeft);
-		Split(StartRight, End);
 	}
+	checkSlow(EndLeft + 1 == StartRight);
+	checkSlow(EndLeft >= Start);
+	checkSlow(End >= StartRight);
 
-	void BuildInstanceBuffer()
-	{
-		// build new instance buffer
-		FRandomStream RandomStream = FRandomStream(InstancingRandomSeed);
-		BuiltInstanceData = MakeUnique<FStaticMeshInstanceData>(/*bInUseHalfFloat = */true);
+	Split(Start, EndLeft);
+	Split(StartRight, End);
+}
+
+void UHierarchicalInstancedStaticMeshComponent::FClusterBuilder::BuildInstanceBuffer()
+{
+	// build new instance buffer
+	FRandomStream RandomStream = FRandomStream(InstancingRandomSeed);
+	BuiltInstanceData = MakeUnique<FStaticMeshInstanceData>(/*bInUseHalfFloat = */true);
 		
-		int32 NumInstances = Result->InstanceReorderTable.Num();
-		int32 NumRenderInstances = Result->SortedInstances.Num();
+	int32 NumInstances = Result->InstanceReorderTable.Num();
+	int32 NumRenderInstances = Result->SortedInstances.Num();
 		
-		if (NumRenderInstances > 0)
-		{
-			BuiltInstanceData->AllocateInstances(NumRenderInstances, NumCustomDataFloats, GIsEditor ? EResizeBufferFlags::AllowSlackOnGrow|EResizeBufferFlags::AllowSlackOnReduce : EResizeBufferFlags::None, false); // In Editor always permit overallocation, to prevent too much realloc
+	if (NumRenderInstances > 0)
+	{
+		BuiltInstanceData->AllocateInstances(NumRenderInstances, NumCustomDataFloats, GIsEditor ? EResizeBufferFlags::AllowSlackOnGrow|EResizeBufferFlags::AllowSlackOnReduce : EResizeBufferFlags::None, false); // In Editor always permit overallocation, to prevent too much realloc
 
-			FVector2D LightmapUVBias = FVector2D(-1.0f, -1.0f);
-			FVector2D ShadowmapUVBias = FVector2D(-1.0f, -1.0f);
+		FVector2D LightmapUVBias = FVector2D(-1.0f, -1.0f);
+		FVector2D ShadowmapUVBias = FVector2D(-1.0f, -1.0f);
 
-			// we loop over all instances to ensure that render instances will get same RandomID regardless of density settings
-			for (int32 i = 0; i < NumInstances; ++i)
+		// we loop over all instances to ensure that render instances will get same RandomID regardless of density settings
+		for (int32 i = 0; i < NumInstances; ++i)
+		{
+			int32 RenderIndex = Result->InstanceReorderTable[i];
+			float RandomID = RandomStream.GetFraction();
+			if (RenderIndex >= 0)
 			{
-				int32 RenderIndex = Result->InstanceReorderTable[i];
-				float RandomID = RandomStream.GetFraction();
-				if (RenderIndex >= 0)
+				// LWC_TODO: Precision loss here has been compensated for by use of TranslatedInstanceSpaceOrigin.
+				BuiltInstanceData->SetInstance(RenderIndex, FMatrix44f(Transforms[i]), RandomID, LightmapUVBias, ShadowmapUVBias);
+				for (int32 DataIndex = 0; DataIndex < NumCustomDataFloats; ++DataIndex)
 				{
-					// LWC_TODO: Precision loss here has been compensated for by use of TranslatedInstanceSpaceOrigin.
-					BuiltInstanceData->SetInstance(RenderIndex, FMatrix44f(Transforms[i]), RandomID, LightmapUVBias, ShadowmapUVBias);
-					for (int32 DataIndex = 0; DataIndex < NumCustomDataFloats; ++DataIndex)
-					{
-						BuiltInstanceData->SetInstanceCustomData(RenderIndex, DataIndex, CustomDataFloats[NumCustomDataFloats * i + DataIndex]);
-					}
+					BuiltInstanceData->SetInstanceCustomData(RenderIndex, DataIndex, CustomDataFloats[NumCustomDataFloats * i + DataIndex]);
 				}
-				// correct light/shadow map bias will be setup on game thread side if needed
 			}
+			// correct light/shadow map bias will be setup on game thread side if needed
 		}
 	}
+}
 
-	void Init()
-	{
-		SortIndex.Empty();
-		SortPoints.SetNumUninitialized(OriginalNum);
+void UHierarchicalInstancedStaticMeshComponent::FClusterBuilder::Init()
+{
+	SortIndex.Empty();
+	SortPoints.SetNumUninitialized(OriginalNum);
 					
-		FRandomStream DensityRand = FRandomStream(InstancingRandomSeed);
+	FRandomStream DensityRand = FRandomStream(InstancingRandomSeed);
 
-		SortIndex.Empty(OriginalNum*DensityScaling);
+	SortIndex.Empty(OriginalNum*DensityScaling);
 
-		for (int32 Index = 0; Index < OriginalNum; Index++)
-		{
-			SortPoints[Index] = Transforms[Index].GetOrigin();
-
-			if (DensityScaling < 1.0f && DensityRand.GetFraction() > DensityScaling)
-			{
-				continue;
-			}
+	for (int32 Index = 0; Index < OriginalNum; Index++)
+	{
+		SortPoints[Index] = Transforms[Index].GetOrigin();
 
-			SortIndex.Add(Index);
+		if (DensityScaling < 1.0f && DensityRand.GetFraction() > DensityScaling)
+		{
+			continue;
 		}
 
-		Num = SortIndex.Num();
+		SortIndex.Add(Index);
+	}
 
-		OcclusionLayerTarget = CVarMaxOcclusionQueriesPerComponent.GetValueOnAnyThread();
-		int32 MinInstancesPerOcclusionQuery = CVarMinInstancesPerOcclusionQuery.GetValueOnAnyThread();
+	Num = SortIndex.Num();
 
-		if (Num / MinInstancesPerOcclusionQuery < OcclusionLayerTarget)
+	OcclusionLayerTarget = CVarMaxOcclusionQueriesPerComponent.GetValueOnAnyThread();
+	int32 MinInstancesPerOcclusionQuery = CVarMinInstancesPerOcclusionQuery.GetValueOnAnyThread();
+
+	if (Num / MinInstancesPerOcclusionQuery < OcclusionLayerTarget)
+	{
+		OcclusionLayerTarget = Num / MinInstancesPerOcclusionQuery;
+		if (OcclusionLayerTarget < CVarMinOcclusionQueriesPerComponent.GetValueOnAnyThread())
 		{
-			OcclusionLayerTarget = Num / MinInstancesPerOcclusionQuery;
-			if (OcclusionLayerTarget < CVarMinOcclusionQueriesPerComponent.GetValueOnAnyThread())
-			{
-				OcclusionLayerTarget = 0;
-			}
+			OcclusionLayerTarget = 0;
 		}
-		InternalNodeBranchingFactor = CVarFoliageSplitFactor.GetValueOnAnyThread();
+	}
+	InternalNodeBranchingFactor = CVarFoliageSplitFactor.GetValueOnAnyThread();
 		
-		if (Num / MaxInstancesPerLeaf < InternalNodeBranchingFactor) // if there are less than InternalNodeBranchingFactor leaf nodes
-		{
-			MaxInstancesPerLeaf = FMath::Clamp<int32>(Num / InternalNodeBranchingFactor, 1, 1024); // then make sure we have at least InternalNodeBranchingFactor leaves
-		}
+	if (Num / MaxInstancesPerLeaf < InternalNodeBranchingFactor) // if there are less than InternalNodeBranchingFactor leaf nodes
+	{
+		MaxInstancesPerLeaf = FMath::Clamp<int32>(Num / InternalNodeBranchingFactor, 1, 1024); // then make sure we have at least InternalNodeBranchingFactor leaves
 	}
-
-public:
-	TUniquePtr<FClusterTree> Result;
-	TUniquePtr<FStaticMeshInstanceData> BuiltInstanceData;
+}
 	
-	FClusterBuilder(TArray<FMatrix> InTransforms, TArray<float> InCustomDataFloats, int32 InNumCustomDataFloats, const FBox& InInstBox, int32 InMaxInstancesPerLeaf, float InDensityScaling, int32 InInstancingRandomSeed, bool InGenerateInstanceScalingRange)
-		: OriginalNum(InTransforms.Num())
-		, InstBox(InInstBox)
-		, MaxInstancesPerLeaf(InMaxInstancesPerLeaf)
-		, InstancingRandomSeed(InInstancingRandomSeed)
-		, DensityScaling(InDensityScaling)
-		, GenerateInstanceScalingRange(InGenerateInstanceScalingRange)
-		, Transforms(MoveTemp(InTransforms))
-		, CustomDataFloats(MoveTemp(InCustomDataFloats))
-		, NumCustomDataFloats(InNumCustomDataFloats)
-		, Result(nullptr)
+
+void UHierarchicalInstancedStaticMeshComponent::FClusterBuilder::BuildTreeAndBufferAsync(ENamedThreads::Type CurrentThread, const FGraphEventRef& MyCompletionGraphEvent)
+{
+#if WITH_EDITOR
+	if (!FMath::IsNearlyZero(GDebugBuildTreeAsyncDelayInSeconds))
 	{
+		UE_LOG(LogStaticMesh, Warning, TEXT("BuildTree Debug Delay %5.1f (CVar foliage.DebugBuildTreeAsyncDelayInSeconds)"), GDebugBuildTreeAsyncDelayInSeconds);
+		FPlatformProcess::Sleep(GDebugBuildTreeAsyncDelayInSeconds);
 	}
+#endif
+	BuildTreeAndBuffer();
+}
+
+void UHierarchicalInstancedStaticMeshComponent::FClusterBuilder::BuildTreeAndBuffer()
+{
+	BuildTree();
+	BuildInstanceBuffer();
+}
 
-	void BuildTreeAndBufferAsync(ENamedThreads::Type CurrentThread, const FGraphEventRef& MyCompletionGraphEvent)
+void UHierarchicalInstancedStaticMeshComponent::FClusterBuilder::BuildTree()
+{
+	Init();
+		
+	Result = MakeUnique<FClusterTree>();
+
+	if (Num == 0)
 	{
-#if WITH_EDITOR
-		if (!FMath::IsNearlyZero(GDebugBuildTreeAsyncDelayInSeconds))
-		{
-			UE_LOG(LogStaticMesh, Warning, TEXT("BuildTree Debug Delay %5.1f (CVar foliage.DebugBuildTreeAsyncDelayInSeconds)"), GDebugBuildTreeAsyncDelayInSeconds);
-			FPlatformProcess::Sleep(GDebugBuildTreeAsyncDelayInSeconds);
-		}
-#endif
-		BuildTreeAndBuffer();
+		// Can happen if all instances are excluded due to scalability
+		// It doesn't only happen with a scalability factor of 0 - 
+		// even with a scalability factor of 0.99, if there's only one instance of this type you can end up with Num == 0 if you're unlucky
+		Result->InstanceReorderTable.Init(INDEX_NONE, OriginalNum);
+		return;
 	}
 
-	void BuildTreeAndBuffer()
+	bool bIsOcclusionLayer = false;
+	BranchingFactor = MaxInstancesPerLeaf;
+	if (BranchingFactor > 2 && OcclusionLayerTarget && Num / BranchingFactor <= OcclusionLayerTarget)
 	{
-		BuildTree();
-		BuildInstanceBuffer();
+		BranchingFactor = FMath::Max<int32>(2, (Num + OcclusionLayerTarget - 1) / OcclusionLayerTarget);
+		OcclusionLayerTarget = 0;
+		bIsOcclusionLayer = true;
 	}
-
-	void BuildTree()
+	Split(Num);
+	if (bIsOcclusionLayer)
 	{
-		Init();
+		Result->OutOcclusionLayerNum = Clusters.Num();
+		bIsOcclusionLayer = false;
+	}
+
+	Result->SortedInstances.Append(SortIndex);
 		
-		Result = MakeUnique<FClusterTree>();
+	NumRoots = Clusters.Num();
+	Result->Nodes.Init(FClusterNode(), Clusters.Num());
 
-		if (Num == 0)
+	for (int32 Index = 0; Index < NumRoots; Index++)
+	{
+		FClusterNode& Node = Result->Nodes[Index];
+		Node.FirstInstance = Clusters[Index].Start;
+		Node.LastInstance = Clusters[Index].Start + Clusters[Index].Num - 1;
+		FBox NodeBox(ForceInit);
+		for (int32 InstanceIndex = Node.FirstInstance; InstanceIndex <= Node.LastInstance; InstanceIndex++)
 		{
-			// Can happen if all instances are excluded due to scalability
-			// It doesn't only happen with a scalability factor of 0 - 
-			// even with a scalability factor of 0.99, if there's only one instance of this type you can end up with Num == 0 if you're unlucky
-			Result->InstanceReorderTable.Init(INDEX_NONE, OriginalNum);
-			return;
+			const FMatrix& ThisInstTrans = Transforms[Result->SortedInstances[InstanceIndex]];
+			FBox ThisInstBox = InstBox.TransformBy(ThisInstTrans);
+			NodeBox += ThisInstBox;
+
+			if (GenerateInstanceScalingRange)
+			{
+				FVector3f CurrentScale(ThisInstTrans.GetScaleVector());
+
+				Node.MinInstanceScale = Node.MinInstanceScale.ComponentMin(CurrentScale);
+				Node.MaxInstanceScale = Node.MaxInstanceScale.ComponentMax(CurrentScale);
+			}
 		}
+		Node.BoundMin = (FVector3f)NodeBox.Min;
+		Node.BoundMax = (FVector3f)NodeBox.Max;
+	}
+	TArray<int32> NodesPerLevel;
+	NodesPerLevel.Add(NumRoots);
+	int32 LOD = 0;
 
-		bool bIsOcclusionLayer = false;
-		BranchingFactor = MaxInstancesPerLeaf;
-		if (BranchingFactor > 2 && OcclusionLayerTarget && Num / BranchingFactor <= OcclusionLayerTarget)
+	TArray<int32> InverseSortIndex;
+	TArray<int32> RemapSortIndex;
+	TArray<int32> InverseInstanceIndex;
+	TArray<int32> OldInstanceIndex;
+	TArray<int32> LevelStarts;
+	TArray<int32> InverseChildIndex;
+	TArray<FClusterNode> OldNodes;
+
+	while (NumRoots > 1)
+	{
+		SortIndex.Reset();
+		SortPoints.Reset();
+		SortIndex.AddUninitialized(NumRoots);
+		SortPoints.AddUninitialized(NumRoots);
+		for (int32 Index = 0; Index < NumRoots; Index++)
 		{
-			BranchingFactor = FMath::Max<int32>(2, (Num + OcclusionLayerTarget - 1) / OcclusionLayerTarget);
+			SortIndex[Index] = Index;
+			FClusterNode& Node = Result->Nodes[Index];
+			SortPoints[Index] = (FVector)(Node.BoundMin + Node.BoundMax) * 0.5f;
+		}
+		BranchingFactor = InternalNodeBranchingFactor;
+		if (BranchingFactor > 2 && OcclusionLayerTarget && NumRoots / BranchingFactor <= OcclusionLayerTarget)
+		{
+			BranchingFactor = FMath::Max<int32>(2, (NumRoots + OcclusionLayerTarget - 1) / OcclusionLayerTarget);
 			OcclusionLayerTarget = 0;
 			bIsOcclusionLayer = true;
 		}
-		Split(Num);
+		Split(NumRoots);
 		if (bIsOcclusionLayer)
 		{
 			Result->OutOcclusionLayerNum = Clusters.Num();
 			bIsOcclusionLayer = false;
 		}
 
-		TArray<int32>& SortedInstances = Result->SortedInstances;
-		SortedInstances.Append(SortIndex);
-		
-		NumRoots = Clusters.Num();
-		Result->Nodes.Init(FClusterNode(), Clusters.Num());
-
+		InverseSortIndex.Reset();
+		InverseSortIndex.AddUninitialized(NumRoots);
 		for (int32 Index = 0; Index < NumRoots; Index++)
 		{
-			FClusterNode& Node = Result->Nodes[Index];
-			Node.FirstInstance = Clusters[Index].Start;
-			Node.LastInstance = Clusters[Index].Start + Clusters[Index].Num - 1;
-			FBox NodeBox(ForceInit);
-			for (int32 InstanceIndex = Node.FirstInstance; InstanceIndex <= Node.LastInstance; InstanceIndex++)
-			{
-				const FMatrix& ThisInstTrans = Transforms[SortedInstances[InstanceIndex]];
-				FBox ThisInstBox = InstBox.TransformBy(ThisInstTrans);
-				NodeBox += ThisInstBox;
+			InverseSortIndex[SortIndex[Index]] = Index;
+		}
 
-				if (GenerateInstanceScalingRange)
+		{
+			// rearrange the instances to match the new order of the old roots
+			RemapSortIndex.Reset();
+			RemapSortIndex.AddUninitialized(Num);
+			int32 OutIndex = 0;
+			for (int32 Index = 0; Index < NumRoots; Index++)
+			{
+				FClusterNode& Node = Result->Nodes[SortIndex[Index]];
+				for (int32 InstanceIndex = Node.FirstInstance; InstanceIndex <= Node.LastInstance; InstanceIndex++)
 				{
-					FVector3f CurrentScale(ThisInstTrans.GetScaleVector());
-
-					Node.MinInstanceScale = Node.MinInstanceScale.ComponentMin(CurrentScale);
-					Node.MaxInstanceScale = Node.MaxInstanceScale.ComponentMax(CurrentScale);
+					RemapSortIndex[OutIndex++] = InstanceIndex;
 				}
 			}
-			Node.BoundMin = (FVector3f)NodeBox.Min;
-			Node.BoundMax = (FVector3f)NodeBox.Max;
-		}
-		TArray<int32> NodesPerLevel;
-		NodesPerLevel.Add(NumRoots);
-		int32 LOD = 0;
-
-		TArray<int32> InverseSortIndex;
-		TArray<int32> RemapSortIndex;
-		TArray<int32> InverseInstanceIndex;
-		TArray<int32> OldInstanceIndex;
-		TArray<int32> LevelStarts;
-		TArray<int32> InverseChildIndex;
-		TArray<FClusterNode> OldNodes;
-
-		while (NumRoots > 1)
-		{
-			SortIndex.Reset();
-			SortPoints.Reset();
-			SortIndex.AddUninitialized(NumRoots);
-			SortPoints.AddUninitialized(NumRoots);
-			for (int32 Index = 0; Index < NumRoots; Index++)
+			InverseInstanceIndex.Reset();
+			InverseInstanceIndex.AddUninitialized(Num);
+			for (int32 Index = 0; Index < Num; Index++)
+			{
+				InverseInstanceIndex[RemapSortIndex[Index]] = Index;
+			}
+			for (int32 Index = 0; Index < Result->Nodes.Num(); Index++)
 			{
-				SortIndex[Index] = Index;
 				FClusterNode& Node = Result->Nodes[Index];
-				SortPoints[Index] = (FVector)(Node.BoundMin + Node.BoundMax) * 0.5f;
+				Node.FirstInstance = InverseInstanceIndex[Node.FirstInstance];
+				Node.LastInstance = InverseInstanceIndex[Node.LastInstance];
 			}
-			BranchingFactor = InternalNodeBranchingFactor;
-			if (BranchingFactor > 2 && OcclusionLayerTarget && NumRoots / BranchingFactor <= OcclusionLayerTarget)
+			OldInstanceIndex.Reset();
+			Swap(OldInstanceIndex, Result->SortedInstances);
+			Result->SortedInstances.AddUninitialized(Num);
+			for (int32 Index = 0; Index < Num; Index++)
 			{
-				BranchingFactor = FMath::Max<int32>(2, (NumRoots + OcclusionLayerTarget - 1) / OcclusionLayerTarget);
-				OcclusionLayerTarget = 0;
-				bIsOcclusionLayer = true;
+				Result->SortedInstances[Index] = OldInstanceIndex[RemapSortIndex[Index]];
 			}
-			Split(NumRoots);
-			if (bIsOcclusionLayer)
+		}
+		{
+			// rearrange the nodes to match the new order of the old roots
+			RemapSortIndex.Reset();
+			int32 NewNum = Result->Nodes.Num() + Clusters.Num();
+			// RemapSortIndex[new index] == old index
+			RemapSortIndex.AddUninitialized(NewNum);
+			LevelStarts.Reset();
+			LevelStarts.Add(Clusters.Num());
+			for (int32 Index = 0; Index < NodesPerLevel.Num() - 1; Index++)
 			{
-				Result->OutOcclusionLayerNum = Clusters.Num();
-				bIsOcclusionLayer = false;
+				LevelStarts.Add(LevelStarts[Index] + NodesPerLevel[Index]);
 			}
 
-			InverseSortIndex.Reset();
-			InverseSortIndex.AddUninitialized(NumRoots);
 			for (int32 Index = 0; Index < NumRoots; Index++)
 			{
-				InverseSortIndex[SortIndex[Index]] = Index;
-			}
+				FClusterNode& Node = Result->Nodes[SortIndex[Index]];
+				RemapSortIndex[LevelStarts[0]++] = SortIndex[Index];
 
-			{
-				// rearrange the instances to match the new order of the old roots
-				RemapSortIndex.Reset();
-				RemapSortIndex.AddUninitialized(Num);
-				int32 OutIndex = 0;
-				for (int32 Index = 0; Index < NumRoots; Index++)
+				int32 LeftIndex = Node.FirstChild;
+				int32 RightIndex = Node.LastChild;
+				int32 LevelIndex = 1;
+				while (RightIndex >= 0)
 				{
-					FClusterNode& Node = Result->Nodes[SortIndex[Index]];
-					for (int32 InstanceIndex = Node.FirstInstance; InstanceIndex <= Node.LastInstance; InstanceIndex++)
+					int32 NextLeftIndex = MAX_int32;
+					int32 NextRightIndex = -1;
+					for (int32 ChildIndex = LeftIndex; ChildIndex <= RightIndex; ChildIndex++)
 					{
-						RemapSortIndex[OutIndex++] = InstanceIndex;
+						RemapSortIndex[LevelStarts[LevelIndex]++] = ChildIndex;
+						int32 LeftChild = Result->Nodes[ChildIndex].FirstChild;
+						int32 RightChild = Result->Nodes[ChildIndex].LastChild;
+						if (LeftChild >= 0 && LeftChild <  NextLeftIndex)
+						{
+							NextLeftIndex = LeftChild;
+						}
+						if (RightChild >= 0 && RightChild >  NextRightIndex)
+						{
+							NextRightIndex = RightChild;
+						}
 					}
+					LeftIndex = NextLeftIndex;
+					RightIndex = NextRightIndex;
+					LevelIndex++;
 				}
-				InverseInstanceIndex.Reset();
-				InverseInstanceIndex.AddUninitialized(Num);
-				for (int32 Index = 0; Index < Num; Index++)
+			}
+			checkSlow(LevelStarts[LevelStarts.Num() - 1] == NewNum);
+			InverseChildIndex.Reset();
+			// InverseChildIndex[old index] == new index
+			InverseChildIndex.AddUninitialized(NewNum);
+			for (int32 Index = Clusters.Num(); Index < NewNum; Index++)
+			{
+				InverseChildIndex[RemapSortIndex[Index]] = Index;
+			}
+			for (int32 Index = 0; Index < Result->Nodes.Num(); Index++)
+			{
+				FClusterNode& Node = Result->Nodes[Index];
+				if (Node.FirstChild >= 0)
 				{
-					InverseInstanceIndex[RemapSortIndex[Index]] = Index;
+					Node.FirstChild = InverseChildIndex[Node.FirstChild];
+					Node.LastChild = InverseChildIndex[Node.LastChild];
 				}
-				for (int32 Index = 0; Index < Result->Nodes.Num(); Index++)
+			}
+			{
+				Swap(OldNodes, Result->Nodes);
+				Result->Nodes.Empty(NewNum);
+				for (int32 Index = 0; Index < Clusters.Num(); Index++)
 				{
-					FClusterNode& Node = Result->Nodes[Index];
-					Node.FirstInstance = InverseInstanceIndex[Node.FirstInstance];
-					Node.LastInstance = InverseInstanceIndex[Node.LastInstance];
+					Result->Nodes.Add(FClusterNode());
 				}
-				OldInstanceIndex.Reset();
-				Swap(OldInstanceIndex, SortedInstances);
-				SortedInstances.AddUninitialized(Num);
-				for (int32 Index = 0; Index < Num; Index++)
+				Result->Nodes.AddUninitialized(OldNodes.Num());
+				for (int32 Index = 0; Index < OldNodes.Num(); Index++)
 				{
-					SortedInstances[Index] = OldInstanceIndex[RemapSortIndex[Index]];
+					Result->Nodes[InverseChildIndex[Index]] = OldNodes[Index];
 				}
 			}
+			int32 OldIndex = Clusters.Num();
+			int32 InstanceTracker = 0;
+			for (int32 Index = 0; Index < Clusters.Num(); Index++)
 			{
-				// rearrange the nodes to match the new order of the old roots
-				RemapSortIndex.Reset();
-				int32 NewNum = Result->Nodes.Num() + Clusters.Num();
-				// RemapSortIndex[new index] == old index
-				RemapSortIndex.AddUninitialized(NewNum);
-				LevelStarts.Reset();
-				LevelStarts.Add(Clusters.Num());
-				for (int32 Index = 0; Index < NodesPerLevel.Num() - 1; Index++)
-				{
-					LevelStarts.Add(LevelStarts[Index] + NodesPerLevel[Index]);
-				}
-
-				for (int32 Index = 0; Index < NumRoots; Index++)
+				FClusterNode& Node = Result->Nodes[Index];
+				Node.FirstChild = OldIndex;
+				OldIndex += Clusters[Index].Num;
+				Node.LastChild = OldIndex - 1;
+				Node.FirstInstance = Result->Nodes[Node.FirstChild].FirstInstance;
+				checkSlow(Node.FirstInstance == InstanceTracker);
+				Node.LastInstance = Result->Nodes[Node.LastChild].LastInstance;
+				InstanceTracker = Node.LastInstance + 1;
+				checkSlow(InstanceTracker <= Num);
+				FBox NodeBox(ForceInit);
+				for (int32 ChildIndex = Node.FirstChild; ChildIndex <= Node.LastChild; ChildIndex++)
 				{
-					FClusterNode& Node = Result->Nodes[SortIndex[Index]];
-					RemapSortIndex[LevelStarts[0]++] = SortIndex[Index];
+					FClusterNode& ChildNode = Result->Nodes[ChildIndex];
+					NodeBox += (FVector)ChildNode.BoundMin;
+					NodeBox += (FVector)ChildNode.BoundMax;
 
-					int32 LeftIndex = Node.FirstChild;
-					int32 RightIndex = Node.LastChild;
-					int32 LevelIndex = 1;
-					while (RightIndex >= 0)
-					{
-						int32 NextLeftIndex = MAX_int32;
-						int32 NextRightIndex = -1;
-						for (int32 ChildIndex = LeftIndex; ChildIndex <= RightIndex; ChildIndex++)
-						{
-							RemapSortIndex[LevelStarts[LevelIndex]++] = ChildIndex;
-							int32 LeftChild = Result->Nodes[ChildIndex].FirstChild;
-							int32 RightChild = Result->Nodes[ChildIndex].LastChild;
-							if (LeftChild >= 0 && LeftChild <  NextLeftIndex)
-							{
-								NextLeftIndex = LeftChild;
-							}
-							if (RightChild >= 0 && RightChild >  NextRightIndex)
-							{
-								NextRightIndex = RightChild;
-							}
-						}
-						LeftIndex = NextLeftIndex;
-						RightIndex = NextRightIndex;
-						LevelIndex++;
-					}
-				}
-				checkSlow(LevelStarts[LevelStarts.Num() - 1] == NewNum);
-				InverseChildIndex.Reset();
-				// InverseChildIndex[old index] == new index
-				InverseChildIndex.AddUninitialized(NewNum);
-				for (int32 Index = Clusters.Num(); Index < NewNum; Index++)
-				{
-					InverseChildIndex[RemapSortIndex[Index]] = Index;
-				}
-				for (int32 Index = 0; Index < Result->Nodes.Num(); Index++)
-				{
-					FClusterNode& Node = Result->Nodes[Index];
-					if (Node.FirstChild >= 0)
-					{
-						Node.FirstChild = InverseChildIndex[Node.FirstChild];
-						Node.LastChild = InverseChildIndex[Node.LastChild];
-					}
-				}
-				{
-					Swap(OldNodes, Result->Nodes);
-					Result->Nodes.Empty(NewNum);
-					for (int32 Index = 0; Index < Clusters.Num(); Index++)
-					{
-						Result->Nodes.Add(FClusterNode());
-					}
-					Result->Nodes.AddUninitialized(OldNodes.Num());
-					for (int32 Index = 0; Index < OldNodes.Num(); Index++)
+					if (GenerateInstanceScalingRange)
 					{
-						Result->Nodes[InverseChildIndex[Index]] = OldNodes[Index];
+						Node.MinInstanceScale = Node.MinInstanceScale.ComponentMin(ChildNode.MinInstanceScale);
+						Node.MaxInstanceScale = Node.MaxInstanceScale.ComponentMax(ChildNode.MaxInstanceScale);
 					}
 				}
-				int32 OldIndex = Clusters.Num();
-				int32 InstanceTracker = 0;
-				for (int32 Index = 0; Index < Clusters.Num(); Index++)
-				{
-					FClusterNode& Node = Result->Nodes[Index];
-					Node.FirstChild = OldIndex;
-					OldIndex += Clusters[Index].Num;
-					Node.LastChild = OldIndex - 1;
-					Node.FirstInstance = Result->Nodes[Node.FirstChild].FirstInstance;
-					checkSlow(Node.FirstInstance == InstanceTracker);
-					Node.LastInstance = Result->Nodes[Node.LastChild].LastInstance;
-					InstanceTracker = Node.LastInstance + 1;
-					checkSlow(InstanceTracker <= Num);
-					FBox NodeBox(ForceInit);
-					for (int32 ChildIndex = Node.FirstChild; ChildIndex <= Node.LastChild; ChildIndex++)
-					{
-						FClusterNode& ChildNode = Result->Nodes[ChildIndex];
-						NodeBox += (FVector)ChildNode.BoundMin;
-						NodeBox += (FVector)ChildNode.BoundMax;
-
-						if (GenerateInstanceScalingRange)
-						{
-							Node.MinInstanceScale = Node.MinInstanceScale.ComponentMin(ChildNode.MinInstanceScale);
-							Node.MaxInstanceScale = Node.MaxInstanceScale.ComponentMax(ChildNode.MaxInstanceScale);
-						}
-					}
-					Node.BoundMin = (FVector3f)NodeBox.Min;
-					Node.BoundMax = (FVector3f)NodeBox.Max;
-				}
-				NumRoots = Clusters.Num();
-				NodesPerLevel.Insert(NumRoots, 0);
+				Node.BoundMin = (FVector3f)NodeBox.Min;
+				Node.BoundMax = (FVector3f)NodeBox.Max;
 			}
+			NumRoots = Clusters.Num();
+			NodesPerLevel.Insert(NumRoots, 0);
 		}
+	}
 
-		// Save inverse map
-		Result->InstanceReorderTable.Init(INDEX_NONE, OriginalNum);
-		for (int32 Index = 0; Index < Num; Index++)
-		{
-			Result->InstanceReorderTable[SortedInstances[Index]] = Index;
-		}
+	// Save inverse map
+	Result->InstanceReorderTable.Init(INDEX_NONE, OriginalNum);
+	for (int32 Index = 0; Index < Num; Index++)
+	{
+		Result->InstanceReorderTable[Result->SortedInstances[Index]] = Index;
+	}
 
-		// Output a general scale of 1 if we dont want the scaling range
-		if (!GenerateInstanceScalingRange)
-		{
-			Result->Nodes[0].MinInstanceScale = FVector3f::OneVector;
-			Result->Nodes[0].MaxInstanceScale = FVector3f::OneVector;
-		}
+	// Output a general scale of 1 if we dont want the scaling range
+	if (!GenerateInstanceScalingRange)
+	{
+		Result->Nodes[0].MinInstanceScale = FVector3f::OneVector;
+		Result->Nodes[0].MaxInstanceScale = FVector3f::OneVector;
 	}
-};
+}
 
-static bool PrintLevel(const FClusterTree& Tree, int32 NodeIndex, int32 Level, int32 CurrentLevel, int32 Parent)
+bool UHierarchicalInstancedStaticMeshComponent::FClusterTree::PrintLevel(int32 NodeIndex, int32 Level, int32 CurrentLevel, int32 Parent)
 {
-	const FClusterNode& Node = Tree.Nodes[NodeIndex];
+	const FClusterNode& Node = Nodes[NodeIndex];
 	if (Level == CurrentLevel)
 	{
 		UE_LOG(LogConsoleResponse, Display, TEXT("Level %2d  Parent %3d"),
@@ -732,7 +669,7 @@ static bool PrintLevel(const FClusterTree& Tree, int32 NodeIndex, int32 Level, i
 	bool Ret = false;
 	for (int32 Child = Node.FirstChild; Child <= Node.LastChild; Child++)
 	{
-		Ret = PrintLevel(Tree, Child, Level, CurrentLevel + 1, NodeIndex) || Ret;
+		Ret = PrintLevel(Child, Level, CurrentLevel + 1, NodeIndex) || Ret;
 	}
 	return Ret;
 }
@@ -764,14 +701,14 @@ static void TestFoliage(const TArray<FString>& Args)
 		InstanceTransforms[Index] = Instances[Index].Transform;
 	}
 
-	FClusterBuilder Builder(InstanceTransforms, InstanceCustomDataDummy, 0, TempBox, 16, 1.0f, 1, 0);
+	UHierarchicalInstancedStaticMeshComponent::FClusterBuilder Builder(InstanceTransforms, InstanceCustomDataDummy, 0, TempBox, 16, 1.0f, 1, 0);
 	Builder.BuildTree();
 
 	int32 Level = 0;
 
 	UE_LOG(LogConsoleResponse, Display, TEXT("-----"));
 
-	while(PrintLevel(*Builder.Result, 0, Level++, 0, -1))
+	while(Builder.Result->PrintLevel(0, Level++, 0, -1))
 	{
 	}
 }
@@ -2067,7 +2004,7 @@ FBoxSphereBounds UHierarchicalInstancedStaticMeshComponent::CalcBounds(const FTr
 	}
 }
 
-static FBox GetClusterTreeBounds(TArray<FClusterNode> const& InClusterTree, FVector InOffset)
+FBox UHierarchicalInstancedStaticMeshComponent::GetClusterTreeBounds(TArray<FClusterNode> const& InClusterTree, const FVector& InOffset)
 {
 	// Return top node of cluster tree. Apply offset on node bounds.
 	return (InClusterTree.Num() > 0 ? FBox(InOffset + FVector(InClusterTree[0].BoundMin), InOffset + FVector(InClusterTree[0].BoundMax)) : FBox(ForceInit));
@@ -2827,85 +2764,6 @@ void UHierarchicalInstancedStaticMeshComponent::BuildTree()
 	}
 }
 
-// TODO: Move this implementation to UGrassInstancedStaticMeshComponent.cpp for UE 5.2
-void UHierarchicalInstancedStaticMeshComponent::BuildTreeAnyThread(
-	TArray<FMatrix>& InstanceTransforms, 
-	TArray<float>& InstanceCustomDataFloats,
-	int32 NumCustomDataFloats,
-	const FBox& MeshBox,
-	TArray<FClusterNode>& OutClusterTree,
-	TArray<int32>& OutSortedInstances,
-	TArray<int32>& OutInstanceReorderTable,
-	int32& OutOcclusionLayerNum,
-	int32 MaxInstancesPerLeaf,
-	bool InGenerateInstanceScalingRange
-	)
-{
-	check(MaxInstancesPerLeaf > 0);
-
-	// do grass need this?
-	float DensityScaling = 1.0f;
-	int32 InstancingRandomSeed = 1;
-
-	FClusterBuilder Builder(InstanceTransforms, InstanceCustomDataFloats, NumCustomDataFloats, MeshBox, MaxInstancesPerLeaf, DensityScaling, InstancingRandomSeed, InGenerateInstanceScalingRange);
-	Builder.BuildTree();
-	OutOcclusionLayerNum = Builder.Result->OutOcclusionLayerNum;
-
-	OutClusterTree = MoveTemp(Builder.Result->Nodes);
-	OutInstanceReorderTable = MoveTemp(Builder.Result->InstanceReorderTable);
-	OutSortedInstances = MoveTemp(Builder.Result->SortedInstances);
-}
-
-// TODO: Move this implementation to UGrassInstancedStaticMeshComponent.cpp for UE 5.2
-void UHierarchicalInstancedStaticMeshComponent::AcceptPrebuiltTree(TArray<FInstancedStaticMeshInstanceData>& InInstanceData, TArray<FClusterNode>& InClusterTree, int32 InOcclusionLayerNumNodes, int32 InNumBuiltRenderInstances)
-{
-	checkSlow(IsInGameThread());
-
-	QUICK_SCOPE_CYCLE_COUNTER(STAT_UHierarchicalInstancedStaticMeshComponent_AcceptPrebuiltTree);
-	
-	// this is only for prebuild data, already in the correct order
-	check(!PerInstanceSMData.Num());
-
-	NumBuiltInstances = 0;
-	TranslatedInstanceSpaceOrigin = FVector::Zero();
-	check(PerInstanceRenderData.IsValid());
-	NumBuiltRenderInstances = InNumBuiltRenderInstances;
-	check(NumBuiltRenderInstances);
-	UnbuiltInstanceBounds.Init();
-	UnbuiltInstanceBoundsList.Empty();
-	ClusterTreePtr = MakeShareable(new TArray<FClusterNode>);
-	InstanceReorderTable.Empty();
-	SortedInstances.Empty();
-	OcclusionLayerNumNodes = InOcclusionLayerNumNodes;
-	BuiltInstanceBounds = GetClusterTreeBounds(InClusterTree, FVector::Zero());
-	InstanceCountToRender = InNumBuiltRenderInstances;
-
-	// Verify that the mesh is valid before using it.
-	const bool bMeshIsValid =
-		// make sure we have instances
-		NumBuiltRenderInstances > 0 &&
-		// make sure we have an actual staticmesh
-		GetStaticMesh() &&
-		GetStaticMesh()->HasValidRenderData();
-
-	if(bMeshIsValid)
-	{
-		*ClusterTreePtr = MoveTemp(InClusterTree);
-
-		// We only need to copy off the instances if it is a Nanite mesh, since the Nanite scene proxy uses them instead
-		if (ShouldCreateNaniteProxy())
-		{
-			PerInstanceSMData = MoveTemp(InInstanceData);
-		}
-		
-		PostBuildStats();
-
-	}
-	QUICK_SCOPE_CYCLE_COUNTER(STAT_UHierarchicalInstancedStaticMeshComponent_AcceptPrebuiltTree_Mark);
-
-	MarkRenderStateDirty();
-}
-
 void UHierarchicalInstancedStaticMeshComponent::ApplyBuildTreeAsync(ENamedThreads::Type CurrentThread, const FGraphEventRef& MyCompletionGraphEvent, TSharedRef<FClusterBuilder, ESPMode::ThreadSafe> Builder, double StartTime)
 {
 	check(IsInGameThread());
diff --git a/Engine/Source/Runtime/Foliage/Private/InstancedGrass.cpp b/Engine/Source/Runtime/Foliage/Private/InstancedGrass.cpp
--- a/Engine/Source/Runtime/Foliage/Private/InstancedGrass.cpp
+++ b/Engine/Source/Runtime/Foliage/Private/InstancedGrass.cpp
@@ -1,6 +1,7 @@
 // Copyright Epic Games, Inc. All Rights Reserved.
 
 #include "GrassInstancedStaticMeshComponent.h"
+#include "Engine/StaticMesh.h"
 
 UGrassInstancedStaticMeshComponent::UGrassInstancedStaticMeshComponent(const FObjectInitializer& ObjectInitializer) :
 	Super(ObjectInitializer)
@@ -21,31 +22,66 @@ void UGrassInstancedStaticMeshComponent::BuildTreeAnyThread(
 	bool InGenerateInstanceScalingRange
 	)
 {
-	// TODO: Move the implementation of this function in the base here in 5.2
-	PRAGMA_DISABLE_DEPRECATION_WARNINGS
-	Super::BuildTreeAnyThread(
-		InstanceTransforms,
-		InstanceCustomDataFloats,
-		NumCustomDataFloats,
-		MeshBox,
-		OutClusterTree,
-		OutSortedInstances,
-		OutInstanceReorderTable,
-		OutOcclusionLayerNum,
-		MaxInstancesPerLeaf,
-		InGenerateInstanceScalingRange);
-	PRAGMA_ENABLE_DEPRECATION_WARNINGS
+	check(MaxInstancesPerLeaf > 0);
+
+	// do grass need this?
+	float DensityScaling = 1.0f;
+	int32 InstancingRandomSeed = 1;
+
+	FClusterBuilder Builder(InstanceTransforms, InstanceCustomDataFloats, NumCustomDataFloats, MeshBox, MaxInstancesPerLeaf, DensityScaling, InstancingRandomSeed, InGenerateInstanceScalingRange);
+	Builder.BuildTree();
+	OutOcclusionLayerNum = Builder.Result->OutOcclusionLayerNum;
+
+	OutClusterTree = MoveTemp(Builder.Result->Nodes);
+	OutInstanceReorderTable = MoveTemp(Builder.Result->InstanceReorderTable);
+	OutSortedInstances = MoveTemp(Builder.Result->SortedInstances);
 }
 
 void UGrassInstancedStaticMeshComponent::AcceptPrebuiltTree(TArray<FInstancedStaticMeshInstanceData>& InInstanceData, TArray<FClusterNode>& InClusterTree, int32 InOcclusionLayerNumNodes, int32 InNumBuiltRenderInstances)
 {
-	// TODO: Move the implementation of this function in the base to here in 5.2
-	PRAGMA_DISABLE_DEPRECATION_WARNINGS
-	Super::AcceptPrebuiltTree(
-		InInstanceData,
-		InClusterTree,
-		InOcclusionLayerNumNodes,
-		InNumBuiltRenderInstances);
-	PRAGMA_ENABLE_DEPRECATION_WARNINGS
+	checkSlow(IsInGameThread());
+
+	QUICK_SCOPE_CYCLE_COUNTER(STAT_UGrassInstancedStaticMeshComponent_AcceptPrebuiltTree);
+
+	// this is only for prebuild data, already in the correct order
+	check(!PerInstanceSMData.Num());
+
+	NumBuiltInstances = 0;
+	TranslatedInstanceSpaceOrigin = FVector::Zero();
+	check(PerInstanceRenderData.IsValid());
+	NumBuiltRenderInstances = InNumBuiltRenderInstances;
+	check(NumBuiltRenderInstances);
+	UnbuiltInstanceBounds.Init();
+	UnbuiltInstanceBoundsList.Empty();
+	ClusterTreePtr = MakeShareable(new TArray<FClusterNode>);
+	InstanceReorderTable.Empty();
+	SortedInstances.Empty();
+	OcclusionLayerNumNodes = InOcclusionLayerNumNodes;
+	BuiltInstanceBounds = GetClusterTreeBounds(InClusterTree, FVector::Zero());
+	InstanceCountToRender = InNumBuiltRenderInstances;
+
+	// Verify that the mesh is valid before using it.
+	const bool bMeshIsValid =
+		// make sure we have instances
+		NumBuiltRenderInstances > 0 &&
+		// make sure we have an actual staticmesh
+		GetStaticMesh() &&
+		GetStaticMesh()->HasValidRenderData();
+
+	if (bMeshIsValid)
+	{
+		*ClusterTreePtr = MoveTemp(InClusterTree);
+
+		if (RequiresInstanceDataForTree())
+		{
+			PerInstanceSMData = MoveTemp(InInstanceData);
+		}
+
+		PostBuildStats();
+
+	}
+	QUICK_SCOPE_CYCLE_COUNTER(STAT_UGrassInstancedStaticMeshComponent_AcceptPrebuiltTree_Mark);
+
+	MarkRenderStateDirty();
 }
 
diff --git a/Engine/Source/Runtime/Foliage/Public/GrassInstancedStaticMeshComponent.h b/Engine/Source/Runtime/Foliage/Public/GrassInstancedStaticMeshComponent.h
--- a/Engine/Source/Runtime/Foliage/Public/GrassInstancedStaticMeshComponent.h
+++ b/Engine/Source/Runtime/Foliage/Public/GrassInstancedStaticMeshComponent.h
@@ -20,6 +20,7 @@ public:
 
 	inline bool RequiresInstanceDataForTree() const
 	{
+		// We only need to copy off the instances if it is a Nanite mesh, since the Nanite scene proxy uses them instead
 		return ShouldCreateNaniteProxy();
 	}
 };
diff --git a/Engine/Source/Runtime/Landscape/Classes/LandscapeComponent.h b/Engine/Source/Runtime/Landscape/Classes/LandscapeComponent.h
--- a/Engine/Source/Runtime/Landscape/Classes/LandscapeComponent.h
+++ b/Engine/Source/Runtime/Landscape/Classes/LandscapeComponent.h
@@ -750,9 +750,6 @@ public:
 	LANDSCAPE_API void RemoveLayerData(const FGuid& InLayerGuid);
 	LANDSCAPE_API void ForEachLayer(TFunctionRef<void(const FGuid&, struct FLandscapeLayerComponentData&)> Fn);
 
-	UE_DEPRECATED(5.1, "SetEditingLayer has been deprecated, use GetLandscapeActor()->SetEditingLayer() instead.")
-	void SetEditingLayer(const FGuid& InEditingLayer) {}
-
 	/** Get the Landscape Actor's editing layer data */
 	FLandscapeLayerComponentData* GetEditingLayer();
 	const FLandscapeLayerComponentData* GetEditingLayer() const;
diff --git a/Engine/Source/Runtime/Landscape/Classes/LandscapeHeightfieldCollisionComponent.h b/Engine/Source/Runtime/Landscape/Classes/LandscapeHeightfieldCollisionComponent.h
--- a/Engine/Source/Runtime/Landscape/Classes/LandscapeHeightfieldCollisionComponent.h
+++ b/Engine/Source/Runtime/Landscape/Classes/LandscapeHeightfieldCollisionComponent.h
@@ -235,10 +235,6 @@ public:
 	virtual void BeginDestroy() override;
 	virtual void GetResourceSizeEx(FResourceSizeEx& CumulativeResourceSize) override;
 	virtual void PostLoad() override;
-	PRAGMA_DISABLE_DEPRECATION_WARNINGS // Suppress compiler warning on override of deprecated function
-	UE_DEPRECATED(5.0, "Use version that takes FObjectPreSaveContext instead.")
-	virtual void PreSave(const class ITargetPlatform* TargetPlatform) override;
-	PRAGMA_ENABLE_DEPRECATION_WARNINGS
 	virtual void PreSave(FObjectPreSaveContext ObjectSaveContext) override;
 #if WITH_EDITOR
 	virtual void ExportCustomProperties(FOutputDevice& Out, uint32 Indent) override;
diff --git a/Engine/Source/Runtime/Landscape/Classes/LandscapeProxy.h b/Engine/Source/Runtime/Landscape/Classes/LandscapeProxy.h
--- a/Engine/Source/Runtime/Landscape/Classes/LandscapeProxy.h
+++ b/Engine/Source/Runtime/Landscape/Classes/LandscapeProxy.h
@@ -716,9 +716,6 @@ public:
 	UPROPERTY(EditAnywhere, Category=Collision, meta = (LandscapeOverridable))
 	int32 SimpleCollisionMipLevel;
 
-	UE_DEPRECATED(5.2, "CollisionThickness is not supported by Chaos Physics and therefore deprecated. Please remove any usage of this property")
-	float CollisionThickness;
-
 	/** Collision profile settings for this landscape */
 	UPROPERTY(EditAnywhere, BlueprintReadOnly, Category=Collision, meta=(ShowOnlyInnerProperties))
 	FBodyInstance BodyInstance;
@@ -1066,10 +1063,6 @@ public:
 #endif // WITH_EDITOR
 
 	//~ Begin UObject Interface.
-	PRAGMA_DISABLE_DEPRECATION_WARNINGS // Suppress compiler warning on override of deprecated function
-	UE_DEPRECATED(5.0, "Use version that takes FObjectPreSaveContext instead.")
-	virtual void PreSave(const class ITargetPlatform* TargetPlatform) override;
-	PRAGMA_ENABLE_DEPRECATION_WARNINGS
 	virtual void PreSave(FObjectPreSaveContext ObjectSaveContext) override;
 	virtual void Serialize(FArchive& Ar) override;
 	static void AddReferencedObjects(UObject* InThis, FReferenceCollector& Collector);
@@ -1231,41 +1224,6 @@ public:
 	LANDSCAPE_API void Import(const FGuid& InGuid, int32 InMinX, int32 InMinY, int32 InMaxX, int32 InMaxY, int32 InNumSubsections, int32 InSubsectionSizeQuads, const TMap<FGuid, TArray<uint16>>& InImportHeightData,
 							  const TCHAR* const InHeightmapFileName, const TMap<FGuid, TArray<FLandscapeImportLayerInfo>>& InImportMaterialLayerInfos, ELandscapeImportAlphamapType InImportMaterialLayerType, const TArray<struct FLandscapeLayer>* InImportLayers = nullptr);
 
-	/**
-	 * Exports landscape into raw mesh
-	 * 
-	 * @param InExportLOD Landscape LOD level to use while exporting, INDEX_NONE will use ALanscapeProxy::ExportLOD settings
-	 * @param OutRawMesh - Resulting raw mesh
-	 * @return true if successful
-	 */
-	UE_DEPRECATED(5.2, "Use the version of this function taking a FRawMeshExportParams as a parameter")
-	LANDSCAPE_API bool ExportToRawMesh(int32 InExportLOD, FMeshDescription& OutRawMesh) const;
-
-	/**
-	* Exports landscape geometry contained within InBounds into a raw mesh
-	*
-	* @param InExportLOD - Landscape LOD level to use while exporting, INDEX_NONE will use ALanscapeProxy::ExportLOD settings
-	* @param OutRawMesh - Resulting raw mesh
-	* @param InBounds - Box/Sphere bounds which limit the geometry exported out into OutRawMesh
-	* @param bIgnoreBounds - If false, InBounds will be ignored during export
-	* @return true if successful
-	*/
-	UE_DEPRECATED(5.2, "Use the version of this function taking a FRawMeshExportParams as a parameter")
-	LANDSCAPE_API bool ExportToRawMesh(int32 InExportLOD, FMeshDescription& OutRawMesh, const FBoxSphereBounds& InBounds, bool bIgnoreBounds = false) const;
-
-	/**
-	* Exports landscape geometry contained within InBounds into a raw mesh
-	*
-	* @param InComponents - Specific landscape component(s) to export
-	* @param InExportLOD - Landscape LOD level to use while exporting, INDEX_NONE will use ALanscapeProxy::ExportLOD settings
-	* @param OutRawMesh - Resulting raw mesh
-	* @param InBounds - Box/Sphere bounds which limit the geometry exported out into OutRawMesh
-	* @param bIgnoreBounds - If false, InBounds will be ignored during export
-	* @return true if successful
-	*/
-	UE_DEPRECATED(5.2, "Use the version of this function taking a FRawMeshExportParams as a parameter")
-	LANDSCAPE_API bool ExportToRawMesh(const TArrayView<ULandscapeComponent*>& InComponents, int32 InExportLOD, FMeshDescription& OutRawMesh, const FBoxSphereBounds& InBounds, bool bIgnoreBounds = false, bool bGenerateOnePolygonGroupPerComponent = false) const;
-
 	struct FRawMeshExportParams
 	{
 		static constexpr int32 MaxUVCount = 6;
@@ -1360,10 +1318,6 @@ public:
 
 	bool ExportToRawMeshDataCopyOld(const FRawMeshExportParams& InExportParams, FMeshDescription& OutRawMesh, const UE::Landscape::Nanite::FAsyncBuildData& AsyncData) const;
 
-	UE_DEPRECATED(5.1, "CheckGenerateLandscapePlatformData has been deprecated, please use CheckGenerateMobilePlatformData instead.")
-	LANDSCAPE_API void CheckGenerateLandscapePlatformData(bool bIsCooking, const ITargetPlatform* TargetPlatform);
-
-	/** Generate mobile platform data if it's missing or outdated */
 	LANDSCAPE_API void CheckGenerateMobilePlatformData(bool bIsCooking, const ITargetPlatform* TargetPlatform);
 
 	/** Returns true if the Nanite representation is missing or outdated */
diff --git a/Engine/Source/Runtime/Landscape/Classes/LandscapeSplinesComponent.h b/Engine/Source/Runtime/Landscape/Classes/LandscapeSplinesComponent.h
--- a/Engine/Source/Runtime/Landscape/Classes/LandscapeSplinesComponent.h
+++ b/Engine/Source/Runtime/Landscape/Classes/LandscapeSplinesComponent.h
@@ -168,9 +168,6 @@ public:
 
 	bool HasAnyControlPointsOrSegments() const { return ControlPoints.Num() > 0 || Segments.Num() > 0; }
 
-	UE_DEPRECATED(5.2, "This function was needed for forcing the SplineEditorMesh to be loaded when a commandlet is running, but commandlets are not prevented from loading it anymore so it has become useless")
-	LANDSCAPE_API void SetDefaultEditorSplineMesh();
-
 	virtual void ShowSplineEditorMesh(bool bShow);
 
 	// Rebuilds all spline points and meshes for all spline control points and segments in this splines component
diff --git a/Engine/Source/Runtime/Landscape/Private/Landscape.cpp b/Engine/Source/Runtime/Landscape/Private/Landscape.cpp
--- a/Engine/Source/Runtime/Landscape/Private/Landscape.cpp
+++ b/Engine/Source/Runtime/Landscape/Private/Landscape.cpp
@@ -318,12 +318,6 @@ void ULandscapeComponent::BeginCacheForCookedPlatformData(const ITargetPlatform*
 	}
 }
 
-// Deprecated, use CheckGenerateMobilePlatformData
-void ALandscapeProxy::CheckGenerateLandscapePlatformData(bool bIsCooking, const ITargetPlatform* TargetPlatform)
-{
-	return CheckGenerateMobilePlatformData(bIsCooking, TargetPlatform);
-}
-
 void ALandscapeProxy::CheckGenerateMobilePlatformData(bool bIsCooking, const ITargetPlatform* TargetPlatform)
 {
 	for (ULandscapeComponent* Component : LandscapeComponents)
@@ -3159,13 +3153,6 @@ void ALandscapeProxy::OnFeatureLevelChanged(ERHIFeatureLevel::Type NewFeatureLev
 }
 #endif
 
-void ALandscapeProxy::PreSave(const class ITargetPlatform* TargetPlatform)
-{
-	PRAGMA_DISABLE_DEPRECATION_WARNINGS;
-	Super::PreSave(TargetPlatform);
-	PRAGMA_ENABLE_DEPRECATION_WARNINGS;
-}
-
 void ALandscapeProxy::PreSave(FObjectPreSaveContext ObjectSaveContext)
 {
 	Super::PreSave(ObjectSaveContext);
diff --git a/Engine/Source/Runtime/Landscape/Private/LandscapeCollision.cpp b/Engine/Source/Runtime/Landscape/Private/LandscapeCollision.cpp
--- a/Engine/Source/Runtime/Landscape/Private/LandscapeCollision.cpp
+++ b/Engine/Source/Runtime/Landscape/Private/LandscapeCollision.cpp
@@ -2489,13 +2489,6 @@ void ULandscapeHeightfieldCollisionComponent::PostLoad()
 #endif // WITH_EDITOR
 }
 
-void ULandscapeHeightfieldCollisionComponent::PreSave(const class ITargetPlatform* TargetPlatform)
-{
-	PRAGMA_DISABLE_DEPRECATION_WARNINGS;
-	Super::PreSave(TargetPlatform);
-	PRAGMA_ENABLE_DEPRECATION_WARNINGS;
-}
-
 void ULandscapeHeightfieldCollisionComponent::PreSave(FObjectPreSaveContext ObjectSaveContext)
 {
 	Super::PreSave(ObjectSaveContext);
diff --git a/Engine/Source/Runtime/Landscape/Private/LandscapeEdit.cpp b/Engine/Source/Runtime/Landscape/Private/LandscapeEdit.cpp
--- a/Engine/Source/Runtime/Landscape/Private/LandscapeEdit.cpp
+++ b/Engine/Source/Runtime/Landscape/Private/LandscapeEdit.cpp
@@ -3553,46 +3553,8 @@ LANDSCAPE_API void ALandscapeProxy::Import(const FGuid& InGuid, int32 InMinX, in
 	
 }
 
-PRAGMA_DISABLE_DEPRECATION_WARNINGS
-bool ALandscapeProxy::ExportToRawMesh(int32 InExportLOD, FMeshDescription& OutRawMesh) const
-{
-	ALandscapeProxy::FRawMeshExportParams Params;
-	Params.ExportLOD = InExportLOD;
-
-	return ExportToRawMesh(Params, OutRawMesh);
-}
-
-bool ALandscapeProxy::ExportToRawMesh(int32 InExportLOD, FMeshDescription& OutRawMesh, const FBoxSphereBounds& InBounds, bool bIgnoreBounds) const
-{
-	ALandscapeProxy::FRawMeshExportParams Params;
-	Params.ExportLOD = InExportLOD;
-	if (!bIgnoreBounds)
-	{
-		Params.ExportBounds = InBounds;
-	}
-
-	return ExportToRawMesh(Params, OutRawMesh);
-}
 
-bool ALandscapeProxy::ExportToRawMesh(
-	const TArrayView<ULandscapeComponent*>& RegisteredLandscapeComponents,
-	int32 InExportLOD,
-	FMeshDescription& OutRawMesh,
-	const FBoxSphereBounds& InBounds,
-	bool bIgnoreBounds,
-	bool bGenerateOnePolygonGroupPerComponent) const
-{
-	ALandscapeProxy::FRawMeshExportParams Params;
-	Params.ExportLOD = InExportLOD;
-	if (!bIgnoreBounds)
-	{
-		Params.ExportBounds = InBounds;
-	}
-	Params.ComponentsToExport = RegisteredLandscapeComponents;
-
-	return ExportToRawMesh(Params, OutRawMesh);
-}
-PRAGMA_ENABLE_DEPRECATION_WARNINGS
+// ----------------------------------------------------------------------------------
 
 ALandscapeProxy::FRawMeshExportParams::FUVConfiguration::FUVConfiguration()
 {
@@ -3617,6 +3579,9 @@ int32 ALandscapeProxy::FRawMeshExportParams::FUVConfiguration::GetNumUVChannelsN
 	return Result;
 }
 
+
+// ----------------------------------------------------------------------------------
+
 const ALandscapeProxy::FRawMeshExportParams::FUVConfiguration& ALandscapeProxy::FRawMeshExportParams::GetUVConfiguration(int32 InComponentIndex) const
 {
 	return ComponentsUVConfiguration.IsSet() ? (*ComponentsUVConfiguration)[InComponentIndex] : UVConfiguration;
@@ -3640,6 +3605,9 @@ int32 ALandscapeProxy::FRawMeshExportParams::GetNumUVChannelsNeeded() const
 	return Result;
 }
 
+
+// ----------------------------------------------------------------------------------
+
 namespace UE::Landscape
 {
 	const FIntPoint QuadPattern[4] =
diff --git a/Engine/Source/Runtime/Landscape/Private/LandscapeGrass.cpp b/Engine/Source/Runtime/Landscape/Private/LandscapeGrass.cpp
--- a/Engine/Source/Runtime/Landscape/Private/LandscapeGrass.cpp
+++ b/Engine/Source/Runtime/Landscape/Private/LandscapeGrass.cpp
@@ -850,14 +850,14 @@ static void GrassCVarSinkFunction()
 		{
 			if (!World->bIsTearingDown)
 			{
-				if (ULandscapeSubsystem* LandscapeSubsystem = World->GetSubsystem<ULandscapeSubsystem>())
-				{
-					LandscapeSubsystem->RegenerateGrass(/*bInFlushGrass = */true, /*bInForceSync = */true);
-				}
+			if (ULandscapeSubsystem* LandscapeSubsystem = World->GetSubsystem<ULandscapeSubsystem>())
+			{
+				LandscapeSubsystem->RegenerateGrass(/*bInFlushGrass = */true, /*bInForceSync = */true);
 			}
 		}
 	}
 }
+}
 
 static FAutoConsoleVariableSink CVarGrassSink(FConsoleCommandDelegate::CreateStatic(&GrassCVarSinkFunction));
 
@@ -1133,7 +1133,7 @@ class FLandscapeGrassWeightExporter_RenderThread
 {
 	FLandscapeGrassWeightExporter_RenderThread(const TArray<int32>& InHeightMips)
 		: HeightMips(InHeightMips)
-	{
+		{
 	}
 
 	friend class FLandscapeGrassWeightExporter;
@@ -1300,10 +1300,10 @@ public:
 		}
 
 		// center of target area in world
-		FVector TargetCenter = LandscapeProxy->GetTransform().TransformPosition(FVector(TargetSizeMinusOne, 0.f) * 0.5f);
+		FVector TargetCenter = LandscapeProxy->GetTransform().TransformPosition(FVector(TargetSizeMinusOne, 0.f)*0.5f);
 
 		// extent of target in world space
-		FVector TargetExtent = FVector(TargetSize, 0.0f) * LandscapeProxy->GetActorScale() * 0.5f;
+		FVector TargetExtent = FVector(TargetSize, 0.0f)*LandscapeProxy->GetActorScale()*0.5f;
 
 		ViewOrigin = TargetCenter;
 		ViewRotationMatrix = FInverseRotationMatrix(LandscapeProxy->GetActorRotation());
@@ -1342,7 +1342,7 @@ public:
 	TMap<ULandscapeComponent*, TUniquePtr<FLandscapeComponentGrassData>, TInlineSetAllocator<1>> FetchResults()
 	{
 		TArray<FColor> Samples;
-		Samples.SetNumUninitialized(TargetSize.X * TargetSize.Y);
+		Samples.SetNumUninitialized(TargetSize.X*TargetSize.Y);
 
 		// Copy the contents of the remote texture to system memory
 		FReadSurfaceDataFlags ReadSurfaceDataFlags;
@@ -1351,7 +1351,7 @@ public:
 
 		TMap<ULandscapeComponent*, TUniquePtr<FLandscapeComponentGrassData>, TInlineSetAllocator<1>> Results;
 		Results.Reserve(ComponentInfos.Num());
-
+		
 		// Local data will be moved in contiguous array at the end of export (to minimize slack waste)
 		TArray<uint16> HeightData;
 		TMap<ULandscapeGrassType*, TArray<uint8>> WeightData;
@@ -1404,7 +1404,7 @@ public:
 
 			for (int32 PassIdx = 0; PassIdx < ComponentInfo.NumPasses; PassIdx++)
 			{
-				FColor* SampleData = &Samples[ComponentInfo.PixelOffsetX + PassIdx * ComponentSizeVerts];
+				FColor* SampleData = &Samples[ComponentInfo.PixelOffsetX + PassIdx*ComponentSizeVerts];
 				if (PassIdx < ComponentInfo.FirstHeightMipsPassIndex)
 				{
 					if (PassIdx == 0)
@@ -1541,7 +1541,7 @@ FLandscapeComponentGrassData::FLandscapeComponentGrassData(ULandscapeComponent*
 }
 
 void ULandscapeComponent::UpdateGrassTypes()
-{
+	{
 	GrassTypesMaxDiscardDistance = 0.0f;
 	GrassTypes.Reset();
 	
@@ -1550,11 +1550,11 @@ void ULandscapeComponent::UpdateGrassTypes()
 		GrassTypes = Material->GetMaterial()->GetCachedExpressionData().GrassTypes;
 
 		for (const ULandscapeGrassType* GrassType : GrassTypes)
-		{
+{
 			if (GrassType != nullptr)
 			{
 				for (const FGrassVariety& GrassVariety : GrassType->GrassVarieties)
-				{
+	{
 					GrassTypesMaxDiscardDistance = FMath::Max((float)GrassVariety.EndCullDistance.GetValue(), GrassTypesMaxDiscardDistance);
 				}
 			}
@@ -1585,14 +1585,14 @@ uint32 ULandscapeComponent::ComputeGrassMapGenerationHash() const
 			}
 
 			Hash = FCrc::TypeCrc32(GrassType->StateHash, Hash);
+			}
 		}
-	}
 
 	return Hash;
-}
+		}
 
 bool ULandscapeComponent::IsGrassMapOutdated() const
-{
+		{
 	return GrassData->HasValidData() ? (ComputeGrassMapGenerationHash() != GrassData->GenerationHash) : false;
 }
 
@@ -2197,23 +2197,23 @@ FArchive& operator<<(FArchive& Ar, FLandscapeComponentGrassData& Data)
 	{
 		if (Ar.CustomVer(FFortniteMainBranchObjectVersion::GUID) < FFortniteMainBranchObjectVersion::LandscapeSupportPerComponentGrassTypes)
 		{
-			if (Ar.CustomVer(FLandscapeCustomVersion::GUID) >= FLandscapeCustomVersion::GrassMaterialInstanceFix)
-			{
+		if (Ar.CustomVer(FLandscapeCustomVersion::GUID) >= FLandscapeCustomVersion::GrassMaterialInstanceFix)
+		{
 				Ar << Data.MaterialStateIds_DEPRECATED;
-			}
-			else
-			{
+		}
+		else
+		{
 				Data.MaterialStateIds_DEPRECATED.Empty(1);
-				if (Ar.UEVer() >= VER_UE4_SERIALIZE_LANDSCAPE_GRASS_DATA_MATERIAL_GUID)
-				{
-					FGuid MaterialStateId;
-					Ar << MaterialStateId;
+			if (Ar.UEVer() >= VER_UE4_SERIALIZE_LANDSCAPE_GRASS_DATA_MATERIAL_GUID)
+			{
+				FGuid MaterialStateId;
+				Ar << MaterialStateId;
 					Data.MaterialStateIds_DEPRECATED.Add(MaterialStateId);
-				}
 			}
+		}
 
-			if (Ar.CustomVer(FLandscapeCustomVersion::GUID) >= FLandscapeCustomVersion::GrassMaterialWPO)
-			{
+		if (Ar.CustomVer(FLandscapeCustomVersion::GUID) >= FLandscapeCustomVersion::GrassMaterialWPO)
+		{
 				Ar << Data.RotationForWPO_DEPRECATED;
 			}
 		}
@@ -2467,9 +2467,9 @@ struct FAsyncGrassBuilder : public FGrassBuilderBase
 	FFloatInterval ScaleZ;
 	bool bWeightAttenuatesMaxScale;
 	float MaxScaleWeightAttenuation ;
-	bool RandomRotation;
-	bool RandomScale;
-	bool AlignToSurface;
+	bool bRandomRotation;
+	bool bRandomScale;
+	bool bAlignToSurface;
 	float PlacementJitter;
 	FRandomStream RandomStream;
 	FMatrix XForm;
@@ -2480,15 +2480,15 @@ struct FAsyncGrassBuilder : public FGrassBuilderBase
 	int32 TotalInstances;
 	uint32 HaltonBaseIndex;
 
-	bool UseLandscapeLightmap;
+	bool bUseLandscapeLightmap;
 	FVector2D LightmapBaseBias;
 	FVector2D LightmapBaseScale;
 	FVector2D ShadowmapBaseBias;
 	FVector2D ShadowmapBaseScale;
 	FVector2D LightMapComponentBias;
 	FVector2D LightMapComponentScale;
-	bool RequireCPUAccess;
-	bool RequiresInstanceDataForTree;
+	bool bRequiresCPUAccess;
+	bool bRequiresInstanceDataForTree;
 
 	TArray<FBox> ExcludedBoxes;
 
@@ -2507,9 +2507,9 @@ struct FAsyncGrassBuilder : public FGrassBuilderBase
 		, ScaleZ(GrassVariety.ScaleZ)
 		, bWeightAttenuatesMaxScale(GrassVariety.bWeightAttenuatesMaxScale)
 		, MaxScaleWeightAttenuation(GrassVariety.MaxScaleWeightAttenuation)
-		, RandomRotation(GrassVariety.RandomRotation)
-		, RandomScale(false)
-		, AlignToSurface(GrassVariety.AlignToSurface)
+		, bRandomRotation(GrassVariety.RandomRotation)
+		, bRandomScale(false)
+		, bAlignToSurface(GrassVariety.AlignToSurface)
 		, PlacementJitter(GrassVariety.PlacementJitter)
 		, RandomStream(GrassInstancedStaticMeshComponent->InstancingRandomSeed)
 		, XForm(LandscapeToWorld * GrassInstancedStaticMeshComponent->GetComponentTransform().ToMatrixWithScale().Inverse())
@@ -2520,15 +2520,15 @@ struct FAsyncGrassBuilder : public FGrassBuilderBase
 		, TotalInstances(0)
 		, HaltonBaseIndex(InHaltonBaseIndex)
 
-		, UseLandscapeLightmap(GrassVariety.bUseLandscapeLightmap)
+		, bUseLandscapeLightmap(GrassVariety.bUseLandscapeLightmap)
 		, LightmapBaseBias(FVector2D::ZeroVector)
 		, LightmapBaseScale(FVector2D::UnitVector)
 		, ShadowmapBaseBias(FVector2D::ZeroVector)
 		, ShadowmapBaseScale(FVector2D::UnitVector)
 		, LightMapComponentBias(FVector2D::ZeroVector)
 		, LightMapComponentScale(FVector2D::UnitVector)
-		, RequireCPUAccess(GrassVariety.bKeepInstanceBufferCPUCopy)
-		, RequiresInstanceDataForTree(GrassInstancedStaticMeshComponent->RequiresInstanceDataForTree())
+		, bRequiresCPUAccess(GrassVariety.bKeepInstanceBufferCPUCopy)
+		, bRequiresInstanceDataForTree(GrassInstancedStaticMeshComponent->RequiresInstanceDataForTree())
 		// output
 		, InstanceBuffer(/*bSupportsVertexHalfFloat*/ true)
 		, ClusterTree()
@@ -2538,13 +2538,13 @@ struct FAsyncGrassBuilder : public FGrassBuilderBase
 		switch (Scaling)
 		{
 		case EGrassScaling::Uniform:
-			RandomScale = ScaleX.Size() > 0;
+			bRandomScale = ScaleX.Size() > 0;
 			break;
 		case EGrassScaling::Free:
-			RandomScale = ScaleX.Size() > 0 || ScaleY.Size() > 0 || ScaleZ.Size() > 0;
+			bRandomScale = ScaleX.Size() > 0 || ScaleY.Size() > 0 || ScaleZ.Size() > 0;
 			break;
 		case EGrassScaling::LockXY:
-			RandomScale = ScaleX.Size() > 0 || ScaleZ.Size() > 0;
+			bRandomScale = ScaleX.Size() > 0 || ScaleZ.Size() > 0;
 			break;
 		default:
 			check(0);
@@ -2561,11 +2561,11 @@ struct FAsyncGrassBuilder : public FGrassBuilderBase
 
 		bHaveValidData = bHaveValidData && GrassData.IsValid();
 
-		InstanceBuffer.SetAllowCPUAccess(RequireCPUAccess);
+		InstanceBuffer.SetAllowCPUAccess(bRequiresCPUAccess);
 
 		check(DesiredInstancesPerLeaf > 0);
 
-		if (UseLandscapeLightmap)
+		if (bUseLandscapeLightmap)
 		{
 			InitLandscapeLightmap(Component);
 		}
@@ -2618,7 +2618,7 @@ struct FAsyncGrassBuilder : public FGrassBuilderBase
 
 	void SetInstance(int32 InstanceIndex, const FMatrix& InXForm, float RandomFraction)
 	{
-		if (UseLandscapeLightmap)
+		if (bUseLandscapeLightmap)
 		{
 			FMatrix::FReal InstanceX = InXForm.M[3][0];
 			FMatrix::FReal InstanceY = InXForm.M[3][1];
@@ -2737,15 +2737,15 @@ struct FAsyncGrassBuilder : public FGrassBuilderBase
 				FVector LocationWithHeight;
 				FVector ComputedNormal;
 				float Weight = 0.f;
-				SampleLandscapeAtLocationLocal(Location, LocationWithHeight, Weight, AlignToSurface ? &ComputedNormal : nullptr);
+				SampleLandscapeAtLocationLocal(Location, LocationWithHeight, Weight, bAlignToSurface ? &ComputedNormal : nullptr);
 				bool bKeep = Weight > 0.0f && Weight >= RandomStream.GetFraction() && !IsExcluded(LocationWithHeight);
 				if (bKeep)
 				{
-					const FVector Scale = RandomScale ? GetRandomScale(Weight) : DefaultScale;
-					const float Rot = RandomRotation ? RandomStream.GetFraction() * 360.0f : 0.0f;
+					const FVector Scale = bRandomScale ? GetRandomScale(Weight) : DefaultScale;
+					const float Rot = bRandomRotation ? RandomStream.GetFraction() * 360.0f : 0.0f;
 					const FMatrix BaseXForm = FScaleRotationTranslationMatrix(Scale, FRotator(0.0f, Rot, 0.0f), FVector::ZeroVector);
 					FMatrix OutXForm;
-					if (AlignToSurface && !ComputedNormal.IsNearlyZero())
+					if (bAlignToSurface && !ComputedNormal.IsNearlyZero())
 					{
 						const FVector NewZ = ComputedNormal * FMath::Sign(ComputedNormal.Z);
 						const FVector NewX = (FVector(0, -1, 0) ^ NewZ).GetSafeNormal();
@@ -2829,11 +2829,11 @@ struct FAsyncGrassBuilder : public FGrassBuilderBase
 							const FInstanceLocal& Instance = Instances[InstanceIndex];
 							if (Instance.bKeep)
 							{
-								const FVector Scale = RandomScale ? GetRandomScale(Instance.Weight) : DefaultScale;
-								const float Rot = RandomRotation ? RandomStream.GetFraction() * 360.0f : 0.0f;
+								const FVector Scale = bRandomScale ? GetRandomScale(Instance.Weight) : DefaultScale;
+								const float Rot = bRandomRotation ? RandomStream.GetFraction() * 360.0f : 0.0f;
 								const FMatrix BaseXForm = FScaleRotationTranslationMatrix(Scale, FRotator(0.0f, Rot, 0.0f), FVector::ZeroVector);
 								FMatrix OutXForm;
-								if (AlignToSurface)
+								if (bAlignToSurface)
 								{
 									FVector PosX1 = xStart ? Instances[InstanceIndex - SqrtMaxInstances].Pos : Instance.Pos;
 									FVector PosX2 = (xStart + 1 < SqrtMaxInstances) ? Instances[InstanceIndex + SqrtMaxInstances].Pos : Instance.Pos;
@@ -2878,7 +2878,7 @@ struct FAsyncGrassBuilder : public FGrassBuilderBase
 			TArray<float> InstanceCustomDataDummy;
 			UGrassInstancedStaticMeshComponent::BuildTreeAnyThread(InstanceTransforms, InstanceCustomDataDummy, 0, MeshBox, ClusterTree, SortedInstances, InstanceReorderTable, OutOcclusionLayerNum, DesiredInstancesPerLeaf, false);
 
-			if (RequiresInstanceDataForTree)
+			if (bRequiresInstanceDataForTree)
 			{
 				InstanceData.Reset(NumInstances);
 				for (const FMatrix& Transform : InstanceTransforms)
@@ -2896,7 +2896,7 @@ struct FAsyncGrassBuilder : public FGrassBuilderBase
 				{
 					check(LoadFrom > FirstUnfixedIndex);
 					InstanceBuffer.SwapInstance(FirstUnfixedIndex, LoadFrom);
-					if (RequiresInstanceDataForTree)
+					if (bRequiresInstanceDataForTree)
 					{
 						InstanceData.Swap(FirstUnfixedIndex, LoadFrom);
 					}
@@ -3209,7 +3209,7 @@ void ALandscapeProxy::UpdateGrassDataStatus(TSet<UTexture2D*>* OutCurrentForcedS
 	{
 		*OutOutdatedGrassMaps = 0;
 	}
-
+		
 	const UWorld* World = GetWorld();
 	if (!World || World->IsGameWorld())
 	{
@@ -3255,9 +3255,9 @@ void ALandscapeProxy::UpdateGrassDataStatus(TSet<UTexture2D*>* OutCurrentForcedS
 #endif // WITH_EDITOR
 
 				if (OutOutdatedComponents)
-				{
-					OutOutdatedComponents->Add(Component);
-				}
+			{
+				OutOutdatedComponents->Add(Component);
+			}
 			}
 
 			// Needs to be called after UpdateGrassTypes
@@ -3773,7 +3773,7 @@ void ALandscapeProxy::UpdateGrass(const TArray<FVector>& Cameras, int32& InOutNu
 										GrassInstancedStaticMeshComponent->bEvaluateWorldPositionOffset = true;
 										GrassInstancedStaticMeshComponent->WorldPositionOffsetDisableDistance = GrassVariety.InstanceWorldPositionOffsetDisableDistance;
 										GrassInstancedStaticMeshComponent->ShadowCacheInvalidationBehavior = GrassVariety.ShadowCacheInvalidationBehavior;
-
+										
 										GrassInstancedStaticMeshComponent->PrecachePSOs();
 
 										const FMeshMapBuildData* MeshMapBuildData = Component->GetMeshMapBuildData();
@@ -4033,7 +4033,7 @@ void ALandscapeProxy::UpdateGrass(const TArray<FVector>& Cameras, int32& InOutNu
 
 						if (!GrassISMComponent->PerInstanceRenderData.IsValid())
 						{
-							GrassISMComponent->InitPerInstanceRenderData(true, &Inner.Builder->InstanceBuffer, Inner.Builder->RequireCPUAccess);
+							GrassISMComponent->InitPerInstanceRenderData(true, &Inner.Builder->InstanceBuffer, Inner.Builder->bRequiresCPUAccess);
 						}
 						else
 						{
diff --git a/Engine/Source/Runtime/Landscape/Private/LandscapeSplines.cpp b/Engine/Source/Runtime/Landscape/Private/LandscapeSplines.cpp
--- a/Engine/Source/Runtime/Landscape/Private/LandscapeSplines.cpp
+++ b/Engine/Source/Runtime/Landscape/Private/LandscapeSplines.cpp
@@ -1136,13 +1136,6 @@ void ULandscapeSplinesComponent::RequestSplineLayerUpdate()
 	}
 }
 
-PRAGMA_DISABLE_DEPRECATION_WARNINGS
-void ULandscapeSplinesComponent::SetDefaultEditorSplineMesh()
-{
-	SplineEditorMesh = LoadObject<UStaticMesh>(nullptr, TEXT("/Engine/EditorLandscapeResources/SplineEditorMesh"));
-}
-PRAGMA_ENABLE_DEPRECATION_WARNINGS
-
 void ULandscapeSplinesComponent::ShowSplineEditorMesh(bool bShow)
 {
 	bShowSplineEditorMesh = bShow;
diff --git a/Engine/Source/Runtime/Landscape/Public/LandscapeRenderMobile.h b/Engine/Source/Runtime/Landscape/Public/LandscapeRenderMobile.h
deleted file mode 100644
