Remove assert on null AnimScriptInstance in Shared Linked Layers module as this can happen under valid circumstances with PostProcessAnimInstances

diff --git a/Engine/Source/Runtime/Engine/Private/Animation/AnimSubsystem_SharedLinkedAnimLayers.cpp b/Engine/Source/Runtime/Engine/Private/Animation/AnimSubsystem_SharedLinkedAnimLayers.cpp
--- a/Engine/Source/Runtime/Engine/Private/Animation/AnimSubsystem_SharedLinkedAnimLayers.cpp
+++ b/Engine/Source/Runtime/Engine/Private/Animation/AnimSubsystem_SharedLinkedAnimLayers.cpp
@@ -213,8 +213,13 @@ FAnimSubsystem_SharedLinkedAnimLayers* FAnimSubsystem_SharedLinkedAnimLayers::Ge
 	}
 #endif
 	check(SkelMesh);
-	check(SkelMesh->GetAnimInstance());
-	return SkelMesh->GetAnimInstance()->FindSubsystem<FAnimSubsystem_SharedLinkedAnimLayers>();
+
+	// In some cases we have a PostProcessAnimInstance but no AnimScriptInstance
+	if (SkelMesh->GetAnimInstance())
+	{
+		return SkelMesh->GetAnimInstance()->FindSubsystem<FAnimSubsystem_SharedLinkedAnimLayers>();
+	}
+	return nullptr;
 }
 
 void FAnimSubsystem_SharedLinkedAnimLayers::Reset()
