Make sure artist cannot import a hires mesh with more section than the base LOD. Or warn artist if they reimport the LOD 0 with less.

#rnx

diff --git a/Engine/Source/Runtime/Engine/Private/StaticMesh.cpp b/Engine/Source/Runtime/Engine/Private/StaticMesh.cpp
--- a/Engine/Source/Runtime/Engine/Private/StaticMesh.cpp
+++ b/Engine/Source/Runtime/Engine/Private/StaticMesh.cpp
@@ -4119,6 +4119,22 @@ EDataValidationResult UStaticMesh::IsDataValid(FDataValidationContext& Context)
 			Context.AddError(LOCTEXT("StaticMeshValidation_UninitializedLOD0", "This Static Mesh Asset has no LOD0 Source Model mesh. This asset is not repairable, the Asset is corrupted and must be deleted."));
 			ValidationResult = EDataValidationResult::Invalid;
 		}
+		
+		if (!GIsBuildMachine && IsHiResMeshDescriptionValid())
+		{
+			if (const FMeshDescription* BaseLodMeshDescription = GetMeshDescription(0))
+			{
+				if (const FMeshDescription* HiResMeshDescription = GetHiResMeshDescription())
+				{
+					//Validate the number of sections
+					if (HiResMeshDescription->PolygonGroups().Num() > BaseLodMeshDescription->PolygonGroups().Num())
+					{
+						Context.AddError(LOCTEXT("StaticMeshValidation_HiresMoreSectionThanLod0", "Invalid hi-res mesh description. The number of sections from the hires mesh is higher than LOD 0 section count. This is not supported and LOD 0 will be used as a fallback to build nanite data."));
+						ValidationResult = EDataValidationResult::Invalid;
+					}
+				}
+			}
+		}
 	}
 
 	return ValidationResult;
