Deprecating unused view family gamma correction.

diff --git a/Engine/Plugins/Experimental/GPULightmass/Source/GPULightmass/Private/LightmapRenderer.cpp b/Engine/Plugins/Experimental/GPULightmass/Source/GPULightmass/Private/LightmapRenderer.cpp
--- a/Engine/Plugins/Experimental/GPULightmass/Source/GPULightmass/Private/LightmapRenderer.cpp
+++ b/Engine/Plugins/Experimental/GPULightmass/Source/GPULightmass/Private/LightmapRenderer.cpp
@@ -876,8 +876,7 @@ bool FSceneRenderState::SetupRayTracingScene(FRDGBuilder& GraphBuilder, int32 LO
 		nullptr,
 		nullptr,
 		FEngineShowFlags(ESFIM_Game))
-		.SetTime(FGameTime())
-		.SetGammaCorrection(1.0f));
+		.SetTime(FGameTime()));
 
 	const FIntRect ViewRect(FIntPoint(0, 0), FIntPoint(GPreviewLightmapPhysicalTileSize, GPreviewLightmapPhysicalTileSize));
 
diff --git a/Engine/Plugins/FX/Niagara/Source/Niagara/Private/NiagaraGpuComputeDispatch.cpp b/Engine/Plugins/FX/Niagara/Source/Niagara/Private/NiagaraGpuComputeDispatch.cpp
--- a/Engine/Plugins/FX/Niagara/Source/Niagara/Private/NiagaraGpuComputeDispatch.cpp
+++ b/Engine/Plugins/FX/Niagara/Source/Niagara/Private/NiagaraGpuComputeDispatch.cpp
@@ -406,7 +406,6 @@ void FNiagaraGpuComputeDispatch::ProcessPendingTicksFlush(FRHICommandListImmedia
 			FSceneViewFamilyContext ViewFamily(
 				FSceneViewFamily::ConstructionValues(nullptr, GetSceneInterface(), FEngineShowFlags(ESFIM_Game))
 				.SetTime(CachedViewInitOptions.GameTime)
-				.SetGammaCorrection(CachedViewInitOptions.GammaCorrection)
 			);
 
 			FSceneViewInitOptions ViewInitOptions;
@@ -2026,7 +2025,6 @@ void FNiagaraGpuComputeDispatch::PostRenderOpaque(FRDGBuilder& GraphBuilder, TCo
 		if (const FSceneViewFamily* ViewFamily = Views[0].Family)
 		{
 			CachedViewInitOptions.GameTime			= ViewFamily->Time;
-			CachedViewInitOptions.GammaCorrection	= ViewFamily->GammaCorrection;
 		}
 
 		CachedViewInitOptions.ViewRect				= UE::FXRenderingUtils::GetRawViewRectUnsafe(Views[0]);
diff --git a/Engine/Plugins/FX/Niagara/Source/Niagara/Private/NiagaraGpuComputeDispatch.h b/Engine/Plugins/FX/Niagara/Source/Niagara/Private/NiagaraGpuComputeDispatch.h
--- a/Engine/Plugins/FX/Niagara/Source/Niagara/Private/NiagaraGpuComputeDispatch.h
+++ b/Engine/Plugins/FX/Niagara/Source/Niagara/Private/NiagaraGpuComputeDispatch.h
@@ -270,7 +270,6 @@ private:
 	struct FCachedViewInitOptions
 	{
 		FGameTime	GameTime;
-		float		GammaCorrection		= 1.0f;
 		FIntRect	ViewRect			= FIntRect(0, 0, 64, 64);
 		FVector		ViewOrigin			= FVector::ZeroVector;
 		FMatrix		ViewRotationMatrix	= FMatrix::Identity;
diff --git a/Engine/Plugins/FX/Niagara/Source/NiagaraEditor/Private/NiagaraBakerRenderer.cpp b/Engine/Plugins/FX/Niagara/Source/NiagaraEditor/Private/NiagaraBakerRenderer.cpp
--- a/Engine/Plugins/FX/Niagara/Source/NiagaraEditor/Private/NiagaraBakerRenderer.cpp
+++ b/Engine/Plugins/FX/Niagara/Source/NiagaraEditor/Private/NiagaraBakerRenderer.cpp
@@ -418,7 +418,6 @@ void FNiagaraBakerRenderer::RenderBufferVisualization(UTextureRenderTarget2D* Re
 	}
 
 	const FIntRect ViewRect = FIntRect(0, 0, RenderTarget->GetSurfaceWidth(), RenderTarget->GetSurfaceHeight());
-	const float GammaCorrection = 1.0f;
 	const float WorldTime = GetWorldTime();
 	UWorld* World = GetWorld();
 
@@ -429,7 +428,6 @@ void FNiagaraBakerRenderer::RenderBufferVisualization(UTextureRenderTarget2D* Re
 	FSceneViewFamilyContext ViewFamily(
 		FSceneViewFamily::ConstructionValues(RenderTarget->GameThread_GetRenderTargetResource(), World->Scene, FEngineShowFlags(ESFIM_Game))
 		.SetTime(FGameTime::CreateUndilated(WorldTime, FApp::GetDeltaTime()))
-		.SetGammaCorrection(GammaCorrection)
 	);
 	
 	ViewFamily.EngineShowFlags.SetScreenPercentage(false);
diff --git a/Engine/Plugins/Media/PixelStreaming/Source/PixelStreaming/Private/UtilsRender.h b/Engine/Plugins/Media/PixelStreaming/Source/PixelStreaming/Private/UtilsRender.h
--- a/Engine/Plugins/Media/PixelStreaming/Source/PixelStreaming/Private/UtilsRender.h
+++ b/Engine/Plugins/Media/PixelStreaming/Source/PixelStreaming/Private/UtilsRender.h
@@ -152,8 +152,7 @@ namespace UE::PixelStreaming
 
 			//Dummy ViewFamily/ViewInfo created to use built in Draw Screen/Texture Pass
 			FSceneViewFamilyContext ViewFamily(FSceneViewFamily::ConstructionValues(nullptr, nullptr, FEngineShowFlags(ESFIM_Game))
-				.SetTime(FGameTime())
-				.SetGammaCorrection(1.0f));
+				.SetTime(FGameTime()));
 			FSceneViewInitOptions ViewInitOptions;
 			ViewInitOptions.ViewFamily = &ViewFamily;
 			ViewInitOptions.SetViewRectangle(ViewRect);
diff --git a/Engine/Plugins/Media/PixelStreaming/Source/PixelStreamingEditor/Private/PixelStreamingVideoInputBackBufferComposited.cpp b/Engine/Plugins/Media/PixelStreaming/Source/PixelStreamingEditor/Private/PixelStreamingVideoInputBackBufferComposited.cpp
--- a/Engine/Plugins/Media/PixelStreaming/Source/PixelStreamingEditor/Private/PixelStreamingVideoInputBackBufferComposited.cpp
+++ b/Engine/Plugins/Media/PixelStreaming/Source/PixelStreamingEditor/Private/PixelStreamingVideoInputBackBufferComposited.cpp
@@ -236,8 +236,7 @@ void FPixelStreamingVideoInputBackBufferComposited::CompositeWindows()
 
 			// Dummy ViewFamily/ViewInfo created to use built in Draw Screen/Texture Pass
 			FSceneViewFamilyContext ViewFamily(FSceneViewFamily::ConstructionValues(nullptr, nullptr, FEngineShowFlags(ESFIM_Game))
-											.SetTime(FGameTime())
-											.SetGammaCorrection(1.0f));
+											.SetTime(FGameTime()));
 			FSceneViewInitOptions ViewInitOptions;
 			ViewInitOptions.ViewFamily = &ViewFamily;
 			ViewInitOptions.SetViewRectangle(ViewRect);
diff --git a/Engine/Plugins/Runtime/nDisplay/Source/DisplayCluster/Private/Render/Viewport/DisplayClusterViewportManager.cpp b/Engine/Plugins/Runtime/nDisplay/Source/DisplayCluster/Private/Render/Viewport/DisplayClusterViewportManager.cpp
--- a/Engine/Plugins/Runtime/nDisplay/Source/DisplayCluster/Private/Render/Viewport/DisplayClusterViewportManager.cpp
+++ b/Engine/Plugins/Runtime/nDisplay/Source/DisplayCluster/Private/Render/Viewport/DisplayClusterViewportManager.cpp
@@ -737,7 +737,6 @@ FSceneViewFamily::ConstructionValues FDisplayClusterViewportManager::CreateViewF
 	return FSceneViewFamily::ConstructionValues(InFrameTarget.RenderTargetPtr, InScene, InEngineShowFlags)
 		.SetResolveScene(bResolveScene)
 		.SetRealtimeUpdate(true)
-		.SetGammaCorrection(1.0f)
 		.SetAdditionalViewFamily(bInAdditionalViewFamily);
 }
 
diff --git a/Engine/Plugins/Runtime/nDisplay/Source/DisplayClusterLightCardEditorShaders/Private/DisplayClusterMeshProjectionRenderer.cpp b/Engine/Plugins/Runtime/nDisplay/Source/DisplayClusterLightCardEditorShaders/Private/DisplayClusterMeshProjectionRenderer.cpp
--- a/Engine/Plugins/Runtime/nDisplay/Source/DisplayClusterLightCardEditorShaders/Private/DisplayClusterMeshProjectionRenderer.cpp
+++ b/Engine/Plugins/Runtime/nDisplay/Source/DisplayClusterLightCardEditorShaders/Private/DisplayClusterMeshProjectionRenderer.cpp
@@ -917,8 +917,7 @@ void FDisplayClusterMeshProjectionRenderer::Render(FCanvas* Canvas, FSceneInterf
 				RenderTarget,
 				Scene,
 				RenderSettings.EngineShowFlags)
-				.SetTime(FGameTime::GetTimeSinceAppStart())
-				.SetGammaCorrection(1.0f));
+				.SetTime(FGameTime::GetTimeSinceAppStart()));
 
 			if (Scene)
 			{
diff --git a/Engine/Plugins/Runtime/nDisplay/Source/DisplayClusterShaders/Private/Shaders/DisplayClusterShadersPreprocess_UVLightCards.cpp b/Engine/Plugins/Runtime/nDisplay/Source/DisplayClusterShaders/Private/Shaders/DisplayClusterShadersPreprocess_UVLightCards.cpp
--- a/Engine/Plugins/Runtime/nDisplay/Source/DisplayClusterShaders/Private/Shaders/DisplayClusterShadersPreprocess_UVLightCards.cpp
+++ b/Engine/Plugins/Runtime/nDisplay/Source/DisplayClusterShaders/Private/Shaders/DisplayClusterShadersPreprocess_UVLightCards.cpp
@@ -232,8 +232,7 @@ bool FDisplayClusterShadersPreprocess_UVLightCards::RenderPreprocess_UVLightCard
 		InRenderTarget,
 		InScene,
 		EngineShowFlags)
-		.SetTime(FGameTime::GetTimeSinceAppStart())
-		.SetGammaCorrection(1.0f));
+		.SetTime(FGameTime::GetTimeSinceAppStart()));
 
 	FScenePrimitiveRenderingContextScopeHelper ScenePrimitiveRenderingContextScopeHelper(GetRendererModule().BeginScenePrimitiveRendering(GraphBuilder, &ViewFamily));
 
diff --git a/Engine/Plugins/VirtualProduction/TextureShare/Source/TextureShare/Private/Game/ViewExtension/TextureShareSceneViewExtension.cpp b/Engine/Plugins/VirtualProduction/TextureShare/Source/TextureShare/Private/Game/ViewExtension/TextureShareSceneViewExtension.cpp
--- a/Engine/Plugins/VirtualProduction/TextureShare/Source/TextureShare/Private/Game/ViewExtension/TextureShareSceneViewExtension.cpp
+++ b/Engine/Plugins/VirtualProduction/TextureShare/Source/TextureShare/Private/Game/ViewExtension/TextureShareSceneViewExtension.cpp
@@ -66,7 +66,6 @@ namespace UE::TextureShare::SceneViewExtension
 
 		OutViewFamily.FrameNumber = InViewFamily.FrameNumber;
 		OutViewFamily.bIsHDR = InViewFamily.bIsHDR;
-		OutViewFamily.GammaCorrection = InViewFamily.GammaCorrection;
 		OutViewFamily.SecondaryViewFraction = InViewFamily.SecondaryViewFraction;
 	}
 };
diff --git a/Engine/Source/Developer/MaterialBaking/Private/MaterialBakingModule.cpp b/Engine/Source/Developer/MaterialBaking/Private/MaterialBakingModule.cpp
--- a/Engine/Source/Developer/MaterialBaking/Private/MaterialBakingModule.cpp
+++ b/Engine/Source/Developer/MaterialBaking/Private/MaterialBakingModule.cpp
@@ -715,8 +715,7 @@ private:
 
 				FSceneViewFamily ViewFamily(FSceneViewFamily::ConstructionValues(RenderTarget->GetRenderTargetResource(), nullptr,
 					FEngineShowFlags(ESFIM_Game))
-					.SetTime(FGameTime())
-					.SetGammaCorrection(RenderTarget->GetRenderTargetResource()->GetDisplayGamma()));
+					.SetTime(FGameTime()));
 
 				RenderItem->MaterialRenderProxy = ExportMaterialProxy;
 				RenderItem->ViewFamily = &ViewFamily;
diff --git a/Engine/Source/Developer/MaterialUtilities/Private/MeshRendering.cpp b/Engine/Source/Developer/MaterialUtilities/Private/MeshRendering.cpp
--- a/Engine/Source/Developer/MaterialUtilities/Private/MeshRendering.cpp
+++ b/Engine/Source/Developer/MaterialUtilities/Private/MeshRendering.cpp
@@ -647,8 +647,7 @@ bool FMeshRenderer::RenderMaterial(struct FMaterialMergeData& InMaterialData, FM
 			CanvasRenderTarget,
 			NULL,
 			FEngineShowFlags(ESFIM_Game))
-			.SetTime(FGameTime())
-			.SetGammaCorrection(CanvasRenderTarget->GetDisplayGamma()));
+			.SetTime(FGameTime()));
 		
 #if !SHOW_WIREFRAME_MESH
 		Canvas.Clear(InRenderTarget->ClearColor);
@@ -766,8 +765,7 @@ bool FMeshRenderer::RenderMaterialTexCoordScales(struct FMaterialMergeData& InMa
 	ShowFlags.OutputMaterialTextureScales = true; // This will bind the DVSM_OutputMaterialTextureScales
 
 	FSceneViewFamily ViewFamily(FSceneViewFamily::ConstructionValues(CanvasRenderTarget, nullptr, ShowFlags)
-		.SetTime(FGameTime())
-		.SetGammaCorrection(CanvasRenderTarget->GetDisplayGamma()));
+		.SetTime(FGameTime()));
 
 	// The next line ensures a constant view vector of (0,0,1) for all pixels. Required because here SVPositionToTranslatedWorld is identity, making excessive view angle increase per pixel.
 	// That creates bad side effects for anything that depends on the view vector, like parallax or bump offset mappings. For those, we want the tangent
diff --git a/Engine/Source/Editor/UnrealEd/Private/Lightmass/LightmassLandscapeRender.cpp b/Engine/Source/Editor/UnrealEd/Private/Lightmass/LightmassLandscapeRender.cpp
--- a/Engine/Source/Editor/UnrealEd/Private/Lightmass/LightmassLandscapeRender.cpp
+++ b/Engine/Source/Editor/UnrealEd/Private/Lightmass/LightmassLandscapeRender.cpp
@@ -94,8 +94,7 @@ void RenderLandscapeMaterialForLightmass(const FLandscapeStaticLightingMesh* Lan
 		RenderTarget,
 		NULL,
 		FEngineShowFlags(ESFIM_Game))
-		.SetTime(FGameTime())
-		.SetGammaCorrection(RenderTarget->GetDisplayGamma()));
+		.SetTime(FGameTime()));
 
 	FDynamicMeshBuilder DynamicMeshBuilder(ViewFamily.GetFeatureLevel(), 4, 0, true);
 	DynamicMeshBuilder.AddVertices(Vertices);
diff --git a/Engine/Source/Runtime/Engine/Private/SceneView.cpp b/Engine/Source/Runtime/Engine/Private/SceneView.cpp
--- a/Engine/Source/Runtime/Engine/Private/SceneView.cpp
+++ b/Engine/Source/Runtime/Engine/Private/SceneView.cpp
@@ -2712,7 +2712,9 @@ FSceneViewFamily::ConstructionValues::ConstructionValues(
 ,	Scene(InScene)
 ,	EngineShowFlags(InEngineShowFlags)
 ,	ViewModeParam(-1)
+PRAGMA_DISABLE_DEPRECATION_WARNINGS
 ,	GammaCorrection(1.0f)
+PRAGMA_ENABLE_DEPRECATION_WARNINGS
 ,	bAdditionalViewFamily(false)
 ,	bRealtimeUpdate(false)
 ,	bDeferClear(false)
@@ -2751,7 +2753,9 @@ FSceneViewFamily::FSceneViewFamily(const ConstructionValues& CVS)
 	bWorldIsPaused(false),
 	bIsHDR(false),
 	bRequireMultiView(false),
+PRAGMA_DISABLE_DEPRECATION_WARNINGS
 	GammaCorrection(CVS.GammaCorrection),
+PRAGMA_ENABLE_DEPRECATION_WARNINGS
 	SecondaryViewFraction(1.0f),
 	SecondaryScreenPercentageMethod(ESecondaryScreenPercentageMethod::LowerPixelDensitySimulation),
 	ProfileSceneRenderTime(nullptr),
diff --git a/Engine/Source/Runtime/Engine/Private/TileRendering.cpp b/Engine/Source/Runtime/Engine/Private/TileRendering.cpp
--- a/Engine/Source/Runtime/Engine/Private/TileRendering.cpp
+++ b/Engine/Source/Runtime/Engine/Private/TileRendering.cpp
@@ -223,8 +223,7 @@ bool FCanvasTileRendererItem::Render_RenderThread(FCanvasRenderContext& RenderCo
 		CanvasRenderTarget,
 		nullptr,
 		FEngineShowFlags(ESFIM_Game))
-		.SetTime(Time)
-		.SetGammaCorrection(CanvasRenderTarget->GetDisplayGamma()));
+		.SetTime(Time));
 
 	const FIntRect ViewRect(FIntPoint(0, 0), CanvasRenderTarget->GetSizeXY());
 
@@ -268,8 +267,7 @@ bool FCanvasTileRendererItem::Render_GameThread(const FCanvas* Canvas, FCanvasRe
 			CanvasRenderTarget,
 			Canvas->GetScene(),
 			FEngineShowFlags(ESFIM_Game))
-			.SetTime(Time)
-			.SetGammaCorrection(CanvasRenderTarget->GetDisplayGamma()));
+			.SetTime(Time));
 
 		const FIntRect ViewRect(FIntPoint(0, 0), CanvasRenderTarget->GetSizeXY());
 
diff --git a/Engine/Source/Runtime/Engine/Private/TriangleRendering.cpp b/Engine/Source/Runtime/Engine/Private/TriangleRendering.cpp
--- a/Engine/Source/Runtime/Engine/Private/TriangleRendering.cpp
+++ b/Engine/Source/Runtime/Engine/Private/TriangleRendering.cpp
@@ -199,8 +199,7 @@ bool FCanvasTriangleRendererItem::Render_RenderThread(FCanvasRenderContext& Rend
 		CanvasRenderTarget,
 		nullptr,
 		FEngineShowFlags(ESFIM_Game))
-		.SetTime(Time)
-		.SetGammaCorrection(CanvasRenderTarget->GetDisplayGamma()));
+		.SetTime(Time));
 
 	const FIntRect ViewRect(FIntPoint(0, 0), CanvasRenderTarget->GetSizeXY());
 
@@ -243,8 +242,7 @@ bool FCanvasTriangleRendererItem::Render_GameThread(const FCanvas* Canvas, FCanv
 		CanvasRenderTarget,
 		Canvas->GetScene(),
 		FEngineShowFlags(ESFIM_Game))
-		.SetTime(Time)
-		.SetGammaCorrection(CanvasRenderTarget->GetDisplayGamma()));
+		.SetTime(Time));
 
 	const FIntRect ViewRect(FIntPoint(0, 0), CanvasRenderTarget->GetSizeXY());
 
diff --git a/Engine/Source/Runtime/Engine/Public/SceneView.h b/Engine/Source/Runtime/Engine/Public/SceneView.h
--- a/Engine/Source/Runtime/Engine/Public/SceneView.h
+++ b/Engine/Source/Runtime/Engine/Public/SceneView.h
@@ -1931,6 +1931,13 @@ public:
 	*/
 	struct ConstructionValues
 	{
+		PRAGMA_DISABLE_DEPRECATION_WARNINGS
+		ConstructionValues(ConstructionValues&&) = default;
+		ConstructionValues(const ConstructionValues& Other) = default;
+		ConstructionValues& operator=(ConstructionValues&&) = default;
+		ConstructionValues& operator=(const ConstructionValues& Other) = default;
+		PRAGMA_ENABLE_DEPRECATION_WARNINGS
+
 		ENGINE_API ConstructionValues(
 			const FRenderTarget* InRenderTarget,
 			FSceneInterface* InScene,
@@ -1955,6 +1962,7 @@ public:
 		FGameTime Time;
 
 		/** Gamma correction used when rendering this family. Default is 1.0 */
+		UE_DEPRECATED(5.4, "Unused gamma correction.")
 		float GammaCorrection;
 
 		/** Indicates whether the view family is additional. */
@@ -1999,7 +2007,13 @@ public:
 		ConstructionValues& SetResolveScene(const bool Value) { bResolveScene = Value; return *this; }
 		
 		/** Set Gamma correction used when rendering this family. */
-		ConstructionValues& SetGammaCorrection(const float Value) { GammaCorrection = Value; return *this; }		
+		UE_DEPRECATED(5.4, "Unused gamma correction.")
+		ConstructionValues& SetGammaCorrection(const float Value)
+		{
+			PRAGMA_DISABLE_DEPRECATION_WARNINGS
+			GammaCorrection = Value; return *this;
+			PRAGMA_ENABLE_DEPRECATION_WARNINGS
+		}
 
 		/** Set the view param. */
 		ConstructionValues& SetViewModeParam(const int InViewModeParam, const FName& InViewModeParamName) { ViewModeParam = InViewModeParam; ViewModeParamName = InViewModeParamName; return *this; }		
@@ -2105,6 +2119,7 @@ public:
 	bool bRequireMultiView;
 
 	/** Gamma correction used when rendering this family. Default is 1.0 */
+	UE_DEPRECATED(5.4, "Unused gamma correction.")
 	float GammaCorrection;
 	
 	/** DPI scale to be used for debugging font. */
diff --git a/Engine/Source/Runtime/Renderer/Private/SceneRendering.cpp b/Engine/Source/Runtime/Renderer/Private/SceneRendering.cpp
--- a/Engine/Source/Runtime/Renderer/Private/SceneRendering.cpp
+++ b/Engine/Source/Runtime/Renderer/Private/SceneRendering.cpp
@@ -5386,7 +5386,6 @@ static void DisplayInternals(FRDGBuilder& GraphBuilder, FViewInfo& InView)
 			CANVAS_LINE(false, TEXT("  Time (Real/World/DeltaWorld): %.2f/%.2f/%.2f"), Family->Time.GetRealTimeSeconds(), Family->Time.GetWorldTimeSeconds(), Family->Time.GetDeltaWorldTimeSeconds())
 			CANVAS_LINE(false, TEXT("  FrameNumber: %u"), Family->FrameNumber)
 			CANVAS_LINE(false, TEXT("  ExposureSettings: %s"), *Family->ExposureSettings.ToString())
-			CANVAS_LINE(false, TEXT("  GammaCorrection: %.2f"), Family->GammaCorrection)
 
 			CANVAS_HEADER(TEXT("View:"))
 			CANVAS_LINE(false, TEXT("  TemporalJitter: %.2f/%.2f"), ViewInfo.TemporalJitterPixels.X, ViewInfo.TemporalJitterPixels.Y)
diff --git a/Engine/Source/Runtime/SlateRHIRenderer/Private/SlateRHIRenderingPolicy.cpp b/Engine/Source/Runtime/SlateRHIRenderer/Private/SlateRHIRenderingPolicy.cpp
--- a/Engine/Source/Runtime/SlateRHIRenderer/Private/SlateRHIRenderingPolicy.cpp
+++ b/Engine/Source/Runtime/SlateRHIRenderer/Private/SlateRHIRenderingPolicy.cpp
@@ -747,7 +747,6 @@ void FSlateRHIRenderingPolicy::DrawElements(
 					DefaultShowFlags
 				)
 				.SetTime(Params.Time)
-				.SetGammaCorrection(DisplayGamma)
 				.SetRealtimeUpdate(true)
 			);
 			SceneViews[i] = CreateSceneView(SceneViewFamilyContexts[i], BackBuffer, FMatrix(Params.ViewProjectionMatrix));
@@ -762,7 +761,6 @@ void FSlateRHIRenderingPolicy::DrawElements(
 				DefaultShowFlags
 			)
 			.SetTime(Params.Time)
-			.SetGammaCorrection(DisplayGamma)
 			.SetRealtimeUpdate(true)
 		);
 		SceneViews[NumScenes - 1] = CreateSceneView(SceneViewFamilyContexts[NumScenes - 1], BackBuffer, FMatrix(Params.ViewProjectionMatrix));
