Make blend profiles blueprintable to allow pinning to anim node inputs

diff --git a/Engine/Source/Runtime/Engine/Classes/Animation/BlendProfile.h b/Engine/Source/Runtime/Engine/Classes/Animation/BlendProfile.h
--- a/Engine/Source/Runtime/Engine/Classes/Animation/BlendProfile.h
+++ b/Engine/Source/Runtime/Engine/Classes/Animation/BlendProfile.h
@@ -54,7 +54,7 @@ struct FBlendProfileBoneEntry
 /** A blend profile is a set of per-bone scales that can be used in transitions and blend lists
  *  to tweak the weights of specific bones. The scales are applied to the normal weight for that bone
  */
-UCLASS(Within=Skeleton, MinimalAPI)
+UCLASS(Within=Skeleton, MinimalAPI, BlueprintType)
 class UBlendProfile : public UObject, public IInterpolationIndexProvider
 {
 public:
