Fixed compilation error in Mac editor builds when using UWorld::IsStreamingLevelBeingConsidered or UWorld::HasStreamingLevelsToConsider.

diff --git a/Engine/Source/Runtime/Engine/Classes/Engine/World.h b/Engine/Source/Runtime/Engine/Classes/Engine/World.h
--- a/Engine/Source/Runtime/Engine/Classes/Engine/World.h
+++ b/Engine/Source/Runtime/Engine/Classes/Engine/World.h
@@ -784,7 +784,7 @@ private:
 };
 
 USTRUCT()
-struct FStreamingLevelsToConsider
+struct ENGINE_API FStreamingLevelsToConsider
 {
 	GENERATED_BODY()
 
