WIP New HLSL translator material expression implementations
MaterialExpressionTangentOutput
MaterialExpressionAtmosphericLightColor
MaterialExpressionLightVector
MaterialExpressionDecalDerivative
MaterialExpressionHairColor

#tests loaded various materials in editor, loaded level

diff --git a/Engine/Source/Runtime/Engine/Classes/Materials/MaterialExpressionAtmosphericLightColor.h b/Engine/Source/Runtime/Engine/Classes/Materials/MaterialExpressionAtmosphericLightColor.h
--- a/Engine/Source/Runtime/Engine/Classes/Materials/MaterialExpressionAtmosphericLightColor.h
+++ b/Engine/Source/Runtime/Engine/Classes/Materials/MaterialExpressionAtmosphericLightColor.h
@@ -17,6 +17,8 @@ class UMaterialExpressionAtmosphericLightColor : public UMaterialExpression
 #if WITH_EDITOR
 	virtual int32 Compile(class FMaterialCompiler* Compiler, int32 OutputIndex) override;
 	virtual void GetCaption(TArray<FString>& OutCaptions) const override;
+
+	virtual bool GenerateHLSLExpression(FMaterialHLSLGenerator& Generator, UE::HLSLTree::FScope& Scope, int32 OutputIndex, UE::HLSLTree::FExpression const*& OutExpression) const override;
 #endif
 	//~ End UMaterialExpression Interface
 };
diff --git a/Engine/Source/Runtime/Engine/Classes/Materials/MaterialExpressionDecalDerivative.h b/Engine/Source/Runtime/Engine/Classes/Materials/MaterialExpressionDecalDerivative.h
--- a/Engine/Source/Runtime/Engine/Classes/Materials/MaterialExpressionDecalDerivative.h
+++ b/Engine/Source/Runtime/Engine/Classes/Materials/MaterialExpressionDecalDerivative.h
@@ -18,6 +18,8 @@ class UMaterialExpressionDecalDerivative : public UMaterialExpression
 #if WITH_EDITOR
 	virtual int32 Compile(class FMaterialCompiler* Compiler, int32 OutputIndex) override;
 	virtual void GetCaption(TArray<FString>& OutCaptions) const override;
+
+	virtual bool GenerateHLSLExpression(FMaterialHLSLGenerator& Generator, UE::HLSLTree::FScope& Scope, int32 OutputIndex, UE::HLSLTree::FExpression const*& OutExpression) const override;
 #endif
 	//~ End UMaterialExpression Interface
 };
diff --git a/Engine/Source/Runtime/Engine/Classes/Materials/MaterialExpressionHairColor.h b/Engine/Source/Runtime/Engine/Classes/Materials/MaterialExpressionHairColor.h
--- a/Engine/Source/Runtime/Engine/Classes/Materials/MaterialExpressionHairColor.h
+++ b/Engine/Source/Runtime/Engine/Classes/Materials/MaterialExpressionHairColor.h
@@ -30,6 +30,8 @@ class UMaterialExpressionHairColor : public UMaterialExpression
 #if WITH_EDITOR
 	virtual int32 Compile(class FMaterialCompiler* Compiler, int32 OutputIndex) override;
 	virtual void GetCaption(TArray<FString>& OutCaptions) const override;
+
+	virtual bool GenerateHLSLExpression(FMaterialHLSLGenerator& Generator, UE::HLSLTree::FScope& Scope, int32 OutputIndex, UE::HLSLTree::FExpression const*& OutExpression) const override;
 #endif
 	//~ End UMaterialExpression Interface
 };
diff --git a/Engine/Source/Runtime/Engine/Classes/Materials/MaterialExpressionLightVector.h b/Engine/Source/Runtime/Engine/Classes/Materials/MaterialExpressionLightVector.h
--- a/Engine/Source/Runtime/Engine/Classes/Materials/MaterialExpressionLightVector.h
+++ b/Engine/Source/Runtime/Engine/Classes/Materials/MaterialExpressionLightVector.h
@@ -18,6 +18,8 @@ class UMaterialExpressionLightVector : public UMaterialExpression
 #if WITH_EDITOR
 	virtual int32 Compile(class FMaterialCompiler* Compiler, int32 OutputIndex) override;
 	virtual void GetCaption(TArray<FString>& OutCaptions) const override;
+
+	virtual bool GenerateHLSLExpression(FMaterialHLSLGenerator& Generator, UE::HLSLTree::FScope& Scope, int32 OutputIndex, UE::HLSLTree::FExpression const*& OutExpression) const override;
 #endif
 	//~ End UMaterialExpression Interface
 };
diff --git a/Engine/Source/Runtime/Engine/Classes/Materials/MaterialExpressionTangentOutput.h b/Engine/Source/Runtime/Engine/Classes/Materials/MaterialExpressionTangentOutput.h
--- a/Engine/Source/Runtime/Engine/Classes/Materials/MaterialExpressionTangentOutput.h
+++ b/Engine/Source/Runtime/Engine/Classes/Materials/MaterialExpressionTangentOutput.h
@@ -21,6 +21,8 @@ class UMaterialExpressionTangentOutput : public UMaterialExpressionCustomOutput
 	virtual int32 Compile(class FMaterialCompiler* Compiler, int32 OutputIndex) override;
 	virtual void GetCaption(TArray<FString>& OutCaptions) const override;
 	virtual uint32 GetInputType(int32 InputIndex) override { return MCT_Float; }
+	virtual UE::Shader::EValueType GetCustomOutputType(int32 OutputIndex) const override { return UE::Shader::EValueType::Float3; }
+	virtual bool GenerateHLSLExpression(FMaterialHLSLGenerator& Generator, UE::HLSLTree::FScope& Scope, int32 OutputIndex, UE::HLSLTree::FExpression const*& OutExpression) const override;
 #endif
 
 	virtual int32 GetNumOutputs() const override { return 1; }
diff --git a/Engine/Source/Runtime/Engine/Private/Materials/MaterialExpressionHLSL.cpp b/Engine/Source/Runtime/Engine/Private/Materials/MaterialExpressionHLSL.cpp
--- a/Engine/Source/Runtime/Engine/Private/Materials/MaterialExpressionHLSL.cpp
+++ b/Engine/Source/Runtime/Engine/Private/Materials/MaterialExpressionHLSL.cpp
@@ -24,6 +24,7 @@
 #include "Materials/MaterialExpressionArctangent2.h"
 #include "Materials/MaterialExpressionArctangent2Fast.h"
 #include "Materials/MaterialExpressionArctangentFast.h"
+#include "Materials/MaterialExpressionAtmosphericLightColor.h"
 #include "Materials/MaterialExpressionAtmosphericLightVector.h"
 #include "Materials/MaterialExpressionBinaryOp.h"
 #include "Materials/MaterialExpressionBlendMaterialAttributes.h"
@@ -51,6 +52,7 @@
 #include "Materials/MaterialExpressionDDX.h"
 #include "Materials/MaterialExpressionDDY.h"
 #include "Materials/MaterialExpressionDecalColor.h"
+#include "Materials/MaterialExpressionDecalDerivative.h"
 #include "Materials/MaterialExpressionDecalLifetimeOpacity.h"
 #include "Materials/MaterialExpressionDeltaTime.h"
 #include "Materials/MaterialExpressionDepthFade.h"
@@ -85,12 +87,14 @@
 #include "Materials/MaterialExpressionGetLocal.h"
 #include "Materials/MaterialExpressionGetMaterialAttributes.h"
 #include "Materials/MaterialExpressionHairAttributes.h"
+#include "Materials/MaterialExpressionHairColor.h"
 #include "Materials/MaterialExpressionIf.h"
 #include "Materials/MaterialExpressionIfThenElse.h"
 #include "Materials/MaterialExpressionInverseLinearInterpolate.h"
 #include "Materials/MaterialExpressionIsOrthographic.h"
 #include "Materials/MaterialExpressionLength.h"
 #include "Materials/MaterialExpressionLightmapUVs.h"
+#include "Materials/MaterialExpressionLightVector.h"
 #include "Materials/MaterialExpressionLinearInterpolate.h"
 #include "Materials/MaterialExpressionLogarithm.h"
 #include "Materials/MaterialExpressionLogarithm10.h"
@@ -172,6 +176,7 @@
 #include "Materials/MaterialExpressionSubtract.h"
 #include "Materials/MaterialExpressionSwitch.h"
 #include "Materials/MaterialExpressionTangent.h"
+#include "Materials/MaterialExpressionTangentOutput.h"
 #include "Materials/MaterialExpressionTextureCoordinate.h"
 #include "Materials/MaterialExpressionTextureObject.h"
 #include "Materials/MaterialExpressionTextureObjectParameter.h"
@@ -3321,12 +3326,52 @@ bool UMaterialExpressionSkyAtmosphereLightDirection::GenerateHLSLExpression(FMat
 	return true;
 }
 
+bool UMaterialExpressionAtmosphericLightColor::GenerateHLSLExpression(FMaterialHLSLGenerator& Generator, UE::HLSLTree::FScope& Scope, int32 OutputIndex, UE::HLSLTree::FExpression const*& OutExpression) const
+{
+	OutExpression = Generator.GetTree().NewExpression<UE::HLSLTree::FExpressionInlineCustomHLSL>(UE::Shader::EValueType::Float3, TEXT("MaterialExpressionAtmosphericLightColor(Parameters)"));
+	return true;
+}
+
 bool UMaterialExpressionAtmosphericLightVector::GenerateHLSLExpression(FMaterialHLSLGenerator& Generator, UE::HLSLTree::FScope& Scope, int32 OutputIndex, UE::HLSLTree::FExpression const*& OutExpression) const
 {
 	OutExpression = Generator.GetTree().NewExpression<UE::HLSLTree::FExpressionInlineCustomHLSL>(UE::Shader::EValueType::Float3, TEXT("MaterialExpressionAtmosphericLightVector(Parameters)"));
 	return true;
 }
 
+bool UMaterialExpressionDecalDerivative::GenerateHLSLExpression(FMaterialHLSLGenerator& Generator, UE::HLSLTree::FScope& Scope, int32 OutputIndex, UE::HLSLTree::FExpression const*& OutExpression) const
+{
+	if (Generator.GetTargetMaterial()->MaterialDomain != MD_DeferredDecal)
+	{
+		return Generator.Errorf(TEXT("Decal derivatives only available in the decal material domain."));
+	}
+
+	if (OutputIndex == 1)
+	{
+		// DDY Case
+		OutExpression = Generator.GetTree().NewExpression<UE::HLSLTree::FExpressionInlineCustomHLSL>(UE::Shader::EValueType::Float2, TEXT("ComputeDecalDDY(Parameters)"));
+	}
+	else
+	{
+		OutExpression = Generator.GetTree().NewExpression<UE::HLSLTree::FExpressionInlineCustomHLSL>(UE::Shader::EValueType::Float2, TEXT("ComputeDecalDDX(Parameters)"));
+	}
+	return true;
+}
+
+bool UMaterialExpressionLightVector::GenerateHLSLExpression(FMaterialHLSLGenerator& Generator, UE::HLSLTree::FScope& Scope, int32 OutputIndex, UE::HLSLTree::FExpression const*& OutExpression) const
+{
+	using namespace UE::HLSLTree;
+
+	EMaterialDomain Domain = Generator.GetTargetMaterial()->MaterialDomain;
+	if (Domain != MD_LightFunction && Domain != MD_DeferredDecal)
+	{
+		return Generator.Errorf(TEXT("LightVector can only be used in LightFunction or DeferredDecal materials"));
+	}
+
+	const FExpression* Expression = Generator.NewExternalInput(Material::EExternalInput::LightVector);
+	OutExpression = Generator.GetTree().NewExpression<Material::FExpressionLightVector>(Expression);
+	return true;
+}
+
 static UE::Shader::FType GetCustomOutputType(const FMaterialHLSLGenerator& Generator, ECustomMaterialOutputType Type)
 {
 	using namespace UE::Shader;
@@ -3437,6 +3482,19 @@ bool UMaterialExpressionHairAttributes::GenerateHLSLExpression(FMaterialHLSLGene
 	return OutExpression != nullptr;
 }
 
+bool UMaterialExpressionHairColor::GenerateHLSLExpression(FMaterialHLSLGenerator& Generator, UE::HLSLTree::FScope& Scope, int32 OutputIndex, UE::HLSLTree::FExpression const*& OutExpression) const
+{
+	using namespace UE::HLSLTree;
+	using namespace UE::Shader;
+
+	const FExpression* MelaninExpression = Melanin.AcquireHLSLExpressionOrConstant(Generator, Scope, 0.5f);
+	const FExpression* RednessExpression = Redness.AcquireHLSLExpressionOrConstant(Generator, Scope, 0.0f);
+	const FExpression* DyeColorExpression = DyeColor.AcquireHLSLExpressionOrConstant(Generator, Scope, FValue(1.0f, 1.0f, 1.0f));
+
+	OutExpression = Generator.GetTree().NewExpression<Material::FExpressionHairColor>(MelaninExpression, RednessExpression, DyeColorExpression);
+	return true;
+}
+
 bool UMaterialExpressionCloudSampleAttribute::GenerateHLSLExpression(FMaterialHLSLGenerator& Generator, UE::HLSLTree::FScope& Scope, int32 OutputIndex, UE::HLSLTree::FExpression const*& OutExpression) const
 {
 	using namespace UE::HLSLTree;
@@ -3613,6 +3671,12 @@ bool UMaterialExpressionClearCoatNormalCustomOutput::GenerateHLSLExpression(FMat
 	return OutExpression != nullptr;
 }
 
+bool UMaterialExpressionTangentOutput::GenerateHLSLExpression(FMaterialHLSLGenerator& Generator, UE::HLSLTree::FScope& Scope, int32 OutputIndex, UE::HLSLTree::FExpression const*& OutExpression) const
+{
+	OutExpression = Input.AcquireHLSLExpression(Generator, Scope);
+	return OutExpression != nullptr;
+}
+
 bool UMaterialExpressionExecBegin::GenerateHLSLStatements(FMaterialHLSLGenerator& Generator, UE::HLSLTree::FScope& Scope) const
 {
 	Exec.GenerateHLSLStatements(Generator, Scope);
diff --git a/Engine/Source/Runtime/Engine/Private/Materials/MaterialHLSLTree.cpp b/Engine/Source/Runtime/Engine/Private/Materials/MaterialHLSLTree.cpp
--- a/Engine/Source/Runtime/Engine/Private/Materials/MaterialHLSLTree.cpp
+++ b/Engine/Source/Runtime/Engine/Private/Materials/MaterialHLSLTree.cpp
@@ -103,6 +103,7 @@ FExternalInputDescription GetExternalInputDescription(EExternalInput Input)
 	case EExternalInput::RcpResolutionFraction: return FExternalInputDescription(TEXT("RcpResolutionFraction"), Shader::EValueType::Float1);
 
 	case EExternalInput::CameraVector: return FExternalInputDescription(TEXT("CameraVector"), Shader::EValueType::Float3);
+	case EExternalInput::LightVector: return FExternalInputDescription(TEXT("LightVector"), Shader::EValueType::Float3);
 	case EExternalInput::CameraWorldPosition: return FExternalInputDescription(TEXT("CameraWorldPosition"), Shader::EValueType::Double3, EExternalInput::None, EExternalInput::None, EExternalInput::PrevCameraWorldPosition);
 	case EExternalInput::ViewWorldPosition: return FExternalInputDescription(TEXT("ViewWorldPosition"), Shader::EValueType::Double3, EExternalInput::None, EExternalInput::None, EExternalInput::PrevViewWorldPosition);
 	case EExternalInput::PreViewTranslation: return FExternalInputDescription(TEXT("PreViewTranslation"), Shader::EValueType::Double3, EExternalInput::None, EExternalInput::None, EExternalInput::PrevPreViewTranslation);
@@ -318,6 +319,7 @@ void FExpressionExternalInput::EmitValueShader(FEmitContext& Context, FEmitScope
 		case EExternalInput::RcpResolutionFraction: Code = TEXT("View.ResolutionFractionAndInv.y"); break;
 
 		case EExternalInput::CameraVector: Code = TEXT("Parameters.CameraVector"); break;
+		case EExternalInput::LightVector: Code = TEXT("Parameters.LightVector"); break;
 		case EExternalInput::CameraWorldPosition: Code = TEXT("ResolvedView.WorldCameraOrigin"); break;
 		case EExternalInput::ViewWorldPosition: Code = TEXT("ResolvedView.WorldViewOrigin"); break;
 		case EExternalInput::PreViewTranslation: Code = TEXT("ResolvedView.PreViewTranslation"); break;
@@ -1965,6 +1967,52 @@ const TCHAR* FExpressionSamplePhysicsField::GetEmitExpressionFormat() const
 	}
 }
 
+bool FExpressionHairColor::PrepareValue(FEmitContext& Context, FEmitScope& Scope, const FRequestedType& RequestedType, FPrepareValueResult& OutResult) const
+{
+	const FPreparedType& MelaninType = Context.PrepareExpression(MelaninExpression, Scope, Shader::EValueType::Float1);
+	if (MelaninType.IsVoid())
+	{
+		return false;
+	}
+
+	const FPreparedType& RednessType = Context.PrepareExpression(RednessExpression, Scope, Shader::EValueType::Float1);
+	if (RednessType.IsVoid())
+	{
+		return false;
+	}
+
+	const FPreparedType& DyeColorType = Context.PrepareExpression(DyeColorExpression, Scope, Shader::EValueType::Float3);
+	if (DyeColorType.IsVoid())
+	{
+		return false;
+	}
+
+	return OutResult.SetType(Context, RequestedType, EExpressionEvaluation::Shader, Shader::EValueType::Float3);
+}
+
+void FExpressionHairColor::EmitValueShader(FEmitContext& Context, FEmitScope& Scope, const FRequestedType& RequestedType, FEmitValueShaderResult& OutResult) const
+{
+	FEmitShaderExpression* EmitMelanin = MelaninExpression->GetValueShader(Context, Scope, Shader::EValueType::Float1);
+	FEmitShaderExpression* EmitRedness = RednessExpression->GetValueShader(Context, Scope, Shader::EValueType::Float1);
+	FEmitShaderExpression* EmitDyeColor = DyeColorExpression->GetValueShader(Context, Scope, Shader::EValueType::Float3);
+
+	OutResult.Code = Context.EmitInlineExpression(Scope, Shader::EValueType::Float3,
+		TEXT("MaterialExpressionGetHairColorFromMelanin(%, %, %)"),
+		EmitMelanin,
+		EmitRedness,
+		EmitDyeColor);
+}
+
+bool FExpressionLightVector::PrepareValue(FEmitContext& Context, FEmitScope& Scope, const FRequestedType& RequestedType, FPrepareValueResult& OutResult) const
+{
+	if (Context.ShaderFrequency != SF_Pixel && Context.ShaderFrequency != SF_Compute)
+	{
+		return Context.Error(TEXT("LightVector can only be used in Pixel and Compute shaders."));
+	}
+
+	return FExpressionForward::PrepareValue(Context, Scope, RequestedType, OutResult);
+}
+
 bool FExpressionDistanceFieldApproxAO::PrepareValue(FEmitContext& Context, FEmitScope& Scope, const FRequestedType& RequestedType, FPrepareValueResult& OutResult) const
 {
 	if (!Private::PlatformSupportDistanceFields(Context))
diff --git a/Engine/Source/Runtime/Engine/Public/MaterialHLSLTree.h b/Engine/Source/Runtime/Engine/Public/MaterialHLSLTree.h
--- a/Engine/Source/Runtime/Engine/Public/MaterialHLSLTree.h
+++ b/Engine/Source/Runtime/Engine/Public/MaterialHLSLTree.h
@@ -140,6 +140,7 @@ enum class EExternalInput : uint8
 	RcpResolutionFraction,
 
 	CameraVector,
+	LightVector,
 	CameraWorldPosition,
 	ViewWorldPosition,
 	PreViewTranslation,
@@ -781,6 +782,36 @@ public:
 	virtual void EmitValueShader(FEmitContext& Context, FEmitScope& Scope, const FRequestedType& RequestedType, FEmitValueShaderResult& OutResult) const override;
 };
 
+class FExpressionHairColor : public FExpression
+{
+public:
+	const FExpression* MelaninExpression;
+	const FExpression* RednessExpression;
+	const FExpression* DyeColorExpression;
+
+	FExpressionHairColor(
+		const FExpression* InMelaninExpression,
+		const FExpression* InRednessExpression,
+		const FExpression* InDyeColorExpression)
+		: MelaninExpression(InMelaninExpression)
+		, RednessExpression(InRednessExpression)
+		, DyeColorExpression(InDyeColorExpression)
+	{}
+
+	virtual bool PrepareValue(FEmitContext& Context, FEmitScope& Scope, const FRequestedType& RequestedType, FPrepareValueResult& OutResult) const override;
+	virtual void EmitValueShader(FEmitContext& Context, FEmitScope& Scope, const FRequestedType& RequestedType, FEmitValueShaderResult& OutResult) const override;
+};
+
+class FExpressionLightVector : public FExpressionForward
+{
+public:
+	FExpressionLightVector(const FExpression* InExpression)
+		: FExpressionForward(InExpression)
+	{}
+
+	virtual bool PrepareValue(FEmitContext& Context, FEmitScope& Scope, const FRequestedType& RequestedType, FPrepareValueResult& OutResult) const override;
+};
+
 struct FVertexInterpolator
 {
 	FVertexInterpolator() = default;
