Enables support for screen-space (Tier 2) Variable Rate Shading on Xbox Series X.

diff --git a/Engine/Source/Runtime/Engine/Private/Materials/MaterialShared.cpp b/Engine/Source/Runtime/Engine/Private/Materials/MaterialShared.cpp
--- a/Engine/Source/Runtime/Engine/Private/Materials/MaterialShared.cpp
+++ b/Engine/Source/Runtime/Engine/Private/Materials/MaterialShared.cpp
@@ -2621,7 +2621,8 @@ void FMaterial::SetupMaterialEnvironment(
 		SET_SHADER_DEFINE(OutEnvironment, USE_STENCIL_LOD_DITHER_DEFAULT, CVar->GetValueOnAnyThread() != 0 ? 1 : 0);
 	}
 
-	if (GetShadingRate() != MSR_1x1)
+	if (FDataDrivenShaderPlatformInfo::GetSupportsVariableRateShading(Platform) && 
+		(GRHIAttachmentVariableRateShadingEnabled || GetShadingRate() != MSR_1x1))
 	{
 		OutEnvironment.SetCompileArgument(TEXT("USING_VARIABLE_RATE_SHADING"), true);
 	}
diff --git a/Engine/Source/Runtime/RHI/Private/RHI.cpp b/Engine/Source/Runtime/RHI/Private/RHI.cpp
--- a/Engine/Source/Runtime/RHI/Private/RHI.cpp
+++ b/Engine/Source/Runtime/RHI/Private/RHI.cpp
@@ -1175,8 +1175,8 @@ static FAutoConsoleVariableRef CVarEnableVariableRateShading(
 static FAutoConsoleVariableRef CVarEnableAttachmentVariableRateShading(
 	TEXT("r.VRS.EnableImage"),
 	GRHIAttachmentVariableRateShadingEnabled,
-	TEXT("Toggle to enable image-based Variable Rate Shading."),
-	ECVF_RenderThreadSafe);
+	TEXT("Enable image-based Variable Rate Shading."),
+	ECVF_ReadOnly);
 
 
 FString GRHIBindlessResourceConfiguration = TEXT("Disabled");
diff --git a/Engine/Source/Runtime/RenderCore/Private/Shader.cpp b/Engine/Source/Runtime/RenderCore/Private/Shader.cpp
--- a/Engine/Source/Runtime/RenderCore/Private/Shader.cpp
+++ b/Engine/Source/Runtime/RenderCore/Private/Shader.cpp
@@ -2092,6 +2092,11 @@ void ShaderMapAppendKeyString(EShaderPlatform Platform, FString& KeyString)
 		KeyString += TEXT("_sdct");
 	}
 
+	if (FDataDrivenShaderPlatformInfo::GetSupportsVariableRateShading(Platform) && GRHIAttachmentVariableRateShadingEnabled)
+	{
+		KeyString += TEXT("_VRS");
+	}
+
 	{
 		static const auto CVar = IConsoleManager::Get().FindTConsoleVariableDataInt(TEXT("r.GBufferDiffuseSampleOcclusion"));
 		if (CVar && CVar->GetValueOnAnyThread() != 0)
