Chaos - prevent particles with no geometry from being added to the AABBTree with empty bounds (which does not transform).

diff --git a/Engine/Source/Runtime/Engine/Private/PhysicsEngine/Experimental/PhysScene_Chaos.cpp b/Engine/Source/Runtime/Engine/Private/PhysicsEngine/Experimental/PhysScene_Chaos.cpp
--- a/Engine/Source/Runtime/Engine/Private/PhysicsEngine/Experimental/PhysScene_Chaos.cpp
+++ b/Engine/Source/Runtime/Engine/Private/PhysicsEngine/Experimental/PhysScene_Chaos.cpp
@@ -2163,13 +2163,16 @@ static void UpdateAccelerationStructureFromGeometryCollectionProxy(FGeometryColl
 
 	auto GetParticleWorldBounds = [](Chaos::FPBDRigidParticle& Particle) -> Chaos::FAABB3
 	{
-		const Chaos::FRigidTransform3 ParticleWorldTransform(Particle.X(), Particle.R());
-		Chaos::FAABB3 WorldBounds;
-		if (const Chaos::FImplicitObjectRef Geometry = Particle.GetGeometry(); Geometry && Geometry->HasBoundingBox())
+		const Chaos::FImplicitObjectRef Geometry = Particle.GetGeometry();
+		check(Geometry != nullptr)
+
+		if ((Geometry != nullptr) && Geometry->HasBoundingBox())
 		{
-			WorldBounds = Geometry->CalculateTransformedBounds(ParticleWorldTransform);
+			const Chaos::FRigidTransform3 ParticleWorldTransform(Particle.X(), Particle.R());
+			return Geometry->CalculateTransformedBounds(ParticleWorldTransform);
 		}
-		return WorldBounds;
+
+		return Chaos::FAABB3(Particle.X(), Particle.X());
 	};
 
 	const TManagedArray<TUniquePtr<Chaos::FPBDRigidParticle>>& GTParticles = Proxy.GetExternalParticles();
