Fix shipping client missing include
#rnx

diff --git a/Engine/Source/Runtime/Engine/Classes/Components/SceneComponent.h b/Engine/Source/Runtime/Engine/Classes/Components/SceneComponent.h
--- a/Engine/Source/Runtime/Engine/Classes/Components/SceneComponent.h
+++ b/Engine/Source/Runtime/Engine/Classes/Components/SceneComponent.h
@@ -2,6 +2,7 @@
 
 #pragma once
 
+#include "UObject/ObjectMacros.h"
 #include "Math/BoxSphereBounds.h"
 #include "UObject/UObjectGlobals.h"
 #if UE_ENABLE_INCLUDE_ORDER_DEPRECATED_IN_5_2
