World Partition Unit Test: added different actor parent attachements scenarios involving parent or child editor-only actors.

diff --git a/Engine/Source/Runtime/Engine/Private/WorldPartition/WorldPartitionRuntimeLevelStreamingCell.cpp b/Engine/Source/Runtime/Engine/Private/WorldPartition/WorldPartitionRuntimeLevelStreamingCell.cpp
--- a/Engine/Source/Runtime/Engine/Private/WorldPartition/WorldPartitionRuntimeLevelStreamingCell.cpp
+++ b/Engine/Source/Runtime/Engine/Private/WorldPartition/WorldPartitionRuntimeLevelStreamingCell.cpp
@@ -383,10 +383,14 @@ void UWorldPartitionRuntimeLevelStreamingCell::DumpStateLog(FHierarchicalLogArch
 
 	for (const FWorldPartitionRuntimeCellObjectMapping& Mapping : SortedPackages)
 	{
-		Ar.Printf(TEXT("         Actor Path: %s"), *Mapping.Path.ToString());
-		Ar.Printf(TEXT("      Actor Package: %s"), *Mapping.Package.ToString());
-		Ar.Printf(TEXT(" Actor Container ID: %s"), *Mapping.ContainerID.ToString());
-		Ar.Printf(TEXT("Actor Instance Guid: %s"), *Mapping.ActorInstanceGuid.ToString());
+		FHierarchicalLogArchive::FIndentScope ActorIndentScope = Ar.PrintfIndent(TEXT("%s"), *Mapping.Path.ToString());
+		Ar.Printf(TEXT("        Package: %s"), *Mapping.Package.ToString());
+		Ar.Printf(TEXT("    Editor Only: %d"), Mapping.bIsEditorOnly ? 1 : 0);
+		Ar.Printf(TEXT("  Instance Guid: %s"), *Mapping.ActorInstanceGuid.ToString());
+
+		FHierarchicalLogArchive::FIndentScope ContainerIndentScope = Ar.PrintfIndent(TEXT("Container:"));
+		Ar.Printf(TEXT("       ID: %s"), *Mapping.ContainerID.ToString());
+		Ar.Printf(TEXT("Transform: %s"), *Mapping.ContainerTransform.ToString());		
 	}
 }
 #endif
