Add IsGPUSkinPassThrough to shader DDC key

diff --git a/Engine/Plugins/FX/Niagara/Source/Niagara/Private/NiagaraCommon.cpp b/Engine/Plugins/FX/Niagara/Source/Niagara/Private/NiagaraCommon.cpp
--- a/Engine/Plugins/FX/Niagara/Source/Niagara/Private/NiagaraCommon.cpp
+++ b/Engine/Plugins/FX/Niagara/Source/Niagara/Private/NiagaraCommon.cpp
@@ -3,7 +3,7 @@
 #include "NiagaraCommon.h"
 
 #include "DataDrivenShaderPlatformInfo.h"
-#include "LocalVertexFactory.h"
+#include "RenderUtils.h"
 #include "Misc/StringBuilder.h"
 #include "NiagaraComponent.h"
 #include "NiagaraConstants.h"
@@ -963,7 +963,7 @@ bool FNiagaraUtilities::AllowGPUCulling(EShaderPlatform ShaderPlatform)
 
 bool FNiagaraUtilities::AreBufferSRVsAlwaysCreated(EShaderPlatform ShaderPlatform)
 {
-	return RHISupportsManualVertexFetch(ShaderPlatform) || FLocalVertexFactory::IsGPUSkinPassThroughSupported(ShaderPlatform);
+	return RHISupportsManualVertexFetch(ShaderPlatform) || IsGPUSkinPassThroughSupported(ShaderPlatform);
 }
 
 ENiagaraCompileUsageStaticSwitch FNiagaraUtilities::ConvertScriptUsageToStaticSwitchUsage(ENiagaraScriptUsage ScriptUsage)
diff --git a/Engine/Source/Runtime/Engine/Private/Animation/MeshDeformerProvider.cpp b/Engine/Source/Runtime/Engine/Private/Animation/MeshDeformerProvider.cpp
deleted file mode 100644
diff --git a/Engine/Source/Runtime/Engine/Private/GPUSkinCache.cpp b/Engine/Source/Runtime/Engine/Private/GPUSkinCache.cpp
--- a/Engine/Source/Runtime/Engine/Private/GPUSkinCache.cpp
+++ b/Engine/Source/Runtime/Engine/Private/GPUSkinCache.cpp
@@ -198,17 +198,6 @@ static int32 GGPUSkinCacheFlushCounter = 0;
 
 const float MBSize = 1048576.f; // 1024 x 1024 bytes
 
-static inline bool IsGPUSkinCacheAllowed(EShaderPlatform Platform)
-{
-	static FShaderPlatformCachedIniValue<bool> PerPlatformCVar(TEXT("r.SkinCache.Allow"));
-	return PerPlatformCVar.Get(Platform);
-}
-
-ENGINE_API bool IsGPUSkinCacheAvailable(EShaderPlatform Platform)
-{
-	return AreSkinCacheShadersEnabled(Platform) != 0 && IsGPUSkinCacheAllowed(Platform);
-}
-
 static inline bool IsGPUSkinCacheEnable(EShaderPlatform Platform)
 {
 	static FShaderPlatformCachedIniValue<int32> PerPlatformCVar(TEXT("r.SkinCache.Mode"));
diff --git a/Engine/Source/Runtime/Engine/Private/InstancedStaticMesh.cpp b/Engine/Source/Runtime/Engine/Private/InstancedStaticMesh.cpp
--- a/Engine/Source/Runtime/Engine/Private/InstancedStaticMesh.cpp
+++ b/Engine/Source/Runtime/Engine/Private/InstancedStaticMesh.cpp
@@ -1083,7 +1083,7 @@ void FInstancedStaticMeshVertexFactory::GetVertexElements(
 
 		// PreSkinPosition attribute is only used for GPUSkinPassthrough variation of local vertex factory.
 		// It is not used by ISM so fill with dummy buffer.
-		if (FLocalVertexFactory::IsGPUSkinPassThroughSupported(GMaxRHIShaderPlatform))
+		if (IsGPUSkinPassThroughSupported(GMaxRHIShaderPlatform))
 		{
 			FVertexStreamComponent NullComponent(&GNullVertexBuffer, 0, 0, VET_Float4);
 			Elements.Add(AccessStreamComponent(NullComponent, 14, Streams));
diff --git a/Engine/Source/Runtime/Engine/Private/LocalVertexFactory.cpp b/Engine/Source/Runtime/Engine/Private/LocalVertexFactory.cpp
--- a/Engine/Source/Runtime/Engine/Private/LocalVertexFactory.cpp
+++ b/Engine/Source/Runtime/Engine/Private/LocalVertexFactory.cpp
@@ -18,7 +18,6 @@
 #include "ProfilingDebugging/LoadTimeTracker.h"
 #include "GPUSkinCache.h"
 #include "GPUSkinVertexFactory.h"
-#include "Animation/MeshDeformerProvider.h"
 #include "RenderUtils.h"
 #include "SceneInterface.h"
 
@@ -172,14 +171,6 @@ void FLocalVertexFactoryShaderParameters::Bind(const FShaderParameterMap& Parame
 	IsGPUSkinPassThrough.Bind(ParameterMap, TEXT("bIsGPUSkinPassThrough"));
 }
 
-bool FLocalVertexFactory::IsGPUSkinPassThroughSupported(EShaderPlatform Platform)
-{
-	// Enable the GPUSkin passthrough path if we might use the GPUSkinCache or MeshDeformers.
-	static IMeshDeformerProvider* MeshDeformerProvider = IMeshDeformerProvider::Get();
-	bool bMeshDeformersAvailable = MeshDeformerProvider && MeshDeformerProvider->IsSupported(Platform);
-	return bMeshDeformersAvailable || IsGPUSkinCacheAvailable(Platform);
-}
-
 static void GetMeshDeformerVertexStreams(FMeshDeformerGeometry const& InDeformerGeometry, FGPUSkinPassthroughVertexFactory const* InVertexFactory, FVertexInputStreamArray& InOutVertexStreams)
 {
 	const int32 PositionStreamIndex = InVertexFactory->GetAttributeStreamIndex(FGPUSkinPassthroughVertexFactory::VertexPosition);
diff --git a/Engine/Source/Runtime/Engine/Private/PositionVertexBuffer.cpp b/Engine/Source/Runtime/Engine/Private/PositionVertexBuffer.cpp
--- a/Engine/Source/Runtime/Engine/Private/PositionVertexBuffer.cpp
+++ b/Engine/Source/Runtime/Engine/Private/PositionVertexBuffer.cpp
@@ -4,7 +4,8 @@
 
 #include "Components.h"
 #include "DataDrivenShaderPlatformInfo.h"
-#include "LocalVertexFactory.h"
+#include "RenderingThread.h"
+#include "RenderUtils.h"
 #include "RHIResourceUpdates.h"
 #include "StaticMeshVertexData.h"
 
@@ -245,7 +246,7 @@ void FPositionVertexBuffer::InitRHI(FRHICommandListBase& RHICmdList)
 	if (VertexBufferRHI)
 	{
 		// we have decide to create the SRV based on GMaxRHIShaderPlatform because this is created once and shared between feature levels for editor preview.
-		bool bSRV = RHISupportsManualVertexFetch(GMaxRHIShaderPlatform) || FLocalVertexFactory::IsGPUSkinPassThroughSupported(GMaxRHIShaderPlatform);
+		bool bSRV = RHISupportsManualVertexFetch(GMaxRHIShaderPlatform) || IsGPUSkinPassThroughSupported(GMaxRHIShaderPlatform);
 
 		// When bAllowCPUAccess is true, the meshes is likely going to be used for Niagara to spawn particles on mesh surface.
 		// And it can be the case for CPU *and* GPU access: no differenciation today. That is why we create a SRV in this case.
diff --git a/Engine/Source/Runtime/Engine/Private/RawIndexBuffer.cpp b/Engine/Source/Runtime/Engine/Private/RawIndexBuffer.cpp
--- a/Engine/Source/Runtime/Engine/Private/RawIndexBuffer.cpp
+++ b/Engine/Source/Runtime/Engine/Private/RawIndexBuffer.cpp
@@ -8,7 +8,7 @@
 
 #include "DataDrivenShaderPlatformInfo.h"
 #include "Interfaces/ITargetPlatform.h"
-#include "LocalVertexFactory.h"
+#include "RenderUtils.h"
 #include "Modules/ModuleManager.h"
 #include "RenderingThread.h"
 #include "RHIResourceUpdates.h"
@@ -328,7 +328,7 @@ FBufferRHIRef FRawStaticIndexBuffer::CreateRHIBuffer(FRHICommandListBase& RHICmd
 	if (GetNumIndices() > 0)
 	{
 		// Systems that generate data for GPUSkinPassThrough use index buffer as SRV.
-		bool bSRV = RHISupportsManualVertexFetch(GMaxRHIShaderPlatform) || FLocalVertexFactory::IsGPUSkinPassThroughSupported(GMaxRHIShaderPlatform);
+		bool bSRV = RHISupportsManualVertexFetch(GMaxRHIShaderPlatform) || IsGPUSkinPassThroughSupported(GMaxRHIShaderPlatform);
 
 		// When bAllowCPUAccess is true, the meshes is likely going to be used for Niagara to spawn particles on mesh surface.
 		// And it can be the case for CPU *and* GPU access: no differenciation today. That is why we create a SRV in this case.
@@ -446,7 +446,7 @@ void FRawStaticIndexBuffer::Discard()
 bool FRawStaticIndexBuffer16or32Interface::IsSRVNeeded(bool bAllowCPUAccess) const
 {
 	// Systems that generate data for GPUSkinPassThrough use index buffer as SRV.
-	bool bSRV = RHISupportsManualVertexFetch(GMaxRHIShaderPlatform) || FLocalVertexFactory::IsGPUSkinPassThroughSupported(GMaxRHIShaderPlatform);
+	bool bSRV = RHISupportsManualVertexFetch(GMaxRHIShaderPlatform) || IsGPUSkinPassThroughSupported(GMaxRHIShaderPlatform);
 	// When bAllowCPUAccess is true, the meshes is likely going to be used for Niagara to spawn particles on mesh surface.
 	// And it can be the case for CPU *and* GPU access: no differenciation today. That is why we create a SRV in this case.
 	// This also avoid setting lots of states on all the members of all the different buffers used by meshes. Follow up: https://jira.it.epicgames.net/browse/UE-69376.
diff --git a/Engine/Source/Runtime/Engine/Private/Rendering/StaticMeshVertexBuffer.cpp b/Engine/Source/Runtime/Engine/Private/Rendering/StaticMeshVertexBuffer.cpp
--- a/Engine/Source/Runtime/Engine/Private/Rendering/StaticMeshVertexBuffer.cpp
+++ b/Engine/Source/Runtime/Engine/Private/Rendering/StaticMeshVertexBuffer.cpp
@@ -334,7 +334,7 @@ void FStaticMeshVertexBuffer::InitRHI(FRHICommandListBase& RHICmdList)
 	const bool bHadTangentsData = TangentsData != nullptr;
 	const bool bCreateTangentsSRV = bHadTangentsData && TangentsData->GetAllowCPUAccess();
 	TangentsVertexBuffer.VertexBufferRHI = CreateTangentsRHIBuffer(RHICmdList);
-	if (TangentsVertexBuffer.VertexBufferRHI && (bCreateTangentsSRV || RHISupportsManualVertexFetch(GMaxRHIShaderPlatform) || FLocalVertexFactory::IsGPUSkinPassThroughSupported(GMaxRHIShaderPlatform)))
+	if (TangentsVertexBuffer.VertexBufferRHI && (bCreateTangentsSRV || RHISupportsManualVertexFetch(GMaxRHIShaderPlatform) || IsGPUSkinPassThroughSupported(GMaxRHIShaderPlatform)))
 	{
 		uint32       Stride = GetUseHighPrecisionTangentBasis() ? 8 : 4;
 		EPixelFormat Format = GetUseHighPrecisionTangentBasis() ? PF_R16G16B16A16_SNORM : PF_R8G8B8A8_SNORM;
diff --git a/Engine/Source/Runtime/Engine/Public/Animation/MeshDeformerProvider.h b/Engine/Source/Runtime/Engine/Public/Animation/MeshDeformerProvider.h
deleted file mode 100644
diff --git a/Engine/Source/Runtime/Engine/Public/GPUSkinCache.h b/Engine/Source/Runtime/Engine/Public/GPUSkinCache.h
--- a/Engine/Source/Runtime/Engine/Public/GPUSkinCache.h
+++ b/Engine/Source/Runtime/Engine/Public/GPUSkinCache.h
@@ -59,9 +59,6 @@ struct FVertexBufferAndSRV;
 struct FRayTracingGeometrySegment;
 struct FSkinBatchVertexFactoryUserData;
 
-// Can the skin cache be used (ie shaders added, etc)
-extern ENGINE_API bool IsGPUSkinCacheAvailable(EShaderPlatform Platform);
-
 extern bool ShouldWeCompileGPUSkinVFShaders(EShaderPlatform Platform, ERHIFeatureLevel::Type FeatureLevel);
 
 extern ENGINE_API bool GPUSkinCacheNeedsDuplicatedVertices();
diff --git a/Engine/Source/Runtime/Engine/Public/LocalVertexFactory.h b/Engine/Source/Runtime/Engine/Public/LocalVertexFactory.h
--- a/Engine/Source/Runtime/Engine/Public/LocalVertexFactory.h
+++ b/Engine/Source/Runtime/Engine/Public/LocalVertexFactory.h
@@ -76,12 +76,6 @@ public:
 	static ENGINE_API void GetPSOPrecacheVertexFetchElements(EVertexInputStreamType VertexInputStreamType, FVertexDeclarationElementList& Elements);
 	static ENGINE_API void GetVertexElements(ERHIFeatureLevel::Type FeatureLevel, EVertexInputStreamType InputStreamType, bool bSupportsManualVertexFetch, FDataType& Data, FVertexDeclarationElementList& Elements);
 
-	/** 
-	 * Does the platform support GPUSkinPassthrough permutations.
-	 * This knowledge can be used to indicate if we need to create SRV for index/vertex buffers.
-	 */
-	static ENGINE_API bool IsGPUSkinPassThroughSupported(EShaderPlatform Platform);
-
 	/**
 	 * An implementation of the interface used by TSynchronizedResource to update the resource with new data from the game thread.
 	 */
diff --git a/Engine/Source/Runtime/RenderCore/Private/Animation/MeshDeformerProvider.cpp b/Engine/Source/Runtime/RenderCore/Private/Animation/MeshDeformerProvider.cpp
new file mode 100644
--- /dev/null
+++ b/Engine/Source/Runtime/RenderCore/Private/Animation/MeshDeformerProvider.cpp
@@ -0,0 +1,20 @@
+// Copyright Epic Games, Inc. All Rights Reserved.
+
+#include "Animation/MeshDeformerProvider.h"
+#include "Features/IModularFeatures.h"
+
+const FName IMeshDeformerProvider::ModularFeatureName(TEXT("MeshDeformer"));
+
+IMeshDeformerProvider* IMeshDeformerProvider::Get()
+{
+	if (IsAvailable())
+	{
+		return &IModularFeatures::Get().GetModularFeature<IMeshDeformerProvider>(ModularFeatureName);
+	}
+	return nullptr;
+}
+
+bool IMeshDeformerProvider::IsAvailable()
+{
+	return IModularFeatures::Get().IsModularFeatureAvailable(ModularFeatureName);
+}
diff --git a/Engine/Source/Runtime/RenderCore/Private/RenderUtils.cpp b/Engine/Source/Runtime/RenderCore/Private/RenderUtils.cpp
--- a/Engine/Source/Runtime/RenderCore/Private/RenderUtils.cpp
+++ b/Engine/Source/Runtime/RenderCore/Private/RenderUtils.cpp
@@ -16,6 +16,7 @@
 #include "DataDrivenShaderPlatformInfo.h"
 #include "RenderCore.h"
 #include "SubstrateDefinitions.h"
+#include "Animation/MeshDeformerProvider.h"
 
 #if WITH_EDITOR
 #include "Interfaces/ITargetPlatform.h"
@@ -1393,6 +1394,25 @@ RENDERCORE_API bool AreSkinCacheShadersEnabled(EShaderPlatform Platform)
 	return (PerPlatformCVar.Get(Platform) != 0);
 }
 
+RENDERCORE_API bool IsGPUSkinCacheAllowed(EShaderPlatform Platform)
+{
+	static FShaderPlatformCachedIniValue<bool> PerPlatformCVar(TEXT("r.SkinCache.Allow"));
+	return PerPlatformCVar.Get(Platform);
+}
+
+RENDERCORE_API bool IsGPUSkinCacheAvailable(EShaderPlatform Platform)
+{
+	return AreSkinCacheShadersEnabled(Platform) != 0 && IsGPUSkinCacheAllowed(Platform);
+}
+
+RENDERCORE_API bool IsGPUSkinPassThroughSupported(EShaderPlatform Platform)
+{
+	// Enable the GPUSkin passthrough path if we might use the GPUSkinCache or MeshDeformers.
+	static IMeshDeformerProvider* MeshDeformerProvider = IMeshDeformerProvider::Get();
+	bool bMeshDeformersAvailable = MeshDeformerProvider && MeshDeformerProvider->IsSupported(Platform);
+	return bMeshDeformersAvailable || IsGPUSkinCacheAvailable(Platform);
+}
+
 RENDERCORE_API bool DoesRuntimeSupportOnePassPointLightShadows(EShaderPlatform Platform)
 {
 	static const auto CVar = IConsoleManager::Get().FindTConsoleVariableDataInt(TEXT("r.Shadow.DetectVertexShaderLayerAtRuntime"));
diff --git a/Engine/Source/Runtime/RenderCore/Private/Shader.cpp b/Engine/Source/Runtime/RenderCore/Private/Shader.cpp
--- a/Engine/Source/Runtime/RenderCore/Private/Shader.cpp
+++ b/Engine/Source/Runtime/RenderCore/Private/Shader.cpp
@@ -2265,6 +2265,15 @@ void ShaderMapAppendKeyString(EShaderPlatform Platform, FString& KeyString)
 	{
 		KeyString += TEXT("_SLWDP");
 	}
+
+	if (IsGPUSkinPassThroughSupported(Platform))
+	{
+		KeyString += TEXT("_SKPassThrough1");
+	}
+	else
+	{
+		KeyString += TEXT("_SKPassThrough0");
+	}
 }
 
 EShaderPermutationFlags GetShaderPermutationFlags(const FPlatformTypeLayoutParameters& LayoutParams)
diff --git a/Engine/Source/Runtime/RenderCore/Public/Animation/MeshDeformerProvider.h b/Engine/Source/Runtime/RenderCore/Public/Animation/MeshDeformerProvider.h
new file mode 100644
--- /dev/null
+++ b/Engine/Source/Runtime/RenderCore/Public/Animation/MeshDeformerProvider.h
@@ -0,0 +1,44 @@
+// Copyright Epic Games, Inc. All Rights Reserved.
+
+#pragma once
+
+#include "CoreTypes.h"
+#include "CoreMinimal.h"
+#include "Features/IModularFeature.h"
+#if UE_ENABLE_INCLUDE_ORDER_DEPRECATED_IN_5_2
+#include "RHIDefinitions.h"
+#endif
+#include "UObject/SoftObjectPtr.h"
+
+enum EShaderPlatform : uint16;
+
+/**
+ * Modular feature interface for mesh deformer providers. 
+ * Modules that inherit from this need to be loaded before shader compilation starts (PostConfigInit)
+ * so that the correct vertex factories can be created.
+ */
+class IMeshDeformerProvider : public IModularFeature
+{
+public:
+	virtual ~IMeshDeformerProvider() {}
+
+	static RENDERCORE_API const FName ModularFeatureName; // "MeshDeformer"
+	static RENDERCORE_API bool IsAvailable();
+	static RENDERCORE_API IMeshDeformerProvider* Get();
+
+	/** Returns true if the platform is supported. */
+	virtual bool IsSupported(EShaderPlatform Platform) const = 0;
+
+	/** Structure for passing to GetDefaultMeshDeformer(). */
+	struct FDefaultMeshDeformerSetup
+	{
+		bool bIsUsingSkinCache = false;
+		bool bIsRequestingRecomputeTangent = false;
+	};
+
+	/** 
+	 * Returns a default mesh deformer. 
+	 * This can allow a mesh deformer plugin to automatically replace the UE fixed function animation path.
+	 */
+	virtual TObjectPtr<class UMeshDeformer> GetDefaultMeshDeformer(FDefaultMeshDeformerSetup const& Setup) = 0;
+};
diff --git a/Engine/Source/Runtime/RenderCore/Public/RenderUtils.h b/Engine/Source/Runtime/RenderCore/Public/RenderUtils.h
--- a/Engine/Source/Runtime/RenderCore/Public/RenderUtils.h
+++ b/Engine/Source/Runtime/RenderCore/Public/RenderUtils.h
@@ -442,6 +442,22 @@ RENDERCORE_API bool UseRemoveUnsedInterpolators(EShaderPlatform ShaderPlatform);
  */
 RENDERCORE_API bool AreSkinCacheShadersEnabled(EShaderPlatform Platform);
 
+/**
+ * Checks if skin cache shaders are allowed for the platform (via r.SkinCache.Allow)
+ */
+RENDERCORE_API bool IsGPUSkinCacheAllowed(EShaderPlatform Platform);
+
+/**
+ * Can the skin cache be used (ie shaders added, etc)
+ */
+RENDERCORE_API bool IsGPUSkinCacheAvailable(EShaderPlatform Platform);
+
+/**
+ * Does the platform support GPUSkinPassthrough permutations.
+ * This knowledge can be used to indicate if we need to create SRV for index/vertex buffers.
+ */
+RENDERCORE_API bool IsGPUSkinPassThroughSupported(EShaderPlatform Platform);
+
 /*
  * Detect (at runtime) if the runtime supports rendering one-pass point light shadows (i.e., cube maps)
  */
