Added SkyAtmosphereLightIlluminanceOnGround node

diff --git a/Engine/Shaders/Private/MaterialTemplate.ush b/Engine/Shaders/Private/MaterialTemplate.ush
--- a/Engine/Shaders/Private/MaterialTemplate.ush
+++ b/Engine/Shaders/Private/MaterialTemplate.ush
@@ -1720,9 +1720,18 @@ float3 MaterialExpressionGetHairColorFromMelanin(float Melanin, float Redness, f
 DEFINE_ATMOSPHERELIGHTVECTOR(FMaterialPixelParameters)
 DEFINE_ATMOSPHERELIGHTVECTOR(FMaterialVertexParameters)
 
-float3 MaterialExpressionAtmosphericLightColor(FMaterialPixelParameters Parameters)
+float3 MaterialExpressionSkyAtmosphereLightIlluminanceOnGround(FMaterialPixelParameters Parameters, uint LightIndex)
 {
-	return ResolvedView.AtmosphereLightIlluminanceOnGroundPostTransmittance[0].rgb;
+#if PROJECT_SUPPORT_SKY_ATMOSPHERE // can be used anywhere, not only in sky material with MATERIAL_SKY_ATMOSPHERE
+	return ResolvedView.AtmosphereLightIlluminanceOnGroundPostTransmittance[LightIndex].rgb;
+#else
+	return float3(0.0f, 0.0f, 0.0f);
+#endif
+}
+
+float3 MaterialExpressionAtmosphericLightColor(FMaterialPixelParameters Parameters) // Deprecated but back compatible.
+{
+	return MaterialExpressionSkyAtmosphereLightIlluminanceOnGround(Parameters, 0);
 }
 
 float3 MaterialExpressionSkyAtmosphereLightIlluminance(FMaterialPixelParameters Parameters, FLWCVector3 WorldPosition, uint LightIndex)
diff --git a/Engine/Source/Runtime/Engine/Classes/Materials/MaterialExpressionSkyAtmosphereLightIlluminance.h b/Engine/Source/Runtime/Engine/Classes/Materials/MaterialExpressionSkyAtmosphereLightIlluminance.h
--- a/Engine/Source/Runtime/Engine/Classes/Materials/MaterialExpressionSkyAtmosphereLightIlluminance.h
+++ b/Engine/Source/Runtime/Engine/Classes/Materials/MaterialExpressionSkyAtmosphereLightIlluminance.h
@@ -25,6 +25,27 @@ class UMaterialExpressionSkyAtmosphereLightIlluminance : public UMaterialExpress
 #if WITH_EDITOR
 	virtual int32 Compile(class FMaterialCompiler* Compiler, int32 OutputIndex) override;
 	virtual void GetCaption(TArray<FString>& OutCaptions) const override;
+	virtual void GetExpressionToolTip(TArray<FString>& OutToolTip) override;
+#endif
+	//~ End UMaterialExpression Interface
+};
+
+
+
+UCLASS()
+class UMaterialExpressionSkyAtmosphereLightIlluminanceOnGround : public UMaterialExpression
+{
+	GENERATED_UCLASS_BODY()
+
+	/** Index of the atmosphere light to sample. */
+	UPROPERTY(EditAnywhere, BlueprintReadWrite, Category = MaterialExpressionTextureCoordinate, meta = (UIMin = "0", UIMax = "1", ClampMin = "0", ClampMax = "1", ShowAsInputPin = "Primary"))
+	int32 LightIndex;
+
+	//~ Begin UMaterialExpression Interface
+#if WITH_EDITOR
+	virtual int32 Compile(class FMaterialCompiler* Compiler, int32 OutputIndex) override;
+	virtual void GetCaption(TArray<FString>& OutCaptions) const override;
+	virtual void GetExpressionToolTip(TArray<FString>& OutToolTip) override;
 #endif
 	//~ End UMaterialExpression Interface
 };
diff --git a/Engine/Source/Runtime/Engine/Private/Materials/HLSLMaterialTranslator.cpp b/Engine/Source/Runtime/Engine/Private/Materials/HLSLMaterialTranslator.cpp
--- a/Engine/Source/Runtime/Engine/Private/Materials/HLSLMaterialTranslator.cpp
+++ b/Engine/Source/Runtime/Engine/Private/Materials/HLSLMaterialTranslator.cpp
@@ -10869,6 +10869,12 @@ int32 FHLSLMaterialTranslator::SkyAtmosphereLightIlluminance(int32 WorldPosition
 	return AddCodeChunk(MCT_Float3, TEXT("MaterialExpressionSkyAtmosphereLightIlluminance(Parameters, %s, %d)"), *WorldPosCode, LightIndex);
 }
 
+int32 FHLSLMaterialTranslator::SkyAtmosphereLightIlluminanceOnGround(int32 LightIndex)
+{
+	bUsesSkyAtmosphere = true;
+	return AddCodeChunk(MCT_Float3, TEXT("MaterialExpressionSkyAtmosphereLightIlluminanceOnGround(Parameters, %d)"), LightIndex);
+}
+
 int32 FHLSLMaterialTranslator::SkyAtmosphereLightDirection(int32 LightIndex)
 {
 	bUsesSkyAtmosphere = true;
diff --git a/Engine/Source/Runtime/Engine/Private/Materials/HLSLMaterialTranslator.h b/Engine/Source/Runtime/Engine/Private/Materials/HLSLMaterialTranslator.h
--- a/Engine/Source/Runtime/Engine/Private/Materials/HLSLMaterialTranslator.h
+++ b/Engine/Source/Runtime/Engine/Private/Materials/HLSLMaterialTranslator.h
@@ -1145,6 +1145,7 @@ protected:
 	virtual int32 AtmosphericLightColor() override;
 
 	virtual int32 SkyAtmosphereLightIlluminance(int32 WorldPosition, int32 LightIndex) override;
+	virtual int32 SkyAtmosphereLightIlluminanceOnGround(int32 LightIndex) override;
 	virtual int32 SkyAtmosphereLightDirection(int32 LightIndex) override;
 	virtual int32 SkyAtmosphereLightDiskLuminance(int32 LightIndex, int32 OverrideAtmosphereLightDiscCosHalfApexAngle) override;
 	virtual int32 SkyAtmosphereViewLuminance() override;
diff --git a/Engine/Source/Runtime/Engine/Private/Materials/MaterialExpressions.cpp b/Engine/Source/Runtime/Engine/Private/Materials/MaterialExpressions.cpp
--- a/Engine/Source/Runtime/Engine/Private/Materials/MaterialExpressions.cpp
+++ b/Engine/Source/Runtime/Engine/Private/Materials/MaterialExpressions.cpp
@@ -21749,6 +21749,52 @@ void UMaterialExpressionSkyAtmosphereLightIlluminance::GetCaption(TArray<FString
 {
 	OutCaptions.Add(FString::Printf(TEXT("SkyAtmosphereLightIlluminance[%i]"), LightIndex));
 }
+
+void UMaterialExpressionSkyAtmosphereLightIlluminance::GetExpressionToolTip(TArray<FString>& OutToolTip)
+{
+	ConvertToMultilineToolTip(TEXT("Returns the light illuminance at the specified world position. This node samples the transmitance texture once."), 80, OutToolTip);
+
+}
+#endif // WITH_EDITOR
+
+///////////////////////////////////////////////////////////////////////////////
+// UMaterialExpressionSkyAtmosphereLightIlluminanceOnGround
+///////////////////////////////////////////////////////////////////////////////
+UMaterialExpressionSkyAtmosphereLightIlluminanceOnGround::UMaterialExpressionSkyAtmosphereLightIlluminanceOnGround(const FObjectInitializer& ObjectInitializer)
+	: Super(ObjectInitializer)
+{
+	// Structure to hold one-time initialization
+	struct FConstructorStatics
+	{
+		FText NAME_Sky;
+		FConstructorStatics()
+			: NAME_Sky(LOCTEXT("Sky", "Sky"))
+		{
+		}
+	};
+	static FConstructorStatics ConstructorStatics;
+
+#if WITH_EDITORONLY_DATA
+	MenuCategories.Add(ConstructorStatics.NAME_Sky);
+#endif
+}
+
+#if WITH_EDITOR
+int32 UMaterialExpressionSkyAtmosphereLightIlluminanceOnGround::Compile(class FMaterialCompiler* Compiler, int32 OutputIndex)
+{
+	return Compiler->SkyAtmosphereLightIlluminanceOnGround(LightIndex);
+}
+
+void UMaterialExpressionSkyAtmosphereLightIlluminanceOnGround::GetCaption(TArray<FString>& OutCaptions) const
+{
+	OutCaptions.Add(FString::Printf(TEXT("SkyAtmosphereLightIlluminanceOnGround[%i]"), LightIndex));
+}
+
+void UMaterialExpressionSkyAtmosphereLightIlluminanceOnGround::GetExpressionToolTip(TArray<FString>& OutToolTip)
+{
+	ConvertToMultilineToolTip(TEXT("Returns the light illuminance at level = 0 meter on the ground at the top of the planet. This node is cheap: it does not sample the transmitance texture."), 80, OutToolTip);
+
+}
 #endif // WITH_EDITOR
 
 ///////////////////////////////////////////////////////////////////////////////
diff --git a/Engine/Source/Runtime/Engine/Public/MaterialCompiler.h b/Engine/Source/Runtime/Engine/Public/MaterialCompiler.h
--- a/Engine/Source/Runtime/Engine/Public/MaterialCompiler.h
+++ b/Engine/Source/Runtime/Engine/Public/MaterialCompiler.h
@@ -502,6 +502,7 @@ public:
 	virtual int32 AtmosphericLightColor() = 0;
 
 	virtual int32 SkyAtmosphereLightIlluminance(int32 WorldPosition, int32 LightIndex) = 0;
+	virtual int32 SkyAtmosphereLightIlluminanceOnGround(int32 LightIndex) = 0;
 	virtual int32 SkyAtmosphereLightDirection(int32 LightIndex) = 0;
 	virtual int32 SkyAtmosphereLightDiskLuminance(int32 LightIndex, int32 OverrideAtmosphereLightDiscCosHalfApexAngle) = 0;
 	virtual int32 SkyAtmosphereViewLuminance() = 0;
@@ -1021,6 +1022,11 @@ public:
 		return Compiler->SkyAtmosphereLightIlluminance(WorldPosition, LightIndex);
 	}
 
+	virtual int32 SkyAtmosphereLightIlluminanceOnGround(int32 LightIndex) override
+	{
+		return Compiler->SkyAtmosphereLightIlluminanceOnGround(LightIndex);
+	}
+
 	virtual int32 SkyAtmosphereLightDirection(int32 LightIndex) override
 	{
 		return Compiler->SkyAtmosphereLightDirection(LightIndex);
