Add an assert during world cleanup to make sure the world is not being used by a cook package splitter.

diff --git a/Engine/Source/Runtime/Engine/Classes/Engine/World.h b/Engine/Source/Runtime/Engine/Classes/Engine/World.h
--- a/Engine/Source/Runtime/Engine/Classes/Engine/World.h
+++ b/Engine/Source/Runtime/Engine/Classes/Engine/World.h
@@ -1295,6 +1295,9 @@ public:
 
 	/** Indicates that the renderer scene for this editor world was purged while Play-in-Editor. */
 	uint8 bPurgedScene : 1;
+
+	/** Indicates that the world is currently being used by a cook splitter.  */
+	uint8 bUsedByCookSplitter : 1;
 #endif
 
 	/** Keeps track whether actors moved via PostEditMove and therefore constraint syncup should be performed. */
diff --git a/Engine/Source/Runtime/Engine/Private/World.cpp b/Engine/Source/Runtime/Engine/Private/World.cpp
--- a/Engine/Source/Runtime/Engine/Private/World.cpp
+++ b/Engine/Source/Runtime/Engine/Private/World.cpp
@@ -593,6 +593,7 @@ UWorld::UWorld( const FObjectInitializer& ObjectInitializer )
 , bDebugFrameStepExecutedThisFrame(false)
 , bToggledBetweenPIEandSIEThisFrame(false)
 , bPurgedScene(false)
+, bUsedByCookSplitter(false)
 #endif
 , bShouldTick(true)
 , bHasEverBeenInitialized(false)
@@ -5311,6 +5312,10 @@ bool UWorld::IsNavigationRebuilt() const
 
 void UWorld::CleanupWorld(bool bSessionEnded, bool bCleanupResources, UWorld* NewWorld)
 {
+#if WITH_EDITOR
+	check(!bUsedByCookSplitter);
+#endif
+
     CleanupWorldGlobalTag++;
 	if (!bIsWorldInitialized)
 	{
diff --git a/Engine/Source/Runtime/Engine/Private/WorldPartition/WorldPartitionCookPackageSplitter.cpp b/Engine/Source/Runtime/Engine/Private/WorldPartition/WorldPartitionCookPackageSplitter.cpp
--- a/Engine/Source/Runtime/Engine/Private/WorldPartition/WorldPartitionCookPackageSplitter.cpp
+++ b/Engine/Source/Runtime/Engine/Private/WorldPartition/WorldPartitionCookPackageSplitter.cpp
@@ -31,20 +31,23 @@ FWorldPartitionCookPackageSplitter::~FWorldPartitionCookPackageSplitter()
 
 void FWorldPartitionCookPackageSplitter::Teardown(ETeardown Status)
 {
- 	if (bInitializedWorldPartition)
+	if (UWorld* LocalWorld = ReferencedWorld.Get())
 	{
-		if (UWorld* LocalWorld = ReferencedWorld.Get())
+		check(LocalWorld->bUsedByCookSplitter);
+		LocalWorld->bUsedByCookSplitter = false;
+
+		if (bInitializedWorldPartition)
 		{
-			UWorldPartition* WorldPartition = LocalWorld->PersistentLevel->GetWorldPartition();
-			if (WorldPartition)
+			if (UWorldPartition* WorldPartition = LocalWorld->PersistentLevel->GetWorldPartition())
 			{
 				WorldPartition->EndCook(CookContext);
 				WorldPartition->Uninitialize();
 			}
 		}
-		bInitializedWorldPartition = false;
 	}
 
+	bInitializedWorldPartition = false;
+
 	if (bInitializedPhysicsSceneForSave)
 	{
 		GEditor->CleanupPhysicsSceneThatWasInitializedForSave(ReferencedWorld.Get(), bForceInitializedWorld);
@@ -90,6 +93,9 @@ TArray<ICookPackageSplitter::FGeneratedPackage> FWorldPartitionCookPackageSplitt
 	// that is necessary for populate 
 	ReferencedWorld = PartitionedWorld;
 
+	check(!PartitionedWorld->bUsedByCookSplitter);
+	PartitionedWorld->bUsedByCookSplitter = true;
+
 	check(!bInitializedPhysicsSceneForSave && !bForceInitializedWorld);
 	bInitializedPhysicsSceneForSave = GEditor->InitializePhysicsSceneForSaveIfNecessary(PartitionedWorld, bForceInitializedWorld);
 
