Fixed UTexture2D::CreateTransient unit test

#rnx

diff --git a/Engine/Source/Runtime/Engine/Private/Tests/Texture2DTests.cpp b/Engine/Source/Runtime/Engine/Private/Tests/Texture2DTests.cpp
--- a/Engine/Source/Runtime/Engine/Private/Tests/Texture2DTests.cpp
+++ b/Engine/Source/Runtime/Engine/Private/Tests/Texture2DTests.cpp
@@ -18,7 +18,7 @@ bool FTexture2DTestTransient::RunTest(const FString& Parameters)
 {
 	{
 		// Each test in this scope is expected to give one warning about invalid parameters
-		AddExpectedError(TEXT("Invalid parameters specified for UTexture2D::CreateTransient()"), EAutomationExpectedErrorFlags::Contains, 3);
+		AddExpectedError(TEXT("Negative size specified for UTexture2D::CreateTransient()"), EAutomationExpectedErrorFlags::Contains, 3);
 
 		UTexture2D* ZeroSizedTexture = UTexture2D::CreateTransient(0, 0);
 		TestTrue(TEXT("Creating a transient texture with a zero length dimension should fail!"), ZeroSizedTexture == nullptr);
