WorldPartition: Fix bad condition preventing unload of ContentBundle actors in Editor

#rnx

diff --git a/Engine/Source/Runtime/Engine/Private/WorldPartition/WorldPartitionActorLoaderInterface.cpp b/Engine/Source/Runtime/Engine/Private/WorldPartition/WorldPartitionActorLoaderInterface.cpp
--- a/Engine/Source/Runtime/Engine/Private/WorldPartition/WorldPartitionActorLoaderInterface.cpp
+++ b/Engine/Source/Runtime/Engine/Private/WorldPartition/WorldPartitionActorLoaderInterface.cpp
@@ -223,7 +223,7 @@ void IWorldPartitionActorLoaderInterface::ILoaderAdapter::OnActorDescContainerUn
 	{
 		const FActorReferenceMap* ActorReferences = GetContainerReferencesConst(It->GetContainer());
 
-		if (!ActorReferences || !ActorReferences->Find(It->GetGuid()))
+		if (ActorReferences && ActorReferences->Find(It->GetGuid()))
 		{
 			ActorsToUnload.Emplace(Container, It->GetGuid());
 		}
