Content Bundle: Make sure actors get loaded when World Partition is set to Disable Streaming or Disabled Loading in Editor

#rnx

diff --git a/Engine/Source/Runtime/Engine/Private/WorldPartition/WorldPartition.cpp b/Engine/Source/Runtime/Engine/Private/WorldPartition/WorldPartition.cpp
--- a/Engine/Source/Runtime/Engine/Private/WorldPartition/WorldPartition.cpp
+++ b/Engine/Source/Runtime/Engine/Private/WorldPartition/WorldPartition.cpp
@@ -2060,6 +2060,9 @@ UActorDescContainer* UWorldPartition::RegisterActorDescContainer(const FContaine
 
 		AddContainer(ContainerToRegister);
 
+		// @todo_ow: there is some redundancy here with UWorldPartition::Initialize, at some point we will also need to support instancing of CB containers to 
+		// support CB/EDL in Level Instances so this code should probably be shared with the code in UWorldPartition::Initialize
+		// See https://jira.it.epicgames.com/browse/UE-195953
 		if (IsInitialized() && EditorHash != nullptr)
 		{
 			FWorldPartitionReference WDLReference;
@@ -2072,10 +2075,25 @@ UActorDescContainer* UWorldPartition::RegisterActorDescContainer(const FContaine
 				}
 			}
 
+			const bool bIsStreamingEnabled = IsStreamingEnabledInEditor();
+			TArray<FGuid> ForceLoadedActorGuids;
+
 			for (UActorDescContainer::TIterator<> It(ContainerToRegister); It; ++It)
 			{
+				It->bIsForcedNonSpatiallyLoaded = !bIsStreamingEnabled;
+
+				if (ForceLoadedActors)
+				{
+					ForceLoadedActorGuids.Add(It->GetGuid());
+				}
+			
 				HashActorDesc(*It);
 			}
+
+			if (ForceLoadedActors && ForceLoadedActorGuids.Num() > 0)
+			{
+				ForceLoadedActors->AddActors(ForceLoadedActorGuids);
+			}
 		}
 
 		OnActorDescContainerRegistered.Broadcast(ContainerToRegister);
