Made skeletal mesh and landscape not create and update an instance data buffer (instead relying on the default behavior).

diff --git a/Engine/Source/Runtime/Engine/Private/SkeletalMesh.cpp b/Engine/Source/Runtime/Engine/Private/SkeletalMesh.cpp
--- a/Engine/Source/Runtime/Engine/Private/SkeletalMesh.cpp
+++ b/Engine/Source/Runtime/Engine/Private/SkeletalMesh.cpp
@@ -5793,15 +5793,8 @@ FSkeletalMeshSceneProxy::FSkeletalMeshSceneProxy(const USkinnedMeshComponent* Co
 	if (FeatureLevel > ERHIFeatureLevel::ES3_1) // TODO: support gpuscene for a mobile SkinVF
 	{
 		EnableGPUSceneSupportFlags();
-		const bool bUseGPUScene = UseGPUScene(GMaxRHIShaderPlatform, FeatureLevel);
 
-		if (bUseGPUScene)
-		{
-			bSupportsInstanceDataBuffer = true;
-			UpdateDefaultInstanceSceneData();
-		}
 	}
-
 #if RHI_RAYTRACING
 	if (IsRayTracingAllowed())
 	{
@@ -6915,9 +6908,6 @@ void FSkeletalMeshSceneProxy::OnTransformChanged(FRHICommandListBase& RHICmdList
 	// OnTransformChanged is called on the following frame after FSkeletalMeshObject::Update(), thus omit '+ 1' to frame number.
 	MeshObject->SetTransform(GetLocalToWorld(), GetScene().GetFrameNumber());
 	MeshObject->RefreshClothingTransforms(GetLocalToWorld(), GetScene().GetFrameNumber());
-
-	// Update the default-instance
-	UpdateDefaultInstanceSceneData();
 }
 
 //////////////////////////////////////////////////////////////////////////
diff --git a/Engine/Source/Runtime/Landscape/Private/LandscapeRender.cpp b/Engine/Source/Runtime/Landscape/Private/LandscapeRender.cpp
--- a/Engine/Source/Runtime/Landscape/Private/LandscapeRender.cpp
+++ b/Engine/Source/Runtime/Landscape/Private/LandscapeRender.cpp
@@ -1331,9 +1331,6 @@ FLandscapeComponentSceneProxy::FLandscapeComponentSceneProxy(ULandscapeComponent
 		}
 	}
 
-	bSupportsInstanceDataBuffer = true;
-	UpdateDefaultInstanceSceneData();
-
 #if WITH_EDITOR
 	const TArray<FWeightmapLayerAllocationInfo>& ComponentWeightmapLayerAllocations = InComponent->GetWeightmapLayerAllocations();
 	for (const FWeightmapLayerAllocationInfo& Allocation : ComponentWeightmapLayerAllocations)
@@ -2011,8 +2008,6 @@ void FLandscapeComponentSceneProxy::OnTransformChanged(FRHICommandListBase& RHIC
 
 	// Recache mesh draw commands for changed uniform buffers
 	GetScene().UpdateCachedRenderStates(this);
-
-	UpdateDefaultInstanceSceneData();
 }
 
 /** Creates a mesh batch for virtual texture or water info texture rendering. The caller is responsible for setting the required flags (bRenderToVirtualTexture, bUseForWaterInfoTextureDepth). Will render a simple fixed grid with combined subsections. */
