Prevent a potential integer overflow in AGameSession::GetNextPlayerID by specifying a max player ID and wrapping the NextPlayerId if it exceeds it back to a minimum of 256.

#rnx

diff --git a/Engine/Source/Runtime/Engine/Private/GameSession.cpp b/Engine/Source/Runtime/Engine/Private/GameSession.cpp
--- a/Engine/Source/Runtime/Engine/Private/GameSession.cpp
+++ b/Engine/Source/Runtime/Engine/Private/GameSession.cpp
@@ -225,7 +225,18 @@ void AGameSession::PostLogin(APlayerController* NewPlayer)
 int32 AGameSession::GetNextPlayerID()
 {
 	// Start at 256, because 255 is special (means all team for some UT Emote stuff)
-	static int32 NextPlayerID = 256;
+	static constexpr int32 MinPlayerId = 256;
+	static constexpr int32 MaxPlayerId = TNumericLimits<int32>::Max() - 1;
+	
+	static int32 NextPlayerID = MinPlayerId;
+	
+	// Prevent possible integer overflow by wrapping the value to the max player ID
+	if (NextPlayerID >= MaxPlayerId)
+	{
+		UE_LOG(LogGameSession, Warning, TEXT("AGameSession::GetNextPlayerID had to wrap the Player ID, this probably shouldn't have happened. PlayerID collisions may occur! Is this function being called incorrectly in a loop?"));
+		NextPlayerID = MinPlayerId;
+	}
+	
 	return NextPlayerID++;
 }
 
