Only update navigation for actor instance when initially populating the ISMC to prevent unnecessary dirtiness.

diff --git a/Engine/Plugins/Runtime/MassGameplay/Source/MassRepresentation/Private/MassVisualizationComponent.cpp b/Engine/Plugins/Runtime/MassGameplay/Source/MassRepresentation/Private/MassVisualizationComponent.cpp
--- a/Engine/Plugins/Runtime/MassGameplay/Source/MassRepresentation/Private/MassVisualizationComponent.cpp
+++ b/Engine/Plugins/Runtime/MassGameplay/Source/MassRepresentation/Private/MassVisualizationComponent.cpp
@@ -425,9 +425,6 @@ void UMassVisualizationComponent::HandleChangesWithExternalIDTracking(UInstanced
 			// remove instance
 			if (ISMComponent.PerInstanceSMData.IsValidIndex(InstanceIndex))
 			{
-				// Request navigation update
-				ISMComponent.PartialNavigationUpdate(InstanceIndex);
-				
 				ISMComponent.PerInstanceSMData.RemoveAtSwap(InstanceIndex, 1, false);
 				ISMComponent.PerInstanceSMCustomData.RemoveAt(InstanceIndex * ISMComponent.NumCustomDataFloats, ISMComponent.NumCustomDataFloats, false);
 			}
@@ -499,6 +496,8 @@ void UMassVisualizationComponent::HandleChangesWithExternalIDTracking(UInstanced
 		}
 	}
 
+	const bool bInitiallyEmpty = ISMComponent.PerInstanceSMData.Num() == 0;
+	
 	if (SharedData.UpdateInstanceIds.Num())
 	{
 		INC_DWORD_STAT_BY(STAT_Mass_VisualizationComponent_InstancesAddedNum, SharedData.UpdateInstanceIds.Num());
@@ -544,7 +543,7 @@ void UMassVisualizationComponent::HandleChangesWithExternalIDTracking(UInstanced
 			}
 
 			check(InstanceIds.Num() == InstanceTransforms.Num());
-			const TArray<int32> NewIndices = ISMComponent.AddInstances(InstanceTransforms, /*bShouldReturnIndices=*/true, /*bWorldSpace=*/true);
+			const TArray<int32> NewIndices = ISMComponent.AddInstances(InstanceTransforms, /*bShouldReturnIndices=*/true, /*bWorldSpace=*/true, /*bUpdateNavigation=*/bInitiallyEmpty);
 
 			check(InstanceIds.Num() == NewIndices.Num());
 			// note that in this case ISMComponent.PerInstanceSMData.Num() is always going to be greater than ISMComponent.PerInstanceIds
diff --git a/Engine/Source/Runtime/Engine/Classes/Components/HierarchicalInstancedStaticMeshComponent.h b/Engine/Source/Runtime/Engine/Classes/Components/HierarchicalInstancedStaticMeshComponent.h
--- a/Engine/Source/Runtime/Engine/Classes/Components/HierarchicalInstancedStaticMeshComponent.h
+++ b/Engine/Source/Runtime/Engine/Classes/Components/HierarchicalInstancedStaticMeshComponent.h
@@ -303,7 +303,7 @@ public:
 
 	// UInstancedStaticMesh interface
 	ENGINE_API virtual int32 AddInstance(const FTransform& InstanceTransform, bool bWorldSpace = false) override;
-	ENGINE_API virtual TArray<int32> AddInstances(const TArray<FTransform>& InstanceTransforms, bool bShouldReturnIndices, bool bWorldSpace = false) override;
+	ENGINE_API virtual TArray<int32> AddInstances(const TArray<FTransform>& InstanceTransforms, bool bShouldReturnIndices, bool bWorldSpace = false, bool bUpdateNavigation = true) override;
 	ENGINE_API virtual bool RemoveInstance(int32 InstanceIndex) override;
 	ENGINE_API virtual bool RemoveInstances(const TArray<int32>& InstancesToRemove) override;
 	ENGINE_API virtual bool RemoveInstances(const TArray<int32>& InstancesToRemove, bool bInstanceArrayAlreadySortedInReverseOrder) override;
diff --git a/Engine/Source/Runtime/Engine/Classes/Components/InstancedStaticMeshComponent.h b/Engine/Source/Runtime/Engine/Classes/Components/InstancedStaticMeshComponent.h
--- a/Engine/Source/Runtime/Engine/Classes/Components/InstancedStaticMeshComponent.h
+++ b/Engine/Source/Runtime/Engine/Classes/Components/InstancedStaticMeshComponent.h
@@ -161,7 +161,7 @@ class UInstancedStaticMeshComponent : public UStaticMeshComponent, public ISMIns
 
 	/** Add multiple instances to this component. Transform is given in local space of this component unless bWorldSpace is set. */
 	UFUNCTION(BlueprintCallable, Category="Components|InstancedStaticMesh")
-	ENGINE_API virtual TArray<int32> AddInstances(const TArray<FTransform>& InstanceTransforms, bool bShouldReturnIndices, bool bWorldSpace = false);
+	ENGINE_API virtual TArray<int32> AddInstances(const TArray<FTransform>& InstanceTransforms, bool bShouldReturnIndices, bool bWorldSpace = false, bool bUpdateNavigation = true);
 
 	/** Add an instance to this component. Transform is given in world space. */
 	UE_DEPRECATED(5.0, "Use AddInstance or AddInstances with bWorldSpace set to true.")
@@ -535,7 +535,7 @@ protected:
 	ENGINE_API int32 AddInstanceInternal(int32 InstanceIndex, FInstancedStaticMeshInstanceData* InNewInstanceData, const FTransform& InstanceTransform, bool bWorldSpace);
 
 	/** Internal implementation of AddInstances */
-	ENGINE_API TArray<int32> AddInstancesInternal(TConstArrayView<FTransform> InstanceTransforms, bool bShouldReturnIndices, bool bWorldSpace);
+	ENGINE_API TArray<int32> AddInstancesInternal(TConstArrayView<FTransform> InstanceTransforms, bool bShouldReturnIndices, bool bWorldSpace, bool bUpdateNavigation = true);
 
 	/** Internal version of RemoveInstance */	
 	ENGINE_API bool RemoveInstanceInternal(int32 InstanceIndex, bool InstanceAlreadyRemoved);
diff --git a/Engine/Source/Runtime/Engine/Private/HierarchicalInstancedStaticMesh.cpp b/Engine/Source/Runtime/Engine/Private/HierarchicalInstancedStaticMesh.cpp
--- a/Engine/Source/Runtime/Engine/Private/HierarchicalInstancedStaticMesh.cpp
+++ b/Engine/Source/Runtime/Engine/Private/HierarchicalInstancedStaticMesh.cpp
@@ -2560,11 +2560,11 @@ int32 UHierarchicalInstancedStaticMeshComponent::AddInstance(const FTransform& I
 	return InstanceIndex;
 }
 
-TArray<int32> UHierarchicalInstancedStaticMeshComponent::AddInstances(const TArray<FTransform>& InstanceTransforms, bool bShouldReturnIndices, bool bWorldSpace)
+TArray<int32> UHierarchicalInstancedStaticMeshComponent::AddInstances(const TArray<FTransform>& InstanceTransforms, bool bShouldReturnIndices, bool bWorldSpace, bool bUpdateNavigation)
 {
 	SCOPE_CYCLE_COUNTER(STAT_HISMCAddInstances);
 
-	TArray<int32> InstanceIndices = UInstancedStaticMeshComponent::AddInstances(InstanceTransforms, true, bWorldSpace);
+	TArray<int32> InstanceIndices = UInstancedStaticMeshComponent::AddInstances(InstanceTransforms, true, bWorldSpace, bUpdateNavigation);
 
 	// The tree will be fully rebuilt once the static mesh compilation is finished, no need for incremental update in that case.
 	if (InstanceIndices.Num() > 0 && GetStaticMesh() && !GetStaticMesh()->IsCompiling() && GetStaticMesh()->HasValidRenderData(false))
diff --git a/Engine/Source/Runtime/Engine/Private/InstancedStaticMesh.cpp b/Engine/Source/Runtime/Engine/Private/InstancedStaticMesh.cpp
--- a/Engine/Source/Runtime/Engine/Private/InstancedStaticMesh.cpp
+++ b/Engine/Source/Runtime/Engine/Private/InstancedStaticMesh.cpp
@@ -3714,7 +3714,7 @@ int32 UInstancedStaticMeshComponent::AddInstance(const FTransform& InstanceTrans
 	return AddInstanceInternal(PerInstanceSMData.Num(), nullptr, InstanceTransform, bWorldSpace);
 }
 
-TArray<int32> UInstancedStaticMeshComponent::AddInstancesInternal(TConstArrayView<FTransform> InstanceTransforms, bool bShouldReturnIndices, bool bWorldSpace)
+TArray<int32> UInstancedStaticMeshComponent::AddInstancesInternal(TConstArrayView<FTransform> InstanceTransforms, bool bShouldReturnIndices, bool bWorldSpace, bool bUpdateNavigation)
 {
 	const int32 Count = InstanceTransforms.Num();
 
@@ -3747,7 +3747,7 @@ TArray<int32> UInstancedStaticMeshComponent::AddInstancesInternal(TConstArrayVie
 			NewInstanceIndices.Add(InstanceIndex);
 		}
 
-		if (SupportsPartialNavigationUpdate())
+		if (bUpdateNavigation && SupportsPartialNavigationUpdate())
 		{
 			// If it's the first instance, register the component. 
 			// If there was no instance on component register, component registration was skipped because of UInstancedStaticMeshComponent::IsNavigationRelevant().
@@ -3772,7 +3772,7 @@ TArray<int32> UInstancedStaticMeshComponent::AddInstancesInternal(TConstArrayVie
 		++InstanceIndex;
 	}
 
-	if (!SupportsPartialNavigationUpdate())
+	if (bUpdateNavigation && !SupportsPartialNavigationUpdate())
 	{
 		FullNavigationUpdate();
 	}
@@ -3784,9 +3784,9 @@ TArray<int32> UInstancedStaticMeshComponent::AddInstancesInternal(TConstArrayVie
 	return NewInstanceIndices;
 }
 
-TArray<int32> UInstancedStaticMeshComponent::AddInstances(const TArray<FTransform>& InstanceTransforms, bool bShouldReturnIndices, bool bWorldSpace)
+TArray<int32> UInstancedStaticMeshComponent::AddInstances(const TArray<FTransform>& InstanceTransforms, bool bShouldReturnIndices, bool bWorldSpace, bool bUpdateNavigation)
 {
-	return AddInstancesInternal(InstanceTransforms, bShouldReturnIndices, bWorldSpace);
+	return AddInstancesInternal(InstanceTransforms, bShouldReturnIndices, bWorldSpace, bUpdateNavigation);
 }
 
 // Per Instance Custom Data - Updating custom data for specific instance
