- Changed AWorldDataLayers::DataLayersStateEpoch to be local to its associated world partition so that only this world partition UpdateStreamingState will be affected
- UWorldPartition::GetStreamingStateEpoch now combines both WP StreamingStateEpoch and it's AWorldDataLayers DataLayersStateEpoch

diff --git a/Engine/Source/Runtime/Engine/Private/WorldPartition/DataLayer/WorldDataLayers.cpp b/Engine/Source/Runtime/Engine/Private/WorldPartition/DataLayer/WorldDataLayers.cpp
--- a/Engine/Source/Runtime/Engine/Private/WorldPartition/DataLayer/WorldDataLayers.cpp
+++ b/Engine/Source/Runtime/Engine/Private/WorldPartition/DataLayer/WorldDataLayers.cpp
@@ -29,8 +29,6 @@
 
 #define LOCTEXT_NAMESPACE "WorldDataLayers"
 
-int32 AWorldDataLayers::DataLayersStateEpoch = 0;
-
 static FString JoinDataLayerShortNamesFromInstanceNames(AWorldDataLayers* InWorldDataLayers, const TArray<FName>& InDataLayerInstanceNames)
 {
 	check(InWorldDataLayers);
@@ -51,6 +49,7 @@ AWorldDataLayers::AWorldDataLayers(const FObjectInitializer& ObjectInitializer)
 #if WITH_EDITORONLY_DATA
 	, bAllowRuntimeDataLayerEditing(true)
 #endif
+	, DataLayersStateEpoch(0)
 {
 	bAlwaysRelevant = true;
 	bReplicates = true;
diff --git a/Engine/Source/Runtime/Engine/Private/WorldPartition/WorldPartition.cpp b/Engine/Source/Runtime/Engine/Private/WorldPartition/WorldPartition.cpp
--- a/Engine/Source/Runtime/Engine/Private/WorldPartition/WorldPartition.cpp
+++ b/Engine/Source/Runtime/Engine/Private/WorldPartition/WorldPartition.cpp
@@ -1255,6 +1255,14 @@ bool UWorldPartition::UseMakingInvisibleTransactionRequests() const
 	return bCachedUseMakingInvisibleTransactionRequests.Get(false);
 }
 
+int32 UWorldPartition::GetStreamingStateEpoch() const
+{
+	// Merge WorldPartition's StreamingStateEpoch and AWorldDataLayers DataLayersStateEpoch
+	const UWorld* OuterWorld = GetTypedOuter<UWorld>();
+	const AWorldDataLayers* WorldDataLayers = OuterWorld->GetWorldDataLayers();
+	return HashCombineFast(StreamingStateEpoch, WorldDataLayers ? WorldDataLayers->GetDataLayersStateEpoch() : 0);
+}
+
 bool UWorldPartition::IsSimulating(bool bIncludeTestEnableSimulationStreamingSource)
 {
 #if WITH_EDITOR
diff --git a/Engine/Source/Runtime/Engine/Private/WorldPartition/WorldPartitionRuntimeCell.cpp b/Engine/Source/Runtime/Engine/Private/WorldPartition/WorldPartitionRuntimeCell.cpp
--- a/Engine/Source/Runtime/Engine/Private/WorldPartition/WorldPartitionRuntimeCell.cpp
+++ b/Engine/Source/Runtime/Engine/Private/WorldPartition/WorldPartitionRuntimeCell.cpp
@@ -130,17 +130,20 @@ UDataLayerManager* UWorldPartitionRuntimeCell::GetDataLayerManager() const
 
 EDataLayerRuntimeState UWorldPartitionRuntimeCell::GetCellEffectiveWantedState() const
 {
-	if (EffectiveWantedStateEpoch != AWorldDataLayers::GetDataLayersStateEpoch())
+	if (!HasDataLayers())
 	{
-		if (!HasDataLayers())
-		{
-			EffectiveWantedState = EDataLayerRuntimeState::Activated;
-		}
-		else
+		EffectiveWantedState = EDataLayerRuntimeState::Activated;
+	}
+	else
+	{
+		const UWorld* OuterWorld = GetOuterWorld();
+		const AWorldDataLayers* WorldDataLayers = OuterWorld->GetWorldDataLayers();
+		check(WorldDataLayers);
+		if (EffectiveWantedStateEpoch != WorldDataLayers->GetDataLayersStateEpoch())
 		{
 			EffectiveWantedState = EDataLayerRuntimeState::Unloaded;
 
-			UWorldPartition* WorldPartition = GetOuterWorld()->GetWorldPartition();
+			UWorldPartition* WorldPartition = OuterWorld->GetWorldPartition();
 			if (const UDataLayerManager* DataLayerManager = WorldPartition->GetDataLayerManager())
 			{
 				switch (WorldPartition->GetDataLayersLogicOperator())
@@ -169,9 +172,9 @@ EDataLayerRuntimeState UWorldPartitionRuntimeCell::GetCellEffectiveWantedState()
 					checkNoEntry();
 				}
 			}
-		}
 
-		EffectiveWantedStateEpoch = AWorldDataLayers::GetDataLayersStateEpoch();
+			EffectiveWantedStateEpoch = WorldDataLayers->GetDataLayersStateEpoch();
+		}
 	}
 
 	return EffectiveWantedState;
diff --git a/Engine/Source/Runtime/Engine/Private/WorldPartition/WorldPartitionStreamingPolicy.cpp b/Engine/Source/Runtime/Engine/Private/WorldPartition/WorldPartitionStreamingPolicy.cpp
--- a/Engine/Source/Runtime/Engine/Private/WorldPartition/WorldPartitionStreamingPolicy.cpp
+++ b/Engine/Source/Runtime/Engine/Private/WorldPartition/WorldPartitionStreamingPolicy.cpp
@@ -69,7 +69,6 @@ UWorldPartitionStreamingPolicy::UWorldPartitionStreamingPolicy(const FObjectInit
 	, ProcessedToLoadCells(0)
 	, bCriticalPerformanceRequestedBlockTillOnWorld(false)
 	, CriticalPerformanceBlockTillLevelStreamingCompletedEpoch(0)
-	, ServerDataLayersStatesEpoch(INT_MIN)
 	, ServerStreamingStateEpoch(INT_MIN)
 	, ServerStreamingEnabledEpoch(INT_MIN)
 	, UpdateStreamingHash(0)
@@ -163,8 +162,7 @@ uint32 UWorldPartitionStreamingPolicy::ComputeUpdateStreamingHash(bool bCanOptim
 			HashBuilder << WorldPartition->RuntimeHash->ComputeUpdateStreamingHash();
 		}
 		HashBuilder << ComputeServerStreamingEnabledEpoch();
-		HashBuilder << GetOuterUWorldPartition()->GetStreamingStateEpoch();
-		HashBuilder << AWorldDataLayers::GetDataLayersStateEpoch();
+		HashBuilder << WorldPartition->GetStreamingStateEpoch();
 		HashBuilder << bIsStreaming3D;
 		for (const FWorldPartitionStreamingSource& Source : StreamingSources)
 		{
@@ -365,8 +363,7 @@ void UWorldPartitionStreamingPolicy::UpdateStreamingState()
 			if (!bIsServerStreamingEnabled && 
 				bLastUpdateCompletedLoadingAndActivation && 
 				(ServerStreamingEnabledEpoch == NewServerStreamingEnabledEpoch) &&
-				(ServerStreamingStateEpoch == WorldPartition->GetStreamingStateEpoch()) &&
-				(ServerDataLayersStatesEpoch == AWorldDataLayers::GetDataLayersStateEpoch()))
+				(ServerStreamingStateEpoch == WorldPartition->GetStreamingStateEpoch()))
 			{
 				// Server as nothing to do early out
 				return; 
@@ -599,7 +596,6 @@ void UWorldPartitionStreamingPolicy::UpdateStreamingState()
 	{
 		if (bUpdateServerEpoch)
 		{
-			ServerDataLayersStatesEpoch = AWorldDataLayers::GetDataLayersStateEpoch();
 			ServerStreamingStateEpoch = WorldPartition->GetStreamingStateEpoch();
 			ServerStreamingEnabledEpoch = NewServerStreamingEnabledEpoch;
 			UE_LOG(LogWorldPartition, Verbose, TEXT("Server epoch updated"));
diff --git a/Engine/Source/Runtime/Engine/Public/WorldPartition/DataLayer/WorldDataLayers.h b/Engine/Source/Runtime/Engine/Public/WorldPartition/DataLayer/WorldDataLayers.h
--- a/Engine/Source/Runtime/Engine/Public/WorldPartition/DataLayer/WorldDataLayers.h
+++ b/Engine/Source/Runtime/Engine/Public/WorldPartition/DataLayer/WorldDataLayers.h
@@ -122,7 +122,7 @@ public:
 	const TSet<FName>& GetEffectiveLoadedDataLayerNames() const;
 	UFUNCTION(NetMulticast, Reliable)
 	ENGINE_API void OnDataLayerRuntimeStateChanged(const UDataLayerInstance* InDataLayer, EDataLayerRuntimeState InState);
-	static int32 GetDataLayersStateEpoch() { return DataLayersStateEpoch; }
+	ENGINE_API int32 GetDataLayersStateEpoch() const { return DataLayersStateEpoch; }
 
 	//~ Begin Deprecated
 
@@ -258,7 +258,7 @@ private:
 	mutable int32 AllEffectiveLoadedDataLayerNamesEpoch;
 	mutable TSet<FName> AllEffectiveLoadedDataLayerNames;
 
-	static ENGINE_API int32 DataLayersStateEpoch;
+	int32 DataLayersStateEpoch;
 
 	static_assert(DATALAYER_TO_INSTANCE_RUNTIME_CONVERSION_ENABLED, "bHasDeprecatedDataLayers is deprecated and needs to be deleted.");
 	bool bHasDeprecatedDataLayers;
diff --git a/Engine/Source/Runtime/Engine/Public/WorldPartition/WorldPartition.h b/Engine/Source/Runtime/Engine/Public/WorldPartition/WorldPartition.h
--- a/Engine/Source/Runtime/Engine/Public/WorldPartition/WorldPartition.h
+++ b/Engine/Source/Runtime/Engine/Public/WorldPartition/WorldPartition.h
@@ -389,7 +389,7 @@ public:
 
 public:
 	static ENGINE_API bool IsSimulating(bool bIncludeTestEnableSimulationStreamingSource = true);
-	int32 GetStreamingStateEpoch() const { return StreamingStateEpoch; }
+	ENGINE_API int32 GetStreamingStateEpoch() const;
 
 	ENGINE_API bool CanInitialize(UWorld* InWorld) const;
 	ENGINE_API void Initialize(UWorld* World, const FTransform& InTransform);
diff --git a/Engine/Source/Runtime/Engine/Public/WorldPartition/WorldPartitionStreamingPolicy.h b/Engine/Source/Runtime/Engine/Public/WorldPartition/WorldPartitionStreamingPolicy.h
--- a/Engine/Source/Runtime/Engine/Public/WorldPartition/WorldPartitionStreamingPolicy.h
+++ b/Engine/Source/Runtime/Engine/Public/WorldPartition/WorldPartitionStreamingPolicy.h
@@ -140,7 +140,6 @@ private:
 
 	bool bCriticalPerformanceRequestedBlockTillOnWorld;
 	int32 CriticalPerformanceBlockTillLevelStreamingCompletedEpoch;
-	int32 ServerDataLayersStatesEpoch;
 	int32 ServerStreamingStateEpoch;
 	int32 ServerStreamingEnabledEpoch;
 	uint32 UpdateStreamingHash;
