Fix crash when importing ADPCM audio files.

diff --git a/Engine/Source/Editor/AudioEditor/Private/AudioEditorModule.cpp b/Engine/Source/Editor/AudioEditor/Private/AudioEditorModule.cpp
--- a/Engine/Source/Editor/AudioEditor/Private/AudioEditorModule.cpp
+++ b/Engine/Source/Editor/AudioEditor/Private/AudioEditorModule.cpp
@@ -125,7 +125,7 @@ public:
 
 		SetupIcons();
 #if WITH_SNDFILE_IO
-		if (!Audio::InitSoundFileIOManager())
+		if (!Audio::SoundFileUtils::InitSoundFileIOManager())
 		{
 			UE_LOG(LogAudioEditor, Display, TEXT("LibSoundFile failed to load. Importing audio will not work correctly."));
 		}
@@ -136,7 +136,7 @@ public:
 	{
 		LLM_SCOPE(ELLMTag::AudioMisc);
 #if WITH_SNDFILE_IO
-		Audio::ShutdownSoundFileIOManager();
+		Audio::SoundFileUtils::ShutdownSoundFileIOManager();
 #endif // WITH_SNDFILE_IO
 
 		SoundClassExtensibility.Reset();
diff --git a/Engine/Source/Editor/AudioEditor/Private/Factories/SoundFactory.cpp b/Engine/Source/Editor/AudioEditor/Private/Factories/SoundFactory.cpp
--- a/Engine/Source/Editor/AudioEditor/Private/Factories/SoundFactory.cpp
+++ b/Engine/Source/Editor/AudioEditor/Private/Factories/SoundFactory.cpp
@@ -207,7 +207,7 @@ UObject* USoundFactory::FactoryCreateBinary
 
 		// Convert audio data to a wav file in memory
 		TArray<uint8> RawWaveData;
-		if (Audio::ConvertAudioToWav(RawAudioData, RawWaveData))
+		if (Audio::SoundFileUtils::ConvertAudioToWav(RawAudioData, RawWaveData))
 		{
 			const uint8* Ptr = &RawWaveData[0];
 
@@ -376,10 +376,10 @@ UObject* USoundFactory::CreateObject
 		if (*WaveInfo.pBitsPerSample != 16)
 		{
 #if WITH_SNDFILE_IO
-			const uint32 OrigNumSamples = WaveInfo.GetNumSamples();
+			const uint32 OrigNumSamples = Audio::SoundFileUtils::GetNumSamples(RawWaveData);
 
 			// Attempt to convert to 16 bit audio
-			if (Audio::ConvertAudioToWav(RawWaveData, ConvertedRawWaveData))
+			if (Audio::SoundFileUtils::ConvertAudioToWav(RawWaveData, ConvertedRawWaveData))
 			{
 				WaveInfo = FWaveModInfo();				
 				if (!WaveInfo.ReadWaveInfo(ConvertedRawWaveData.GetData(), ConvertedRawWaveData.Num(), &ErrorMessage))
diff --git a/Engine/Source/Runtime/AudioMixer/Private/SoundFileIO/SoundFileIO.cpp b/Engine/Source/Runtime/AudioMixer/Private/SoundFileIO/SoundFileIO.cpp
--- a/Engine/Source/Runtime/AudioMixer/Private/SoundFileIO/SoundFileIO.cpp
+++ b/Engine/Source/Runtime/AudioMixer/Private/SoundFileIO/SoundFileIO.cpp
@@ -16,7 +16,7 @@
 #include "Stats/Stats.h"
 
 
-namespace Audio
+namespace Audio::SoundFileUtils
 {
 	static void CopyOptionalWavChunks(TSharedPtr<ISoundFileReader>& InSoundDataReader, const int32 InInputFormat, TSharedPtr<ISoundFileWriter>& InSoundFileWriter, const int32 InOutputFormat);
 
@@ -30,6 +30,26 @@ namespace Audio
 		return Audio::SoundFileIOManagerShutdown();
 	}
 
+	uint32 AUDIOMIXER_API GetNumSamples(const TArray<uint8>& InAudioData)
+	{
+		FSoundFileIOManager SoundIOManager;
+		TSharedPtr<ISoundFileReader> InputSoundDataReader = SoundIOManager.CreateSoundDataReader();
+
+		ESoundFileError::Type Error = InputSoundDataReader->Init(&InAudioData);
+		if (Error != ESoundFileError::Type::NONE)
+		{
+			return 0;
+		}
+
+		TArray<ESoundFileChannelMap::Type> ChannelMap;
+
+		FSoundFileDescription InputDescription;
+		InputSoundDataReader->GetDescription(InputDescription, ChannelMap);
+		InputSoundDataReader->Release();
+
+		return InputDescription.NumFrames * InputDescription.NumChannels;
+	}
+
 	bool AUDIOMIXER_API ConvertAudioToWav(const TArray<uint8>& InAudioData, TArray<uint8>& OutWaveData)
 	{
 		const FSoundFileConvertFormat ConvertFormat = FSoundFileConvertFormat::CreateDefault();
diff --git a/Engine/Source/Runtime/AudioMixer/Public/SoundFileIO/SoundFileIO.h b/Engine/Source/Runtime/AudioMixer/Public/SoundFileIO/SoundFileIO.h
--- a/Engine/Source/Runtime/AudioMixer/Public/SoundFileIO/SoundFileIO.h
+++ b/Engine/Source/Runtime/AudioMixer/Public/SoundFileIO/SoundFileIO.h
@@ -4,9 +4,10 @@
 
 #include "CoreMinimal.h"
 
-namespace Audio
+namespace Audio::SoundFileUtils
 {
 	bool AUDIOMIXER_API InitSoundFileIOManager();
 	bool AUDIOMIXER_API ShutdownSoundFileIOManager();
+	uint32 AUDIOMIXER_API GetNumSamples(const TArray<uint8>& InAudioData);
 	bool AUDIOMIXER_API ConvertAudioToWav(const TArray<uint8>& InAudioData, TArray<uint8>& OutWaveData);
 }
diff --git a/Engine/Source/Runtime/Engine/Private/Audio.cpp b/Engine/Source/Runtime/Engine/Private/Audio.cpp
--- a/Engine/Source/Runtime/Engine/Private/Audio.cpp
+++ b/Engine/Source/Runtime/Engine/Private/Audio.cpp
@@ -1620,10 +1620,7 @@ bool FWaveModInfo::ReadWaveInfo( const uint8* WaveData, int32 WaveDataSize, FStr
 	SampleDataSize = INTEL_ORDER32( RiffChunk->ChunkLen );
 	SampleDataEnd = SampleDataStart + SampleDataSize;
 
-	if (*pFormatTag != 0x0001 // WAVE_FORMAT_PCM
-		&& *pFormatTag != 0x0002 // WAVE_FORMAT_ADPCM
-		&& *pFormatTag != 0x0011 // WAVE_FORMAT_DVI_ADPCM
-		&& *pFormatTag != 0x0003) // WAVE_FORMAT_IEEE_FLOAT
+	if (!IsFormatSupported())
 	{
 		ReportImportFailure();
 		if (ErrorReason) *ErrorReason = TEXT("Unsupported wave file format.  Only PCM, ADPCM, and DVI ADPCM can be imported.");
@@ -1906,7 +1903,9 @@ void FWaveModInfo::ReportImportFailure() const
 
 uint32 FWaveModInfo::GetNumSamples() const
 {
-	if (*pBitsPerSample != 0)
+	// The calculation below only works for uncompressed formats.
+	// For compressed formats see Audio::SoundFileUtils::GetNumSamples().
+	if (IsFormatUncompressed() && *pBitsPerSample >= 8)
 	{
 		return SampleDataSize / (*pBitsPerSample / 8);
 	}
@@ -1914,6 +1913,20 @@ uint32 FWaveModInfo::GetNumSamples() const
 	return 0;
 }
 
+bool FWaveModInfo::IsFormatSupported() const
+{
+	return (*pFormatTag == WAVE_FORMAT_PCM
+		|| *pFormatTag == WAVE_FORMAT_ADPCM
+		|| *pFormatTag == WAVE_FORMAT_DVI_ADPCM
+		|| *pFormatTag == WAVE_FORMAT_IEEE_FLOAT);
+}
+
+bool FWaveModInfo::IsFormatUncompressed() const
+{
+	return (*pFormatTag == WAVE_FORMAT_PCM
+		|| *pFormatTag == WAVE_FORMAT_IEEE_FLOAT);
+}
+
 static void WriteUInt32ToByteArrayLE(TArray<uint8>& InByteArray, int32& Index, const uint32 Value)
 {
 	InByteArray[Index++] = (uint8)(Value >> 0);
diff --git a/Engine/Source/Runtime/Engine/Public/Audio.h b/Engine/Source/Runtime/Engine/Public/Audio.h
--- a/Engine/Source/Runtime/Engine/Public/Audio.h
+++ b/Engine/Source/Runtime/Engine/Public/Audio.h
@@ -884,6 +884,12 @@ class FWaveModInfo
 {
 public:
 
+	// Format specifiers
+	static constexpr uint16 WAVE_FORMAT_PCM = 0x0001;
+	static constexpr uint16 WAVE_FORMAT_ADPCM = 0x0002;
+	static constexpr uint16 WAVE_FORMAT_IEEE_FLOAT = 0x0003;
+	static constexpr uint16 WAVE_FORMAT_DVI_ADPCM = 0x0011;
+
 	// Pointers to variables in the in-memory WAVE file.
 	const uint32* pSamplesPerSec;
 	const uint32* pAvgBytesPerSec;
@@ -939,6 +945,11 @@ public:
 	/** Return total number of samples */
 	ENGINE_API uint32 GetNumSamples() const;
 
+	/** Return whether file format is supported for import */
+	ENGINE_API bool IsFormatSupported() const;
+	/** Return whether file format contains uncompressed PCM data */
+	ENGINE_API bool IsFormatUncompressed() const;
+
 };
 
 /** Utility to serialize raw PCM data into a wave file. */
