Fix for C2059 compile error caused by preprocessor symbol collision.

diff --git a/Engine/Source/Runtime/Engine/Private/Audio.cpp b/Engine/Source/Runtime/Engine/Private/Audio.cpp
--- a/Engine/Source/Runtime/Engine/Private/Audio.cpp
+++ b/Engine/Source/Runtime/Engine/Private/Audio.cpp
@@ -1915,16 +1915,16 @@ uint32 FWaveModInfo::GetNumSamples() const
 
 bool FWaveModInfo::IsFormatSupported() const
 {
-	return (*pFormatTag == WAVE_FORMAT_PCM
-		|| *pFormatTag == WAVE_FORMAT_ADPCM
-		|| *pFormatTag == WAVE_FORMAT_DVI_ADPCM
-		|| *pFormatTag == WAVE_FORMAT_IEEE_FLOAT);
+	return (*pFormatTag == WAVE_INFO_FORMAT_PCM
+		|| *pFormatTag == WAVE_INFO_FORMAT_ADPCM
+		|| *pFormatTag == WAVE_INFO_FORMAT_DVI_ADPCM
+		|| *pFormatTag == WAVE_INFO_FORMAT_IEEE_FLOAT);
 }
 
 bool FWaveModInfo::IsFormatUncompressed() const
 {
-	return (*pFormatTag == WAVE_FORMAT_PCM
-		|| *pFormatTag == WAVE_FORMAT_IEEE_FLOAT);
+	return (*pFormatTag == WAVE_INFO_FORMAT_PCM
+		|| *pFormatTag == WAVE_INFO_FORMAT_IEEE_FLOAT);
 }
 
 static void WriteUInt32ToByteArrayLE(TArray<uint8>& InByteArray, int32& Index, const uint32 Value)
diff --git a/Engine/Source/Runtime/Engine/Public/Audio.h b/Engine/Source/Runtime/Engine/Public/Audio.h
--- a/Engine/Source/Runtime/Engine/Public/Audio.h
+++ b/Engine/Source/Runtime/Engine/Public/Audio.h
@@ -885,10 +885,10 @@ class FWaveModInfo
 public:
 
 	// Format specifiers
-	static constexpr uint16 WAVE_FORMAT_PCM = 0x0001;
-	static constexpr uint16 WAVE_FORMAT_ADPCM = 0x0002;
-	static constexpr uint16 WAVE_FORMAT_IEEE_FLOAT = 0x0003;
-	static constexpr uint16 WAVE_FORMAT_DVI_ADPCM = 0x0011;
+	static constexpr uint16 WAVE_INFO_FORMAT_PCM = 0x0001;
+	static constexpr uint16 WAVE_INFO_FORMAT_ADPCM = 0x0002;
+	static constexpr uint16 WAVE_INFO_FORMAT_IEEE_FLOAT = 0x0003;
+	static constexpr uint16 WAVE_INFO_FORMAT_DVI_ADPCM = 0x0011;
 
 	// Pointers to variables in the in-memory WAVE file.
 	const uint32* pSamplesPerSec;
