World Partition - Helpers: ForEachIntersectingActorDesc() ignored the provided function return value, calling the user function on all intersecting actors

diff --git a/Engine/Source/Runtime/Engine/Private/WorldPartition/WorldPartitionHelpers.cpp b/Engine/Source/Runtime/Engine/Private/WorldPartition/WorldPartitionHelpers.cpp
--- a/Engine/Source/Runtime/Engine/Private/WorldPartition/WorldPartitionHelpers.cpp
+++ b/Engine/Source/Runtime/Engine/Private/WorldPartition/WorldPartitionHelpers.cpp
@@ -94,11 +94,13 @@ bool FWorldPartitionHelpers::IsActorDescClassCompatibleWith(const FWorldPartitio
 
 void FWorldPartitionHelpers::ForEachIntersectingActorDesc(UWorldPartition* WorldPartition, const FBox& Box, TSubclassOf<AActor> ActorClass, TFunctionRef<bool(const FWorldPartitionActorDesc*)> Func)
 {
-	WorldPartition->EditorHash->ForEachIntersectingActor(Box, [&ActorClass, Func](const FWorldPartitionActorDesc* ActorDesc)
+	bool bProcessNextActors = true;
+
+	WorldPartition->EditorHash->ForEachIntersectingActor(Box, [&ActorClass, Func, &bProcessNextActors](const FWorldPartitionActorDesc* ActorDesc)
 	{
-		if (IsActorDescClassCompatibleWith(ActorDesc, ActorClass))
+		if (bProcessNextActors && IsActorDescClassCompatibleWith(ActorDesc, ActorClass))
 		{
-			Func(ActorDesc);
+			bProcessNextActors = Func(ActorDesc);
 		}
 	});
 }
