Prevent the server from sending the server-only-option of ?listen when doing a Client Travel.  This will prevent the client from incorrectly thinking it's a listen server during travel.

diff --git a/Engine/Source/Runtime/Engine/Private/GameModeBase.cpp b/Engine/Source/Runtime/Engine/Private/GameModeBase.cpp
--- a/Engine/Source/Runtime/Engine/Private/GameModeBase.cpp
+++ b/Engine/Source/Runtime/Engine/Private/GameModeBase.cpp
@@ -49,6 +49,14 @@ namespace UE::GameModeBase::Private
 		bAllowPIESeamlessTravel,
 		TEXT("When true, allow seamless travels in single process PIE.")
 	);
+
+	// Fixed in UE5.4
+	static bool bAllowListenParamSentToClient = false;
+	static FAutoConsoleVariableRef CVarAllowListenParamSentToClient(
+		TEXT("net.fix.AllowListenParamSentToClient"),
+		bAllowListenParamSentToClient,
+		TEXT("When true, allow '?listen' sent to the client during Travel. This would confuse the client and cause crashes with RPCs.")
+	);
 }
 
 
@@ -391,8 +399,18 @@ APlayerController* AGameModeBase::ProcessClientTravel(FString& FURL, bool bSeaml
 		{
 			if (Cast<UNetConnection>(PlayerController->Player) != nullptr)
 			{
+				FString ClientUrl = FURL;
+				if (!UE::GameModeBase::Private::bAllowListenParamSentToClient)
+				{
+					// We are going to lop off the listen option when sending to the client, so they don't think they're the server.
+					// That sounds silly, but in GetNetMode, we will eventually arrive at AttemptDeriveFromURL in certain circumstances (such as during travel!)
+					// The client would then think that they're supposed to execute things locally rather than RPC them based on the listen param.
+					ClientUrl = FURL.Replace(TEXT("?listen?"), TEXT("?"));
+					ClientUrl.RemoveFromEnd(TEXT("?listen"));
+				}
+
 				// Remote player
-				PlayerController->ClientTravel(FURL, TRAVEL_Relative, bSeamless);
+				PlayerController->ClientTravel(ClientUrl, TRAVEL_Relative, bSeamless);
 			}
 			else
 			{
