Exposing FAnimTrack::IsAdditive() function
#rnx

diff --git a/Engine/Source/Runtime/Engine/Classes/Animation/AnimCompositeBase.h b/Engine/Source/Runtime/Engine/Classes/Animation/AnimCompositeBase.h
--- a/Engine/Source/Runtime/Engine/Classes/Animation/AnimCompositeBase.h
+++ b/Engine/Source/Runtime/Engine/Classes/Animation/AnimCompositeBase.h
@@ -257,7 +257,7 @@ struct FAnimTrack
 
 	FAnimTrack() {}
 	ENGINE_API float GetLength() const;
-	bool IsAdditive() const;
+	ENGINE_API bool IsAdditive() const;
 	bool IsRotationOffsetAdditive() const;
 
 	ENGINE_API int32 GetTrackAdditiveType() const;
