Renamed UHLODSubsystem to UWorldPartitionHLODRuntimeSubsystem

diff --git a/Engine/Source/Editor/UnrealEd/Private/WorldPartition/WorldPartitionHLODsBuilder.cpp b/Engine/Source/Editor/UnrealEd/Private/WorldPartition/WorldPartitionHLODsBuilder.cpp
--- a/Engine/Source/Editor/UnrealEd/Private/WorldPartition/WorldPartitionHLODsBuilder.cpp
+++ b/Engine/Source/Editor/UnrealEd/Private/WorldPartition/WorldPartitionHLODsBuilder.cpp
@@ -30,7 +30,7 @@
 #include "WorldPartition/HLOD/HLODActorDesc.h"
 #include "WorldPartition/HLOD/HLODLayer.h"
 #include "WorldPartition/HLOD/HLODProviderInterface.h"
-#include "WorldPartition/HLOD/HLODSubsystem.h"
+#include "WorldPartition/HLOD/HLODRuntimeSubsystem.h"
 
 DEFINE_LOG_CATEGORY_STATIC(LogWorldPartitionHLODsBuilder, All, All);
 
@@ -666,7 +666,7 @@ bool UWorldPartitionHLODsBuilder::SubmitHLODActors()
 bool UWorldPartitionHLODsBuilder::DumpStats()
 {
 	const FString HLODStatsOutputFilename = FPaths::ProjectSavedDir() / TEXT("WorldPartition") / FString::Printf(TEXT("HLODStats-%08x.csv"), FPlatformProcess::GetCurrentProcessId());
-	return UHLODSubsystem::WriteHLODStatsCSV(World, HLODStatsOutputFilename);
+	return UWorldPartitionHLODRuntimeSubsystem::WriteHLODStatsCSV(World, HLODStatsOutputFilename);
 }
 
 bool UWorldPartitionHLODsBuilder::GetHLODActorsToBuild(TArray<FGuid>& HLODActorsToBuild) const
diff --git a/Engine/Source/Runtime/Engine/Private/UnrealEngine.cpp b/Engine/Source/Runtime/Engine/Private/UnrealEngine.cpp
--- a/Engine/Source/Runtime/Engine/Private/UnrealEngine.cpp
+++ b/Engine/Source/Runtime/Engine/Private/UnrealEngine.cpp
@@ -238,7 +238,7 @@ UnrealEngine.cpp: Implements the UEngine class and helpers.
 #endif
 
 #if !(UE_BUILD_SHIPPING || UE_BUILD_TEST)
-#include "WorldPartition/HLOD/HLODSubsystem.h"
+#include "WorldPartition/HLOD/HLODRuntimeSubsystem.h"
 #endif
 
 #if WITH_DUMPGPU
@@ -12168,7 +12168,7 @@ float DrawMapWarnings(UWorld* World, FViewport* Viewport, FCanvas* Canvas, UCanv
 	}
 
 #if WITH_EDITOR
-	if (UHLODSubsystem* HLODSubsystem = World->GetSubsystem<UHLODSubsystem>())
+	if (UWorldPartitionHLODRuntimeSubsystem* HLODSubsystem = World->GetSubsystem<UWorldPartitionHLODRuntimeSubsystem>())
 	{
 		uint32 NumOutdatedHLODActors = HLODSubsystem->GetNumOutdatedHLODActors();
 		if (NumOutdatedHLODActors)
diff --git a/Engine/Source/Runtime/Engine/Private/WorldPartition/HLOD/HLODActor.cpp b/Engine/Source/Runtime/Engine/Private/WorldPartition/HLOD/HLODActor.cpp
--- a/Engine/Source/Runtime/Engine/Private/WorldPartition/HLOD/HLODActor.cpp
+++ b/Engine/Source/Runtime/Engine/Private/WorldPartition/HLOD/HLODActor.cpp
@@ -2,7 +2,7 @@
 
 #include "WorldPartition/HLOD/HLODActor.h"
 #include "Engine/World.h"
-#include "WorldPartition/HLOD/HLODSubsystem.h"
+#include "WorldPartition/HLOD/HLODRuntimeSubsystem.h"
 #include "Components/PrimitiveComponent.h"
 #include "Misc/PackageName.h"
 #include "UObject/UE5MainStreamObjectVersion.h"
@@ -94,12 +94,12 @@ void AWorldPartitionHLOD::SetVisibility(bool bInVisible)
 void AWorldPartitionHLOD::BeginPlay()
 {
 	Super::BeginPlay();
-	GetWorld()->GetSubsystem<UHLODSubsystem>()->RegisterHLODActor(this);
+	GetWorld()->GetSubsystem<UWorldPartitionHLODRuntimeSubsystem>()->RegisterHLODActor(this);
 }
 
 void AWorldPartitionHLOD::EndPlay(const EEndPlayReason::Type EndPlayReason)
 {
-	GetWorld()->GetSubsystem<UHLODSubsystem>()->UnregisterHLODActor(this);
+	GetWorld()->GetSubsystem<UWorldPartitionHLODRuntimeSubsystem>()->UnregisterHLODActor(this);
 	Super::EndPlay(EndPlayReason);
 }
 
diff --git a/Engine/Source/Runtime/Engine/Private/WorldPartition/HLOD/HLODDestruction.cpp b/Engine/Source/Runtime/Engine/Private/WorldPartition/HLOD/HLODDestruction.cpp
--- a/Engine/Source/Runtime/Engine/Private/WorldPartition/HLOD/HLODDestruction.cpp
+++ b/Engine/Source/Runtime/Engine/Private/WorldPartition/HLOD/HLODDestruction.cpp
@@ -8,7 +8,7 @@
 #include "UObject/ScriptInterface.h"
 #include "WorldPartition/HLOD/DestructibleHLODComponent.h"
 #include "WorldPartition/HLOD/HLODActor.h"
-#include "WorldPartition/HLOD/HLODSubsystem.h"
+#include "WorldPartition/HLOD/HLODRuntimeSubsystem.h"
 
 #include UE_INLINE_GENERATED_CPP_BY_NAME(HLODDestruction)
 
@@ -33,7 +33,7 @@ namespace
 
 		if (RuntimeCell)
 		{
-			for (AWorldPartitionHLOD* HLODActor : InActor->GetWorld()->GetSubsystem<UHLODSubsystem>()->GetHLODActorsForCell(RuntimeCell))
+			for (AWorldPartitionHLOD* HLODActor : InActor->GetWorld()->GetSubsystem<UWorldPartitionHLODRuntimeSubsystem>()->GetHLODActorsForCell(RuntimeCell))
 			{
 				for (UActorComponent* ActorComponent : HLODActor->GetComponents())
 				{
diff --git a/Engine/Source/Runtime/Engine/Private/WorldPartition/HLOD/HLODRuntimeSubsystem.cpp b/Engine/Source/Runtime/Engine/Private/WorldPartition/HLOD/HLODRuntimeSubsystem.cpp
new file mode 100644
--- /dev/null
+++ b/Engine/Source/Runtime/Engine/Private/WorldPartition/HLOD/HLODRuntimeSubsystem.cpp
@@ -0,0 +1,930 @@
+// Copyright Epic Games, Inc. All Rights Reserved.
+
+#include "WorldPartition/HLOD/HLODRuntimeSubsystem.h"
+#include "Engine/Level.h"
+#include "WorldPartition/HLOD/HLODActor.h"
+#include "Engine/LevelStreaming.h"
+#include "Engine/World.h"
+#include "WorldPartition/HLOD/HLODActorDesc.h"
+#include "Materials/MaterialInterface.h"
+#include "WorldPartition/DataLayer/DataLayerManager.h"
+#include "Misc/Paths.h"
+#include "WorldPartition/WorldPartition.h"
+#include "RenderUtils.h"
+#include "WorldPartition/WorldPartitionSubsystem.h"
+#include "SceneInterface.h"
+#include "SceneView.h"
+#include "WorldPartition/WorldPartitionRuntimeHash.h"
+#include "SceneManagement.h"
+#include "Engine/Engine.h"
+#include "Engine/StaticMesh.h"
+#include "EngineModule.h"
+#include "LevelUtils.h"
+#include "Components/InstancedStaticMeshComponent.h"
+#include "Misc/FileHelper.h"
+#include "SceneViewExtension.h"
+#include "StaticMeshResources.h"
+#include "Subsystems/Subsystem.h"
+#include "UObject/UObjectIterator.h"
+#include "WorldPartition/DataLayer/WorldDataLayers.h"
+#include "WorldPartition/HLOD/HLODStats.h"
+
+#include UE_INLINE_GENERATED_CPP_BY_NAME(HLODRuntimeSubsystem)
+
+#define LOCTEXT_NAMESPACE "HLODRuntimeSubsystem"
+
+DEFINE_LOG_CATEGORY_STATIC(LogHLODRuntimeSubsystem, Log, All);
+
+static TAutoConsoleVariable<int32> CVarHLODWarmupEnabled(
+	TEXT("wp.Runtime.HLOD.WarmupEnabled"),
+	1,
+	TEXT("Enable HLOD assets warmup. Will delay unloading of cells & transition to HLODs for wp.Runtime.HLOD.WarmupNumFrames frames."));
+
+static TAutoConsoleVariable<int32> CVarHLODWarmupVT(
+	TEXT("wp.Runtime.HLOD.WarmupVT"),
+	1,
+	TEXT("Enable virtual texture warmup for HLOD assets. Requires wp.Runtime.HLOD.WarmupEnabled to be 1."));
+
+static TAutoConsoleVariable<int32> CVarHLODWarmupNanite(
+	TEXT("wp.Runtime.HLOD.WarmupNanite"),
+	1,
+	TEXT("Enable Nanite warmup for HLOD assets. Requires wp.Runtime.HLOD.WarmupEnabled to be 1."));
+
+static TAutoConsoleVariable<int32> CVarHLODWarmupNumFrames(
+	TEXT("wp.Runtime.HLOD.WarmupNumFrames"),
+	5,
+	TEXT("Delay unloading of a cell for this amount of frames to ensure HLOD assets are ready to be shown at the proper resolution. Set to 0 to force disable warmup."));
+
+static TAutoConsoleVariable<int32> CVarHLODWarmupDebugDraw(
+	TEXT("wp.Runtime.HLOD.WarmupDebugDraw"),
+	0,
+	TEXT("Draw debug display for the warmup requests"));
+
+static TAutoConsoleVariable<float> CVarHLODWarmupVTScaleFactor(
+	TEXT("wp.Runtime.HLOD.WarmupVTScaleFactor"),
+	2.0f,
+	TEXT("Scale the VT size we ask to prefetch by this factor."));
+
+static TAutoConsoleVariable<int32> CVarHLODWarmupVTSizeClamp(
+	TEXT("wp.Runtime.HLOD.WarmupVTSizeClamp"),
+	2048,
+	TEXT("Clamp VT warmup requests for safety."));
+
+static TAutoConsoleVariable<int32> CVarCullDistanceWorkSlicer(
+	TEXT("wp.Runtime.HLOD.CullDistanceWorkSlicer"),
+	60,
+	TEXT("Maximum amount of actors that can have it's CullDistance changed per frames."));
+
+static TAutoConsoleVariable<int32> CVarCullDistanceDifferenceNeeded(
+	TEXT("wp.Runtime.HLOD.CVarCullDistanceDifferenceNeeded"),
+	300,
+	TEXT("The difference needed for the CullDistance changed to be processed"));
+
+static void HLODRuntimeSubsystemCVarSinkFunction()
+{
+	for (UWorld* World : TObjectRange<UWorld>(RF_ClassDefaultObject | RF_ArchetypeObject, true, EInternalObjectFlags::Garbage))
+	{
+		if (World->WorldType == EWorldType::Game || World->WorldType == EWorldType::PIE)
+		{
+			if (UWorldPartitionHLODRuntimeSubsystem* HLODRuntimeSubsystem = World->GetSubsystem<UWorldPartitionHLODRuntimeSubsystem>())
+			{
+				HLODRuntimeSubsystem->OnCVarsChanged();
+			}
+		}
+	}
+}
+
+static FAutoConsoleVariableSink CVarHLODSink(FConsoleCommandDelegate::CreateStatic(&HLODRuntimeSubsystemCVarSinkFunction));
+
+namespace FHLODRuntimeSubsystem
+{
+    static const UWorldPartitionRuntimeCell* GetActorRuntimeCell(AActor* InActor)
+    {
+	    const ULevel* ActorLevel = InActor->GetLevel();
+	    const ULevelStreaming* LevelStreaming = FLevelUtils::FindStreamingLevel(ActorLevel);
+	    return LevelStreaming ? Cast<const UWorldPartitionRuntimeCell>(LevelStreaming->GetWorldPartitionCell()) : nullptr;
+    }
+    
+    static UWorldPartition* GetWorldPartition(AWorldPartitionHLOD* InWorldPartitionHLOD)
+    {
+	    // Alwaysloaded Cell level will have a WorldPartition
+	    if (UWorldPartition* WorldPartition = InWorldPartitionHLOD->GetLevel()->GetWorldPartition())
+	    {
+		    return WorldPartition;
+	    } // If not find it through the cell
+	    else if (const UWorldPartitionRuntimeCell* RuntimeCell = GetActorRuntimeCell(InWorldPartitionHLOD))
+	    {
+		    return RuntimeCell->GetOuterWorld()->GetWorldPartition();
+	    }
+    
+	    return nullptr;
+    }
+}
+
+
+class FHLODResourcesResidencySceneViewExtension : public FWorldSceneViewExtension
+{
+public:
+	FHLODResourcesResidencySceneViewExtension(const FAutoRegister& AutoRegister, UWorld* InWorld)
+		: FWorldSceneViewExtension(AutoRegister, InWorld)
+	{
+	}
+
+	virtual void SetupViewFamily(FSceneViewFamily& InViewFamily) override {}
+	virtual void SetupView(FSceneViewFamily& InViewFamily, FSceneView& InView) override {}
+	virtual void BeginRenderViewFamily(FSceneViewFamily& InViewFamily) override
+	{
+		GetWorld()->GetSubsystem<UWorldPartitionHLODRuntimeSubsystem>()->OnBeginRenderViews(InViewFamily);
+	}
+};
+
+
+UWorldPartitionHLODRuntimeSubsystem::UWorldPartitionHLODRuntimeSubsystem()
+	: UWorldSubsystem()
+	, bCachedShouldPerformWarmup(true)
+{
+}
+
+UWorldPartitionHLODRuntimeSubsystem::~UWorldPartitionHLODRuntimeSubsystem()
+{
+	if (TickHandle.IsValid())
+	{
+		FTSTicker::GetCoreTicker().RemoveTicker(TickHandle);
+		TickHandle.Reset();
+	}
+}
+
+void UWorldPartitionHLODRuntimeSubsystem::SetHLODAlwaysLoadedCullDistance(int32 InCullDistance)
+{
+	const float MaxDrawDistance = FMath::Max(InCullDistance, 0);
+	const bool bNeverDistanceCull = InCullDistance <= 0;
+
+	//If the WorkSlicer is set to process 0 items, revert back to processing everything at once
+	if (CVarCullDistanceWorkSlicer.GetValueOnAnyThread() > 0)
+	{
+		int32 CullDiff = FMath::Abs(LastSetCullDistance - InCullDistance);
+		if (!OperationQueue.IsEmpty() || CullDiff < CVarCullDistanceDifferenceNeeded.GetValueOnAnyThread())
+		{
+			return;
+		}
+		LastSetCullDistance = InCullDistance;
+
+		OperationQueue.Reserve(AlwaysLoadedHLODActors.Num());
+		for (AWorldPartitionHLOD* HLODActor : AlwaysLoadedHLODActors)
+		{
+			FDrawDistanceQueue Item = FDrawDistanceQueue();
+			Item.DrawDistance = MaxDrawDistance;
+			Item.HLODActor = HLODActor;
+			OperationQueue.Add(Item);
+		}
+
+		if (!TickHandle.IsValid())
+		{
+			TickHandle = FTSTicker::GetCoreTicker().AddTicker(FTickerDelegate::CreateUObject(this, &UWorldPartitionHLODRuntimeSubsystem::Tick), 0.0f);
+		}
+	}
+	else
+	{
+		for (AWorldPartitionHLOD* HLODActor : AlwaysLoadedHLODActors)
+		{
+			HLODActor->ForEachComponent<UPrimitiveComponent>(false, [MaxDrawDistance, bNeverDistanceCull](UPrimitiveComponent* HLODComponent)
+			{
+				HLODComponent->bNeverDistanceCull = bNeverDistanceCull;
+				HLODComponent->SetCachedMaxDrawDistance(MaxDrawDistance);
+				// Enable per instance culling to avoid ISM components dissapearing as a whole
+				if (UInstancedStaticMeshComponent* ISMComponent = Cast<UInstancedStaticMeshComponent>(HLODComponent))
+				{
+					ISMComponent->SetCullDistances(0, MaxDrawDistance);
+				}
+			});
+		}
+	}
+	
+}
+
+void UWorldPartitionHLODRuntimeSubsystem::OnWorldBeginPlay(UWorld& InWorld)
+{
+	Super::OnWorldBeginPlay(InWorld);
+
+	check(InWorld.IsGameWorld());
+
+	AlwaysLoadedHLODActors.Reset();
+
+	// Cache always loaded HLODs
+	for (AActor* Actor : InWorld.PersistentLevel->Actors)
+	{
+		if (AWorldPartitionHLOD* HLODActor = Cast<AWorldPartitionHLOD>(Actor))
+		{
+			AlwaysLoadedHLODActors.Add(HLODActor);
+		}
+	}
+
+}
+
+bool UWorldPartitionHLODRuntimeSubsystem::Tick(float DeltaTime)
+{
+	//If the CVarCullDistanceAllowWorkSlicer is false we will never add this function to FTSTicker
+	if (OperationQueue.Num() > 0)
+	{
+		int32 Amount = FMath::Min(OperationQueue.Num(), CVarCullDistanceWorkSlicer.GetValueOnAnyThread());	
+		for (int32 i = OperationQueue.Num() - Amount; i < OperationQueue.Num(); i++)
+		{
+			const float MaxDrawDistance = OperationQueue[i].DrawDistance;
+			TObjectPtr<AWorldPartitionHLOD> HLODActor = OperationQueue[i].HLODActor;
+			if (HLODActor != nullptr)
+			{
+				HLODActor->ForEachComponent<UPrimitiveComponent>(false, [MaxDrawDistance](UPrimitiveComponent* HLODComponent)
+				{
+					HLODComponent->bNeverDistanceCull = false;
+					HLODComponent->SetCachedMaxDrawDistance(MaxDrawDistance);
+					// Enable per instance culling to avoid ISM components dissapearing as a whole
+					if (UInstancedStaticMeshComponent* ISMComponent = Cast<UInstancedStaticMeshComponent>(HLODComponent))
+					{
+						ISMComponent->SetCullDistances(0, MaxDrawDistance);
+					}
+				});
+			}
+			
+		}
+		OperationQueue.RemoveAt(OperationQueue.Num() - Amount, Amount);
+	}
+	return true;
+}
+
+bool UWorldPartitionHLODRuntimeSubsystem::WorldPartitionHLODEnabled = true;
+
+FAutoConsoleCommand UWorldPartitionHLODRuntimeSubsystem::EnableHLODCommand(
+	TEXT("wp.Runtime.HLOD"),
+	TEXT("Turn on/off loading & rendering of world partition HLODs."),
+	FConsoleCommandWithArgsDelegate::CreateLambda([](const TArray<FString>& Args)
+	{
+		UWorldPartitionHLODRuntimeSubsystem::WorldPartitionHLODEnabled = (Args.Num() != 1) || (Args[0] != TEXT("0"));
+		for (const FWorldContext& Context : GEngine->GetWorldContexts())
+		{
+			UWorld* World = Context.World();
+			if (World && World->IsGameWorld())
+			{
+				if (UWorldPartitionHLODRuntimeSubsystem* HLODSubSystem = World->GetSubsystem<UWorldPartitionHLODRuntimeSubsystem>()) 
+				{
+					for (const auto& KeyValuePair : HLODSubSystem->WorldPartitionsHLODRuntimeData)
+					{
+						for (const auto& CellHLODMapping : KeyValuePair.Value.CellsData)
+						{
+							const FCellData& CellData = CellHLODMapping.Value;
+							bool bIsHLODVisible = UWorldPartitionHLODRuntimeSubsystem::WorldPartitionHLODEnabled && !CellData.bIsCellVisible;
+							for (AWorldPartitionHLOD* HLODActor : CellData.LoadedHLODs)
+							{
+								HLODActor->SetVisibility(bIsHLODVisible);
+							}
+						}
+					}
+				}
+			}
+		}
+	})
+);
+
+bool UWorldPartitionHLODRuntimeSubsystem::IsHLODEnabled()
+{
+	return UWorldPartitionHLODRuntimeSubsystem::WorldPartitionHLODEnabled;
+}
+
+bool UWorldPartitionHLODRuntimeSubsystem::DoesSupportWorldType(const EWorldType::Type WorldType) const
+{
+	return WorldType == EWorldType::Game || WorldType == EWorldType::PIE;
+}
+
+void UWorldPartitionHLODRuntimeSubsystem::Initialize(FSubsystemCollectionBase& Collection)
+{
+	// Ensure the WorldPartitionSubsystem gets created before the HLODRuntimeSubsystem
+	Collection.InitializeDependency<UWorldPartitionSubsystem>();
+
+	Super::Initialize(Collection);
+
+	UWorld* World = GetWorld();
+	check(World->IsGameWorld());
+
+	GetWorld()->OnWorldPartitionInitialized().AddUObject(this, &UWorldPartitionHLODRuntimeSubsystem::OnWorldPartitionInitialized);
+	GetWorld()->OnWorldPartitionUninitialized().AddUObject(this, &UWorldPartitionHLODRuntimeSubsystem::OnWorldPartitionUninitialized);
+
+	bCachedShouldPerformWarmup = ShouldPerformWarmup();
+
+	SceneViewExtension = FSceneViewExtensions::NewExtension<FHLODResourcesResidencySceneViewExtension>(World);
+}
+
+void UWorldPartitionHLODRuntimeSubsystem::Deinitialize()
+{
+	Super::Deinitialize();
+
+	GetWorld()->OnWorldPartitionInitialized().RemoveAll(this);
+	GetWorld()->OnWorldPartitionUninitialized().RemoveAll(this);
+}
+
+void UWorldPartitionHLODRuntimeSubsystem::OnWorldPartitionInitialized(UWorldPartition* InWorldPartition)
+{
+	if (InWorldPartition && InWorldPartition->IsStreamingEnabled())
+	{
+		check(!WorldPartitionsHLODRuntimeData.Contains(InWorldPartition));
+
+		FWorldPartitionHLODRuntimeData& WorldPartitionHLODRuntimeData = WorldPartitionsHLODRuntimeData.Add(InWorldPartition, FWorldPartitionHLODRuntimeData());
+	
+		// Build cell to HLOD mapping
+		if (InWorldPartition && InWorldPartition->RuntimeHash)
+		{
+			InWorldPartition->RuntimeHash->ForEachStreamingCells([&WorldPartitionHLODRuntimeData](const UWorldPartitionRuntimeCell* Cell)
+			{
+				UE_LOG(LogHLODRuntimeSubsystem, Verbose, TEXT("Registering cell %s - %s"), *Cell->GetGuid().ToString(), *Cell->GetDebugName());
+
+				WorldPartitionHLODRuntimeData.CellsData.Emplace(Cell->GetGuid());
+				return true;
+			});
+		}
+	}
+}
+
+void UWorldPartitionHLODRuntimeSubsystem::OnWorldPartitionUninitialized(UWorldPartition* InWorldPartition)
+{
+	if (InWorldPartition && InWorldPartition->IsStreamingEnabled())
+	{
+		check(WorldPartitionsHLODRuntimeData.Contains(InWorldPartition));
+
+#if !NO_LOGGING
+		FWorldPartitionHLODRuntimeData& WorldPartitionHLODRuntimeData = WorldPartitionsHLODRuntimeData.FindChecked(InWorldPartition);
+		InWorldPartition->RuntimeHash->ForEachStreamingCells([&WorldPartitionHLODRuntimeData](const UWorldPartitionRuntimeCell* Cell)
+		{
+			UE_LOG(LogHLODRuntimeSubsystem, Verbose, TEXT("Unregistering cell %s - %s"), *Cell->GetGuid().ToString(), *Cell->GetDebugName());
+			return true;
+		});
+#endif
+
+		WorldPartitionsHLODRuntimeData.Remove(InWorldPartition);
+	}
+}
+
+void UWorldPartitionHLODRuntimeSubsystem::OnExternalStreamingObjectInjected(URuntimeHashExternalStreamingObjectBase* ExternalStreamingObject)
+{
+	UWorldPartition* OwnerPartition = ExternalStreamingObject->GetOuterWorld()->GetWorldPartition();
+	FWorldPartitionHLODRuntimeData* WorldPartitionHLODRuntimeData = WorldPartitionsHLODRuntimeData.Find(OwnerPartition);
+	if (WorldPartitionHLODRuntimeData)
+	{
+		ExternalStreamingObject->ForEachStreamingCells([WorldPartitionHLODRuntimeData](const UWorldPartitionRuntimeCell& Cell)
+		{
+			UE_LOG(LogHLODRuntimeSubsystem, Verbose, TEXT("Registering external cell %s - %s"), *Cell.GetGuid().ToString(), *Cell.GetDebugName());
+			WorldPartitionHLODRuntimeData->CellsData.Emplace(Cell.GetGuid());
+			return true;
+		});
+	}
+}
+
+void UWorldPartitionHLODRuntimeSubsystem::OnExternalStreamingObjectRemoved(URuntimeHashExternalStreamingObjectBase* ExternalStreamingObject)
+{
+	UWorldPartition* OwnerPartition = ExternalStreamingObject->GetOuterWorld()->GetWorldPartition();
+	FWorldPartitionHLODRuntimeData* WorldPartitionHLODRuntimeData = WorldPartitionsHLODRuntimeData.Find(OwnerPartition);
+	if (WorldPartitionHLODRuntimeData)
+	{
+		ExternalStreamingObject->ForEachStreamingCells([WorldPartitionHLODRuntimeData](const UWorldPartitionRuntimeCell& Cell)
+		{
+			UE_LOG(LogHLODRuntimeSubsystem, Verbose, TEXT("Unregistering external cell %s - %s"), *Cell.GetGuid().ToString(), *Cell.GetDebugName());
+			WorldPartitionHLODRuntimeData->CellsData.Remove(Cell.GetGuid());
+			return true;
+		});
+	}
+}
+
+const UWorldPartitionHLODRuntimeSubsystem::FCellData* UWorldPartitionHLODRuntimeSubsystem::GetCellData(const UWorldPartitionRuntimeCell* InCell) const
+{
+	return const_cast<UWorldPartitionHLODRuntimeSubsystem*>(this)->GetCellData(InCell);
+}
+
+UWorldPartitionHLODRuntimeSubsystem::FCellData* UWorldPartitionHLODRuntimeSubsystem::GetCellData(const UWorldPartitionRuntimeCell* InCell)
+{
+	const UWorldPartition* WorldPartition = InCell->GetOuterWorld()->GetWorldPartition();
+	if (WorldPartition)
+	{
+		FWorldPartitionHLODRuntimeData* WorldPartitionHLODRuntimeData = WorldPartitionsHLODRuntimeData.Find(WorldPartition);
+		if (WorldPartitionHLODRuntimeData)
+		{
+			check(WorldPartition->IsStreamingEnabled());
+			return WorldPartitionHLODRuntimeData->CellsData.Find(InCell->GetGuid());
+		}
+	}
+
+	return nullptr;
+}
+
+UWorldPartitionHLODRuntimeSubsystem::FCellData* UWorldPartitionHLODRuntimeSubsystem::GetCellData(AWorldPartitionHLOD* InWorldPartitionHLOD)
+{
+	const UWorldPartition* WorldPartition = FHLODRuntimeSubsystem::GetWorldPartition(InWorldPartitionHLOD);
+	if (WorldPartition)
+	{
+		FWorldPartitionHLODRuntimeData* WorldPartitionHLODRuntimeData = WorldPartitionsHLODRuntimeData.Find(WorldPartition);
+		if (WorldPartitionHLODRuntimeData)
+		{
+			check(WorldPartition->IsStreamingEnabled());
+			const FGuid CellGuid = InWorldPartitionHLOD->GetSourceCellGuid();
+			return WorldPartitionHLODRuntimeData->CellsData.Find(CellGuid);
+		}
+	}
+
+	return nullptr;		
+}
+
+const TArray<AWorldPartitionHLOD*>& UWorldPartitionHLODRuntimeSubsystem::GetHLODActorsForCell(const UWorldPartitionRuntimeCell* InCell) const
+{
+	if (const FCellData* CellData = GetCellData(InCell))
+	{
+		return CellData->LoadedHLODs;
+	}
+
+	// No HLOD found for the given cell, return a dummy array
+	static const TArray<AWorldPartitionHLOD*> DummyArray;
+	return DummyArray;
+}
+
+void UWorldPartitionHLODRuntimeSubsystem::RegisterHLODActor(AWorldPartitionHLOD* InWorldPartitionHLOD)
+{
+	TRACE_CPUPROFILER_EVENT_SCOPE(UWorldPartitionHLODRuntimeSubsystem::RegisterHLODActor);
+	
+	if (FCellData* CellData = GetCellData(InWorldPartitionHLOD))
+	{
+		UE_LOG(LogHLODRuntimeSubsystem, Verbose, TEXT("Registering HLOD %s for cell %s"), *InWorldPartitionHLOD->GetActorNameOrLabel(), *InWorldPartitionHLOD->GetSourceCellGuid().ToString());
+
+		CellData->LoadedHLODs.Add(InWorldPartitionHLOD);
+		InWorldPartitionHLOD->SetVisibility(UWorldPartitionHLODRuntimeSubsystem::WorldPartitionHLODEnabled && !CellData->bIsCellVisible);
+	}
+	else
+	{
+		UE_LOG(LogHLODRuntimeSubsystem, Verbose, TEXT("Found HLOD %s referencing nonexistent cell '%s'"), *InWorldPartitionHLOD->GetActorNameOrLabel(), *InWorldPartitionHLOD->GetSourceCellGuid().ToString());
+		InWorldPartitionHLOD->SetVisibility(false);
+
+#if WITH_EDITOR
+		OutdatedHLODActors.Add(InWorldPartitionHLOD);
+#endif
+	}
+
+	HLODActorsToWarmup.Remove(InWorldPartitionHLOD);
+	HLODActorRegisteredEvent.Broadcast(InWorldPartitionHLOD);
+}
+
+void UWorldPartitionHLODRuntimeSubsystem::UnregisterHLODActor(AWorldPartitionHLOD* InWorldPartitionHLOD)
+{
+	TRACE_CPUPROFILER_EVENT_SCOPE(UWorldPartitionHLODRuntimeSubsystem::UnregisterHLODActor);
+
+	if (FCellData* CellData = GetCellData(InWorldPartitionHLOD))
+	{
+		UE_LOG(LogHLODRuntimeSubsystem, Verbose, TEXT("Unregistering HLOD %s for cell %s"), *InWorldPartitionHLOD->GetActorNameOrLabel(), *InWorldPartitionHLOD->GetSourceCellGuid().ToString());
+
+		verify(CellData->LoadedHLODs.Remove(InWorldPartitionHLOD));
+	}
+	else
+	{
+#if WITH_EDITOR
+		OutdatedHLODActors.Remove(InWorldPartitionHLOD);
+#endif
+	}
+
+	HLODActorUnregisteredEvent.Broadcast(InWorldPartitionHLOD);
+}
+
+void UWorldPartitionHLODRuntimeSubsystem::OnCellShown(const UWorldPartitionRuntimeCell* InCell)
+{
+	if (FCellData* CellData = GetCellData(InCell))
+	{
+		CellData->bIsCellVisible = true;
+
+		if (!CellData->LoadedHLODs.IsEmpty())
+		{
+			UE_LOG(LogHLODRuntimeSubsystem, Verbose, TEXT("Cell shown - %s - hiding %d HLOD actors"), *InCell->GetGuid().ToString(), CellData->LoadedHLODs.Num());
+
+			for (AWorldPartitionHLOD* HLODActor : CellData->LoadedHLODs)
+			{
+				UE_LOG(LogHLODRuntimeSubsystem, Verbose, TEXT("\t\t* %s"), *HLODActor->GetActorNameOrLabel());
+				HLODActor->SetVisibility(false);
+			}
+		}
+	}
+}
+
+void UWorldPartitionHLODRuntimeSubsystem::OnCellHidden(const UWorldPartitionRuntimeCell* InCell)
+{
+	if (FCellData* CellData = GetCellData(InCell))
+	{
+		CellData->bIsCellVisible = false;
+
+		if (!CellData->LoadedHLODs.IsEmpty())
+		{
+			UE_LOG(LogHLODRuntimeSubsystem, Verbose, TEXT("Cell hidden - %s - showing %d HLOD actors"), *InCell->GetGuid().ToString(), CellData->LoadedHLODs.Num());
+
+			for (AWorldPartitionHLOD* HLODActor : CellData->LoadedHLODs)
+			{
+				UE_LOG(LogHLODRuntimeSubsystem, Verbose, TEXT("\t\t* %s"), *HLODActor->GetActorNameOrLabel());
+				HLODActor->SetVisibility(UWorldPartitionHLODRuntimeSubsystem::WorldPartitionHLODEnabled);
+				HLODActorsToWarmup.Remove(HLODActor);
+			}
+		}
+	}
+}
+
+static void PrepareVTRequests(TMap<UMaterialInterface*, float>& InOutVTRequests, UStaticMeshComponent* InStaticMeshComponent, float InPixelSize)
+{
+	float PixelSize = InPixelSize;
+
+	// Assume the texture is wrapped around the object, so the screen size is actually less than the resolution we require.
+	PixelSize *= CVarHLODWarmupVTScaleFactor.GetValueOnAnyThread();
+
+	// Clamp for safety
+	PixelSize = FMath::Min(PixelSize, CVarHLODWarmupVTSizeClamp.GetValueOnAnyThread());
+
+	for (UMaterialInterface* MaterialInterface : InStaticMeshComponent->GetMaterials())
+	{
+		// We have a VT we'd like to prefetch, add or update a request in our request map.
+		// If the texture was already requested by another component, fetch the highest required resolution only.
+		float& CurrentMaxPixel = InOutVTRequests.FindOrAdd(MaterialInterface);
+		CurrentMaxPixel = FMath::Max(CurrentMaxPixel, PixelSize);
+	}
+}
+
+static void PrepareNaniteRequests(TSet<Nanite::FResources*>& InOutNaniteRequests, UStaticMeshComponent* InStaticMeshComponent)
+{
+	UStaticMesh* StaticMesh = InStaticMeshComponent->GetStaticMesh();
+	if (StaticMesh && StaticMesh->HasValidNaniteData())
+	{
+		InOutNaniteRequests.Add(StaticMesh->GetRenderData()->NaniteResourcesPtr.Get());
+	}
+}
+
+bool UWorldPartitionHLODRuntimeSubsystem::PrepareToWarmup(const UWorldPartitionRuntimeCell* InCell, AWorldPartitionHLOD* InHLODActor)
+{
+	TRACE_CPUPROFILER_EVENT_SCOPE(UWorldPartitionHLODRuntimeSubsystem::PrepareToWarmup)
+
+	bool bHLODActorNeedsWarmUp = false;
+
+	if (InHLODActor->DoesRequireWarmup())
+	{
+		FWorldPartitionHLODWarmupState& WarmupState = HLODActorsToWarmup.FindOrAdd(InHLODActor);
+
+		// In case a previous request to unload was aborted and the cell never unloaded... assume warmup requests are expired after a given amount of frames.
+		const uint32 WarmupExpiredFrames = 30;
+		const uint32 CurrentFrameNumber = GFrameNumberRenderThread;
+
+		// Trigger warmup on the first request to unload, or if a warmup request expired
+		const bool bInitiateWarmup = WarmupState.WarmupEndFrame == INDEX_NONE || CurrentFrameNumber > (WarmupState.WarmupEndFrame + WarmupExpiredFrames);
+		const bool bWarmupCompleted = !bInitiateWarmup && CurrentFrameNumber >= WarmupState.WarmupEndFrame;
+		
+		if (bInitiateWarmup)
+		{
+			// Warmup will be triggered in the next BeginRenderView() call, at which point the frame number will have been incremented.
+			WarmupState.WarmupStartFrame = CurrentFrameNumber + 1;
+			WarmupState.WarmupEndFrame = WarmupState.WarmupStartFrame + CVarHLODWarmupNumFrames.GetValueOnGameThread();
+			WarmupState.Location = InCell->GetCellBounds().GetCenter();
+		}
+
+		bHLODActorNeedsWarmUp = !bWarmupCompleted;
+	}
+
+	return bHLODActorNeedsWarmUp;
+}
+
+void UWorldPartitionHLODRuntimeSubsystem::OnCVarsChanged()
+{
+	bCachedShouldPerformWarmup = ShouldPerformWarmup();
+}
+
+bool UWorldPartitionHLODRuntimeSubsystem::ShouldPerformWarmup() const
+{
+	// Test if warmup is disabled globally.
+	const bool bWarmupEnabled = CVarHLODWarmupEnabled.GetValueOnGameThread() != 0;
+	if (!bWarmupEnabled)
+	{
+		return false;
+	}
+
+	// If warmup num of frames is invalid, no warmup needed
+	if (CVarHLODWarmupNumFrames.GetValueOnGameThread() <= 0)
+	{
+		return false;
+	}
+
+	// If warmup num of frames is invalid, no warmup needed
+	const EShaderPlatform ShaderPlatform = GShaderPlatformForFeatureLevel[GMaxRHIFeatureLevel];
+	const bool bNaniteEnabled = UseNanite(ShaderPlatform);
+	const bool bVirtualTextureEnabled = UseVirtualTexturing(ShaderPlatform);
+	const bool bWarmupNanite = CVarHLODWarmupNanite.GetValueOnGameThread() != 0;
+	const bool bWarmupVT = CVarHLODWarmupVT.GetValueOnGameThread() != 0;
+	const bool bWarmupNeeded = (bNaniteEnabled && bWarmupNanite) || (bVirtualTextureEnabled && bWarmupVT);
+	if (!bWarmupNeeded)
+	{
+		return false;
+	}
+
+	// If we're running a dedicated server, no warmup needed
+	const bool bIsDedicatedServer = GetWorld()->GetNetMode() == NM_DedicatedServer;
+	if (bIsDedicatedServer)
+	{
+		return false;
+	}
+
+	return true;
+}
+
+bool UWorldPartitionHLODRuntimeSubsystem::ShouldPerformWarmupForCell(const UWorldPartitionRuntimeCell* InCell) const
+{
+	if (!bCachedShouldPerformWarmup)
+	{
+		return false;
+	}
+
+	const UWorld* World = GetWorld();
+
+	// Blocking loading shouldn't trigger warmup
+	const bool bIsInBlockingLoading = World->GetIsInBlockTillLevelStreamingCompleted();
+	if (bIsInBlockingLoading)
+	{
+		return false;
+	}
+
+	// If processing for this cell in Add/RemoveFromWorld() has already started, it's too late for warmup
+	const ULevel* CellLevel = InCell->GetLevel();
+	const bool bCurrentlyProcessingLevel = CellLevel == World->GetCurrentLevelPendingVisibility() ||
+										   CellLevel == World->GetCurrentLevelPendingInvisibility();
+	if (bCurrentlyProcessingLevel)
+	{
+		return false;
+	}
+
+	return true;
+}
+
+bool UWorldPartitionHLODRuntimeSubsystem::CanMakeVisible(const UWorldPartitionRuntimeCell* InCell)
+{
+	TRACE_CPUPROFILER_EVENT_SCOPE(UWorldPartitionHLODRuntimeSubsystem::CanMakeVisible)
+
+	if (!ShouldPerformWarmupForCell(InCell))
+	{
+		return true;
+	}
+
+	bool bCanMakeVisible = true;
+
+	// Prevent cells containing HLODs actors from being made visible until warmup has been performed
+	if (InCell->GetIsHLOD() && InCell->GetLevel())
+	{
+		for (AActor* Actor : InCell->GetLevel()->Actors)
+		{
+			if (AWorldPartitionHLOD* HLODActor = Cast<AWorldPartitionHLOD>(Actor))
+			{
+				const bool bHLODActorNeedsWarmup = PrepareToWarmup(InCell, HLODActor);
+				bCanMakeVisible &= !bHLODActorNeedsWarmup;
+			}
+		}
+	}
+
+	return bCanMakeVisible;
+}
+
+bool UWorldPartitionHLODRuntimeSubsystem::CanMakeInvisible(const UWorldPartitionRuntimeCell* InCell)
+{
+	TRACE_CPUPROFILER_EVENT_SCOPE(UWorldPartitionHLODRuntimeSubsystem::CanMakeInvisible)
+
+	if (!ShouldPerformWarmupForCell(InCell))
+	{
+		return true;
+	}
+
+	bool bCanMakeInvisible = true;
+
+	// Prevent cells which have an HLOD representation from being hidden until their matching HLOD actors have been warmed up
+	if (FCellData* CellData = GetCellData(InCell))
+	{
+		for (AWorldPartitionHLOD* HLODActor : CellData->LoadedHLODs)
+		{
+			const bool bHLODActorNeedsWarmup = PrepareToWarmup(InCell, HLODActor);
+			bCanMakeInvisible &= !bHLODActorNeedsWarmup;
+		}
+	}
+
+	return bCanMakeInvisible;
+}
+
+static float IsInView(const FVector& BoundsOrigin, const FVector& BoundsExtent, const FSceneViewFamily& InViewFamily, bool bComputeScreenSize, float& MaxScreenSizePixels)
+{
+	MaxScreenSizePixels = 0;
+
+	// Estimate the highest screen pixel size of this component in the provided views
+	for (const FSceneView* View : InViewFamily.Views)
+	{
+		// Make sure the HLOD actor we're about to show is actually in the frustum
+		if (View->ViewFrustum.IntersectBox(BoundsOrigin, BoundsExtent))
+		{
+			if (bComputeScreenSize)
+			{
+				float ScreenDiameter = ComputeBoundsScreenSize(BoundsOrigin, BoundsExtent.Size(), *View);
+				float ScreenSizePixels = ScreenDiameter * View->ViewMatrices.GetScreenScale() * 2.0f;
+
+				MaxScreenSizePixels = FMath::Max(MaxScreenSizePixels, ScreenSizePixels);
+			}
+			else
+			{
+				return true;
+			}
+		}
+	}
+
+	return MaxScreenSizePixels > 0;
+}
+
+static void MakeHLODRenderResourcesResident(TMap<UMaterialInterface*, float>& VTRequests, TSet<Nanite::FResources*>& NaniteRequests, const FSceneViewFamily& InViewFamily)
+{
+	if (!VTRequests.IsEmpty() || !NaniteRequests.IsEmpty())
+	{
+		ENQUEUE_RENDER_COMMAND(MakeHLODRenderResourcesResident)(
+			[VTRequests = MoveTemp(VTRequests), NaniteRequests = MoveTemp(NaniteRequests), FeatureLevel = InViewFamily.GetFeatureLevel()](FRHICommandListImmediate& RHICmdList)
+		{
+			for (const TPair<UMaterialInterface*, float>& VTRequest : VTRequests)
+			{
+				UMaterialInterface* Material = VTRequest.Key;
+				FMaterialRenderProxy* MaterialRenderProxy = Material->GetRenderProxy();
+
+				GetRendererModule().RequestVirtualTextureTiles(MaterialRenderProxy, FVector2D(VTRequest.Value, VTRequest.Value), FeatureLevel);
+			}
+
+			const uint32 NumFramesBeforeRender = CVarHLODWarmupNumFrames.GetValueOnRenderThread();
+			for (const Nanite::FResources* Resource : NaniteRequests)
+			{
+				GetRendererModule().PrefetchNaniteResource(Resource, NumFramesBeforeRender);
+			}
+		});
+	}
+}
+
+void UWorldPartitionHLODRuntimeSubsystem::OnBeginRenderViews(const FSceneViewFamily& InViewFamily)
+{
+	TRACE_CPUPROFILER_EVENT_SCOPE(UWorldPartitionHLODRuntimeSubsystem::OnBeginRenderViews)
+
+	TMap<UMaterialInterface*, float> VTRequests;
+	TSet<Nanite::FResources*> NaniteRequests;
+
+	const bool bWarmupNanite = CVarHLODWarmupNanite.GetValueOnGameThread() != 0;
+	const bool bWarmupVT = CVarHLODWarmupVT.GetValueOnGameThread() != 0;
+
+	for (FHLODWarmupStateMap::TIterator HLODActorWarmupStateIt(HLODActorsToWarmup); HLODActorWarmupStateIt; ++HLODActorWarmupStateIt)
+	{
+		const FObjectKey& HLODActorObjectKey = HLODActorWarmupStateIt.Key();
+		const FWorldPartitionHLODWarmupState& HLODWarmupState = HLODActorWarmupStateIt.Value();
+		const AWorldPartitionHLOD* HLODActor = Cast<AWorldPartitionHLOD>(HLODActorObjectKey.ResolveObjectPtr());
+		
+		if (HLODActor && InViewFamily.FrameNumber < HLODWarmupState.WarmupEndFrame)
+		{
+			HLODActor->ForEachComponent<UStaticMeshComponent>(false, [&](UStaticMeshComponent* SMC)
+			{
+				// Assume ISM HLOD don't need warmup, as they are actually found in the source level
+				if (SMC->IsA<UInstancedStaticMeshComponent>() || !SMC->GetStaticMesh())
+				{
+					return;
+				}
+
+				// Retrieve this component's bound - we must support getting the bounds before the component is even registered.
+				FVector BoundsOrigin = SMC->GetStaticMesh()->GetBounds().Origin + SMC->GetRelativeLocation();
+				FVector BoundsExtent = SMC->GetStaticMesh()->GetBounds().BoxExtent;
+				
+				float ScreenSizePixels = 0;
+				if (IsInView(BoundsOrigin, BoundsExtent, InViewFamily, bWarmupVT, ScreenSizePixels))
+				{
+					if (bWarmupVT)
+					{
+						PrepareVTRequests(VTRequests, SMC, ScreenSizePixels);
+					}
+
+					if (bWarmupNanite)
+					{
+						// Only issue Nanite requests on the first warmup frame
+						if (HLODWarmupState.WarmupStartFrame == InViewFamily.FrameNumber)
+						{
+							PrepareNaniteRequests(NaniteRequests, SMC);
+						}
+					}
+
+#if ENABLE_DRAW_DEBUG
+					if (CVarHLODWarmupDebugDraw.GetValueOnAnyThread())
+					{
+						DrawDebugBox(HLODActor->GetWorld(), BoundsOrigin, BoundsExtent, FColor::Yellow, /*bPersistentLine*/ false, /*Lifetime*/ 1.0f);
+					}
+#endif
+				}
+			});
+		}
+
+		if (!HLODActor)
+		{
+			// HLOD actor has been unloaded, stop tracking it's warmup state
+			HLODActorWarmupStateIt.RemoveCurrent();
+		}
+	}
+
+	MakeHLODRenderResourcesResident(VTRequests, NaniteRequests, InViewFamily);
+}
+
+#if WITH_EDITOR
+
+bool UWorldPartitionHLODRuntimeSubsystem::WriteHLODStatsCSV(UWorld* InWorld, const FString& InFilename)
+{
+	UWorldPartition* WorldPartition = InWorld ? InWorld->GetWorldPartition() : nullptr;
+	if (!WorldPartition)
+	{
+		return false;
+	}
+	
+	typedef TFunction<FString(const FHLODActorDesc&)> FGetStatFunc;
+
+	auto GetHLODStat = [](FName InStatName)
+	{
+		return TPair<FName, FGetStatFunc>(InStatName, [InStatName](const FHLODActorDesc& InActorDesc)
+		{
+			return FString::Printf(TEXT("%lld"), InActorDesc.GetStat(InStatName));
+		});
+	};
+
+	const UDataLayerManager* DataLayerManager = WorldPartition->GetDataLayerManager();
+
+	auto GetDataLayerShortName = [DataLayerManager](FName DataLayerInstanceName)
+	{ 
+		const UDataLayerInstance* DataLayerInstance = DataLayerManager ? DataLayerManager->GetDataLayerInstance(DataLayerInstanceName) : nullptr;
+		return DataLayerInstance ? DataLayerInstance->GetDataLayerShortName() : DataLayerInstanceName.ToString();
+	};
+
+	TArray<TPair<FName, FGetStatFunc>> StatsToWrite =
+	{	
+		{ "WorldPackage",		[InWorld](const FHLODActorDesc& InActorDesc) { return InWorld->GetPackage()->GetName(); } },
+		{ "Name",				[](const FHLODActorDesc& InActorDesc) { return InActorDesc.GetActorLabel().ToString(); } },
+		{ "HLODLayer",			[](const FHLODActorDesc& InActorDesc) { return InActorDesc.GetSourceHLODLayerName().ToString(); } },
+		{ "SpatiallyLoaded",	[](const FHLODActorDesc& InActorDesc) { return InActorDesc.GetIsSpatiallyLoaded() ? TEXT("true") : TEXT("false"); } },
+		{ "DataLayers",			[&GetDataLayerShortName](const FHLODActorDesc& InActorDesc) { return FString::JoinBy(InActorDesc.GetDataLayerInstanceNames(), TEXT(" | "), GetDataLayerShortName); } },
+
+		GetHLODStat(FWorldPartitionHLODStats::InputActorCount),
+		GetHLODStat(FWorldPartitionHLODStats::InputTriangleCount),
+		GetHLODStat(FWorldPartitionHLODStats::InputVertexCount),
+
+		GetHLODStat(FWorldPartitionHLODStats::MeshInstanceCount),
+		GetHLODStat(FWorldPartitionHLODStats::MeshNaniteTriangleCount),
+		GetHLODStat(FWorldPartitionHLODStats::MeshNaniteVertexCount),
+		GetHLODStat(FWorldPartitionHLODStats::MeshTriangleCount),
+		GetHLODStat(FWorldPartitionHLODStats::MeshVertexCount),
+		GetHLODStat(FWorldPartitionHLODStats::MeshUVChannelCount),
+
+		GetHLODStat(FWorldPartitionHLODStats::MaterialBaseColorTextureSize),
+		GetHLODStat(FWorldPartitionHLODStats::MaterialNormalTextureSize),
+		GetHLODStat(FWorldPartitionHLODStats::MaterialEmissiveTextureSize),
+		GetHLODStat(FWorldPartitionHLODStats::MaterialMetallicTextureSize),
+		GetHLODStat(FWorldPartitionHLODStats::MaterialRoughnessTextureSize),
+		GetHLODStat(FWorldPartitionHLODStats::MaterialSpecularTextureSize),
+		
+		GetHLODStat(FWorldPartitionHLODStats::MemoryMeshResourceSizeBytes),
+		GetHLODStat(FWorldPartitionHLODStats::MemoryTexturesResourceSizeBytes),
+		GetHLODStat(FWorldPartitionHLODStats::MemoryDiskSizeBytes),
+		
+		GetHLODStat(FWorldPartitionHLODStats::BuildTimeLoadMilliseconds),
+		GetHLODStat(FWorldPartitionHLODStats::BuildTimeBuildMilliseconds),
+		GetHLODStat(FWorldPartitionHLODStats::BuildTimeTotalMilliseconds)				
+	};
+
+	FStringOutputDevice Output;
+
+	// Write header if file doesn't exist
+	if (!IFileManager::Get().FileExists(*InFilename))
+	{
+		const FString StatHeader = FString::JoinBy(StatsToWrite, TEXT(","), [](const TPair<FName, FGetStatFunc>& Pair) { return Pair.Key.ToString(); });
+		Output.Logf(TEXT("%s" LINE_TERMINATOR_ANSI), *StatHeader);
+	}
+
+	// Write one line per HLOD actor desc
+	for (FActorDescContainerCollection::TIterator<AWorldPartitionHLOD> HLODIterator(WorldPartition); HLODIterator; ++HLODIterator)
+	{
+		const FString StatLine = FString::JoinBy(StatsToWrite, TEXT(","), [&HLODIterator](const TPair<FName, FGetStatFunc>& Pair) { return Pair.Value(**HLODIterator); });
+		Output.Logf(TEXT("%s" LINE_TERMINATOR_ANSI), *StatLine);
+	}
+
+	// Write to file
+	return FFileHelper::SaveStringToFile(Output, *InFilename, FFileHelper::EEncodingOptions::AutoDetect, &IFileManager::Get(), EFileWrite::FILEWRITE_Append);
+}
+
+FAutoConsoleCommand HLODDumpStats(
+	TEXT("wp.Editor.HLOD.DumpStats"),
+	TEXT("Write various HLOD stats to a CSV formatted file."),
+	FConsoleCommandWithArgsDelegate::CreateLambda([](const TArray<FString>& Args)
+	{
+		const FString HLODStatsOutputFilename = FPaths::ProjectLogDir() / TEXT("WorldPartition") / FString::Printf(TEXT("HLODStats-%08x-%s.csv"), FPlatformProcess::GetCurrentProcessId(), *FDateTime::Now().ToString());
+
+		for (const FWorldContext& Context : GEngine->GetWorldContexts())
+		{
+			if (UWorld* World = Context.World())
+			{
+				UWorldPartitionHLODRuntimeSubsystem::WriteHLODStatsCSV(World, HLODStatsOutputFilename);
+			}
+		}
+	})
+);
+
+#endif // #if WITH_EDITOR
+
+#undef LOCTEXT_NAMESPACE
+
diff --git a/Engine/Source/Runtime/Engine/Private/WorldPartition/HLOD/HLODSubsystem.cpp b/Engine/Source/Runtime/Engine/Private/WorldPartition/HLOD/HLODSubsystem.cpp
deleted file mode 100644
diff --git a/Engine/Source/Runtime/Engine/Private/WorldPartition/WorldPartition.cpp b/Engine/Source/Runtime/Engine/Private/WorldPartition/WorldPartition.cpp
--- a/Engine/Source/Runtime/Engine/Private/WorldPartition/WorldPartition.cpp
+++ b/Engine/Source/Runtime/Engine/Private/WorldPartition/WorldPartition.cpp
@@ -12,7 +12,7 @@
 #include "WorldPartition/WorldPartitionLevelStreamingPolicy.h"
 #include "WorldPartition/WorldPartitionReplay.h"
 #include "WorldPartition/WorldPartitionSubsystem.h"
-#include "WorldPartition/HLOD/HLODSubsystem.h"
+#include "WorldPartition/HLOD/HLODRuntimeSubsystem.h"
 #include "WorldPartition/DataLayer/DataLayerManager.h"
 #include "WorldPartition/WorldPartitionSettings.h"
 #include "GameFramework/WorldSettings.h"
@@ -1711,7 +1711,7 @@ bool UWorldPartition::InjectExternalStreamingObject(URuntimeHashExternalStreamin
 		{
 			StreamingPolicy->InjectExternalStreamingObject(InExternalStreamingObject);
 		}
-		GetWorld()->GetSubsystem<UHLODSubsystem>()->OnExternalStreamingObjectInjected(InExternalStreamingObject);
+		GetWorld()->GetSubsystem<UWorldPartitionHLODRuntimeSubsystem>()->OnExternalStreamingObjectInjected(InExternalStreamingObject);
 		++StreamingStateEpoch;
 
 #if DO_CHECK
@@ -1738,7 +1738,7 @@ bool UWorldPartition::RemoveExternalStreamingObject(URuntimeHashExternalStreamin
 			StreamingPolicy->RemoveExternalStreamingObject(InExternalStreamingObject);
 		}
 		
-		GetWorld()->GetSubsystem<UHLODSubsystem>()->OnExternalStreamingObjectRemoved(InExternalStreamingObject);
+		GetWorld()->GetSubsystem<UWorldPartitionHLODRuntimeSubsystem>()->OnExternalStreamingObjectRemoved(InExternalStreamingObject);
 		++StreamingStateEpoch;
 	}
 
@@ -1795,7 +1795,7 @@ void UWorldPartition::OnCellShown(const UWorldPartitionRuntimeCell* InCell)
 	{
 		if (IsStreamingEnabled())
 		{
-			GetWorld()->GetSubsystem<UHLODSubsystem>()->OnCellShown(InCell);
+			GetWorld()->GetSubsystem<UWorldPartitionHLODRuntimeSubsystem>()->OnCellShown(InCell);
 		}
 		StreamingPolicy->OnCellShown(InCell);
 	}
@@ -1809,7 +1809,7 @@ void UWorldPartition::OnCellHidden(const UWorldPartitionRuntimeCell* InCell)
 	{
 		if (IsStreamingEnabled())
 		{
-			GetWorld()->GetSubsystem<UHLODSubsystem>()->OnCellHidden(InCell);
+			GetWorld()->GetSubsystem<UWorldPartitionHLODRuntimeSubsystem>()->OnCellHidden(InCell);
 		}
 		StreamingPolicy->OnCellHidden(InCell);
 	}
diff --git a/Engine/Source/Runtime/Engine/Private/WorldPartition/WorldPartitionLevelStreamingDynamic.cpp b/Engine/Source/Runtime/Engine/Private/WorldPartition/WorldPartitionLevelStreamingDynamic.cpp
--- a/Engine/Source/Runtime/Engine/Private/WorldPartition/WorldPartitionLevelStreamingDynamic.cpp
+++ b/Engine/Source/Runtime/Engine/Private/WorldPartition/WorldPartitionLevelStreamingDynamic.cpp
@@ -4,7 +4,7 @@
 #include "Engine/LevelStreaming.h"
 #include "WorldPartition/WorldPartition.h"
 #include "UObject/PropertyPortFlags.h"
-#include "WorldPartition/HLOD/HLODSubsystem.h"
+#include "WorldPartition/HLOD/HLODRuntimeSubsystem.h"
 
 #include UE_INLINE_GENERATED_CPP_BY_NAME(WorldPartitionLevelStreamingDynamic)
 
@@ -673,7 +673,7 @@ bool UWorldPartitionLevelStreamingDynamic::CanChangeVisibility(bool bMakeVisible
 				{
 					if (const UWorldPartition* WorldPartition = OuterWorld->GetWorldPartition())
 					{
-						if (UHLODSubsystem* HLODSubsystem = GetWorld()->GetSubsystem<UHLODSubsystem>())
+						if (UWorldPartitionHLODRuntimeSubsystem* HLODSubsystem = GetWorld()->GetSubsystem<UWorldPartitionHLODRuntimeSubsystem>())
 						{
 							if (bMakeVisible)
 							{
diff --git a/Engine/Source/Runtime/Engine/Private/WorldPartition/WorldPartitionStreamingPolicy.cpp b/Engine/Source/Runtime/Engine/Private/WorldPartition/WorldPartitionStreamingPolicy.cpp
--- a/Engine/Source/Runtime/Engine/Private/WorldPartition/WorldPartitionStreamingPolicy.cpp
+++ b/Engine/Source/Runtime/Engine/Private/WorldPartition/WorldPartitionStreamingPolicy.cpp
@@ -13,7 +13,7 @@
 #include "WorldPartition/WorldPartitionReplay.h"
 #include "WorldPartition/DataLayer/DataLayerManager.h"
 #include "WorldPartition/DataLayer/WorldDataLayers.h"
-#include "WorldPartition/HLOD/HLODSubsystem.h"
+#include "WorldPartition/HLOD/HLODRuntimeSubsystem.h"
 #include "WorldPartition/ContentBundle/ContentBundle.h"
 #include "WorldPartition/WorldPartitionDebugHelper.h"
 #include "Engine/Level.h"
@@ -469,7 +469,7 @@ void UWorldPartitionStreamingPolicy::UpdateStreamingState()
 
 	auto ShouldSkipDisabledHLODCell = [](const UWorldPartitionRuntimeCell* Cell)
 	{
-		return Cell->GetIsHLOD() && !UHLODSubsystem::IsHLODEnabled();
+		return Cell->GetIsHLOD() && !UWorldPartitionHLODRuntimeSubsystem::IsHLODEnabled();
 	};
 
 	// Activation superseeds Loading
@@ -912,7 +912,7 @@ bool UWorldPartitionStreamingPolicy::IsStreamingCompleted(const TArray<FWorldPar
 bool UWorldPartitionStreamingPolicy::IsStreamingCompleted(EWorldPartitionRuntimeCellState QueryState, const TArray<FWorldPartitionStreamingQuerySource>& QuerySources, bool bExactState) const
 {
 	const UDataLayerManager* DataLayerManager = WorldPartition->GetDataLayerManager();
-	const bool bIsHLODEnabled = UHLODSubsystem::IsHLODEnabled();
+	const bool bIsHLODEnabled = UWorldPartitionHLODRuntimeSubsystem::IsHLODEnabled();
 
 	bool bResult = true;
 	for (const FWorldPartitionStreamingQuerySource& QuerySource : QuerySources)
diff --git a/Engine/Source/Runtime/Engine/Public/WorldPartition/HLOD/HLODActor.h b/Engine/Source/Runtime/Engine/Public/WorldPartition/HLOD/HLODActor.h
--- a/Engine/Source/Runtime/Engine/Public/WorldPartition/HLOD/HLODActor.h
+++ b/Engine/Source/Runtime/Engine/Public/WorldPartition/HLOD/HLODActor.h
@@ -12,7 +12,6 @@
 #include "HLODActor.generated.h"
 
 class UHLODLayer;
-class UHLODSubsystem;
 class UWorldPartitionHLODSourceActors;
 
 ENGINE_API DECLARE_LOG_CATEGORY_EXTERN(LogHLODHash, Log, All);
diff --git a/Engine/Source/Runtime/Engine/Public/WorldPartition/HLOD/HLODRuntimeSubsystem.h b/Engine/Source/Runtime/Engine/Public/WorldPartition/HLOD/HLODRuntimeSubsystem.h
new file mode 100644
--- /dev/null
+++ b/Engine/Source/Runtime/Engine/Public/WorldPartition/HLOD/HLODRuntimeSubsystem.h
@@ -0,0 +1,139 @@
+// Copyright Epic Games, Inc. All Rights Reserved.
+#pragma once
+
+#include "Subsystems/WorldSubsystem.h"
+#include "Containers/Array.h"
+#include "Containers/Map.h"
+#include "UObject/ObjectKey.h"
+#include "Containers/Ticker.h"
+#include "HLODRuntimeSubsystem.generated.h"
+
+
+class AWorldPartitionHLOD;
+class FSceneViewFamily;
+class FHLODResourcesResidencySceneViewExtension;
+class UWorldPartition;
+class UWorldPartitionRuntimeCell;
+class URuntimeHashExternalStreamingObjectBase;
+class UWorld;
+
+DECLARE_MULTICAST_DELEGATE_OneParam(FWorldPartitionHLODActorRegisteredEvent, AWorldPartitionHLOD* /* InHLODActor */);
+DECLARE_MULTICAST_DELEGATE_OneParam(FWorldPartitionHLODActorUnregisteredEvent, AWorldPartitionHLOD* /* InHLODActor */);
+
+
+/**
+ * UWorldPartitionHLODRuntimeSubsystem
+ */
+UCLASS(MinimalAPI)
+class UWorldPartitionHLODRuntimeSubsystem : public UWorldSubsystem
+{
+	GENERATED_BODY()
+
+public:
+	ENGINE_API UWorldPartitionHLODRuntimeSubsystem();
+	ENGINE_API virtual ~UWorldPartitionHLODRuntimeSubsystem();
+
+	//~ Begin USubsystem Interface.
+	ENGINE_API virtual void Initialize(FSubsystemCollectionBase& Collection) override;
+	ENGINE_API virtual void Deinitialize() override;
+	//~ End USubsystem Interface.
+
+	//~ Begin UWorldSubsystem Interface.
+	ENGINE_API virtual bool DoesSupportWorldType(const EWorldType::Type WorldType) const override;
+	ENGINE_API virtual void OnWorldBeginPlay(UWorld& InWorld) override;
+	//~ End UWorldSubsystem Interface.
+
+	ENGINE_API bool Tick(float DeltaTime);
+
+	ENGINE_API void RegisterHLODActor(AWorldPartitionHLOD* InWorldPartitionHLOD);
+	ENGINE_API void UnregisterHLODActor(AWorldPartitionHLOD* InWorldPartitionHLOD);
+
+	ENGINE_API void OnCellShown(const UWorldPartitionRuntimeCell* InCell);
+	ENGINE_API void OnCellHidden(const UWorldPartitionRuntimeCell* InCell);
+
+	ENGINE_API bool CanMakeVisible(const UWorldPartitionRuntimeCell* InCell);
+	ENGINE_API bool CanMakeInvisible(const UWorldPartitionRuntimeCell* InCell);
+
+	ENGINE_API const TArray<AWorldPartitionHLOD*>& GetHLODActorsForCell(const UWorldPartitionRuntimeCell* InCell) const;
+
+	static ENGINE_API bool IsHLODEnabled();
+
+	FWorldPartitionHLODActorRegisteredEvent& OnHLODActorRegisteredEvent() { return HLODActorRegisteredEvent; }
+	FWorldPartitionHLODActorUnregisteredEvent& OnHLODActorUnregisteredEvent() { return HLODActorUnregisteredEvent; }
+
+	ENGINE_API void OnExternalStreamingObjectInjected(URuntimeHashExternalStreamingObjectBase* ExternalStreamingObject);
+	ENGINE_API void OnExternalStreamingObjectRemoved(URuntimeHashExternalStreamingObjectBase* ExternalStreamingObject);
+
+	ENGINE_API void SetHLODAlwaysLoadedCullDistance(int32 InCullDistance);
+
+	ENGINE_API void OnCVarsChanged();
+
+#if WITH_EDITOR
+	uint32 GetNumOutdatedHLODActors() const { return OutdatedHLODActors.Num(); }
+	static ENGINE_API bool WriteHLODStatsCSV(UWorld* InWorld, const FString& InFilename);
+#endif
+	
+private:
+	struct FCellData
+	{
+		bool bIsCellVisible = false;
+		TArray<AWorldPartitionHLOD*> LoadedHLODs;
+	};
+
+	struct FWorldPartitionHLODRuntimeData
+	{
+		TMap<FGuid, FCellData> CellsData;
+	};
+
+	struct FDrawDistanceQueue
+	{
+		float DrawDistance;
+		TObjectPtr<AWorldPartitionHLOD> HLODActor;
+	};
+	
+	TMap<TObjectPtr<UWorldPartition>, FWorldPartitionHLODRuntimeData> WorldPartitionsHLODRuntimeData;
+	ENGINE_API const FCellData* GetCellData(const UWorldPartitionRuntimeCell* InCell) const;
+	ENGINE_API FCellData* GetCellData(const UWorldPartitionRuntimeCell* InCell);
+	ENGINE_API FCellData* GetCellData(AWorldPartitionHLOD* InWorldPartitionHLOD);
+
+	TArray<AWorldPartitionHLOD*> AlwaysLoadedHLODActors;
+
+	struct FWorldPartitionHLODWarmupState
+	{
+		uint32 WarmupStartFrame = INDEX_NONE;
+		uint32 WarmupEndFrame = INDEX_NONE;
+		FVector Location;
+	};
+
+	typedef TMap<FObjectKey, FWorldPartitionHLODWarmupState> FHLODWarmupStateMap;
+	FHLODWarmupStateMap HLODActorsToWarmup;
+
+	ENGINE_API void OnBeginRenderViews(const FSceneViewFamily& InViewFamily);
+	ENGINE_API bool PrepareToWarmup(const UWorldPartitionRuntimeCell* InCell, AWorldPartitionHLOD* InHLODActor);
+	ENGINE_API bool ShouldPerformWarmup() const;
+	ENGINE_API bool ShouldPerformWarmupForCell(const UWorldPartitionRuntimeCell* InCell) const;
+	bool bCachedShouldPerformWarmup;
+
+	TArray<FDrawDistanceQueue> OperationQueue;
+	FTSTicker::FDelegateHandle TickHandle;
+	int32 LastSetCullDistance;
+
+	/** Console command used to turn on/off loading & rendering of world partition HLODs */
+	static ENGINE_API class FAutoConsoleCommand EnableHLODCommand;
+
+	static ENGINE_API bool WorldPartitionHLODEnabled;
+
+	friend class FHLODResourcesResidencySceneViewExtension;
+	TSharedPtr<FHLODResourcesResidencySceneViewExtension, ESPMode::ThreadSafe> SceneViewExtension;
+
+	ENGINE_API void OnWorldPartitionInitialized(UWorldPartition* InWorldPartition);
+	ENGINE_API void OnWorldPartitionUninitialized(UWorldPartition* InWorldPartition);
+
+	FWorldPartitionHLODActorRegisteredEvent		HLODActorRegisteredEvent;
+	FWorldPartitionHLODActorUnregisteredEvent	HLODActorUnregisteredEvent;
+
+#if WITH_EDITOR
+	TSet<AWorldPartitionHLOD*> OutdatedHLODActors;
+#endif
+};
+
diff --git a/Engine/Source/Runtime/Engine/Public/WorldPartition/HLOD/HLODSubsystem.h b/Engine/Source/Runtime/Engine/Public/WorldPartition/HLOD/HLODSubsystem.h
--- a/Engine/Source/Runtime/Engine/Public/WorldPartition/HLOD/HLODSubsystem.h
+++ b/Engine/Source/Runtime/Engine/Public/WorldPartition/HLOD/HLODSubsystem.h
@@ -1,139 +1,8 @@
 // Copyright Epic Games, Inc. All Rights Reserved.
 #pragma once
 
-#include "Subsystems/WorldSubsystem.h"
-#include "Containers/Array.h"
-#include "Containers/Map.h"
-#include "UObject/ObjectKey.h"
-#include "Containers/Ticker.h"
-#include "HLODSubsystem.generated.h"
+#include "WorldPartition/HLOD/HLODRuntimeSubsystem.h"
 
 
-class AWorldPartitionHLOD;
-class FSceneViewFamily;
-class FHLODResourcesResidencySceneViewExtension;
-class UWorldPartition;
-class UWorldPartitionRuntimeCell;
-class URuntimeHashExternalStreamingObjectBase;
-class UWorld;
-
-DECLARE_MULTICAST_DELEGATE_OneParam(FWorldPartitionHLODActorRegisteredEvent, AWorldPartitionHLOD* /* InHLODActor */);
-DECLARE_MULTICAST_DELEGATE_OneParam(FWorldPartitionHLODActorUnregisteredEvent, AWorldPartitionHLOD* /* InHLODActor */);
-
-
-/**
- * UHLODSubsystem
- */
-UCLASS(MinimalAPI)
-class UHLODSubsystem : public UWorldSubsystem
-{
-	GENERATED_BODY()
-
-public:
-	ENGINE_API UHLODSubsystem();
-	ENGINE_API virtual ~UHLODSubsystem();
-
-	//~ Begin USubsystem Interface.
-	ENGINE_API virtual void Initialize(FSubsystemCollectionBase& Collection) override;
-	ENGINE_API virtual void Deinitialize() override;
-	//~ End USubsystem Interface.
-
-	//~ Begin UWorldSubsystem Interface.
-	ENGINE_API virtual bool DoesSupportWorldType(const EWorldType::Type WorldType) const override;
-	ENGINE_API virtual void OnWorldBeginPlay(UWorld& InWorld) override;
-	//~ End UWorldSubsystem Interface.
-
-	ENGINE_API bool Tick(float DeltaTime);
-
-	ENGINE_API void RegisterHLODActor(AWorldPartitionHLOD* InWorldPartitionHLOD);
-	ENGINE_API void UnregisterHLODActor(AWorldPartitionHLOD* InWorldPartitionHLOD);
-
-	ENGINE_API void OnCellShown(const UWorldPartitionRuntimeCell* InCell);
-	ENGINE_API void OnCellHidden(const UWorldPartitionRuntimeCell* InCell);
-
-	ENGINE_API bool CanMakeVisible(const UWorldPartitionRuntimeCell* InCell);
-	ENGINE_API bool CanMakeInvisible(const UWorldPartitionRuntimeCell* InCell);
-
-	ENGINE_API const TArray<AWorldPartitionHLOD*>& GetHLODActorsForCell(const UWorldPartitionRuntimeCell* InCell) const;
-
-	static ENGINE_API bool IsHLODEnabled();
-
-	FWorldPartitionHLODActorRegisteredEvent& OnHLODActorRegisteredEvent() { return HLODActorRegisteredEvent; }
-	FWorldPartitionHLODActorUnregisteredEvent& OnHLODActorUnregisteredEvent() { return HLODActorUnregisteredEvent; }
-
-	ENGINE_API void OnExternalStreamingObjectInjected(URuntimeHashExternalStreamingObjectBase* ExternalStreamingObject);
-	ENGINE_API void OnExternalStreamingObjectRemoved(URuntimeHashExternalStreamingObjectBase* ExternalStreamingObject);
-
-	ENGINE_API void SetHLODAlwaysLoadedCullDistance(int32 InCullDistance);
-
-	ENGINE_API void OnCVarsChanged();
-
-#if WITH_EDITOR
-	uint32 GetNumOutdatedHLODActors() const { return OutdatedHLODActors.Num(); }
-	static ENGINE_API bool WriteHLODStatsCSV(UWorld* InWorld, const FString& InFilename);
-#endif
-	
-private:
-	struct FCellData
-	{
-		bool bIsCellVisible = false;
-		TArray<AWorldPartitionHLOD*> LoadedHLODs;
-	};
-
-	struct FWorldPartitionHLODRuntimeData
-	{
-		TMap<FGuid, FCellData> CellsData;
-	};
-
-	struct FDrawDistanceQueue
-	{
-		float DrawDistance;
-		TObjectPtr<AWorldPartitionHLOD> HLODActor;
-	};
-	
-	TMap<TObjectPtr<UWorldPartition>, FWorldPartitionHLODRuntimeData> WorldPartitionsHLODRuntimeData;
-	ENGINE_API const FCellData* GetCellData(const UWorldPartitionRuntimeCell* InCell) const;
-	ENGINE_API FCellData* GetCellData(const UWorldPartitionRuntimeCell* InCell);
-	ENGINE_API FCellData* GetCellData(AWorldPartitionHLOD* InWorldPartitionHLOD);
-
-	TArray<AWorldPartitionHLOD*> AlwaysLoadedHLODActors;
-
-	struct FWorldPartitionHLODWarmupState
-	{
-		uint32 WarmupStartFrame = INDEX_NONE;
-		uint32 WarmupEndFrame = INDEX_NONE;
-		FVector Location;
-	};
-
-	typedef TMap<FObjectKey, FWorldPartitionHLODWarmupState> FHLODWarmupStateMap;
-	FHLODWarmupStateMap HLODActorsToWarmup;
-
-	ENGINE_API void OnBeginRenderViews(const FSceneViewFamily& InViewFamily);
-	ENGINE_API bool PrepareToWarmup(const UWorldPartitionRuntimeCell* InCell, AWorldPartitionHLOD* InHLODActor);
-	ENGINE_API bool ShouldPerformWarmup() const;
-	ENGINE_API bool ShouldPerformWarmupForCell(const UWorldPartitionRuntimeCell* InCell) const;
-	bool bCachedShouldPerformWarmup;
-
-	TArray<FDrawDistanceQueue> OperationQueue;
-	FTSTicker::FDelegateHandle TickHandle;
-	int32 LastSetCullDistance;
-
-	/** Console command used to turn on/off loading & rendering of world partition HLODs */
-	static ENGINE_API class FAutoConsoleCommand EnableHLODCommand;
-
-	static ENGINE_API bool WorldPartitionHLODEnabled;
-
-	friend class FHLODResourcesResidencySceneViewExtension;
-	TSharedPtr<FHLODResourcesResidencySceneViewExtension, ESPMode::ThreadSafe> SceneViewExtension;
-
-	ENGINE_API void OnWorldPartitionInitialized(UWorldPartition* InWorldPartition);
-	ENGINE_API void OnWorldPartitionUninitialized(UWorldPartition* InWorldPartition);
-
-	FWorldPartitionHLODActorRegisteredEvent		HLODActorRegisteredEvent;
-	FWorldPartitionHLODActorUnregisteredEvent	HLODActorUnregisteredEvent;
-
-#if WITH_EDITOR
-	TSet<AWorldPartitionHLOD*> OutdatedHLODActors;
-#endif
-};
-
+UE_DEPRECATED(5.4, "Use UWorldPartitionHLODRuntimeSubsystem")
+typedef UWorldPartitionHLODRuntimeSubsystem UHLODSubsystem;
