Fix an issue in multiplayer PIE where some world partition level package names weren't being renamed in the net driver's PackageMap. This would cause broken actor channels and failed replication for actors in that level.

diff --git a/Engine/Source/Runtime/Engine/Private/LevelStreaming.cpp b/Engine/Source/Runtime/Engine/Private/LevelStreaming.cpp
--- a/Engine/Source/Runtime/Engine/Private/LevelStreaming.cpp
+++ b/Engine/Source/Runtime/Engine/Private/LevelStreaming.cpp
@@ -165,16 +165,16 @@ static void NetDriverRenameStreamingLevelPackageForPIE(const UWorld* World, FNam
 	for (FNamedNetDriver& Driver : WorldContext->ActiveNetDrivers)
 	{
 		if (Driver.NetDriver && Driver.NetDriver->GuidCache.IsValid())
-	{
-			for (TPair<FNetworkGUID, FNetGuidCacheObject>& GuidPair : Driver.NetDriver->GuidCache->ObjectLookup)
-	{
-		// Only look for packages, which will have a static GUID and an invalid OuterGUID.
-		const bool bIsPackage = GuidPair.Key.IsStatic() && !GuidPair.Value.OuterGUID.IsValid();
-		if (bIsPackage && GuidPair.Value.PathName == UnPrefixedPackageName)
 		{
-			GuidPair.Value.PathName = *UWorld::ConvertToPIEPackageName(GuidPair.Value.PathName.ToString(), WorldContext->PIEInstance);
-		}
-	}
+			for (TPair<FNetworkGUID, FNetGuidCacheObject>& GuidPair : Driver.NetDriver->GuidCache->ObjectLookup)
+			{
+				// Only look for packages, which will have a static GUID and an invalid OuterGUID.
+				const bool bIsPackage = GuidPair.Key.IsStatic() && !GuidPair.Value.OuterGUID.IsValid();
+				if (bIsPackage && GuidPair.Value.PathName == UnPrefixedPackageName)
+				{
+					GuidPair.Value.PathName = *UWorld::ConvertToPIEPackageName(GuidPair.Value.PathName.ToString(), WorldContext->PIEInstance);
+				}
+			}
 		}
 	}
 }
@@ -2030,13 +2030,13 @@ void ULevelStreaming::RenameForPIE(int32 PIEInstanceID, bool bKeepWorldAssetName
 	// Apply PIE prefix so this level references
 	if (!WorldAsset.IsNull())
 	{
+		FName NonPrefixedName = *UWorld::StripPIEPrefixFromPackageName(GetWorldAssetPackageName(), UWorld::BuildPIEPackagePrefix(PIEInstanceID));
+		NetDriverRenameStreamingLevelPackageForPIE(World, NonPrefixedName);
+		
 		// Store original name 
 		if (PackageNameToLoad == NAME_None)
 		{
-			FString NonPrefixedName = UWorld::StripPIEPrefixFromPackageName(
-				GetWorldAssetPackageName(), 
-				UWorld::BuildPIEPackagePrefix(PIEInstanceID));
-			PackageNameToLoad = FName(*NonPrefixedName);
+			PackageNameToLoad = NonPrefixedName;
 		}
 		FName PlayWorldStreamingPackageName = FName(*UWorld::ConvertToPIEPackageName(GetWorldAssetPackageName(), PIEInstanceID));
 		FSoftObjectPath::AddPIEPackageName(PlayWorldStreamingPackageName);
@@ -2048,8 +2048,6 @@ void ULevelStreaming::RenameForPIE(int32 PIEInstanceID, bool bKeepWorldAssetName
 		{
 			SetWorldAssetByPackageName(PlayWorldStreamingPackageName);
 		}
-
-		NetDriverRenameStreamingLevelPackageForPIE(World, PackageNameToLoad);
 	}
 	
 	// Rename LOD levels if any
