Changes for the WIP FMaterialHLSLGenerator:
- Fixed a bug where static parameter values can change for the same material between post load and translate times. The cause was using an incomplete FStaticParameterSet in UMaterialInstanceConstant::UpdateCachedData. UMaterialInstance::GetStaticParameters only gets the params overridden by the instance and doesn't include the overrides in parent instances. However, we need to have all the overrides to correctly preparing an HLSL tree
- Implemented EmitValueObject for FExpressionSelect
- Implemented GenerateHLSLExpression for the following material nodes: ReflectionCapturePassSwitch, ParticleMotionBlurFade, SceneDepthWithoutWater, SkyAtmosphereLightDiskLuminance, SkyAtmosphereAerialPerspective, SkyAtmosphereViewLuminance, HsvToRgb, RgbToHsv, and SingleLayerWaterMaterialOutput
- Misc. bug fixes

diff --git a/Engine/Source/Runtime/Engine/Classes/Materials/MaterialExpressionHsvToRgb.h b/Engine/Source/Runtime/Engine/Classes/Materials/MaterialExpressionHsvToRgb.h
--- a/Engine/Source/Runtime/Engine/Classes/Materials/MaterialExpressionHsvToRgb.h
+++ b/Engine/Source/Runtime/Engine/Classes/Materials/MaterialExpressionHsvToRgb.h
@@ -18,6 +18,8 @@ class UMaterialExpressionHsvToRgb : public UMaterialExpression
 	virtual void GetCaption(TArray<FString>& OutCaptions) const override;
 	virtual void GetExpressionToolTip(TArray<FString>& OutToolTip) override;
 	virtual FText GetKeywords() const override {return FText::FromString(TEXT("hsvtorgb"));}
+
+	virtual bool GenerateHLSLExpression(FMaterialHLSLGenerator& Generator, UE::HLSLTree::FScope& Scope, int32 OutputIndex, UE::HLSLTree::FExpression const*& OutExpression) const override;
 #endif
 	//~ End UMaterialExpressionMaterialX Interface
 };
diff --git a/Engine/Source/Runtime/Engine/Classes/Materials/MaterialExpressionParticleMotionBlurFade.h b/Engine/Source/Runtime/Engine/Classes/Materials/MaterialExpressionParticleMotionBlurFade.h
--- a/Engine/Source/Runtime/Engine/Classes/Materials/MaterialExpressionParticleMotionBlurFade.h
+++ b/Engine/Source/Runtime/Engine/Classes/Materials/MaterialExpressionParticleMotionBlurFade.h
@@ -22,6 +22,8 @@ class UMaterialExpressionParticleMotionBlurFade : public UMaterialExpression
 #if WITH_EDITOR
 	virtual int32 Compile(class FMaterialCompiler* Compiler, int32 OutputIndex) override;
 	virtual void GetCaption(TArray<FString>& OutCaptions) const override;
+
+	virtual bool GenerateHLSLExpression(FMaterialHLSLGenerator& Generator, UE::HLSLTree::FScope& Scope, int32 OutputIndex, UE::HLSLTree::FExpression const*& OutExpression) const override;
 #endif
 	//~ End UMaterialExpression Interface
 };
diff --git a/Engine/Source/Runtime/Engine/Classes/Materials/MaterialExpressionReflectionCapturePassSwitch.h b/Engine/Source/Runtime/Engine/Classes/Materials/MaterialExpressionReflectionCapturePassSwitch.h
--- a/Engine/Source/Runtime/Engine/Classes/Materials/MaterialExpressionReflectionCapturePassSwitch.h
+++ b/Engine/Source/Runtime/Engine/Classes/Materials/MaterialExpressionReflectionCapturePassSwitch.h
@@ -29,6 +29,8 @@ class UMaterialExpressionReflectionCapturePassSwitch : public UMaterialExpressio
 	virtual bool IsResultSubstrateMaterial(int32 OutputIndex) override;
 	virtual void GatherSubstrateMaterialInfo(FSubstrateMaterialInfo& SubstrateMaterialInfo, int32 OutputIndex) override;
 	virtual FSubstrateOperator* SubstrateGenerateMaterialTopologyTree(class FMaterialCompiler* Compiler, class UMaterialExpression* Parent, int32 OutputIndex) override;
+
+	virtual bool GenerateHLSLExpression(FMaterialHLSLGenerator& Generator, UE::HLSLTree::FScope& Scope, int32 OutputIndex, UE::HLSLTree::FExpression const*& OutExpression) const override;
 #endif
 	//~ End UMaterialExpression Interface
 };
diff --git a/Engine/Source/Runtime/Engine/Classes/Materials/MaterialExpressionRgbToHsv.h b/Engine/Source/Runtime/Engine/Classes/Materials/MaterialExpressionRgbToHsv.h
--- a/Engine/Source/Runtime/Engine/Classes/Materials/MaterialExpressionRgbToHsv.h
+++ b/Engine/Source/Runtime/Engine/Classes/Materials/MaterialExpressionRgbToHsv.h
@@ -18,6 +18,8 @@ class UMaterialExpressionRgbToHsv : public UMaterialExpression
 	virtual void GetCaption(TArray<FString>& OutCaptions) const override;
 	virtual void GetExpressionToolTip(TArray<FString>& OutToolTip) override;
 	virtual FText GetKeywords() const override {return FText::FromString(TEXT("rgbtohsv"));}
+
+	virtual bool GenerateHLSLExpression(FMaterialHLSLGenerator& Generator, UE::HLSLTree::FScope& Scope, int32 OutputIndex, UE::HLSLTree::FExpression const*& OutExpression) const override;
 #endif
 	//~ End UMaterialExpressionMaterialX Interface
 };
diff --git a/Engine/Source/Runtime/Engine/Classes/Materials/MaterialExpressionSceneDepthWithoutWater.h b/Engine/Source/Runtime/Engine/Classes/Materials/MaterialExpressionSceneDepthWithoutWater.h
--- a/Engine/Source/Runtime/Engine/Classes/Materials/MaterialExpressionSceneDepthWithoutWater.h
+++ b/Engine/Source/Runtime/Engine/Classes/Materials/MaterialExpressionSceneDepthWithoutWater.h
@@ -46,6 +46,8 @@ public:
 	virtual int32 Compile(class FMaterialCompiler* Compiler, int32 OutputIndex) override;
 	virtual void GetCaption(TArray<FString>& OutCaptions) const override;
 	virtual FName GetInputName(int32 InputIndex) const override;
+
+	virtual bool GenerateHLSLExpression(FMaterialHLSLGenerator& Generator, UE::HLSLTree::FScope& Scope, int32 OutputIndex, UE::HLSLTree::FExpression const*& OutExpression) const override;
 #endif
 	//~ End UMaterialExpression Interface
 };
diff --git a/Engine/Source/Runtime/Engine/Classes/Materials/MaterialExpressionSingleLayerWaterMaterialOutput.h b/Engine/Source/Runtime/Engine/Classes/Materials/MaterialExpressionSingleLayerWaterMaterialOutput.h
--- a/Engine/Source/Runtime/Engine/Classes/Materials/MaterialExpressionSingleLayerWaterMaterialOutput.h
+++ b/Engine/Source/Runtime/Engine/Classes/Materials/MaterialExpressionSingleLayerWaterMaterialOutput.h
@@ -34,6 +34,8 @@ public:
 	//~ Begin UMaterialExpression Interface
 	virtual int32 Compile(class FMaterialCompiler* Compiler, int32 OutputIndex) override;
 	virtual void GetCaption(TArray<FString>& OutCaptions) const override;
+
+	virtual bool GenerateHLSLExpression(FMaterialHLSLGenerator& Generator, UE::HLSLTree::FScope& Scope, int32 OutputIndex, UE::HLSLTree::FExpression const*& OutExpression) const override;
 	//~ End UMaterialExpression Interface
 #endif
 
@@ -41,5 +43,8 @@ public:
 	virtual int32 GetNumOutputs() const override;
 	virtual FString GetFunctionName() const override;
 	virtual FString GetDisplayName() const override;
+#if WITH_EDITOR
+	virtual UE::Shader::EValueType GetCustomOutputType(int32 OutputIndex) const override;
+#endif
 	//~ End UMaterialExpressionCustomOutput Interface
 };
diff --git a/Engine/Source/Runtime/Engine/Classes/Materials/MaterialExpressionSkyAtmosphereLightIlluminance.h b/Engine/Source/Runtime/Engine/Classes/Materials/MaterialExpressionSkyAtmosphereLightIlluminance.h
--- a/Engine/Source/Runtime/Engine/Classes/Materials/MaterialExpressionSkyAtmosphereLightIlluminance.h
+++ b/Engine/Source/Runtime/Engine/Classes/Materials/MaterialExpressionSkyAtmosphereLightIlluminance.h
@@ -69,6 +69,8 @@ class UMaterialExpressionSkyAtmosphereLightDiskLuminance : public UMaterialExpre
 #if WITH_EDITOR
 	virtual int32 Compile(class FMaterialCompiler* Compiler, int32 OutputIndex) override;
 	virtual void GetCaption(TArray<FString>& OutCaptions) const override;
+
+	virtual bool GenerateHLSLExpression(FMaterialHLSLGenerator& Generator, UE::HLSLTree::FScope& Scope, int32 OutputIndex, UE::HLSLTree::FExpression const*& OutExpression) const override;
 #endif
 	//~ End UMaterialExpression Interface
 };
@@ -89,6 +91,8 @@ class UMaterialExpressionSkyAtmosphereAerialPerspective : public UMaterialExpres
 #if WITH_EDITOR
 	virtual int32 Compile(class FMaterialCompiler* Compiler, int32 OutputIndex) override;
 	virtual void GetCaption(TArray<FString>& OutCaptions) const override;
+
+	virtual bool GenerateHLSLExpression(FMaterialHLSLGenerator& Generator, UE::HLSLTree::FScope& Scope, int32 OutputIndex, UE::HLSLTree::FExpression const*& OutExpression) const override;
 #endif
 	//~ End UMaterialExpression Interface
 };
diff --git a/Engine/Source/Runtime/Engine/Classes/Materials/MaterialExpressionSkyAtmosphereViewLuminance.h b/Engine/Source/Runtime/Engine/Classes/Materials/MaterialExpressionSkyAtmosphereViewLuminance.h
--- a/Engine/Source/Runtime/Engine/Classes/Materials/MaterialExpressionSkyAtmosphereViewLuminance.h
+++ b/Engine/Source/Runtime/Engine/Classes/Materials/MaterialExpressionSkyAtmosphereViewLuminance.h
@@ -17,6 +17,8 @@ class UMaterialExpressionSkyAtmosphereViewLuminance : public UMaterialExpression
 #if WITH_EDITOR
 	virtual int32 Compile(class FMaterialCompiler* Compiler, int32 OutputIndex) override;
 	virtual void GetCaption(TArray<FString>& OutCaptions) const override;
+
+	virtual bool GenerateHLSLExpression(FMaterialHLSLGenerator& Generator, UE::HLSLTree::FScope& Scope, int32 OutputIndex, UE::HLSLTree::FExpression const*& OutExpression) const override;
 #endif
 	//~ End UMaterialExpression Interface
 };
diff --git a/Engine/Source/Runtime/Engine/Private/HLSLTree/HLSLTreeCommon.cpp b/Engine/Source/Runtime/Engine/Private/HLSLTree/HLSLTreeCommon.cpp
--- a/Engine/Source/Runtime/Engine/Private/HLSLTree/HLSLTreeCommon.cpp
+++ b/Engine/Source/Runtime/Engine/Private/HLSLTree/HLSLTreeCommon.cpp
@@ -493,6 +493,10 @@ bool FExpressionSelect::PrepareValue(FEmitContext& Context, FEmitScope& Scope, c
 			return Context.Error(TEXT("Cannot mix numeric vectors with different number of components"));
 		}
 	}
+	else if (LhsType.IsObject() && !IsConstantEvaluation(ConditionEvaluation))
+	{
+		return Context.Error(TEXT("Condition must be constant when selecting between objects"));
+	}
 
 	ResultType.MergeEvaluation(ConditionEvaluation);
 	return OutResult.SetType(Context, RequestedType, ResultType);
@@ -577,6 +581,14 @@ void FExpressionSelect::EmitValuePreshader(FEmitContext& Context, FEmitScope& Sc
 	}
 }
 
+bool FExpressionSelect::EmitValueObject(FEmitContext& Context, FEmitScope& Scope, const FName& ObjectTypeName, void* OutObjectBase) const
+{
+	// We have already checked that the condition is constant during PrepareValue
+	const bool bCondition = ConditionExpression->GetValueConstant(Context, Scope, Shader::EValueType::Bool1).AsBoolScalar();
+	const FExpression* SelectedExpression = bCondition ? TrueExpression : FalseExpression;
+	return SelectedExpression->GetValueObject(Context, Scope, ObjectTypeName, OutObjectBase);
+}
+
 void FExpressionDerivative::ComputeAnalyticDerivatives(FTree& Tree, FExpressionDerivatives& OutResult) const
 {
 	// TODO
@@ -717,7 +729,7 @@ void SwizzleEmitValueShader(FEmitContext& Context,
 
 	// Make sure the input is cast to the explicit requested type, this ensures it will have enough components for the swizzle
 	// Alternately, we could avoid the cast, and update the logic to insert 0s for swizzle access to invalid components
-	FEmitShaderExpression* EmitInput = Input->GetValueShader(Context, Scope, RequestedInputType, RequestedInputType.Type);
+	FEmitShaderExpression* EmitInput = Input->GetValueShader(Context, Scope, RequestedInputType, RequestedInputType.Type.GetConcreteType());
 
 	const Shader::FValueTypeDescription InputTypeDesc = Shader::GetValueTypeDescription(EmitInput->Type);
 	bool bSkipSwizzle = (InputTypeDesc.NumComponents == Parameters.NumComponents);
@@ -1023,6 +1035,11 @@ void FExpressionAppend::EmitValuePreshader(FEmitContext& Context, FEmitScope& Sc
 	OutResult.Type = Types.ResultType.GetResultType();
 }
 
+const FExpression* FTree::NewAppend(const FExpression* Lhs, const FExpression* Rhs)
+{
+	return NewExpression<FExpressionAppend>(Lhs, Rhs);
+}
+
 void FExpressionSwitchBase::ComputeAnalyticDerivatives(FTree& Tree, FExpressionDerivatives& OutResult) const
 {
 	const FExpression* InputDdx[MaxInputs];
@@ -1192,6 +1209,11 @@ FExpressionVirtualTextureFeatureSwitch::FExpressionVirtualTextureFeatureSwitch(T
 
 bool FExpressionVirtualTextureFeatureSwitch::IsInputActive(const FEmitContext& Context, int32 Index) const
 {
+	if (Context.TargetParameters.IsGenericTarget())
+	{
+		return true;
+	}
+
 	if (UseVirtualTexturing(Context.TargetParameters.FeatureLevel, Context.TargetParameters.TargetPlatform))
 	{
 		return Index == 0;
diff --git a/Engine/Source/Runtime/Engine/Private/HLSLTree/HLSLTreeEmit.cpp b/Engine/Source/Runtime/Engine/Private/HLSLTree/HLSLTreeEmit.cpp
--- a/Engine/Source/Runtime/Engine/Private/HLSLTree/HLSLTreeEmit.cpp
+++ b/Engine/Source/Runtime/Engine/Private/HLSLTree/HLSLTreeEmit.cpp
@@ -492,7 +492,9 @@ FPreparedType FEmitContext::PrepareExpression(const FExpression* InExpression, F
 		}
 		else
 		{
-			const Shader::FType CombinedType = Shader::CombineTypes(RequestedType.Type, PastRequestedType.Type);
+			checkf(!RequestedType.Type.IsNumericMatrix() || RequestedType.Type.ValueType != Shader::EValueType::DoubleInverse4x4,
+				TEXT("DoubleInverse4x4 shouldn't be explicitly requested"));
+			const Shader::FType CombinedType = Shader::CombineTypes(RequestedType.Type, PastRequestedType.Type, true);
 			check(!CombinedType.IsVoid());
 
 			if (CombinedType == PastRequestedType.Type && AllComponentsRequestedInPast(RequestedType, PastRequestedType))
diff --git a/Engine/Source/Runtime/Engine/Private/Materials/MaterialExpressionHLSL.cpp b/Engine/Source/Runtime/Engine/Private/Materials/MaterialExpressionHLSL.cpp
--- a/Engine/Source/Runtime/Engine/Private/Materials/MaterialExpressionHLSL.cpp
+++ b/Engine/Source/Runtime/Engine/Private/Materials/MaterialExpressionHLSL.cpp
@@ -88,6 +88,7 @@
 #include "Materials/MaterialExpressionGetMaterialAttributes.h"
 #include "Materials/MaterialExpressionHairAttributes.h"
 #include "Materials/MaterialExpressionHairColor.h"
+#include "Materials/MaterialExpressionHsvToRgb.h"
 #include "Materials/MaterialExpressionIf.h"
 #include "Materials/MaterialExpressionIfThenElse.h"
 #include "Materials/MaterialExpressionInverseLinearInterpolate.h"
@@ -120,6 +121,7 @@
 #include "Materials/MaterialExpressionParticleColor.h"
 #include "Materials/MaterialExpressionParticleDirection.h"
 #include "Materials/MaterialExpressionParticleMacroUV.h"
+#include "Materials/MaterialExpressionParticleMotionBlurFade.h"
 #include "Materials/MaterialExpressionParticlePositionWS.h"
 #include "Materials/MaterialExpressionParticleRadius.h"
 #include "Materials/MaterialExpressionParticleRandom.h"
@@ -142,8 +144,10 @@
 #include "Materials/MaterialExpressionSamplePhysicsField.h"
 #include "Materials/MaterialExpressionQualitySwitch.h"
 #include "Materials/MaterialExpressionRayTracingQualitySwitch.h"
+#include "Materials/MaterialExpressionReflectionCapturePassSwitch.h"
 #include "Materials/MaterialExpressionReflectionVectorWS.h"
 #include "Materials/MaterialExpressionReroute.h"
+#include "Materials/MaterialExpressionRgbToHsv.h"
 #include "Materials/MaterialExpressionRotateAboutAxis.h"
 #include "Materials/MaterialExpressionRotator.h"
 #include "Materials/MaterialExpressionRound.h"
@@ -152,6 +156,7 @@
 #include "Materials/MaterialExpressionSaturate.h"
 #include "Materials/MaterialExpressionSceneColor.h"
 #include "Materials/MaterialExpressionSceneDepth.h"
+#include "Materials/MaterialExpressionSceneDepthWithoutWater.h"
 #include "Materials/MaterialExpressionSceneTexelSize.h"
 #include "Materials/MaterialExpressionSceneTexture.h"
 #include "Materials/MaterialExpressionScreenPosition.h"
@@ -161,9 +166,12 @@
 #include "Materials/MaterialExpressionShadingModel.h"
 #include "Materials/MaterialExpressionShadingPathSwitch.h"
 #include "Materials/MaterialExpressionShadowReplace.h"
+#include "Materials/MaterialExpressionSingleLayerWaterMaterialOutput.h"
 #include "Materials/MaterialExpressionSign.h"
 #include "Materials/MaterialExpressionSine.h"
 #include "Materials/MaterialExpressionSkyAtmosphereLightDirection.h"
+#include "Materials/MaterialExpressionSkyAtmosphereLightIlluminance.h"
+#include "Materials/MaterialExpressionSkyAtmosphereViewLuminance.h"
 #include "Materials/MaterialExpressionSkyLightEnvMapSample.h"
 #include "Materials/MaterialExpressionSmoothStep.h"
 #include "Materials/MaterialExpressionSphereMask.h"
@@ -206,6 +214,7 @@
 #include "Materials/MaterialParameterCollection.h"
 #include "MaterialShared.h"
 #include "Misc/MemStackUtility.h"
+#include "RenderUtils.h"
 
 bool UMaterialExpression::GenerateHLSLExpression(FMaterialHLSLGenerator& Generator, UE::HLSLTree::FScope& Scope, int32 OutputIndex, UE::HLSLTree::FExpression const*& OutExpression) const
 {
@@ -493,6 +502,29 @@ bool UMaterialExpressionNaniteReplace::GenerateHLSLExpression(FMaterialHLSLGener
 	return OutExpression != nullptr;
 }
 
+bool UMaterialExpressionReflectionCapturePassSwitch::GenerateHLSLExpression(FMaterialHLSLGenerator& Generator, UE::HLSLTree::FScope& Scope, int32 OutputIndex, UE::HLSLTree::FExpression const*& OutExpression) const
+{
+	using namespace UE::HLSLTree;
+
+	if (!Default.GetTracedInput().Expression)
+	{
+		return Generator.Error(TEXT("Missing input Default"));
+	}
+	else if (!Reflection.GetTracedInput().Expression)
+	{
+		return Generator.Error(TEXT("Missing input Reflection"));
+	}
+	else
+	{
+		const FExpression* ConditionExpression = Generator.GetTree().NewExpression<FExpressionInlineCustomHLSL>(UE::Shader::EValueType::Bool1, TEXT("GetReflectionCapturePassSwitchState()"));
+		const FExpression* DefaultExpression = Default.AcquireHLSLExpression(Generator, Scope);
+		const FExpression* ReflectionExpression = Reflection.AcquireHLSLExpression(Generator, Scope);
+
+		OutExpression = Generator.GetTree().NewExpression<FExpressionSelect>(ConditionExpression, ReflectionExpression, DefaultExpression);
+		return true;
+	}
+}
+
 bool UMaterialExpressionQualitySwitch::GenerateHLSLExpression(FMaterialHLSLGenerator& Generator, UE::HLSLTree::FScope& Scope, int32 OutputIndex, UE::HLSLTree::FExpression const*& OutExpression) const
 {
 	using namespace UE::HLSLTree;
@@ -724,8 +756,7 @@ bool UMaterialExpressionCollectionParameter::GenerateHLSLExpression(FMaterialHLS
 
 	if (ParameterIndex != INDEX_NONE)
 	{
-		const int32 CollectionIndex = Generator.FindOrAddParameterCollection(Collection);
-		OutExpression = Generator.GetTree().NewExpression<UE::HLSLTree::Material::FExpressionCollectionParameter>(CollectionIndex, ParameterIndex);
+		OutExpression = Generator.GetTree().NewExpression<UE::HLSLTree::Material::FExpressionCollectionParameter>(Collection, ParameterIndex);
 		OutExpression = Generator.GetTree().NewExpression<UE::HLSLTree::FExpressionSwizzle>(
 			UE::HLSLTree::MakeSwizzleMask(
 				ComponentIndex == INDEX_NONE ? true : ComponentIndex % 4 == 0,
@@ -1249,6 +1280,13 @@ bool UMaterialExpressionParticleMacroUV::GenerateHLSLExpression(FMaterialHLSLGen
 	return true;
 }
 
+bool UMaterialExpressionParticleMotionBlurFade::GenerateHLSLExpression(FMaterialHLSLGenerator& Generator, UE::HLSLTree::FScope& Scope, int32 OutputIndex, UE::HLSLTree::FExpression const*& OutExpression) const
+{
+	using namespace UE::HLSLTree;
+	OutExpression = Generator.NewExternalInput(Material::EExternalInput::ParticleMotionBlurFade);
+	return true;
+}
+
 bool UMaterialExpressionTextureObject::GenerateHLSLExpression(FMaterialHLSLGenerator& Generator, UE::HLSLTree::FScope& Scope, int32 OutputIndex, UE::HLSLTree::FExpression const*& OutExpression) const
 {
 	using namespace UE::HLSLTree;
@@ -1856,6 +1894,51 @@ bool UMaterialExpressionSceneDepth::GenerateHLSLExpression(FMaterialHLSLGenerato
 	return true;
 }
 
+bool UMaterialExpressionSceneDepthWithoutWater::GenerateHLSLExpression(FMaterialHLSLGenerator& Generator, UE::HLSLTree::FScope& Scope, int32 OutputIndex, UE::HLSLTree::FExpression const*& OutExpression) const
+{
+	using namespace UE::HLSLTree;
+
+	const UMaterial* BaseMaterial = Generator.GetTargetMaterial();
+	if (BaseMaterial)
+	{
+		const EMaterialDomain MaterialDomain = BaseMaterial->MaterialDomain;
+
+		if (MaterialDomain != MD_PostProcess)
+		{
+			if (!BaseMaterial->GetShadingModels().HasShadingModel(MSM_SingleLayerWater))
+			{
+				return Generator.Error(TEXT("Can only read scene depth below water when material Shading Model is Single Layer Water or when material Domain is PostProcess."));
+			}
+
+			if (MaterialDomain != MD_Surface)
+			{
+				return Generator.Error(TEXT("Can only read scene depth below water when material Domain is set to Surface or PostProcess."));
+			}
+
+			if (IsTranslucentBlendMode(*BaseMaterial))
+			{
+				return Generator.Error(TEXT("Can only read scene depth below water when material Blend Mode isn't translucent."));
+			}
+		}
+	}
+
+	const FExpression* OffsetExpression = nullptr;
+	const FExpression* CoordinateExpression = nullptr;
+
+	if (InputMode == EMaterialSceneAttributeInputMode::OffsetFraction)
+	{
+		OffsetExpression = Input.AcquireHLSLExpressionOrConstant(Generator, Scope, FVector2f(ConstInput.X, ConstInput.Y));
+	}
+	else if (InputMode == EMaterialSceneAttributeInputMode::Coordinates)
+	{
+		CoordinateExpression = Input.TryAcquireHLSLExpression(Generator, Scope);
+	}
+
+	const FExpression* ScreenUVExpression = Generator.GetTree().NewExpression<Material::FExpressionScreenAlignedUV>(OffsetExpression, CoordinateExpression);
+	OutExpression = Generator.GetTree().NewExpression<Material::FExpressionSceneDepthWithoutWater>(ScreenUVExpression, FallbackDepth);
+	return true;
+}
+
 bool UMaterialExpressionDepthFade::GenerateHLSLExpression(FMaterialHLSLGenerator& Generator, UE::HLSLTree::FScope& Scope, int32 OutputIndex, UE::HLSLTree::FExpression const*& OutExpression) const
 {
 	using namespace UE::HLSLTree;
@@ -2516,11 +2599,9 @@ bool UMaterialExpressionComponentMask::GenerateHLSLExpression(FMaterialHLSLGener
 
 bool UMaterialExpressionGetMaterialAttributes::GenerateHLSLExpression(FMaterialHLSLGenerator& Generator, UE::HLSLTree::FScope& Scope, int32 OutputIndex, UE::HLSLTree::FExpression const*& OutExpression) const
 {
-	const UE::HLSLTree::FExpression* AttributesExpression = MaterialAttributes.AcquireHLSLExpression(Generator, Scope);
-	if (!AttributesExpression)
-	{
-		return false;
-	}
+	const UE::HLSLTree::FExpression* AttributesExpression = MaterialAttributes.AcquireHLSLExpressionOrConstant(Generator, Scope, Generator.GetMaterialAttributesDefaultValue());
+	check(AttributesExpression);
+
 	if (OutputIndex == 0)
 	{
 		OutExpression = AttributesExpression;
@@ -2554,12 +2635,15 @@ bool UMaterialExpressionSetMaterialAttributes::GenerateHLSLExpression(FMaterialH
 			// Only compile code to set attributes of the current shader frequency
 			const EShaderFrequency AttributeFrequency = FMaterialAttributeDefinitionMap::GetShaderFrequency(AttributeID);
 			const FString& AttributeName = FMaterialAttributeDefinitionMap::GetAttributeName(AttributeID);
-			//if (AttributeFrequency == Compiler->GetCurrentShaderFrequency())
+			const UE::Shader::FStructField* AttributeField = Generator.GetMaterialAttributesType()->FindFieldByName(*AttributeName);
+
+			// if (AttributeFrequency == Compiler->GetCurrentShaderFrequency())
+			// TODO: Currently we just ignore setting hidden material attributes. Should anything be done?
+			if (AttributeField)
 			{
 				const UE::HLSLTree::FExpression* ValueExpression = AttributeInput.TryAcquireHLSLExpression(Generator, Scope);
 				if (ValueExpression)
 				{
-					const UE::Shader::FStructField* AttributeField = Generator.GetMaterialAttributesType()->FindFieldByName(*AttributeName);
 					AttributesExpression = Generator.GetTree().NewExpression<UE::HLSLTree::FExpressionSetStructField>(Generator.GetMaterialAttributesType(), AttributeField, AttributesExpression, ValueExpression);
 				}
 			}
@@ -2702,23 +2786,23 @@ bool UMaterialExpressionMaterialAttributeLayers::GenerateHLSLExpression(FMateria
 
 	const FMaterialLayersFunctions* LayerOverrides = Generator.GetLayerOverrides();
 	const FMaterialLayersFunctions& MaterialLayers = LayerOverrides ? *LayerOverrides : DefaultLayers;
-	if (MaterialLayers.Layers.Num() == 0)
+	const int32 NumLayers = MaterialLayers.Layers.Num();
+	if (!NumLayers)
 	{
 		return Generator.Error(TEXT("No layers"));
 	}
 
 	TArray<FFunctionExpressionInput> FunctionInputs;
 	TArray<FFunctionExpressionOutput> FunctionOutputs;
-
 	const FExpression* ExpressionLayerInput = Input.AcquireHLSLExpressionOrConstant(Generator, Scope, Generator.GetMaterialAttributesDefaultValue());
-	TArray<FMaterialHLSLGenerator::FConnectedInput, TInlineAllocator<1>> LayerInputExpressions;
 
+	TArray<FMaterialHLSLGenerator::FConnectedInput, TInlineAllocator<1>> LayerInputExpressions;
 	TArray<const FExpression*, TInlineAllocator<16>> LayerExpressions;
-	LayerExpressions.Reserve(MaterialLayers.Layers.Num());
-	for (int32 LayerIndex = 0; LayerIndex < MaterialLayers.Layers.Num(); ++LayerIndex)
+
+	for (int32 LayerIndex = 0; LayerIndex < NumLayers; ++LayerIndex)
 	{
 		UMaterialFunctionInterface* LayerFunction = MaterialLayers.Layers[LayerIndex];
-		const FExpression* LayerExpression = nullptr;
+		
 		if (LayerFunction && MaterialLayers.EditorOnly.LayerStates[LayerIndex])
 		{
 			const EMaterialFunctionUsage Usage = LayerFunction->GetMaterialFunctionUsage();
@@ -2742,33 +2826,36 @@ bool UMaterialExpressionMaterialAttributeLayers::GenerateHLSLExpression(FMateria
 			{
 				LayerInputExpressions.Emplace(ExpressionLayerInput);
 			}
-			LayerExpression = Generator.GenerateFunctionCall(Scope, LayerFunction, LayerParameter, LayerIndex, LayerInputExpressions, 0);
+			const FExpression* LayerExpression = Generator.GenerateFunctionCall(Scope, LayerFunction, LayerParameter, LayerIndex, LayerInputExpressions, 0);
+			LayerExpressions.Add(LayerExpression);
+		}
+		else if (LayerIndex == 0)
+		{
+			LayerExpressions.Add(ExpressionLayerInput);
 		}
-		LayerExpressions.Add(LayerExpression);
 	}
 
+	check(LayerExpressions.Num() > 0);
 	const FExpression* BottomLayerExpression = LayerExpressions[0];
-	if (!BottomLayerExpression)
-	{
-		return Generator.Error(TEXT("No layers"));
-	}
+	check(BottomLayerExpression);
 
 	TArray<FMaterialHLSLGenerator::FConnectedInput, TInlineAllocator<2>> BlendInputExpressions;
-	for (int32 BlendIndex = 0; BlendIndex < MaterialLayers.Blends.Num(); ++BlendIndex)
+	const int32 NumBlends = MaterialLayers.Blends.Num();
+	int32 NumActiveBlends = 0;
+
+	for (int32 BlendIndex = 0; BlendIndex < NumBlends; ++BlendIndex)
 	{
 		const int32 LayerIndex = BlendIndex + 1;
 		if (!MaterialLayers.Layers.IsValidIndex(LayerIndex))
 		{
-			return Generator.Errorf(TEXT("Invalid number of layers (%d) and blends (%d)"), MaterialLayers.Layers.Num(), MaterialLayers.Blends.Num());
+			return Generator.Errorf(TEXT("Invalid number of layers (%d) and blends (%d)"), NumLayers, NumBlends);
 		}
 
 		if (MaterialLayers.Layers[LayerIndex] && MaterialLayers.EditorOnly.LayerStates[LayerIndex])
 		{
-			const FExpression* LayerExpression = LayerExpressions[LayerIndex];
-			if (!LayerExpression)
-			{
-				return Generator.Errorf(TEXT("Missing layer %d"), LayerIndex);
-			}
+			const int32 ActiveLayerIndex = ++NumActiveBlends;
+			const FExpression* LayerExpression = LayerExpressions[ActiveLayerIndex];
+			check(LayerExpression);
 
 			UMaterialFunctionInterface* BlendFunction = MaterialLayers.Blends[BlendIndex];
 			if (BlendFunction)
@@ -3332,6 +3419,42 @@ bool UMaterialExpressionAtmosphericLightColor::GenerateHLSLExpression(FMaterialH
 	return true;
 }
 
+bool UMaterialExpressionSkyAtmosphereLightDiskLuminance::GenerateHLSLExpression(FMaterialHLSLGenerator& Generator, UE::HLSLTree::FScope& Scope, int32 OutputIndex, UE::HLSLTree::FExpression const*& OutExpression) const
+{
+	using namespace UE::HLSLTree;
+
+	FTree& Tree = Generator.GetTree();
+	const FExpression* CosHalfDiskRadiusExpression;
+	if (DiskAngularDiameterOverride.GetTracedInput().Expression)
+	{
+		const FExpression* DiskAngularDiameter = DiskAngularDiameterOverride.AcquireHLSLExpression(Generator, Scope);
+		const FExpression* HalfDiskRadius = Tree.NewMul(Tree.NewConstant(0.5f * float(UE_PI) / 180.0f), DiskAngularDiameter);
+		CosHalfDiskRadiusExpression = Tree.NewCos(HalfDiskRadius);
+	}
+	else
+	{
+		CosHalfDiskRadiusExpression = Tree.NewConstant(-1.f);
+	}
+
+	OutExpression = Tree.NewExpression<Material::FExpressionSkyAtmosphereLightDiskLuminance>(CosHalfDiskRadiusExpression, LightIndex);
+	return true;
+}
+
+bool UMaterialExpressionSkyAtmosphereAerialPerspective::GenerateHLSLExpression(FMaterialHLSLGenerator& Generator, UE::HLSLTree::FScope& Scope, int32 OutputIndex, UE::HLSLTree::FExpression const*& OutExpression) const
+{
+	using namespace UE::HLSLTree;
+	const FExpression* WorldPositionExpression = WorldPosition.AcquireHLSLExpressionOrExternalInput(Generator, Scope, Material::EExternalInput::WorldPosition);
+	OutExpression = Generator.GetTree().NewExpression<Material::FExpressionSkyAtmosphereAerialPerspective>(WorldPositionExpression);
+	return true;
+}
+
+bool UMaterialExpressionSkyAtmosphereViewLuminance::GenerateHLSLExpression(FMaterialHLSLGenerator& Generator, UE::HLSLTree::FScope& Scope, int32 OutputIndex, UE::HLSLTree::FExpression const*& OutExpression) const
+{
+	using namespace UE::HLSLTree;
+	OutExpression = Generator.NewExternalInput(Material::EExternalInput::SkyAtmosphereViewLuminance);
+	return true;
+}
+
 bool UMaterialExpressionAtmosphericLightVector::GenerateHLSLExpression(FMaterialHLSLGenerator& Generator, UE::HLSLTree::FScope& Scope, int32 OutputIndex, UE::HLSLTree::FExpression const*& OutExpression) const
 {
 	OutExpression = Generator.GetTree().NewExpression<UE::HLSLTree::FExpressionInlineCustomHLSL>(UE::Shader::EValueType::Float3, TEXT("MaterialExpressionAtmosphericLightVector(Parameters)"));
@@ -3962,4 +4085,125 @@ bool UMaterialExpressionSamplePhysicsIntegerField::GenerateHLSLExpression(FMater
 	return true;
 }
 
+bool UMaterialExpressionHsvToRgb::GenerateHLSLExpression(FMaterialHLSLGenerator& Generator, UE::HLSLTree::FScope& Scope, int32 OutputIndex, UE::HLSLTree::FExpression const*& OutExpression) const
+{
+	using namespace UE::HLSLTree;
+
+	FTree& Tree = Generator.GetTree();
+	const FExpression* Hsv = Input.AcquireHLSLExpression(Generator, Scope);
+	const FExpression* H = Tree.NewSwizzle(FSwizzleParameters(0), Hsv);
+	const FExpression* S = Tree.NewSwizzle(FSwizzleParameters(1), Hsv);
+	const FExpression* V = Tree.NewSwizzle(FSwizzleParameters(2), Hsv);
+	// W will be 0 if Hsv is a vec3
+	const FExpression* W = Tree.NewSwizzle(FSwizzleParameters(3), Hsv);
+
+	const FExpression* Rgb = Tree.NewAbs(Tree.NewSub(Tree.NewMul(H, Tree.NewConstant(6.f)), Tree.NewConstant(FVector3f(3.f, 2.f, 4.f))));
+	Rgb = Tree.NewSaturate(Tree.NewAdd(Tree.NewMul(Rgb, Tree.NewConstant(FVector3f(1.f, -1.f, -1.f))), Tree.NewConstant(FVector3f(-1.f, 2.f, 2.f))));
+	const FExpression* ConstantOne = Tree.NewConstant(1.f);
+	Rgb = Tree.NewSub(Rgb, ConstantOne);
+	Rgb = Tree.NewAdd(Tree.NewMul(Rgb, S), ConstantOne);
+	Rgb = Tree.NewMul(Rgb, V);
+
+	OutExpression = Tree.NewExpression<FExpressionAppend>(Rgb, W);
+	return true;
+}
+
+bool UMaterialExpressionRgbToHsv::GenerateHLSLExpression(FMaterialHLSLGenerator& Generator, UE::HLSLTree::FScope& Scope, int32 OutputIndex, UE::HLSLTree::FExpression const*& OutExpression) const
+{
+	using namespace UE::HLSLTree;
+
+	FTree& Tree = Generator.GetTree();
+
+	const FExpression* Rgb = Input.AcquireHLSLExpression(Generator, Scope);
+	const FExpression* R = Tree.NewSwizzle(FSwizzleParameters(0), Rgb);
+	const FExpression* G = Tree.NewSwizzle(FSwizzleParameters(1), Rgb);
+	const FExpression* B = Tree.NewSwizzle(FSwizzleParameters(2), Rgb);
+	// A will be 0 if Rgb is a vec3
+	const FExpression* A = Tree.NewSwizzle(FSwizzleParameters(3), Rgb);
+
+	const FExpression* P = Tree.NewExpression<FExpressionSelect>(
+		Tree.NewLess(G, B),
+		Tree.NewAppend(B, G, Tree.NewConstant(FVector2f(-1.f, 2.f / 3.f))),
+		Tree.NewAppend(G, B, Tree.NewConstant(FVector2f(0.f, -1.f / 3.f))));
+
+	const FExpression* Q = Tree.NewExpression<FExpressionSelect>(
+		Tree.NewLess(R, Tree.NewSwizzle(FSwizzleParameters(0), P)),
+		Tree.NewAppend(Tree.NewSwizzle(FSwizzleParameters(0, 1, 3), P), R),
+		Tree.NewAppend(R, Tree.NewSwizzle(FSwizzleParameters(1, 2, 0), P)));
+	const FExpression* Qx = Tree.NewSwizzle(FSwizzleParameters(0), Q);
+	const FExpression* Qy = Tree.NewSwizzle(FSwizzleParameters(1), Q);
+	const FExpression* Qz = Tree.NewSwizzle(FSwizzleParameters(2), Q);
+	const FExpression* Qw = Tree.NewSwizzle(FSwizzleParameters(3), Q);
+
+	const FExpression* Chroma = Tree.NewSub(Qx, Tree.NewMin(Qw, Qy));
+
+	const FExpression* Epsilon = Tree.NewConstant(1e-10f);
+	const FExpression* Hue = Tree.NewDiv(Tree.NewSub(Qw, Qy), Tree.NewAdd(Tree.NewMul(Tree.NewConstant(6.f), Chroma), Epsilon));
+	Hue = Tree.NewAbs(Tree.NewAdd(Hue, Qz));
+
+	const FExpression* S = Tree.NewDiv(Chroma, Tree.NewAdd(Qx, Epsilon));
+	
+	OutExpression = Tree.NewAppend(Hue, S, Qx, A);
+	return true;
+}
+
+UE::Shader::EValueType UMaterialExpressionSingleLayerWaterMaterialOutput::GetCustomOutputType(int32 OutputIndex) const
+{
+	using namespace UE::Shader;
+
+	if (OutputIndex == 0)
+	{
+		return EValueType::Float3;
+	}
+	else if (OutputIndex == 1)
+	{
+		return EValueType::Float3;
+	}
+	else if (OutputIndex == 2)
+	{
+		return EValueType::Float1;
+	}
+	else if (OutputIndex == 3)
+	{
+		return EValueType::Float1;
+	}
+	else
+	{
+		return EValueType::Void;
+	}
+}
+
+bool UMaterialExpressionSingleLayerWaterMaterialOutput::GenerateHLSLExpression(FMaterialHLSLGenerator& Generator, UE::HLSLTree::FScope& Scope, int32 OutputIndex, UE::HLSLTree::FExpression const*& OutExpression) const
+{
+	const bool bSubstrate = Substrate::IsSubstrateEnabled();
+
+	if (!ScatteringCoefficients.IsConnected() && !AbsorptionCoefficients.IsConnected() && !PhaseG.IsConnected() && !bSubstrate)
+	{
+		return Generator.Error(TEXT("No inputs to Single Layer Water Material."));
+	}
+
+	if (OutputIndex == 0)
+	{
+		OutExpression = ScatteringCoefficients.AcquireHLSLExpressionOrConstant(Generator, Scope, FVector3f::ZeroVector);
+	}
+	else if (OutputIndex == 1)
+	{
+		OutExpression = AbsorptionCoefficients.AcquireHLSLExpressionOrConstant(Generator, Scope, FVector3f::ZeroVector);
+	}
+	else if (OutputIndex == 2)
+	{
+		OutExpression = PhaseG.AcquireHLSLExpressionOrConstant(Generator, Scope, 0.f);
+	}
+	else if (OutputIndex == 3)
+	{
+		OutExpression = ColorScaleBehindWater.AcquireHLSLExpressionOrConstant(Generator, Scope, 1.f);
+	}
+	else
+	{
+		OutExpression = nullptr;
+	}
+
+	return OutExpression != nullptr;
+}
+
 #endif // WITH_EDITOR
diff --git a/Engine/Source/Runtime/Engine/Private/Materials/MaterialHLSLEmitter.cpp b/Engine/Source/Runtime/Engine/Private/Materials/MaterialHLSLEmitter.cpp
--- a/Engine/Source/Runtime/Engine/Private/Materials/MaterialHLSLEmitter.cpp
+++ b/Engine/Source/Runtime/Engine/Private/Materials/MaterialHLSLEmitter.cpp
@@ -370,7 +370,9 @@ static void GetMaterialEnvironment(EShaderPlatform InPlatform,
 		OutEnvironment.SetDefine(TEXT("NEEDS_PARTICLE_TIME"), 1);
 	}
 
-	if (false)//bUsesParticleMotionBlur)
+	const bool bUsesParticleMotionBlur = EmitMaterialData.IsExternalInputUsed(Material::EExternalInput::ParticleMotionBlurFade);
+
+	if (bUsesParticleMotionBlur)
 	{
 		OutEnvironment.SetDefine(TEXT("USES_PARTICLE_MOTION_BLUR"), 1);
 	}
@@ -504,7 +506,7 @@ static void GetMaterialEnvironment(EShaderPlatform InPlatform,
 		OutEnvironment.SetDefine(PageTableName.ToString(), PageTableValue.ToString());
 	}
 
-	const TArray<UMaterialParameterCollection*>& ParameterCollections = InMaterial.GetCachedHLSLTree()->GetParameterCollections();
+	const TConstArrayView<const UMaterialParameterCollection*> ParameterCollections = EmitMaterialData.ParameterCollections;
 	for (int32 CollectionIndex = 0; CollectionIndex < ParameterCollections.Num(); CollectionIndex++)
 	{
 		// Add uniform buffer declarations for any parameter collections referenced
@@ -1005,7 +1007,7 @@ bool MaterialEmitHLSL(const FMaterialCompileTargetParameters& InCompilerTarget,
 			OutCompilationOutput);
 	}
 
-	OutCompilationOutput.UniformExpressionSet.SetParameterCollections(CachedTree->GetParameterCollections());
+	OutCompilationOutput.UniformExpressionSet.SetParameterCollections(EmitMaterialData.ParameterCollections);
 
 	OutMaterialEnvironment = new FSharedShaderCompilerEnvironment();
 	OutMaterialEnvironment->TargetPlatform = InCompilerTarget.TargetPlatform;
diff --git a/Engine/Source/Runtime/Engine/Private/Materials/MaterialHLSLGenerator.cpp b/Engine/Source/Runtime/Engine/Private/Materials/MaterialHLSLGenerator.cpp
--- a/Engine/Source/Runtime/Engine/Private/Materials/MaterialHLSLGenerator.cpp
+++ b/Engine/Source/Runtime/Engine/Private/Materials/MaterialHLSLGenerator.cpp
@@ -899,22 +899,4 @@ int32 FMaterialHLSLGenerator::FindOrAddCustomExpressionOutputStructId(TArrayView
 	}
 }
 
-int32 FMaterialHLSLGenerator::FindOrAddParameterCollection(UMaterialParameterCollection* ParameterCollection)
-{
-	int32 CollectionIndex = CachedTree.ParameterCollections.Find(ParameterCollection);
-
-	if (CollectionIndex == INDEX_NONE)
-	{
-		if (CachedTree.ParameterCollections.Num() >= MaxNumParameterCollectionsPerMaterial)
-		{
-			return Error(TEXT("Material references too many MaterialParameterCollections!  A material may only reference 2 different collections."));
-		}
-
-		CachedTree.ParameterCollections.Add(ParameterCollection);
-		CollectionIndex = CachedTree.ParameterCollections.Num() - 1;
-	}
-
-	return CollectionIndex;
-}
-
 #endif // WITH_EDITOR
diff --git a/Engine/Source/Runtime/Engine/Private/Materials/MaterialHLSLTree.cpp b/Engine/Source/Runtime/Engine/Private/Materials/MaterialHLSLTree.cpp
--- a/Engine/Source/Runtime/Engine/Private/Materials/MaterialHLSLTree.cpp
+++ b/Engine/Source/Runtime/Engine/Private/Materials/MaterialHLSLTree.cpp
@@ -154,10 +154,13 @@ FExternalInputDescription GetExternalInputDescription(EExternalInput Input)
 	case EExternalInput::ParticleSubUVCoords0: return FExternalInputDescription(TEXT("ParticleSubUVCoords0"), Shader::EValueType::Float2);
 	case EExternalInput::ParticleSubUVCoords1: return FExternalInputDescription(TEXT("ParticleSubUVCoords1"), Shader::EValueType::Float2);
 	case EExternalInput::ParticleSubUVLerp: return FExternalInputDescription(TEXT("ParticleSubUVLerp"), Shader::EValueType::Float1);
+	case EExternalInput::ParticleMotionBlurFade: return FExternalInputDescription(TEXT("ParticleMotionBlurFade"), Shader::EValueType::Float1);
 
 	case EExternalInput::PerInstanceFadeAmount: return FExternalInputDescription(TEXT("PerInstanceFadeAmount"), Shader::EValueType::Float1);
 	case EExternalInput::PerInstanceRandom: return FExternalInputDescription(TEXT("PerInstanceRandom"), Shader::EValueType::Float1);
 
+	case EExternalInput::SkyAtmosphereViewLuminance: return FExternalInputDescription(TEXT("SkyAtmosphereViewLuminance"), Shader::EValueType::Float3);
+
 	case EExternalInput::IsOrthographic: return FExternalInputDescription(TEXT("IsOrthographic"), Shader::EValueType::Float1);
 
 	default: checkNoEntry(); return FExternalInputDescription(TEXT("Invalid"), Shader::EValueType::Void);
@@ -233,6 +236,8 @@ bool FExpressionExternalInput::PrepareValue(FEmitContext& Context, FEmitScope& S
 			case EExternalInput::PerInstanceRandom:
 				Context.MaterialCompilationOutput->bUsesPerInstanceRandom = true;
 				break;
+			case EExternalInput::SkyAtmosphereViewLuminance:
+				Context.bUsesSkyAtmosphere = true;
 			default:
 				break;
 			}
@@ -368,10 +373,13 @@ void FExpressionExternalInput::EmitValueShader(FEmitContext& Context, FEmitScope
 		case EExternalInput::ParticleSubUVCoords0: Code = TEXT("Parameters.Particle.SubUVCoords[0].xy"); break;
 		case EExternalInput::ParticleSubUVCoords1: Code = TEXT("Parameters.Particle.SubUVCoords[1].xy"); break;
 		case EExternalInput::ParticleSubUVLerp: Code = TEXT("Parameters.Particle.SubUVLerp"); break;
+		case EExternalInput::ParticleMotionBlurFade: Code = TEXT("Parameters.Particle.MotionBlurFade"); break;
 
 		case EExternalInput::PerInstanceFadeAmount: Code = TEXT("GetPerInstanceFadeAmount(Parameters)"); break;
 		case EExternalInput::PerInstanceRandom: Code = TEXT("GetPerInstanceRandom(Parameters)"); break;
 
+		case EExternalInput::SkyAtmosphereViewLuminance: Code = TEXT("MaterialExpressionSkyAtmosphereViewLuminance(Parameters)"); break;
+
 		case EExternalInput::IsOrthographic: Code = TEXT("((View.ViewToClip[3][3] < 1.0f) ? 0.0f : 1.0f)"); break;
 
 		default:
@@ -661,11 +669,24 @@ void FExpressionCollectionParameter::ComputeAnalyticDerivatives(FTree& Tree, FEx
 
 bool FExpressionCollectionParameter::PrepareValue(FEmitContext& Context, FEmitScope& Scope, const FRequestedType& RequestedType, FPrepareValueResult& OutResult) const
 {
+	if (Context.bMarkLiveValues && !Context.TargetParameters.IsGenericTarget())
+	{
+		FEmitData& EmitMaterialData = Context.FindData<FEmitData>();
+		const int32 CollectionIndex = EmitMaterialData.FindOrAddParameterCollection(ParameterCollection);
+
+		if (CollectionIndex == INDEX_NONE)
+		{
+			return Context.Error(TEXT("Material references too many MaterialParameterCollections! A material may only reference 2 different collections."));
+		}
+	}
 	return OutResult.SetType(Context, RequestedType, EExpressionEvaluation::Shader, Shader::EValueType::Float4);
 }
 
 void FExpressionCollectionParameter::EmitValueShader(FEmitContext& Context, FEmitScope& Scope, const FRequestedType& RequestedType, FEmitValueShaderResult& OutResult) const
 {
+	FEmitData& EmitMaterialData = Context.FindData<FEmitData>();
+	const int32 CollectionIndex = EmitMaterialData.ParameterCollections.Find(ParameterCollection);
+	check(CollectionIndex != INDEX_NONE);
 	OutResult.Code = Context.EmitInlineExpression(Scope, Shader::EValueType::Float4, TEXT("MaterialCollection%.Vectors[%]"), CollectionIndex, ParameterIndex);
 }
 
@@ -1599,6 +1620,44 @@ void FExpressionSceneDepth::EmitValueShader(FEmitContext& Context, FEmitScope& S
 	OutResult.Code = Context.EmitExpression(Scope, Shader::EValueType::Float1, TEXT("CalcSceneDepth(%)"), EmitScreenUV);
 }
 
+bool FExpressionSceneDepthWithoutWater::PrepareValue(FEmitContext& Context, FEmitScope& Scope, const FRequestedType& RequestedType, FPrepareValueResult& OutResult) const
+{
+	if (Context.ShaderFrequency == SF_Vertex)
+	{
+		// Mobile currently does not support this, we need to read a separate copy of the depth, we must disable framebuffer fetch and force scene texture reads.
+		// (Texture bindings are not setup properly for any platform so we're disallowing usage in vertex shader altogether now)
+		return Context.Error(TEXT("Cannot read scene depth without water from the vertex shader."));
+	}
+
+	// Need to check again since material instances can override shading models and blend modes
+	if (Context.Material)
+	{
+		if (!Context.Material->GetShadingModels().HasShadingModel(MSM_SingleLayerWater))
+		{
+			return Context.Error(TEXT("Can only read scene depth below water when material Shading Model is Single Layer Water or when material Domain is PostProcess."));
+		}
+
+		if (IsTranslucentBlendMode(*Context.Material))
+		{
+			return Context.Error(TEXT("Can only read scene depth below water when material Blend Mode isn't translucent."));
+		}
+	}
+
+	const FPreparedType& ScreenUVType = Context.PrepareExpression(ScreenUVExpression, Scope, Shader::EValueType::Float2);
+	if (ScreenUVType.IsVoid())
+	{
+		return false;
+	}
+
+	return OutResult.SetType(Context, RequestedType, EExpressionEvaluation::Shader, Shader::EValueType::Float1);
+}
+
+void FExpressionSceneDepthWithoutWater::EmitValueShader(FEmitContext& Context, FEmitScope& Scope, const FRequestedType& RequestedType, FEmitValueShaderResult& OutResult) const
+{
+	FEmitShaderExpression* EmitScreenUV = ScreenUVExpression->GetValueShader(Context, Scope, Shader::EValueType::Float2);
+	OutResult.Code = Context.EmitExpression(Scope, Shader::EValueType::Float1, TEXT("MaterialExpressionSceneDepthWithoutWater(%, %)"), EmitScreenUV, FallbackDepth);
+}
+
 bool FExpressionSceneColor::PrepareValue(FEmitContext& Context, FEmitScope& Scope, const FRequestedType& RequestedType, FPrepareValueResult& OutResult) const
 {
 	if (Context.ShaderFrequency != SF_Pixel)
@@ -1794,7 +1853,15 @@ void FExpressionVertexInterpolator::EmitValueShader(FEmitContext& Context, FEmit
 			}
 		}
 
-		FEmitShaderExpression* EmitPreshader = Context.EmitPreshaderOrConstant(Scope, RequestedPreshaderType, LocalType, VertexExpression);
+		FEmitShaderExpression* EmitPreshader;
+		if (RequestedPreshaderType.IsEmpty())
+		{
+			EmitPreshader = Context.EmitConstantZero(Scope, LocalType);
+		}
+		else
+		{
+			EmitPreshader = Context.EmitPreshaderOrConstant(Scope, RequestedPreshaderType, LocalType, VertexExpression);
+		}
 		OutResult.Code = Context.EmitExpression(Scope, LocalType, TEXT("MaterialVertexInterpolator%(Parameters, %)"), InterpolatorIndex, EmitPreshader);
 	}
 	else
@@ -1812,15 +1879,48 @@ void FExpressionVertexInterpolator::EmitValuePreshader(FEmitContext& Context, FE
 
 bool FExpressionSkyAtmosphereLightDirection::PrepareValue(FEmitContext& Context, FEmitScope& Scope, const FRequestedType& RequestedType, FPrepareValueResult& OutResult) const
 {
+	if (Context.bMarkLiveValues)
+	{
+		Context.bUsesSkyAtmosphere = true;
+	}
 	return OutResult.SetType(Context, RequestedType, EExpressionEvaluation::Shader, Shader::EValueType::Float3);
 }
 
 void FExpressionSkyAtmosphereLightDirection::EmitValueShader(FEmitContext& Context, FEmitScope& Scope, const FRequestedType& RequestedType, FEmitValueShaderResult& OutResult) const
 {
-	Context.bUsesSkyAtmosphere = true;
 	OutResult.Code = Context.EmitInlineExpression(Scope, Shader::EValueType::Float3, TEXT("MaterialExpressionSkyAtmosphereLightDirection(Parameters, %)"), LightIndex);
 }
 
+bool FExpressionSkyAtmosphereLightDiskLuminance::PrepareValue(FEmitContext& Context, FEmitScope& Scope, const FRequestedType& RequestedType, FPrepareValueResult& OutResult) const
+{
+	if (Context.bMarkLiveValues)
+	{
+		Context.bUsesSkyAtmosphere = true;
+	}
+	return OutResult.SetType(Context, RequestedType, EExpressionEvaluation::Shader, Shader::EValueType::Float3);
+}
+
+void FExpressionSkyAtmosphereLightDiskLuminance::EmitValueShader(FEmitContext& Context, FEmitScope& Scope, const FRequestedType& RequestedType, FEmitValueShaderResult& OutResult) const
+{
+	FEmitShaderExpression* EmitCosHalfDiskRadius = CosHalfDiskRadiusExpression->GetValueShader(Context, Scope, Shader::EValueType::Float1);
+	OutResult.Code = Context.EmitInlineExpression(Scope, Shader::EValueType::Float3, TEXT("MaterialExpressionSkyAtmosphereLightDiskLuminance(Parameters, %, %)"), LightIndex, EmitCosHalfDiskRadius);
+}
+
+bool FExpressionSkyAtmosphereAerialPerspective::PrepareValue(FEmitContext& Context, FEmitScope& Scope, const FRequestedType& RequestedType, FPrepareValueResult& OutResult) const
+{
+	if (Context.bMarkLiveValues)
+	{
+		Context.bUsesSkyAtmosphere = true;
+	}
+	return OutResult.SetType(Context, RequestedType, EExpressionEvaluation::Shader, Shader::EValueType::Float4);
+}
+
+void FExpressionSkyAtmosphereAerialPerspective::EmitValueShader(FEmitContext& Context, FEmitScope& Scope, const FRequestedType& RequestedType, FEmitValueShaderResult& OutResult) const
+{
+	FEmitShaderExpression* EmitWorldPosition = WorldPositionExpression->GetValueShader(Context, Scope, Shader::EValueType::Double3);
+	OutResult.Code = Context.EmitInlineExpression(Scope, Shader::EValueType::Float4, TEXT("MaterialExpressionSkyAtmosphereAerialPerspective(Parameters, %)"), EmitWorldPosition);
+}
+
 namespace Private
 {
 	bool PlatformSupportDistanceFields(FEmitContext& Context)
@@ -2107,7 +2207,7 @@ bool FExpressionNaniteReplaceFunction::PrepareValue(FEmitContext& Context, FEmit
 	}
 
 	// skip preparing if platform doesn't support Nanite
-	if (FDataDrivenShaderPlatformInfo::GetSupportsNanite(Context.TargetParameters.ShaderPlatform) || Context.TargetParameters.IsGenericTarget())
+	if (Context.TargetParameters.IsGenericTarget() || FDataDrivenShaderPlatformInfo::GetSupportsNanite(Context.TargetParameters.ShaderPlatform))
 	{
 		const FPreparedType& NaniteType = Context.PrepareExpression(NaniteExpression, Scope, RequestedType);
 		if (NaniteType.IsVoid())
@@ -2386,6 +2486,24 @@ void FEmitData::EmitInterpolatorShader(FEmitContext& Context, FStringBuilderBase
 	NumInterpolatorComponents = InterpolatorOffset;
 }
 
+int32 FEmitData::FindOrAddParameterCollection(const class UMaterialParameterCollection* ParameterCollection)
+{
+	int32 CollectionIndex = ParameterCollections.Find(ParameterCollection);
+
+	if (CollectionIndex == INDEX_NONE)
+	{
+		if (ParameterCollections.Num() >= MaxNumParameterCollectionsPerMaterial)
+		{
+			return INDEX_NONE;
+		}
+
+		ParameterCollections.Add(ParameterCollection);
+		CollectionIndex = ParameterCollections.Num() - 1;
+	}
+
+	return CollectionIndex;
+}
+
 } // namespace UE::HLSLTree::Material
 
 #endif // WITH_EDITOR
diff --git a/Engine/Source/Runtime/Engine/Private/Materials/MaterialInstance.cpp b/Engine/Source/Runtime/Engine/Private/Materials/MaterialInstance.cpp
--- a/Engine/Source/Runtime/Engine/Private/Materials/MaterialInstance.cpp
+++ b/Engine/Source/Runtime/Engine/Private/Materials/MaterialInstance.cpp
@@ -1942,24 +1942,6 @@ void UMaterialInterface::GetStaticParameterValues(FStaticParameterSet& OutStatic
 }
 #endif // WITH_EDITORONLY_DATA
 
-template<typename TArrayType>
-static void RemapLayersForParent(TArrayType& LayerIndexRemap, int32 NumParentLayers, TArrayView<const int32> ParentLayerIndexRemap)
-{
-	TArrayType NewLayerIndexRemap;
-	NewLayerIndexRemap.Init(INDEX_NONE, NumParentLayers);
-
-	check(LayerIndexRemap.Num() == ParentLayerIndexRemap.Num());
-	for (int32 i = 0; i < ParentLayerIndexRemap.Num(); ++i)
-	{
-		const int32 ParentLayerIndex = ParentLayerIndexRemap[i];
-		if (ParentLayerIndex != INDEX_NONE)
-		{
-			NewLayerIndexRemap[ParentLayerIndex] = LayerIndexRemap[i];
-		}
-	}
-	LayerIndexRemap = MoveTemp(NewLayerIndexRemap);
-}
-
 void UMaterialInstance::GetAllParametersOfType(EMaterialParameterType Type, TMap<FMaterialParameterInfo, FMaterialParameterMetadata>& OutParameters) const
 {
 	FMaterialInheritanceChain InstanceChain;
diff --git a/Engine/Source/Runtime/Engine/Private/Materials/MaterialInstanceConstant.cpp b/Engine/Source/Runtime/Engine/Private/Materials/MaterialInstanceConstant.cpp
--- a/Engine/Source/Runtime/Engine/Private/Materials/MaterialInstanceConstant.cpp
+++ b/Engine/Source/Runtime/Engine/Private/Materials/MaterialInstanceConstant.cpp
@@ -215,7 +215,7 @@ void UMaterialInstanceConstant::UpdateCachedData()
 		TUniquePtr<FMaterialCachedExpressionData> LocalCachedExpressionData;
 		TUniquePtr<FMaterialCachedHLSLTree> LocalCachedTree;
 
-		// If we have overriden material layers, need to create a local cached expression data
+		// If we have overridden material layers, need to create a local cached expression data
 		// Otherwise we can leave it as null, and use cached data from our parent
 		const FStaticParameterSet& LocalStaticParameters = GetStaticParameters();
 		if (LocalStaticParameters.bHasMaterialLayers)
@@ -242,9 +242,62 @@ void UMaterialInstanceConstant::UpdateCachedData()
 
 		if (bUsingNewHLSLGenerator && bHasStaticPermutationResource)
 		{
+			// Find the values of all overridden parameters including those overridden by parent instances. This is needed
+			// for correct HLSL tree preparation where static parameter values can affect the final result. We cannot call
+			// GetStaticParameterValues() here because it uses CachedExpressionData which is not ready yet
+			FStaticParameterSet OverriddenStaticParameters;
+#if WITH_EDITORONLY_DATA
+			{
+				FMaterialInheritanceChain InstanceChain;
+				GetMaterialInheritanceChain(InstanceChain);
+
+				// Maps a layer in parent to a layer in this instance
+				TArray<int32, TInlineAllocator<16>> LayerIndexRemap;
+				LayerIndexRemap.Empty(GetCachedInstanceData().ParentLayerIndexRemap.Num());
+				for (int32 LayerIndex = 0; LayerIndex < GetCachedInstanceData().ParentLayerIndexRemap.Num(); ++LayerIndex)
+				{
+					LayerIndexRemap.Add(LayerIndex);
+				}
+
+				TMap<FMaterialParameterInfo, FMaterialParameterMetadata> OverriddenStaticParametersMap[2];
+				TSet<FMaterialParameterInfo> SeenParameters[2];
+				for (int32 InstanceIndex = 0; InstanceIndex < InstanceChain.MaterialInstances.Num(); ++InstanceIndex)
+				{
+					const UMaterialInstance* Instance = InstanceChain.MaterialInstances[InstanceIndex];
+					const bool bSetOverride = InstanceIndex == 0;
+
+					if (InstanceIndex > 0)
+					{
+						const UMaterialInstance* ChildInstance = InstanceChain.MaterialInstances[InstanceIndex - 1];
+						const int32 NumLayers = Instance->GetCachedInstanceData().ParentLayerIndexRemap.Num();
+						RemapLayersForParent(LayerIndexRemap, NumLayers, ChildInstance->GetCachedInstanceData().ParentLayerIndexRemap);
+					}
+
+					const FStaticParameterSet& InstanceStaticParameters = Instance->GetStaticParameters();
+					GameThread_ApplyParameterOverrides(
+						InstanceStaticParameters.StaticSwitchParameters,
+						LayerIndexRemap,
+						bSetOverride,
+						SeenParameters[0],
+						OverriddenStaticParametersMap[0],
+						true);
+					GameThread_ApplyParameterOverrides(
+						InstanceStaticParameters.EditorOnly.StaticComponentMaskParameters,
+						LayerIndexRemap,
+						bSetOverride,
+						SeenParameters[1],
+						OverriddenStaticParametersMap[1],
+						true);
+				}
+
+				OverriddenStaticParameters.AddParametersOfType(EMaterialParameterType::StaticSwitch, OverriddenStaticParametersMap[0]);
+				OverriddenStaticParameters.AddParametersOfType(EMaterialParameterType::StaticComponentMask, OverriddenStaticParametersMap[1]);
+			}
+#endif // WITH_EDITORONLY_DATA
+
 			check(!LocalCachedExpressionData);
 			LocalCachedExpressionData.Reset(new FMaterialCachedExpressionData());
-			LocalCachedExpressionData->UpdateForCachedHLSLTree(GetCachedHLSLTree(), &LocalStaticParameters);
+			LocalCachedExpressionData->UpdateForCachedHLSLTree(GetCachedHLSLTree(), &OverriddenStaticParameters);
 		}
 
 		CachedExpressionData = MoveTemp(LocalCachedExpressionData);
diff --git a/Engine/Source/Runtime/Engine/Private/Materials/MaterialInstanceSupport.h b/Engine/Source/Runtime/Engine/Private/Materials/MaterialInstanceSupport.h
--- a/Engine/Source/Runtime/Engine/Private/Materials/MaterialInstanceSupport.h
+++ b/Engine/Source/Runtime/Engine/Private/Materials/MaterialInstanceSupport.h
@@ -513,11 +513,13 @@ inline bool GameThread_GetParameterValue(const TArray<ParameterType>& Parameters
 }
 
 template <typename ParameterType, typename OverridenParametersType>
-inline void GameThread_ApplyParameterOverrides(const TArray<ParameterType>& Parameters,
+inline void GameThread_ApplyParameterOverrides(
+	const TArray<ParameterType>& Parameters,
 	TArrayView<const int32> LayerIndexRemap,
 	bool bSetOverride,
 	OverridenParametersType& OverridenParameters, // TSet<FMaterialParameterInfo, ...>
-	TMap<FMaterialParameterInfo, FMaterialParameterMetadata>& OutParameters)
+	TMap<FMaterialParameterInfo, FMaterialParameterMetadata>& OutParameters,
+	bool bAddIfNotFound = false)
 {
 	for (int32 ParameterIndex = 0; ParameterIndex < Parameters.Num(); ParameterIndex++)
 	{
@@ -531,7 +533,7 @@ inline void GameThread_ApplyParameterOverrides(const TArray<ParameterType>& Para
 				OverridenParameters.Add(ParameterInfo, &bPreviouslyOverriden);
 				if (!bPreviouslyOverriden)
 				{
-					FMaterialParameterMetadata* Result = OutParameters.Find(ParameterInfo);
+					FMaterialParameterMetadata* Result = bAddIfNotFound ? &OutParameters.FindOrAdd(ParameterInfo) : OutParameters.Find(ParameterInfo);
 					if (Result)
 					{
 						Parameter->GetValue(*Result);
@@ -547,3 +549,21 @@ inline void GameThread_ApplyParameterOverrides(const TArray<ParameterType>& Para
 		}
 	}
 }
+
+template<typename TArrayType>
+inline void RemapLayersForParent(TArrayType& LayerIndexRemap, int32 NumParentLayers, TArrayView<const int32> ParentLayerIndexRemap)
+{
+	TArrayType NewLayerIndexRemap;
+	NewLayerIndexRemap.Init(INDEX_NONE, NumParentLayers);
+
+	check(LayerIndexRemap.Num() == ParentLayerIndexRemap.Num());
+	for (int32 i = 0; i < ParentLayerIndexRemap.Num(); ++i)
+	{
+		const int32 ParentLayerIndex = ParentLayerIndexRemap[i];
+		if (ParentLayerIndex != INDEX_NONE)
+		{
+			NewLayerIndexRemap[ParentLayerIndex] = LayerIndexRemap[i];
+		}
+	}
+	LayerIndexRemap = MoveTemp(NewLayerIndexRemap);
+}
diff --git a/Engine/Source/Runtime/Engine/Private/Materials/MaterialUniformExpressions.cpp b/Engine/Source/Runtime/Engine/Private/Materials/MaterialUniformExpressions.cpp
--- a/Engine/Source/Runtime/Engine/Private/Materials/MaterialUniformExpressions.cpp
+++ b/Engine/Source/Runtime/Engine/Private/Materials/MaterialUniformExpressions.cpp
@@ -339,7 +339,7 @@ FString FUniformExpressionSet::GetSummaryString() const
 		);
 }
 
-void FUniformExpressionSet::SetParameterCollections(const TArray<UMaterialParameterCollection*>& InCollections)
+void FUniformExpressionSet::SetParameterCollections(TConstArrayView<const UMaterialParameterCollection*> InCollections)
 {
 	ParameterCollections.Empty(InCollections.Num());
 
diff --git a/Engine/Source/Runtime/Engine/Private/Shader/ShaderValue.cpp b/Engine/Source/Runtime/Engine/Private/Shader/ShaderValue.cpp
--- a/Engine/Source/Runtime/Engine/Private/Shader/ShaderValue.cpp
+++ b/Engine/Source/Runtime/Engine/Private/Shader/ShaderValue.cpp
@@ -171,7 +171,7 @@ EValueType FType::GetFlatFieldType(int32 Index) const
 	}
 }
 
-FType CombineTypes(const FType& Lhs, const FType& Rhs)
+FType CombineTypes(const FType& Lhs, const FType& Rhs, bool bMergeMatrixTypes)
 {
 	if (Lhs.IsVoid() || Lhs.IsAny())
 	{
@@ -182,7 +182,7 @@ FType CombineTypes(const FType& Lhs, const FType& Rhs)
 		return Lhs;
 	}
 
-	if (Lhs.IsNumericVector() && Rhs.IsNumericVector())
+	if ((Lhs.IsNumericVector() && Rhs.IsNumericVector()) || (bMergeMatrixTypes && Lhs.IsNumericMatrix() && Rhs.IsNumericMatrix()))
 	{
 		const FValueTypeDescription& LhsDesc = GetValueTypeDescription(Lhs);
 		const FValueTypeDescription& RhsDesc = GetValueTypeDescription(Rhs);
diff --git a/Engine/Source/Runtime/Engine/Public/HLSLTree/HLSLTree.h b/Engine/Source/Runtime/Engine/Public/HLSLTree/HLSLTree.h
--- a/Engine/Source/Runtime/Engine/Public/HLSLTree/HLSLTree.h
+++ b/Engine/Source/Runtime/Engine/Public/HLSLTree/HLSLTree.h
@@ -678,6 +678,14 @@ public:
 
 	const FExpression* NewTruncateLWC(const FExpression* Input) { return NewUnaryOp(EOperation::TruncateLWC, Input); }
 
+	const FExpression* NewAppend(const FExpression* Lhs, const FExpression* Rhs);
+
+	template <typename... ArgTypes>
+	const FExpression* NewAppend(const FExpression* Lhs, ArgTypes&&... Rhs)
+	{
+		return NewAppend(Lhs, NewAppend(Forward<ArgTypes>(Rhs)...));
+	}
+
 	FActiveStructFieldStack ActiveStructFieldStack;
 
 private:
diff --git a/Engine/Source/Runtime/Engine/Public/HLSLTree/HLSLTreeCommon.h b/Engine/Source/Runtime/Engine/Public/HLSLTree/HLSLTreeCommon.h
--- a/Engine/Source/Runtime/Engine/Public/HLSLTree/HLSLTreeCommon.h
+++ b/Engine/Source/Runtime/Engine/Public/HLSLTree/HLSLTreeCommon.h
@@ -159,6 +159,7 @@ public:
 	virtual bool PrepareValue(FEmitContext& Context, FEmitScope& Scope, const FRequestedType& RequestedType, FPrepareValueResult& OutResult) const override;
 	virtual void EmitValueShader(FEmitContext& Context, FEmitScope& Scope, const FRequestedType& RequestedType, FEmitValueShaderResult& OutResult) const override;
 	virtual void EmitValuePreshader(FEmitContext& Context, FEmitScope& Scope, const FRequestedType& RequestedType, FEmitValuePreshaderResult& OutResult) const override;
+	virtual bool EmitValueObject(FEmitContext& Context, FEmitScope& Scope, const FName& ObjectTypeName, void* OutObjectBase) const override;
 };
 
 class FExpressionDerivative : public FExpression
diff --git a/Engine/Source/Runtime/Engine/Public/MaterialCachedHLSLTree.h b/Engine/Source/Runtime/Engine/Public/MaterialCachedHLSLTree.h
--- a/Engine/Source/Runtime/Engine/Public/MaterialCachedHLSLTree.h
+++ b/Engine/Source/Runtime/Engine/Public/MaterialCachedHLSLTree.h
@@ -63,8 +63,6 @@ public:
 	const UE::Shader::FStructType* GetVTPageTableResultType() const { return VTPageTableResultType; }
 	const UE::Shader::FValue& GetMaterialAttributesDefaultValue() const { return MaterialAttributesDefaultValue; }
 
-	const TArray<UMaterialParameterCollection*>& GetParameterCollections() const { return ParameterCollections; }
-
 	const TMap<FMaterialConnectionKey, const UE::HLSLTree::FExpression*>& GetConnections() const { return ConnectionMap; }
 
 	ENGINE_API void SetRequestedFields(EShaderFrequency ShaderFrequency, UE::HLSLTree::FRequestedType& OutRequestedType) const;
@@ -83,7 +81,6 @@ private:
 	UE::HLSLTree::FScope* ResultScope = nullptr;
 
 	TArray<UMaterialExpressionCustomOutput*> MaterialCustomOutputs;
-	TArray<UMaterialParameterCollection*> ParameterCollections;
 	TMap<FMaterialConnectionKey, const UE::HLSLTree::FExpression*> ConnectionMap;
 	const UE::Shader::FStructType* MaterialAttributesType = nullptr;
 	const UE::Shader::FStructType* VTPageTableResultType = nullptr;
diff --git a/Engine/Source/Runtime/Engine/Public/MaterialHLSLGenerator.h b/Engine/Source/Runtime/Engine/Public/MaterialHLSLGenerator.h
--- a/Engine/Source/Runtime/Engine/Public/MaterialHLSLGenerator.h
+++ b/Engine/Source/Runtime/Engine/Public/MaterialHLSLGenerator.h
@@ -210,8 +210,6 @@ public:
 	FMaterialParameterInfo GetParameterInfo(const FName& ParameterName) const;
 
 	int32 FindOrAddCustomExpressionOutputStructId(TArrayView<UE::Shader::FStructFieldInitializer> StructFields);
-	
-	int32 FindOrAddParameterCollection(UMaterialParameterCollection* ParameterCollection);
 
 private:
 	static constexpr int32 MaxNumPreviousScopes = UE::HLSLTree::MaxNumPreviousScopes;
diff --git a/Engine/Source/Runtime/Engine/Public/MaterialHLSLTree.h b/Engine/Source/Runtime/Engine/Public/MaterialHLSLTree.h
--- a/Engine/Source/Runtime/Engine/Public/MaterialHLSLTree.h
+++ b/Engine/Source/Runtime/Engine/Public/MaterialHLSLTree.h
@@ -11,6 +11,7 @@
 #include "VT/RuntimeVirtualTextureEnum.h"
 #include "Field/FieldSystemTypes.h"
 #include "MaterialCompiler.h"
+#include "ParameterCollection.h"
 
 class UTexture;
 enum class EMaterialParameterType : uint8;
@@ -192,10 +193,13 @@ enum class EExternalInput : uint8
 	ParticleSubUVCoords0,
 	ParticleSubUVCoords1,
 	ParticleSubUVLerp,
+	ParticleMotionBlurFade,
 
 	PerInstanceFadeAmount,
 	PerInstanceRandom,
 
+	SkyAtmosphereViewLuminance,
+
 	IsOrthographic,
 
 	Num,
@@ -289,11 +293,11 @@ public:
 class FExpressionCollectionParameter : public FExpression
 {
 public:
-	int32 CollectionIndex;
+	const class UMaterialParameterCollection* ParameterCollection;
 	int32 ParameterIndex;
 
-	FExpressionCollectionParameter(int32 InCollectionIndex, int32 InParameterIndex)
-		: CollectionIndex(InCollectionIndex)
+	FExpressionCollectionParameter(const class UMaterialParameterCollection* InParameterCollection, int32 InParameterIndex)
+		: ParameterCollection(InParameterCollection)
 		, ParameterIndex(InParameterIndex)
 	{}
 
@@ -578,6 +582,21 @@ public:
 	virtual void EmitValueShader(FEmitContext& Context, FEmitScope& Scope, const FRequestedType& RequestedType, FEmitValueShaderResult& OutResult) const override;
 };
 
+class FExpressionSceneDepthWithoutWater : public FExpression
+{
+public:
+	FExpressionSceneDepthWithoutWater(const FExpression* InScreenUVExpression, float InFallbackDepth)
+		: ScreenUVExpression(InScreenUVExpression)
+		, FallbackDepth(InFallbackDepth)
+	{}
+
+	const FExpression* ScreenUVExpression;
+	float FallbackDepth;
+
+	virtual bool PrepareValue(FEmitContext& Context, FEmitScope& Scope, const FRequestedType& RequestedType, FPrepareValueResult& OutResult) const override;
+	virtual void EmitValueShader(FEmitContext& Context, FEmitScope& Scope, const FRequestedType& RequestedType, FEmitValueShaderResult& OutResult) const override;
+};
+
 class FExpressionSceneColor : public FExpression
 {
 public:
@@ -683,6 +702,34 @@ public:
 	virtual void EmitValueShader(FEmitContext& Context, FEmitScope& Scope, const FRequestedType& RequestedType, FEmitValueShaderResult& OutResult) const override;
 };
 
+class FExpressionSkyAtmosphereLightDiskLuminance : public FExpression
+{
+public:
+	const FExpression* CosHalfDiskRadiusExpression;
+	int32 LightIndex;
+
+	FExpressionSkyAtmosphereLightDiskLuminance(const FExpression* InCosHalfDiskRaidusExpression, int32 InLightIndex)
+		: CosHalfDiskRadiusExpression(InCosHalfDiskRaidusExpression)
+		, LightIndex(InLightIndex)
+	{}
+
+	virtual bool PrepareValue(FEmitContext& Context, FEmitScope& Scope, const FRequestedType& RequestedType, FPrepareValueResult& OutResult) const override;
+	virtual void EmitValueShader(FEmitContext& Context, FEmitScope& Scope, const FRequestedType& RequestedType, FEmitValueShaderResult& OutResult) const override;
+};
+
+class FExpressionSkyAtmosphereAerialPerspective : public FExpression
+{
+public:
+	const FExpression* WorldPositionExpression;
+
+	FExpressionSkyAtmosphereAerialPerspective(const FExpression* InWorldPositionExpression)
+		: WorldPositionExpression(InWorldPositionExpression)
+	{}
+
+	virtual bool PrepareValue(FEmitContext& Context, FEmitScope& Scope, const FRequestedType& RequestedType, FPrepareValueResult& OutResult) const override;
+	virtual void EmitValueShader(FEmitContext& Context, FEmitScope& Scope, const FRequestedType& RequestedType, FEmitValueShaderResult& OutResult) const override;
+};
+
 class FExpressionDistanceToNearestSurface : public FExpression
 {
 public:
@@ -858,6 +905,7 @@ public:
 	TMap<Shader::FValue, uint32> DefaultUniformValues;
 	TArray<FVertexInterpolator, TInlineAllocator<8>> VertexInterpolators;
 	TArray<FVTStackEntry, TInlineAllocator<8>> VTStacks;
+	TArray<const class UMaterialParameterCollection*, TInlineAllocator<MaxNumParameterCollectionsPerMaterial>> ParameterCollections;
 	FHashTable VTStackHash;
 	TBitArray<> ExternalInputMask[SF_NumFrequencies];
 	FMaterialShadingModelField ShadingModelsFromCompilation;
@@ -872,6 +920,8 @@ public:
 	void PrepareInterpolators(FEmitContext& Context, FEmitScope& Scope);
 	void EmitInterpolatorStatements(FEmitContext& Context, FEmitScope& Scope) const;
 	void EmitInterpolatorShader(FEmitContext& Context, FStringBuilderBase& OutCode);
+
+	int32 FindOrAddParameterCollection(const class UMaterialParameterCollection* ParameterCollection);
 };
 
 } // namespace UE::HLSLTree::Material
diff --git a/Engine/Source/Runtime/Engine/Public/MaterialShared.h b/Engine/Source/Runtime/Engine/Public/MaterialShared.h
--- a/Engine/Source/Runtime/Engine/Public/MaterialShared.h
+++ b/Engine/Source/Runtime/Engine/Public/MaterialShared.h
@@ -588,7 +588,7 @@ public:
 
 	FShaderParametersMetadata* CreateBufferStruct();
 
-	void SetParameterCollections(const TArray<class UMaterialParameterCollection*>& Collections);
+	void SetParameterCollections(TConstArrayView<const class UMaterialParameterCollection*> Collections);
 
 	ENGINE_API void FillUniformBuffer(const FMaterialRenderContext& MaterialRenderContext, const FUniformExpressionCache& UniformExpressionCache, const FRHIUniformBufferLayout* UniformBufferLayout, uint8* TempBuffer, int TempBufferSize) const;
 
diff --git a/Engine/Source/Runtime/Engine/Public/Shader/ShaderTypes.h b/Engine/Source/Runtime/Engine/Public/Shader/ShaderTypes.h
--- a/Engine/Source/Runtime/Engine/Public/Shader/ShaderTypes.h
+++ b/Engine/Source/Runtime/Engine/Public/Shader/ShaderTypes.h
@@ -242,7 +242,7 @@ inline bool operator!=(const EValueType& Lhs, const FType& Rhs)
 	return !operator==(Lhs, Rhs);
 }
 
-ENGINE_API FType CombineTypes(const FType& Lhs, const FType& Rhs);
+ENGINE_API FType CombineTypes(const FType& Lhs, const FType& Rhs, bool bMergeMatrixTypes = false);
 
 struct FStructField
 {
