Switched engine-level LWI interface to use the new interface class, the IActorInstanceManagerInterface, rather than the ALightWeightInstanceManager class directly.
* ALightWeightInstanceManager implements IActorInstanceManagerInterface now and all the relevant public API has been moved to the interface functions' implementation
* removed the whole `UActorInstanceHandleInterface` funcitonality since it was not being used anywhere.
* cut other bits of LWI interface that were never being called  (ConvertLightWeightIndexToCollisionIndex, the whole LightWeightInstanceUtils.h header)
* removed UWorld.LWILastAssignedUID as it was never being used.

diff --git a/Engine/Source/Runtime/Engine/Classes/Engine/ActorInstanceHandle.h b/Engine/Source/Runtime/Engine/Classes/Engine/ActorInstanceHandle.h
--- a/Engine/Source/Runtime/Engine/Classes/Engine/ActorInstanceHandle.h
+++ b/Engine/Source/Runtime/Engine/Classes/Engine/ActorInstanceHandle.h
@@ -2,26 +2,36 @@
 
 #pragma once
 
-#include "CoreMinimal.h"
+#include "UObject/WeakInterfacePtr.h"
 #include "ActorInstanceHandle.generated.h"
 
-// Handle to a unique object. This may specify a full weigh actor or it may only specify the light weight instance that represents the same object.
+class USceneComponent;
+class AActor;
+class UActorInstanceManager;
+class IActorInstanceManagerInterface;
+class ULevel;
+
+using FActorInstanceManagerInterface = TWeakInterfacePtr<IActorInstanceManagerInterface>;
+
+// Handle to a unique object. This may specify a full weigh actor or it may only specify the actor instance that represents the same object.
 USTRUCT(BlueprintType)
 struct FActorInstanceHandle
 {
 	GENERATED_BODY()
 
-	friend struct FLightWeightInstanceSubsystem;
-	friend class ALightWeightInstanceManager;
-	friend class UActorInstanceHandleInterface;
-
 	ENGINE_API FActorInstanceHandle();
 
 	ENGINE_API explicit FActorInstanceHandle(AActor* InActor);
-	ENGINE_API explicit FActorInstanceHandle(class ALightWeightInstanceManager* Manager, int32 InInstanceIndex);
-
+	ENGINE_API FActorInstanceHandle(UObject* InManager, const UPrimitiveComponent* RelevantComponent, int32 CollisionInstanceIndex);
+	ENGINE_API FActorInstanceHandle(FActorInstanceManagerInterface InManagerInterface, int32 InstanceIndex);
 	ENGINE_API FActorInstanceHandle(const FActorInstanceHandle& Other);
 
+	/** 
+	 * A path dedicated to creation of handles while converting actor to a dehydrated representation. This path ensures
+	 * an actor won't be spawned as a side effect of looking for the actor given Manager/Index represents
+	 */
+	static FActorInstanceHandle MakeDehydratedActorHandle(UObject& Manager, int32 InInstanceIndex);
+
 	ENGINE_API bool IsValid() const;
 
 	ENGINE_API bool DoesRepresentClass(const UClass* OtherClass) const;
@@ -30,7 +40,7 @@ struct FActorInstanceHandle
 	bool DoesRepresent() const;
 
 	ENGINE_API UClass* GetRepresentedClass() const;
-
+	ENGINE_API ULevel* GetLevel() const;
 	ENGINE_API FVector GetLocation() const;
 	ENGINE_API FRotator GetRotation() const;
 	ENGINE_API FTransform GetTransform() const;
@@ -46,15 +56,16 @@ struct FActorInstanceHandle
 
 	/** Returns the actor specified by this handle. This may require loading and creating the actor object. */
 	ENGINE_API AActor* FetchActor() const;
+
 	template <typename T>
 	T* FetchActor() const;
 
+	AActor* GetCachedActor() const { return Actor.Get(); }
+	void SetCachedActor(AActor* InActor) const;
+
 	/* Returns the index used internally by the manager */
 	FORCEINLINE int32 GetInstanceIndex() const { return InstanceIndex; }
 
-	/* Returns the index used by rendering and collision */
-	ENGINE_API int32 GetRenderingInstanceIndex() const;
-
 	FActorInstanceHandle& operator=(const FActorInstanceHandle& Other) = default;
 	FActorInstanceHandle& operator=(FActorInstanceHandle&& Other) = default;
 	ENGINE_API FActorInstanceHandle& operator=(AActor* OtherActor);
@@ -69,9 +80,17 @@ struct FActorInstanceHandle
 
 	friend ENGINE_API FArchive& operator<<(FArchive& Ar, FActorInstanceHandle& Handle);
 
-	uint32 GetInstanceUID() const { return InstanceUID; }
+	FActorInstanceManagerInterface GetManagerInterface() const { return ManagerInterface; }
+
+	template<typename T>
+	T* GetManager() const 
+	{
+		return Cast<T>(ManagerInterface.GetObject());
+	}
 
 private:
+	void SetInternal(IActorInstanceManagerInterface& InManagerInterface, const UPrimitiveComponent* RelevantComponent, int32 CollisionInstanceIndex);
+
 	/**
 	 * helper functions that let us treat the actor pointer as a UObject in templated functions
 	 * these do NOT fetch the actor so they will return nullptr if we don't have a full actor representation
@@ -84,16 +103,13 @@ private:
 
 	/** this is cached here for convenience */
 	UPROPERTY()
-		mutable TWeakObjectPtr<AActor> Actor;
+	mutable TWeakObjectPtr<AActor> Actor;
 
-	/** Identifies the light weight instance manager to use */
-	TWeakObjectPtr<ALightWeightInstanceManager> Manager;
+	/** Identifies the actor instance manager to use */
+	FActorInstanceManagerInterface ManagerInterface;
 
 	/** Identifies the instance within the manager */
 	int32 InstanceIndex;
-
-	/** Unique identifier for instances represented by the handle */
-	uint32 InstanceUID;
 };
 
 template<typename T>
diff --git a/Engine/Source/Runtime/Engine/Classes/Engine/ActorInstanceManagerInterface.h b/Engine/Source/Runtime/Engine/Classes/Engine/ActorInstanceManagerInterface.h
new file mode 100644
--- /dev/null
+++ b/Engine/Source/Runtime/Engine/Classes/Engine/ActorInstanceManagerInterface.h
@@ -0,0 +1,45 @@
+// Copyright Epic Games, Inc. All Rights Reserved.
+
+#pragma once
+
+#include "UObject/Interface.h"
+#include "Chaos/Transform.h"
+#include "ActorInstanceManagerInterface.generated.h"
+
+
+class ULevel;
+class UClass;
+class AActor;
+class UPrimitiveComponent;
+struct FActorInstanceHandle;
+
+UINTERFACE(MinimalAPI)
+class UActorInstanceManagerInterface : public UInterface
+{
+	GENERATED_BODY()
+};
+
+class IActorInstanceManagerInterface
+{
+	GENERATED_BODY()
+
+public:
+	// Returns the index used internally by the actor instance manager that is associated with the instance 
+	// referred to by InIndex used by collision and rendering
+	virtual int32 ConvertCollisionIndexToInstanceIndex(int32 InIndex, const UPrimitiveComponent* RelevantComponent) const = 0;
+
+	// Returns the actor associated with Handle if one exists
+	virtual AActor* FindActor(const FActorInstanceHandle& Handle) = 0;
+
+	// Returns the actor associated with Handle. If one does not exist the system is expected to spawn one instantly. 
+	virtual AActor* FindOrCreateActor(const FActorInstanceHandle& Handle) = 0;
+
+	// Returns the specific class InstanceIndex represents
+	virtual UClass* GetRepresentedClass(int32 InstanceIndex) const = 0;
+
+	// Returns the ULevel given instance is represented in. Potentially null if InstanceIndex doesn't correspond to an actor
+	virtual ULevel* GetLevelForInstance(int32 InstanceIndex) const = 0;
+
+	// Returns the transform of the instance specified by Handle
+	virtual FTransform GetTransform(const FActorInstanceHandle& Handle) const = 0;
+};
diff --git a/Engine/Source/Runtime/Engine/Classes/Engine/World.h b/Engine/Source/Runtime/Engine/Classes/Engine/World.h
--- a/Engine/Source/Runtime/Engine/Classes/Engine/World.h
+++ b/Engine/Source/Runtime/Engine/Classes/Engine/World.h
@@ -1513,10 +1513,6 @@ public:
 	UPROPERTY(Transient)
 	TObjectPtr<class UPhysicsFieldComponent> PhysicsField;
 
-	/** Tracks the last assigned unique id for light weight instances in this world. */
-	UPROPERTY()
-	uint32 LWILastAssignedUID;
-
 private:
 
 	/** Array of components that need to wait on tasks before end of frame updates */
diff --git a/Engine/Source/Runtime/Engine/Classes/GameFramework/LightWeightInstanceManager.h b/Engine/Source/Runtime/Engine/Classes/GameFramework/LightWeightInstanceManager.h
--- a/Engine/Source/Runtime/Engine/Classes/GameFramework/LightWeightInstanceManager.h
+++ b/Engine/Source/Runtime/Engine/Classes/GameFramework/LightWeightInstanceManager.h
@@ -11,13 +11,14 @@
 #include "Engine/EngineBaseTypes.h"
 #include "Templates/SharedPointer.h"
 
-#include "Actor.h"
+#include "GameFramework/Actor.h"
+#include "Engine/ActorInstanceManagerInterface.h"
 
 #include "LightWeightInstanceManager.generated.h"
 
 
+class AActor;
 struct FActorSpawnParameters;
-//DECLARE_DYNAMIC_DELEGATE_RetVal_OneParam(FActorInstanceHandle, FOnActorReady, FActorInstanceHandle, InHandle);
 
 // Used for initializing light weight instances.
 struct FLWIData
@@ -26,31 +27,8 @@ struct FLWIData
 };
 
 
-// Base class for interfaces for each handle
-UCLASS(Experimental, MinimalAPI)
-class UActorInstanceHandleInterface : public UObject
-{
-	GENERATED_UCLASS_BODY()
-
-	friend class ALightWeightInstanceManager;
-
-	// Returns a pointer to the actor for this handle. This will create a new copy of the actor if it was being stored as a light weight instance.
-	template<typename T>
-	T* FetchActor();
-
-protected:
-
-	// Handle to the actor or instance that is associated with this interface
-	FActorInstanceHandle Handle;
-
-	// A cached pointer to the manager associated with Handle so we don't need to find it every time
-	ALightWeightInstanceManager* Manager;
-};
-
-
-
 UCLASS(BlueprintType, Blueprintable, Experimental, MinimalAPI)
-class ALightWeightInstanceManager : public AActor
+class ALightWeightInstanceManager : public AActor, public IActorInstanceManagerInterface
 {
 	GENERATED_UCLASS_BODY()
 
@@ -69,9 +47,6 @@ public:
 	// Returns the rotation of the instance specified by Handle
 	ENGINE_API FRotator GetRotation(const FActorInstanceHandle& Handle) const;
 
-	// Returns the transform of the instance specified by Handle
-	ENGINE_API FTransform GetTransform(const FActorInstanceHandle& Handle) const;
-
 	// Returns the name of the instance specified by Handle
 	ENGINE_API FString GetName(const FActorInstanceHandle& Handle) const;
 
@@ -81,30 +56,18 @@ public:
 	// Returns true if this manager is capable of representing objects of type OtherClass
 	ENGINE_API bool DoesAcceptClass(const UClass* OtherClass) const;
 
-	// Returns the specific class that this manages
-	ENGINE_API UClass* GetRepresentedClass() const;
-
 	// Returns the base class of types that this can manage
 	ENGINE_API UClass* GetAcceptedClass() const;
 
 	// Sets the specific class that this manages
 	ENGINE_API virtual void SetRepresentedClass(UClass* ActorClass);
 
-	// Returns the actor associated with Handle if one exists
-	ENGINE_API AActor* FetchActorFromHandle(const FActorInstanceHandle& Handle);
-
 	// Returns the index of the light weight instance associated with InActor if one exists; otherwise we return INDEX_NONE
 	ENGINE_API int32 FindIndexForActor(const AActor* InActor) const;
 
 	// Returns a handle to a light weight instance representing the same object as InActor and calls destroy on InActor if successful.
 	ENGINE_API FActorInstanceHandle ConvertActorToLightWeightInstance(AActor* InActor);
 
-	// Returns the index used internally by the light weight instance manager that is associated with the instance referred to by InIndex used by collision and rendering
-	ENGINE_API virtual int32 ConvertCollisionIndexToLightWeightIndex(int32 InIndex) const;
-
-	// Returns the index used by collision and rendering that is associated with the instance referred to by InIndex
-	ENGINE_API virtual int32 ConvertLightWeightIndexToCollisionIndex(int32 InIndex) const;
-
 	// LWI grid size to use for this manager.
 	ENGINE_API virtual int32 GetGridSize() const;
 
@@ -114,42 +77,23 @@ public:
 	// Helper that retrieve the world space bounds of the LWI grid cell encompassing the provided coordinate.
 	ENGINE_API FBox ConvertPositionToGridBounds(const FVector& InPosition) const;
 
-	template<typename U>
-	bool IsInterfaceSupported() const
-	{
-		const UClass* const InterfaceClass = GetInterfaceClass();
-		return InterfaceClass && InterfaceClass->ImplementsInterface(U::StaticClass());
-	}
-
-	template<typename I>
-	I* FetchInterfaceObject(const FActorInstanceHandle& Handle)
-	{
-		// if we have a valid actor, use it
-		if (Handle.Actor.IsValid())
-		{
-			return Cast<I>(Handle.Actor.Get());
-		}
-
-		// TODO: once we have a better idea of how we're going to use this we should add some kind of caching so we don't always need to create a new uobject.
-		// if we can support the interface using just the handle then do that
-		if (I* InterfaceObj = Cast<I>(CreateInterfaceObject(Handle)))
-		{
-			return InterfaceObj;
-		}
-
-		// fallback to creating the actor
-		return Cast<I>(FetchActorFromHandle(Handle));
-	}
-
-	ENGINE_API AActor* FindActorForInstanceIndex(const int32 InstanceIndex);
-
 	ENGINE_API bool HasAnyValidInstancesOrManagedActors() const;
 
+	// IActorInstanceManagerInterface begin
+	ENGINE_API virtual int32 ConvertCollisionIndexToInstanceIndex(int32 InIndex, const UPrimitiveComponent* RelevantComponent) const;
+	ENGINE_API virtual AActor* FindActor(const FActorInstanceHandle& Handle) override;
+	ENGINE_API virtual AActor* FindOrCreateActor(const FActorInstanceHandle& Handle) override;
+	ENGINE_API virtual UClass* GetRepresentedClass(const int32 InstanceIndex) const override { return GetRepresentedClassInternal(); }
+	ENGINE_API virtual ULevel* GetLevelForInstance(const int32 InstanceIndex) const override { return GetLevel(); }
+	ENGINE_API virtual FTransform GetTransform(const FActorInstanceHandle& Handle) const;
+	// IActorInstanceManagerInterface end
+
 protected:
+	UClass* GetRepresentedClassInternal() const;
+
 	// Creates an actor to replace the instance specified by Handle
 	ENGINE_API AActor* ConvertInstanceToActor(const FActorInstanceHandle& Handle);
 
-	// Takes a polymorphic struct to set the initial data for a new instance
 	ENGINE_API int32 AddNewInstance(FLWIData* InitData);
 
 	// Adds a new instance at the specified index. This function should only be called by AddNewInstance.
@@ -180,18 +124,6 @@ protected:
 	// Called after a spawned actor is destroyed
 	ENGINE_API virtual void OnSpawnedActorDestroyed(AActor* DestroyedActor, const int32 DestroyedActorInstanceIndex);
 
-	// Create an object that implements interfaces for the light weight instance specified by Handle
-	UObject* CreateInterfaceObject(const FActorInstanceHandle& Handle)
-	{
-		UActorInstanceHandleInterface* InterfaceObj = NewObject<UActorInstanceHandleInterface>(this, GetInterfaceClass());
-		InterfaceObj->Handle = Handle;
-		InterfaceObj->Manager = this;
-		return InterfaceObj;
-	}
-
-	// Gets the class that implements interfaces for light weight instances owned by this manager
-	ENGINE_API virtual UClass* GetInterfaceClass() const;
-
 	// Helper functions for converting between our internal storage indices and the indices used by external bookkeeping
 	virtual int32 ConvertInternalIndexToHandleIndex(int32 InInternalIndex) const { return InInternalIndex; }
 	virtual int32 ConvertHandleIndexToInternalIndex(int32 InHandleIndex) const { return InHandleIndex; }
diff --git a/Engine/Source/Runtime/Engine/Classes/GameFramework/LightWeightInstanceStaticMeshManager.h b/Engine/Source/Runtime/Engine/Classes/GameFramework/LightWeightInstanceStaticMeshManager.h
--- a/Engine/Source/Runtime/Engine/Classes/GameFramework/LightWeightInstanceStaticMeshManager.h
+++ b/Engine/Source/Runtime/Engine/Classes/GameFramework/LightWeightInstanceStaticMeshManager.h
@@ -4,7 +4,7 @@
 
 #include "CoreMinimal.h"
 #include "Engine/StaticMesh.h"
-#include "GameFramework/LightWeightInstanceManager.h"
+#include "LightWeightInstanceManager.h"
 #include "Elements/SMInstance/SMInstanceManager.h"
 
 #include "LightWeightInstanceStaticMeshManager.generated.h"
@@ -24,9 +24,7 @@ class ALightWeightInstanceStaticMeshManager : public ALightWeightInstanceManager
 	// Clears the static mesh used for rendering instances
 	ENGINE_API void ClearStaticMesh();
 
-	ENGINE_API virtual int32 ConvertCollisionIndexToLightWeightIndex(int32 InIndex) const override;
-
-	ENGINE_API virtual int32 ConvertLightWeightIndexToCollisionIndex(int32 InIndex) const override;
+	ENGINE_API virtual int32 ConvertCollisionIndexToInstanceIndex(int32 InIndex, const UPrimitiveComponent* RelevantComponent) const override;
 
 protected:
 	ENGINE_API virtual void AddNewInstanceAt(FLWIData* InitData, int32 Index) override;
diff --git a/Engine/Source/Runtime/Engine/Classes/GameFramework/LightWeightInstanceSubsystem.h b/Engine/Source/Runtime/Engine/Classes/GameFramework/LightWeightInstanceSubsystem.h
--- a/Engine/Source/Runtime/Engine/Classes/GameFramework/LightWeightInstanceSubsystem.h
+++ b/Engine/Source/Runtime/Engine/Classes/GameFramework/LightWeightInstanceSubsystem.h
@@ -6,7 +6,6 @@
 #include "GameFramework/LightWeightInstanceManager.h"
 #include "Templates/SharedPointer.h"
 
-#include "Actor.h"
 
 ENGINE_API DECLARE_LOG_CATEGORY_EXTERN(LogLightWeightInstance, Log, All);
 
@@ -49,9 +48,6 @@ struct FLightWeightInstanceSubsystem
 	// Returns the instance manager that handles instances of type Class that live in Level
 	ENGINE_API ALightWeightInstanceManager* FindOrAddLightWeightInstanceManager(UClass& ActorClass, UWorld& World, const FVector& InPos, const UDataLayerInstance* DataLayer = nullptr);
 
-	// Returns the actor specified by Handle. This may require loading and creating the actor object.
-	ENGINE_API AActor* FetchActor(const FActorInstanceHandle& Handle);
-
 	// Returns the actor specified by Handle if it exists. Returns nullptr if it doesn't
 	ENGINE_API AActor* GetActor_NoCreate(const FActorInstanceHandle& Handle) const;
 
@@ -73,30 +69,6 @@ struct FLightWeightInstanceSubsystem
 	// deletes the instance identified by Handle
 	ENGINE_API void DeleteInstance(const FActorInstanceHandle& Handle);
 
-	// Returns true if the handle can return an object that implements the interface U
-	template<typename U>
-	bool IsInterfaceSupported(const FActorInstanceHandle& Handle) const
-	{
-		if (ALightWeightInstanceManager* InstanceManager = FindLightWeightInstanceManager(Handle))
-		{
-			return InstanceManager->IsInterfaceSupported<U>();
-		}
-
-		return false;
-	}
-
-	// Returns an object that implements the interface I for Handle
-	template<typename I>
-	I* FetchInterfaceObject(const FActorInstanceHandle& Handle)
-	{
-		if (ALightWeightInstanceManager* InstanceManager = FindLightWeightInstanceManager(Handle))
-		{
-			return InstanceManager->FetchInterfaceObject<I>(Handle);
-		}
-
-		return nullptr;
-	}
-
 	// Helper that converts a position (world space) into a coordinate for the LWI grid.
 	UE_DEPRECATED(5.3, "Use LWI Managers version of ConvertPositionToCoord()")
 	static ENGINE_API FInt32Vector3 ConvertPositionToCoord(const FVector& InPosition);
diff --git a/Engine/Source/Runtime/Engine/Classes/GameFramework/LightWeightInstanceUtils.h b/Engine/Source/Runtime/Engine/Classes/GameFramework/LightWeightInstanceUtils.h
deleted file mode 100644
diff --git a/Engine/Source/Runtime/Engine/Private/Engine/ActorInstanceHandle.cpp b/Engine/Source/Runtime/Engine/Private/Engine/ActorInstanceHandle.cpp
--- a/Engine/Source/Runtime/Engine/Private/Engine/ActorInstanceHandle.cpp
+++ b/Engine/Source/Runtime/Engine/Private/Engine/ActorInstanceHandle.cpp
@@ -1,15 +1,19 @@
 // Copyright Epic Games, Inc. All Rights Reserved.
 
 #include "Engine/ActorInstanceHandle.h"
-#include "GameFramework/LightWeightInstanceSubsystem.h"
 #include "Engine/World.h"
+#include "Engine/ActorInstanceManagerInterface.h"
+#include "GameFramework/Actor.h"
 
 #include UE_INLINE_GENERATED_CPP_BY_NAME(ActorInstanceHandle)
 
+
+//-----------------------------------------------------------------------------
+// FActorInstanceHandle
+//-----------------------------------------------------------------------------
 FActorInstanceHandle::FActorInstanceHandle()
 	: Actor(nullptr)
 	, InstanceIndex(INDEX_NONE)
-	, InstanceUID(0)
 {
 	// do nothing
 }
@@ -17,26 +21,28 @@ FActorInstanceHandle::FActorInstanceHandle()
 FActorInstanceHandle::FActorInstanceHandle(AActor* InActor)
 	: Actor(InActor)
 	, InstanceIndex(INDEX_NONE)
-	, InstanceUID(0)
 {
 }
 
-FActorInstanceHandle::FActorInstanceHandle(ALightWeightInstanceManager* InManager, int32 InInstanceIndex)
+FActorInstanceHandle::FActorInstanceHandle(UObject* InManager, const UPrimitiveComponent* RelevantComponent, int32 CollisionInstanceIndex)
+	: Actor(nullptr)
+	, ManagerInterface(InManager)
+	, InstanceIndex(CollisionInstanceIndex)
+{
+	if (IActorInstanceManagerInterface* ManagerInterfacePtr = ManagerInterface.Get())
+	{
+		SetInternal(*ManagerInterfacePtr, RelevantComponent, CollisionInstanceIndex);
+	}
+}
+
+FActorInstanceHandle::FActorInstanceHandle(FActorInstanceManagerInterface InManagerInterface, int32 CollisionInstanceIndex)
 	: Actor(nullptr)
-	, InstanceIndex(InInstanceIndex)
-	, InstanceUID(0)
+	, ManagerInterface(InManagerInterface)
+	, InstanceIndex(CollisionInstanceIndex)
 {
-	Manager = InManager;
-	if (Manager.IsValid())
+	if (IActorInstanceManagerInterface* ManagerInterfacePtr = ManagerInterface.Get())
 	{
-		InstanceIndex = Manager->ConvertCollisionIndexToLightWeightIndex(InInstanceIndex);
-		Actor = Manager->FindActorForInstanceIndex(InstanceIndex);
-
-		UWorld* World = Manager->GetWorld();
-		if (ensure(World))
-		{
-			InstanceUID = World->LWILastAssignedUID++;
-		}
+		SetInternal(*ManagerInterfacePtr, /*RelevantComponent=*/nullptr, CollisionInstanceIndex);
 	}
 }
 
@@ -44,14 +50,30 @@ FActorInstanceHandle::FActorInstanceHandle(const FActorInstanceHandle& Other)
 {
 	Actor = Other.Actor;
 
-	Manager = Other.Manager;
+	ManagerInterface = Other.ManagerInterface;
 	InstanceIndex = Other.InstanceIndex;
-	InstanceUID = Other.InstanceUID;
+}
+
+FActorInstanceHandle FActorInstanceHandle::MakeDehydratedActorHandle(UObject& Manager, int32 InInstanceIndex)
+{
+	FActorInstanceHandle ReturnHandle;
+	ReturnHandle.ManagerInterface = FActorInstanceManagerInterface(&Manager);
+	ReturnHandle.InstanceIndex = InInstanceIndex;
+
+	return ReturnHandle;
+}
+
+void FActorInstanceHandle::SetInternal(IActorInstanceManagerInterface& InManagerInterface, const UPrimitiveComponent* RelevantComponent, int32 CollisionInstanceIndex)
+{
+	check(ManagerInterface.Get() == &InManagerInterface);
+
+	InstanceIndex = InManagerInterface.ConvertCollisionIndexToInstanceIndex(CollisionInstanceIndex, RelevantComponent);
+	Actor = InManagerInterface.FindActor(*this);
 }
 
 bool FActorInstanceHandle::IsValid() const
 {
-	return (Manager.IsValid() && InstanceIndex != INDEX_NONE) || IsActorValid();
+	return (ManagerInterface.IsValid() && InstanceIndex != INDEX_NONE) || IsActorValid();
 }
 
 bool FActorInstanceHandle::DoesRepresentClass(const UClass* OtherClass) const
@@ -65,22 +87,28 @@ bool FActorInstanceHandle::DoesRepresentClass(const UClass* OtherClass) const
 
 UClass* FActorInstanceHandle::GetRepresentedClass() const
 {
-	if (!IsValid())
-	{
-		return nullptr;
-	}
-
 	if (IsActorValid())
 	{
 		return Actor->GetClass();
 	}
 
-	if (Manager.IsValid())
+	IActorInstanceManagerInterface* ManagerInterfacePtr = ManagerInterface.Get();
+	return ManagerInterfacePtr && (InstanceIndex != INDEX_NONE)
+		? ManagerInterfacePtr->GetRepresentedClass(InstanceIndex)
+		: nullptr;
+}
+
+ULevel* FActorInstanceHandle::GetLevel() const
+{
+	if (IsActorValid())
 	{
-		return Manager->GetRepresentedClass();
+		return Actor->GetLevel();
 	}
 
-	return nullptr;
+	IActorInstanceManagerInterface* ManagerInterfacePtr = ManagerInterface.Get();
+	return ManagerInterfacePtr && (InstanceIndex != INDEX_NONE)
+		? ManagerInterfacePtr->GetLevelForInstance(InstanceIndex)
+		: nullptr;
 }
 
 FVector FActorInstanceHandle::GetLocation() const
@@ -90,12 +118,10 @@ FVector FActorInstanceHandle::GetLocation() const
 		return Actor->GetActorLocation();
 	}
 
-	if (Manager.IsValid())
-	{
-		return Manager->GetLocation(*this);
-	}
-
-	return FVector();
+	IActorInstanceManagerInterface* ManagerInterfacePtr = ManagerInterface.Get();
+	return ManagerInterfacePtr && (InstanceIndex != INDEX_NONE)
+		? ManagerInterfacePtr->GetTransform(*this).GetLocation()
+		: FVector();
 }
 
 FRotator FActorInstanceHandle::GetRotation() const
@@ -105,12 +131,10 @@ FRotator FActorInstanceHandle::GetRotation() const
 		return Actor->GetActorRotation();
 	}
 
-	if (Manager.IsValid())
-	{
-		return Manager->GetRotation(*this);
-	}
-
-	return FRotator();
+	IActorInstanceManagerInterface* ManagerInterfacePtr = ManagerInterface.Get();
+	return ManagerInterfacePtr && (InstanceIndex != INDEX_NONE)
+		? ManagerInterfacePtr->GetTransform(*this).GetRotation().Rotator()
+		: FRotator();
 }
 
 FTransform FActorInstanceHandle::GetTransform() const
@@ -120,12 +144,10 @@ FTransform FActorInstanceHandle::GetTransform() const
 		return Actor->GetActorTransform();
 	}
 
-	if (Manager.IsValid())
-	{
-		return Manager->GetTransform(*this);
-	}
-
-	return FTransform();
+	IActorInstanceManagerInterface* ManagerInterfacePtr = ManagerInterface.Get();
+	return ManagerInterfacePtr && (InstanceIndex != INDEX_NONE)
+		? ManagerInterfacePtr->GetTransform(*this)
+		: FTransform();
 }
 
 FName FActorInstanceHandle::GetFName() const
@@ -145,12 +167,12 @@ FString FActorInstanceHandle::GetName() const
 		return Actor->GetName();
 	}
 
-	if (Manager.IsValid())
+	if (ManagerInterface.IsValid())
 	{
-		return Manager->GetName(*this);
+		return FString::Printf(TEXT("%s:d"), *GetNameSafe(ManagerInterface.GetObject()), InstanceIndex);
 	}
 
-	return FString();
+	return TEXT("Invalid");
 }
 
 AActor* FActorInstanceHandle::GetManagingActor() const
@@ -160,7 +182,7 @@ AActor* FActorInstanceHandle::GetManagingActor() const
 		return Actor.Get();
 	}
 
-	return Manager.Get();
+	return Cast<AActor>(ManagerInterface.GetObject());
 }
 
 USceneComponent* FActorInstanceHandle::GetRootComponent() const
@@ -170,7 +192,8 @@ USceneComponent* FActorInstanceHandle::GetRootComponent() const
 		return Actor->GetRootComponent();
 	}
 
-	return Manager.IsValid() ? Manager->GetRootComponent() : nullptr;
+	AActor* AsActor = Cast<AActor>(ManagerInterface.GetObject());
+	return AsActor ? AsActor->GetRootComponent() : nullptr;
 }
 
 AActor* FActorInstanceHandle::FetchActor() const
@@ -180,22 +203,13 @@ AActor* FActorInstanceHandle::FetchActor() const
 		return Actor.Get();
 	}
 
-	return FLightWeightInstanceSubsystem::Get().FetchActor(*this);
-}
-
-int32 FActorInstanceHandle::GetRenderingInstanceIndex() const
-{
-	return Manager.IsValid() ? Manager->ConvertLightWeightIndexToCollisionIndex(InstanceIndex) : INDEX_NONE;
+	return ManagerInterface.IsValid() ? ManagerInterface->FindOrCreateActor(*this) : nullptr;
 }
 
 UObject* FActorInstanceHandle::GetActorAsUObject()
 {
-	if (IsActorValid())
-	{
-		return Cast<UObject>(Actor.Get());
-	}
-
-	return nullptr;
+	// 
+	return Cast<UObject>(Actor.Get());
 }
 
 const UObject* FActorInstanceHandle::GetActorAsUObject() const
@@ -213,10 +227,16 @@ bool FActorInstanceHandle::IsActorValid() const
 	return Actor.IsValid();
 }
 
+void FActorInstanceHandle::SetCachedActor(AActor* InActor) const
+{
+	check(Actor.IsValid() == false);
+	Actor = InActor;
+}
+
 FActorInstanceHandle& FActorInstanceHandle::operator=(AActor* OtherActor)
 {
 	Actor = OtherActor;
-	Manager.Reset();
+	ManagerInterface.Reset();
 	InstanceIndex = INDEX_NONE;
 
 	return *this;
@@ -225,9 +245,9 @@ FActorInstanceHandle& FActorInstanceHandle::operator=(AActor* OtherActor)
 bool FActorInstanceHandle::operator==(const FActorInstanceHandle& Other) const
 {
 	// try to compare managers and indices first if we have them
-	if (Manager.IsValid() && Other.Manager.IsValid() && InstanceIndex != INDEX_NONE && Other.InstanceIndex != INDEX_NONE)
+	if (ManagerInterface.IsValid() && Other.ManagerInterface.IsValid() && InstanceIndex != INDEX_NONE && Other.InstanceIndex != INDEX_NONE)
 	{
-		return Manager == Other.Manager && InstanceIndex == Other.InstanceIndex;
+		return ManagerInterface == Other.ManagerInterface && InstanceIndex == Other.InstanceIndex;
 	}
 
 	// try to compare the actors
@@ -245,38 +265,19 @@ bool FActorInstanceHandle::operator!=(const FActorInstanceHandle& Other) const
 bool FActorInstanceHandle::operator==(const AActor* OtherActor) const
 {
 	// if we have an actor, compare the two actors
-	if (Actor.IsValid())
+	if (AActor* AsActor = Actor.Get())
 	{
-		return Actor.Get() == OtherActor;
+		return AsActor == OtherActor;
 	}
 
 	// if OtherActor is null then we're only equal if this doesn't refer to a valid instance
 	if (OtherActor == nullptr)
 	{
-		return !Manager.IsValid() && InstanceIndex == INDEX_NONE;
-	}
-
-	// we don't have an actor so see if we can look up an instance associated with OtherActor and see if we refer to the same instance
-
-#if WITH_EDITOR
-	// use the first layer the actor is in if it's in multiple layers
-	TArray<const UDataLayerInstance*> DataLayerInstances = OtherActor->GetDataLayerInstances();
-	const UDataLayerInstance* DataLayerInstance = DataLayerInstances.Num() > 0 ? DataLayerInstances[0] : nullptr;
-#else
-	const UDataLayerInstance* DataLayerInstance = nullptr;
-#endif // WITH_EDITOR
-
-	if (ALightWeightInstanceManager* LWIManager = FLightWeightInstanceSubsystem::Get().FindLightWeightInstanceManager(*OtherActor->GetClass(), *OtherActor->GetWorld(), OtherActor->GetActorLocation(), DataLayerInstance))
-	{
-		if (Manager.Get() != LWIManager)
-		{
-			return false;
-		}
-
-		return Manager->FindIndexForActor(OtherActor) == InstanceIndex;
+		return !ManagerInterface.IsValid() && InstanceIndex == INDEX_NONE;
 	}
 
-	return false;
+	IActorInstanceManagerInterface* ManagerInterfacePtr = ManagerInterface.Get();
+	return ManagerInterfacePtr && (ManagerInterfacePtr->FindActor(*this) == OtherActor);
 }
 
 bool FActorInstanceHandle::operator!=(const AActor* OtherActor) const
@@ -287,13 +288,13 @@ bool FActorInstanceHandle::operator!=(const AActor* OtherActor) const
 uint32 GetTypeHash(const FActorInstanceHandle& Handle)
 {
 	uint32 Hash = 0;
-	if (Handle.Actor.IsValid())
+	if (AActor* Actor = Handle.Actor.Get())
 	{
-		FCrc::StrCrc32(*(Handle.Actor->GetPathName()), Hash);
+		FCrc::StrCrc32(*(Actor->GetPathName()), Hash);
 	}
-	if (Handle.Manager.IsValid())
+	if (UObject* ManagerInterfaceObject = Handle.ManagerInterface.GetObject())
 	{
-		Hash = HashCombine(Hash, GetTypeHash(Handle.Manager.Get()));
+		Hash = HashCombine(Hash, GetTypeHash(ManagerInterfaceObject));
 	}
 	Hash = HashCombine(Hash, Handle.InstanceIndex);
 
@@ -303,9 +304,15 @@ uint32 GetTypeHash(const FActorInstanceHandle& Handle)
 FArchive& operator<<(FArchive& Ar, FActorInstanceHandle& Handle)
 {
 	Ar << Handle.Actor;
-	Ar << Handle.Manager;
 	Ar << Handle.InstanceIndex;
 
+	TSoftObjectPtr<UObject> SoftObject(Handle.ManagerInterface.GetObject());
+	Ar << SoftObject;
+
+	if (Ar.IsLoading())
+	{
+		Handle.ManagerInterface = FActorInstanceManagerInterface(SoftObject.Get());
+	}
+
 	return Ar;
 }
-
diff --git a/Engine/Source/Runtime/Engine/Private/LightWeightInstanceManager.cpp b/Engine/Source/Runtime/Engine/Private/LightWeightInstanceManager.cpp
--- a/Engine/Source/Runtime/Engine/Private/LightWeightInstanceManager.cpp
+++ b/Engine/Source/Runtime/Engine/Private/LightWeightInstanceManager.cpp
@@ -20,12 +20,6 @@ static FAutoConsoleVariableRef CVarLWIGridSize
 );
 
 
-UActorInstanceHandleInterface::UActorInstanceHandleInterface(const FObjectInitializer& ObjectInitializer)
-	:Super(ObjectInitializer)
-{
-	// do nothing
-}
-
 ALightWeightInstanceManager::ALightWeightInstanceManager(const FObjectInitializer& ObjectInitializer)
 	:Super(ObjectInitializer)
 {
@@ -47,20 +41,22 @@ void ALightWeightInstanceManager::SetRepresentedClass(UClass* ActorClass)
 	RepresentedClass = ActorClass;
 }
 
-UClass* ALightWeightInstanceManager::GetInterfaceClass() const
+void ALightWeightInstanceManager::Tick(float DeltaSeconds)
 {
-	return UActorInstanceHandleInterface::StaticClass();
+	// do nothing
 }
 
-void ALightWeightInstanceManager::Tick(float DeltaSeconds)
+AActor* ALightWeightInstanceManager::FindActor(const FActorInstanceHandle& Handle)
 {
-	// do nothing
+	TObjectPtr<AActor>* FoundActor = Actors.Find(Handle.GetInstanceIndex());
+
+	return FoundActor ? *FoundActor : nullptr;
 }
 
-AActor* ALightWeightInstanceManager::FetchActorFromHandle(const FActorInstanceHandle& Handle)
+AActor* ALightWeightInstanceManager::FindOrCreateActor(const FActorInstanceHandle& Handle)
 {
 	// make sure the handle doesn't have an actor already
-	if (ensure(!Handle.Actor.IsValid()))
+	if (ensure(!Handle.GetCachedActor()))
 	{
 		// check if we already have an actor for this handle
 		TObjectPtr<AActor>* FoundActor = nullptr;
@@ -82,14 +78,14 @@ AActor* ALightWeightInstanceManager::FetchActorFromHandle(const FActorInstanceHa
 			}
 #endif //!UE_BUILD_SHIPPING
 
-			Handle.Actor = *FoundActor;
+			Handle.SetCachedActor(*FoundActor);
 		}
 	}
 
 	// Unless we are on the server or the actor has been spawned, this ensure will fail.
 	// Commented out until there is more robust replication.
 	// ensure(Handle.Actor.IsValid());
-	return Handle.Actor.Get();
+	return Handle.GetCachedActor();
 }
 
 AActor* ALightWeightInstanceManager::ConvertInstanceToActor(const FActorInstanceHandle& Handle)
@@ -140,7 +136,7 @@ AActor* ALightWeightInstanceManager::ConvertInstanceToActor(const FActorInstance
 		check(NewActor);
 
 		//should have been assigned in CustomPreSpawnInitialization
-		check(Handle.Actor == NewActor);
+		check(Handle.GetCachedActor() == NewActor);
 		check(NewActor == Actors.FindRef(Handle.GetInstanceIndex()));
 
 		NewActor->OnDestroyed.AddUniqueDynamic(this, &ALightWeightInstanceManager::OnSpawnedActorDestroyed);
@@ -198,10 +194,9 @@ int32 ALightWeightInstanceManager::FindIndexForActor(const AActor* InActor) cons
 
 FActorInstanceHandle ALightWeightInstanceManager::ConvertActorToLightWeightInstance(AActor* InActor)
 {
-	FActorInstanceHandle ReturnHandle;
 	if (!InActor)
 	{
-		return ReturnHandle;
+		return FActorInstanceHandle();
 	}
 
 	if (FLWIData* Data = AllocateInitData())
@@ -228,23 +223,18 @@ FActorInstanceHandle ALightWeightInstanceManager::ConvertActorToLightWeightInsta
 			UpdateDataAtIndex(Data, Idx);
 		}
 
-		// Update our handle
-		ReturnHandle.Manager = this;
-		ReturnHandle.InstanceIndex = Idx;
-		ReturnHandle.InstanceUID = InActor->GetWorld()->LWILastAssignedUID++;
-
 		// cleanup
 		InActor->Destroy();
 		delete Data;
+
+		return FActorInstanceHandle::MakeDehydratedActorHandle(*this, Idx);
 	}
 	else
 	{
 		// something went wrong and we can't manage this actor
 		// just return a handle to the actor
-		ReturnHandle.Actor = InActor;
+		return FActorInstanceHandle(InActor);
 	}
-
-	return ReturnHandle;
 }
 
 FLWIData* ALightWeightInstanceManager::AllocateInitData() const
@@ -265,12 +255,7 @@ bool ALightWeightInstanceManager::SetDataFromActor(FLWIData* InData, AActor* InA
 	return true;
 }
 
-int32 ALightWeightInstanceManager::ConvertCollisionIndexToLightWeightIndex(int32 InIndex) const
-{
-	return InIndex;
-}
-
-int32 ALightWeightInstanceManager::ConvertLightWeightIndexToCollisionIndex(int32 InIndex) const
+int32 ALightWeightInstanceManager::ConvertCollisionIndexToInstanceIndex(int32 InIndex, const UPrimitiveComponent* RelevantComponent) const
 {
 	return InIndex;
 }
@@ -325,7 +310,7 @@ UClass* ALightWeightInstanceManager::GetActorClassToSpawn(const FActorInstanceHa
 
 void ALightWeightInstanceManager::PreSpawnInitalization(const FActorInstanceHandle& Handle, AActor* SpawnedActor)
 {
-	Handle.Actor = SpawnedActor;
+	Handle.SetCachedActor(SpawnedActor);
 	Actors.Add(Handle.GetInstanceIndex(), SpawnedActor);
 }
 
@@ -346,19 +331,15 @@ bool ALightWeightInstanceManager::IsIndexValid(int32 Index) const
 
 bool ALightWeightInstanceManager::FindActorForHandle(const FActorInstanceHandle& Handle) const
 {
-	if (Handle.Actor.IsValid())
+	if (Handle.GetCachedActor())
 	{
 		return true;
 	}
 
 	const TObjectPtr<AActor>* FoundActor = Actors.Find(Handle.GetInstanceIndex());
-	Handle.Actor = FoundActor ? *FoundActor : nullptr;
-	return Handle.Actor != nullptr;
-}
-
-AActor* ALightWeightInstanceManager::FindActorForInstanceIndex(const int32 InstanceIndex)
-{
-	return Actors.FindRef(InstanceIndex);
+	AActor* ResolvedActor = FoundActor ? *FoundActor : nullptr;
+	Handle.SetCachedActor(ResolvedActor);
+	return ResolvedActor != nullptr;
 }
 
 FVector ALightWeightInstanceManager::GetLocation(const FActorInstanceHandle& Handle) const
@@ -375,7 +356,8 @@ FTransform ALightWeightInstanceManager::GetTransform(const FActorInstanceHandle&
 {
 	if (FindActorForHandle(Handle))
 	{
-		return Handle.Actor->GetActorTransform();
+		check(Handle.GetCachedActor());
+		return Handle.GetCachedActor()->GetActorTransform();
 	}
 
 	if (ensure(IsIndexValid(Handle.GetInstanceIndex())))
@@ -391,7 +373,8 @@ FString ALightWeightInstanceManager::GetName(const FActorInstanceHandle& Handle)
 {
 	if (FindActorForHandle(Handle))
 	{
-		return Handle.Actor->GetName();
+		check(Handle.GetCachedActor());
+		return Handle.GetCachedActor()->GetName();
 	}
 
 	return FString::Printf(TEXT("%s_%u"), *BaseInstanceName, Handle.GetInstanceIndex());
@@ -399,7 +382,7 @@ FString ALightWeightInstanceManager::GetName(const FActorInstanceHandle& Handle)
 
 bool ALightWeightInstanceManager::DoesRepresentClass(const UClass* OtherClass) const
 {
-	return OtherClass ? OtherClass->IsChildOf(GetRepresentedClass()) : false;
+	return OtherClass ? OtherClass->IsChildOf(GetRepresentedClassInternal()) : false;
 }
 
 bool ALightWeightInstanceManager::DoesAcceptClass(const UClass* OtherClass) const
@@ -407,7 +390,7 @@ bool ALightWeightInstanceManager::DoesAcceptClass(const UClass* OtherClass) cons
 	return OtherClass ? OtherClass->IsChildOf(GetAcceptedClass()) : false;
 }
 
-UClass* ALightWeightInstanceManager::GetRepresentedClass() const
+UClass* ALightWeightInstanceManager::GetRepresentedClassInternal() const
 {
 	return IsValid(RepresentedClass) ? RepresentedClass : nullptr;
 }
diff --git a/Engine/Source/Runtime/Engine/Private/LightWeightInstanceStaticMeshManager.cpp b/Engine/Source/Runtime/Engine/Private/LightWeightInstanceStaticMeshManager.cpp
--- a/Engine/Source/Runtime/Engine/Private/LightWeightInstanceStaticMeshManager.cpp
+++ b/Engine/Source/Runtime/Engine/Private/LightWeightInstanceStaticMeshManager.cpp
@@ -116,7 +116,7 @@ PRAGMA_ENABLE_DEPRECATION_WARNINGS
 	}
 }
 
-int32 ALightWeightInstanceStaticMeshManager::ConvertCollisionIndexToLightWeightIndex(int32 InIndex) const
+int32 ALightWeightInstanceStaticMeshManager::ConvertCollisionIndexToInstanceIndex(int32 InIndex, const UPrimitiveComponent* RelevantComponent) const
 {
 	if (ensureMsgf(RenderingIndicesToDataIndices.IsValidIndex(InIndex), TEXT("Invalid index [ %d ]"), InIndex))
 	{
@@ -126,16 +126,6 @@ int32 ALightWeightInstanceStaticMeshManager::ConvertCollisionIndexToLightWeightI
 	return InIndex;
 }
 
-int32 ALightWeightInstanceStaticMeshManager::ConvertLightWeightIndexToCollisionIndex(int32 InIndex) const
-{
-	if (ensureMsgf(DataIndicesToRenderingIndices.IsValidIndex(InIndex), TEXT("Invalid index [ %d ]"), InIndex))
-	{
-		return DataIndicesToRenderingIndices[InIndex];
-	}
-
-	return InIndex;
-}
-
 void ALightWeightInstanceStaticMeshManager::AddNewInstanceAt(FLWIData* InitData, int32 Index)
 {
 	Super::AddNewInstanceAt(InitData, Index);
diff --git a/Engine/Source/Runtime/Engine/Private/LightWeightInstanceSubsystem.cpp b/Engine/Source/Runtime/Engine/Private/LightWeightInstanceSubsystem.cpp
--- a/Engine/Source/Runtime/Engine/Private/LightWeightInstanceSubsystem.cpp
+++ b/Engine/Source/Runtime/Engine/Private/LightWeightInstanceSubsystem.cpp
@@ -79,18 +79,21 @@ bool FLightWeightInstanceSubsystem::RemoveManager(ALightWeightInstanceManager* M
 
 ALightWeightInstanceManager* FLightWeightInstanceSubsystem::FindLightWeightInstanceManager(const FActorInstanceHandle& Handle) const
 {
-	if (Handle.Manager.IsValid())
+	if (Handle.GetManagerInterface().IsValid())
 	{
-		return Handle.Manager.Get();
+		return Handle.GetManager<ALightWeightInstanceManager>();
 	}
 
-	if (Handle.Actor.IsValid())
+	if (AActor* CachedActor = Handle.GetCachedActor())
 	{
+		const UClass* CachedActorClass = CachedActor->GetClass();
+		CA_ASSUME(CachedActorClass);
+
 		FReadScopeLock Lock(LWIManagersRWLock);
 		// see if we already have a match
 		for (ALightWeightInstanceManager* LWInstance : LWInstanceManagers)
 		{
-			if (Handle.Actor->GetClass() == LWInstance->GetRepresentedClass())
+			if (CachedActorClass == LWInstance->GetRepresentedClass(Handle.GetInstanceIndex()))
 			{
 				const FInt32Vector3 GridCoord = LWInstance->ConvertPositionToCoord(Handle.GetLocation());
 				const FInt32Vector3 ManagerGridCoord = LWInstance->ConvertPositionToCoord(LWInstance->GetActorLocation());
@@ -139,7 +142,7 @@ ALightWeightInstanceManager* FLightWeightInstanceSubsystem::FindLightWeightInsta
 	{
 		if (IsValid(InstanceManager))
 		{
-			if (InstanceManager->GetRepresentedClass() == &ActorClass)
+			if (InstanceManager->GetRepresentedClassInternal() == &ActorClass)
 			{
 				const FInt32Vector3 GridCoord = InstanceManager->ConvertPositionToCoord(InPos);
 				const FInt32Vector3 ManagerGridCoord = InstanceManager->ConvertPositionToCoord(InstanceManager->GetActorLocation());
@@ -240,7 +243,7 @@ UClass* FLightWeightInstanceSubsystem::FindBestInstanceManagerClass(const UClass
 	{
 		if (ManagerClass->GetDefaultObject<ALightWeightInstanceManager>()->DoesAcceptClass(InActorClass))
 		{
-			const UClass* HandledClass = ManagerClass->GetDefaultObject<ALightWeightInstanceManager>()->GetRepresentedClass();
+			const UClass* HandledClass = ManagerClass->GetDefaultObject<ALightWeightInstanceManager>()->GetRepresentedClassInternal();
 			if (!HandledClass)
 			{
 				HandledClass = ManagerClass->GetDefaultObject<ALightWeightInstanceManager>()->GetAcceptedClass();
@@ -268,37 +271,21 @@ UClass* FLightWeightInstanceSubsystem::FindBestInstanceManagerClass(const UClass
 	return BestManagerClass;
 }
 
-AActor* FLightWeightInstanceSubsystem::FetchActor(const FActorInstanceHandle& Handle)
-{
-	// if the actor is valid return it
-	if (Handle.Actor.IsValid())
-	{
-		return Handle.Actor.Get();
-	}
-
-	if (ALightWeightInstanceManager* LWIManager = FindLightWeightInstanceManager(Handle))
-	{
-		return LWIManager->FetchActorFromHandle(Handle);
-	}
-
-	return nullptr;
-}
-
 AActor* FLightWeightInstanceSubsystem::GetActor_NoCreate(const FActorInstanceHandle& Handle) const
 {
-	return Handle.Actor.Get();
+	return Handle.GetCachedActor();
 }
 
 UClass* FLightWeightInstanceSubsystem::GetActorClass(const FActorInstanceHandle& Handle)
 {
-	if (Handle.Actor.IsValid())
+	if (const AActor* CachedActor = Handle.GetCachedActor())
 	{
-		return Handle.Actor->StaticClass();
+		return CachedActor->StaticClass();
 	}
 
 	if (ALightWeightInstanceManager* LWIManager = FindLightWeightInstanceManager(Handle))
 	{
-		return LWIManager->GetRepresentedClass();
+		return LWIManager->GetRepresentedClass(Handle.GetInstanceIndex());
 	}
 
 	return nullptr;
@@ -308,9 +295,9 @@ FVector FLightWeightInstanceSubsystem::GetLocation(const FActorInstanceHandle& H
 {
 	ensure(Handle.IsValid());
 
-	if (Handle.Actor.IsValid())
+	if (const AActor* CachedActor = Handle.GetCachedActor())
 	{
-		return Handle.Actor->GetActorLocation();
+		return CachedActor->GetActorLocation();
 	}
 
 	if (ALightWeightInstanceManager* InstanceManager = FindLightWeightInstanceManager(Handle))
@@ -325,9 +312,9 @@ FString FLightWeightInstanceSubsystem::GetName(const FActorInstanceHandle& Handl
 {
 	ensure(Handle.IsValid());
 
-	if (Handle.Actor.IsValid())
+	if (const AActor* CachedActor = Handle.GetCachedActor())
 	{
-		return Handle.Actor->GetName();
+		return CachedActor->GetName();
 	}
 
 	if (ALightWeightInstanceManager* InstanceManager = FindLightWeightInstanceManager(Handle))
@@ -342,9 +329,9 @@ ULevel* FLightWeightInstanceSubsystem::GetLevel(const FActorInstanceHandle& Hand
 {
 	ensure(Handle.IsValid());
 
-	if (Handle.Actor.IsValid())
+	if (AActor* CachedActor = Handle.GetCachedActor())
 	{
-		return Handle.Actor->GetLevel();
+		return CachedActor->GetLevel();
 	}
 
 	if (ALightWeightInstanceManager* InstanceManager = FindLightWeightInstanceManager(Handle))
@@ -359,9 +346,9 @@ bool FLightWeightInstanceSubsystem::IsInLevel(const FActorInstanceHandle& Handle
 {
 	ensure(Handle.IsValid());
 
-	if (Handle.Actor.IsValid())
+	if (AActor* CachedActor = Handle.GetCachedActor())
 	{
-		return Handle.Actor->IsInLevel(InLevel);
+		return CachedActor->IsInLevel(InLevel);
 	}
 
 	if (ALightWeightInstanceManager* InstanceManager = FindLightWeightInstanceManager(Handle))
@@ -382,7 +369,7 @@ FActorInstanceHandle FLightWeightInstanceSubsystem::CreateNewLightWeightInstance
 			// create an instance with the given data
 			int32 InstanceIdx = LWIManager->AddNewInstance(InitData);
 			InstanceIdx = LWIManager->ConvertInternalIndexToHandleIndex(InstanceIdx);
-			return FActorInstanceHandle(LWIManager, InstanceIdx);
+			return FActorInstanceHandle(FActorInstanceManagerInterface(LWIManager), InstanceIdx);
 		}
 	}
 
diff --git a/Engine/Source/Runtime/Engine/Private/World.cpp b/Engine/Source/Runtime/Engine/Private/World.cpp
--- a/Engine/Source/Runtime/Engine/Private/World.cpp
+++ b/Engine/Source/Runtime/Engine/Private/World.cpp
@@ -602,7 +602,6 @@ UWorld::UWorld( const FObjectInitializer& ObjectInitializer )
 #if WITH_EDITOR
 , HierarchicalLODBuilder(new FHierarchicalLODBuilder(this))
 #endif
-, LWILastAssignedUID(0)
 , URL(FURL(NULL))
 ,	FXSystem(NULL)
 ,	TickTaskLevel(FTickTaskManagerInterface::Get().AllocateTickTaskLevel())
