Updated LWI call sites to use new generalized API

diff --git a/Engine/Plugins/Runtime/GameplayAbilities/Source/GameplayAbilities/Private/Abilities/GameplayAbilityTargetActor_SingleLineTrace.cpp b/Engine/Plugins/Runtime/GameplayAbilities/Source/GameplayAbilities/Private/Abilities/GameplayAbilityTargetActor_SingleLineTrace.cpp
--- a/Engine/Plugins/Runtime/GameplayAbilities/Source/GameplayAbilities/Private/Abilities/GameplayAbilityTargetActor_SingleLineTrace.cpp
+++ b/Engine/Plugins/Runtime/GameplayAbilities/Source/GameplayAbilities/Private/Abilities/GameplayAbilityTargetActor_SingleLineTrace.cpp
@@ -2,7 +2,6 @@
 
 #include "Abilities/GameplayAbilityTargetActor_SingleLineTrace.h"
 #include "DrawDebugHelpers.h"
-#include "GameFramework/LightWeightInstanceSubsystem.h"
 
 #include UE_INLINE_GENERATED_CPP_BY_NAME(GameplayAbilityTargetActor_SingleLineTrace)
 
@@ -44,7 +43,9 @@ FHitResult AGameplayAbilityTargetActor_SingleLineTrace::PerformTrace(AActor* InS
 	if (AGameplayAbilityWorldReticle* LocalReticleActor = ReticleActor.Get())
 	{
 		const bool bHitActor = (ReturnHitResult.bBlockingHit && (ReturnHitResult.HitObjectHandle.IsValid()));
-		const FVector ReticleLocation = (bHitActor && LocalReticleActor->bSnapToTargetedActor) ? FLightWeightInstanceSubsystem::Get().GetLocation(ReturnHitResult.HitObjectHandle) : ReturnHitResult.Location;
+		const FVector ReticleLocation = (bHitActor && LocalReticleActor->bSnapToTargetedActor) 
+			? ReturnHitResult.HitObjectHandle.GetLocation()
+			: ReturnHitResult.Location;
 
 		LocalReticleActor->SetActorLocation(ReticleLocation);
 		LocalReticleActor->SetIsTargetAnActor(bHitActor);
diff --git a/Engine/Source/Runtime/Engine/Private/Collision/CollisionConversions.cpp b/Engine/Source/Runtime/Engine/Private/Collision/CollisionConversions.cpp
--- a/Engine/Source/Runtime/Engine/Private/Collision/CollisionConversions.cpp
+++ b/Engine/Source/Runtime/Engine/Private/Collision/CollisionConversions.cpp
@@ -14,7 +14,7 @@
 #include "Physics/Experimental/ChaosInterfaceWrapper.h"
 
 #include "PhysicsEngine/CollisionQueryFilterCallback.h"
-#include "GameFramework/LightWeightInstanceManager.h"
+#include "Engine/ActorInstanceManagerInterface.h"
 #include "PhysicsProxy/SingleParticlePhysicsProxy.h"
 #include "PhysicsProxy/GeometryCollectionPhysicsProxy.h"
 
@@ -186,14 +186,7 @@ static void SetHitResultFromShapeAndFaceIndex(const FPhysicsShape& Shape,  const
 	{
 		OutResult.Component = OwningComponent;
 		AActor* Owner = OwningComponent->GetOwner();
-		if (ALightWeightInstanceManager* LWIManager = Cast<ALightWeightInstanceManager>(Owner))
-		{
-			OutResult.HitObjectHandle = FActorInstanceHandle(LWIManager, OutResult.Item);
-		}
-		else
-		{
-			OutResult.HitObjectHandle = FActorInstanceHandle(OwningComponent->GetOwner());
-		}
+		OutResult.HitObjectHandle = FActorInstanceHandle(Owner);
 
 		if (bReturnPhysMat)
 		{
@@ -610,9 +603,10 @@ void ConvertQueryOverlap(const FPhysicsShape& Shape, const FPhysicsActor& Actor,
         BodyInst = FPhysicsInterface::ShapeToOriginalBodyInstance(BodyInst, &Shape);
 		if (const UPrimitiveComponent* OwnerComponent = BodyInst->OwnerComponent.Get())
 		{
-			if (ALightWeightInstanceManager* LWIManager = Cast<ALightWeightInstanceManager>(OwnerComponent->GetOwner()))
+			UObject* Outer = OwnerComponent->GetOuter();
+			if (Outer && Outer->Implements<UActorInstanceManagerInterface>())
 			{
-				OutOverlap.OverlapObjectHandle = FActorInstanceHandle(LWIManager, BodyInst->InstanceBodyIndex);
+				OutOverlap.OverlapObjectHandle = FActorInstanceHandle(Outer, OwnerComponent, BodyInst->InstanceBodyIndex);
 			}
 			else
 			{
diff --git a/Engine/Source/Runtime/Engine/Private/DebugCameraHUD.cpp b/Engine/Source/Runtime/Engine/Private/DebugCameraHUD.cpp
--- a/Engine/Source/Runtime/Engine/Private/DebugCameraHUD.cpp
+++ b/Engine/Source/Runtime/Engine/Private/DebugCameraHUD.cpp
@@ -12,7 +12,6 @@
 #include "Engine/World.h"
 #include "GameFramework/PawnMovementComponent.h"
 #include "GameFramework/SpectatorPawn.h"
-#include "GameFramework/LightWeightInstanceSubsystem.h"
 #include "Engine/Canvas.h"
 #include "Engine/DebugCameraController.h"
 #include "Materials/MaterialInterface.h"
@@ -155,7 +154,7 @@ void ADebugCameraHUD::PostRender()
 				yl += Y;
 				Canvas->DrawText(RenderFont, FString::Printf(TEXT("HitDist: %f"), Hit.Distance), X, yl, 1.f, 1.f, FontRenderInfo);
 				yl += Y;
-				Canvas->DrawText(RenderFont, FString::Printf(TEXT("HitObject: '%s'"), *FLightWeightInstanceSubsystem::Get().GetName(HitHandle)), X, yl, 1.f, 1.f, FontRenderInfo);
+				Canvas->DrawText(RenderFont, FString::Printf(TEXT("HitObject: '%s'"), *HitHandle.GetName()), X, yl, 1.f, 1.f, FontRenderInfo);
 				yl += Y;
 				Canvas->DrawText(RenderFont, FString::Printf(TEXT("HitComponent: '%s'"), Hit.Component.Get() ? *Hit.Component.Get()->GetFName().ToString() : TEXT("<NULL>")), X, yl, 1.f, 1.f, FontRenderInfo);
 				yl += Y;
diff --git a/Engine/Source/Runtime/Engine/Private/Engine/HitResult.cpp b/Engine/Source/Runtime/Engine/Private/Engine/HitResult.cpp
--- a/Engine/Source/Runtime/Engine/Private/Engine/HitResult.cpp
+++ b/Engine/Source/Runtime/Engine/Private/Engine/HitResult.cpp
@@ -2,7 +2,6 @@
 
 #include "Engine/HitResult.h"
 #include "PhysicalMaterials/PhysicalMaterial.h"
-#include "GameFramework/LightWeightInstanceSubsystem.h"
 #include "Components/PrimitiveComponent.h"
 
 #include UE_INLINE_GENERATED_CPP_BY_NAME(HitResult)
@@ -158,7 +157,7 @@ FString FHitResult::ToString() const
 		PenetrationDepth,
 		Item,
 		PhysMaterial.IsValid() ? *PhysMaterial->GetName() : TEXT("None"),
-		HitObjectHandle.IsValid() ? *FLightWeightInstanceSubsystem::Get().GetName(HitObjectHandle) : TEXT("None"),
+		HitObjectHandle.IsValid() ? *HitObjectHandle.GetName() : TEXT("None"),
 		Component.IsValid() ? *Component->GetName() : TEXT("None"),
 		BoneName.IsValid() ? *BoneName.ToString() : TEXT("None"),
 		FaceIndex);
diff --git a/Engine/Source/Runtime/Engine/Private/EngineUtils.cpp b/Engine/Source/Runtime/Engine/Private/EngineUtils.cpp
--- a/Engine/Source/Runtime/Engine/Private/EngineUtils.cpp
+++ b/Engine/Source/Runtime/Engine/Private/EngineUtils.cpp
@@ -13,7 +13,6 @@
 #include "Misc/PackageName.h"
 #include "Misc/EngineVersion.h"
 #include "GameFramework/PlayerController.h"
-#include "GameFramework/LightWeightInstanceSubsystem.h"
 #include "GenericPlatform/ICursor.h"
 #include "Elements/Framework/EngineElementsLibrary.h"
 #include "Components/PrimitiveComponent.h"
@@ -520,7 +519,7 @@ TArray<FSubLevelStatus> GetSubLevelsStatus( UWorld* World, bool SortByActorCount
 
 				if (Hit.HitObjectHandle.IsValid())
 				{
-					LevelPlayerIsIn = FLightWeightInstanceSubsystem::Get().GetLevel(Hit.HitObjectHandle);
+					LevelPlayerIsIn = Hit.HitObjectHandle.GetLevel();
 				}
 				else if (UPrimitiveComponent* HitComponent = Hit.Component.Get())
 				{
