Refactored UWorldPartitionLevelStreamingDynamic to go through ULevel::CleanupLevel for trashing packages and clearing standalone flag.

diff --git a/Engine/Source/Runtime/Engine/Classes/Engine/Level.h b/Engine/Source/Runtime/Engine/Classes/Engine/Level.h
--- a/Engine/Source/Runtime/Engine/Classes/Engine/Level.h
+++ b/Engine/Source/Runtime/Engine/Classes/Engine/Level.h
@@ -584,8 +584,13 @@ public:
 
 	/** Whether the level should call FixupActorFolders on its actors when loading the level/actors (only used when level is using actor folder objects) */
 	uint8										bFixupActorFoldersAtLoad:1;
+
+private:
+	/** Whether the level is set not to be reusable after unload (editor-only) */
+	uint8										bForceCantReuseUnloadedButStillAround:1;
 #endif
 	
+public:
 	/** The below variables are used temporarily while making a level visible.				*/
 
 	/** Whether we already moved actors.													*/
@@ -733,6 +738,9 @@ public:
 
 	ENGINE_API void SetEditorPathOwner(UObject* InEditorPathOwner) { EditorPathOwner = InEditorPathOwner; }
 	ENGINE_API virtual UObject* GetEditorPathOwner() const override { return EditorPathOwner.Get(); }
+
+	ENGINE_API bool GetForceCantReuseUnloadedButStillAround() const { return bForceCantReuseUnloadedButStillAround; }
+	ENGINE_API void SetForceCantReuseUnloadedButStillAround(bool bNewValue) { bForceCantReuseUnloadedButStillAround = bNewValue; }
 #endif
 
 private:
diff --git a/Engine/Source/Runtime/Engine/Private/Level.cpp b/Engine/Source/Runtime/Engine/Private/Level.cpp
--- a/Engine/Source/Runtime/Engine/Private/Level.cpp
+++ b/Engine/Source/Runtime/Engine/Private/Level.cpp
@@ -391,6 +391,7 @@ ULevel::ULevel( const FObjectInitializer& ObjectInitializer )
 	bPromptWhenAddingToLevelOutsideBounds = true;
 	bUseActorFolders = false;
 	bFixupActorFoldersAtLoad = IsActorFolderObjectsFeatureAvailable();
+	bForceCantReuseUnloadedButStillAround = false;
 #endif	
 	bActorClusterCreated = false;
 	bGarbageCollectionClusteringEnabled = true;
diff --git a/Engine/Source/Runtime/Engine/Private/LevelStreaming.cpp b/Engine/Source/Runtime/Engine/Private/LevelStreaming.cpp
--- a/Engine/Source/Runtime/Engine/Private/LevelStreaming.cpp
+++ b/Engine/Source/Runtime/Engine/Private/LevelStreaming.cpp
@@ -130,6 +130,12 @@ bool ULevelStreaming::ShouldServerUseMakingVisibleTransactionRequest()
 
 bool ULevelStreaming::ShouldReuseUnloadedButStillAroundLevels(const ULevel* InLevel)
 {
+#if WITH_EDITOR
+	if (InLevel && InLevel->GetForceCantReuseUnloadedButStillAround())
+	{
+		return false;
+	}
+#endif
 	UWorld* OuterWorld = InLevel ? InLevel->GetTypedOuter<UWorld>() : nullptr;
 	if (OuterWorld && OuterWorld->IsGameWorld() && !LevelStreamingCVars::bShouldReuseUnloadedButStillAroundLevels)
 	{
@@ -144,7 +150,7 @@ int32 ULevelStreamingDynamic::UniqueLevelInstanceId = 0;
  * This helper function is defined here so that it can go into the 4.18.1 hotfix (for UE-51791),
  * even though it would make more logical sense to have this logic in a member function of UNetDriver.
  * We're getting away with this because UNetDriver::GuidCache is (unfortunately) public.
- *
+ * 
  * Renames any package entries in the GuidCache with a path matching UnPrefixedName to have a PIE prefix.
  * This is needed because a client may receive an export for a level package before it's loaded and
  * its name registered with FSoftObjectPath::AddPIEPackageName. In this case, the entry in the GuidCache
diff --git a/Engine/Source/Runtime/Engine/Private/WorldPartition/WorldPartitionLevelStreamingDynamic.cpp b/Engine/Source/Runtime/Engine/Private/WorldPartition/WorldPartitionLevelStreamingDynamic.cpp
--- a/Engine/Source/Runtime/Engine/Private/WorldPartition/WorldPartitionLevelStreamingDynamic.cpp
+++ b/Engine/Source/Runtime/Engine/Private/WorldPartition/WorldPartitionLevelStreamingDynamic.cpp
@@ -168,6 +168,9 @@ void UWorldPartitionLevelStreamingDynamic::CreateRuntimeLevel()
 	RuntimeLevel = FWorldPartitionLevelHelper::CreateEmptyLevelForRuntimeCell(StreamingCell.Get(), World, GetWorldAsset().ToString());
 	check(RuntimeLevel);
 
+	// Force world partition level/actor packages not to be reused
+	RuntimeLevel->SetForceCantReuseUnloadedButStillAround(true);
+
 	// Make sure Actor Folders is disabled on generated runtime levels to avoid any problems with duplicate folders that
 	// can be caused by level instances injecting their actors, which can cause duplicate folders (which only happens during PIE).
 	FLevelActorFoldersHelper::SetUseActorFolders(RuntimeLevel, false);
@@ -532,38 +535,6 @@ void UWorldPartitionLevelStreamingDynamic::OnCleanupLevel()
 		check(OnCleanupLevelDelegateHandle.IsValid());
 		RuntimeLevel->OnCleanupLevel.Remove(OnCleanupLevelDelegateHandle);
 		OnCleanupLevelDelegateHandle.Reset();
-
-		// If reusing levels is enabled, trash world partition level/actor packages
-		// as it won't be done by ULevel::CleanupLevel
-		if (ShouldReuseUnloadedButStillAroundLevels(RuntimeLevel))
-		{
-			TSet<UPackage*> TrashedPackages;
-			auto TrashPackage = [&TrashedPackages](UPackage* Package)
-			{
-				bool bWasAlreadyInSet;
-				TrashedPackages.Add(Package, &bWasAlreadyInSet);
-
-				if (!bWasAlreadyInSet)
-				{
-					// Clears RF_Standalone flag on objects in package (UMetaData)
-					ForEachObjectWithPackage(Package, [](UObject* Object) { Object->ClearFlags(RF_Standalone); return true; }, false);
-
-					// Rename package to avoid having to deal with pending kill objects in subsequent RequestLevel calls
-					FName NewPackageName = MakeUniqueObjectName(nullptr, UPackage::StaticClass(), FName(*FString::Printf(TEXT("%s_Trashed"), *Package->GetName())));
-					Package->Rename(*NewPackageName.ToString(), nullptr, REN_ForceNoResetLoaders | REN_DontCreateRedirectors | REN_NonTransactional | REN_DoNotDirty);
-				}
-			};
-
-			TrashPackage(RuntimeLevel->GetPackage());
-			for (AActor* Actor : RuntimeLevel->Actors)
-			{
-				if (UPackage* ActorPackage = Actor ? Actor->GetExternalPackage() : nullptr)
-				{
-					TrashPackage(ActorPackage);
-				}
-			}
-		}
-
 		RuntimeLevel = nullptr;
 	}
 	else
