Add some checks to try to track down source of corruption.

These checks reflect invariants that should always hold.

diff --git a/Engine/Source/Runtime/Engine/Private/Streaming/TextureInstanceState.cpp b/Engine/Source/Runtime/Engine/Private/Streaming/TextureInstanceState.cpp
--- a/Engine/Source/Runtime/Engine/Private/Streaming/TextureInstanceState.cpp
+++ b/Engine/Source/Runtime/Engine/Private/Streaming/TextureInstanceState.cpp
@@ -109,6 +109,7 @@ void FRenderAssetInstanceState::AddElement(const UPrimitiveComponent* InComponen
 	{
 		const bool bAllowShrinking = false;
 		ElementIndex = FreeElementIndices.Pop(bAllowShrinking);
+		check(ElementIndex < Elements.Num());
 	}
 	else
 	{
@@ -127,6 +128,7 @@ void FRenderAssetInstanceState::AddElement(const UPrimitiveComponent* InComponen
 	FRenderAssetDesc* AssetDesc = RenderAssetMap.Find(InAsset);
 	if (AssetDesc)
 	{
+		check(AssetDesc->HeadLink < Elements.Num());
 		FElement& AssetLinkElement = Elements[AssetDesc->HeadLink];
 
 		// The new inserted element as the head element.
@@ -137,8 +139,11 @@ void FRenderAssetInstanceState::AddElement(const UPrimitiveComponent* InComponen
 	else
 	{
 		RenderAssetMap.Add(InAsset, FRenderAssetDesc(ElementIndex, InAsset->GetLODGroupForStreaming()));
+		check(Element.NextRenderAssetLink == INDEX_NONE);
 	}
 
+	check(Element.PrevRenderAssetLink == INDEX_NONE);
+
 	// Simple sanity check to ensure that the component link passed in param is the right one
 	checkSlow(ComponentLink == ComponentMap.Find(InComponent));
 	if (ComponentLink)
@@ -175,6 +180,7 @@ void FRenderAssetInstanceState::AddElement(const UPrimitiveComponent* InComponen
 
 void FRenderAssetInstanceState::RemoveElement(int32 ElementIndex, int32& NextComponentLink, int32& BoundsIndex, const UStreamableRenderAsset*& Asset)
 {
+	check(ElementIndex < Elements.Num());
 	FElement& Element = Elements[ElementIndex];
 	NextComponentLink = Element.NextComponentLink; 
 	BoundsIndex = Element.BoundsIndex; 
@@ -231,6 +237,7 @@ void FRenderAssetInstanceState::RemoveElement(int32 ElementIndex, int32& NextCom
 	}
 	else
 	{
+		check(RenderAssetMap.IsEmpty());
 		Elements.Empty();
 		FreeElementIndices.Empty();
 	}
