Adding 'iris' CSV metadata tag and add to PRS tables so we can track which builds had it enabled for when we extend playtests outside of dedicated Iris playtests.
Flag is set to '1' if the GameNetDriver has "IsUsingIrisReplication" set to true.
This logic is done in UNetDriver::PostCreation because doing it inside InitBase is too early on clients, as the PendingNetDriver will not have been set to use Iris yet.
Works on Client and Server - tested against a Cooked build dedicated windows server + client
Preflight: https://horde.devtools.epicgames.com/job/650c04854a2c4e888528fb17

diff --git a/Engine/Source/Runtime/Engine/Private/NetDriver.cpp b/Engine/Source/Runtime/Engine/Private/NetDriver.cpp
--- a/Engine/Source/Runtime/Engine/Private/NetDriver.cpp
+++ b/Engine/Source/Runtime/Engine/Private/NetDriver.cpp
@@ -5543,6 +5543,13 @@ void UNetDriver::PostCreation(bool bInitializeWithIris)
 #if UE_WITH_IRIS
 	bIsUsingIris = bInitializeWithIris;
 #endif //UE_WITH_IRIS
+
+	if (NetDriverDefinition == NAME_GameNetDriver)
+	{
+		//Add to CSV whether we're using Iris on the GameNetDriver or not
+		CSV_METADATA(TEXT("Iris"), IsUsingIrisReplication() ? TEXT("1") : TEXT("0"));
+	}
+
 }
 
 #if NET_DEBUG_RELEVANT_ACTORS
