Fix missing material connection tracking with material attributes causing Substrate material conversion error.

Material attributes connection tracking was not done correctly for material function.

diff --git a/Engine/Source/Runtime/Engine/Private/Materials/MaterialExpressions.cpp b/Engine/Source/Runtime/Engine/Private/Materials/MaterialExpressions.cpp
--- a/Engine/Source/Runtime/Engine/Private/Materials/MaterialExpressions.cpp
+++ b/Engine/Source/Runtime/Engine/Private/Materials/MaterialExpressions.cpp
@@ -26474,12 +26474,10 @@ void UMaterialExpressionSubstrateThinFilm::GetExpressionToolTip(TArray<FString>&
 
 // Return a conservative list of connected material attribute inputs
 #if WITH_EDITOR
-static uint64 GetConnectedMaterialAttributesInputs(const UMaterial* InMaterial)
+static uint64 GetConnectedMaterialAttributesInputs(TConstArrayView<TObjectPtr<UMaterialExpression>> Expressions)
 {
-	if (!InMaterial) return 0;
-
 	uint64 Out = 0ull;
-	for (const UMaterialExpression* Expression : InMaterial->GetExpressions())
+	for (const UMaterialExpression* Expression : Expressions)
 	{
 		if (Expression)
 		{
@@ -26493,11 +26491,25 @@ static uint64 GetConnectedMaterialAttributesInputs(const UMaterial* InMaterial)
 				const UMaterialExpressionMakeMaterialAttributes* Attr = Cast<UMaterialExpressionMakeMaterialAttributes>(Expression);
 				Out |= Attr->GetConnectedInputs();
 			}
+			else if (Expression->IsA(UMaterialExpressionMaterialFunctionCall::StaticClass()))
+			{
+				const UMaterialExpressionMaterialFunctionCall* Attr = Cast<UMaterialExpressionMaterialFunctionCall>(Expression);
+				if (Attr->MaterialFunction)
+				{
+					Out |= GetConnectedMaterialAttributesInputs(Attr->MaterialFunction->GetExpressions());
+				}
+			}
 		}
 	}
 	return Out;
 }
 
+static uint64 GetConnectedMaterialAttributesInputs(const UMaterial* InMaterial)
+{
+	if (!InMaterial) return 0;
+	return GetConnectedMaterialAttributesInputs(InMaterial->GetExpressions());
+}
+
 static bool IsCustomMaterialAttributeInputConnected(uint64 InCache, FGuid InProperty)
 {
 	static TArray<FMaterialCustomOutputAttributeDefintion> CustomAttributeList;
