UInstancedStaticMeshComponent::Serialize replace check with if

It is possible for an ISM component to lose its BP Archetype between its Owning Actors latest save / load if the Packed BP has been repacked and removed some components in the process.
In this specific case the ISM will have no instance data serialized and will actually get destroyed on rerun construction script.

#rnx

diff --git a/Engine/Source/Runtime/Engine/Private/InstancedStaticMesh.cpp b/Engine/Source/Runtime/Engine/Private/InstancedStaticMesh.cpp
--- a/Engine/Source/Runtime/Engine/Private/InstancedStaticMesh.cpp
+++ b/Engine/Source/Runtime/Engine/Private/InstancedStaticMesh.cpp
@@ -3629,10 +3629,10 @@ void UInstancedStaticMeshComponent::Serialize(FArchive& Ar)
 
 			// Make sure that if we inherit the PerInstanceSMCustomData we also do inherit this value
 			NumCustomDataFloats = Archetype->NumCustomDataFloats;
-		}
-		else
+		} 
+		// It is possible for a component to lose its BP archetype between a save / load so in this case we have no per instance data (usually this component gets deleted through construction script)
+		else if(bHasSkipSerializationPropertiesData)
 		{
-			check(bHasSkipSerializationPropertiesData);
 			PerInstanceSMData = MoveTemp(TempPerInstanceSMData);
 			PerInstanceSMCustomData = MoveTemp(TempPerInstanceSMCustomData);
 		}
