Actually call the  "OnForceFeedbackFinished" on Force Feedback components when the component stops playing the effect.

diff --git a/Engine/Source/Runtime/Engine/Private/Components/ForceFeedbackComponent.cpp b/Engine/Source/Runtime/Engine/Private/Components/ForceFeedbackComponent.cpp
--- a/Engine/Source/Runtime/Engine/Private/Components/ForceFeedbackComponent.cpp
+++ b/Engine/Source/Runtime/Engine/Private/Components/ForceFeedbackComponent.cpp
@@ -316,6 +316,11 @@ void UForceFeedbackComponent::Stop()
 
 void UForceFeedbackComponent::StopInternal(const bool bRemoveFromManager)
 {
+	if (OnForceFeedbackFinished.IsBound())
+	{
+		OnForceFeedbackFinished.Broadcast(this);	
+	}
+	
 	// Set this to immediately be inactive
 	SetActiveFlag(false);
 	PlayTime = 0.f;
