WIP New HLSL translator bug fixes
Repair LWC support in vec/matric muls
Fix for FExpressionDynamicParameter not preparing its default value
Implement ThinTranslucentMaterialOutput GetCustomOutputType
Fix VT incorrect temporary local variable reuse out of scope
Fix conditionals coercing their inputs to bool types

#tests loaded various materials in editor, loaded level

diff --git a/Engine/Source/Runtime/Engine/Classes/Materials/MaterialExpressionThinTranslucentMaterialOutput.h b/Engine/Source/Runtime/Engine/Classes/Materials/MaterialExpressionThinTranslucentMaterialOutput.h
--- a/Engine/Source/Runtime/Engine/Classes/Materials/MaterialExpressionThinTranslucentMaterialOutput.h
+++ b/Engine/Source/Runtime/Engine/Classes/Materials/MaterialExpressionThinTranslucentMaterialOutput.h
@@ -31,5 +31,8 @@ public:
 	virtual int32 GetNumOutputs() const override;
 	virtual FString GetFunctionName() const override;
 	virtual FString GetDisplayName() const override;
+#if WITH_EDITOR
+	virtual UE::Shader::EValueType GetCustomOutputType(int32 OutputIndex) const override;
+#endif
 	//~ End UMaterialExpressionCustomOutput Interface
 };
diff --git a/Engine/Source/Runtime/Engine/Private/HLSLTree/HLSLTreeOperation.cpp b/Engine/Source/Runtime/Engine/Private/HLSLTree/HLSLTreeOperation.cpp
--- a/Engine/Source/Runtime/Engine/Private/HLSLTree/HLSLTreeOperation.cpp
+++ b/Engine/Source/Runtime/Engine/Private/HLSLTree/HLSLTreeOperation.cpp
@@ -195,6 +195,14 @@ FOperationRequestedTypes GetOperationRequestedTypes(EOperation Op, const FReques
 		Types.InputType[0] = Shader::EValueType::Float4x4;
 		Types.InputType[1] = Shader::EValueType::Float3;
 		break;
+	case EOperation::Less:
+	case EOperation::Greater:
+	case EOperation::LessEqual:
+	case EOperation::GreaterEqual:
+		// inputs for comparisons are some numeric type
+		Types.InputType[0] = Shader::EValueType::Numeric1;
+		Types.InputType[1] = Shader::EValueType::Numeric1;
+		break;
 	default:
 		break;
 	}
@@ -237,6 +245,7 @@ FOperationTypes GetOperationTypes(EOperation Op, TConstArrayView<FPreparedType>
 				Types.InputType[0] = Shader::EValueType::Double3;
 				Types.InputType[1] = Shader::EValueType::DoubleInverse4x4;
 				Types.ResultType = FPreparedType(Shader::EValueType::Float3, IntermediateComponent);
+				Types.bIsLWC = true;
 				break;
 			default:
 				// float3 * float4x4 -> float3
diff --git a/Engine/Source/Runtime/Engine/Private/Materials/MaterialExpressionHLSL.cpp b/Engine/Source/Runtime/Engine/Private/Materials/MaterialExpressionHLSL.cpp
--- a/Engine/Source/Runtime/Engine/Private/Materials/MaterialExpressionHLSL.cpp
+++ b/Engine/Source/Runtime/Engine/Private/Materials/MaterialExpressionHLSL.cpp
@@ -3547,6 +3547,17 @@ bool UMaterialExpressionThinTranslucentMaterialOutput::GenerateHLSLExpression(FM
 	return Generator.Error(TEXT("Invalid output"));
 }
 
+UE::Shader::EValueType UMaterialExpressionThinTranslucentMaterialOutput::GetCustomOutputType(int32 OutputIndex) const
+{
+	using namespace UE::Shader;
+	if (OutputIndex == 0)
+	{
+		return EValueType::Float3;
+	}
+
+	return EValueType::Void;
+}
+
 UE::Shader::EValueType UMaterialExpressionVolumetricAdvancedMaterialOutput::GetCustomOutputType(int32 OutputIndex) const
 {
 	if (OutputIndex >= 0 && OutputIndex < 6) return UE::Shader::EValueType::Float1;
diff --git a/Engine/Source/Runtime/Engine/Private/Materials/MaterialHLSLTree.cpp b/Engine/Source/Runtime/Engine/Private/Materials/MaterialHLSLTree.cpp
--- a/Engine/Source/Runtime/Engine/Private/Materials/MaterialHLSLTree.cpp
+++ b/Engine/Source/Runtime/Engine/Private/Materials/MaterialHLSLTree.cpp
@@ -703,6 +703,12 @@ bool FExpressionDynamicParameter::PrepareValue(FEmitContext& Context, FEmitScope
 		return Context.Error(TEXT("Invalid node used in hull/domain shader input!"));
 	}
 
+	const FPreparedType& DefaultType = Context.PrepareExpression(DefaultValueExpression, Scope, RequestedType);
+	if (DefaultType.IsVoid())
+	{
+		return false;
+	}
+
 	return OutResult.SetType(Context, RequestedType, EExpressionEvaluation::Shader, Shader::EValueType::Float4);
 }
 
@@ -846,7 +852,7 @@ uint32 AcquireVTStackIndex(
 	bool bGenerateFeedback)
 {
 	FHasher Hasher;
-	AppendHashes(Hasher, MipValueMode, AddressU, AddressV, AspectRatio, EmitTexCoordValue, EmitTexCoordValueDdx, EmitTexCoordValueDdy, EmitMipValue, PreallocatedStackTextureIndex, bAdaptive, bGenerateFeedback);
+	AppendHashes(Hasher, &Scope, MipValueMode, AddressU, AddressV, AspectRatio, EmitTexCoordValue, EmitTexCoordValueDdx, EmitTexCoordValueDdy, EmitMipValue, PreallocatedStackTextureIndex, bAdaptive, bGenerateFeedback);
 	const FXxHash64 Hash = Hasher.Finalize();
 
 	// First check to see if we have an existing VTStack that matches this key, that can still fit another layer
@@ -858,6 +864,7 @@ uint32 AcquireVTStackIndex(
 			Entry.EmitTexCoordValue == EmitTexCoordValue &&
 			Entry.EmitTexCoordValueDdx == EmitTexCoordValueDdx &&
 			Entry.EmitTexCoordValueDdy == EmitTexCoordValueDdy &&
+			Entry.Scope == &Scope &&
 			Entry.EmitMipValue == EmitMipValue &&
 			Entry.MipValueMode == MipValueMode &&
 			Entry.AddressU == AddressU &&
@@ -878,6 +885,7 @@ uint32 AcquireVTStackIndex(
 	Entry.EmitTexCoordValue = EmitTexCoordValue;
 	Entry.EmitTexCoordValueDdx = EmitTexCoordValueDdx;
 	Entry.EmitTexCoordValueDdy = EmitTexCoordValueDdy;
+	Entry.Scope = &Scope;
 	Entry.MipValueMode = MipValueMode;
 	Entry.AddressU = AddressU;
 	Entry.AddressV = AddressV;
diff --git a/Engine/Source/Runtime/Engine/Public/MaterialHLSLTree.h b/Engine/Source/Runtime/Engine/Public/MaterialHLSLTree.h
--- a/Engine/Source/Runtime/Engine/Public/MaterialHLSLTree.h
+++ b/Engine/Source/Runtime/Engine/Public/MaterialHLSLTree.h
@@ -875,6 +875,7 @@ struct FVTStackEntry
 	FEmitShaderExpression* EmitTexCoordValueDdx;
 	FEmitShaderExpression* EmitTexCoordValueDdy;
 	FEmitShaderExpression* EmitMipValue;
+	FEmitScope* Scope;
 	FEmitShaderExpression* EmitResult;
 	ETextureMipValueMode MipValueMode;
 	TextureAddress AddressU;
