World Partition = Fix WorldPartition streaming to take ViewTarget into account instead of player/spectator pawn location.

#rnx

diff --git a/Engine/Source/Runtime/Engine/Private/PlayerController.cpp b/Engine/Source/Runtime/Engine/Private/PlayerController.cpp
--- a/Engine/Source/Runtime/Engine/Private/PlayerController.cpp
+++ b/Engine/Source/Runtime/Engine/Private/PlayerController.cpp
@@ -3626,12 +3626,14 @@ void APlayerController::GetStreamingSourceLocationAndRotation(FVector& OutLocati
 {
 	if (!PlayerControllerCVars::ForceUsingCameraAsStreamingSource)
 	{
-		GetActorEyesViewPoint(OutLocation, OutRotation);
-	}
-	else
-	{
-		GetPlayerViewPoint(OutLocation, OutRotation);
+		if (const AActor* ViewTarget = GetViewTarget())
+		{
+			ViewTarget->GetActorEyesViewPoint(OutLocation, OutRotation);
+			return;
+		}
 	}
+
+	GetPlayerViewPoint(OutLocation, OutRotation);
 }
 
 void APlayerController::GetStreamingSourceShapes(TArray<FStreamingSourceShape>& OutShapes) const
