Changes for the WIP FHLSLMaterialTransaltor:
- Properly handle circles in material graphs during HLSL tree generartion. A material with inactive circles can still be succesfully translated. Errors will only be thrown when circles are connected to the final output.
- FExpressionFunctionCall and FExpressionLocalPHI now forward calls to EmitValueObject when possible.

diff --git a/Engine/Source/Runtime/Engine/Private/HLSLTree/HLSLTree.cpp b/Engine/Source/Runtime/Engine/Private/HLSLTree/HLSLTree.cpp
--- a/Engine/Source/Runtime/Engine/Private/HLSLTree/HLSLTree.cpp
+++ b/Engine/Source/Runtime/Engine/Private/HLSLTree/HLSLTree.cpp
@@ -50,6 +50,7 @@ public:
 	virtual bool PrepareValue(FEmitContext& Context, FEmitScope& Scope, const FRequestedType& RequestedType, FPrepareValueResult& OutResult) const override;
 	virtual void EmitValueShader(FEmitContext& Context, FEmitScope& Scope, const FRequestedType& RequestedType, FEmitValueShaderResult& OutResult) const override;
 	virtual void EmitValuePreshader(FEmitContext& Context, FEmitScope& Scope, const FRequestedType& RequestedType, FEmitValuePreshaderResult& OutResult) const override;
+	virtual bool EmitValueObject(FEmitContext& Context, FEmitScope& Scope, const FName& ObjectTypeName, void* OutObjectBase) const override;
 
 	TArray<FLocalPHIChainEntry, TInlineAllocator<8>> Chain;
 	FName LocalName;
@@ -71,6 +72,7 @@ public:
 	virtual bool PrepareValue(FEmitContext& Context, FEmitScope& Scope, const FRequestedType& RequestedType, FPrepareValueResult& OutResult) const override;
 	virtual void EmitValueShader(FEmitContext& Context, FEmitScope& Scope, const FRequestedType& RequestedType, FEmitValueShaderResult& OutResult) const override;
 	virtual void EmitValuePreshader(FEmitContext& Context, FEmitScope& Scope, const FRequestedType& RequestedType, FEmitValuePreshaderResult& OutResult) const override;
+	virtual bool EmitValueObject(FEmitContext& Context, FEmitScope& Scope, const FName& ObjectTypeName, void* OutObjectBase) const override;
 
 	FFunction* Function;
 	int32 OutputIndex;
@@ -489,6 +491,18 @@ void FExpressionLocalPHI::EmitValuePreshader(FEmitContext& Context, FEmitScope&
 	}
 }
 
+bool FExpressionLocalPHI::EmitValueObject(FEmitContext& Context, FEmitScope& Scope, const FName& ObjectTypeName, void* OutObjectBase) const
+{
+	Private::FLocalPHILiveScopes LiveScopes;
+	if (!Private::GetLiveScopes(Context, *this, LiveScopes) || !LiveScopes.bCanForwardValue)
+	{
+		// Cannot get if no live scope. Don't know which scope to use if there is more than one live
+		return false;
+	}
+
+	return LiveScopes.LiveValues[0]->GetValueObject(Context, Scope, ObjectTypeName, OutObjectBase);
+}
+
 void FExpressionFunctionCall::ComputeAnalyticDerivatives(FTree& Tree, FExpressionDerivatives& OutResult) const
 {
 	check(Function->OutputExpressions.IsValidIndex(OutputIndex));
@@ -529,6 +543,10 @@ void FExpressionFunctionCall::EmitValuePreshader(FEmitContext& Context, FEmitSco
 	OutResult.Type = Function->OutputExpressions[OutputIndex]->GetValuePreshader(Context, Scope, RequestedType, OutResult.Preshader);
 }
 
+bool FExpressionFunctionCall::EmitValueObject(FEmitContext& Context, FEmitScope& Scope, const FName& ObjectTypeName, void* OutObjectBase) const
+{
+	return Function->OutputExpressions[OutputIndex]->GetValueObject(Context, Scope, ObjectTypeName, OutObjectBase);
+}
 
 FRequestedType::FRequestedType(const Shader::FType& InType, bool bDefaultRequest) : Type(InType)
 {
diff --git a/Engine/Source/Runtime/Engine/Private/Materials/MaterialHLSLGenerator.cpp b/Engine/Source/Runtime/Engine/Private/Materials/MaterialHLSLGenerator.cpp
--- a/Engine/Source/Runtime/Engine/Private/Materials/MaterialHLSLGenerator.cpp
+++ b/Engine/Source/Runtime/Engine/Private/Materials/MaterialHLSLGenerator.cpp
@@ -182,6 +182,7 @@ bool FMaterialHLSLGenerator::GenerateResult(UE::HLSLTree::FScope& Scope)
 		{
 			UMaterialExpressionFunctionOutput* ExpressionOutput = FunctionEntry->FunctionOutputs[OutputIndex];
 			//FOwnerScope TreeOwnerScope(GetTree(), ExpressionOutput);
+			FScopedGenerateFunctionOutput GenerateOutputScope(FunctionEntry, OutputIndex);
 			HLSLFunction->OutputExpressions.Add(ExpressionOutput->A.TryAcquireHLSLExpression(*this, Scope, 0));
 		}
 		FunctionEntry->bGeneratedResult = true;
@@ -708,6 +709,15 @@ const UE::HLSLTree::FExpression* FMaterialHLSLGenerator::GenerateFunctionCall(UE
 				Hasher.Update(&ConnectedInput.Expression, sizeof(ConnectedInput.Expression));
 			}
 		}
+
+		// Hash the callstack too in case this function call is inside a material function.
+		// In that case, the addresses of FExpressionInputs won't change even if the owner
+		// MF is called at different places with different inputs
+		for (int32 Index = FunctionCallStack.Num() - 1; Index >= 0; --Index)
+		{
+			const FFunctionCallEntry* CallstackEntry = FunctionCallStack[Index];
+			Hasher.Update(&CallstackEntry, sizeof(CallstackEntry));
+		}
 		Hash = Hasher.Finalize();
 	}
 
@@ -742,11 +752,34 @@ const UE::HLSLTree::FExpression* FMaterialHLSLGenerator::GenerateFunctionCall(UE
 			check(FunctionCall->bGeneratedResult);
 		}
 	}
+	else if (FunctionCall->IsGeneratingOutput(OutputIndex))
+	{
+		FString CallstackStr;
+		for (int32 Index = FunctionCallStack.Num() - 1; Index >= 0; --Index)
+		{
+			const FFunctionCallEntry* Entry = FunctionCallStack[Index];
+			if (Entry && Entry->MaterialFunction)
+			{
+				CallstackStr.Appendf(TEXT("%s%s"), CallstackStr.Len() > 0 ? TEXT("<<") : TEXT(""), *Entry->MaterialFunction->GetName());
+			}
+		}
+		if (TargetMaterial)
+		{
+			CallstackStr.Appendf(TEXT("%s%s"), CallstackStr.Len() > 0 ? TEXT("<<") : TEXT(""), *TargetMaterial->GetName());
+		}
+
+		return NewErrorExpressionf(
+			TEXT("Circle found in material graph. MaterialFunction: %s, output %d. Callstack: %s"),
+			*MaterialFunction->GetName(),
+			OutputIndex,
+			*CallstackStr);
+	}
 
 	const FExpression* Result = nullptr;
 	FunctionCallStack.Add(FunctionCall);
 	if (bInlineFunction)
 	{
+		FScopedGenerateFunctionOutput GenerateOutputScope(FunctionCall, OutputIndex);
 		Result = ExpressionOutput->A.AcquireHLSLExpression(*this, Scope);
 	}
 	else
diff --git a/Engine/Source/Runtime/Engine/Public/MaterialHLSLGenerator.h b/Engine/Source/Runtime/Engine/Public/MaterialHLSLGenerator.h
--- a/Engine/Source/Runtime/Engine/Public/MaterialHLSLGenerator.h
+++ b/Engine/Source/Runtime/Engine/Public/MaterialHLSLGenerator.h
@@ -243,9 +243,46 @@ private:
 		FFunctionInputArray FunctionInputs;
 		FFunctionOutputArray FunctionOutputs;
 		FConnectedInputArray ConnectedInputs;
+		uint64 GeneratingOutputMask = 0; // Do we need to support more than 64 outputs?
 		EMaterialParameterAssociation ParameterAssociation = GlobalParameter;
 		int32 ParameterIndex = INDEX_NONE;
 		bool bGeneratedResult = false;
+
+		void BeginGeneratingOutput(int32 OutputIndex)
+		{
+			check(OutputIndex < 64);
+			GeneratingOutputMask |= (1llu << OutputIndex);
+		}
+
+		void EndGeneratingOutput(int32 OutputIndex)
+		{
+			check(IsGeneratingOutput(OutputIndex));
+			GeneratingOutputMask ^= (1llu << OutputIndex);
+		}
+
+		bool IsGeneratingOutput(int32 OutputIndex) const
+		{
+			return GeneratingOutputMask & (1llu << OutputIndex);
+		}
+	};
+
+	struct FScopedGenerateFunctionOutput
+	{
+		FScopedGenerateFunctionOutput(FFunctionCallEntry* InFunctionCall, int32 InOutputIndex)
+			: FunctionCall(InFunctionCall)
+			, OutputIndex(InOutputIndex)
+		{
+			FunctionCall->BeginGeneratingOutput(OutputIndex);
+		}
+
+		~FScopedGenerateFunctionOutput()
+		{
+			FunctionCall->EndGeneratingOutput(OutputIndex);
+		}
+
+	private:
+		FFunctionCallEntry* FunctionCall;
+		int32 OutputIndex;
 	};
 
 	struct FStatementEntry
