Fixing a non-deterministic replay crash caused by level streaming not always completing during replay playback by the same frame it completed when recording the replay. This would result in dropping dependent packets on the floor (or, at least, delivering them out of order in a way that the client isn't expecting to handle). Created a new cvar demo.DoNotDropPackets which will prefer to pause the playback until the prerequisites are satisfied. Gating behind a cvar to ensure backward compatibility.

diff --git a/Engine/Source/Runtime/Engine/Classes/Engine/DemoNetDriver.h b/Engine/Source/Runtime/Engine/Classes/Engine/DemoNetDriver.h
--- a/Engine/Source/Runtime/Engine/Classes/Engine/DemoNetDriver.h
+++ b/Engine/Source/Runtime/Engine/Classes/Engine/DemoNetDriver.h
@@ -451,11 +451,7 @@ public:
 	ENGINE_API bool ConditionallyReadDemoFrameIntoPlaybackPackets(FArchive& Ar);
 
 	ENGINE_API bool ProcessPacket(const uint8* Data, int32 Count);
-	bool ProcessPacket(const FPlaybackPacket& PlaybackPacket)
-	{
-		return ShouldSkipPlaybackPacket(PlaybackPacket) ||
-				ProcessPacket(PlaybackPacket.Data.GetData(), PlaybackPacket.Data.Num());
-	}
+	bool ProcessPacket(const FPlaybackPacket& PlaybackPacket);
 	
 	ENGINE_API void WriteDemoFrameFromQueuedDemoPackets(FArchive& Ar, TArray<FQueuedDemoPacket>& QueuedPackets, float FrameTime, EWriteDemoFrameFlags Flags);
 
diff --git a/Engine/Source/Runtime/Engine/Private/DemoNetDriver.cpp b/Engine/Source/Runtime/Engine/Private/DemoNetDriver.cpp
--- a/Engine/Source/Runtime/Engine/Private/DemoNetDriver.cpp
+++ b/Engine/Source/Runtime/Engine/Private/DemoNetDriver.cpp
@@ -88,6 +88,7 @@ static TAutoConsoleVariable<float> CVarDemoMinimumRepPrioritizeTime(TEXT("demo.M
 static TAutoConsoleVariable<float> CVarDemoMaximumRepPrioritizeTime(TEXT("demo.MaximumRepPrioritizePercent"), 0.7, TEXT("Maximum percent of time that may be spent prioritizing actors, regardless of throttling."));
 
 static TAutoConsoleVariable<int32> CVarFastForwardLevelsPausePlayback(TEXT("demo.FastForwardLevelsPausePlayback"), 0, TEXT("If true, pause channels and playback while fast forward levels task is running."));
+static TAutoConsoleVariable<bool> CVarDemoPlaybackDoNotDropPackets(TEXT("demo.DoNotDropPackets"), false, TEXT("If true, will prefer to wait to advance the demo until prerequisites such as level streaming have completed."));
 
 namespace ReplayTaskNames
 {
@@ -2616,6 +2617,24 @@ bool UDemoNetDriver::ProcessPacket(const uint8* Data, int32 Count)
 	return true;
 }
 
+bool UDemoNetDriver::ProcessPacket(const FPlaybackPacket& PlaybackPacket)
+{
+	// This is a bit unfortunate. In order to guarantee performance for when this is used in a shipping game
+	// we need to be willing to drop packets on the floor if we're not ready for them (e.g., because a level hasn't streamed in yet)
+	// But on the other hand, to get reliable behavior for debug/perf analysis we need to do the opposite and wait for the level to stream
+	// in if needed
+	if (CVarDemoPlaybackDoNotDropPackets.GetValueOnGameThread())
+	{
+		return !ShouldSkipPlaybackPacket(PlaybackPacket) &&
+			ProcessPacket(PlaybackPacket.Data.GetData(), PlaybackPacket.Data.Num());
+	}
+	else
+	{
+		return ShouldSkipPlaybackPacket(PlaybackPacket) ||
+			ProcessPacket(PlaybackPacket.Data.GetData(), PlaybackPacket.Data.Num());
+	}
+}
+
 void UDemoNetDriver::WriteDemoFrameFromQueuedDemoPackets(FArchive& Ar, TArray<FQueuedDemoPacket>& QueuedPackets, float FrameTime, EWriteDemoFrameFlags Flags)
 {
 	ReplayHelper.WriteDemoFrame(ClientConnections[0], Ar, QueuedPackets, FrameTime, Flags);
