PR #9729: Fix Visual Logger Capsule first param to match cpp

diff --git a/Engine/Source/Runtime/Engine/Public/VisualLogger/VisualLoggerTypes.h b/Engine/Source/Runtime/Engine/Public/VisualLogger/VisualLoggerTypes.h
--- a/Engine/Source/Runtime/Engine/Public/VisualLogger/VisualLoggerTypes.h
+++ b/Engine/Source/Runtime/Engine/Public/VisualLogger/VisualLoggerTypes.h
@@ -213,11 +213,11 @@ struct FVisualLogEntry
 	// box
 	ENGINE_API void AddElement(const FBox& Box, const FMatrix& Matrix, const FName& CategoryName, ELogVerbosity::Type Verbosity, const FColor& Color = FColor::White, const FString& Description = TEXT(""), uint16 Thickness = 0);
 	// Cone
-	ENGINE_API void AddElement(const FVector& Orgin, const FVector& Direction, float Length, float AngleWidth, float AngleHeight, const FName& CategoryName, ELogVerbosity::Type Verbosity, const FColor& Color = FColor::White, const FString& Description = TEXT(""), uint16 Thickness = 0);
+	ENGINE_API void AddElement(const FVector& Origin, const FVector& Direction, float Length, float AngleWidth, float AngleHeight, const FName& CategoryName, ELogVerbosity::Type Verbosity, const FColor& Color = FColor::White, const FString& Description = TEXT(""), uint16 Thickness = 0);
 	// Cylinder
 	ENGINE_API void AddElement(const FVector& Start, const FVector& End, float Radius, const FName& CategoryName, ELogVerbosity::Type Verbosity, const FColor& Color = FColor::White, const FString& Description = TEXT(""), uint16 Thickness = 0);
 	// capsule
-	ENGINE_API void AddElement(const FVector& Center, float HalfHeight, float Radius, const FQuat & Rotation, const FName& CategoryName, ELogVerbosity::Type Verbosity, const FColor& Color = FColor::White, const FString& Description = TEXT(""));
+	ENGINE_API void AddElement(const FVector& Base, float HalfHeight, float Radius, const FQuat & Rotation, const FName& CategoryName, ELogVerbosity::Type Verbosity, const FColor& Color = FColor::White, const FString& Description = TEXT(""));
 	// custom element
 	ENGINE_API void AddElement(const FVisualLogShapeElement& Element);
 	// NavAreaMesh
