FWorldPartitionCookPackageSplitter: added logs to track down an intermittent cooking issue.

diff --git a/Engine/Source/Runtime/Engine/Private/WorldPartition/WorldPartitionCookPackageSplitter.cpp b/Engine/Source/Runtime/Engine/Private/WorldPartition/WorldPartitionCookPackageSplitter.cpp
--- a/Engine/Source/Runtime/Engine/Private/WorldPartition/WorldPartitionCookPackageSplitter.cpp
+++ b/Engine/Source/Runtime/Engine/Private/WorldPartition/WorldPartitionCookPackageSplitter.cpp
@@ -31,6 +31,8 @@ FWorldPartitionCookPackageSplitter::~FWorldPartitionCookPackageSplitter()
 
 void FWorldPartitionCookPackageSplitter::Teardown(ETeardown Status)
 {
+	UE_LOG(LogWorldPartition, Log, TEXT("[Cook] Debug(TearDown): OwnerObject=%s bInitializedWorldPartition=%d bForceInitializedWorld=%d bInitializedPhysicsSceneForSave=%d"), *GetFullNameSafe(ReferencedWorld), bInitializedWorldPartition ? 1 : 0, bForceInitializedWorld ? 1 : 0, bInitializedPhysicsSceneForSave ? 1 : 0);
+
 	if (UWorld* LocalWorld = ReferencedWorld.Get())
 	{
 		check(LocalWorld->bUsedByCookSplitter);
@@ -117,6 +119,7 @@ TArray<ICookPackageSplitter::FGeneratedPackage> FWorldPartitionCookPackageSplitt
 	BuildPackagesToGenerateList(PackagesToGenerate);
 
 	UE_LOG(LogWorldPartition, Log, TEXT("[Cook] Sending %u packages to be generated."), PackagesToGenerate.Num());
+	UE_LOG(LogWorldPartition, Log, TEXT("[Cook] Debug(GetGenerateList) : OwnerObject=%s bInitializedWorldPartition=%d bForceInitializedWorld=%d bInitializedPhysicsSceneForSave=%d"), *GetFullNameSafe(OwnerObject), bInitializedWorldPartition ? 1 : 0, bForceInitializedWorld ? 1 : 0, bInitializedPhysicsSceneForSave ? 1 : 0);
 
 	return PackagesToGenerate;
 }
