Fixed bug in debug bone rendering, inconsistent use of scale when rendering cone connectors

diff --git a/Engine/Source/Runtime/Engine/Private/SkeletalDebugRendering.cpp b/Engine/Source/Runtime/Engine/Private/SkeletalDebugRendering.cpp
--- a/Engine/Source/Runtime/Engine/Private/SkeletalDebugRendering.cpp
+++ b/Engine/Source/Runtime/Engine/Private/SkeletalDebugRendering.cpp
@@ -71,16 +71,18 @@ void DrawWireBoneAdvanced(
 #if ENABLE_DRAW_DEBUG
 
 	const FVector BoneLocation = InBoneTransform.GetLocation();
+	FTransform BoneNoScale = InBoneTransform;
+	BoneNoScale.SetScale3D(FVector::OneVector);
 
 	// draw wire sphere at joint origin, oriented with the bone
-	DrawWireSphere(PDI, InBoneTransform, InColor, SphereRadius, NumSphereSides, InDepthPriority, 0.0f, 1.0f);
+	DrawWireSphere(PDI, BoneNoScale, InColor, SphereRadius, NumSphereSides, InDepthPriority, 0.0f, 1.0f);
 
 	// draw axes at joint location
 	if (InAxisConfig.bDraw)
 	{
 		const float Thickness = InAxisConfig.Thickness > 0.f ? InAxisConfig.Thickness : 0.f;
 		const float Length = InAxisConfig.Length > 0.f ? InAxisConfig.Length : SphereRadius;
-		SkeletalDebugRendering::DrawAxes(PDI, InBoneTransform, SDPG_Foreground, Thickness, Length);
+		SkeletalDebugRendering::DrawAxes(PDI, BoneNoScale, SDPG_Foreground, Thickness, Length);
 	}
 
 	// draw wire cones to each child
