[shaders]
- decouple preprocessed job cache & per shader ddc cook settings
- update cvar descriptions for the above

diff --git a/Engine/Source/Runtime/Engine/Private/ShaderCompiler/ShaderCompiler.cpp b/Engine/Source/Runtime/Engine/Private/ShaderCompiler/ShaderCompiler.cpp
--- a/Engine/Source/Runtime/Engine/Private/ShaderCompiler/ShaderCompiler.cpp
+++ b/Engine/Source/Runtime/Engine/Private/ShaderCompiler/ShaderCompiler.cpp
@@ -123,7 +123,7 @@ static FAutoConsoleVariableRef CVarShaderCompilerJobCacheOverflowReducePercent(
 static TAutoConsoleVariable<bool> CVarPreprocessedJobCache(
 	TEXT("r.ShaderCompiler.PreprocessedJobCache"),
 	false,
-	TEXT("EXPERIMENTAL: if enabled will preprocess jobs up-front when the job is queued and generate job input hashes based on preprocessed source."),
+	TEXT("If enabled will shader compile jobs will be preprocessed at submission time in the cook process (when the job is queued) and generate job input hashes based on preprocessed source."),
 	ECVF_Default
 );
 
@@ -136,7 +136,7 @@ static TAutoConsoleVariable<bool> CVarJobCacheDDC(
 static TAutoConsoleVariable<bool> CVarJobCacheDDCPolicy(
 	TEXT("r.ShaderCompiler.JobCacheDDCEnableRemotePolicy"),
 	false,
-	TEXT("Whether to cache shaders in the job cache to your local machine or remotely to the network.\n"),
+	TEXT("If true, individual shader jobs will be cached to remote/shared DDC instances; if false they will only cache to DDC instances on the local machine.\n"),
 	ECVF_ReadOnly);
 
 static TAutoConsoleVariable<bool> CVarDebugDumpWorkerInputs(
@@ -186,8 +186,8 @@ static FAutoConsoleVariableRef CVarShaderCompilerPerShaderDDCGlobal(
 
 static TAutoConsoleVariable<bool> CVarShaderCompilerPerShaderDDCCook(
 	TEXT("r.ShaderCompiler.PerShaderDDCCook"),
-	true,
-	TEXT("if != 0, Per-shader DDC queries enabled during cooks (requires r.ShaderCompiler.PreprocessedJobCache to also be enabled)."),
+	false,
+	TEXT("If true, per-shader DDC caching will be enabled during cooks."),
 	ECVF_Default
 );
 
@@ -201,7 +201,7 @@ bool IsShaderJobCacheDDCEnabled()
 {
 #if WITH_EDITOR
 	static const bool bForceAllowShaderCompilerJobCache = FParse::Param(FCommandLine::Get(), TEXT("forceAllowShaderCompilerJobCache")) ||
-		(CVarShaderCompilerPerShaderDDCCook.GetValueOnAnyThread() && CVarPreprocessedJobCache.GetValueOnAnyThread());
+		CVarShaderCompilerPerShaderDDCCook.GetValueOnAnyThread();
 #else
 	const bool bForceAllowShaderCompilerJobCache = false;
 #endif
