Cloth - Added a skeletal mesh component velocity scale property to reduce the component induced velocities on all its cloths. This also replaces the bLocalSpaceSimulation property that was used with Chaos Cloth and had been deprecated.

diff --git a/Engine/Source/Runtime/ClothingSystemRuntimeCommon/Private/ClothingSimulation.cpp b/Engine/Source/Runtime/ClothingSystemRuntimeCommon/Private/ClothingSimulation.cpp
--- a/Engine/Source/Runtime/ClothingSystemRuntimeCommon/Private/ClothingSimulation.cpp
+++ b/Engine/Source/Runtime/ClothingSystemRuntimeCommon/Private/ClothingSimulation.cpp
@@ -173,7 +173,7 @@ void FClothingSimulationContextCommon::FillTeleportMode(const USkeletalMeshCompo
 		InComponent->ClothTeleportMode;
 
 	VelocityScale = (TeleportMode == EClothingTeleportMode::None && InDeltaSeconds > 0.f) ?
-		FMath::Min(InDeltaSeconds, InMaxPhysicsDelta) / InDeltaSeconds : 0.f;
+		InComponent->ClothVelocityScale * FMath::Min(InDeltaSeconds, InMaxPhysicsDelta) / InDeltaSeconds : 0.f;
 }
 
 void FClothingSimulationContextCommon::FillMaxDistanceScale(const USkeletalMeshComponent* InComponent)
diff --git a/Engine/Source/Runtime/Engine/Classes/Components/SkeletalMeshComponent.h b/Engine/Source/Runtime/Engine/Classes/Components/SkeletalMeshComponent.h
--- a/Engine/Source/Runtime/Engine/Classes/Components/SkeletalMeshComponent.h
+++ b/Engine/Source/Runtime/Engine/Classes/Components/SkeletalMeshComponent.h
@@ -679,10 +679,14 @@ public:
 	UPROPERTY(EditAnywhere, BlueprintReadWrite, Category=Clothing)
 	uint8 bForceCollisionUpdate : 1;
 
-	/** Deprecated. */
-	UE_DEPRECATED(5.2, "Cloth simulation always happens in local space.")
-	UPROPERTY(BlueprintReadWrite, Category = Clothing, meta = (DeprecationMessage = "This property is deprecated. Cloth simulation always happens in local space."))
-	uint8 bLocalSpaceSimulation : 1;
+	/**
+	 * Scale applied to the component induced velocity of all cloth particles prior to stepping the cloth solver.
+	 * Use 1.0 for fully induced velocity (default), or use 0.0 for no induced velocity, and any other values in between for a reduced induced velocity.
+	 * When set to 0.0, it also provides a way to force the clothing to simulate in local space.
+	 * This value multiplies to individual cloth's velocity scale settings, still allowing for differences in behavior between the various pieces of clothing within the same component.
+	 */
+	UPROPERTY(EditAnywhere, BlueprintReadWrite, Category = Clothing, Meta = (UIMin = 0.f, UIMax = 1.f, ClampMin = 0.f, ClampMax = 1.f))
+	float ClothVelocityScale = 1.f;
 
 	/** reset the clothing after moving the clothing position (called teleport) */
 	UPROPERTY(EditAnywhere, BlueprintReadWrite, Category=Clothing)
