Fixed description of MTP_BeforeDOF in enum EMaterialTranslucencyPass.

diff --git a/Engine/Source/Runtime/Engine/Classes/Materials/Material.h b/Engine/Source/Runtime/Engine/Classes/Materials/Material.h
--- a/Engine/Source/Runtime/Engine/Classes/Materials/Material.h
+++ b/Engine/Source/Runtime/Engine/Classes/Materials/Material.h
@@ -147,7 +147,7 @@ enum EMaterialDecalResponse : int
 UENUM()
 enum EMaterialTranslucencyPass : int
 {
-	/** Render after depth of field. */
+	/** Render before depth of field. */
 	MTP_BeforeDOF UMETA(DisplayName="Before DOF"),
 	/** Render after depth of field. */
 	MTP_AfterDOF UMETA(DisplayName="After DOF"),
