[VisualLogger] Fixed VisualLogger ensure for objects without associated world.

diff --git a/Engine/Source/Runtime/Engine/Private/VisualLogger/VisualLogger.cpp b/Engine/Source/Runtime/Engine/Private/VisualLogger/VisualLogger.cpp
--- a/Engine/Source/Runtime/Engine/Private/VisualLogger/VisualLogger.cpp
+++ b/Engine/Source/Runtime/Engine/Private/VisualLogger/VisualLogger.cpp
@@ -275,9 +275,9 @@ TMap<const UWorld*, FVisualLogger::FOwnerToChildrenRedirectionMap> FVisualLogger
 int32 FVisualLogger::bIsRecording = false;
 FVisualLogger::FNavigationDataDump FVisualLogger::NavigationDataDumpDelegate;
 
-bool FVisualLogger::CheckVisualLogInputInternal(const UObject* Object, const FName& CategoryName, ELogVerbosity::Type Verbosity, UWorld **World, FVisualLogEntry **CurrentEntry)
+bool FVisualLogger::CheckVisualLogInputInternal(const UObject* Object, const FName& CategoryName, ELogVerbosity::Type, UWorld** OutWorld, FVisualLogEntry** OutCurrentEntry)
 {
-	if (FVisualLogger::IsRecording() == false || !Object || !GEngine || GEngine->bDisableAILogging || Object->HasAnyFlags(RF_ClassDefaultObject))
+	if (IsRecording() == false || !Object || !GEngine || GEngine->bDisableAILogging || Object->HasAnyFlags(RF_ClassDefaultObject))
 	{
 		return false;
 	}
@@ -288,14 +288,14 @@ bool FVisualLogger::CheckVisualLogInputInternal(const UObject* Object, const FNa
 		return false;
 	}
 
-	*World = GEngine->GetWorldFromContextObject(Object, EGetWorldErrorMode::ReturnNull);
-	if (ensure(*World != nullptr) == false)
+	*OutWorld = GetWorldForVisualLogger(Object);
+	if (!ensure(*OutWorld != nullptr))
 	{
 		return false;
 	}
 
-	*CurrentEntry = VisualLogger.GetEntryToWrite(Object, VisualLogger.GetTimeStampForObject(Object));
-	if (*CurrentEntry == nullptr)
+	*OutCurrentEntry = VisualLogger.GetEntryToWrite(Object, VisualLogger.GetTimeStampForObject(Object));
+	if (*OutCurrentEntry == nullptr)
 	{
 		return false;
 	}
@@ -310,12 +310,12 @@ double FVisualLogger::GetTimeStampForObject(const UObject* Object) const
 		return GetTimeStampFunc(Object);
 	}
 
-	if (const UWorld* World = GEngine->GetWorldFromContextObject(Object, EGetWorldErrorMode::ReturnNull))
+	if (const UWorld* World = GetWorldForVisualLogger(Object))
 	{
 		return World->TimeSeconds;
 	}
 
-	return 0.;
+	return 0;
 }
 
 void FVisualLogger::SetGetTimeStampFunc(const TFunction<double(const UObject*)> Function)
diff --git a/Engine/Source/Runtime/Engine/Public/VisualLogger/VisualLogger.h b/Engine/Source/Runtime/Engine/Public/VisualLogger/VisualLogger.h
--- a/Engine/Source/Runtime/Engine/Public/VisualLogger/VisualLogger.h
+++ b/Engine/Source/Runtime/Engine/Public/VisualLogger/VisualLogger.h
@@ -554,7 +554,7 @@ public:
 	const TMap<FName, FVisualLogExtensionInterface*>& GetAllExtensions() const { return AllExtensions; }
 
 	/** internal check for each usage of visual logger */
-	static ENGINE_API bool CheckVisualLogInputInternal(const UObject* Object, const FName& CategoryName, ELogVerbosity::Type Verbosity, UWorld **World, FVisualLogEntry **CurrentEntry);
+	static ENGINE_API bool CheckVisualLogInputInternal(const UObject* Object, const FName& CategoryName, ELogVerbosity::Type Verbosity, UWorld **OutWorld, FVisualLogEntry **OutCurrentEntry);
 	
 	/** Returns time stamp for object */
 	ENGINE_API double GetTimeStampForObject(const UObject* Object) const;
