Iris
-Added CVar to allow for additonal data packets to be written by ReplicationWriter as long as bandwidth settings are respected

diff --git a/Engine/Plugins/Runtime/ReplicationSystemTestPlugin/Source/Private/Tests/ReplicationSystem/ReplicationSystemServerClientTestFixture.cpp b/Engine/Plugins/Runtime/ReplicationSystemTestPlugin/Source/Private/Tests/ReplicationSystem/ReplicationSystemServerClientTestFixture.cpp
--- a/Engine/Plugins/Runtime/ReplicationSystemTestPlugin/Source/Private/Tests/ReplicationSystem/ReplicationSystemServerClientTestFixture.cpp
+++ b/Engine/Plugins/Runtime/ReplicationSystemTestPlugin/Source/Private/Tests/ReplicationSystem/ReplicationSystemServerClientTestFixture.cpp
@@ -250,8 +250,9 @@ bool FReplicationSystemTestNode::SendUpdate(uint32 ConnectionId, const TCHAR* De
 	FConnectionInfo& Connection = GetConnectionInfo(ConnectionId);
 
 	const FDataStreamRecord* Record = nullptr;
+	UDataStream::FBeginWriteParameters BeginWriteParameters;
 
-	const bool bResult = (Connection.DataStreamManager->BeginWrite() != UDataStream::EWriteResult::NoData) && (Connection.DataStreamManager->WriteData(Context, Record)  != UDataStream::EWriteResult::NoData);
+	const bool bResult = (Connection.DataStreamManager->BeginWrite(BeginWriteParameters) != UDataStream::EWriteResult::NoData) && (Connection.DataStreamManager->WriteData(Context, Record)  != UDataStream::EWriteResult::NoData);
 	if (bResult)
 	{
 		Writer.CommitWrites();
diff --git a/Engine/Source/Runtime/Engine/Private/Net/Experimental/Iris/DataStreamChannel.cpp b/Engine/Source/Runtime/Engine/Private/Net/Experimental/Iris/DataStreamChannel.cpp
--- a/Engine/Source/Runtime/Engine/Private/Net/Experimental/Iris/DataStreamChannel.cpp
+++ b/Engine/Source/Runtime/Engine/Private/Net/Experimental/Iris/DataStreamChannel.cpp
@@ -18,6 +18,7 @@
 #include "Iris/Serialization/NetSerializationContext.h"
 #include "Iris/Core/IrisProfiler.h"
 #include "Iris/Core/IrisMemoryTracker.h"
+#include "Iris/Core/IrisLog.h"
 #include "Net/Core/Trace/NetTrace.h"
 #include "PacketHandler.h"
 #include "ProfilingDebugging/CsvProfiler.h"
@@ -164,11 +165,18 @@ void UDataStreamChannel::Tick()
 		return;
 	}
 
-	if (!IsNetReady(UE::Net::Private::bIrisSaturateBandwidth) || IsPacketWindowFull() || !Connection->HasReceivedClientPacket() || (Connection->Handler != nullptr && !Connection->Handler->IsFullyInitialized()))
+	if (IsPacketWindowFull() || !Connection->HasReceivedClientPacket() || (Connection->Handler != nullptr && !Connection->Handler->IsFullyInitialized()))
 	{
 		return;
 	}
 
+	// We probably want separate bandwidth management for iris as we are not pre-filling sendbuffer before call to NetReady.
+	if (!IsNetReady(UE::Net::Private::bIrisSaturateBandwidth))
+	{
+		UE_LOG(LogIris, Warning, TEXT("Disallowed to write first object in batch, with Iris this is not good!"))
+		return;
+	}
+
 	// Wait for channel to open
 	if (!bHandshakeComplete)
 	{
@@ -183,10 +191,28 @@ void UDataStreamChannel::Tick()
 	IRIS_PROFILER_SCOPE(UDataStreamChannel_Tick);
 	LLM_SCOPE_BYTAG(Iris);
 
+	// Limit the amount of bits to minimum of a bunch and our buffer. NetBitStreamWriter requires the number of bytes to be a multiple of 4.
+	const uint32 MaxBitCount = uint32(Connection->GetMaxSingleBunchSizeBits());
+	const uint32 MaxBytes = FPlatformMath::Min((MaxBitCount/32U)*4U, (uint32)sizeof(BitStreamBuffer));
+	const int64 MaxBunchBits = MaxBytes*8;
+
+	// Try to determine if we have headroom to write more than a single packet if needed.
+	UDataStream::FBeginWriteParameters BeginWriteParams;
+	{
+		int32 CurrentQueuedBits = Connection->QueuedBits + Connection->SendBuffer.GetNumBits();
+		if (CurrentQueuedBits < 0)
+		{
+			const int32 NumBunchesThatMightBeWritten = (-CurrentQueuedBits) / MaxBunchBits;
+
+			// Indicate that DataStream can request more data
+			BeginWriteParams.bCanWriteMoreData = NumBunchesThatMightBeWritten > 1;
+		}
+	}
+
 	// Currently we want to use a full bunch so we flush if we have to
 	bool bNeedsPreSendFlush = Connection->SendBuffer.GetNumBits() > MAX_PACKET_HEADER_BITS;
 
-	auto WriteDataFunction = [this, &bNeedsPreSendFlush]()
+	auto WriteDataFunction = [this, &bNeedsPreSendFlush, &BeginWriteParams, MaxBunchBits, MaxBytes]()
 	{
 		if (bNeedsPreSendFlush)
 		{
@@ -201,10 +227,6 @@ void UDataStreamChannel::Tick()
 			Connection->WriteBitsToSendBuffer(nullptr, 0);
 		}
 
-		// Limit the amount of bits to minimum of a bunch and our buffer. NetBitStreamWriter requires the number of bytes to be a multiple of 4.
-		const uint32 MaxBitCount = uint32(Connection->GetMaxSingleBunchSizeBits());
-		const uint32 MaxBytes = FPlatformMath::Min((MaxBitCount/32U)*4U, (uint32)sizeof(BitStreamBuffer));
-
 		FNetBitStreamWriter BitWriter;
 		BitWriter.InitBytes(BitStreamBuffer, MaxBytes);
 
@@ -222,6 +244,7 @@ void UDataStreamChannel::Tick()
 		UDataStream::EWriteResult WriteResult = DataStreamManager->WriteData(SerializationContext, Record);
 		if (WriteResult == UDataStream::EWriteResult::NoData || SerializationContext.HasError())
 		{
+			IRIS_PROFILER_SCOPE(UDataStreamChannel_NoDataSent);
 			// Do not report the bunch
 			UE_NET_TRACE_DISCARD_BUNCH(Collector);
 			
@@ -233,7 +256,6 @@ void UDataStreamChannel::Tick()
 
 		IRIS_PROFILER_SCOPE(UDataStreamChannel_SendBunchAndFlushNet);
 
-		const int64 MaxBunchBits = MaxBytes*8;
 		FOutBunch OutBunch(MaxBunchBits);
 #if UE_NET_TRACE_ENABLED
 		SetTraceCollector(OutBunch, Collector);
@@ -274,7 +296,7 @@ void UDataStreamChannel::Tick()
 	};
 
 	// Begin the write, if we have nothing todo, just return
-	if (DataStreamManager->BeginWrite() == UDataStream::EWriteResult::NoData)
+	if (DataStreamManager->BeginWrite(BeginWriteParams) == UDataStream::EWriteResult::NoData)
 	{
 		return;
 	}
diff --git a/Engine/Source/Runtime/Experimental/Iris/Core/Private/Iris/DataStream/DataStream.cpp b/Engine/Source/Runtime/Experimental/Iris/Core/Private/Iris/DataStream/DataStream.cpp
--- a/Engine/Source/Runtime/Experimental/Iris/Core/Private/Iris/DataStream/DataStream.cpp
+++ b/Engine/Source/Runtime/Experimental/Iris/Core/Private/Iris/DataStream/DataStream.cpp
@@ -6,7 +6,7 @@ UDataStream::~UDataStream()
 {
 }
 
-UDataStream::EWriteResult UDataStream::BeginWrite()
+UDataStream::EWriteResult UDataStream::BeginWrite(const FBeginWriteParameters& Params)
 {
 	return EWriteResult::HasMoreData;
 }
diff --git a/Engine/Source/Runtime/Experimental/Iris/Core/Private/Iris/DataStream/DataStreamManager.cpp b/Engine/Source/Runtime/Experimental/Iris/Core/Private/Iris/DataStream/DataStreamManager.cpp
--- a/Engine/Source/Runtime/Experimental/Iris/Core/Private/Iris/DataStream/DataStreamManager.cpp
+++ b/Engine/Source/Runtime/Experimental/Iris/Core/Private/Iris/DataStream/DataStreamManager.cpp
@@ -20,7 +20,7 @@ public:
 	void Init(const FDataStreamManagerInitParams& InitParams);
 	void Deinit();
 
-	EWriteResult BeginWrite();
+	EWriteResult BeginWrite(const UDataStream::FBeginWriteParameters& Params);
 	UDataStream::EWriteResult WriteData(UE::Net::FNetSerializationContext& context, FDataStreamRecord const*& OutRecord);
 	void EndWrite();
 	void ReadData(UE::Net::FNetSerializationContext& context);
@@ -95,9 +95,9 @@ void UDataStreamManager::Deinit()
 	Impl->Deinit();
 }
 
-UDataStream::EWriteResult UDataStreamManager::BeginWrite()
+UDataStream::EWriteResult UDataStreamManager::BeginWrite(const UDataStream::FBeginWriteParameters& Params)
 {
-	return Impl->BeginWrite();
+	return Impl->BeginWrite(Params);
 }
 
 void UDataStreamManager::EndWrite()
@@ -191,7 +191,7 @@ void UDataStreamManager::FImpl::Deinit()
 	StreamSendStatus.Reset();
 }
 
-UDataStreamManager::EWriteResult UDataStreamManager::FImpl::BeginWrite()
+UDataStreamManager::EWriteResult UDataStreamManager::FImpl::BeginWrite(const UDataStream::FBeginWriteParameters& Params)
 {
 	const SIZE_T StreamCount = Streams.Num();
 	if (StreamCount == 0)
@@ -212,7 +212,7 @@ UDataStreamManager::EWriteResult UDataStreamManager::FImpl::BeginWrite()
 			continue;
 		}
 
-		const EWriteResult WriteResult = Stream->BeginWrite();
+		const EWriteResult WriteResult = Stream->BeginWrite(Params);
 		CombinedWriteResult = (CombinedWriteResult == EWriteResult::HasMoreData || WriteResult == EWriteResult::HasMoreData) ? EWriteResult::HasMoreData : WriteResult;
 	}
 
diff --git a/Engine/Source/Runtime/Experimental/Iris/Core/Private/Iris/ReplicationSystem/ReplicationDataStream.cpp b/Engine/Source/Runtime/Experimental/Iris/Core/Private/Iris/ReplicationSystem/ReplicationDataStream.cpp
--- a/Engine/Source/Runtime/Experimental/Iris/Core/Private/Iris/ReplicationSystem/ReplicationDataStream.cpp
+++ b/Engine/Source/Runtime/Experimental/Iris/Core/Private/Iris/ReplicationSystem/ReplicationDataStream.cpp
@@ -14,11 +14,11 @@ UReplicationDataStream::~UReplicationDataStream()
 {
 }
 
-UDataStream::EWriteResult UReplicationDataStream::BeginWrite()
+UDataStream::EWriteResult UReplicationDataStream::BeginWrite(const FBeginWriteParameters& Params)
 {
 	if (ReplicationWriter != nullptr)
 	{
-		return ReplicationWriter->BeginWrite();
+		return ReplicationWriter->BeginWrite(Params);
 	}
 
 	return EWriteResult::NoData;
diff --git a/Engine/Source/Runtime/Experimental/Iris/Core/Private/Iris/ReplicationSystem/ReplicationDataStream.h b/Engine/Source/Runtime/Experimental/Iris/Core/Private/Iris/ReplicationSystem/ReplicationDataStream.h
--- a/Engine/Source/Runtime/Experimental/Iris/Core/Private/Iris/ReplicationSystem/ReplicationDataStream.h
+++ b/Engine/Source/Runtime/Experimental/Iris/Core/Private/Iris/ReplicationSystem/ReplicationDataStream.h
@@ -24,7 +24,7 @@ private:
 	virtual ~UReplicationDataStream();
 
 	// UDataStream interface
-	virtual EWriteResult BeginWrite() override;
+	virtual EWriteResult BeginWrite(const FBeginWriteParameters& Params) override;
 	virtual EWriteResult WriteData(UE::Net::FNetSerializationContext& Context, FDataStreamRecord const*& OutRecord) override;
 	virtual void EndWrite() override;
 	virtual void ReadData(UE::Net::FNetSerializationContext& Context) override;
diff --git a/Engine/Source/Runtime/Experimental/Iris/Core/Private/Iris/ReplicationSystem/ReplicationWriter.cpp b/Engine/Source/Runtime/Experimental/Iris/Core/Private/Iris/ReplicationSystem/ReplicationWriter.cpp
--- a/Engine/Source/Runtime/Experimental/Iris/Core/Private/Iris/ReplicationSystem/ReplicationWriter.cpp
+++ b/Engine/Source/Runtime/Experimental/Iris/Core/Private/Iris/ReplicationSystem/ReplicationWriter.cpp
@@ -50,6 +50,10 @@ static FAutoConsoleVariableRef CVarWarnAboutDroppedAttachmentsToObjectsNotInScop
 	TEXT("Warn when attachments are dropped due to object not in scope. Default is false."
 	));
 
+static int32 GReplicationWriterMaxAllowedPacketsIfNotHugeObject = 3;
+static FAutoConsoleVariableRef CVarReplicationWriterMaxAllowedPacketsIfNotHugeObject(TEXT("net.iris.ReplicationWriterMaxAllowedPacketsIfNotHugeObject"), GReplicationWriterMaxAllowedPacketsIfNotHugeObject,
+	TEXT("Allow ReplicationWriter to overcommit data if we have more data to write."));
+
 static bool bValidateObjectsWithDirtyChanges = false;
 static FAutoConsoleVariableRef CvarValidateObjectsWithDirtyChanges(TEXT("net.Iris.ReplicationWriter.ValidateObjectsWithDirtyChanges"), bValidateObjectsWithDirtyChanges, TEXT("Ensure that we don't try to mark invalid objects as dirty when they shouldn't."));
 
@@ -2919,6 +2923,8 @@ int FReplicationWriter::HandleObjectBatchSuccess(const FBatchInfo& BatchInfo, FR
 
 FReplicationWriter::EWriteObjectRetryMode FReplicationWriter::HandleObjectBatchFailure(FReplicationWriter::EWriteObjectStatus WriteObjectStatus, const FBatchInfo& BatchInfo, const FReplicationWriter::FBitStreamInfo& BatchBitStreamInfo) const
 {
+	IRIS_PROFILER_SCOPE(FReplicationWriter_HandleObjectBatchFailure);
+	
 	// Cleanup data stored in BatchInfo
 	for (const FBatchObjectInfo& BatchObjectInfo : BatchInfo.ObjectInfos)
 	{
@@ -2951,8 +2957,14 @@ FReplicationWriter::EWriteObjectRetryMode FReplicationWriter::HandleObjectBatchF
 			return EWriteObjectRetryMode::SplitHugeObject;
 		}
 	}
+	else
+	{
+		IRIS_PROFILER_SCOPE(FReplicationWriter_BlockedByHugeOBjectAlreadyBeingSent);
+	}
 
-	if (WriteContext.FailedToWriteSmallObjectCount >= Parameters.MaxFailedSmallObjectCount)	
+	// If we are allowed to request more packets to write, we should abort once we failed to write a small object
+	const uint32 MaxFailedSmallObjectCount = WriteContext.bCanWriteMoreData ? 1U : Parameters.MaxFailedSmallObjectCount;
+	if (WriteContext.FailedToWriteSmallObjectCount >= MaxFailedSmallObjectCount)	
 	{
 		return EWriteObjectRetryMode::Abort;
 	}
@@ -2962,7 +2974,7 @@ FReplicationWriter::EWriteObjectRetryMode FReplicationWriter::HandleObjectBatchF
 	return EWriteObjectRetryMode::TrySmallObject;
 }
 
-UDataStream::EWriteResult FReplicationWriter::BeginWrite()
+UDataStream::EWriteResult FReplicationWriter::BeginWrite(const UDataStream::FBeginWriteParameters& Params)
 {
 	IRIS_PROFILER_SCOPE(FReplicationWriter_PrepareWrite);
 
@@ -2994,6 +3006,8 @@ UDataStream::EWriteResult FReplicationWriter::BeginWrite()
 	WriteContext.CurrentIndex = 0U;
 	WriteContext.FailedToWriteSmallObjectCount = 0U;
 	WriteContext.SortedObjectCount = 0U;
+	WriteContext.NumWrittenPacketsInThisBatch = 0U;
+	WriteContext.bCanWriteMoreData = Params.bCanWriteMoreData;
 
 	// Reset dependent object array
 	WriteContext.DependentObjectsPendingSend.Reset();
@@ -3129,6 +3143,10 @@ UDataStream::EWriteResult FReplicationWriter::Write(FNetSerializationContext& Co
 		{
 			WriteResult = UDataStream::EWriteResult::HasMoreData;
 		}
+		else if (WriteContext.bCanWriteMoreData && ((int32)WriteContext.NumWrittenPacketsInThisBatch < GReplicationWriterMaxAllowedPacketsIfNotHugeObject))
+		{
+			WriteResult = UDataStream::EWriteResult::HasMoreData;
+		}
 		else
 		{
 			WriteResult = UDataStream::EWriteResult::Ok;
@@ -3162,6 +3180,7 @@ UDataStream::EWriteResult FReplicationWriter::Write(FNetSerializationContext& Co
 		}
 #endif
 
+		++WriteContext.NumWrittenPacketsInThisBatch;
 	}
 	else 
 	{
diff --git a/Engine/Source/Runtime/Experimental/Iris/Core/Private/Iris/ReplicationSystem/ReplicationWriter.h b/Engine/Source/Runtime/Experimental/Iris/Core/Private/Iris/ReplicationSystem/ReplicationWriter.h
--- a/Engine/Source/Runtime/Experimental/Iris/Core/Private/Iris/ReplicationSystem/ReplicationWriter.h
+++ b/Engine/Source/Runtime/Experimental/Iris/Core/Private/Iris/ReplicationSystem/ReplicationWriter.h
@@ -51,7 +51,7 @@ class FReplicationWriter
 {
 public:
 	// Scheduling constants
-	static constexpr float CreatePriority = 1.f;
+	static constexpr float CreatePriority = 0.f;
 	static constexpr float TearOffPriority = 1.f;
 	static constexpr float LostStatePriorityBump = 1.f;
 	static constexpr float SchedulingThresholdPriority = 1.f;
@@ -166,7 +166,7 @@ public:
 	// UpdatedPriorities contains priorities for all objects. Objects in need of a priority update should use the newly calculated priorities.
 	void UpdatePriorities(const float* UpdatedPriorities);
 
-	UDataStream::EWriteResult BeginWrite();
+	UDataStream::EWriteResult BeginWrite(const UDataStream::FBeginWriteParameters& Params);
 
 	// WriteData to Packet, returns true for now if data was written
 	UDataStream::EWriteResult Write(FNetSerializationContext& Context);
@@ -225,6 +225,10 @@ private:
 		// How many objects that were attempted to be replicated but which ultimately didn't fit in the packet.
 		uint32 FailedToWriteSmallObjectCount;
 
+		// How many packets have we written to?
+		uint32 NumWrittenPacketsInThisBatch = 0U;
+		bool bCanWriteMoreData = false;
+
 		uint32 bHasDestroyedObjectsToSend : 1;
 		uint32 bHasUpdatedObjectsToSend : 1;
 		uint32 bHasHugeObjectToSend : 1;
diff --git a/Engine/Source/Runtime/Experimental/Iris/Core/Public/Iris/DataStream/DataStream.h b/Engine/Source/Runtime/Experimental/Iris/Core/Public/Iris/DataStream/DataStream.h
--- a/Engine/Source/Runtime/Experimental/Iris/Core/Public/Iris/DataStream/DataStream.h
+++ b/Engine/Source/Runtime/Experimental/Iris/Core/Public/Iris/DataStream/DataStream.h
@@ -53,6 +53,11 @@ public:
 		HasMoreData,
 	};
 
+	struct FBeginWriteParameters
+	{
+		bool bCanWriteMoreData = false;
+	};
+
 public:
 	IRISCORE_API virtual ~UDataStream();
 
@@ -60,7 +65,7 @@ public:
 	 * Called before any calls to potential WriteData, if it returns EWriteData::NoData no other calls will be made.
 	 * The purpose of the method is to enable a DataStream to setup data that can persist over multiple calls to WriteData if bandwidth allows.
 	*/
-	IRISCORE_API virtual EWriteResult BeginWrite();
+	IRISCORE_API virtual EWriteResult BeginWrite(const FBeginWriteParameters& Params);
 
 	/**
 	 * Serialize data to a bitstream and optionally store record of what was serialized to a custom FDataStreamRecord.
diff --git a/Engine/Source/Runtime/Experimental/Iris/Core/Public/Iris/DataStream/DataStreamManager.h b/Engine/Source/Runtime/Experimental/Iris/Core/Public/Iris/DataStream/DataStreamManager.h
--- a/Engine/Source/Runtime/Experimental/Iris/Core/Public/Iris/DataStream/DataStreamManager.h
+++ b/Engine/Source/Runtime/Experimental/Iris/Core/Public/Iris/DataStream/DataStreamManager.h
@@ -57,7 +57,7 @@ public:
 	// UDataStream interface
 
 	/** Call BeginWrite on all active data streams. */
-	IRISCORE_API virtual EWriteResult BeginWrite() override;
+	IRISCORE_API virtual EWriteResult BeginWrite(const FBeginWriteParameters& Params) override;
 
 	/** Call WriteData on all active data streams. */
 	IRISCORE_API virtual EWriteResult WriteData(UE::Net::FNetSerializationContext& context, FDataStreamRecord const*& OutRecord) override;
