World Partition - HLOD: HLOD Actor Descriptor now stores the path to the source actors HLOD Layer asset
* This is the HLOD Layer asset used the build the HLOD

diff --git a/Engine/Source/Editor/UnrealEd/Private/WorldPartition/WorldPartitionHLODsBuilder.cpp b/Engine/Source/Editor/UnrealEd/Private/WorldPartition/WorldPartitionHLODsBuilder.cpp
--- a/Engine/Source/Editor/UnrealEd/Private/WorldPartition/WorldPartitionHLODsBuilder.cpp
+++ b/Engine/Source/Editor/UnrealEd/Private/WorldPartition/WorldPartitionHLODsBuilder.cpp
@@ -735,7 +735,7 @@ TArray<TArray<FGuid>> UWorldPartitionHLODsBuilder::GetHLODWorkloads(int32 NumWor
 		}
 
 		// Filter by HLOD layer
-		if (!HLODLayerToBuild.IsNone() && HLODIterator->GetSourceHLODLayerName() != HLODLayerToBuild)
+		if (!HLODLayerToBuild.IsNone() && HLODIterator->GetSourceHLODLayer().GetAssetName() != HLODLayerToBuild)
 		{
 			continue;
 		}
diff --git a/Engine/Source/Runtime/Core/Public/UObject/FortniteMainBranchObjectVersion.h b/Engine/Source/Runtime/Core/Public/UObject/FortniteMainBranchObjectVersion.h
--- a/Engine/Source/Runtime/Core/Public/UObject/FortniteMainBranchObjectVersion.h
+++ b/Engine/Source/Runtime/Core/Public/UObject/FortniteMainBranchObjectVersion.h
@@ -372,6 +372,9 @@ struct FFortniteMainBranchObjectVersion
 		// Remove the UAnimCurveCompressionCodec::InstanceGuid which causes cook determinism issues
 		RemoveAnimCurveCompressionCodecInstanceGuid,
 
+		// Serialize the source HLOD Layer for HLOD actor descriptors.
+		WorldPartitionHLODActorDescSerializeSourceHLODLayer,
+
 		// -----<new versions can be added above this line>-------------------------------------------------
 		VersionPlusOne,
 		LatestVersion = VersionPlusOne - 1
diff --git a/Engine/Source/Runtime/Engine/Private/WorldPartition/HLOD/HLODActorDesc.cpp b/Engine/Source/Runtime/Engine/Private/WorldPartition/HLOD/HLODActorDesc.cpp
--- a/Engine/Source/Runtime/Engine/Private/WorldPartition/HLOD/HLODActorDesc.cpp
+++ b/Engine/Source/Runtime/Engine/Private/WorldPartition/HLOD/HLODActorDesc.cpp
@@ -44,11 +44,7 @@ void FHLODActorDesc::Init(const AActor* InActor)
 
 	if (HLODSourceActors)
 	{
-		SourceHLODLayerName = NAME_None;
-		if (const UHLODLayer* SourceHLODLayer = HLODSourceActors->GetHLODLayer())
-		{
-			SourceHLODLayerName = SourceHLODLayer->GetFName();
-		}
+		SourceHLODLayer = FTopLevelAssetPath(HLODSourceActors->GetHLODLayer());
 	}
 	
 	HLODStats = HLODActor->GetStats();
@@ -130,7 +126,12 @@ void FHLODActorDesc::Serialize(FArchive& Ar)
 				Ar << SourceCellName;
 			}
 
-			Ar << SourceHLODLayerName;
+			if (Ar.CustomVer(FFortniteMainBranchObjectVersion::GUID) < FFortniteMainBranchObjectVersion::WorldPartitionHLODActorDescSerializeSourceHLODLayer)
+			{
+				FString SourceHLODLayerName;
+				Ar << SourceHLODLayerName;
+			}
+
 			Ar << HLODStats;
 
 			// Update package size stat on load
@@ -139,6 +140,11 @@ void FHLODActorDesc::Serialize(FArchive& Ar)
 				HLODStats.Add(FWorldPartitionHLODStats::MemoryDiskSizeBytes, GetPackageSize());
 			}
 		}
+
+		if (Ar.CustomVer(FFortniteMainBranchObjectVersion::GUID) >= FFortniteMainBranchObjectVersion::WorldPartitionHLODActorDescSerializeSourceHLODLayer)
+		{
+			Ar << SourceHLODLayer;
+		}
 	}
 }
 
@@ -147,7 +153,7 @@ bool FHLODActorDesc::Equals(const FWorldPartitionActorDesc* Other) const
 	if (FWorldPartitionActorDesc::Equals(Other))
 	{
 		const FHLODActorDesc& HLODActorDesc = *(FHLODActorDesc*)Other;
-		return SourceHLODLayerName == HLODActorDesc.SourceHLODLayerName &&
+		return SourceHLODLayer == HLODActorDesc.SourceHLODLayer &&
 			   HLODStats.OrderIndependentCompareEqual(HLODActorDesc.GetStats()) &&
 			   CompareUnsortedArrays(ChildHLODActors, HLODActorDesc.ChildHLODActors);
 	}
diff --git a/Engine/Source/Runtime/Engine/Private/WorldPartition/HLOD/HLODRuntimeSubsystem.cpp b/Engine/Source/Runtime/Engine/Private/WorldPartition/HLOD/HLODRuntimeSubsystem.cpp
--- a/Engine/Source/Runtime/Engine/Private/WorldPartition/HLOD/HLODRuntimeSubsystem.cpp
+++ b/Engine/Source/Runtime/Engine/Private/WorldPartition/HLOD/HLODRuntimeSubsystem.cpp
@@ -856,7 +856,7 @@ bool UWorldPartitionHLODRuntimeSubsystem::WriteHLODStatsCSV(UWorld* InWorld, con
 	{	
 		{ "WorldPackage",		[InWorld](const FHLODActorDesc& InActorDesc) { return InWorld->GetPackage()->GetName(); } },
 		{ "Name",				[](const FHLODActorDesc& InActorDesc) { return InActorDesc.GetActorLabel().ToString(); } },
-		{ "HLODLayer",			[](const FHLODActorDesc& InActorDesc) { return InActorDesc.GetSourceHLODLayerName().ToString(); } },
+		{ "HLODLayer",			[](const FHLODActorDesc& InActorDesc) { return InActorDesc.GetSourceHLODLayer().GetAssetName().ToString(); }},
 		{ "SpatiallyLoaded",	[](const FHLODActorDesc& InActorDesc) { return InActorDesc.GetIsSpatiallyLoaded() ? TEXT("true") : TEXT("false"); } },
 		{ "DataLayers",			[&GetDataLayerShortName](const FHLODActorDesc& InActorDesc) { return FString::JoinBy(InActorDesc.GetDataLayerInstanceNames(), TEXT(" | "), GetDataLayerShortName); } },
 
diff --git a/Engine/Source/Runtime/Engine/Public/WorldPartition/HLOD/HLODActorDesc.h b/Engine/Source/Runtime/Engine/Public/WorldPartition/HLOD/HLODActorDesc.h
--- a/Engine/Source/Runtime/Engine/Public/WorldPartition/HLOD/HLODActorDesc.h
+++ b/Engine/Source/Runtime/Engine/Public/WorldPartition/HLOD/HLODActorDesc.h
@@ -26,7 +26,7 @@ public:
 	typedef TMap<FName, int64>	FStats;
 
 	inline const TArray<FGuid>& GetChildHLODActors() const { return ChildHLODActors; }
-	inline const FName GetSourceHLODLayerName() const { return SourceHLODLayerName; }
+	inline const FTopLevelAssetPath& GetSourceHLODLayer() const { return SourceHLODLayer; }
 	inline const FStats& GetStats() const { return HLODStats; }
 	inline int64 GetStat(FName InStatName) const { return HLODStats.FindRef(InStatName); }
 
@@ -44,7 +44,9 @@ protected:
 	//~ End FWorldPartitionActorDesc Interface.
 
 	TArray<FGuid> ChildHLODActors;
-	FName SourceHLODLayerName;
+
+	FTopLevelAssetPath SourceHLODLayer;
+
 	FStats HLODStats;
 };
 #endif
