World Partition - Editor: Added support for a delegate to be called on loader adapters state change

diff --git a/Engine/Source/Runtime/Engine/Private/WorldPartition/WorldPartition.cpp b/Engine/Source/Runtime/Engine/Private/WorldPartition/WorldPartition.cpp
--- a/Engine/Source/Runtime/Engine/Private/WorldPartition/WorldPartition.cpp
+++ b/Engine/Source/Runtime/Engine/Private/WorldPartition/WorldPartition.cpp
@@ -2232,6 +2232,16 @@ void UWorldPartition::LoadLastLoadedRegions()
 	}
 }
 
+void UWorldPartition::OnLoaderAdapterStateChanged(IWorldPartitionActorLoaderInterface::ILoaderAdapter* InLoaderAdapter)
+{
+	if (InLoaderAdapter->GetUserCreated())
+	{
+		NumUserCreatedLoadedRegions += InLoaderAdapter->IsLoaded() ? 1 : -1;
+	}
+
+	LoaderAdapterStateChanged.Broadcast(InLoaderAdapter);
+}
+
 void UWorldPartition::OnWorldRenamed(UWorld* RenamedWorld)
 {
 	if (GetWorld() == RenamedWorld)
diff --git a/Engine/Source/Runtime/Engine/Private/WorldPartition/WorldPartitionActorLoaderInterface.cpp b/Engine/Source/Runtime/Engine/Private/WorldPartition/WorldPartitionActorLoaderInterface.cpp
--- a/Engine/Source/Runtime/Engine/Private/WorldPartition/WorldPartitionActorLoaderInterface.cpp
+++ b/Engine/Source/Runtime/Engine/Private/WorldPartition/WorldPartitionActorLoaderInterface.cpp
@@ -52,12 +52,9 @@ void IWorldPartitionActorLoaderInterface::ILoaderAdapter::Load()
 		RefreshLoadedState();
 		RegisterDelegates();
 
-		if (bUserCreated)
+		if (UWorldPartition* WorldPartition = World->GetWorldPartition())
 		{
-			if (UWorldPartition* WorldPartition = World->GetWorldPartition())
-			{
-				WorldPartition->OnUserCreatedRegionLoaded();
-			}
+			WorldPartition->OnLoaderAdapterStateChanged(this);
 		}
 	}
 }
@@ -84,10 +81,7 @@ void IWorldPartitionActorLoaderInterface::ILoaderAdapter::Unload()
 
 			bLoaded = false;
 
-			if (bUserCreated)
-			{
-				WorldPartition->OnUserCreatedRegionUnloaded();
-			}
+			WorldPartition->OnLoaderAdapterStateChanged(this);
 
 			if (NumUnloads)
 			{
diff --git a/Engine/Source/Runtime/Engine/Public/WorldPartition/WorldPartition.h b/Engine/Source/Runtime/Engine/Public/WorldPartition/WorldPartition.h
--- a/Engine/Source/Runtime/Engine/Public/WorldPartition/WorldPartition.h
+++ b/Engine/Source/Runtime/Engine/Public/WorldPartition/WorldPartition.h
@@ -379,8 +379,11 @@ public:
 	ENGINE_API void LoadLastLoadedRegions();
 
 	bool HasLoadedUserCreatedRegions() const { return !!NumUserCreatedLoadedRegions; }
-	void OnUserCreatedRegionLoaded() { NumUserCreatedLoadedRegions++; }
-	void OnUserCreatedRegionUnloaded() { check(HasLoadedUserCreatedRegions()); NumUserCreatedLoadedRegions--; }
+
+	DECLARE_MULTICAST_DELEGATE_OneParam(FLoaderAdapterStateChangedDelegate, const IWorldPartitionActorLoaderInterface::ILoaderAdapter*);
+	FLoaderAdapterStateChangedDelegate LoaderAdapterStateChanged;
+		
+	ENGINE_API void OnLoaderAdapterStateChanged(IWorldPartitionActorLoaderInterface::ILoaderAdapter* InLoaderAdapter);
 
 	bool IsEnablingStreamingJustified() const { return bEnablingStreamingJustified; }
 
