Do not set StaticMesh RenderData in headless runs (e.g. on servers).

- Some memory savings on the server (removing few tens from a thousand paper cuts).

diff --git a/Engine/Source/Runtime/Engine/Private/StaticMesh.cpp b/Engine/Source/Runtime/Engine/Private/StaticMesh.cpp
--- a/Engine/Source/Runtime/Engine/Private/StaticMesh.cpp
+++ b/Engine/Source/Runtime/Engine/Private/StaticMesh.cpp
@@ -1915,6 +1915,8 @@ void FStaticMeshRenderData::InitResources(ERHIFeatureLevel::Type InFeatureLevel,
 	}
 #endif // #if WITH_EDITOR
 
+	checkf(FApp::CanEverRender(), TEXT("RenderData should not initialize resources in headless runs"));
+
 	for (int32 LODIndex = 0; LODIndex < LODResources.Num(); ++LODIndex)
 	{
 		// Skip LODs that have their render data stripped
@@ -3477,6 +3479,8 @@ void UStaticMesh::SetRenderData(TUniquePtr<class FStaticMeshRenderData>&& InRend
 {
 	WaitUntilAsyncPropertyReleased(EStaticMeshAsyncProperties::RenderData);
 
+	checkf(FApp::CanEverRender() || InRenderData == nullptr, TEXT("Non-null RenderData should not be set in headless runs."));
+
 	PRAGMA_DISABLE_DEPRECATION_WARNINGS
 	RenderData = MoveTemp(InRenderData);
 	PRAGMA_ENABLE_DEPRECATION_WARNINGS
@@ -5740,8 +5744,12 @@ void UStaticMesh::Serialize(FArchive& Ar)
 		if (Ar.IsLoading())
 		{
 			SCOPE_MS_ACCUMULATOR(STAT_StaticMesh_RenderData);
-			SetRenderData(MakeUnique<FStaticMeshRenderData>());
-			GetRenderData()->Serialize(Ar, this, bCooked);
+			TUniquePtr<class FStaticMeshRenderData> LocalRenderData = MakeUnique<FStaticMeshRenderData>();
+			LocalRenderData->Serialize(Ar, this, bCooked);
+			if (FApp::CanEverRender())
+			{
+				SetRenderData(MoveTemp(LocalRenderData));
+			}
 		}
 #if WITH_EDITOR
 		else if (Ar.IsSaving())
