Compile out PSO precache class members when not needed.

- sizeof(UPrimitiveComponent) 1424 -> 1392,  sizeof(UStaticMeshComponent) 1616 -> 1584
- with these classes (and subclasses) being in tens of thousands on the server we're saving some 100MB res/200 virt between this and not setting StaticMesh render data (separate change).

diff --git a/Engine/Plugins/FX/Niagara/Source/Niagara/Private/NiagaraComponent.cpp b/Engine/Plugins/FX/Niagara/Source/Niagara/Private/NiagaraComponent.cpp
--- a/Engine/Plugins/FX/Niagara/Source/Niagara/Private/NiagaraComponent.cpp
+++ b/Engine/Plugins/FX/Niagara/Source/Niagara/Private/NiagaraComponent.cpp
@@ -2286,6 +2286,7 @@ FPrimitiveSceneProxy* UNiagaraComponent::CreateSceneProxy()
 	SCOPE_CYCLE_COUNTER(STAT_NiagaraCreateSceneProxy);
 	SCOPE_CYCLE_COUNTER(STAT_NiagaraOverview_GT);
     
+#if UE_WITH_PSO_PRECACHING
 	if (Asset != nullptr)
 	{
 		// PSO request should have been handled by the Asset itself, so here we just ensure the request gets boosted
@@ -2300,6 +2301,7 @@ FPrimitiveSceneProxy* UNiagaraComponent::CreateSceneProxy()
 			return nullptr;
 		}
 	}
+#endif // UE_WITH_PSO_PRECACHING
 
 	// The constructor will set up the System renderers from the component.
 	FNiagaraSceneProxy* Proxy = new FNiagaraSceneProxy(this);
diff --git a/Engine/Source/Runtime/Core/Public/HAL/Platform.h b/Engine/Source/Runtime/Core/Public/HAL/Platform.h
--- a/Engine/Source/Runtime/Core/Public/HAL/Platform.h
+++ b/Engine/Source/Runtime/Core/Public/HAL/Platform.h
@@ -582,6 +582,10 @@
 	#define PLATFORM_USE_FALLBACK_PSO 0
 #endif
 
+#ifndef PLATFORM_SUPPORTS_PSO_PRECACHING
+	#define PLATFORM_SUPPORTS_PSO_PRECACHING (!UE_SERVER)
+#endif
+
 #ifndef PLATFORM_USES_UNFAIR_LOCKS
 	#define PLATFORM_USES_UNFAIR_LOCKS 0
 #endif
diff --git a/Engine/Source/Runtime/Engine/Classes/Components/PrimitiveComponent.h b/Engine/Source/Runtime/Engine/Classes/Components/PrimitiveComponent.h
--- a/Engine/Source/Runtime/Engine/Classes/Components/PrimitiveComponent.h
+++ b/Engine/Source/Runtime/Engine/Classes/Components/PrimitiveComponent.h
@@ -653,6 +653,7 @@ protected:
 		ToolTip = "When false, the underlying physics body will contain all sim data (mass, inertia tensor, etc) even if mobility is not set to Moveable"))
 	uint8 bStaticWhenNotMoveable:1;
 
+#if UE_WITH_PSO_PRECACHING
 	/** Helper flag to check if PSOs have been precached already */
 	uint8 bPSOPrecacheCalled : 1;
 
@@ -664,6 +665,7 @@ protected:
 
 	/** Graph event used to track all the PSO precache events */
 	FGraphEventRef PSOPrecacheCompileEvent;
+#endif
 
 	uint8 bIgnoreBoundsForEditorFocus : 1;
 #if WITH_EDITOR
diff --git a/Engine/Source/Runtime/Engine/Private/Components/PrimitiveComponent.cpp b/Engine/Source/Runtime/Engine/Private/Components/PrimitiveComponent.cpp
--- a/Engine/Source/Runtime/Engine/Private/Components/PrimitiveComponent.cpp
+++ b/Engine/Source/Runtime/Engine/Private/Components/PrimitiveComponent.cpp
@@ -358,8 +358,10 @@ UPrimitiveComponent::UPrimitiveComponent(const FObjectInitializer& ObjectInitial
 	bBulkReregister = false;
 	LastCheckedAllCollideableDescendantsTime = 0.f;
 
+#if UE_WITH_PSO_PRECACHING
 	bPSOPrecacheCalled = false;
 	bPSOPrecacheRequestBoosted = false;
+#endif // UE_WITH_PSO_PRECACHING
 	
 	bApplyImpulseOnDamage = true;
 	bReplicatePhysicsToAutonomousProxy = true;
@@ -4581,6 +4583,7 @@ public:
 
 void UPrimitiveComponent::PrecachePSOs()
 {
+#if UE_WITH_PSO_PRECACHING
 	if (!FApp::CanEverRender() || !IsComponentPSOPrecachingEnabled())
 	{
 		return;
@@ -4604,10 +4607,12 @@ void UPrimitiveComponent::PrecachePSOs()
 	}	
 
 	RequestRecreateRenderStateWhenPSOPrecacheFinished(GraphEvents);
+#endif
 }
 
 void UPrimitiveComponent::RequestRecreateRenderStateWhenPSOPrecacheFinished(const FGraphEventArray& PSOPrecacheCompileEvents)
 {
+#if UE_WITH_PSO_PRECACHING
 	// If the proxy creation strategy relies on knowing when the precached PSO has been compiled,
 	// schedule a task to mark the render state dirty when all PSOs are compiled so the proxy gets recreated.
 	if (GetPSOPrecacheProxyCreationStrategy() != EPSOPrecacheProxyCreationStrategy::AlwaysCreate && !PSOPrecacheCompileEvents.IsEmpty())
@@ -4616,20 +4621,30 @@ void UPrimitiveComponent::RequestRecreateRenderStateWhenPSOPrecacheFinished(cons
 	}
 
 	bPSOPrecacheCalled = true;
+#endif // UE_WITH_PSO_PRECACHING
 }
 
 bool UPrimitiveComponent::IsPSOPrecaching() const
 {
+#if UE_WITH_PSO_PRECACHING
 	return PSOPrecacheCompileEvent && !PSOPrecacheCompileEvent->IsComplete();
+#else
+	return false;
+#endif // UE_WITH_PSO_PRECACHING
 }
 
 bool UPrimitiveComponent::ShouldRenderProxyFallbackToDefaultMaterial() const
 {
+#if UE_WITH_PSO_PRECACHING
 	return IsPSOPrecaching() && GetPSOPrecacheProxyCreationStrategy() == EPSOPrecacheProxyCreationStrategy::UseDefaultMaterialUntilPSOPrecached;
+#else
+	return false;
+#endif // UE_WITH_PSO_PRECACHING
 }
 
 bool UPrimitiveComponent::CheckPSOPrecachingAndBoostPriority()
 {
+#if UE_WITH_PSO_PRECACHING
 	ensure(!IsComponentPSOPrecachingEnabled() || bPSOPrecacheCalled);
 
 	if (PSOPrecacheCompileEvent && !PSOPrecacheCompileEvent->IsComplete())
@@ -4646,6 +4661,9 @@ bool UPrimitiveComponent::CheckPSOPrecachingAndBoostPriority()
 	}
 
 	return IsPSOPrecaching();
+#else
+	return false;
+#endif
 }
 
 #if WITH_EDITOR
diff --git a/Engine/Source/Runtime/Engine/Private/Particles/ParticleComponents.cpp b/Engine/Source/Runtime/Engine/Private/Particles/ParticleComponents.cpp
--- a/Engine/Source/Runtime/Engine/Private/Particles/ParticleComponents.cpp
+++ b/Engine/Source/Runtime/Engine/Private/Particles/ParticleComponents.cpp
@@ -3577,6 +3577,7 @@ bool UFXSystemComponent::RequiresLWCTileRecache(const FVector3f CurrentTile, con
 
 void UFXSystemComponent::PrecacheAssetPSOs(UFXSystemAsset* FXSystemAsset)
 {
+#if UE_WITH_PSO_PRECACHING
 	if (!FApp::CanEverRender() || !IsComponentPSOPrecachingEnabled() || FXSystemAsset == nullptr)
 	{
 		return;
@@ -3602,6 +3603,7 @@ void UFXSystemComponent::PrecacheAssetPSOs(UFXSystemAsset* FXSystemAsset)
 	}
 
 	bPSOPrecacheCalled = true;
+#endif // UE_WITH_PSO_PRECACHING
 }
 
 FOnSystemPreActivationChange UParticleSystemComponent::OnSystemPreActivationChange;
diff --git a/Engine/Source/Runtime/Engine/Private/Particles/ParticleSystemRender.cpp b/Engine/Source/Runtime/Engine/Private/Particles/ParticleSystemRender.cpp
--- a/Engine/Source/Runtime/Engine/Private/Particles/ParticleSystemRender.cpp
+++ b/Engine/Source/Runtime/Engine/Private/Particles/ParticleSystemRender.cpp
@@ -6958,6 +6958,7 @@ FPrimitiveSceneProxy* UParticleSystemComponent::CreateSceneProxy()
 	//@fixme Get non-instanced path working in ES!
 	if ((IsActive() == true)/** && (EmitterInstances.Num() > 0)*/ && Template)
 	{
+#if UE_WITH_PSO_PRECACHING
 		if (!bPSOPrecacheCalled)
 		{
 			PrecacheAssetPSOs(Template);
@@ -6968,6 +6969,7 @@ FPrimitiveSceneProxy* UParticleSystemComponent::CreateSceneProxy()
 			UE_LOG(LogParticles, Verbose, TEXT("Skipping CreateSceneProxy for UParticleSystemComponent %s (UParticleSystem PSOs are still compiling)"), *GetFullName());
 			return nullptr;
 		}
+#endif // UE_WITH_PSO_PRECACHING
 
 		FInGameScopedCycleCounter InGameCycleCounter(GetWorld(), EInGamePerfTrackers::VFXSignificance, EInGamePerfTrackerThreads::GameThread, bIsManagingSignificance);
 
diff --git a/Engine/Source/Runtime/Engine/Public/PSOPrecache.h b/Engine/Source/Runtime/Engine/Public/PSOPrecache.h
--- a/Engine/Source/Runtime/Engine/Public/PSOPrecache.h
+++ b/Engine/Source/Runtime/Engine/Public/PSOPrecache.h
@@ -19,6 +19,11 @@ class FVertexFactoryType;
 class FGraphicsPipelineStateInitializer;
 enum class EVertexInputStreamType : uint8;
 
+// master switch that decides whether to compile out some PSO precaching code (most importantly, reduce sizeofs of common classes)
+#ifndef UE_WITH_PSO_PRECACHING
+	#define UE_WITH_PSO_PRECACHING		(PLATFORM_SUPPORTS_PSO_PRECACHING)
+#endif // UE_WITH_PSO_PRECACHING
+
 #define PSO_PRECACHING_VALIDATE !WITH_EDITOR
 
 /**
