- WorldPartition Server streaming : Make sure the server streaming always loads a bit more than the client. This is necessary to avoid making the client wait indefinitly for the server to finish loading cells that are not even requested by the server because of a slight difference between client and server streaming source locations. Network quantization and world partition location/rotation quantization can contribute to this difference.
- Added wp.Runtime.ServerStreamingSourceMinimumExtraRadius to control the minimum extra radius applied on server streaming sources.
- Added wp.Runtime.ServerStreamingSourceMinimumExtraAngle to control the minimum extra angle applied on server streaming source shapes.

diff --git a/Engine/Source/Runtime/Engine/Private/WorldPartition/WorldPartitionStreamingSource.cpp b/Engine/Source/Runtime/Engine/Private/WorldPartition/WorldPartitionStreamingSource.cpp
--- a/Engine/Source/Runtime/Engine/Private/WorldPartition/WorldPartitionStreamingSource.cpp
+++ b/Engine/Source/Runtime/Engine/Private/WorldPartition/WorldPartitionStreamingSource.cpp
@@ -68,6 +68,16 @@ FString FWorldPartitionStreamingSource::ToString() const
 		StringBuilder.RemoveSuffix(1);
 	}
 
+	if (ExtraRadius > 0.f)
+	{
+		StringBuilder.Appendf(TEXT(" | Extra Radius: %d "), (int32)ExtraRadius);
+	}
+
+	if (ExtraAngle > 0.f)
+	{
+		StringBuilder.Appendf(TEXT(" | Extra Angle: %d "), (int32)ExtraAngle);
+	}
+
 	if (TargetGrids.Num())
 	{
 		StringBuilder.Appendf(TEXT(" | %s TargetGrids: "), (TargetBehavior == EStreamingSourceTargetBehavior::Include) ? TEXT("Included") : TEXT("Excluded"));
@@ -99,7 +109,7 @@ void FWorldPartitionStreamingSource::UpdateHash()
 	}
 
 	FHashBuilder HashBuilder;
-	HashBuilder	<< Name << TargetState << bBlockOnSlowLoading << bReplay << bRemote << Priority << TargetBehavior << TargetGrids << TargetHLODLayers << Shapes;
+	HashBuilder	<< Name << TargetState << bBlockOnSlowLoading << bReplay << bRemote << Priority << TargetBehavior << TargetGrids << TargetHLODLayers << Shapes << ExtraRadius  << ExtraAngle;
 
 	if (LocationQuantization)
 	{
diff --git a/Engine/Source/Runtime/Engine/Private/WorldPartition/WorldPartitionSubsystem.cpp b/Engine/Source/Runtime/Engine/Private/WorldPartition/WorldPartitionSubsystem.cpp
--- a/Engine/Source/Runtime/Engine/Private/WorldPartition/WorldPartitionSubsystem.cpp
+++ b/Engine/Source/Runtime/Engine/Private/WorldPartition/WorldPartitionSubsystem.cpp
@@ -41,6 +41,22 @@
 extern int32 GBlockOnSlowStreaming;
 static const FName NAME_WorldPartitionRuntimeHash("WorldPartitionRuntimeHash");
 
+static int32 GServerStreamingSourceMinimumExtraRadius = 400;
+static FAutoConsoleVariableRef CVarServerStreamingSourceMinimumExtraRadius(
+	TEXT("wp.Runtime.ServerStreamingSourceMinimumExtraRadius"),
+	GServerStreamingSourceMinimumExtraRadius,
+	TEXT("Minimum value added to the radius of the streaming sources used by the server (in Unreal unit)."),
+	ECVF_Default
+);
+
+static int32 GServerStreamingSourceMinimumExtraAngle = 1;
+static FAutoConsoleVariableRef CVarServerStreamingSourceMinimumExtraAngle(
+	TEXT("wp.Runtime.ServerStreamingSourceMinimumExtraAngle"),
+	GServerStreamingSourceMinimumExtraAngle,
+	TEXT("Minimum value added to the angle of the streaming source shape sector used by the server (in degree)."),
+	ECVF_Default
+);
+
 static int32 GDrawWorldPartitionIndex = -1;
 static FAutoConsoleCommand CVarDrawWorldPartitionIndex(
 	TEXT("wp.Runtime.DrawWorldPartitionIndex"),
@@ -140,7 +156,7 @@ static FAutoConsoleCommandWithOutputDevice GDumpWorldPartitionsCmd(
 			if (World && World->IsGameWorld())
 			{
 				if (const UWorldPartitionSubsystem* WorldPartitionSubsystem = World->GetSubsystem<UWorldPartitionSubsystem>())
-				{				
+				{
 					WorldPartitionSubsystem->DumpWorldPartitions(OutputDevice);
 				}
 			}
@@ -963,7 +979,9 @@ void UWorldPartitionSubsystem::UpdateStreamingSources()
 
 	if (!bIsUsingReplayStreamingSources)
 	{
-		if (!IsServer(World) || HasAnyWorldPartitionServerStreamingEnabled() || AWorldPartitionReplay::IsRecordingEnabled(World))
+		const bool bIsServer = IsServer(World);
+		const bool bServerStreamingEnabled = HasAnyWorldPartitionServerStreamingEnabled();
+		if (!bIsServer || bServerStreamingEnabled || AWorldPartitionReplay::IsRecordingEnabled(World))
 		{
 			bool bAllowPlayerControllerStreamingSources = true;
 #if WITH_EDITOR
@@ -992,6 +1010,26 @@ void UWorldPartitionSubsystem::UpdateStreamingSources()
 				}
 			}
 		}
+
+		// Make sure the server streaming always loads a bit more than the client.
+		// This is necessary to avoid making the client wait indefinitly for the server to finish loading cells
+		// that are not even requested by the server because of a slight difference between client and server
+		// streaming source locations / rotation.
+		// Network quantization and world partition location/rotation quantization can contribute to this difference.
+		if (bIsServer && bServerStreamingEnabled)
+		{
+			// Double location quantization for safety
+			const int32 LocationQuantization = UWorldPartitionStreamingPolicy::IsUpdateStreamingOptimEnabled() && (FWorldPartitionStreamingSource::GetLocationQuantization() > 0.f) ? FWorldPartitionStreamingSource::GetLocationQuantization() : 0.f;
+			const float ExtraRadius = FMath::Max<int32>(FMath::Max<int32>(GServerStreamingSourceMinimumExtraRadius, LocationQuantization * 2), 0);
+			// Double rotation quantization for safety
+			const int32 RotationQuantization = UWorldPartitionStreamingPolicy::IsUpdateStreamingOptimEnabled() && (FWorldPartitionStreamingSource::GetRotationQuantization() > 0.f) ? FWorldPartitionStreamingSource::GetRotationQuantization() : 0.f;
+			const float ExtraAngle = FMath::Max<int32>(FMath::Max<int32>(GServerStreamingSourceMinimumExtraAngle, RotationQuantization * 2), 0);
+			for (FWorldPartitionStreamingSource& StreamingSource : StreamingSources)
+			{
+				FSetStreamingSourceExtraRadius SetExtraRadius(StreamingSource, ExtraRadius);
+				FSetStreamingSourceExtraAngle SetExtraAngle(StreamingSource, ExtraAngle);
+			}
+		}
 	}
 
 	for (auto& Pair : StreamingSourcesVelocity)
diff --git a/Engine/Source/Runtime/Engine/Public/WorldPartition/WorldPartitionStreamingSource.h b/Engine/Source/Runtime/Engine/Public/WorldPartition/WorldPartitionStreamingSource.h
--- a/Engine/Source/Runtime/Engine/Public/WorldPartition/WorldPartitionStreamingSource.h
+++ b/Engine/Source/Runtime/Engine/Public/WorldPartition/WorldPartitionStreamingSource.h
@@ -180,12 +180,13 @@ public:
 		}
 		return false;
 	}
-	FORCEINLINE static void ForEachShape(float InGridLoadingRange, float InDefaultRadius, bool bInProjectIn2D, const FVector& InLocation, const FRotator& InRotation, const TArray<FStreamingSourceShape>& InShapes, TFunctionRef<void(const FSphericalSector&)> InOperation)
+
+	FORCEINLINE static void ForEachShape(float InGridLoadingRange, float InDefaultRadius, bool bInProjectIn2D, const FVector& InLocation, const FRotator& InRotation, const TArray<FStreamingSourceShape>& InShapes, TFunctionRef<void(const FSphericalSector&)> InOperation, float InExtraRadius = 0.f, float InExtraAngle = 0.f)
 	{
 		const FTransform Transform(bInProjectIn2D ? FRotator(0, InRotation.Yaw, 0) : InRotation, InLocation);
 		if (InShapes.IsEmpty())
 		{
-			FSphericalSector LocalShape(FVector::ZeroVector, InDefaultRadius);
+			FSphericalSector LocalShape(FVector::ZeroVector, InDefaultRadius + InExtraRadius);
 			if (LocalShape.IsValid())
 			{
 				InOperation(LocalShape.TransformBy(Transform));
@@ -195,8 +196,8 @@ public:
 		{
 			for (const FStreamingSourceShape& Shape : InShapes)
 			{
-				const FVector::FReal ShapeRadius = Shape.bUseGridLoadingRange ? (InGridLoadingRange * Shape.LoadingRangeScale) : Shape.Radius;
-				const FVector::FReal ShapeAngle = Shape.bIsSector ? Shape.SectorAngle : 360.0f;
+				const FVector::FReal ShapeRadius = (Shape.bUseGridLoadingRange ? (InGridLoadingRange * Shape.LoadingRangeScale) : Shape.Radius) + InExtraRadius;
+				const FVector::FReal ShapeAngle = Shape.bIsSector ? (Shape.SectorAngle + InExtraAngle) : 360.0f;
 				const FVector ShapeAxis = bInProjectIn2D ? FRotator(0, Shape.Rotation.Yaw, 0).Vector() : Shape.Rotation.Vector();
 				FSphericalSector LocalShape(bInProjectIn2D ? FVector(Shape.Location.X, Shape.Location.Y, 0) : Shape.Location, ShapeRadius, ShapeAxis, ShapeAngle);
 				if (LocalShape.IsValid())
@@ -372,6 +373,8 @@ struct FWorldPartitionStreamingSource
 		, Hash3D(0)
 		, OldLocation(FVector::ZeroVector)
 		, OldRotation(FRotator::ZeroRotator)
+		, ExtraRadius(0)
+		, ExtraAngle(0)
 	{}
 
 	FWorldPartitionStreamingSource(FName InName, const FVector& InLocation, const FRotator& InRotation, EStreamingSourceTargetState InTargetState, bool bInBlockOnSlowLoading, EStreamingSourcePriority InPriority, bool bRemote, float InVelocity = 0.f)
@@ -390,6 +393,8 @@ struct FWorldPartitionStreamingSource
 		, Hash3D(0)
 		, OldLocation(InLocation)
 		, OldRotation(InRotation)
+		, ExtraRadius(0)
+		, ExtraAngle(0)
 	{}
 
 	// Define Copy Constructor to avoid deprecation warnings
@@ -416,6 +421,10 @@ struct FWorldPartitionStreamingSource
 		bRemote = Other.bRemote;
 		Hash2D = Other.Hash2D;
 		Hash3D = Other.Hash3D;
+		OldLocation = Other.OldLocation;
+		OldRotation = Other.OldRotation;
+		ExtraRadius = Other.ExtraRadius;
+		ExtraAngle = Other.ExtraAngle;
 		return *this;
 	}
 	
@@ -490,7 +499,7 @@ struct FWorldPartitionStreamingSource
 	{
 		if (FStreamingSourceShapeHelper::IsSourceAffectingGrid(TargetGrids, TargetHLODLayers, TargetBehavior, InGridName, InGridHLODLayer))
 		{
-			FStreamingSourceShapeHelper::ForEachShape(InGridLoadingRange, InGridLoadingRange, bInProjectIn2D, Location, Rotation, Shapes, InOperation);
+			FStreamingSourceShapeHelper::ForEachShape(InGridLoadingRange, InGridLoadingRange, bInProjectIn2D, Location, Rotation, Shapes, InOperation, ExtraRadius, ExtraAngle);
 		}
 	}
 
@@ -512,6 +521,41 @@ private:
 	/** Source values used for hash computations. */
 	FVector OldLocation;
 	FRotator OldRotation;
+
+	/** Used internally for server streaming */
+	float ExtraRadius;
+	float ExtraAngle;
+
+	friend struct FSetStreamingSourceExtraRadius;
+	friend struct FSetStreamingSourceExtraAngle;
+};
+
+struct FSetStreamingSourceExtraRadius
+{
+private:
+	FSetStreamingSourceExtraRadius(FWorldPartitionStreamingSource& InStreamingSource, float InExtraRadius)
+	{
+		if (ensure(InExtraRadius >= 0.f))
+		{
+			InStreamingSource.ExtraRadius = InExtraRadius;
+		}
+	}
+
+	friend class UWorldPartitionSubsystem;
+};
+
+struct FSetStreamingSourceExtraAngle
+{
+private:
+	FSetStreamingSourceExtraAngle(FWorldPartitionStreamingSource& InStreamingSource, float InExtraAngle)
+	{
+		if (ensure(InExtraAngle >= 0.f))
+		{
+			InStreamingSource.ExtraAngle = InExtraAngle;
+		}
+	}
+
+	friend class UWorldPartitionSubsystem;
 };
 
 /**
