World Partition - HLOD: Store custom "editor bounds" in the HLOD actor descriptor
* For the HLOD actor, the streaming bounds is forced to be the bounds of cell represented by the HLOD.
* Editor bounds should be the correct actor bounds in order for queries of the editor hash to behave properly.
* Removed incorrect implementation of AWorldPartitionHLOD::GetActorBounds() which returned the streaming bounds. The default implementation is correct.

diff --git a/Engine/Source/Runtime/Core/Public/UObject/FortniteMainBranchObjectVersion.h b/Engine/Source/Runtime/Core/Public/UObject/FortniteMainBranchObjectVersion.h
--- a/Engine/Source/Runtime/Core/Public/UObject/FortniteMainBranchObjectVersion.h
+++ b/Engine/Source/Runtime/Core/Public/UObject/FortniteMainBranchObjectVersion.h
@@ -375,6 +375,9 @@ struct FFortniteMainBranchObjectVersion
 		// Serialize the source HLOD Layer for HLOD actor descriptors.
 		WorldPartitionHLODActorDescSerializeSourceHLODLayer,
 
+		// Serialize custom editor bounds for HLOD actor descriptors.
+		WorldPartitionHLODActorDescSerializeEditorBounds,
+
 		// -----<new versions can be added above this line>-------------------------------------------------
 		VersionPlusOne,
 		LatestVersion = VersionPlusOne - 1
diff --git a/Engine/Source/Runtime/Engine/Private/WorldPartition/HLOD/HLODActor.cpp b/Engine/Source/Runtime/Engine/Private/WorldPartition/HLOD/HLODActor.cpp
--- a/Engine/Source/Runtime/Engine/Private/WorldPartition/HLOD/HLODActor.cpp
+++ b/Engine/Source/Runtime/Engine/Private/WorldPartition/HLOD/HLODActor.cpp
@@ -396,11 +396,6 @@ void AWorldPartitionHLOD::SetHLODBounds(const FBox& InBounds)
 	HLODBounds = InBounds;
 }
 
-void AWorldPartitionHLOD::GetActorBounds(bool bOnlyCollidingComponents, FVector& Origin, FVector& BoxExtent, bool bIncludeFromChildActors) const
-{
-	HLODBounds.GetCenterAndExtents(Origin, BoxExtent);
-}
-
 FBox AWorldPartitionHLOD::GetStreamingBounds() const
 {
 	return HLODBounds;
diff --git a/Engine/Source/Runtime/Engine/Private/WorldPartition/HLOD/HLODActorDesc.cpp b/Engine/Source/Runtime/Engine/Private/WorldPartition/HLOD/HLODActorDesc.cpp
--- a/Engine/Source/Runtime/Engine/Private/WorldPartition/HLOD/HLODActorDesc.cpp
+++ b/Engine/Source/Runtime/Engine/Private/WorldPartition/HLOD/HLODActorDesc.cpp
@@ -48,6 +48,10 @@ void FHLODActorDesc::Init(const AActor* InActor)
 	}
 	
 	HLODStats = HLODActor->GetStats();
+
+	FVector Origin, Extent;
+	HLODActor->GetActorBounds(false, Origin, Extent);
+	EditorBounds = FBox::BuildAABB(Origin, Extent);
 }
 
 struct FHLODSubActorDescDeprecated
@@ -145,6 +149,11 @@ void FHLODActorDesc::Serialize(FArchive& Ar)
 		{
 			Ar << SourceHLODLayer;
 		}
+		
+		if (Ar.CustomVer(FFortniteMainBranchObjectVersion::GUID) >= FFortniteMainBranchObjectVersion::WorldPartitionHLODActorDescSerializeEditorBounds)
+		{
+			Ar << EditorBounds;
+		}
 	}
 }
 
@@ -155,7 +164,8 @@ bool FHLODActorDesc::Equals(const FWorldPartitionActorDesc* Other) const
 		const FHLODActorDesc& HLODActorDesc = *(FHLODActorDesc*)Other;
 		return SourceHLODLayer == HLODActorDesc.SourceHLODLayer &&
 			   HLODStats.OrderIndependentCompareEqual(HLODActorDesc.GetStats()) &&
-			   CompareUnsortedArrays(ChildHLODActors, HLODActorDesc.ChildHLODActors);
+			   CompareUnsortedArrays(ChildHLODActors, HLODActorDesc.ChildHLODActors) &&
+			   EditorBounds == HLODActorDesc.EditorBounds;
 	}
 	return false;
 }
diff --git a/Engine/Source/Runtime/Engine/Public/WorldPartition/HLOD/HLODActor.h b/Engine/Source/Runtime/Engine/Public/WorldPartition/HLOD/HLODActor.h
--- a/Engine/Source/Runtime/Engine/Public/WorldPartition/HLOD/HLODActor.h
+++ b/Engine/Source/Runtime/Engine/Public/WorldPartition/HLOD/HLODActor.h
@@ -79,8 +79,6 @@ protected:
 	virtual bool IsRuntimeOnly() const override { return true; }
 #if WITH_EDITOR
 	ENGINE_API virtual TUniquePtr<class FWorldPartitionActorDesc> CreateClassActorDesc() const override;
-
-	ENGINE_API virtual void GetActorBounds(bool bOnlyCollidingComponents, FVector& Origin, FVector& BoxExtent, bool bIncludeFromChildActors) const override;
 	ENGINE_API virtual FBox GetStreamingBounds() const override;
 
 	virtual bool ShouldImport(FStringView ActorPropString, bool IsMovingLevel) override { return false; }
diff --git a/Engine/Source/Runtime/Engine/Public/WorldPartition/HLOD/HLODActorDesc.h b/Engine/Source/Runtime/Engine/Public/WorldPartition/HLOD/HLODActorDesc.h
--- a/Engine/Source/Runtime/Engine/Public/WorldPartition/HLOD/HLODActorDesc.h
+++ b/Engine/Source/Runtime/Engine/Public/WorldPartition/HLOD/HLODActorDesc.h
@@ -41,6 +41,7 @@ protected:
 	ENGINE_API virtual void Serialize(FArchive& Ar) override;
 	virtual bool IsRuntimeRelevant(const FActorContainerID& InContainerID) const override { return !bIsForcedNonSpatiallyLoaded; }
 	virtual bool ShouldValidateRuntimeGrid() const override { return false; }
+	virtual FBox GetEditorBounds() const override { return EditorBounds; }
 	//~ End FWorldPartitionActorDesc Interface.
 
 	TArray<FGuid> ChildHLODActors;
@@ -48,5 +49,7 @@ protected:
 	FTopLevelAssetPath SourceHLODLayer;
 
 	FStats HLODStats;
+
+	FBox EditorBounds;
 };
 #endif
