Initialize the ISM instance buffer on the scene pipe. Fixes race condition with proxy shutdown.

diff --git a/Engine/Source/Runtime/Engine/Private/InstancedStaticMesh.cpp b/Engine/Source/Runtime/Engine/Private/InstancedStaticMesh.cpp
--- a/Engine/Source/Runtime/Engine/Private/InstancedStaticMesh.cpp
+++ b/Engine/Source/Runtime/Engine/Private/InstancedStaticMesh.cpp
@@ -1283,7 +1283,7 @@ FPerInstanceRenderData::FPerInstanceRenderData(FStaticMeshInstanceData& Other, E
 	InstanceBuffer_GameThread = InstanceBuffer.InstanceData;
 	if (!InstanceBuffer.CondSetFlushToGPUPending())
 	{
-		BeginInitResource(&InstanceBuffer);
+		BeginInitResource(&InstanceBuffer, &UE::RenderCommandPipe::Scene);
 	}
 	UpdateBoundsTransforms_Concurrent();
 }
