WorldPartitionCookPackageSplitter:
- Listen to world cleanup event when the package splitter is active to make sure the world is not uninitialized during UEditorEngine::Save.
- Partial backout of CL 28031122 since it's now replaced by the world cleanup event instead.

diff --git a/Engine/Source/Runtime/Engine/Classes/Engine/World.h b/Engine/Source/Runtime/Engine/Classes/Engine/World.h
--- a/Engine/Source/Runtime/Engine/Classes/Engine/World.h
+++ b/Engine/Source/Runtime/Engine/Classes/Engine/World.h
@@ -1295,9 +1295,6 @@ public:
 
 	/** Indicates that the renderer scene for this editor world was purged while Play-in-Editor. */
 	uint8 bPurgedScene : 1;
-
-	/** Indicates that the world is currently being used by a cook splitter.  */
-	uint8 bUsedByCookSplitter : 1;
 #endif
 
 	/** Keeps track whether actors moved via PostEditMove and therefore constraint syncup should be performed. */
diff --git a/Engine/Source/Runtime/Engine/Private/World.cpp b/Engine/Source/Runtime/Engine/Private/World.cpp
--- a/Engine/Source/Runtime/Engine/Private/World.cpp
+++ b/Engine/Source/Runtime/Engine/Private/World.cpp
@@ -593,7 +593,6 @@ UWorld::UWorld( const FObjectInitializer& ObjectInitializer )
 , bDebugFrameStepExecutedThisFrame(false)
 , bToggledBetweenPIEandSIEThisFrame(false)
 , bPurgedScene(false)
-, bUsedByCookSplitter(false)
 #endif
 , bShouldTick(true)
 , bHasEverBeenInitialized(false)
@@ -5311,10 +5310,6 @@ bool UWorld::IsNavigationRebuilt() const
 
 void UWorld::CleanupWorld(bool bSessionEnded, bool bCleanupResources, UWorld* NewWorld)
 {
-#if WITH_EDITOR
-	check(!bUsedByCookSplitter);
-#endif
-
     CleanupWorldGlobalTag++;
 	if (!bIsWorldInitialized)
 	{
diff --git a/Engine/Source/Runtime/Engine/Private/WorldPartition/WorldPartitionCookPackageSplitter.cpp b/Engine/Source/Runtime/Engine/Private/WorldPartition/WorldPartitionCookPackageSplitter.cpp
--- a/Engine/Source/Runtime/Engine/Private/WorldPartition/WorldPartitionCookPackageSplitter.cpp
+++ b/Engine/Source/Runtime/Engine/Private/WorldPartition/WorldPartitionCookPackageSplitter.cpp
@@ -6,6 +6,7 @@
 
 #include "AssetRegistry/IAssetRegistry.h"
 #include "Engine/Level.h"
+#include "Engine/World.h"
 #include "WorldPartition/WorldPartition.h"
 #include "WorldPartition/WorldPartitionLog.h"
 #include "WorldPartition/WorldPartitionRuntimeCell.h"
@@ -31,28 +32,20 @@ FWorldPartitionCookPackageSplitter::~FWorldPartitionCookPackageSplitter()
 
 void FWorldPartitionCookPackageSplitter::Teardown(ETeardown Status)
 {
-	UE_LOG(LogWorldPartition, Log, TEXT("[Cook] Debug(TearDown): OwnerObject=%s bInitializedWorldPartition=%d bForceInitializedWorld=%d bInitializedPhysicsSceneForSave=%d"), *GetFullNameSafe(ReferencedWorld), bInitializedWorldPartition ? 1 : 0, bForceInitializedWorld ? 1 : 0, bInitializedPhysicsSceneForSave ? 1 : 0);
+	UE_LOG(LogWorldPartition, Log, TEXT("[Cook] Debug(TearDown): OwnerObject=%s bForceInitializedWorld=%d bInitializedPhysicsSceneForSave=%d"), *GetFullNameSafe(ReferencedWorld),  bForceInitializedWorld ? 1 : 0, bInitializedPhysicsSceneForSave ? 1 : 0);
 
-	if (UWorld* LocalWorld = ReferencedWorld.Get())
-	{
-		check(LocalWorld->bUsedByCookSplitter);
-		LocalWorld->bUsedByCookSplitter = false;
+	FWorldDelegates::OnWorldCleanup.RemoveAll(this);
 
-		if (bInitializedWorldPartition)
-		{
-			if (UWorldPartition* WorldPartition = LocalWorld->PersistentLevel->GetWorldPartition())
-			{
-				WorldPartition->EndCook(CookContext);
-				WorldPartition->Uninitialize();
-			}
-		}
-	}
+	// Assume that the world is partitioned as per FWorldPartitionCookPackageSplitter::ShouldSplit
+	UWorldPartition* WorldPartition = ReferencedWorld->PersistentLevel->GetWorldPartition();
+	check(WorldPartition);
 
-	bInitializedWorldPartition = false;
+	WorldPartition->EndCook(CookContext);
+	WorldPartition->Uninitialize();
 
 	if (bInitializedPhysicsSceneForSave)
 	{
-		GEditor->CleanupPhysicsSceneThatWasInitializedForSave(ReferencedWorld.Get(), bForceInitializedWorld);
+		GEditor->CleanupPhysicsSceneThatWasInitializedForSave(ReferencedWorld, bForceInitializedWorld);
 		bInitializedPhysicsSceneForSave = false;
 		bForceInitializedWorld = false;
 	}
@@ -60,6 +53,11 @@ void FWorldPartitionCookPackageSplitter::Teardown(ETeardown Status)
 	ReferencedWorld = nullptr;
 }
 
+void FWorldPartitionCookPackageSplitter::OnWorldCleanup(UWorld* InWorld, bool bSessionEnded, bool bCleanupResources)
+{
+	check(InWorld != ReferencedWorld);
+}
+
 void FWorldPartitionCookPackageSplitter::AddReferencedObjects(FReferenceCollector& Collector)
 {
 	Collector.AddReferencedObject(ReferencedWorld);
@@ -70,44 +68,25 @@ FString FWorldPartitionCookPackageSplitter::GetReferencerName() const
 	return TEXT("FWorldPartitionCookPackageSplitter");
 }
 
-UWorld* FWorldPartitionCookPackageSplitter::ValidateDataObject(UObject* SplitData)
-{
-	UWorld* PartitionedWorld = CastChecked<UWorld>(SplitData);
-	check(PartitionedWorld);
-	check(PartitionedWorld->PersistentLevel);
-	check(PartitionedWorld->IsPartitionedWorld());
-	return PartitionedWorld;
-}
-
-const UWorld* FWorldPartitionCookPackageSplitter::ValidateDataObject(const UObject* SplitData)
-{
-	return ValidateDataObject(const_cast<UObject*>(SplitData));
-}
-
 TArray<ICookPackageSplitter::FGeneratedPackage> FWorldPartitionCookPackageSplitter::GetGenerateList(const UPackage* OwnerPackage, const UObject* OwnerObject)
 {
-	// TODO: Make WorldPartition functions const so we can honor the constness of the OwnerObject in this API function
-	const UWorld* ConstPartitionedWorld = ValidateDataObject(OwnerObject);
-	UWorld* PartitionedWorld = const_cast<UWorld*>(ConstPartitionedWorld);
+	UE_LOG(LogWorldPartition, Display, TEXT("[Cook] Gathering packages to cook from generators for owner object %s."), *GetFullNameSafe(OwnerObject));
 
 	// Store the World pointer to declare it to GarbageCollection; we do not want to allow the World to be Garbage Collected
 	// until we have finished all of our PreSaveGeneratedPackage calls, because we store information on the World 
 	// that is necessary for populate 
-	ReferencedWorld = PartitionedWorld;
-
-	check(!PartitionedWorld->bUsedByCookSplitter);
-	PartitionedWorld->bUsedByCookSplitter = true;
+	ReferencedWorld = const_cast<UWorld*>(CastChecked<const UWorld>(OwnerObject));
 
 	check(!bInitializedPhysicsSceneForSave && !bForceInitializedWorld);
-	bInitializedPhysicsSceneForSave = GEditor->InitializePhysicsSceneForSaveIfNecessary(PartitionedWorld, bForceInitializedWorld);
+	bInitializedPhysicsSceneForSave = GEditor->InitializePhysicsSceneForSaveIfNecessary(ReferencedWorld, bForceInitializedWorld);
+
+	// Assume that the world is partitioned as per FWorldPartitionCookPackageSplitter::ShouldSplit
+	UWorldPartition* WorldPartition = ReferencedWorld->PersistentLevel->GetWorldPartition();
+	check(WorldPartition);
 
-	// Manually initialize WorldPartition
-	UWorldPartition* WorldPartition = PartitionedWorld->PersistentLevel->GetWorldPartition();
 	// We expect the WorldPartition has not yet been initialized
 	ensure(!WorldPartition->IsInitialized());
-	WorldPartition->Initialize(PartitionedWorld, FTransform::Identity);
-	bInitializedWorldPartition = true;
-
+	WorldPartition->Initialize(ReferencedWorld, FTransform::Identity);
 	WorldPartition->BeginCook(CookContext);
 
 	bool bIsSuccess = CookContext.GatherPackagesToCook();
@@ -119,7 +98,9 @@ TArray<ICookPackageSplitter::FGeneratedPackage> FWorldPartitionCookPackageSplitt
 	BuildPackagesToGenerateList(PackagesToGenerate);
 
 	UE_LOG(LogWorldPartition, Log, TEXT("[Cook] Sending %u packages to be generated."), PackagesToGenerate.Num());
-	UE_LOG(LogWorldPartition, Log, TEXT("[Cook] Debug(GetGenerateList) : OwnerObject=%s bInitializedWorldPartition=%d bForceInitializedWorld=%d bInitializedPhysicsSceneForSave=%d"), *GetFullNameSafe(OwnerObject), bInitializedWorldPartition ? 1 : 0, bForceInitializedWorld ? 1 : 0, bInitializedPhysicsSceneForSave ? 1 : 0);
+	UE_LOG(LogWorldPartition, Log, TEXT("[Cook] Debug(GetGenerateList) : OwnerObject=%s bForceInitializedWorld=%d bInitializedPhysicsSceneForSave=%d"), *GetFullNameSafe(OwnerObject), bForceInitializedWorld ? 1 : 0, bInitializedPhysicsSceneForSave ? 1 : 0);
+
+	FWorldDelegates::OnWorldCleanup.AddRaw(this, &FWorldPartitionCookPackageSplitter::OnWorldCleanup);
 
 	return PackagesToGenerate;
 }
diff --git a/Engine/Source/Runtime/Engine/Private/WorldPartition/WorldPartitionCookPackageSplitter.h b/Engine/Source/Runtime/Engine/Private/WorldPartition/WorldPartitionCookPackageSplitter.h
--- a/Engine/Source/Runtime/Engine/Private/WorldPartition/WorldPartitionCookPackageSplitter.h
+++ b/Engine/Source/Runtime/Engine/Private/WorldPartition/WorldPartitionCookPackageSplitter.h
@@ -31,13 +31,12 @@ public:
 	//~ End of ICookPackageSplitter
 
 private:
+	void OnWorldCleanup(UWorld* InWorld, bool bSessionEnded, bool bCleanupResources);
+
 	/** FGCObject interface */
 	virtual void AddReferencedObjects(FReferenceCollector& Collector) override;
 	virtual FString GetReferencerName() const override;
 
-	const UWorld* ValidateDataObject(const UObject* SplitData);
-	UWorld* ValidateDataObject(UObject* SplitData);
-
 	void BuildPackagesToGenerateList(TArray<ICookPackageSplitter::FGeneratedPackage>& PackagesToGenerate) const;
 	bool MapGeneratePackageToCookPackage(const TArray<ICookPackageSplitter::FGeneratedPackageForPreSave>& GeneratedPackages);
 	
@@ -45,7 +44,6 @@ private:
 
 	FWorldPartitionCookPackageContext CookContext;
 
-	bool bInitializedWorldPartition = false;
 	bool bForceInitializedWorld = false;
 	bool bInitializedPhysicsSceneForSave = false;
 };
