Fix CIS: amend check for setting render data to allow commandlets.

#rnx

diff --git a/Engine/Source/Runtime/Engine/Private/StaticMesh.cpp b/Engine/Source/Runtime/Engine/Private/StaticMesh.cpp
--- a/Engine/Source/Runtime/Engine/Private/StaticMesh.cpp
+++ b/Engine/Source/Runtime/Engine/Private/StaticMesh.cpp
@@ -3479,7 +3479,8 @@ void UStaticMesh::SetRenderData(TUniquePtr<class FStaticMeshRenderData>&& InRend
 {
 	WaitUntilAsyncPropertyReleased(EStaticMeshAsyncProperties::RenderData);
 
-	checkf(FApp::CanEverRender() || InRenderData == nullptr, TEXT("Non-null RenderData should not be set in headless runs."));
+	// DDC generation can cause a non-null data in non-rendering commandlets
+	checkf(FApp::CanEverRender() || IsRunningCommandlet() || InRenderData == nullptr, TEXT("Non-null RenderData should not be set in headless runs."));
 
 	PRAGMA_DISABLE_DEPRECATION_WARNINGS
 	RenderData = MoveTemp(InRenderData);
