World Partition - HLOD: Initialize the HLOD actor desc EditorBounds properly in case we're loading a version without it

diff --git a/Engine/Source/Runtime/Engine/Private/WorldPartition/HLOD/HLODActorDesc.cpp b/Engine/Source/Runtime/Engine/Private/WorldPartition/HLOD/HLODActorDesc.cpp
--- a/Engine/Source/Runtime/Engine/Private/WorldPartition/HLOD/HLODActorDesc.cpp
+++ b/Engine/Source/Runtime/Engine/Private/WorldPartition/HLOD/HLODActorDesc.cpp
@@ -17,6 +17,11 @@
 #include "WorldPartition/HLOD/HLODLayer.h"
 #include "WorldPartition/HLOD/HLODSourceActorsFromCell.h"
 
+FHLODActorDesc::FHLODActorDesc()
+	: EditorBounds(ForceInit)
+{
+}
+
 
 void FHLODActorDesc::Init(const AActor* InActor)
 {
@@ -154,6 +159,10 @@ void FHLODActorDesc::Serialize(FArchive& Ar)
 		{
 			Ar << EditorBounds;
 		}
+		else
+		{
+			EditorBounds = GetRuntimeBounds();
+		}
 	}
 }
 
diff --git a/Engine/Source/Runtime/Engine/Public/WorldPartition/HLOD/HLODActorDesc.h b/Engine/Source/Runtime/Engine/Public/WorldPartition/HLOD/HLODActorDesc.h
--- a/Engine/Source/Runtime/Engine/Public/WorldPartition/HLOD/HLODActorDesc.h
+++ b/Engine/Source/Runtime/Engine/Public/WorldPartition/HLOD/HLODActorDesc.h
@@ -20,6 +20,7 @@ class AWorldPartitionHLOD;
  */
 class FHLODActorDesc : public FWorldPartitionActorDesc
 {
+	friend class AWorldPartitionHLOD;
 	friend class FHLODActorDescFactory;
 
 public:
@@ -34,6 +35,8 @@ public:
 	static ENGINE_API int64 GetPackageSize(const AWorldPartitionHLOD* InHLODActor);
 
 protected:
+	ENGINE_API FHLODActorDesc();
+
 	//~ Begin FWorldPartitionActorDesc Interface.
 	ENGINE_API virtual void Init(const AActor* InActor) override;
 	ENGINE_API virtual bool Equals(const FWorldPartitionActorDesc* Other) const override;
