Fix wording.

#rnx

diff --git a/Engine/Source/Runtime/Engine/Public/PSOPrecache.h b/Engine/Source/Runtime/Engine/Public/PSOPrecache.h
--- a/Engine/Source/Runtime/Engine/Public/PSOPrecache.h
+++ b/Engine/Source/Runtime/Engine/Public/PSOPrecache.h
@@ -19,7 +19,7 @@ class FVertexFactoryType;
 class FGraphicsPipelineStateInitializer;
 enum class EVertexInputStreamType : uint8;
 
-// master switch that decides whether to compile out some PSO precaching code (most importantly, reduce sizeofs of common classes)
+// General switch that decides whether to compile out some PSO precaching code (most importantly, reduce sizeofs of common classes)
 #ifndef UE_WITH_PSO_PRECACHING
 	#define UE_WITH_PSO_PRECACHING		(PLATFORM_SUPPORTS_PSO_PRECACHING)
 #endif // UE_WITH_PSO_PRECACHING
