Don't popup a slow task dialog when compiling shaders in PIE since it grabs the focus and stops player movements.

diff --git a/Engine/Source/Runtime/Engine/Private/ShaderCompiler/ShaderCompiler.cpp b/Engine/Source/Runtime/Engine/Private/ShaderCompiler/ShaderCompiler.cpp
--- a/Engine/Source/Runtime/Engine/Private/ShaderCompiler/ShaderCompiler.cpp
+++ b/Engine/Source/Runtime/Engine/Private/ShaderCompiler/ShaderCompiler.cpp
@@ -5982,7 +5982,7 @@ void FShaderCompilingManager::BlockOnShaderMapCompletion(const TArray<int32>& Sh
 			}
 		}
 
-		FScopedSlowTask SlowTask(NumJobs, FText::Format(LOCTEXT("BlockOnShaderMapCompletion", "Compiling Shaders ({0})"), NumJobs), GIsEditor && !IsRunningCommandlet());
+		FScopedSlowTask SlowTask(NumJobs, FText::Format(LOCTEXT("BlockOnShaderMapCompletion", "Compiling Shaders ({0})"), NumJobs), GIsEditor && !IsRunningCommandlet() && GPlayInEditorID == INDEX_NONE);
 		SlowTask.MakeDialog(false, true);
 
 		int32 NumPendingJobs = 0;
