Shader Compiler:  Fix for startup hang due to nesting write locks when running with -nothreading.  Need to go through blocking DDC query code path, since the -nothreading flag forces the query to run on the same thread.

#rnx

diff --git a/Engine/Source/Runtime/Engine/Private/ShaderCompiler/ShaderCompiler.cpp b/Engine/Source/Runtime/Engine/Private/ShaderCompiler/ShaderCompiler.cpp
--- a/Engine/Source/Runtime/Engine/Private/ShaderCompiler/ShaderCompiler.cpp
+++ b/Engine/Source/Runtime/Engine/Private/ShaderCompiler/ShaderCompiler.cpp
@@ -10217,7 +10217,7 @@ FShaderJobCacheRef FShaderJobCache::FindOrAdd(const FJobInputHash& Hash, EShader
 			UE::DerivedData::FRequestOwner* RequestOwner;
 
 			static const bool PerShaderDDCAsync = CVarShaderCompilerPerShaderDDCAsync.GetValueOnAnyThread();
-			if (PerShaderDDCAsync)
+			if (PerShaderDDCAsync && FGenericPlatformProcess::SupportsMultithreading())
 			{
 				switch (JobPriority)
 				{
