Holdouts: Implement the holdout flag for the deferred renderer

When the holdout flag is enabled, the object will render black with an alpha matching the background. This is implemented with an extra condition in the basepass.

The logic is only added to the base pass pixel shader when the project is configure to have post processing of the alpha channel to minimize the impact on games.

diff --git a/Engine/Shaders/Private/BasePassPixelShader.usf b/Engine/Shaders/Private/BasePassPixelShader.usf
--- a/Engine/Shaders/Private/BasePassPixelShader.usf
+++ b/Engine/Shaders/Private/BasePassPixelShader.usf
@@ -2115,6 +2115,23 @@ void FPixelShaderInOut_MainPS(
 	Out.MRT[0] = SubstrateMaterialExport;
 #endif
 
+#if POST_PROCESS_ALPHA != 0 && MATERIAL_DOMAIN_SURFACE
+	// Only include the holdout logic if the Renderer is forwarding its alpha channel to post-processing
+	if (GetPrimitiveData(MaterialParameters).Flags & PRIMITIVE_SCENE_DATA_FLAG_HOLDOUT)
+	{
+		// Discard this pixel so that nothing gets written to the G-Buffers and alpha is kept as its default value (background)
+
+		// TODO: The translucent case doesn't seem to get the right alpha (compared to the path traced reference)
+		Out.MRT[0] = half4(0.0, 0.0, 0.0, Opacity);
+
+#if USES_GBUFFER
+		Out.MRT[1] = 0.0;
+		Out.MRT[2] = 0.0;
+		Out.MRT[3] = 0.0;
+#endif
+	}
+#endif
+
 // Duplicate coverage bit values in case multiple pixels are written by one invocation (VRS)
 #if OUTPUTS_COVERAGE
 	// If centroids are integers, it means we're at a non-1x1 shading rate
diff --git a/Engine/Source/Runtime/Engine/Classes/Components/PrimitiveComponent.h b/Engine/Source/Runtime/Engine/Classes/Components/PrimitiveComponent.h
--- a/Engine/Source/Runtime/Engine/Classes/Components/PrimitiveComponent.h
+++ b/Engine/Source/Runtime/Engine/Classes/Components/PrimitiveComponent.h
@@ -415,8 +415,8 @@ public:
 	UPROPERTY(EditAnywhere, AdvancedDisplay, BlueprintReadOnly, Category=Rendering)
 	uint8 bReceivesDecals:1;
 
-	/** If this is True, this primitive will render black with an alpha of 0, but all secondary effects (shadows, reflections, indirect lighting) remain. This feature is currently only implemented in the Path Tracer. */
-	UPROPERTY(EditAnywhere, BlueprintReadOnly, Category = PathTracing, Interp)
+	/** If this is True, this primitive will render black with an alpha of 0, but all secondary effects (shadows, reflections, indirect lighting) remain. This feature required the project setting "Enable alpha channel support in post processing". */
+	UPROPERTY(EditAnywhere, AdvancedDisplay, BlueprintReadOnly, Category = Rendering, Interp)
 	uint8 bHoldout : 1;
 
 	/** If this is True, this component won't be visible when the view actor is the component's owner, directly or indirectly. */
