Fix material function crash with Substrate on preview.

There were two issues:
* The material function preview was not correctly converting material attribute output.
* The cached expression were not populated, causing the IsMaterialAttributeOutput() function to not return the correct value.

Down the line, this was causing material attribute data to be connected to a float3 input causing a crash down the line.

diff --git a/Engine/Source/Editor/MaterialEditor/Private/MaterialEditor.cpp b/Engine/Source/Editor/MaterialEditor/Private/MaterialEditor.cpp
--- a/Engine/Source/Editor/MaterialEditor/Private/MaterialEditor.cpp
+++ b/Engine/Source/Editor/MaterialEditor/Private/MaterialEditor.cpp
@@ -2540,6 +2540,7 @@ void FMaterialEditor::UpdatePreviewMaterial( bool bForce )
 
 	if( PreviewExpression && ExpressionPreviewMaterial )
 	{
+		ExpressionPreviewMaterial->UpdateCachedExpressionData();
 		PreviewExpression->ConnectToPreviewMaterial(ExpressionPreviewMaterial,0);
 	}
 
diff --git a/Engine/Source/Runtime/Engine/Private/Materials/MaterialExpressions.cpp b/Engine/Source/Runtime/Engine/Private/Materials/MaterialExpressions.cpp
--- a/Engine/Source/Runtime/Engine/Private/Materials/MaterialExpressions.cpp
+++ b/Engine/Source/Runtime/Engine/Private/Materials/MaterialExpressions.cpp
@@ -668,7 +668,6 @@ int32 CompileShadingModelBlendFunction(FMaterialCompiler* Compiler, const int32
 
 int32 CompileSubstrateBlendFunction(FMaterialCompiler* Compiler, const int32 A, const int32 B, const int32 Alpha)
 {
-	check(false);	// we should never blend Substrate data from material layers.
 	return INDEX_NONE;
 }
 
@@ -2076,21 +2075,23 @@ void UMaterialExpression::ConnectToPreviewMaterial(UMaterial* InMaterial, int32
 			}
 			else if (IsResultMaterialAttributes(0))
 			{
-				// Set preview material as unlit and not using material attribute
-				InMaterial->SetShadingModel(MSM_Unlit);
-				InMaterial->bUseMaterialAttributes = false;
-
-				// Connect the material attribute output to the material attribute break node
-				UMaterialExpressionBreakMaterialAttributes* BreakMatAtt = NewObject<UMaterialExpressionBreakMaterialAttributes>(this);
-				BreakMatAtt->MaterialAttributes.Connect(OutputIndex, this);
+				// Propagate material attributes to MaterialAttributes input
+				InMaterial->SetShadingModel(MSM_DefaultLit);
+				InMaterial->bUseMaterialAttributes = true;
+				FExpressionInput* MaterialInput = InMaterial->GetExpressionInputForProperty(MP_MaterialAttributes);
+				check(MaterialInput);
+				ConnectExpression( MaterialInput, OutputIndex );
 
-				// Get BaseColor as unlit color preview from material attribute
-				UMaterialExpressionSubstrateUnlitBSDF* UnlitBSDF = NewObject<UMaterialExpressionSubstrateUnlitBSDF>(this);
-				UnlitBSDF->EmissiveColor.Connect(0, BreakMatAtt);
+				// Converte material input into Substrate data
+				UMaterialExpressionSubstrateConvertMaterialAttributes* ConvertAttributeNode = NewObject<UMaterialExpressionSubstrateConvertMaterialAttributes>(this);
+				ConvertAttributeNode->Material = InMaterial;
+				ConvertAttributeNode->MaterialAttributes.Connect(OutputIndex, this);
 
-				FExpressionInput* MaterialInput = InMaterial->GetExpressionInputForProperty(MP_FrontMaterial);
-				check(MaterialInput);
-				MaterialInput->Connect(0, UnlitBSDF);
+				// Connect substrate data into material FrontMaterial input
+				if (UMaterialEditorOnlyData* MaterialEditorOnlyData = InMaterial->GetEditorOnlyData())
+				{
+					MaterialEditorOnlyData->FrontMaterial.Connect(0, ConvertAttributeNode);
+				}
 			}
 			else
 			{
