[NetPhysics] Resim and replication improvements

- Fix issue where PI could keep hold of a target state indefinitely after a resim was performed if the server stopped sending states around the same time.
- Implemented functionality to stop PI and Default physics replication while the particles interact with a resim particle, the resim particle will trigger a resimulation if the simulation goes out of sync and the resim uses the received targets from the server for all particles in the islands resimulated.
- Implement functionality to find a valid resim frame where all particles in a resim island have a valid target to rewind to.

diff --git a/Engine/Source/Runtime/Engine/Private/PhysicsEngine/PhysicsReplication.cpp b/Engine/Source/Runtime/Engine/Private/PhysicsEngine/PhysicsReplication.cpp
--- a/Engine/Source/Runtime/Engine/Private/PhysicsEngine/PhysicsReplication.cpp
+++ b/Engine/Source/Runtime/Engine/Private/PhysicsEngine/PhysicsReplication.cpp
@@ -92,6 +92,9 @@ namespace PhysicsReplicationCVars
 		static bool bRuntimeCorrectionEnabled = true;
 		static FAutoConsoleVariableRef CVarResimRuntimeCorrectionEnabled(TEXT("np2.Resim.RuntimeCorrectionEnabled"), bRuntimeCorrectionEnabled, TEXT("Apply runtime corrections while error is smalle enough not to trigger a resim."));
 
+		static bool bDisableReplicationOnInteraction = true;
+		static FAutoConsoleVariableRef CVarResimDisableReplicationOnInteraction(TEXT("np2.Resim.DisableReplicationOnInteraction"), bDisableReplicationOnInteraction, TEXT("If a resim object interacts with another object not running resimulation, deactivate that objects replication until interaction stops."));
+
 		static float PosStabilityMultiplier = 0.5f;
 		static FAutoConsoleVariableRef CVarResimPosStabilityMultiplier(TEXT("np2.Resim.PosStabilityMultiplier"), PosStabilityMultiplier, TEXT("Recommended range between 0.0-1.0. Lower value means more stable positional corrections."));
 
@@ -152,7 +155,7 @@ namespace PhysicsReplicationCVars
 		static float EarlyOutAngle = 1.f;
 		static FAutoConsoleVariableRef CVarEarlyOutAngle(TEXT("np2.PredictiveInterpolation.EarlyOutAngle"), EarlyOutAngle, TEXT("If object is within this rotational angle (in degrees) from the source target, early out from replication and apply sleep if replicated."));
 		
-		static bool PostResimWaitForUpdate = true;
+		static bool PostResimWaitForUpdate = false;
 		static FAutoConsoleVariableRef CVarPostResimWaitForUpdate(TEXT("np2.PredictiveInterpolation.PostResimWaitForUpdate"), PostResimWaitForUpdate, TEXT("After a resimulation, wait for replicated states that correspond to post-resim state before processing replication again."));
 		
 		static bool bVelocityBased = true;
@@ -732,7 +735,7 @@ void FPhysicsReplicationAsync::OnPreSimulate_Internal()
 		const bool bRewindDataExist = RewindData != nullptr;
 		if (bRewindDataExist && RewindData->IsResim())
 		{
-			// TODO, Handle the transition from post-resim to interpolation better.
+			// TODO, Handle the transition from post-resim to interpolation better (disabled by default, resim vs replication interaction is handled via FPhysicsReplicationAsync::CacheResimInteractions)
 			if (PhysicsReplicationCVars::PredictiveInterpolationCVars::PostResimWaitForUpdate && RewindData->IsFinalResim())
 			{
 				for (auto Itr = ObjectToTarget.CreateIterator(); Itr; ++Itr)
@@ -775,6 +778,11 @@ void FPhysicsReplicationAsync::OnPreSimulate_Internal()
 			UpdateAsyncTarget(Input, RigidsSolver);
 		}
 
+		if (Chaos::FPBDRigidsSolver::IsNetworkPhysicsPredictionEnabled())
+		{
+			CacheResimInteractions();
+		}
+
 		ApplyTargetStatesAsync(GetDeltaTime_Internal(), AsyncInput->ErrorCorrection, AsyncInput->InputData);
 	}
 }
@@ -829,13 +837,13 @@ void FPhysicsReplicationAsync::UpdateAsyncTarget(const FPhysicsRepAsyncInputData
 		Target->PrevLinVel = Input.TargetState.LinVel;
 	}
 
-	if (Input.ServerFrame >= Target->ServerFrame)
+	if (Input.ServerFrame >= Target->PrevServerFrame)
 	{
 		const int32 PrevTickCount = Target->TickCount;
 		const int32 PrevReceiveInterval = Target->ReceiveInterval;
 
-		Target->PrevServerFrame = (Target->ServerFrame == INDEX_NONE) ? (Input.ServerFrame - 1) : Target->ServerFrame;
-		Target->ServerFrame = Input.ServerFrame;
+		Target->PrevServerFrame = Target->bWaiting ? Input.ServerFrame : Target->ServerFrame;
+		Target->ServerFrame = Target->bWaiting ? Target->ServerFrame : Input.ServerFrame;
 		Target->PrevReceiveFrame = (Target->ReceiveFrame == INDEX_NONE) ? (RigidsSolver->GetCurrentFrame() - 1) : Target->ReceiveFrame;
 		Target->ReceiveFrame = RigidsSolver->GetCurrentFrame();
 		Target->ReceiveInterval = FMath::Clamp((Target->ReceiveFrame - Target->PrevReceiveFrame), 1, 255);
@@ -843,7 +851,6 @@ void FPhysicsReplicationAsync::UpdateAsyncTarget(const FPhysicsRepAsyncInputData
 		Target->RepMode = Input.RepMode;
 		Target->FrameOffset = Input.FrameOffset;
 		Target->TickCount = 0;
-		Target->bWaiting = false;
 
 		if (Input.RepMode == EPhysicsReplicationMode::PredictiveInterpolation)
 		{
@@ -853,7 +860,7 @@ void FPhysicsReplicationAsync::UpdateAsyncTarget(const FPhysicsRepAsyncInputData
 
 			// If we extrapolated the previous target past the receive interval, extrapolate this target by the overshoot
 			const int32 OvershootingExtrapolation = FMath::Clamp((PrevTickCount - PrevReceiveInterval), -1, FMath::CeilToInt(Target->AverageReceiveInterval));
-			if (OvershootingExtrapolation > 0)
+			if (!Target->bWaiting && OvershootingExtrapolation > 0)
 			{
 				ExtrapolateTarget(*Target, OvershootingExtrapolation, GetDeltaTime_Internal());
 			}
@@ -864,6 +871,35 @@ void FPhysicsReplicationAsync::UpdateAsyncTarget(const FPhysicsRepAsyncInputData
 	LatencyOneWay = Input.LatencyOneWay;
 }
 
+void FPhysicsReplicationAsync::CacheResimInteractions()
+{
+	Chaos::FPBDRigidsSolver* RigidsSolver = static_cast<Chaos::FPBDRigidsSolver*>(GetSolver());
+	if (RigidsSolver == nullptr)
+	{
+		return;
+	}
+
+	ParticlesInResimIslands.Empty(FMath::CeilToInt(static_cast<float>(ParticlesInResimIslands.Num()) * 0.9f));
+	Chaos::Private::FPBDIslandManager& IslandManager = RigidsSolver->GetEvolution()->GetIslandManager();
+	Chaos::FWritePhysicsObjectInterface_Internal Interface = Chaos::FPhysicsObjectInternalInterface::GetWrite();
+	for (auto Itr = ObjectToTarget.CreateIterator(); Itr; ++Itr)
+	{
+		FReplicatedPhysicsTargetAsync& Target = Itr.Value();
+		if (Target.RepMode == EPhysicsReplicationMode::Resimulation)
+		{
+			Chaos::FConstPhysicsObjectHandle& POHandle = Itr.Key();
+			if (Chaos::FGeometryParticleHandle* Handle = Interface.GetParticle(POHandle))
+			{
+				// Get a list of particles from the same island as a resim particle is in, i.e. particles interacting with a resim particle
+				for (const Chaos::FGeometryParticleHandle* InteractParticle : IslandManager.FindParticlesInIslands(IslandManager.FindParticleIslands(Handle)))
+				{
+					ParticlesInResimIslands.Add(InteractParticle->GetHandleIdx());
+				}
+			}
+		}
+	}
+}
+
 void FPhysicsReplicationAsync::ApplyTargetStatesAsync(const float DeltaSeconds, const FPhysicsRepErrorCorrectionData& ErrorCorrection, const TArray<FPhysicsRepAsyncInputData>& InputData)
 {
 	using namespace Chaos;
@@ -882,20 +918,19 @@ void FPhysicsReplicationAsync::ApplyTargetStatesAsync(const float DeltaSeconds,
 	}
 
 	// PhysicsObject flow
+	Chaos::FWritePhysicsObjectInterface_Internal Interface = Chaos::FPhysicsObjectInternalInterface::GetWrite();
 	for (auto Itr = ObjectToTarget.CreateIterator(); Itr; ++Itr)
 	{
 		bool bRemoveItr = true; // Remove current cached replication target unless replication logic tells us to store it for next tick
 
-		Chaos::FConstPhysicsObjectHandle POHandle = Itr.Key();
-		Chaos::FWritePhysicsObjectInterface_Internal Interface = Chaos::FPhysicsObjectInternalInterface::GetWrite();
+		Chaos::FConstPhysicsObjectHandle& POHandle = Itr.Key();
 		if (FGeometryParticleHandle* Handle = Interface.GetParticle(POHandle))
 		{
 			FReplicatedPhysicsTargetAsync& Target = Itr.Value();
 
 			if (FPBDRigidParticleHandle* RigidHandle = Handle->CastToRigidParticle())
 			{
-				EPhysicsReplicationMode RepMode = Target.RepMode;
-
+				const EPhysicsReplicationMode RepMode = Target.RepMode;
 				switch (RepMode)
 				{
 					case EPhysicsReplicationMode::Default:
@@ -984,6 +1019,11 @@ bool FPhysicsReplicationAsync::DefaultReplication(Chaos::FPBDRigidParticleHandle
 		return true;
 	}
 
+	if (PhysicsReplicationCVars::ResimulationCVars::bDisableReplicationOnInteraction && ParticlesInResimIslands.Contains(Handle->GetHandleIdx()))
+	{
+		return false;
+	}
+
 	//
 	// NOTES:
 	//
@@ -1222,10 +1262,11 @@ bool FPhysicsReplicationAsync::PredictiveInterpolation(Chaos::FPBDRigidParticleH
 		return true;
 	}
 
-	if (Target.bWaiting)
+	if (Target.bWaiting && (Target.PrevServerFrame + Target.TickCount) < Target.ServerFrame)
 	{
 		return false;
 	}
+	Target.bWaiting = false;
 
 	Chaos::FPBDRigidsSolver* RigidsSolver = static_cast<Chaos::FPBDRigidsSolver*>(GetSolver());
 	if (RigidsSolver == nullptr)
@@ -1233,6 +1274,14 @@ bool FPhysicsReplicationAsync::PredictiveInterpolation(Chaos::FPBDRigidParticleH
 		return true;
 	}
 
+	if (PhysicsReplicationCVars::ResimulationCVars::bDisableReplicationOnInteraction && ParticlesInResimIslands.Contains(Handle->GetHandleIdx()))
+	{
+		// If particle is in an island with a resim object, don't run replication and wait for an up to date target (after leaving the island)
+		Target.bWaiting = true;
+		Target.ServerFrame = RigidsSolver->GetCurrentFrame() + Target.FrameOffset;
+		return false;
+	}
+
 	const bool bCanSimulate = Handle->IsDynamic() || Handle->IsSleeping();
 
 	// Helper for sleep and target clearing at replication end
@@ -1245,7 +1294,7 @@ bool FPhysicsReplicationAsync::PredictiveInterpolation(Chaos::FPBDRigidParticleH
 			RigidsSolver->GetEvolution()->SetParticleObjectState(Handle, Chaos::EObjectStateType::Sleeping);
 		}
 
-		const bool bClearTarget = (!bCanSimulate || bShouldSleep) && !PhysicsReplicationCVars::PredictiveInterpolationCVars::bDontClearTarget;
+		const bool bClearTarget = (!bCanSimulate || (bAllowSleep && bShouldSleep)) && !PhysicsReplicationCVars::PredictiveInterpolationCVars::bDontClearTarget;
 		return bClearTarget;
 	};
 
@@ -1459,6 +1508,7 @@ bool FPhysicsReplicationAsync::ResimulationReplication(Chaos::FPBDRigidParticleH
 	}
 
 	const int32 LocalFrame = Target.ServerFrame - Target.FrameOffset;
+
 	if (LocalFrame > RewindData->CurrentFrame() || LocalFrame < RewindData->GetEarliestFrame_Internal())
 	{
 		if (LocalFrame > 0)
diff --git a/Engine/Source/Runtime/Engine/Public/PhysicsReplication.h b/Engine/Source/Runtime/Engine/Public/PhysicsReplication.h
--- a/Engine/Source/Runtime/Engine/Public/PhysicsReplication.h
+++ b/Engine/Source/Runtime/Engine/Public/PhysicsReplication.h
@@ -97,6 +97,8 @@ struct FReplicatedPhysicsTargetAsync
 		, ReceiveFrame(INDEX_NONE)
 		, ReceiveInterval(5)
 		, AverageReceiveInterval(5.f)
+		, PrevServerFrame(INDEX_NONE)
+		, bWaiting(false)
 	{ }
 
 	/** The target state replicated by server */
@@ -157,10 +159,12 @@ private:
 	float LatencyOneWay;
 	FRigidBodyErrorCorrection ErrorCorrectionDefault;
 	TMap<Chaos::FConstPhysicsObjectHandle, FReplicatedPhysicsTargetAsync> ObjectToTarget;
+	TArray<int32> ParticlesInResimIslands;
 
 private:
 	void UpdateAsyncTarget(const FPhysicsRepAsyncInputData& Input, Chaos::FPBDRigidsSolver* RigidsSolver);
 	void UpdateRewindDataTarget(const FPhysicsRepAsyncInputData& Input);
+	void CacheResimInteractions();
 	void ExtrapolateTarget(FReplicatedPhysicsTargetAsync& Target, const int32 ExtrapolateFrames, const float DeltaSeconds);
 
 public:
diff --git a/Engine/Source/Runtime/Experimental/Chaos/Private/RewindData.cpp b/Engine/Source/Runtime/Experimental/Chaos/Private/RewindData.cpp
--- a/Engine/Source/Runtime/Experimental/Chaos/Private/RewindData.cpp
+++ b/Engine/Source/Runtime/Experimental/Chaos/Private/RewindData.cpp
@@ -846,8 +846,15 @@ void FRewindData::ExtendHistoryWithFrame(const int32 Frame)
 	FramesSaved = FMath::Max(CurFrame - Frame+1, FramesSaved);
 }
 
-CHAOS_API bool bResimFrameValidationLeniency = true;
-FAutoConsoleVariableRef CVarResimFrameValidationLeniency(TEXT("p.Resim.FrameValidationLeniency"), bResimFrameValidationLeniency, TEXT("Lenient resim frame validation finds a frame to resim from where the particle that is triggering the resim has a valid target. Setting this to false will require all replicated particles to have a valid target and both valid component input and state history."));
+// todo, implement into settings
+enum class EResimFrameValidation : int32
+{
+	FullValidation = 0, // No leniency, validate all dirty particle
+	IslandValidation = 1, // Validate dirty particles inside the islands that have resim trigger particles in them
+	TriggerParticleValidation = 2 // Only validate the resim triggering particle(s)
+};
+CHAOS_API int32 ResimFrameValidation = (int32)EResimFrameValidation::IslandValidation;
+FAutoConsoleVariableRef CVarResimFrameValidationLeniency(TEXT("p.Resim.ResimFrameValidation"), ResimFrameValidation, TEXT("0 = no leniency, all dirty particles need a valid target. 1 = Island leniency, all particles in resim islands need a valid target. 2 = Full leniency, only the particle triggering the resim need a valid target."));
 CHAOS_API bool bResimIncompleteHistory = true;
 FAutoConsoleVariableRef CVarResimIncompleteHistory(TEXT("p.Resim.IncompleteHistory"), bResimIncompleteHistory, TEXT("If a valid resim frame can't be found, use the requested resim frame and perform a resimulation with incomplete data."));
 
@@ -866,35 +873,66 @@ int32 FRewindData::FindValidResimFrame(const int32 RequestedFrame)
 
 	EnsureIsInPhysicsThreadContext();
 
-	// First frame of the history datas
-	const int32 EarliestFrame = FMath::Max(GetEarliestFrame_Internal(), BlockResimFrame + 1);
-	bool bHasTargetHistory = false;
+	auto TargetFinderHelper = [&](FDirtyParticleInfo* DirtyParticleInfo, const FFrameAndPhase FrameAndPhase) -> bool
+	{
+		bool bValid = true;
+		const bool bResimAsFollower = DirtyParticleInfo->bResimAsFollower;
+		FGeometryParticleStateBase& History = DirtyParticleInfo->GetHistory();
+		if (const FParticleDynamicMisc* DynamicMisc = History.DynamicsMisc.Read(FrameAndPhase, PropertiesPool))
+		{
+			if (!DynamicMisc->Disabled() && (DynamicMisc->ObjectState() == EObjectStateType::Dynamic) && !History.TargetPositions.IsEmpty() && !History.TargetVelocities.IsEmpty() && !History.TargetStates.IsEmpty())
+			{
+				if (bResimAsFollower || History.TargetPositions.IsClean(FrameAndPhase) || History.TargetVelocities.IsClean(FrameAndPhase) || History.TargetStates.IsClean(FrameAndPhase))
+				{
+					bValid = false;
+				}
+			}
+		}
 
-	for (ValidFrame = RequestedFrame; ValidFrame > EarliestFrame; ValidFrame--)
+		return bValid;
+	};
+
+	Private::FPBDIslandManager& IslandManager = Solver->GetEvolution()->GetIslandManager();
+
+	// Cache all particles in islands that have a resim triggering particle
+	TArray<const FGeometryParticleHandle*> ResimIslandParticles;
+	if (ResimFrameValidation == (int32)EResimFrameValidation::IslandValidation)
 	{
-		bHasTargetHistory = true;
+		TArray<const Private::FPBDIsland*> ResimIslands;
 		for (FDirtyParticleInfo& DirtyParticleInfo : DirtyParticles)
 		{
-			FGeometryParticleStateBase& History = DirtyParticleInfo.GetHistory();
 			FGeometryParticleHandle* Handle = DirtyParticleInfo.GetObjectPtr();
-			const bool bResimAsFollower = DirtyParticleInfo.bResimAsFollower;
-
-			// If the particle is not marked for resimulation, don't bother checking for valid target states.
-			if (bResimFrameValidationLeniency && Solver->GetEvolution()->GetIslandManager().GetParticleResimFrame(Handle) == INDEX_NONE)
+			if (IslandManager.GetParticleResimFrame(Handle) != INDEX_NONE)
 			{
-				continue;
+				ResimIslands.AddUnique(IslandManager.GetParticleIsland(Handle));
 			}
+		}
+
+		ResimIslandParticles = IslandManager.FindParticlesInIslands(ResimIslands);
+	}
 
+	// First frame of the history datas
+	const int32 EarliestFrame = FMath::Max(GetEarliestFrame_Internal(), BlockResimFrame + 1);
+	bool bHasTargetHistory = false;
+
+	for (ValidFrame = RequestedFrame; ValidFrame > EarliestFrame; ValidFrame--)
+	{
+		const FFrameAndPhase FrameAndPhase{ ValidFrame, FFrameAndPhase::PostPushData };
+		bHasTargetHistory = true;
+		
 #if DEBUG_REWIND_DATA
-			UE_LOG(LogTemp, Log, TEXT("COMMON | PT | FindValidResimFrame | Processing resim particle | History Frame: %d | Total Particle Count: %d"), ValidFrame, DirtyParticles.Num());
+		UE_LOG(LogTemp, Log, TEXT("COMMON | PT | FindValidResimFrame | Processing resim particles | History Frame: %d | Total Particle Count: %d | ResimIslands Particle Count: %d | ResimFrameValidation: %d"), ValidFrame, DirtyParticles.Num(), ResimIslandParticles.Num(), ResimFrameValidation);
 #endif
 
-			const FFrameAndPhase FrameAndPhase{ ValidFrame, FFrameAndPhase::PostPushData };
-			if (const FParticleDynamicMisc* DynamicMisc = History.DynamicsMisc.Read(FrameAndPhase, PropertiesPool))
+		if ((EResimFrameValidation)ResimFrameValidation == EResimFrameValidation::IslandValidation)
+		{
+			// Iterate over islands previously found having resim particles in them and check if the particles in the islands have targets
+			for (const FGeometryParticleHandle* IslandParticle : ResimIslandParticles)
 			{
-				if (!DynamicMisc->Disabled() && (DynamicMisc->ObjectState() == EObjectStateType::Dynamic) && !History.TargetPositions.IsEmpty() && !History.TargetVelocities.IsEmpty() && !History.TargetStates.IsEmpty())
+				// Cache particle handles for objects in islands that need resim
+				if (FDirtyParticleInfo* DirtyParticleInfo = FindDirtyObj(*IslandParticle))
 				{
-					if (bResimAsFollower || History.TargetPositions.IsClean(FrameAndPhase) || History.TargetVelocities.IsClean(FrameAndPhase) || History.TargetStates.IsClean(FrameAndPhase))
+					if (!TargetFinderHelper(DirtyParticleInfo, FrameAndPhase))
 					{
 						bHasTargetHistory = false;
 						break;
@@ -902,29 +940,48 @@ int32 FRewindData::FindValidResimFrame(const int32 RequestedFrame)
 				}
 			}
 		}
-
-		if (!bResimFrameValidationLeniency)
+		else
 		{
-			if (bHasTargetHistory)
+			for (FDirtyParticleInfo& DirtyParticleInfo : DirtyParticles)
 			{
-				for (auto& InputsHistory : InputsHistories)
+				// If running validation leniency, check if the particle is marked for resimulation else don't bother checking for valid target states.
+				if ((EResimFrameValidation)ResimFrameValidation == EResimFrameValidation::TriggerParticleValidation)
 				{
-					if (!InputsHistory.Pin().Get()->HasValidDatas(ValidFrame))
+					FGeometryParticleHandle* Handle = DirtyParticleInfo.GetObjectPtr();
+					if(IslandManager.GetParticleResimFrame(Handle) == INDEX_NONE)
 					{
-						bHasTargetHistory = false;
-						break;
+						continue;
 					}
 				}
+
+				if (!TargetFinderHelper(&DirtyParticleInfo, FrameAndPhase))
+				{
+					bHasTargetHistory = false;
+					break;
+				}
+				break;
 			}
-			if (bHasTargetHistory)
+		}
+		
+		if (bHasTargetHistory)
+		{
+			for (auto& InputsHistory : InputsHistories)
 			{
-				for (auto& StatesHistory : StatesHistories)
+				if (!InputsHistory.Pin().Get()->HasValidDatas(ValidFrame))
 				{
-					if (!StatesHistory.Pin().Get()->HasValidDatas(ValidFrame))
-					{
-						bHasTargetHistory = false;
-						break;
-					}
+					bHasTargetHistory = false;
+					break;
+				}
+			}
+		}
+		if (bHasTargetHistory)
+		{
+			for (auto& StatesHistory : StatesHistories)
+			{
+				if (!StatesHistory.Pin().Get()->HasValidDatas(ValidFrame))
+				{
+					bHasTargetHistory = false;
+					break;
 				}
 			}
 		}
@@ -940,7 +997,7 @@ int32 FRewindData::FindValidResimFrame(const int32 RequestedFrame)
 		ValidFrame = bResimIncompleteHistory ? RequestedFrame : INDEX_NONE;
 
 #if DEBUG_REWIND_DATA
-		UE_LOG(LogTemp, Warning, TEXT("COMMON | PT | FindValidResimFrame | No valid resim frame found | RequestedFrame: %d | ValidFrame: %d | EarliestFrame: %d | HasTargetHistory: %d | EarliestHistoryFrame: %d | CurrentFrame: %d | FramesSaved: %d, ResimFrameValidationLeniency:%d"), RequestedFrame, ValidFrame, EarliestFrame, bHasTargetHistory, GetEarliestFrame_Internal(), CurrentFrame(), FramesSaved, ResimFrameValidationLeniency);
+		UE_LOG(LogTemp, Warning, TEXT("COMMON | PT | FindValidResimFrame | No valid resim frame found | RequestedFrame: %d | ValidFrame: %d | EarliestFrame: %d | HasTargetHistory: %d | EarliestHistoryFrame: %d | CurrentFrame: %d | FramesSaved: %d, ResimFrameValidation:%d"), RequestedFrame, ValidFrame, EarliestFrame, bHasTargetHistory, GetEarliestFrame_Internal(), CurrentFrame(), FramesSaved, ResimFrameValidation);
 #endif
 	}
 
diff --git a/Engine/Source/Runtime/Experimental/Chaos/Public/RewindData.h b/Engine/Source/Runtime/Experimental/Chaos/Public/RewindData.h
--- a/Engine/Source/Runtime/Experimental/Chaos/Public/RewindData.h
+++ b/Engine/Source/Runtime/Experimental/Chaos/Public/RewindData.h
@@ -1571,7 +1571,18 @@ private:
 	};
 
 	template <typename TDirtyObjs, typename TObj>
-	auto& FindOrAddDirtyObjImp(TDirtyObjs& DirtyObjs, TObj& Handle, const int32 InitializedOnFrame = INDEX_NONE)
+	auto* FindDirtyObjImp(TDirtyObjs& DirtyObjs, TObj& Handle)
+	{
+		return DirtyObjs.Find(&Handle);
+	}
+
+	FDirtyParticleInfo* FindDirtyObj(const FGeometryParticleHandle& Handle)
+	{
+		return FindDirtyObjImp(DirtyParticles, Handle);
+	}
+
+	template <typename TDirtyObjs, typename TObj>
+	auto& FindOrAddDirtyObjImp(TDirtyObjs & DirtyObjs, TObj & Handle, const int32 InitializedOnFrame = INDEX_NONE)
 	{
 		if (auto Info = DirtyObjs.Find(&Handle))
 		{
