Change logging a callstack when ForceGC is called to add context

diff --git a/Engine/Source/Runtime/Engine/Private/UnrealEngine.cpp b/Engine/Source/Runtime/Engine/Private/UnrealEngine.cpp
--- a/Engine/Source/Runtime/Engine/Private/UnrealEngine.cpp
+++ b/Engine/Source/Runtime/Engine/Private/UnrealEngine.cpp
@@ -1577,19 +1577,6 @@ static FAutoConsoleVariableRef CVarLowMemoryTimeBetweenPurgingPendingKillObjects
 	ECVF_Default
 );
 
-#ifndef UE_ENABLE_LOG_STACK_ON_FORCE_GC
-#define UE_ENABLE_LOG_STACK_ON_FORCE_GC 0
-#endif
-#if UE_ENABLE_LOG_STACK_ON_FORCE_GC
-static bool GLogStackOnForceGC = false;
-static FAutoConsoleVariableRef CVarLogStackOnForceGC(
-	TEXT("gc.LogStackOnForceGC"),
-	GLogStackOnForceGC,
-	TEXT("Whether to log a stack trace when ForceGarbageCollection is called."),
-	ECVF_Default
-);
-#endif
-
 // see also: s.ForceGCAfterLevelStreamedOut, s.ContinuouslyIncrementalGCWhileLevelsPendingPurge
 static float GLowMemoryTimeBetweenPurgingPendingLevels = 15.0f;
 static FAutoConsoleVariableRef CVarLowMemoryTimeBetweenPurgingPendingLevels(
@@ -1623,6 +1610,10 @@ static FAutoConsoleVariableRef CVarIncrementalGCTimePerFrame(
 	ECVF_Default
 );
 
+#if UE_ENABLE_LOG_STACK_ON_FORCE_GC
+bool GLogStackOnForceGC = false;
+#endif
+
 void UEngine::SendWorldEndOfFrameUpdates()
 {
 	// Gather worlds that need EOF updates
@@ -1727,6 +1718,7 @@ void UEngine::ForceGarbageCollection(bool bForcePurge/*=false*/)
 #if UE_ENABLE_LOG_STACK_ON_FORCE_GC
 	if (GLogStackOnForceGC)
 	{
+		UE_LOG(LogOutputDevice, Log, TEXT("ForceGarbageCollection called, logging stack"));
 		FDebug::DumpStackTraceToLog(ELogVerbosity::Log);
 	}
 #endif
diff --git a/Engine/Source/Runtime/Engine/Public/UnrealEngine.h b/Engine/Source/Runtime/Engine/Public/UnrealEngine.h
--- a/Engine/Source/Runtime/Engine/Public/UnrealEngine.h
+++ b/Engine/Source/Runtime/Engine/Public/UnrealEngine.h
@@ -492,6 +492,13 @@ struct FSystemResolution
 ENGINE_API extern FSystemResolution GSystemResolution;
 ENGINE_API extern int32 GUnbuiltHLODCount;
 
+#ifndef UE_ENABLE_LOG_STACK_ON_FORCE_GC
+#define UE_ENABLE_LOG_STACK_ON_FORCE_GC 0
+#endif
+#if UE_ENABLE_LOG_STACK_ON_FORCE_GC
+ENGINE_API extern bool GLogStackOnForceGC;
+#endif
+
 // Update the debugging aid GPlayInEditorContextString based on the current world context (does nothing in WITH_EDITOR=0 builds)
 ENGINE_API void UpdatePlayInEditorWorldDebugString(const FWorldContext* WorldContext);
 
