[Custom Primitive Data] Fix issue where editing the Custom Primitive Data fields on a BP instance in the world didn't stick. Issue caused by the simple construction script getting rerun after editing, causing the transient data not to be populated in the recreated version.

diff --git a/Engine/Source/Runtime/Engine/Private/Components/PrimitiveComponent.cpp b/Engine/Source/Runtime/Engine/Private/Components/PrimitiveComponent.cpp
--- a/Engine/Source/Runtime/Engine/Private/Components/PrimitiveComponent.cpp
+++ b/Engine/Source/Runtime/Engine/Private/Components/PrimitiveComponent.cpp
@@ -730,6 +730,8 @@ void FPrimitiveComponentInstanceData::ApplyToComponent(UActorComponent* Componen
 
 	if (Component->IsRegistered() && ((VisibilityId != INDEX_NONE) || SavedProperties.Num() > 0))
 	{
+		// This is needed to restore transient primitive data from serialized defaults
+		PrimitiveComponent->ResetCustomPrimitiveData();
 		Component->MarkRenderStateDirty();
 	}
 }
