Fix for World being initialized to Inactive (from the crash logs) and then being switched to Editor
which would mean the FContentBundleContainer would be initialized in "Game mode" because of UseEditorContentBundle() returning false when world type is Inactive and then true when world type would change to Editor

There is a known issue that explains this behavior: UE-196359

#rnx

diff --git a/Engine/Source/Runtime/Engine/Private/WorldPartition/ContentBundle/ContentBundleContainer.cpp b/Engine/Source/Runtime/Engine/Private/WorldPartition/ContentBundle/ContentBundleContainer.cpp
--- a/Engine/Source/Runtime/Engine/Private/WorldPartition/ContentBundle/ContentBundleContainer.cpp
+++ b/Engine/Source/Runtime/Engine/Private/WorldPartition/ContentBundle/ContentBundleContainer.cpp
@@ -158,8 +158,7 @@ TArray<TSharedPtr<FContentBundleEditor>>& FContentBundleContainer::GetEditorCont
 
 bool FContentBundleContainer::UseEditorContentBundle() const
 {
-	bool bIsEditorEditWorld = !GetInjectedWorld()->IsGameWorld() && GetInjectedWorld()->IsEditorWorld();
-	return bIsEditorEditWorld || IsRunningCookCommandlet();
+	return !GetInjectedWorld()->IsGameWorld() || IsRunningCookCommandlet();
 }
 
 #endif
