World Partition Streaming Generation: move streaming generation validation checks out of the streaming generation context.

diff --git a/Engine/Source/Runtime/Engine/Private/WorldPartition/WorldPartitionStreamingGeneration.cpp b/Engine/Source/Runtime/Engine/Private/WorldPartition/WorldPartitionStreamingGeneration.cpp
--- a/Engine/Source/Runtime/Engine/Private/WorldPartition/WorldPartitionStreamingGeneration.cpp
+++ b/Engine/Source/Runtime/Engine/Private/WorldPartition/WorldPartitionStreamingGeneration.cpp
@@ -152,14 +152,17 @@ class FWorldPartitionStreamingGenerator
 					const FWorldPartitionActorDescView& ReferenceActorDescView = ActorSetContainer.ActorDescViewMap->FindByGuidChecked(ActorSet.Actors[0]);
 
 					bool bContainsUnfilteredActors = !FilteredActors;
-					// Validate assumptions
-					for (const FGuid& ActorGuid : ActorSet.Actors)
+
+					if (!bContainsUnfilteredActors)
 					{
-						const FWorldPartitionActorDescView& ActorDescView = ActorSetContainer.ActorDescViewMap->FindByGuidChecked(ActorGuid);
-						check(ActorDescView.GetRuntimeGrid() == ReferenceActorDescView.GetRuntimeGrid());
-						check(ActorDescView.GetIsSpatiallyLoaded() == ReferenceActorDescView.GetIsSpatiallyLoaded());
-						check(ActorDescView.GetContentBundleGuid() == ReferenceActorDescView.GetContentBundleGuid());
-						bContainsUnfilteredActors |= (FilteredActors && !FilteredActors->Contains(ActorGuid));
+						for (const FGuid& ActorGuid : ActorSet.Actors)
+						{
+							if (FilteredActors && !FilteredActors->Contains(ActorGuid))
+							{
+								bContainsUnfilteredActors = true;
+								break;
+							}
+						}
 					}
 
 					// Skip if all actors are filtered out for this container
@@ -555,7 +558,7 @@ class FWorldPartitionStreamingGenerator
 					UpdateContainerDescriptor(ContainerCollectionDescriptor);
 				}
 
-				// Calculate Bounds of Non-container ActorDescViews
+				// Calculate Bounds of non-container actor descriptor views
 				check(!ContainerCollectionInstanceDescriptor.Bounds.IsValid);
 				ContainerCollectionDescriptor.ActorDescViewMap.ForEachActorDescView([&ContainerCollectionInstanceDescriptor](const FWorldPartitionActorDescView& ActorDescView)
 				{
@@ -1079,6 +1082,31 @@ PRAGMA_ENABLE_DEPRECATION_WARNINGS
 		ContainerCollectionDescriptor.Clusters = GenerateObjectsClusters(ActorsWithRefs);
 	}
 
+	/**
+	 * Validate the streaming generator internal state.
+	 */
+	void ValidateInternalState()
+	{
+		// Validate that all actors part of the same actor set share the same set of values
+		for (const auto& [ContainerID, ContainerCollectionInstanceDescriptor] : ContainerCollectionInstanceDescriptorsMap)
+		{
+			const FContainerCollectionDescriptor& ContainerCollectionDescriptor = ContainerCollectionDescriptorsMap.FindChecked(ContainerCollectionInstanceDescriptor.ActorDescCollection->GetMainContainerPackageName());
+					
+			for (const TArray<FGuid>& Cluster : ContainerCollectionDescriptor.Clusters)
+			{
+				const FWorldPartitionActorDescView& ReferenceActorDescView = ContainerCollectionDescriptor.ActorDescViewMap.FindByGuidChecked(Cluster[0]);
+
+				for (const FGuid& ActorGuid : Cluster)
+				{
+					const FWorldPartitionActorDescView& ActorDescView = ContainerCollectionDescriptor.ActorDescViewMap.FindByGuidChecked(ActorGuid);
+					check(ActorDescView.GetRuntimeGrid() == ReferenceActorDescView.GetRuntimeGrid());
+					check(ActorDescView.GetIsSpatiallyLoaded() == ReferenceActorDescView.GetIsSpatiallyLoaded());
+					check(ActorDescView.GetContentBundleGuid() == ReferenceActorDescView.GetContentBundleGuid());
+				}
+			}
+		}
+	}
+
 public:
 	struct FWorldPartitionStreamingGeneratorParams
 	{
@@ -1128,6 +1156,7 @@ public:
 	void PreparationPhase(const FStreamingGenerationActorDescCollection& ActorDescCollection)
 	{
 		CreateActorContainers(ActorDescCollection);
+		ValidateInternalState();
 	}
 
 	static TUniquePtr<FArchive> CreateDumpStateLogArchive(const TCHAR* Suffix, bool bTimeStamped = true)
