Wrapped GPU skin visualization system in a mutex in preparation for parallel GDME.

diff --git a/Engine/Source/Runtime/Engine/Private/GPUSkinCacheVisualizationData.cpp b/Engine/Source/Runtime/Engine/Private/GPUSkinCacheVisualizationData.cpp
--- a/Engine/Source/Runtime/Engine/Private/GPUSkinCacheVisualizationData.cpp
+++ b/Engine/Source/Runtime/Engine/Private/GPUSkinCacheVisualizationData.cpp
@@ -4,6 +4,7 @@
 #include "HAL/IConsoleManager.h"
 #include "RenderUtils.h"
 #include "RHIDefinitions.h"
+#include "Misc/ScopeLock.h"
 
 #define LOCTEXT_NAMESPACE "FGPUSkinCacheVisualizationData"
 
@@ -13,6 +14,7 @@ static FGPUSkinCacheVisualizationData GGPUSkinCacheVisualizationData;
 
 void FGPUSkinCacheVisualizationData::Initialize()
 {
+	UE::TScopeLock Lock(Mutex);
 	if (!bIsInitialized)
 	{
 		// NOTE: The first parameter determines the console command parameter. "none", "off" and "list" are reserved
diff --git a/Engine/Source/Runtime/Engine/Public/GPUSkinCacheVisualizationData.h b/Engine/Source/Runtime/Engine/Public/GPUSkinCacheVisualizationData.h
--- a/Engine/Source/Runtime/Engine/Public/GPUSkinCacheVisualizationData.h
+++ b/Engine/Source/Runtime/Engine/Public/GPUSkinCacheVisualizationData.h
@@ -3,6 +3,7 @@
 #pragma once
 
 #include "CoreMinimal.h"
+#include "Async/Mutex.h"
 
 class FGPUSkinCacheVisualizationData
 {
@@ -29,16 +30,12 @@ public:
 	typedef TMap<FName, FModeRecord> TModeMap;
 
 public:
-	FGPUSkinCacheVisualizationData()
-	: bIsInitialized(false)
-	{
-	}
 
 	/** Initialize the system. */
 	void Initialize();
 
 	/** Check if system was initialized. */
-	inline bool IsInitialized() const { return bIsInitialized; }
+	inline bool IsInitialized() const { return bIsInitialized.load(std::memory_order_relaxed); }
 
 	/** Check if visualization is active. */
 	ENGINE_API bool IsActive() const;
@@ -86,6 +83,8 @@ private:
 	void SetActiveMode(FModeType ModeType, const FName& ModeName);
 
 private:
+	UE::FMutex Mutex;
+
 	/** The name->mode mapping table */
 	TModeMap ModeMap;
 
@@ -96,7 +95,7 @@ private:
 	FString ConsoleDocumentationVisualizationMode;
 
 	/** Flag indicating if system is initialized. **/
-	bool bIsInitialized = false;
+	std::atomic_bool bIsInitialized = { false };
 };
 
 ENGINE_API FGPUSkinCacheVisualizationData& GetGPUSkinCacheVisualizationData();
