Implemented GenerateHLSLExpression for the following material nodes: BlackBody, DistanceCullFade, RuntimeVirtualTextureOutput, SkyAtmosphereLightIlluminance, SkyAtmosphereLightIlluminanceOnGround, SkyAtmosphereDistantLightScatteredLuminance, LandscapeGrassOutput, LandscapeLayerBlend, LandscapeLayerSample, LandscapePhysicalMaterialOutput, and LandscapeVisibilityMask.

diff --git a/Engine/Source/Runtime/Engine/Classes/Materials/MaterialExpressionBlackBody.h b/Engine/Source/Runtime/Engine/Classes/Materials/MaterialExpressionBlackBody.h
--- a/Engine/Source/Runtime/Engine/Classes/Materials/MaterialExpressionBlackBody.h
+++ b/Engine/Source/Runtime/Engine/Classes/Materials/MaterialExpressionBlackBody.h
@@ -22,6 +22,8 @@ class UMaterialExpressionBlackBody : public UMaterialExpression
 #if WITH_EDITOR
 	virtual int32 Compile(class FMaterialCompiler* Compiler, int32 OutputIndex) override;
 	virtual void GetCaption(TArray<FString>& OutCaptions) const override;
+
+	virtual bool GenerateHLSLExpression(FMaterialHLSLGenerator& Generator, UE::HLSLTree::FScope& Scope, int32 OutputIndex, UE::HLSLTree::FExpression const*& OutExpression) const override;
 #endif
 	//~ End UMaterialExpression Interface
 };
diff --git a/Engine/Source/Runtime/Engine/Classes/Materials/MaterialExpressionDistanceCullFade.h b/Engine/Source/Runtime/Engine/Classes/Materials/MaterialExpressionDistanceCullFade.h
--- a/Engine/Source/Runtime/Engine/Classes/Materials/MaterialExpressionDistanceCullFade.h
+++ b/Engine/Source/Runtime/Engine/Classes/Materials/MaterialExpressionDistanceCullFade.h
@@ -18,6 +18,8 @@ class UMaterialExpressionDistanceCullFade : public UMaterialExpression
 #if WITH_EDITOR
 	virtual int32 Compile(class FMaterialCompiler* Compiler, int32 OutputIndex) override;
 	virtual void GetCaption(TArray<FString>& OutCaptions) const override;
+
+	virtual bool GenerateHLSLExpression(FMaterialHLSLGenerator& Generator, UE::HLSLTree::FScope& Scope, int32 OutputIndex, UE::HLSLTree::FExpression const*& OutExpression) const override;
 #endif
 	//~ End UMaterialExpression Interface
 };
diff --git a/Engine/Source/Runtime/Engine/Classes/Materials/MaterialExpressionRuntimeVirtualTextureOutput.h b/Engine/Source/Runtime/Engine/Classes/Materials/MaterialExpressionRuntimeVirtualTextureOutput.h
--- a/Engine/Source/Runtime/Engine/Classes/Materials/MaterialExpressionRuntimeVirtualTextureOutput.h
+++ b/Engine/Source/Runtime/Engine/Classes/Materials/MaterialExpressionRuntimeVirtualTextureOutput.h
@@ -46,6 +46,8 @@ public:
 	//~ Begin UMaterialExpression Interface
 	virtual int32 Compile(class FMaterialCompiler* Compiler, int32 OutputIndex) override;
 	virtual void GetCaption(TArray<FString>& OutCaptions) const override;
+
+	virtual bool GenerateHLSLExpression(FMaterialHLSLGenerator& Generator, UE::HLSLTree::FScope& Scope, int32 OutputIndex, UE::HLSLTree::FExpression const*& OutExpression) const override;
 	//~ End UMaterialExpression Interface
 #endif
 
@@ -53,5 +55,9 @@ public:
 	virtual int32 GetNumOutputs() const override;
 	virtual FString GetFunctionName() const override;
 	virtual FString GetDisplayName() const override;
+
+#if WITH_EDITOR
+	virtual UE::Shader::EValueType GetCustomOutputType(int32 OutputIndex) const override;
+#endif
 	//~ End UMaterialExpressionCustomOutput Interface
 };
diff --git a/Engine/Source/Runtime/Engine/Classes/Materials/MaterialExpressionSkyAtmosphereLightIlluminance.h b/Engine/Source/Runtime/Engine/Classes/Materials/MaterialExpressionSkyAtmosphereLightIlluminance.h
--- a/Engine/Source/Runtime/Engine/Classes/Materials/MaterialExpressionSkyAtmosphereLightIlluminance.h
+++ b/Engine/Source/Runtime/Engine/Classes/Materials/MaterialExpressionSkyAtmosphereLightIlluminance.h
@@ -26,6 +26,8 @@ class UMaterialExpressionSkyAtmosphereLightIlluminance : public UMaterialExpress
 	virtual int32 Compile(class FMaterialCompiler* Compiler, int32 OutputIndex) override;
 	virtual void GetCaption(TArray<FString>& OutCaptions) const override;
 	virtual void GetExpressionToolTip(TArray<FString>& OutToolTip) override;
+
+	virtual bool GenerateHLSLExpression(FMaterialHLSLGenerator& Generator, UE::HLSLTree::FScope& Scope, int32 OutputIndex, UE::HLSLTree::FExpression const*& OutExpression) const override;
 #endif
 	//~ End UMaterialExpression Interface
 };
@@ -46,6 +48,8 @@ class UMaterialExpressionSkyAtmosphereLightIlluminanceOnGround : public UMateria
 	virtual int32 Compile(class FMaterialCompiler* Compiler, int32 OutputIndex) override;
 	virtual void GetCaption(TArray<FString>& OutCaptions) const override;
 	virtual void GetExpressionToolTip(TArray<FString>& OutToolTip) override;
+
+	virtual bool GenerateHLSLExpression(FMaterialHLSLGenerator& Generator, UE::HLSLTree::FScope& Scope, int32 OutputIndex, UE::HLSLTree::FExpression const*& OutExpression) const override;
 #endif
 	//~ End UMaterialExpression Interface
 };
@@ -108,6 +112,8 @@ class UMaterialExpressionSkyAtmosphereDistantLightScatteredLuminance : public UM
 #if WITH_EDITOR
 	virtual int32 Compile(class FMaterialCompiler* Compiler, int32 OutputIndex) override;
 	virtual void GetCaption(TArray<FString>& OutCaptions) const override;
+
+	virtual bool GenerateHLSLExpression(FMaterialHLSLGenerator& Generator, UE::HLSLTree::FScope& Scope, int32 OutputIndex, UE::HLSLTree::FExpression const*& OutExpression) const override;
 #endif
 	//~ End UMaterialExpression Interface
 };
diff --git a/Engine/Source/Runtime/Engine/Private/HLSLTree/HLSLTreeCommon.cpp b/Engine/Source/Runtime/Engine/Private/HLSLTree/HLSLTreeCommon.cpp
--- a/Engine/Source/Runtime/Engine/Private/HLSLTree/HLSLTreeCommon.cpp
+++ b/Engine/Source/Runtime/Engine/Private/HLSLTree/HLSLTreeCommon.cpp
@@ -977,6 +977,11 @@ bool FExpressionAppend::PrepareValue(FEmitContext& Context, FEmitScope& Scope, c
 	}
 
 	const FPreparedType& RhsType = Context.PrepareExpression(Rhs, Scope, RhsRequestedType);
+	if (RhsType.IsVoid())
+	{
+		return false;
+	}
+
 	const Private::FAppendTypes Types = Private::GetAppendTypes(RequestedType, LhsType, RhsType);
 	check(!Types.ResultType.IsNumericScalar()); // Appending 2 values should never result in a scalar
 	if (Types.ResultType.IsVoid())
diff --git a/Engine/Source/Runtime/Engine/Private/HLSLTree/HLSLTreeEmit.cpp b/Engine/Source/Runtime/Engine/Private/HLSLTree/HLSLTreeEmit.cpp
--- a/Engine/Source/Runtime/Engine/Private/HLSLTree/HLSLTreeEmit.cpp
+++ b/Engine/Source/Runtime/Engine/Private/HLSLTree/HLSLTreeEmit.cpp
@@ -1005,6 +1005,100 @@ void WriteMaterialUniformAccess(Shader::EValueComponentType ComponentType, uint3
 		OutResult.Append(TEXT(")"));
 	}
 }
+
+void EmitPreshaderField(
+	FEmitContext& Context,
+	TMemoryImageArray<FMaterialUniformPreshaderHeader>& UniformPreshaders,
+	TMemoryImageArray<FMaterialUniformPreshaderField>& UniformPreshaderFields,
+	Shader::FPreshaderData& UniformPreshaderData,
+	FMaterialUniformPreshaderHeader*& PreshaderHeader,
+	TFunction<void (FEmitValuePreshaderResult&)> EmitPreshaderOpcode,
+	const Shader::FValueTypeDescription& TypeDesc,
+	int32 ComponentIndex,
+	FStringBuilderBase& FormattedCode)
+{
+	// Only need to allocate uniform buffer for non-constant components
+	// Constant components can have their value inlined into the shader directly
+	if (!PreshaderHeader)
+	{
+		// Allocate a preshader header the first time we hit a non-constant field
+		PreshaderHeader = &UniformPreshaders.AddDefaulted_GetRef();
+		PreshaderHeader->FieldIndex = UniformPreshaderFields.Num();
+		PreshaderHeader->NumFields = 0u;
+		PreshaderHeader->OpcodeOffset = UniformPreshaderData.Num();
+		FEmitValuePreshaderResult PreshaderResult(UniformPreshaderData);
+		EmitPreshaderOpcode(PreshaderResult);
+		PreshaderHeader->OpcodeSize = UniformPreshaderData.Num() - PreshaderHeader->OpcodeOffset;
+	}
+
+	FMaterialUniformPreshaderField& PreshaderField = UniformPreshaderFields.AddDefaulted_GetRef();
+	PreshaderField.ComponentIndex = ComponentIndex;
+	PreshaderField.Type = TypeDesc.ValueType;
+	PreshaderHeader->NumFields++;
+
+	const int32 NumFieldComponents = TypeDesc.NumComponents;
+	
+	if (TypeDesc.ComponentType == Shader::EValueComponentType::Bool)
+	{
+		// 'Bool' uniforms are packed into bits
+		if (Context.CurrentNumBoolComponents + NumFieldComponents > 32u)
+		{
+			Context.CurrentBoolUniformOffset = Context.UniformPreshaderOffset++;
+			Context.CurrentNumBoolComponents = 0u;
+		}
+
+		const uint32 RegisterIndex = Context.CurrentBoolUniformOffset / 4;
+		const uint32 RegisterOffset = Context.CurrentBoolUniformOffset % 4;
+		FormattedCode.Appendf(TEXT("UnpackUniform_%s(asuint(Material.PreshaderBuffer[%u][%u]), %u)"),
+			TypeDesc.Name,
+			RegisterIndex,
+			RegisterOffset,
+			Context.CurrentNumBoolComponents);
+
+		PreshaderField.BufferOffset = Context.CurrentBoolUniformOffset * 32u + Context.CurrentNumBoolComponents;
+		Context.CurrentNumBoolComponents += NumFieldComponents;
+	}
+	else if (TypeDesc.ComponentType == Shader::EValueComponentType::Double)
+	{
+		// Double uniforms are split into Tile/Offset components to make FLWCScalar/FLWCVectors
+		PreshaderField.BufferOffset = Context.UniformPreshaderOffset;
+
+		if (NumFieldComponents > 1)
+		{
+			FormattedCode.Appendf(TEXT("MakeLWCVector%d("), NumFieldComponents);
+		}
+		else
+		{
+			FormattedCode.Append(TEXT("MakeLWCScalar("));
+		}
+
+		// Write the tile uniform
+		Private::WriteMaterialUniformAccess(Shader::EValueComponentType::Float, NumFieldComponents, Context.UniformPreshaderOffset, FormattedCode);
+		Context.UniformPreshaderOffset += NumFieldComponents;
+		FormattedCode.Append(TEXT(", "));
+
+		// Write the offset uniform
+		Private::WriteMaterialUniformAccess(Shader::EValueComponentType::Float, NumFieldComponents, Context.UniformPreshaderOffset, FormattedCode);
+		Context.UniformPreshaderOffset += NumFieldComponents;
+		FormattedCode.Append(TEXT(")"));
+	}
+	else
+	{
+		// Float/Int uniforms are written directly to the uniform buffer
+		const uint32 RegisterOffset = Context.UniformPreshaderOffset % 4;
+		if (RegisterOffset + NumFieldComponents > 4u)
+		{
+			// If this uniform would span multiple registers, align offset to the next register to avoid this
+			// TODO - we could keep track of this empty padding space, and pack other smaller uniform types here
+			Context.UniformPreshaderOffset = Align(Context.UniformPreshaderOffset, 4u);
+		}
+
+		PreshaderField.BufferOffset = Context.UniformPreshaderOffset;
+		Private::WriteMaterialUniformAccess(TypeDesc.ComponentType, NumFieldComponents, Context.UniformPreshaderOffset, FormattedCode);
+		Context.UniformPreshaderOffset += NumFieldComponents;
+	}
+}
+
 } // namespace Private
 
 FEmitShaderExpression* FEmitContext::EmitPreshaderOrConstant(FEmitScope& Scope, const FRequestedType& RequestedType, const Shader::FType& ResultType, const FExpression* Expression)
@@ -1060,7 +1154,7 @@ FEmitShaderExpression* FEmitContext::EmitPreshaderOrConstant(FEmitScope& Scope,
 		}
 
 		const Shader::EValueType FieldType = ResultType.GetFlatFieldType(FieldIndex);
-		const Shader::FValueTypeDescription TypeDesc = Shader::GetValueTypeDescription(FieldType);
+		const Shader::FValueTypeDescription& TypeDesc = Shader::GetValueTypeDescription(FieldType);
 		const int32 NumFieldComponents = TypeDesc.NumComponents;
 
 		// If this is a struct, use GetFieldEvaluation()
@@ -1071,85 +1165,21 @@ FEmitShaderExpression* FEmitContext::EmitPreshaderOrConstant(FEmitScope& Scope,
 
 		if (FieldEvaluation == EExpressionEvaluation::Preshader)
 		{
-			// Only need to allocate uniform buffer for non-constant components
-			// Constant components can have their value inlined into the shader directly
 			FUniformExpressionSet& UniformExpressionSet = MaterialCompilationOutput->UniformExpressionSet;
-			if (!PreshaderHeader)
-			{
-				// Allocate a preshader header the first time we hit a non-constant field
-				PreshaderHeader = &UniformExpressionSet.UniformPreshaders.AddDefaulted_GetRef();
-				PreshaderHeader->FieldIndex = UniformExpressionSet.UniformPreshaderFields.Num();
-				PreshaderHeader->NumFields = 0u;
-				PreshaderHeader->OpcodeOffset = UniformExpressionSet.UniformPreshaderData.Num();
-				FEmitValuePreshaderResult PreshaderResult(UniformExpressionSet.UniformPreshaderData);
-				Expression->EmitValuePreshader(*this, Scope, RequestedType, PreshaderResult);
-				PreshaderHeader->OpcodeSize = UniformExpressionSet.UniformPreshaderData.Num() - PreshaderHeader->OpcodeOffset;
-			}
-
-			FMaterialUniformPreshaderField& PreshaderField = UniformExpressionSet.UniformPreshaderFields.AddDefaulted_GetRef();
-			PreshaderField.ComponentIndex = ComponentIndex;
-			PreshaderField.Type = FieldType;
-			PreshaderHeader->NumFields++;
-
-			if (TypeDesc.ComponentType == Shader::EValueComponentType::Bool)
-			{
-				// 'Bool' uniforms are packed into bits
-				if (CurrentNumBoolComponents + NumFieldComponents > 32u)
-				{
-					CurrentBoolUniformOffset = UniformPreshaderOffset++;
-					CurrentNumBoolComponents = 0u;
-				}
-
-				const uint32 RegisterIndex = CurrentBoolUniformOffset / 4;
-				const uint32 RegisterOffset = CurrentBoolUniformOffset % 4;
-				FormattedCode.Appendf(TEXT("UnpackUniform_%s(asuint(Material.PreshaderBuffer[%u][%u]), %u)"),
-					TypeDesc.Name,
-					RegisterIndex,
-					RegisterOffset,
-					CurrentNumBoolComponents);
-
-				PreshaderField.BufferOffset = CurrentBoolUniformOffset * 32u + CurrentNumBoolComponents;
-				CurrentNumBoolComponents += NumFieldComponents;
-			}
-			else if (TypeDesc.ComponentType == Shader::EValueComponentType::Double)
-			{
-				// Double uniforms are split into Tile/Offset components to make FLWCScalar/FLWCVectors
-				PreshaderField.BufferOffset = UniformPreshaderOffset;
-
-				if (NumFieldComponents > 1)
-				{
-					FormattedCode.Appendf(TEXT("MakeLWCVector%d("), NumFieldComponents);
-				}
-				else
-				{
-					FormattedCode.Append(TEXT("MakeLWCScalar("));
-				}
-
-				// Write the tile uniform
-				Private::WriteMaterialUniformAccess(Shader::EValueComponentType::Float, NumFieldComponents, UniformPreshaderOffset, FormattedCode);
-				UniformPreshaderOffset += NumFieldComponents;
-				FormattedCode.Append(TEXT(", "));
 
-				// Write the offset uniform
-				Private::WriteMaterialUniformAccess(Shader::EValueComponentType::Float, NumFieldComponents, UniformPreshaderOffset, FormattedCode);
-				UniformPreshaderOffset += NumFieldComponents;
-				FormattedCode.Append(TEXT(")"));
-			}
-			else
-			{
-				// Float/Int uniforms are written directly to the uniform buffer
-				const uint32 RegisterOffset = UniformPreshaderOffset % 4;
-				if (RegisterOffset + NumFieldComponents > 4u)
+			Private::EmitPreshaderField(
+				*this,
+				UniformExpressionSet.UniformPreshaders,
+				UniformExpressionSet.UniformPreshaderFields,
+				UniformExpressionSet.UniformPreshaderData,
+				PreshaderHeader,
+				[Expression, &Context = *this, &Scope, &RequestedType](FEmitValuePreshaderResult& OutResult)
 				{
-					// If this uniform would span multiple registers, align offset to the next register to avoid this
-					// TODO - we could keep track of this empty padding space, and pack other smaller uniform types here
-					UniformPreshaderOffset = Align(UniformPreshaderOffset, 4u);
-				}
-
-				PreshaderField.BufferOffset = UniformPreshaderOffset;
-				Private::WriteMaterialUniformAccess(TypeDesc.ComponentType, NumFieldComponents, UniformPreshaderOffset, FormattedCode);
-				UniformPreshaderOffset += NumFieldComponents;
-			}
+					Expression->EmitValuePreshader(Context, Scope, RequestedType, OutResult);
+				},
+				TypeDesc,
+				ComponentIndex,
+				FormattedCode);
 		}
 		else
 		{
diff --git a/Engine/Source/Runtime/Engine/Private/Materials/MaterialExpressionHLSL.cpp b/Engine/Source/Runtime/Engine/Private/Materials/MaterialExpressionHLSL.cpp
--- a/Engine/Source/Runtime/Engine/Private/Materials/MaterialExpressionHLSL.cpp
+++ b/Engine/Source/Runtime/Engine/Private/Materials/MaterialExpressionHLSL.cpp
@@ -27,6 +27,7 @@
 #include "Materials/MaterialExpressionAtmosphericLightColor.h"
 #include "Materials/MaterialExpressionAtmosphericLightVector.h"
 #include "Materials/MaterialExpressionBinaryOp.h"
+#include "Materials/MaterialExpressionBlackBody.h"
 #include "Materials/MaterialExpressionBlendMaterialAttributes.h"
 #include "Materials/MaterialExpressionBreakMaterialAttributes.h"
 #include "Materials/MaterialExpressionBumpOffset.h"
@@ -60,6 +61,7 @@
 #include "Materials/MaterialExpressionDeriveNormalZ.h"
 #include "Materials/MaterialExpressionDesaturation.h"
 #include "Materials/MaterialExpressionDistance.h"
+#include "Materials/MaterialExpressionDistanceCullFade.h"
 #include "Materials/MaterialExpressionDistanceFieldApproxAO.h"
 #include "Materials/MaterialExpressionDistanceFieldGradient.h"
 #include "Materials/MaterialExpressionDistanceFieldsRenderingSwitch.h"
@@ -151,6 +153,7 @@
 #include "Materials/MaterialExpressionRotateAboutAxis.h"
 #include "Materials/MaterialExpressionRotator.h"
 #include "Materials/MaterialExpressionRound.h"
+#include "Materials/MaterialExpressionRuntimeVirtualTextureOutput.h"
 #include "Materials/MaterialExpressionRuntimeVirtualTextureSample.h"
 #include "Materials/MaterialExpressionRuntimeVirtualTextureSampleParameter.h"
 #include "Materials/MaterialExpressionSaturate.h"
@@ -1729,6 +1732,87 @@ bool UMaterialExpressionRuntimeVirtualTextureSample::GenerateHLSLExpression(FMat
 	return true;
 }
 
+UE::Shader::EValueType UMaterialExpressionRuntimeVirtualTextureOutput::GetCustomOutputType(int32 OutputIndex) const
+{
+	using namespace UE::Shader;
+
+	switch (OutputIndex)
+	{
+	case 0:
+	case 3:
+		return EValueType::Float3;
+	case 1:
+	case 2:
+	case 4:
+	case 5:
+	case 6:
+		return EValueType::Float1;
+	default:
+		checkNoEntry();
+		return EValueType::Void;
+	}
+}
+
+bool UMaterialExpressionRuntimeVirtualTextureOutput::GenerateHLSLExpression(FMaterialHLSLGenerator& Generator, UE::HLSLTree::FScope& Scope, int32 OutputIndex, UE::HLSLTree::FExpression const*& OutExpression) const
+{
+	using namespace UE::HLSLTree;
+
+	auto MakeOutputExpression = [&Generator, &Scope](const FExpressionInput& Input, const UE::Shader::FValue& DefaultValue, ERuntimeVirtualTextureAttributeType AttributeType)
+	{
+		if (!Input.GetTracedInput().Expression)
+		{
+			return Generator.NewConstant(DefaultValue);
+		}
+		
+		const FExpression* AttributeExpression = Input.AcquireHLSLExpression(Generator, Scope);
+		if (AttributeType == ERuntimeVirtualTextureAttributeType::Count)
+		{
+			return AttributeExpression;
+		}
+
+		const uint8 OutputAttributeMask = 1 << (uint8)AttributeType;
+		return Generator.GetTree().NewExpression<Material::FExpressionRuntimeVirtualTextureOutput>(OutputAttributeMask, AttributeExpression);
+	};
+
+	// Order of outputs generates function names GetVirtualTextureOutput{index}
+	// These must match the function names called in VirtualTextureMaterial.usf
+	if (OutputIndex == 0)
+	{
+		OutExpression = MakeOutputExpression(BaseColor, FVector3f::ZeroVector, ERuntimeVirtualTextureAttributeType::BaseColor);
+	}
+	else if (OutputIndex == 1)
+	{
+		OutExpression = MakeOutputExpression(Specular, 0.5f, ERuntimeVirtualTextureAttributeType::Specular);
+	}
+	else if (OutputIndex == 2)
+	{
+		OutExpression = MakeOutputExpression(Roughness, 0.5f, ERuntimeVirtualTextureAttributeType::Roughness);
+	}
+	else if (OutputIndex == 3)
+	{
+		OutExpression = MakeOutputExpression(Normal, FVector3f(0.f, 0.f, 1.f), ERuntimeVirtualTextureAttributeType::Normal);
+	}
+	else if (OutputIndex == 4)
+	{
+		OutExpression = MakeOutputExpression(WorldHeight, 0.f, ERuntimeVirtualTextureAttributeType::WorldHeight);
+	}
+	else if (OutputIndex == 5)
+	{
+		OutExpression = MakeOutputExpression(Opacity, 1.f, ERuntimeVirtualTextureAttributeType::Count);
+	}
+	else if (OutputIndex == 6)
+	{
+		OutExpression = MakeOutputExpression(Mask, 1.f, ERuntimeVirtualTextureAttributeType::Mask);
+	}
+	else
+	{
+		checkNoEntry();
+		return false;
+	}
+
+	return true;
+}
+
 bool UMaterialExpressionFontSample::GenerateHLSLExpression(FMaterialHLSLGenerator& Generator, UE::HLSLTree::FScope& Scope, int32 OutputIndex, UE::HLSLTree::FExpression const*& OutExpression) const
 {
 	using namespace UE::HLSLTree;
@@ -3448,6 +3532,21 @@ bool UMaterialExpressionSkyAtmosphereAerialPerspective::GenerateHLSLExpression(F
 	return true;
 }
 
+bool UMaterialExpressionSkyAtmosphereLightIlluminance::GenerateHLSLExpression(FMaterialHLSLGenerator& Generator, UE::HLSLTree::FScope& Scope, int32 OutputIndex, UE::HLSLTree::FExpression const*& OutExpression) const
+{
+	using namespace UE::HLSLTree;
+	const FExpression* WorldPositionExpression = WorldPosition.AcquireHLSLExpressionOrExternalInput(Generator, Scope, Material::EExternalInput::WorldPosition);
+	OutExpression = Generator.GetTree().NewExpression<Material::FExpressionSkyAtmosphereLightIlluminance>(WorldPositionExpression, LightIndex);
+	return true;
+}
+
+bool UMaterialExpressionSkyAtmosphereLightIlluminanceOnGround::GenerateHLSLExpression(FMaterialHLSLGenerator& Generator, UE::HLSLTree::FScope& Scope, int32 OutputIndex, UE::HLSLTree::FExpression const*& OutExpression) const
+{
+	using namespace UE::HLSLTree;
+	OutExpression = Generator.GetTree().NewExpression<Material::FExpressionSkyAtmosphereLightIlluminanceOnGround>(LightIndex);
+	return true;
+}
+
 bool UMaterialExpressionSkyAtmosphereViewLuminance::GenerateHLSLExpression(FMaterialHLSLGenerator& Generator, UE::HLSLTree::FScope& Scope, int32 OutputIndex, UE::HLSLTree::FExpression const*& OutExpression) const
 {
 	using namespace UE::HLSLTree;
@@ -3455,6 +3554,13 @@ bool UMaterialExpressionSkyAtmosphereViewLuminance::GenerateHLSLExpression(FMate
 	return true;
 }
 
+bool UMaterialExpressionSkyAtmosphereDistantLightScatteredLuminance::GenerateHLSLExpression(FMaterialHLSLGenerator& Generator, UE::HLSLTree::FScope& Scope, int32 OutputIndex, UE::HLSLTree::FExpression const*& OutExpression) const
+{
+	using namespace UE::HLSLTree;
+	OutExpression = Generator.NewExternalInput(Material::EExternalInput::SkyAtmosphereDistantLightScatteredLuminance);
+	return true;
+}
+
 bool UMaterialExpressionAtmosphericLightVector::GenerateHLSLExpression(FMaterialHLSLGenerator& Generator, UE::HLSLTree::FScope& Scope, int32 OutputIndex, UE::HLSLTree::FExpression const*& OutExpression) const
 {
 	OutExpression = Generator.GetTree().NewExpression<UE::HLSLTree::FExpressionInlineCustomHLSL>(UE::Shader::EValueType::Float3, TEXT("MaterialExpressionAtmosphericLightVector(Parameters)"));
@@ -4217,4 +4323,33 @@ bool UMaterialExpressionSingleLayerWaterMaterialOutput::GenerateHLSLExpression(F
 	return OutExpression != nullptr;
 }
 
+bool UMaterialExpressionBlackBody::GenerateHLSLExpression(FMaterialHLSLGenerator& Generator, UE::HLSLTree::FScope& Scope, int32 OutputIndex, UE::HLSLTree::FExpression const*& OutExpression) const
+{
+	using namespace UE::HLSLTree;
+
+	const FExpression* TempExpression = Temp.AcquireHLSLExpression(Generator, Scope);
+	if (!TempExpression)
+	{
+		return false;
+	}
+
+	OutExpression = Generator.GetTree().NewExpression<Material::FExpressionBlackBody>(TempExpression);
+	return true;
+}
+
+bool UMaterialExpressionDistanceCullFade::GenerateHLSLExpression(FMaterialHLSLGenerator& Generator, UE::HLSLTree::FScope& Scope, int32 OutputIndex, UE::HLSLTree::FExpression const*& OutExpression) const
+{
+	using namespace UE::HLSLTree;
+	OutExpression = Generator.NewExternalInput(Material::EExternalInput::DistanceCullFade);
+	return true;
+}
+
+bool GenerateStaticTerrainLayerWeightExpression(FName LayerName, float PreviewWeight, FMaterialHLSLGenerator& Generator, const UE::HLSLTree::FExpression*& OutExpression)
+{
+	using namespace UE::HLSLTree;
+	const FExpression* TexCoordExpression = Generator.NewExternalInput(Material::EExternalInput::TexCoord3);
+	OutExpression = Generator.GetTree().NewExpression<Material::FExpressionStaticTerrainLayerWeight>(Generator.GetParameterInfo(LayerName), TexCoordExpression, PreviewWeight);
+	return true;
+}
+
 #endif // WITH_EDITOR
diff --git a/Engine/Source/Runtime/Engine/Private/Materials/MaterialHLSLTree.cpp b/Engine/Source/Runtime/Engine/Private/Materials/MaterialHLSLTree.cpp
--- a/Engine/Source/Runtime/Engine/Private/Materials/MaterialHLSLTree.cpp
+++ b/Engine/Source/Runtime/Engine/Private/Materials/MaterialHLSLTree.cpp
@@ -2,6 +2,7 @@
 #if WITH_EDITOR
 
 #include "MaterialHLSLTree.h"
+#include "Engine/Engine.h"
 #include "Engine/Texture.h"
 #include "HLSLTree/HLSLTreeEmit.h"
 #include "MaterialDomain.h"
@@ -160,6 +161,9 @@ FExternalInputDescription GetExternalInputDescription(EExternalInput Input)
 	case EExternalInput::PerInstanceRandom: return FExternalInputDescription(TEXT("PerInstanceRandom"), Shader::EValueType::Float1);
 
 	case EExternalInput::SkyAtmosphereViewLuminance: return FExternalInputDescription(TEXT("SkyAtmosphereViewLuminance"), Shader::EValueType::Float3);
+	case EExternalInput::SkyAtmosphereDistantLightScatteredLuminance: return FExternalInputDescription(TEXT("SkyAtmosphereDistanceLightScatteredLuminance"), Shader::EValueType::Float3);
+
+	case EExternalInput::DistanceCullFade: return FExternalInputDescription(TEXT("DistanceCullFade"), Shader::EValueType::Float1);
 
 	case EExternalInput::IsOrthographic: return FExternalInputDescription(TEXT("IsOrthographic"), Shader::EValueType::Float1);
 
@@ -237,7 +241,12 @@ bool FExpressionExternalInput::PrepareValue(FEmitContext& Context, FEmitScope& S
 				Context.MaterialCompilationOutput->bUsesPerInstanceRandom = true;
 				break;
 			case EExternalInput::SkyAtmosphereViewLuminance:
+			case EExternalInput::SkyAtmosphereDistantLightScatteredLuminance:
 				Context.bUsesSkyAtmosphere = true;
+				break;
+			case EExternalInput::DistanceCullFade:
+				Context.MaterialCompilationOutput->bUsesDistanceCullFade = true;
+				break;
 			default:
 				break;
 			}
@@ -379,6 +388,9 @@ void FExpressionExternalInput::EmitValueShader(FEmitContext& Context, FEmitScope
 		case EExternalInput::PerInstanceRandom: Code = TEXT("GetPerInstanceRandom(Parameters)"); break;
 
 		case EExternalInput::SkyAtmosphereViewLuminance: Code = TEXT("MaterialExpressionSkyAtmosphereViewLuminance(Parameters)"); break;
+		case EExternalInput::SkyAtmosphereDistantLightScatteredLuminance: Code = TEXT("MaterialExpressionSkyAtmosphereDistantLightScatteredLuminance(Parameters)"); break;
+
+		case EExternalInput::DistanceCullFade: Code = TEXT("GetDistanceCullFade()"); break;
 
 		case EExternalInput::IsOrthographic: Code = TEXT("((View.ViewToClip[3][3] < 1.0f) ? 0.0f : 1.0f)"); break;
 
@@ -481,6 +493,31 @@ void EmitTextureShader(FEmitContext& Context,
 	OutCode.Appendf(TEXT("Material.%s_%d"), TextureTypeName, TextureParameterIndex);
 }
 
+void EmitNumericParameterPreshader(
+	FEmitContext& Context,
+	FEmitData& EmitData,
+	const FMaterialParameterInfo& ParameterInfo,
+	EMaterialParameterType ParameterType,
+	Shader::FValue DefaultValue,
+	FEmitValuePreshaderResult& OutResult)
+{
+	check(IsNumericMaterialParameter(ParameterType));
+	const uint32* PrevDefaultOffset = EmitData.DefaultUniformValues.Find(DefaultValue);
+	uint32 DefaultOffset;
+	if (PrevDefaultOffset)
+	{
+		DefaultOffset = *PrevDefaultOffset;
+	}
+	else
+	{
+		DefaultOffset = Context.MaterialCompilationOutput->UniformExpressionSet.AddDefaultParameterValue(DefaultValue);
+		EmitData.DefaultUniformValues.Add(DefaultValue, DefaultOffset);
+	}
+	const int32 ParameterIndex = Context.MaterialCompilationOutput->UniformExpressionSet.FindOrAddNumericParameter(ParameterType, ParameterInfo, DefaultOffset);
+	check(ParameterIndex >= 0 && ParameterIndex <= 0xffff);
+	OutResult.Preshader.WriteOpcode(Shader::EPreshaderOpcode::Parameter).Write((uint16)ParameterIndex);
+}
+
 } // namespace Private
 
 void FExpressionParameter::ComputeAnalyticDerivatives(FTree& Tree, FExpressionDerivatives& OutResult) const
@@ -538,7 +575,7 @@ void FExpressionParameter::EmitValueShader(FEmitContext& Context, FEmitScope& Sc
 	{
 		TStringBuilder<64> FormattedTexture;
 		Private::EmitTextureShader(Context, TextureValue, FormattedTexture);
-		// Emit a texture/sampler pair1
+		// Emit a texture/sampler pair
 		OutResult.Code = Context.EmitInlineExpression(Scope, FMaterialTextureValue::GetTypeName(), TEXT("%,%Sampler"), FormattedTexture.ToString(), FormattedTexture.ToString());
 	}
 }
@@ -587,21 +624,7 @@ void FExpressionParameter::EmitValuePreshader(FEmitContext& Context, FEmitScope&
 	}
 	else
 	{
-		check(IsNumericMaterialParameter(ParameterType));
-		const uint32* PrevDefaultOffset = EmitMaterialData.DefaultUniformValues.Find(DefaultValue);
-		uint32 DefaultOffset;
-		if (PrevDefaultOffset)
-		{
-			DefaultOffset = *PrevDefaultOffset;
-		}
-		else
-		{
-			DefaultOffset = Context.MaterialCompilationOutput->UniformExpressionSet.AddDefaultParameterValue(DefaultValue);
-			EmitMaterialData.DefaultUniformValues.Add(DefaultValue, DefaultOffset);
-		}
-		const int32 ParameterIndex = Context.MaterialCompilationOutput->UniformExpressionSet.FindOrAddNumericParameter(ParameterType, ParameterInfo, DefaultOffset);
-		check(ParameterIndex >= 0 && ParameterIndex <= 0xffff);
-		OutResult.Preshader.WriteOpcode(Shader::EPreshaderOpcode::Parameter).Write((uint16)ParameterIndex);
+		Private::EmitNumericParameterPreshader(Context, EmitMaterialData, ParameterInfo, ParameterType, DefaultValue, OutResult);
 	}
 }
 
@@ -977,12 +1000,22 @@ uint32 AcquireVTStackIndex(
 	return StackIndex;
 }
 
-} // namespace Private
-
-void FExpressionTextureSample::EmitValueShader(FEmitContext& Context, FEmitScope& Scope, const FRequestedType& RequestedType, FEmitValueShaderResult& OutResult) const
+FEmitShaderExpression* EmitTextureSampleShader(
+	FEmitContext& Context,
+	FEmitScope& Scope,
+	const FMaterialTextureValue& TextureValue,
+	FEmitShaderExpression* EmitTexCoordValue,
+	FEmitShaderExpression* EmitMipValue,
+	FEmitShaderExpression* EmitTexCoordValueDdx,
+	FEmitShaderExpression* EmitTexCoordValueDdy,
+	ESamplerSourceMode SamplerSource,
+	ETextureMipValueMode MipValueMode,
+	int16 TextureLayerIndex,
+	int16 PageTableLayerIndex,
+	bool bAdaptive,
+	bool bEnableFeedback,
+	bool bAutomaticViewMipBias)
 {
-	FMaterialTextureValue TextureValue;
-	verify(TextureExpression->GetValueObject(Context, Scope, TextureValue));
 	UTexture* Texture = TextureValue.Texture;
 	URuntimeVirtualTexture* RuntimeVirtualTexture = TextureValue.RuntimeVirtualTexture;
 
@@ -1041,26 +1074,20 @@ void FExpressionTextureSample::EmitValueShader(FEmitContext& Context, FEmitScope
 		}
 	}
 
-	FEmitShaderExpression* EmitTexCoordValue = TexCoordExpression->GetValueShader(Context, Scope, TexCoordType);
-	FEmitShaderExpression* EmitMipValue = nullptr;
-	FEmitShaderExpression* EmitTexCoordValueDdx = nullptr;
-	FEmitShaderExpression* EmitTexCoordValueDdy = nullptr;
-	const ETextureMipValueMode LocalMipValueMode = Private::GetMipValueMode(Context, this);
-	if (LocalMipValueMode == TMVM_Derivative)
+	if (MipValueMode == TMVM_Derivative)
 	{
-		EmitTexCoordValueDdx = TexCoordDerivatives.ExpressionDdx->GetValueShader(Context, Scope, TexCoordType);
-		EmitTexCoordValueDdy = TexCoordDerivatives.ExpressionDdy->GetValueShader(Context, Scope, TexCoordType);
+		check(EmitTexCoordValueDdx && EmitTexCoordValueDdy);
+		EmitMipValue = nullptr;
 	}
-	else if (LocalMipValueMode == TMVM_MipLevel || LocalMipValueMode == TMVM_MipBias)
+	else if (MipValueMode == TMVM_MipLevel || MipValueMode == TMVM_MipBias)
 	{
-		if (MipValueExpression)
-		{
-			EmitMipValue = MipValueExpression->GetValueShader(Context, Scope, Shader::EValueType::Float1);
-		}
-		else
+		if (!EmitMipValue)
 		{
 			EmitMipValue = Context.EmitConstantZero(Scope, Shader::EValueType::Float1);
 		}
+
+		EmitTexCoordValueDdx = nullptr;
+		EmitTexCoordValueDdy = nullptr;
 	}
 
 	FEmitShaderExpression* EmitTextureResult = nullptr;
@@ -1086,7 +1113,7 @@ void FExpressionTextureSample::EmitValueShader(FEmitContext& Context, FEmitScope
 			// The layer index in the virtual texture stack is already known
 			// Create a page table sample for each new combination of virtual texture and sample parameters
 			VTStackIndex = Private::AcquireVTStackIndex(Context, Scope, EmitMaterialData,
-				LocalMipValueMode,
+				MipValueMode,
 				StaticAddressX,
 				StaticAddressY,
 				1.0f,
@@ -1110,7 +1137,7 @@ void FExpressionTextureSample::EmitValueShader(FEmitContext& Context, FEmitScope
 
 			// Create a page table sample for each new set of sample parameters
 			VTStackIndex = Private::AcquireVTStackIndex(Context, Scope, EmitMaterialData,
-				LocalMipValueMode,
+				MipValueMode,
 				StaticAddressX,
 				StaticAddressY,
 				TextureAspectRatio,
@@ -1135,7 +1162,7 @@ void FExpressionTextureSample::EmitValueShader(FEmitContext& Context, FEmitScope
 		{
 			// VT doesn't care if the shared sampler is wrap or clamp. It only cares if it is aniso or not.
 			// The wrap/clamp/mirror operation is handled in the shader explicitly.
-			const bool bUseAnisoSampler = VirtualTextureScalability::IsAnisotropicFilteringEnabled() && LocalMipValueMode != TMVM_MipLevel;
+			const bool bUseAnisoSampler = VirtualTextureScalability::IsAnisotropicFilteringEnabled() && MipValueMode != TMVM_MipLevel;
 			const TCHAR* SharedSamplerName = bUseAnisoSampler ? TEXT("View.SharedBilinearAnisoClampedSampler") : TEXT("View.SharedBilinearClampedSampler");
 			FormattedSampler.Appendf(TEXT("GetMaterialSharedSampler(Material.VirtualTexturePhysical_%dSampler, %s)"), TextureParameterIndex, SharedSamplerName);
 		}
@@ -1144,7 +1171,7 @@ void FExpressionTextureSample::EmitValueShader(FEmitContext& Context, FEmitScope
 			FormattedSampler.Appendf(TEXT("Material.VirtualTexturePhysical_%dSampler"), TextureParameterIndex);
 		}
 
-		const TCHAR* SampleFunctionName = (LocalMipValueMode == TMVM_MipLevel) ? TEXT("TextureVirtualSampleLevel") : TEXT("TextureVirtualSample");
+		const TCHAR* SampleFunctionName = (MipValueMode == TMVM_MipLevel) ? TEXT("TextureVirtualSampleLevel") : TEXT("TextureVirtualSample");
 		EmitTextureResult = Context.EmitExpression(Scope, Shader::EValueType::Float4, TEXT("%(%, %, %, %, VTUniform_Unpack(Material.VTPackedUniform[%]))"),
 			SampleFunctionName,
 			FormattedTexture.ToString(),
@@ -1185,7 +1212,6 @@ void FExpressionTextureSample::EmitValueShader(FEmitContext& Context, FEmitScope
 		TStringBuilder<64> FormattedTexture;
 		Private::EmitTextureShader(Context, TextureValue, FormattedTexture);
 
-		const bool AutomaticViewMipBias = AutomaticMipBiasExpression ? AutomaticMipBiasExpression->GetValueConstant(Context, Scope, Shader::EValueType::Bool1).AsBoolScalar() : false;
 		TStringBuilder<256> FormattedSampler;
 		switch (SamplerSource)
 		{
@@ -1195,12 +1221,12 @@ void FExpressionTextureSample::EmitValueShader(FEmitContext& Context, FEmitScope
 		case SSM_Wrap_WorldGroupSettings:
 			FormattedSampler.Appendf(TEXT("GetMaterialSharedSampler(%sSampler,%s)"),
 				FormattedTexture.ToString(),
-				AutomaticViewMipBias ? TEXT("View.MaterialTextureBilinearWrapedSampler") : TEXT("Material.Wrap_WorldGroupSettings"));
+				bAutomaticViewMipBias ? TEXT("View.MaterialTextureBilinearWrapedSampler") : TEXT("Material.Wrap_WorldGroupSettings"));
 			break;
 		case SSM_Clamp_WorldGroupSettings:
 			FormattedSampler.Appendf(TEXT("GetMaterialSharedSampler(%sSampler,%s)"),
 				FormattedTexture.ToString(),
-				AutomaticViewMipBias ? TEXT("View.MaterialTextureBilinearClampedSampler") : TEXT("Material.Clamp_WorldGroupSettings"));
+				bAutomaticViewMipBias ? TEXT("View.MaterialTextureBilinearClampedSampler") : TEXT("Material.Clamp_WorldGroupSettings"));
 			break;
 		case SSM_TerrainWeightmapGroupSettings:
 			FormattedSampler.Appendf(TEXT("GetMaterialSharedSampler(%sSampler,%s)"),
@@ -1212,7 +1238,7 @@ void FExpressionTextureSample::EmitValueShader(FEmitContext& Context, FEmitScope
 			break;
 		}
 
-		switch (LocalMipValueMode)
+		switch (MipValueMode)
 		{
 		case TMVM_Derivative:
 			EmitTextureResult = Context.EmitExpression(Scope, Shader::EValueType::Float4, TEXT("%Grad(%, %, %, %, %)"),
@@ -1253,7 +1279,236 @@ void FExpressionTextureSample::EmitValueShader(FEmitContext& Context, FEmitScope
 	}
 
 	check(EmitTextureResult);
-	OutResult.Code = Context.EmitExpression(Scope, Shader::EValueType::Float4, TEXT("ApplyMaterialSamplerType(%, %)"), EmitTextureResult, TextureValue.SamplerType);
+	return Context.EmitExpression(Scope, Shader::EValueType::Float4, TEXT("ApplyMaterialSamplerType(%, %)"), EmitTextureResult, TextureValue.SamplerType);
+}
+
+} // namespace Private
+
+void FExpressionTextureSample::EmitValueShader(FEmitContext& Context, FEmitScope& Scope, const FRequestedType& RequestedType, FEmitValueShaderResult& OutResult) const
+{
+	FMaterialTextureValue TextureValue;
+	verify(TextureExpression->GetValueObject(Context, Scope, TextureValue));
+
+	EMaterialValueType TextureType = MCT_Unknown;
+	if (TextureValue.Texture)
+	{
+		TextureType = TextureValue.Texture->GetMaterialType();
+	}
+	else if (TextureValue.RuntimeVirtualTexture)
+	{
+		TextureType = MCT_TextureVirtual;
+	}
+
+	const Shader::EValueType TexCoordType = Private::GetTexCoordType(TextureType);
+	FEmitShaderExpression* EmitTexCoordValue = TexCoordExpression->GetValueShader(Context, Scope, TexCoordType);
+	FEmitShaderExpression* EmitMipValue = nullptr;
+	FEmitShaderExpression* EmitTexCoordValueDdx = nullptr;
+	FEmitShaderExpression* EmitTexCoordValueDdy = nullptr;
+	const ETextureMipValueMode LocalMipValueMode = Private::GetMipValueMode(Context, this);
+	if (LocalMipValueMode == TMVM_Derivative)
+	{
+		EmitTexCoordValueDdx = TexCoordDerivatives.ExpressionDdx->GetValueShader(Context, Scope, TexCoordType);
+		EmitTexCoordValueDdy = TexCoordDerivatives.ExpressionDdy->GetValueShader(Context, Scope, TexCoordType);
+	}
+	else if (LocalMipValueMode == TMVM_MipLevel || LocalMipValueMode == TMVM_MipBias)
+	{
+		if (MipValueExpression)
+		{
+			EmitMipValue = MipValueExpression->GetValueShader(Context, Scope, Shader::EValueType::Float1);
+		}
+		else
+		{
+			EmitMipValue = Context.EmitConstantZero(Scope, Shader::EValueType::Float1);
+		}
+	}
+
+	const bool bAutomaticViewMipBias = AutomaticMipBiasExpression ?
+		AutomaticMipBiasExpression->GetValueConstant(Context, Scope, Shader::EValueType::Bool1).AsBoolScalar() : false;
+
+	OutResult.Code = Private::EmitTextureSampleShader(
+		Context,
+		Scope,
+		TextureValue,
+		EmitTexCoordValue,
+		EmitMipValue,
+		EmitTexCoordValueDdx,
+		EmitTexCoordValueDdy,
+		SamplerSource,
+		LocalMipValueMode,
+		TextureLayerIndex,
+		PageTableLayerIndex,
+		bAdaptive,
+		bEnableFeedback,
+		bAutomaticViewMipBias);
+}
+
+bool FExpressionStaticTerrainLayerWeight::PrepareValue(FEmitContext& Context, FEmitScope& Scope, const FRequestedType& RequestedType, FPrepareValueResult& OutResult) const
+{
+	const EMaterialValueType TextureMaterialType = GEngine->WeightMapPlaceholderTexture->GetMaterialType();
+	const FRequestedType TexCoordRequestedType = Private::GetTexCoordType(TextureMaterialType);
+	const FPreparedType& TexCoordPreparedType = Context.PrepareExpression(TexCoordExpression, Scope, TexCoordRequestedType);
+	if (TexCoordPreparedType.IsVoid())
+	{
+		return false;
+	}
+	
+	bool bFoundMatchingParameter = false;
+	FEmitData& EmitData = Context.FindData<FEmitData>();
+
+	if (Context.bMarkLiveValues && EmitData.CachedExpressionData && EmitData.StaticParameters)
+	{
+		for (int32 ParameterIndex = 0; ParameterIndex < EmitData.StaticParameters->EditorOnly.TerrainLayerWeightParameters.Num(); ++ParameterIndex)
+		{
+			const FStaticTerrainLayerWeightParameter& Parameter = EmitData.StaticParameters->EditorOnly.TerrainLayerWeightParameters[ParameterIndex];
+
+			// If there are multiple weight maps with the same name, they should be numbered to allow for unique masks
+			FName LayerNameTest = Parameter.LayerName;
+			LayerNameTest.SetNumber(0);
+
+			if (LayerNameTest != BaseParameterInfo.Name)
+			{
+				continue;
+			}
+
+			const int32 WeightmapIndex = Parameter.WeightmapIndex;
+
+			if (WeightmapIndex == INDEX_NONE)
+			{
+				continue;
+			}
+
+			bFoundMatchingParameter = true;
+
+			FMaterialParameterInfo WeightmapParameterInfo = BaseParameterInfo;
+			WeightmapParameterInfo.Name = *FString::Printf(TEXT("Weightmap%d"), WeightmapIndex);
+
+			FMaterialParameterMetadata WeightmapParameterMeta;
+			WeightmapParameterMeta.Value = GEngine->WeightMapPlaceholderTexture;
+
+			UObject* UnusedReferencedTexture;
+			EmitData.CachedExpressionData->AddParameter(WeightmapParameterInfo, WeightmapParameterMeta, UnusedReferencedTexture);
+
+			FMaterialParameterInfo LayerMaskParameterInfo = BaseParameterInfo;
+			LayerMaskParameterInfo.Name = *FString::Printf(TEXT("LayerMask_%s"), *Parameter.LayerName.ToString());
+
+			FMaterialParameterMetadata LayerMaskParameterMeta;
+			LayerMaskParameterMeta.Value = FLinearColor(1.f, 0.f, 0.f, 0.f);
+
+			EmitData.CachedExpressionData->AddParameter(LayerMaskParameterInfo, LayerMaskParameterMeta, UnusedReferencedTexture);
+		}
+
+		if (bFoundMatchingParameter)
+		{
+			EmitData.CachedExpressionData->ReferencedTextures.AddUnique(GEngine->WeightMapPlaceholderTexture);
+		}
+	}
+
+	return OutResult.SetType(Context, RequestedType, EExpressionEvaluation::Shader, Shader::EValueType::Float4);
+}
+
+void FExpressionStaticTerrainLayerWeight::EmitValueShader(FEmitContext& Context, FEmitScope& Scope, const FRequestedType& RequestedType, FEmitValueShaderResult& OutResult) const
+{
+	FEmitData& EmitData = Context.FindData<FEmitData>();
+
+	if (!EmitData.StaticParameters)
+	{
+		OutResult.Code = Context.EmitConstantZero(Scope, Shader::EValueType::Float4);
+		return;
+	}
+
+	const EMaterialValueType TextureMaterialType = GEngine->WeightMapPlaceholderTexture->GetMaterialType();
+	const Shader::EValueType TexCoordType = Private::GetTexCoordType(TextureMaterialType);
+	FEmitShaderExpression* EmitTexCoordValue = nullptr;
+	FEmitShaderExpression* EmitResult = nullptr;
+	int32 NumWeightmapParameters = 0;
+
+	for (int32 ParameterIndex = 0; ParameterIndex < EmitData.StaticParameters->EditorOnly.TerrainLayerWeightParameters.Num(); ++ParameterIndex)
+	{
+		const FStaticTerrainLayerWeightParameter& Parameter = EmitData.StaticParameters->EditorOnly.TerrainLayerWeightParameters[ParameterIndex];
+
+		// If there are multiple weight maps with the same name, they should be numbered to allow for unique masks
+		FName LayerNameTest = Parameter.LayerName;
+		LayerNameTest.SetNumber(0);
+
+		if (LayerNameTest != BaseParameterInfo.Name)
+		{
+			continue;
+		}
+
+		const int32 WeightmapIndex = Parameter.WeightmapIndex;
+
+		if (WeightmapIndex == INDEX_NONE)
+		{
+			continue;
+		}
+
+		FMaterialTextureValue TextureValue;
+		TextureValue.Texture = GEngine->WeightMapPlaceholderTexture;
+		TextureValue.SamplerType = SAMPLERTYPE_Masks;
+		TextureValue.ParameterInfo = BaseParameterInfo;
+		TextureValue.ParameterInfo.Name = *FString::Printf(TEXT("Weightmap%d"), WeightmapIndex);
+
+		if (!EmitTexCoordValue)
+		{
+			EmitTexCoordValue = TexCoordExpression->GetValueShader(Context, Scope, TexCoordType);
+		}
+
+		FEmitShaderExpression* EmitWeightmapSampleValue = Private::EmitTextureSampleShader(
+			Context, Scope,
+			TextureValue, EmitTexCoordValue, nullptr, nullptr, nullptr,
+			SSM_TerrainWeightmapGroupSettings, TMVM_MipLevel, INDEX_NONE, INDEX_NONE, false, false, false);
+
+		FMaterialParameterInfo LayerMaskParameterInfo = BaseParameterInfo;
+		LayerMaskParameterInfo.Name = *FString::Printf(TEXT("LayerMask_%s"), *Parameter.LayerName.ToString());
+
+		FUniformExpressionSet& UniformExpressionSet = Context.MaterialCompilationOutput->UniformExpressionSet;
+		FMaterialUniformPreshaderHeader* PreshaderHeader = nullptr;
+		TStringBuilder<64> FormattedLayerMaskCode;
+
+		UE::HLSLTree::Private::EmitPreshaderField(
+			Context,
+			UniformExpressionSet.UniformPreshaders,
+			UniformExpressionSet.UniformPreshaderFields,
+			UniformExpressionSet.UniformPreshaderData,
+			PreshaderHeader,
+			[&Context, &EmitData, &LayerMaskParameterInfo](FEmitValuePreshaderResult& OutResult)
+			{
+				Private::EmitNumericParameterPreshader(
+					Context,
+					EmitData,
+					LayerMaskParameterInfo,
+					EMaterialParameterType::Vector,
+					FLinearColor(1.f, 0.f, 0.f, 0.f),
+					OutResult);
+			},
+			Shader::GetValueTypeDescription(Shader::EValueType::Float4),
+			0,
+			FormattedLayerMaskCode);
+
+		FEmitShaderExpression* EmitLayerMaskValue = Context.InternalEmitExpression(
+			Scope, TArrayView<FEmitShaderNode*>(), true, Shader::EValueType::Float4, FormattedLayerMaskCode.ToView());
+
+		FEmitShaderExpression* EmitWeightValue = Context.EmitExpression(Scope, Shader::EValueType::Float1, TEXT("VectorSum(% * %)"), EmitWeightmapSampleValue, EmitLayerMaskValue);
+
+		if (!EmitResult)
+		{
+			EmitResult = EmitWeightValue;
+		}
+		else
+		{
+			EmitResult = Context.EmitExpression(Scope, Shader::EValueType::Float1, TEXT("(% + %)"), EmitResult, EmitWeightValue);
+		}
+		++NumWeightmapParameters;
+	}
+
+	if (NumWeightmapParameters > 0)
+	{
+		OutResult.Code = NumWeightmapParameters > 1 ? Context.EmitExpression(Scope, Shader::EValueType::Float1, TEXT("saturate(%)"), EmitResult) : EmitResult;
+	}
+	else
+	{
+		OutResult.Code = Context.EmitInlineExpression(Scope, Shader::EValueType::Float1, TEXT("((float)%)"), Context.Material && Context.Material->IsPreview() ? DefaultWeight : 0.f);
+	}
 }
 
 bool FExpressionTextureProperty::PrepareValue(FEmitContext& Context, FEmitScope& Scope, const FRequestedType& RequestedType, FPrepareValueResult& OutResult) const
@@ -1327,6 +1582,17 @@ void FExpressionRuntimeVirtualTextureUniform::EmitValuePreshader(FEmitContext& C
 	OutResult.Preshader.WriteOpcode(Shader::EPreshaderOpcode::RuntimeVirtualTextureUniform).Write(ParameterInfo).Write(TextureIndex).Write(VectorIndex);
 }
 
+bool FExpressionRuntimeVirtualTextureOutput::PrepareValue(FEmitContext& Context, FEmitScope& Scope, const FRequestedType& RequestedType, FPrepareValueResult& OutResult) const
+{
+	if (Context.bMarkLiveValues && Context.MaterialCompilationOutput)
+	{
+		Context.MaterialCompilationOutput->bHasRuntimeVirtualTextureOutputNode |= OutputAttributeMask != 0;
+		Context.MaterialCompilationOutput->RuntimeVirtualTextureOutputAttributeMask |= OutputAttributeMask;
+	}
+
+	return FExpressionForward::PrepareValue(Context, Scope, RequestedType, OutResult);
+}
+
 bool FExpressionVirtualTextureUnpack::PrepareValue(FEmitContext& Context, FEmitScope& Scope, const FRequestedType& RequestedType, FPrepareValueResult& OutResult) const
 {
 	if (UnpackType == EVirtualTextureUnpackType::BaseColorYCoCg)
@@ -1901,6 +2167,12 @@ void FExpressionSkyAtmosphereLightDirection::EmitValueShader(FEmitContext& Conte
 
 bool FExpressionSkyAtmosphereLightDiskLuminance::PrepareValue(FEmitContext& Context, FEmitScope& Scope, const FRequestedType& RequestedType, FPrepareValueResult& OutResult) const
 {
+	const FPreparedType& CosHalfDiskRadiusPreparedType = Context.PrepareExpression(CosHalfDiskRadiusExpression, Scope, Shader::EValueType::Float1);
+	if (CosHalfDiskRadiusPreparedType.IsVoid())
+	{
+		return false;
+	}
+
 	if (Context.bMarkLiveValues)
 	{
 		Context.bUsesSkyAtmosphere = true;
@@ -1916,6 +2188,12 @@ void FExpressionSkyAtmosphereLightDiskLuminance::EmitValueShader(FEmitContext& C
 
 bool FExpressionSkyAtmosphereAerialPerspective::PrepareValue(FEmitContext& Context, FEmitScope& Scope, const FRequestedType& RequestedType, FPrepareValueResult& OutResult) const
 {
+	const FPreparedType& WorldPositionPreparedType = Context.PrepareExpression(WorldPositionExpression, Scope, Shader::EValueType::Double3);
+	if (WorldPositionPreparedType.IsVoid())
+	{
+		return false;
+	}
+
 	if (Context.bMarkLiveValues)
 	{
 		Context.bUsesSkyAtmosphere = true;
@@ -1929,6 +2207,41 @@ void FExpressionSkyAtmosphereAerialPerspective::EmitValueShader(FEmitContext& Co
 	OutResult.Code = Context.EmitInlineExpression(Scope, Shader::EValueType::Float4, TEXT("MaterialExpressionSkyAtmosphereAerialPerspective(Parameters, %)"), EmitWorldPosition);
 }
 
+bool FExpressionSkyAtmosphereLightIlluminance::PrepareValue(FEmitContext& Context, FEmitScope& Scope, const FRequestedType& RequestedType, FPrepareValueResult& OutResult) const
+{
+	const FPreparedType& WorldPositionPreparedType = Context.PrepareExpression(WorldPositionExpression, Scope, Shader::EValueType::Double3);
+	if (WorldPositionPreparedType.IsVoid())
+	{
+		return false;
+	}
+
+	if (Context.bMarkLiveValues)
+	{
+		Context.bUsesSkyAtmosphere = true;
+	}
+	return OutResult.SetType(Context, RequestedType, EExpressionEvaluation::Shader, Shader::EValueType::Float3);
+}
+
+void FExpressionSkyAtmosphereLightIlluminance::EmitValueShader(FEmitContext& Context, FEmitScope& Scope, const FRequestedType& RequestedType, FEmitValueShaderResult& OutResult) const
+{
+	FEmitShaderExpression* EmitWorldPosition = WorldPositionExpression->GetValueShader(Context, Scope, Shader::EValueType::Double3);
+	OutResult.Code = Context.EmitInlineExpression(Scope, Shader::EValueType::Float3, TEXT("MaterialExpressionSkyAtmosphereLightIlluminance(Parameters, %, %)"), EmitWorldPosition, LightIndex);
+}
+
+bool FExpressionSkyAtmosphereLightIlluminanceOnGround::PrepareValue(FEmitContext& Context, FEmitScope& Scope, const FRequestedType& RequestedType, FPrepareValueResult& OutResult) const
+{
+	if (Context.bMarkLiveValues)
+	{
+		Context.bUsesSkyAtmosphere = true;
+	}
+	return OutResult.SetType(Context, RequestedType, EExpressionEvaluation::Shader, Shader::EValueType::Float3);
+}
+
+void FExpressionSkyAtmosphereLightIlluminanceOnGround::EmitValueShader(FEmitContext& Context, FEmitScope& Scope, const FRequestedType& RequestedType, FEmitValueShaderResult& OutResult) const
+{
+	OutResult.Code = Context.EmitInlineExpression(Scope, Shader::EValueType::Float3, TEXT("MaterialExpressionSkyAtmosphereLightIlluminanceOnGround(Parameters, %)"), LightIndex);
+}
+
 namespace Private
 {
 	bool PlatformSupportDistanceFields(FEmitContext& Context)
@@ -2111,6 +2424,24 @@ void FExpressionHairColor::EmitValueShader(FEmitContext& Context, FEmitScope& Sc
 		EmitDyeColor);
 }
 
+bool FExpressionBlackBody::PrepareValue(FEmitContext& Context, FEmitScope& Scope, const FRequestedType& RequestedType, FPrepareValueResult& OutResult) const
+{
+	const FPreparedType& TempPreparedType = Context.PrepareExpression(TempExpression, Scope, Shader::EValueType::Float1);
+	if (TempPreparedType.IsVoid())
+	{
+		return false;
+	}
+
+	return OutResult.SetType(Context, RequestedType, EExpressionEvaluation::Shader, Shader::EValueType::Float3);
+}
+
+void FExpressionBlackBody::EmitValueShader(FEmitContext& Context, FEmitScope& Scope, const FRequestedType& RequestedType, FEmitValueShaderResult& OutResult) const
+{
+	FEmitShaderExpression* EmitTemp = TempExpression->GetValueShader(Context, Scope, Shader::EValueType::Float1);
+
+	OutResult.Code = Context.EmitInlineExpression(Scope, Shader::EValueType::Float3, TEXT("MaterialExpressionBlackBody(%)"), EmitTemp);
+}
+
 bool FExpressionLightVector::PrepareValue(FEmitContext& Context, FEmitScope& Scope, const FRequestedType& RequestedType, FPrepareValueResult& OutResult) const
 {
 	if (Context.ShaderFrequency != SF_Pixel && Context.ShaderFrequency != SF_Compute)
diff --git a/Engine/Source/Runtime/Engine/Private/Materials/MaterialInstanceConstant.cpp b/Engine/Source/Runtime/Engine/Private/Materials/MaterialInstanceConstant.cpp
--- a/Engine/Source/Runtime/Engine/Private/Materials/MaterialInstanceConstant.cpp
+++ b/Engine/Source/Runtime/Engine/Private/Materials/MaterialInstanceConstant.cpp
@@ -248,6 +248,8 @@ void UMaterialInstanceConstant::UpdateCachedData()
 			FStaticParameterSet OverriddenStaticParameters;
 #if WITH_EDITORONLY_DATA
 			{
+				OverriddenStaticParameters.EditorOnly.TerrainLayerWeightParameters = LocalStaticParameters.EditorOnly.TerrainLayerWeightParameters;
+
 				FMaterialInheritanceChain InstanceChain;
 				GetMaterialInheritanceChain(InstanceChain);
 
diff --git a/Engine/Source/Runtime/Engine/Public/HLSLTree/HLSLTreeEmit.h b/Engine/Source/Runtime/Engine/Public/HLSLTree/HLSLTreeEmit.h
--- a/Engine/Source/Runtime/Engine/Public/HLSLTree/HLSLTreeEmit.h
+++ b/Engine/Source/Runtime/Engine/Public/HLSLTree/HLSLTreeEmit.h
@@ -15,6 +15,7 @@
 #include "RHIFeatureLevel.h"
 #include "HLSLTree/HLSLTreeTypes.h"
 #include "HLSLTree/HLSLTreeHash.h"
+#include "MaterialShared.h"
 
 class FMaterial;
 class FMaterialCompilationOutput;
@@ -42,6 +43,7 @@ class FEmitShaderExpression;
 class FEmitShaderStatement;
 
 struct FEmitPreshaderScope;
+struct FEmitValuePreshaderResult;
 
 struct FEmitShaderScopeEntry
 {
@@ -637,6 +639,20 @@ struct FEmitOwnerScope
 	const FOwnedNode* Node;
 };
 
+namespace Private
+{
+void EmitPreshaderField(
+	FEmitContext& Context,
+	TMemoryImageArray<FMaterialUniformPreshaderHeader>& UniformPreshaders,
+	TMemoryImageArray<FMaterialUniformPreshaderField>& UniformPreshaderFields,
+	Shader::FPreshaderData& UniformPreshaderData,
+	FMaterialUniformPreshaderHeader*& PreshaderHeader,
+	TFunction<void (FEmitValuePreshaderResult&)> EmitPreshaderOpcode,
+	const Shader::FValueTypeDescription& TypeDesc,
+	int32 ComponentIndex,
+	FStringBuilderBase& FormattedCode);
+}
+
 } // namespace UE::HLSLTree
 
 #endif // WITH_EDITOR
diff --git a/Engine/Source/Runtime/Engine/Public/MaterialHLSLGenerator.h b/Engine/Source/Runtime/Engine/Public/MaterialHLSLGenerator.h
--- a/Engine/Source/Runtime/Engine/Public/MaterialHLSLGenerator.h
+++ b/Engine/Source/Runtime/Engine/Public/MaterialHLSLGenerator.h
@@ -332,4 +332,10 @@ private:
 	TMap<FXxHash64, int32> CustomExpressionOutputStructIdMap;
 };
 
+ENGINE_API bool GenerateStaticTerrainLayerWeightExpression(
+	FName LayerName,
+	float PreviewWeight,
+	FMaterialHLSLGenerator& Generator,
+	const UE::HLSLTree::FExpression*& OutExpression);
+
 #endif // WITH_EDITOR
diff --git a/Engine/Source/Runtime/Engine/Public/MaterialHLSLTree.h b/Engine/Source/Runtime/Engine/Public/MaterialHLSLTree.h
--- a/Engine/Source/Runtime/Engine/Public/MaterialHLSLTree.h
+++ b/Engine/Source/Runtime/Engine/Public/MaterialHLSLTree.h
@@ -199,6 +199,9 @@ enum class EExternalInput : uint8
 	PerInstanceRandom,
 
 	SkyAtmosphereViewLuminance,
+	SkyAtmosphereDistantLightScatteredLuminance,
+
+	DistanceCullFade,
 
 	IsOrthographic,
 
@@ -453,6 +456,25 @@ public:
 	virtual void EmitValueShader(FEmitContext& Context, FEmitScope& Scope, const FRequestedType& RequestedType, FEmitValueShaderResult& OutResult) const override;
 };
 
+class FExpressionStaticTerrainLayerWeight : public FExpression
+{
+public:
+	FExpressionStaticTerrainLayerWeight(const FMaterialParameterInfo& InBaseParameterInfo, const FExpression* InTexCoordExpression, float InDefaultWeight)
+		: BaseParameterInfo(InBaseParameterInfo)
+		, TexCoordExpression(InTexCoordExpression)
+		, DefaultWeight(InDefaultWeight)
+	{
+		check(!BaseParameterInfo.Name.IsNone());
+	}
+
+	FMaterialParameterInfo BaseParameterInfo;
+	const FExpression* TexCoordExpression;
+	float DefaultWeight;
+
+	virtual bool PrepareValue(FEmitContext& Context, FEmitScope& Scope, const FRequestedType& RequestedType, FPrepareValueResult& OutResult) const override;
+	virtual void EmitValueShader(FEmitContext& Context, FEmitScope& Scope, const FRequestedType& RequestedType, FEmitValueShaderResult& OutResult) const override;
+};
+
 class FExpressionTextureProperty : public FExpression
 {
 public:
@@ -485,6 +507,19 @@ public:
 	virtual void EmitValuePreshader(FEmitContext& Context, FEmitScope& Scope, const FRequestedType& RequestedType, FEmitValuePreshaderResult& OutResult) const override;
 };
 
+class FExpressionRuntimeVirtualTextureOutput : public FExpressionForward
+{
+public:
+	FExpressionRuntimeVirtualTextureOutput(uint8 InOutputAttributeMask, const FExpression* OutputExpression)
+		: FExpressionForward(OutputExpression)
+		, OutputAttributeMask(InOutputAttributeMask)
+	{}
+
+	uint8 OutputAttributeMask;
+
+	virtual bool PrepareValue(FEmitContext& Context, FEmitScope& Scope, const FRequestedType& RequestedType, FPrepareValueResult& OutResult) const override;
+};
+
 class FExpressionVirtualTextureUnpack : public FExpression
 {
 public:
@@ -730,6 +765,34 @@ public:
 	virtual void EmitValueShader(FEmitContext& Context, FEmitScope& Scope, const FRequestedType& RequestedType, FEmitValueShaderResult& OutResult) const override;
 };
 
+class FExpressionSkyAtmosphereLightIlluminance : public FExpression
+{
+public:
+	const FExpression* WorldPositionExpression;
+	int32 LightIndex;
+
+	FExpressionSkyAtmosphereLightIlluminance(const FExpression* InWorldPositionExpression, int32 InLightIndex)
+		: WorldPositionExpression(InWorldPositionExpression)
+		, LightIndex(InLightIndex)
+	{}
+
+	virtual bool PrepareValue(FEmitContext& Context, FEmitScope& Scope, const FRequestedType& RequestedType, FPrepareValueResult& OutResult) const override;
+	virtual void EmitValueShader(FEmitContext& Context, FEmitScope& Scope, const FRequestedType& RequestedType, FEmitValueShaderResult& OutResult) const override;
+};
+
+class FExpressionSkyAtmosphereLightIlluminanceOnGround : public FExpression
+{
+public:
+	int32 LightIndex;
+
+	FExpressionSkyAtmosphereLightIlluminanceOnGround(int32 InLightIndex)
+		: LightIndex(InLightIndex)
+	{}
+
+	virtual bool PrepareValue(FEmitContext& Context, FEmitScope& Scope, const FRequestedType& RequestedType, FPrepareValueResult& OutResult) const override;
+	virtual void EmitValueShader(FEmitContext& Context, FEmitScope& Scope, const FRequestedType& RequestedType, FEmitValueShaderResult& OutResult) const override;
+};
+
 class FExpressionDistanceToNearestSurface : public FExpression
 {
 public:
@@ -859,6 +922,19 @@ public:
 	virtual bool PrepareValue(FEmitContext& Context, FEmitScope& Scope, const FRequestedType& RequestedType, FPrepareValueResult& OutResult) const override;
 };
 
+class FExpressionBlackBody : public FExpression
+{
+public:
+	FExpressionBlackBody(const FExpression* InTempExpression)
+		: TempExpression(InTempExpression)
+	{}
+
+	const FExpression* TempExpression;
+
+	virtual bool PrepareValue(FEmitContext& Context, FEmitScope& Scope, const FRequestedType& RequestedType, FPrepareValueResult& OutResult) const override;
+	virtual void EmitValueShader(FEmitContext& Context, FEmitScope& Scope, const FRequestedType& RequestedType, FEmitValueShaderResult& OutResult) const override;
+};
+
 struct FVertexInterpolator
 {
 	FVertexInterpolator() = default;
diff --git a/Engine/Source/Runtime/Engine/Public/MaterialShared.h b/Engine/Source/Runtime/Engine/Public/MaterialShared.h
--- a/Engine/Source/Runtime/Engine/Public/MaterialShared.h
+++ b/Engine/Source/Runtime/Engine/Public/MaterialShared.h
@@ -99,6 +99,11 @@ namespace UE
 namespace HLSLTree
 {
 class FEmitContext;
+
+namespace Material
+{
+class FExpressionStaticTerrainLayerWeight;
+}
 }
 }
 
@@ -660,6 +665,7 @@ protected:
 	friend class FMaterialVirtualTextureStack;
 	friend class FDebugUniformExpressionSet;
 	friend class UE::HLSLTree::FEmitContext;
+	friend class UE::HLSLTree::Material::FExpressionStaticTerrainLayerWeight;
 
 	LAYOUT_FIELD(TMemoryImageArray<FMaterialUniformPreshaderHeader>, UniformPreshaders);
 	LAYOUT_FIELD(TMemoryImageArray<FMaterialUniformPreshaderField>, UniformPreshaderFields);
diff --git a/Engine/Source/Runtime/Landscape/Classes/Materials/MaterialExpressionLandscapeGrassOutput.h b/Engine/Source/Runtime/Landscape/Classes/Materials/MaterialExpressionLandscapeGrassOutput.h
--- a/Engine/Source/Runtime/Landscape/Classes/Materials/MaterialExpressionLandscapeGrassOutput.h
+++ b/Engine/Source/Runtime/Landscape/Classes/Materials/MaterialExpressionLandscapeGrassOutput.h
@@ -52,6 +52,10 @@ class UMaterialExpressionLandscapeGrassOutput : public UMaterialExpressionCustom
 	LANDSCAPE_API virtual FExpressionInput* GetInput(int32 InputIndex) override;
 	LANDSCAPE_API virtual FName GetInputName(int32 InputIndex) const override;
 	LANDSCAPE_API void ValidateInputName(FGrassInput& Input) const;
+
+	virtual bool GenerateHLSLExpression(FMaterialHLSLGenerator& Generator, UE::HLSLTree::FScope& Scope, int32 OutputIndex, UE::HLSLTree::FExpression const*& OutExpression) const override;
+
+	virtual UE::Shader::EValueType GetCustomOutputType(int32 OutputIndex) const override;
 #endif
 
 	//~ Begin UObject Interface
diff --git a/Engine/Source/Runtime/Landscape/Classes/Materials/MaterialExpressionLandscapeLayerBlend.h b/Engine/Source/Runtime/Landscape/Classes/Materials/MaterialExpressionLandscapeLayerBlend.h
--- a/Engine/Source/Runtime/Landscape/Classes/Materials/MaterialExpressionLandscapeLayerBlend.h
+++ b/Engine/Source/Runtime/Landscape/Classes/Materials/MaterialExpressionLandscapeLayerBlend.h
@@ -84,6 +84,8 @@ class UMaterialExpressionLandscapeLayerBlend : public UMaterialExpression
 	LANDSCAPE_API virtual FExpressionInput* GetInput(int32 InputIndex) override;
 	LANDSCAPE_API virtual FName GetInputName(int32 InputIndex) const override;
 
+	virtual bool GenerateHLSLExpression(FMaterialHLSLGenerator& Generator, UE::HLSLTree::FScope& Scope, int32 OutputIndex, UE::HLSLTree::FExpression const*& OutExpression) const override;
+
 	/**
 	 * Gets the landscape layer names
 	 */
diff --git a/Engine/Source/Runtime/Landscape/Classes/Materials/MaterialExpressionLandscapeLayerSample.h b/Engine/Source/Runtime/Landscape/Classes/Materials/MaterialExpressionLandscapeLayerSample.h
--- a/Engine/Source/Runtime/Landscape/Classes/Materials/MaterialExpressionLandscapeLayerSample.h
+++ b/Engine/Source/Runtime/Landscape/Classes/Materials/MaterialExpressionLandscapeLayerSample.h
@@ -34,6 +34,8 @@ public:
 	LANDSCAPE_API virtual void GetCaption(TArray<FString>& OutCaptions) const override;
 	LANDSCAPE_API virtual bool MatchesSearchQuery(const TCHAR* SearchQuery) override;
 
+	virtual bool GenerateHLSLExpression(FMaterialHLSLGenerator& Generator, UE::HLSLTree::FScope& Scope, int32 OutputIndex, UE::HLSLTree::FExpression const*& OutExpression) const override;
+
 	/**
 	 * Gets the landscape layer names
 	 */
diff --git a/Engine/Source/Runtime/Landscape/Classes/Materials/MaterialExpressionLandscapePhysicalMaterialOutput.h b/Engine/Source/Runtime/Landscape/Classes/Materials/MaterialExpressionLandscapePhysicalMaterialOutput.h
--- a/Engine/Source/Runtime/Landscape/Classes/Materials/MaterialExpressionLandscapePhysicalMaterialOutput.h
+++ b/Engine/Source/Runtime/Landscape/Classes/Materials/MaterialExpressionLandscapePhysicalMaterialOutput.h
@@ -50,11 +50,17 @@ class UMaterialExpressionLandscapePhysicalMaterialOutput : public UMaterialExpre
 	LANDSCAPE_API virtual FName GetInputName(int32 InputIndex) const override;
 	virtual uint32 GetInputType(int32 InputIndex) override { return MCT_Float; }
 	LANDSCAPE_API virtual int32 Compile(class FMaterialCompiler* Compiler, int32 OutputIndex) override;
+
+	virtual bool GenerateHLSLExpression(FMaterialHLSLGenerator& Generator, UE::HLSLTree::FScope& Scope, int32 OutputIndex, UE::HLSLTree::FExpression const*& OutExpression) const override;
 	//~ End UMaterialExpression Interface
 #endif
 
 	//~ Begin UMaterialExpressionCustomOutput Interface
 	virtual int32 GetNumOutputs() const override { return Inputs.Num(); }
 	virtual FString GetFunctionName() const override { return TEXT("GetPhysicalMaterial"); }
+
+#if WITH_EDITOR
+	virtual UE::Shader::EValueType GetCustomOutputType(int32 OutputIndex) const override;
+#endif
 	//~ End UMaterialExpressionCustomOutput Interface
 };
diff --git a/Engine/Source/Runtime/Landscape/Classes/Materials/MaterialExpressionLandscapeVisibilityMask.h b/Engine/Source/Runtime/Landscape/Classes/Materials/MaterialExpressionLandscapeVisibilityMask.h
--- a/Engine/Source/Runtime/Landscape/Classes/Materials/MaterialExpressionLandscapeVisibilityMask.h
+++ b/Engine/Source/Runtime/Landscape/Classes/Materials/MaterialExpressionLandscapeVisibilityMask.h
@@ -26,6 +26,8 @@ public:
 	LANDSCAPE_API virtual int32 Compile(class FMaterialCompiler* Compiler, int32 OutputIndex) override;
 	LANDSCAPE_API virtual void GetCaption(TArray<FString>& OutCaptions) const override;
 
+	virtual bool GenerateHLSLExpression(FMaterialHLSLGenerator& Generator, UE::HLSLTree::FScope& Scope, int32 OutputIndex, UE::HLSLTree::FExpression const*& OutExpression) const override;
+
 	/**
 	 * Gets the landscape layer names
 	 */
diff --git a/Engine/Source/Runtime/Landscape/Private/LandscapeGrass.cpp b/Engine/Source/Runtime/Landscape/Private/LandscapeGrass.cpp
--- a/Engine/Source/Runtime/Landscape/Private/LandscapeGrass.cpp
+++ b/Engine/Source/Runtime/Landscape/Private/LandscapeGrass.cpp
@@ -71,6 +71,9 @@
 #include "TextureResource.h"
 #include "Serialization/ArchiveCrc32.h"
 #include "SceneRendererInterface.h"
+#if WITH_EDITOR
+#include "MaterialHLSLGenerator.h"
+#endif
 
 #define LOCTEXT_NAMESPACE "Landscape"
 
@@ -1808,6 +1811,22 @@ int32 UMaterialExpressionLandscapeGrassOutput::Compile(class FMaterialCompiler*
 	return INDEX_NONE;
 }
 
+bool UMaterialExpressionLandscapeGrassOutput::GenerateHLSLExpression(FMaterialHLSLGenerator& Generator, UE::HLSLTree::FScope& Scope, int32 OutputIndex, UE::HLSLTree::FExpression const*& OutExpression) const
+{
+	if (!GrassTypes.IsValidIndex(OutputIndex))
+	{
+		return Generator.Error(TEXT("Invalid LandscapeGrassOutput OutputIndex."));
+	}
+
+	OutExpression = GrassTypes[OutputIndex].Input.AcquireHLSLExpressionOrConstant(Generator, Scope, 0.f);
+	return true;
+}
+
+UE::Shader::EValueType UMaterialExpressionLandscapeGrassOutput::GetCustomOutputType(int32 OutputIndex) const
+{
+	return UE::Shader::EValueType::Float1;
+}
+
 void UMaterialExpressionLandscapeGrassOutput::GetCaption(TArray<FString>& OutCaptions) const
 {
 	OutCaptions.Add(TEXT("Landscape Grass"));
diff --git a/Engine/Source/Runtime/Landscape/Private/Materials/MaterialExpressionLandscapeLayerBlend.cpp b/Engine/Source/Runtime/Landscape/Private/Materials/MaterialExpressionLandscapeLayerBlend.cpp
--- a/Engine/Source/Runtime/Landscape/Private/Materials/MaterialExpressionLandscapeLayerBlend.cpp
+++ b/Engine/Source/Runtime/Landscape/Private/Materials/MaterialExpressionLandscapeLayerBlend.cpp
@@ -7,6 +7,7 @@
 #include "MaterialCompiler.h"
 #if WITH_EDITOR
 #include "MaterialGraph/MaterialGraphNode.h"
+#include "MaterialHLSLGenerator.h"
 #endif
 
 #define LOCTEXT_NAMESPACE "Landscape"
@@ -281,6 +282,107 @@ int32 UMaterialExpressionLandscapeLayerBlend::Compile(class FMaterialCompiler* C
 
 	return OutputCode;
 }
+
+bool UMaterialExpressionLandscapeLayerBlend::GenerateHLSLExpression(FMaterialHLSLGenerator& Generator, UE::HLSLTree::FScope& Scope, int32 OutputIndex, UE::HLSLTree::FExpression const*& OutExpression) const
+{
+	using namespace UE::HLSLTree;
+
+	bool bNeedsRenormalize = false;
+	FTree& Tree = Generator.GetTree();
+	const FExpression* ConstantOne = Tree.NewConstant(1.f);
+	const FExpression* WeightSumExpression = Tree.NewConstant(0.f);
+	TArray<const FExpression*> WeightExpressions;
+
+	WeightExpressions.Empty(Layers.Num());
+
+	for (int32 LayerIdx = 0; LayerIdx < Layers.Num(); LayerIdx++)
+	{
+		WeightExpressions.Add(nullptr);
+
+		const FLayerBlendInput& Layer = Layers[LayerIdx];
+
+		// LB_AlphaBlend layers are blended last
+		if (Layer.BlendType != LB_AlphaBlend)
+		{
+			const FExpression* WeightExpression = nullptr;
+			verify(GenerateStaticTerrainLayerWeightExpression(Layer.LayerName, Layer.PreviewWeight, Generator, WeightExpression));
+
+			if (WeightExpression)
+			{
+				switch (Layer.BlendType)
+				{
+				case LB_WeightBlend:
+				{
+					// Store the weight plus accumulate the sum of all weights so far
+					WeightExpressions[LayerIdx] = WeightExpression;
+					WeightSumExpression = Tree.NewAdd(WeightSumExpression, WeightExpression);
+				}
+				break;
+				case LB_HeightBlend:
+				{
+					bNeedsRenormalize = true;
+
+					// Modify weight with height
+					const FExpression* HeightExpression = Layer.HeightInput.AcquireHLSLExpressionOrConstant(Generator, Scope, Layer.ConstHeightInput);
+					const FExpression* ModifiedWeight = Tree.NewLerp(Tree.NewConstant(-1.f), ConstantOne, WeightExpression);
+					ModifiedWeight = Tree.NewAdd(ModifiedWeight, HeightExpression);
+					ModifiedWeight = Tree.NewMin(Tree.NewMax(ModifiedWeight, Tree.NewConstant(0.0001f)), ConstantOne);
+
+					// Store the final weight plus accumulate the sum of all weights so far
+					WeightExpressions[LayerIdx] = ModifiedWeight;
+					WeightSumExpression = Tree.NewAdd(WeightSumExpression, ModifiedWeight);
+				}
+				break;
+				}
+			}
+		}
+	}
+
+	const FExpression* InvWeightSumExpression = Tree.NewDiv(ConstantOne, WeightSumExpression);
+	OutExpression = Tree.NewConstant(0.f);
+
+	for (int32 LayerIdx = 0; LayerIdx < Layers.Num(); LayerIdx++)
+	{
+		const FLayerBlendInput& Layer = Layers[LayerIdx];
+
+		if (WeightExpressions[LayerIdx])
+		{
+			const FExpression* LayerExpression = Layer.LayerInput.AcquireHLSLExpressionOrConstant(Generator, Scope, FVector3f(Layer.ConstLayerInput));
+
+			if (bNeedsRenormalize)
+			{
+				// Renormalize the weights as our height modification has made them non-uniform
+				OutExpression = Tree.NewAdd(OutExpression, Tree.NewMul(LayerExpression, Tree.NewMul(InvWeightSumExpression, WeightExpressions[LayerIdx])));
+			}
+			else
+			{
+				// No renormalization is necessary, so just add the weights
+				OutExpression = Tree.NewAdd(OutExpression, Tree.NewMul(LayerExpression, WeightExpressions[LayerIdx]));
+			}
+		}
+	}
+
+	// Blend in LB_AlphaBlend layers
+	for (const FLayerBlendInput& Layer : Layers)
+	{
+		if (Layer.BlendType == LB_AlphaBlend)
+		{
+			const FExpression* WeightExpression = nullptr;
+			verify(GenerateStaticTerrainLayerWeightExpression(Layer.LayerName, Layer.PreviewWeight, Generator, WeightExpression));
+
+			if (WeightExpression)
+			{
+				const FExpression* LayerExpression = Layer.LayerInput.AcquireHLSLExpressionOrConstant(Generator, Scope, FVector3f(Layer.ConstLayerInput));
+
+				// Blend in the layer using the alpha value
+				OutExpression = Tree.NewLerp(OutExpression, LayerExpression, WeightExpression);
+			}
+		}
+	}
+
+	return true;
+}
+
 #endif // WITH_EDITOR
 
 UObject* UMaterialExpressionLandscapeLayerBlend::GetReferencedTexture() const
diff --git a/Engine/Source/Runtime/Landscape/Private/Materials/MaterialExpressionLandscapeLayerSample.cpp b/Engine/Source/Runtime/Landscape/Private/Materials/MaterialExpressionLandscapeLayerSample.cpp
--- a/Engine/Source/Runtime/Landscape/Private/Materials/MaterialExpressionLandscapeLayerSample.cpp
+++ b/Engine/Source/Runtime/Landscape/Private/Materials/MaterialExpressionLandscapeLayerSample.cpp
@@ -6,6 +6,9 @@
 #include "EngineGlobals.h"
 #include "MaterialCompiler.h"
 #include "Materials/Material.h"
+#if WITH_EDITOR
+#include "MaterialHLSLGenerator.h"
+#endif
 
 #include UE_INLINE_GENERATED_CPP_BY_NAME(MaterialExpressionLandscapeLayerSample)
 
@@ -48,6 +51,12 @@ int32 UMaterialExpressionLandscapeLayerSample::Compile(class FMaterialCompiler*
 		return WeightCode;
 	}
 }
+
+bool UMaterialExpressionLandscapeLayerSample::GenerateHLSLExpression(FMaterialHLSLGenerator& Generator, UE::HLSLTree::FScope& Scope, int32 OutputIndex, UE::HLSLTree::FExpression const*& OutExpression) const
+{
+	return GenerateStaticTerrainLayerWeightExpression(ParameterName, PreviewWeight, Generator, OutExpression);
+}
+
 #endif // WITH_EDITOR
 
 UObject* UMaterialExpressionLandscapeLayerSample::GetReferencedTexture() const
diff --git a/Engine/Source/Runtime/Landscape/Private/Materials/MaterialExpressionLandscapePhysicalMaterialOutput.cpp b/Engine/Source/Runtime/Landscape/Private/Materials/MaterialExpressionLandscapePhysicalMaterialOutput.cpp
--- a/Engine/Source/Runtime/Landscape/Private/Materials/MaterialExpressionLandscapePhysicalMaterialOutput.cpp
+++ b/Engine/Source/Runtime/Landscape/Private/Materials/MaterialExpressionLandscapePhysicalMaterialOutput.cpp
@@ -7,6 +7,9 @@
 #include "EngineGlobals.h"
 #include "MaterialCompiler.h"
 #include "PhysicalMaterials/PhysicalMaterial.h"
+#if WITH_EDITOR
+#include "MaterialHLSLGenerator.h"
+#endif
 
 #include UE_INLINE_GENERATED_CPP_BY_NAME(MaterialExpressionLandscapePhysicalMaterialOutput)
 
@@ -89,6 +92,32 @@ int32 UMaterialExpressionLandscapePhysicalMaterialOutput::Compile(class FMateria
 	return INDEX_NONE;
 }
 
+UE::Shader::EValueType UMaterialExpressionLandscapePhysicalMaterialOutput::GetCustomOutputType(int32 OutputIndex) const
+{
+	return UE::Shader::EValueType::Float1;
+}
+
+bool UMaterialExpressionLandscapePhysicalMaterialOutput::GenerateHLSLExpression(FMaterialHLSLGenerator& Generator, UE::HLSLTree::FScope& Scope, int32 OutputIndex, UE::HLSLTree::FExpression const*& OutExpression) const
+{
+	if (!Inputs.IsValidIndex(OutputIndex))
+	{
+		return Generator.Error(TEXT("Invalid LandscapePhysicalMaterialOutput OutputIndex."));
+	}
+
+	if (!Inputs[OutputIndex].PhysicalMaterial)
+	{
+		return Generator.Error(TEXT("LandscapePhysicalMaterialOutput PhysicalMaterial not set."));
+	}
+
+	if (!Inputs[OutputIndex].Input.GetTracedInput().Expression)
+	{
+		return Generator.Error(TEXT("LandscapePhysicalMaterialOutput Input missing."));
+	}
+
+	OutExpression = Inputs[OutputIndex].Input.AcquireHLSLExpression(Generator, Scope);
+	return true;
+}
+
 void UMaterialExpressionLandscapePhysicalMaterialOutput::PostEditChangeProperty(FPropertyChangedEvent& PropertyChangedEvent)
 {
 	Super::PostEditChangeProperty(PropertyChangedEvent);
diff --git a/Engine/Source/Runtime/Landscape/Private/Materials/MaterialExpressionLandscapeVisibilityMask.cpp b/Engine/Source/Runtime/Landscape/Private/Materials/MaterialExpressionLandscapeVisibilityMask.cpp
--- a/Engine/Source/Runtime/Landscape/Private/Materials/MaterialExpressionLandscapeVisibilityMask.cpp
+++ b/Engine/Source/Runtime/Landscape/Private/Materials/MaterialExpressionLandscapeVisibilityMask.cpp
@@ -6,6 +6,9 @@
 #include "EngineGlobals.h"
 #include "MaterialCompiler.h"
 #include "Materials/Material.h"
+#if WITH_EDITOR
+#include "MaterialHLSLGenerator.h"
+#endif
 
 #include UE_INLINE_GENERATED_CPP_BY_NAME(MaterialExpressionLandscapeVisibilityMask)
 
@@ -43,6 +46,19 @@ int32 UMaterialExpressionLandscapeVisibilityMask::Compile(class FMaterialCompile
 	int32 MaskLayerCode = Compiler->StaticTerrainLayerWeight(ParameterName, Compiler->Constant(0.f));
 	return MaskLayerCode == INDEX_NONE ? Compiler->Constant(1.f) : Compiler->Sub(Compiler->Constant(1.f), MaskLayerCode);
 }
+
+bool UMaterialExpressionLandscapeVisibilityMask::GenerateHLSLExpression(FMaterialHLSLGenerator& Generator, UE::HLSLTree::FScope& Scope, int32 OutputIndex, UE::HLSLTree::FExpression const*& OutExpression) const
+{
+	using namespace UE::HLSLTree;
+
+	const FExpression* MaskLayerExpression = nullptr;
+	verify(GenerateStaticTerrainLayerWeightExpression(ParameterName, 0.f, Generator, MaskLayerExpression));
+
+	FTree& Tree = Generator.GetTree();
+	const FExpression* ConstantOne = Tree.NewConstant(1.f);
+	OutExpression = MaskLayerExpression ? Tree.NewSub(ConstantOne, MaskLayerExpression) : ConstantOne;
+	return true;
+}
 #endif // WITH_EDITOR
 
 UObject* UMaterialExpressionLandscapeVisibilityMask::GetReferencedTexture() const
