Exposed the ability to check if a cluster union component is anchored

#rnx

diff --git a/Engine/Source/Runtime/Engine/Classes/PhysicsEngine/ClusterUnionComponent.h b/Engine/Source/Runtime/Engine/Classes/PhysicsEngine/ClusterUnionComponent.h
--- a/Engine/Source/Runtime/Engine/Classes/PhysicsEngine/ClusterUnionComponent.h
+++ b/Engine/Source/Runtime/Engine/Classes/PhysicsEngine/ClusterUnionComponent.h
@@ -151,6 +151,8 @@ public:
 	UFUNCTION(BlueprintCallable, Category = "Cluster Union")
 	ENGINE_API void SetIsAnchored(bool bIsAnchored);
 
+	ENGINE_API bool IsAnchored() const;
+
 	// SyncClusterUnionFromProxy will examine the make up of the cluster union (particles, child to parent, etc.) and do whatever is needed on the GT in terms of bookkeeping.
 	ENGINE_API void SyncClusterUnionFromProxy();
 
diff --git a/Engine/Source/Runtime/Engine/Private/PhysicsEngine/ClusterUnionComponent.cpp b/Engine/Source/Runtime/Engine/Private/PhysicsEngine/ClusterUnionComponent.cpp
--- a/Engine/Source/Runtime/Engine/Private/PhysicsEngine/ClusterUnionComponent.cpp
+++ b/Engine/Source/Runtime/Engine/Private/PhysicsEngine/ClusterUnionComponent.cpp
@@ -568,6 +568,16 @@ void UClusterUnionComponent::SetIsAnchored(bool bIsAnchored)
 	PhysicsProxy->SetIsAnchored_External(bIsAnchored);
 }
 
+ENGINE_API bool UClusterUnionComponent::IsAnchored() const
+{
+	if (!PhysicsProxy)
+	{
+		return false;
+	}
+
+	return PhysicsProxy->IsAnchored_External();
+}
+
 bool UClusterUnionComponent::IsAuthority() const
 {
 	ENetMode Mode = GetNetMode();
