Fixes to skin cache Set PSO matching up with Dispatches, and temporarily set r.RenderCommandPipeMode to 1 (instead of 2), until the PSO cache cleanup can be corrected (the refactored thread local compute PSO optimization exposes a race condition in certain cases like GPUSkinCache + render commands).

diff --git a/Engine/Source/Runtime/Engine/Private/GPUSkinCache.cpp b/Engine/Source/Runtime/Engine/Private/GPUSkinCache.cpp
--- a/Engine/Source/Runtime/Engine/Private/GPUSkinCache.cpp
+++ b/Engine/Source/Runtime/Engine/Private/GPUSkinCache.cpp
@@ -1157,9 +1157,6 @@ void FGPUSkinCache::DispatchUpdateSkinTangents(FRHICommandList& RHICmdList, FGPU
 			SCOPED_DRAW_EVENTF(RHICmdList, SkinTangents_PerTrianglePass, TEXT("%sTangentsTri  Mesh=%s, LOD=%d, Chunk=%d, IndexStart=%d Tri=%d BoneInfluenceType=%d UVPrecision=%d"),
 				*RayTracingTag , *GetSkeletalMeshObjectName(Entry->GPUSkin), LODIndex, SectionIndex, DispatchData.IndexBufferOffsetValue, DispatchData.NumTriangles, Entry->BoneInfluenceType, bFullPrecisionUV);
 
-			FRHIComputeShader* ShaderRHI = Shader.GetComputeShader();
-			SetComputePipelineState(RHICmdList, Shader.GetComputeShader());
-
 			if (!GAllowDupedVertsForRecomputeTangents)
 			{
 #if WITH_EDITOR
@@ -1182,6 +1179,9 @@ void FGPUSkinCache::DispatchUpdateSkinTangents(FRHICommandList& RHICmdList, FGPU
 
 			const FRWBuffer& ShaderStagingBuffer = GRecomputeTangentsParallelDispatch ? DispatchData.GetIntermediateAccumulatedTangentBuffer()->Buffer : StagingBuffer->Buffer;
 
+			FRHIComputeShader* ShaderRHI = Shader.GetComputeShader();
+			SetComputePipelineState(RHICmdList, Shader.GetComputeShader());
+
 			SetShaderParametersLegacyCS(RHICmdList, Shader, Entry, DispatchData, ShaderStagingBuffer);
 			DispatchComputeShader(RHICmdList, Shader.GetShader(), ThreadGroupCountValue, 1, 1);
 			UnsetShaderParametersLegacyCS(RHICmdList, Shader);
@@ -1203,8 +1203,6 @@ void FGPUSkinCache::DispatchUpdateSkinTangents(FRHICommandList& RHICmdList, FGPU
 		else
 			ComputeShader = ComputeShader0;
 
-		SetComputePipelineState(RHICmdList, ComputeShader.GetComputeShader());
-
 		uint32 VertexCount = DispatchData.NumVertices;
 		uint32 ThreadGroupCountValue = FMath::DivideAndRoundUp(VertexCount, ComputeShader->ThreadGroupSizeX);
 
@@ -1217,6 +1215,8 @@ void FGPUSkinCache::DispatchUpdateSkinTangents(FRHICommandList& RHICmdList, FGPU
 				});
 		}
 
+		SetComputePipelineState(RHICmdList, ComputeShader.GetComputeShader());
+
 		SetShaderParametersLegacyCS(RHICmdList, ComputeShader, Entry, DispatchData, GRecomputeTangentsParallelDispatch ? DispatchData.GetIntermediateAccumulatedTangentBuffer()->Buffer : StagingBuffer->Buffer);
 		DispatchComputeShader(RHICmdList, ComputeShader.GetShader(), ThreadGroupCountValue, 1, 1);
 		UnsetShaderParametersLegacyCS(RHICmdList, ComputeShader);
@@ -2079,15 +2079,19 @@ void FGPUSkinCache::DispatchUpdateSkinning(FRHICommandList& RHICmdList, FGPUSkin
 
 	check(Shader.IsValid());
 
-	if ((DispatchData.DispatchFlags & ((uint32)EGPUSkinCacheDispatchFlags::DispatchPrevPosition | (uint32)EGPUSkinCacheDispatchFlags::DispatchPosition)) != 0)
+	const bool bDispatchPrevPosition = (DispatchData.DispatchFlags & (uint32)EGPUSkinCacheDispatchFlags::DispatchPrevPosition) != 0;
+	const bool bDispatchPosition = (DispatchData.DispatchFlags & (uint32)EGPUSkinCacheDispatchFlags::DispatchPosition) != 0;
+
+	if (bDispatchPrevPosition || bDispatchPosition)
 	{
-		SetComputePipelineState(RHICmdList, Shader.GetComputeShader());
 		uint32 VertexCountAlign64 = FMath::DivideAndRoundUp(DispatchData.NumVertices, (uint32)64);
 
-		if ((DispatchData.DispatchFlags & (uint32)EGPUSkinCacheDispatchFlags::DispatchPrevPosition) != 0)
+		if (bDispatchPrevPosition)
 		{
 			const FVertexBufferAndSRV& PrevBoneBuffer = ShaderData.GetBoneBufferForReading(true);
 
+			SetComputePipelineState(RHICmdList, Shader.GetComputeShader());
+
 			SetShaderParametersLegacyCS(
 				RHICmdList,
 				Shader,
@@ -2105,10 +2109,12 @@ void FGPUSkinCache::DispatchUpdateSkinning(FRHICommandList& RHICmdList, FGPUSkin
 			BuffersToTransitionToRead.Add(DispatchData.GetPreviousPositionRWBuffer());
 		}
 
-		if ((DispatchData.DispatchFlags & (uint32)EGPUSkinCacheDispatchFlags::DispatchPosition) != 0)
+		if (bDispatchPosition)
 		{
 			const FVertexBufferAndSRV& BoneBuffer = ShaderData.GetBoneBufferForReading(false);
 
+			SetComputePipelineState(RHICmdList, Shader.GetComputeShader());
+
 			SetShaderParametersLegacyCS(
 				RHICmdList,
 				Shader,
diff --git a/Engine/Source/Runtime/RHI/Private/PipelineStateCache.cpp b/Engine/Source/Runtime/RHI/Private/PipelineStateCache.cpp
--- a/Engine/Source/Runtime/RHI/Private/PipelineStateCache.cpp
+++ b/Engine/Source/Runtime/RHI/Private/PipelineStateCache.cpp
@@ -852,7 +852,7 @@ public:
 #if PIPELINESTATECACHE_VERIFYTHREADSAFE
 		FScopeVerifyIncrement S(VerifyMutex);
 #endif
-		// everything is added to the local cache then at end of frame we consoldate them all
+		// everything is added to the local cache then at end of frame we consolidate them all
 		TMap<TMyKey, TMyValue> &LocalCache = GetLocalCache();
 
 		check( LocalCache.Contains(InKey) == false );
diff --git a/Engine/Source/Runtime/RenderCore/Private/RenderingThread.cpp b/Engine/Source/Runtime/RenderCore/Private/RenderingThread.cpp
--- a/Engine/Source/Runtime/RenderCore/Private/RenderingThread.cpp
+++ b/Engine/Source/Runtime/RenderCore/Private/RenderingThread.cpp
@@ -1698,7 +1698,7 @@ inline ERenderCommandPipeMode GetValidatedRenderCommandPipeMode(int32 CVarValue)
 ERenderCommandPipeMode GRenderCommandPipeMode = ERenderCommandPipeMode::None;
 FAutoConsoleVariable CVarRenderCommandPipeMode(
 	TEXT("r.RenderCommandPipeMode"),
-	2,
+	1,
 	TEXT("Controls behavior of the main render thread command pipe.")
 	TEXT(" 0: Render commands are launched individually as tasks;\n")
 	TEXT(" 1: Render commands are enqueued into a render command pipe for the render thread only.;\n")
