Substrate - Fixed volumetric cloud to account for emissive and ambient occlusion.

diff --git a/Engine/Shaders/Public/ShaderVersion.ush b/Engine/Shaders/Public/ShaderVersion.ush
--- a/Engine/Shaders/Public/ShaderVersion.ush
+++ b/Engine/Shaders/Public/ShaderVersion.ush
@@ -2,4 +2,4 @@
 // Each time the console command is executed it generates a new GUID. As this file is included
 // in Platform.ush (which should be included in any shader) it allows to invalidate the shader DDC.
 // 
-#pragma message("UESHADERMETADATA_VERSION E8C6552B-D6B5-4CA4-B26C-1E315D3A1F20")
+#pragma message("UESHADERMETADATA_VERSION 959736DE-626A-4776-A2DC-36960B1A0099")
diff --git a/Engine/Source/Runtime/Engine/Private/Materials/HLSLMaterialTranslator.cpp b/Engine/Source/Runtime/Engine/Private/Materials/HLSLMaterialTranslator.cpp
--- a/Engine/Source/Runtime/Engine/Private/Materials/HLSLMaterialTranslator.cpp
+++ b/Engine/Source/Runtime/Engine/Private/Materials/HLSLMaterialTranslator.cpp
@@ -307,6 +307,8 @@ void FHLSLMaterialTranslator::FSubstrateCompilationContext::Initialise()
 	bSubstrateMaterialIsSimple = false;
 	bSubstrateMaterialIsSingle = false;
 	bSubstrateMaterialIsUnlitNode = false;
+	bSubstrateWritesEmissive = false;
+	bSubstrateWritesAmbientOcclusion = false;
 	bSubstrateTreeOutOfStackDepthOccurred = false;
 
 	// Default value used as the root of the tree for the first path (when a node parent==nullptr).
@@ -436,6 +438,9 @@ FHLSLMaterialTranslator::FHLSLMaterialTranslator(FMaterial* InMaterial,
 		}
 	}
 
+	bSubstrateWritesEmissive = false;
+	bSubstrateWritesAmbientOcclusion = false;
+
 	bSubstrateUsesConversionFromLegacy = false;
 	bSubstrateOutputsOpaqueRoughRefractions = false;
 
@@ -884,6 +889,9 @@ bool FHLSLMaterialTranslator::Translate()
 		UMaterialExpression* FrontMaterialExpr = nullptr;
 		int32 FrontMaterialOutputIndex = INDEX_NONE;
 
+		bSubstrateWritesEmissive = false;
+		bSubstrateWritesAmbientOcclusion = false;
+
 		UMaterialExpression* ExpressionToPreview = Material->GetMaterialGraphNodePreviewExpression();
 		if (ExpressionToPreview)
 		{
@@ -938,6 +946,9 @@ bool FHLSLMaterialTranslator::Translate()
 				{
 					Errorf(TEXT("Substrate material errors encountered."));
 				}
+
+				bSubstrateWritesEmissive |= SubstrateCtx.bSubstrateWritesEmissive;
+				bSubstrateWritesAmbientOcclusion |= SubstrateCtx.bSubstrateWritesAmbientOcclusion;
 			}
 			CurrentSubstrateCompilationContext = ESubstrateCompilationContext::SCC_Default;
 		}
@@ -1233,7 +1244,7 @@ bool FHLSLMaterialTranslator::Translate()
 		// At this point we mark the code chunk generation complete.
 		bAllowCodeChunkGeneration = false;
 
-		bUsesEmissiveColor = IsMaterialPropertyUsed(MP_EmissiveColor, Chunk[MP_EmissiveColor], FLinearColor(0, 0, 0, 0), 3);
+		bUsesEmissiveColor = bSubstrateWritesEmissive || IsMaterialPropertyUsed(MP_EmissiveColor, Chunk[MP_EmissiveColor], FLinearColor(0, 0, 0, 0), 3);
 		bUsesPixelDepthOffset = (AllowPixelDepthOffset(Platform) && IsMaterialPropertyUsed(MP_PixelDepthOffset, Chunk[MP_PixelDepthOffset], FLinearColor(0, 0, 0, 0), 1));
 
 		bool bUsesWorldPositionOffsetCurrent = IsMaterialPropertyUsed(MP_WorldPositionOffset, Chunk[MP_WorldPositionOffset], FLinearColor(0, 0, 0, 0), 3);
@@ -2263,7 +2274,7 @@ void FHLSLMaterialTranslator::GetMaterialEnvironment(EShaderPlatform InPlatform,
 				VolumetricAdvancedNode->ConservativeDensity.IsConnected() ? TEXT("1") : TEXT("0"));
 
 			OutEnvironment.SetDefine(TEXT("MATERIAL_VOLUMETRIC_ADVANCED_OVERRIDE_AMBIENT_OCCLUSION"),
-				Material->HasAmbientOcclusionConnected() ? TEXT("1") : TEXT("0"));
+				(Material->HasAmbientOcclusionConnected() || bSubstrateWritesAmbientOcclusion) ? TEXT("1") : TEXT("0"));
 
 			OutEnvironment.SetDefine(TEXT("MATERIAL_VOLUMETRIC_ADVANCED_GROUND_CONTRIBUTION"),
 				VolumetricAdvancedNode->bGroundContribution ? TEXT("1") : TEXT("0"));
@@ -11498,6 +11509,13 @@ bool FHLSLMaterialTranslator::FSubstrateCompilationContext::SubstrateGenerateDer
 			check(RootIndex == INDEX_NONE);	// There can only be one
 			RootIndex = It.Index;
 		}
+
+		if (It.OperatorType == SUBSTRATE_OPERATOR_BSDF)
+		{
+			// Gather information about data written by BSDF
+			bSubstrateWritesEmissive |= It.bBSDFWritesEmissive > 0;
+			bSubstrateWritesAmbientOcclusion |= It.bBSDFWritesAmbientOcclusion > 0;
+		}
 	}
 	SubstrateMaterialRootOperator = &SubstrateMaterialExpressionRegisteredOperators[RootIndex];
 	if (!SubstrateMaterialRootOperator)
diff --git a/Engine/Source/Runtime/Engine/Private/Materials/HLSLMaterialTranslator.h b/Engine/Source/Runtime/Engine/Private/Materials/HLSLMaterialTranslator.h
--- a/Engine/Source/Runtime/Engine/Private/Materials/HLSLMaterialTranslator.h
+++ b/Engine/Source/Runtime/Engine/Private/Materials/HLSLMaterialTranslator.h
@@ -499,6 +499,9 @@ protected:
 		bool bSubstrateMaterialIsSingle;
 		bool bSubstrateMaterialIsUnlitNode;
 
+		bool bSubstrateWritesEmissive;
+		bool bSubstrateWritesAmbientOcclusion;
+
 		/** Stack of unique id for each node of the Substrate tree
 		* This is transient and updated on the fly in the exact same way when parsing node for
 		*  1- SubstrateGenerateMaterialTopologyTree: generating a picture of the Substrate material tree for code generation and simplifications.
@@ -537,6 +540,9 @@ protected:
 	FString FullySimplifiedSubstrateFrontMaterialTranslatedCodeChunkDefinitions;
 	FString FullySimplifiedSubstrateFrontMaterialTranslatedCodeChunks;
 
+	bool bSubstrateWritesEmissive;
+	bool bSubstrateWritesAmbientOcclusion;
+
 	bool bSubstrateUsesConversionFromLegacy;
 	bool bSubstrateOutputsOpaqueRoughRefractions;
 
diff --git a/Engine/Source/Runtime/Engine/Private/Materials/MaterialExpressions.cpp b/Engine/Source/Runtime/Engine/Private/Materials/MaterialExpressions.cpp
--- a/Engine/Source/Runtime/Engine/Private/Materials/MaterialExpressions.cpp
+++ b/Engine/Source/Runtime/Engine/Private/Materials/MaterialExpressions.cpp
@@ -23705,6 +23705,7 @@ FSubstrateOperator* UMaterialExpressionSubstrateShadingModels::SubstrateGenerate
 		SlabOperator.bBSDFHasMFPPluggedIn = bSSS;
 		SlabOperator.bBSDFHasFuzz = bFuzz;
 		SlabOperator.bBSDFHasAnisotropy = Anisotropy.IsConnected();
+		SlabOperator.bBSDFWritesEmissive = EmissiveColor.IsConnected();
 		SlabOperator.ThicknessIndex = ThicknessIndex;
 
 		return &SlabOperator;
@@ -23736,6 +23737,7 @@ FSubstrateOperator* UMaterialExpressionSubstrateShadingModels::SubstrateGenerate
 			FSubstrateOperator& Operator = Compiler->SubstrateCompilationRegisterOperator(SUBSTRATE_OPERATOR_BSDF_LEGACY, Compiler->SubstrateTreeStackGetPathUniqueId(), this, Parent, Compiler->SubstrateTreeStackGetParentPathUniqueId());
 			Operator.BSDFType = SUBSTRATE_BSDF_TYPE_UNLIT;
 			Operator.ThicknessIndex = ThicknessIndex;
+			Operator.bBSDFWritesEmissive = true;
 			return &Operator;
 		}
 		else if (ShadingModels.HasShadingModel(MSM_DefaultLit))
@@ -23773,6 +23775,7 @@ FSubstrateOperator* UMaterialExpressionSubstrateShadingModels::SubstrateGenerate
 			Operator.ThicknessIndex = ThicknessIndex;
 			Operator.bBSDFHasSecondRoughnessOrSimpleClearCoat = true;
 			Operator.bBSDFHasAnisotropy = Anisotropy.IsConnected();
+			Operator.bBSDFWritesEmissive = EmissiveColor.IsConnected();
 			return &Operator;
 		}
 		else if (ShadingModels.HasShadingModel(MSM_Hair))
@@ -23780,6 +23783,7 @@ FSubstrateOperator* UMaterialExpressionSubstrateShadingModels::SubstrateGenerate
 			FSubstrateOperator& Operator = Compiler->SubstrateCompilationRegisterOperator(SUBSTRATE_OPERATOR_BSDF_LEGACY, Compiler->SubstrateTreeStackGetPathUniqueId(), this, Parent, Compiler->SubstrateTreeStackGetParentPathUniqueId());
 			Operator.BSDFType = SUBSTRATE_BSDF_TYPE_HAIR;
 			Operator.ThicknessIndex = ThicknessIndex;
+			Operator.bBSDFWritesEmissive = EmissiveColor.IsConnected();
 			return &Operator;
 		}
 		else if (ShadingModels.HasShadingModel(MSM_Eye))
@@ -23787,6 +23791,7 @@ FSubstrateOperator* UMaterialExpressionSubstrateShadingModels::SubstrateGenerate
 			FSubstrateOperator& Operator = Compiler->SubstrateCompilationRegisterOperator(SUBSTRATE_OPERATOR_BSDF_LEGACY, Compiler->SubstrateTreeStackGetPathUniqueId(), this, Parent, Compiler->SubstrateTreeStackGetParentPathUniqueId());
 			Operator.BSDFType = SUBSTRATE_BSDF_TYPE_EYE;
 			Operator.ThicknessIndex = ThicknessIndex;
+			Operator.bBSDFWritesEmissive = EmissiveColor.IsConnected();
 			return &Operator;
 		}
 		else if (ShadingModels.HasShadingModel(MSM_SingleLayerWater))
@@ -23794,6 +23799,7 @@ FSubstrateOperator* UMaterialExpressionSubstrateShadingModels::SubstrateGenerate
 			FSubstrateOperator& Operator = Compiler->SubstrateCompilationRegisterOperator(SUBSTRATE_OPERATOR_BSDF_LEGACY, Compiler->SubstrateTreeStackGetPathUniqueId(), this, Parent, Compiler->SubstrateTreeStackGetParentPathUniqueId());
 			Operator.BSDFType = SUBSTRATE_BSDF_TYPE_SINGLELAYERWATER;
 			Operator.ThicknessIndex = ThicknessIndex;
+			Operator.bBSDFWritesEmissive = EmissiveColor.IsConnected();
 			return &Operator;
 		}
 
@@ -24246,6 +24252,7 @@ FSubstrateOperator* UMaterialExpressionSubstrateSlabBSDF::SubstrateGenerateMater
 	SubstrateOperator.bBSDFHasGlint = HasGlint();
 	SubstrateOperator.bBSDFHasSpecularProfile = HasSpecularProfile();
 	SubstrateOperator.ThicknessIndex = Compiler->SubstrateThicknessStackGetThicknessIndex();
+	SubstrateOperator.bBSDFWritesEmissive = EmissiveColor.IsConnected();
 	return &SubstrateOperator;
 }
 
@@ -24479,6 +24486,7 @@ FSubstrateOperator* UMaterialExpressionSubstrateSimpleClearCoatBSDF::SubstrateGe
 	SubstrateOperator.bBSDFHasGlint = false;
 	SubstrateOperator.bBSDFHasSpecularProfile = false;
 	SubstrateOperator.ThicknessIndex = Compiler->SubstrateThicknessStackGetThicknessIndex();
+	SubstrateOperator.bBSDFWritesEmissive = EmissiveColor.IsConnected();
 	return &SubstrateOperator;
 }
 #endif // WITH_EDITOR
@@ -24559,6 +24567,8 @@ FSubstrateOperator* UMaterialExpressionSubstrateVolumetricFogCloudBSDF::Substrat
 	FSubstrateOperator& SubstrateOperator = Compiler->SubstrateCompilationRegisterOperator(SUBSTRATE_OPERATOR_BSDF, Compiler->SubstrateTreeStackGetPathUniqueId(), this, Parent, Compiler->SubstrateTreeStackGetParentPathUniqueId());
 	SubstrateOperator.BSDFType = SUBSTRATE_BSDF_TYPE_VOLUMETRICFOGCLOUD;
 	SubstrateOperator.ThicknessIndex = Compiler->SubstrateThicknessStackGetThicknessIndex();
+	SubstrateOperator.bBSDFWritesEmissive = EmissiveColor.IsConnected();
+	SubstrateOperator.bBSDFWritesAmbientOcclusion = AmbientOcclusion.IsConnected();
 	return &SubstrateOperator;
 }
 #endif // WITH_EDITOR
@@ -24987,6 +24997,7 @@ FSubstrateOperator* UMaterialExpressionSubstrateUnlitBSDF::SubstrateGenerateMate
 	FSubstrateOperator& SubstrateOperator = Compiler->SubstrateCompilationRegisterOperator(SUBSTRATE_OPERATOR_BSDF, Compiler->SubstrateTreeStackGetPathUniqueId(), this, Parent, Compiler->SubstrateTreeStackGetParentPathUniqueId());
 	SubstrateOperator.BSDFType = SUBSTRATE_BSDF_TYPE_UNLIT;
 	SubstrateOperator.ThicknessIndex = Compiler->SubstrateThicknessStackGetThicknessIndex();
+	SubstrateOperator.bBSDFWritesEmissive = EmissiveColor.IsConnected();
 	return &SubstrateOperator;
 }
 #endif // WITH_EDITOR
@@ -25104,6 +25115,7 @@ FSubstrateOperator* UMaterialExpressionSubstrateHairBSDF::SubstrateGenerateMater
 	FSubstrateOperator& SubstrateOperator = Compiler->SubstrateCompilationRegisterOperator(SUBSTRATE_OPERATOR_BSDF, Compiler->SubstrateTreeStackGetPathUniqueId(), this, Parent, Compiler->SubstrateTreeStackGetParentPathUniqueId());
 	SubstrateOperator.BSDFType = SUBSTRATE_BSDF_TYPE_HAIR;
 	SubstrateOperator.ThicknessIndex = Compiler->SubstrateThicknessStackGetThicknessIndex();
+	SubstrateOperator.bBSDFWritesEmissive = EmissiveColor.IsConnected();
 	return &SubstrateOperator;
 }
 #endif // WITH_EDITOR
@@ -25219,6 +25231,7 @@ FSubstrateOperator* UMaterialExpressionSubstrateEyeBSDF::SubstrateGenerateMateri
 	FSubstrateOperator& SubstrateOperator = Compiler->SubstrateCompilationRegisterOperator(SUBSTRATE_OPERATOR_BSDF, Compiler->SubstrateTreeStackGetPathUniqueId(), this, Parent, Compiler->SubstrateTreeStackGetParentPathUniqueId());
 	SubstrateOperator.BSDFType = SUBSTRATE_BSDF_TYPE_EYE;
 	SubstrateOperator.ThicknessIndex = Compiler->SubstrateThicknessStackGetThicknessIndex();
+	SubstrateOperator.bBSDFWritesEmissive = EmissiveColor.IsConnected();
 	return &SubstrateOperator;
 }
 #endif // WITH_EDITOR
@@ -25351,6 +25364,7 @@ FSubstrateOperator* UMaterialExpressionSubstrateSingleLayerWaterBSDF::SubstrateG
 	FSubstrateOperator& SubstrateOperator = Compiler->SubstrateCompilationRegisterOperator(SUBSTRATE_OPERATOR_BSDF, Compiler->SubstrateTreeStackGetPathUniqueId(), this, Parent, Compiler->SubstrateTreeStackGetParentPathUniqueId());
 	SubstrateOperator.BSDFType = SUBSTRATE_BSDF_TYPE_SINGLELAYERWATER;
 	SubstrateOperator.ThicknessIndex = Compiler->SubstrateThicknessStackGetThicknessIndex();
+	SubstrateOperator.bBSDFWritesEmissive = EmissiveColor.IsConnected();
 	return &SubstrateOperator;
 }
 #endif // WITH_EDITOR
@@ -26938,6 +26952,7 @@ FSubstrateOperator* UMaterialExpressionSubstrateConvertMaterialAttributes::Subst
 			FSubstrateOperator& Operator = Compiler->SubstrateCompilationRegisterOperator(SUBSTRATE_OPERATOR_BSDF_LEGACY, Compiler->SubstrateTreeStackGetPathUniqueId(), this, Parent, Compiler->SubstrateTreeStackGetParentPathUniqueId());
 			Operator.BSDFType = SUBSTRATE_BSDF_TYPE_SINGLELAYERWATER;
 			Operator.ThicknessIndex = ThicknessIndex;
+			Operator.bBSDFWritesEmissive = MaterialAttributes.IsConnected(MP_EmissiveColor);
 			return &Operator;
 		}
 
diff --git a/Engine/Source/Runtime/Engine/Public/Rendering/SubstrateMaterialShared.cpp b/Engine/Source/Runtime/Engine/Public/Rendering/SubstrateMaterialShared.cpp
--- a/Engine/Source/Runtime/Engine/Public/Rendering/SubstrateMaterialShared.cpp
+++ b/Engine/Source/Runtime/Engine/Public/Rendering/SubstrateMaterialShared.cpp
@@ -57,6 +57,9 @@ FSubstrateOperator::FSubstrateOperator()
 	bBSDFHasGlint = false;
 	bBSDFHasSpecularProfile = false;
 
+	bBSDFWritesEmissive = false;
+	bBSDFWritesAmbientOcclusion = false;
+
 	MaxDistanceFromLeaves = 0;
 	LayerDepth = 0;
 	bIsTop = false;
diff --git a/Engine/Source/Runtime/Engine/Public/Rendering/SubstrateMaterialShared.h b/Engine/Source/Runtime/Engine/Public/Rendering/SubstrateMaterialShared.h
--- a/Engine/Source/Runtime/Engine/Public/Rendering/SubstrateMaterialShared.h
+++ b/Engine/Source/Runtime/Engine/Public/Rendering/SubstrateMaterialShared.h
@@ -57,6 +57,9 @@ public:
 	LAYOUT_FIELD_EDITORONLY(uint8, bBSDFHasGlint);
 	LAYOUT_FIELD_EDITORONLY(uint8, bBSDFHasSpecularProfile);
 
+	LAYOUT_FIELD_EDITORONLY(uint8, bBSDFWritesEmissive);
+	LAYOUT_FIELD_EDITORONLY(uint8, bBSDFWritesAmbientOcclusion);
+
 	// Data derived after the tree has been built.
 	LAYOUT_FIELD_EDITORONLY(int32, MaxDistanceFromLeaves);
 	LAYOUT_FIELD_EDITORONLY(int32, LayerDepth);
