[SkeletonCompatibility] Added support for skeleton compatibility in combination with blend profiles / masks.

diff --git a/Engine/Plugins/Experimental/Animation/BlendStack/Source/Runtime/Private/AnimNode_BlendStack.cpp b/Engine/Plugins/Experimental/Animation/BlendStack/Source/Runtime/Private/AnimNode_BlendStack.cpp
--- a/Engine/Plugins/Experimental/Animation/BlendStack/Source/Runtime/Private/AnimNode_BlendStack.cpp
+++ b/Engine/Plugins/Experimental/Animation/BlendStack/Source/Runtime/Private/AnimNode_BlendStack.cpp
@@ -38,7 +38,7 @@ void FBlendStackAnimPlayer::Initialize(const FAnimationInitializeContext& Contex
 	USkeleton* Skeleton = Context.AnimInstanceProxy->GetSkeleton();
 	check(Skeleton);
 
-	const FReferenceSkeleton& RefSkeleton = Context.AnimInstanceProxy->GetSkeleton()->GetReferenceSkeleton();
+	const FReferenceSkeleton& RefSkeleton = Skeleton->GetReferenceSkeleton();
 	const bool bApplyDifferentRootBoneBlendTime = RootBoneBlendTime >= 0.f && !FMath::IsNearlyEqual(RootBoneBlendTime, BlendTime);
 	const int32 NumSkeletonBones = RefSkeleton.GetNum();
 	if (NumSkeletonBones <= 0)
@@ -50,11 +50,9 @@ void FBlendStackAnimPlayer::Initialize(const FAnimationInitializeContext& Contex
 		// handling BlendTime > 0 and RootBoneBlendTime >= 0
 		if (BlendProfile != nullptr)
 		{
-			check(BlendProfile->OwningSkeleton && NumSkeletonBones == BlendProfile->OwningSkeleton->GetReferenceSkeleton().GetNum());
-
 			TotalBlendInTimePerBone.Init(BlendTime, NumSkeletonBones);
 
-			BlendProfile->FillSkeletonBoneDurationsArray(TotalBlendInTimePerBone, BlendTime);
+			BlendProfile->FillSkeletonBoneDurationsArray(TotalBlendInTimePerBone, BlendTime, Skeleton);
 
 			if (bApplyDifferentRootBoneBlendTime)
 			{
diff --git a/Engine/Source/Runtime/AnimGraphRuntime/Private/AnimNodes/AnimNode_BlendListBase.cpp b/Engine/Source/Runtime/AnimGraphRuntime/Private/AnimNodes/AnimNode_BlendListBase.cpp
--- a/Engine/Source/Runtime/AnimGraphRuntime/Private/AnimNodes/AnimNode_BlendListBase.cpp
+++ b/Engine/Source/Runtime/AnimGraphRuntime/Private/AnimNodes/AnimNode_BlendListBase.cpp
@@ -7,6 +7,8 @@
 #include "Animation/AnimInstanceProxy.h"
 #include "Animation/AnimNode_Inertialization.h"
 #include "Animation/AnimTrace.h"
+#include "Animation/SkeletonRemapping.h"
+#include "Animation/SkeletonRemappingRegistry.h"
 
 #include UE_INLINE_GENERATED_CPP_BY_NAME(AnimNode_BlendListBase)
 
@@ -318,7 +320,17 @@ void FAnimNode_BlendListBase::Evaluate_AnyThread(FPoseContext& Output)
 		UBlendProfile* CurrentBlendProfile = GetBlendProfile();
 		if (CurrentBlendProfile)
 		{
-			FAnimationRuntime::BlendPosesTogetherPerBone(FilteredPoses, FilteredCurve, FilteredAttributes, CurrentBlendProfile, PerBoneSampleData, PosesToEvaluate, OutAnimationPoseData);
+			const USkeleton* TargetSkeleton = Output.Pose.GetBoneContainer().GetSkeletonAsset();
+			const USkeleton* SourceSkeleton = CurrentBlendProfile->OwningSkeleton;
+			const FSkeletonRemapping& SkeletonRemapping = UE::Anim::FSkeletonRemappingRegistry::Get().GetRemapping(SourceSkeleton, TargetSkeleton);
+			if (SkeletonRemapping.IsValid())
+			{
+				FAnimationRuntime::BlendPosesTogetherPerBoneRemapped(FilteredPoses, FilteredCurve, FilteredAttributes, CurrentBlendProfile, PerBoneSampleData, PosesToEvaluate, SkeletonRemapping, OutAnimationPoseData);
+			}
+			else
+			{
+				FAnimationRuntime::BlendPosesTogetherPerBone(FilteredPoses, FilteredCurve, FilteredAttributes, CurrentBlendProfile, PerBoneSampleData, PosesToEvaluate, OutAnimationPoseData);
+			}
 		}
 		else
 		{
diff --git a/Engine/Source/Runtime/Engine/Classes/Animation/BlendProfile.h b/Engine/Source/Runtime/Engine/Classes/Animation/BlendProfile.h
--- a/Engine/Source/Runtime/Engine/Classes/Animation/BlendProfile.h
+++ b/Engine/Source/Runtime/Engine/Classes/Animation/BlendProfile.h
@@ -165,8 +165,16 @@ public:
 	 * @param OutDurationPerBone Must be sized to the number bones in the skeleton pose. It will be filled with the durations of each bone as setup in the blend profile editor.
 	 * @param Duration The duration of the blend.
 	 */
+	UE_DEPRECATED(5.4, "Please use the FillSkeletonBoneDurationsArray that takes a target skeleton as parameter.")
 	ENGINE_API void FillSkeletonBoneDurationsArray(TCustomBoneIndexArrayView<float, FSkeletonPoseBoneIndex> OutDurationPerBone, float Duration) const;
 
+	/** Fill an array of floats with the bone duration values. One for each bone in the skeleton pose.
+	 * @param OutDurationPerBone Must be sized to the number bones in the skeleton pose. It will be filled with the durations of each bone as setup in the blend profile editor.
+	 * @param Duration The duration of the blend.
+	 * @param TargetSkeleton The target skeleton we are working on. If this is a nullptr, the owning skeleton of the blend profile is assumed. This can be used when using skeleton compatibility.
+	 */
+	ENGINE_API void FillSkeletonBoneDurationsArray(TCustomBoneIndexArrayView<float, FSkeletonPoseBoneIndex> OutDurationPerBone, float Duration, const USkeleton* TargetSkeleton) const;
+
 	// IInterpolationIndexProvider
 	ENGINE_API virtual int32 GetPerBoneInterpolationIndex(const FCompactPoseBoneIndex& InCompactPoseBoneIndex, const FBoneContainer& BoneContainer, const IInterpolationIndexProvider::FPerBoneInterpolationData* Data) const override;
 	// End IInterpolationIndexProvider
diff --git a/Engine/Source/Runtime/Engine/Private/Animation/AnimNode_DeadBlending.cpp b/Engine/Source/Runtime/Engine/Private/Animation/AnimNode_DeadBlending.cpp
--- a/Engine/Source/Runtime/Engine/Private/Animation/AnimNode_DeadBlending.cpp
+++ b/Engine/Source/Runtime/Engine/Private/Animation/AnimNode_DeadBlending.cpp
@@ -761,13 +761,14 @@ void FAnimNode_DeadBlending::Evaluate_AnyThread(FPoseContext& Output)
 		InertializationDuration = BlendTimeMultiplier * RequestQueue[ShortestRequestIdx].Duration;
 		InertializationDurationPerBone.Init(InertializationDuration, NumSkeletonBones);
 
+		const USkeleton* TargetSkeleton = Output.AnimInstanceProxy->GetRequiredBones().GetSkeletonAsset();
 		if (RequestQueue[ShortestRequestIdx].BlendProfile)
 		{
-			RequestQueue[ShortestRequestIdx].BlendProfile->FillSkeletonBoneDurationsArray(InertializationDurationPerBone, InertializationDuration);
+			RequestQueue[ShortestRequestIdx].BlendProfile->FillSkeletonBoneDurationsArray(InertializationDurationPerBone, InertializationDuration, TargetSkeleton);
 		}
 		else if (DefaultBlendProfile)
 		{
-			DefaultBlendProfile->FillSkeletonBoneDurationsArray(InertializationDurationPerBone, InertializationDuration);
+			DefaultBlendProfile->FillSkeletonBoneDurationsArray(InertializationDurationPerBone, InertializationDuration, TargetSkeleton);
 		}
 
 		// Cache the maximum duration across all bones (so we know when to deactivate the inertialization request)
diff --git a/Engine/Source/Runtime/Engine/Private/Animation/AnimNode_Inertialization.cpp b/Engine/Source/Runtime/Engine/Private/Animation/AnimNode_Inertialization.cpp
--- a/Engine/Source/Runtime/Engine/Private/Animation/AnimNode_Inertialization.cpp
+++ b/Engine/Source/Runtime/Engine/Private/Animation/AnimNode_Inertialization.cpp
@@ -290,7 +290,8 @@ void FAnimNode_Inertialization::Evaluate_AnyThread(FPoseContext& Output)
 		
 		const int32 NumSkeletonBones = UE::Anim::GetNumSkeletonBones(Output.AnimInstanceProxy->GetRequiredBones());
 
-		auto FillSkeletonBoneDurationsArray = [this, NumSkeletonBones](auto& DurationPerBone, float Duration, const UBlendProfile* BlendProfile) {
+		const USkeleton* TargetSkeleton = Output.AnimInstanceProxy->GetRequiredBones().GetSkeletonAsset();
+		auto FillSkeletonBoneDurationsArray = [this, NumSkeletonBones, TargetSkeleton](auto& DurationPerBone, float Duration, const UBlendProfile* BlendProfile) {
 			if (BlendProfile == nullptr)
 			{
 				BlendProfile = DefaultBlendProfile;
@@ -299,7 +300,7 @@ void FAnimNode_Inertialization::Evaluate_AnyThread(FPoseContext& Output)
 			if (BlendProfile != nullptr)
 			{
 				DurationPerBone.SetNum(NumSkeletonBones);
-				BlendProfile->FillSkeletonBoneDurationsArray(DurationPerBone, Duration);
+				BlendProfile->FillSkeletonBoneDurationsArray(DurationPerBone, Duration, TargetSkeleton);
 			}
 			else
 			{
diff --git a/Engine/Source/Runtime/Engine/Private/Animation/AnimNode_StateMachine.cpp b/Engine/Source/Runtime/Engine/Private/Animation/AnimNode_StateMachine.cpp
--- a/Engine/Source/Runtime/Engine/Private/Animation/AnimNode_StateMachine.cpp
+++ b/Engine/Source/Runtime/Engine/Private/Animation/AnimNode_StateMachine.cpp
@@ -15,6 +15,8 @@
 #include "Logging/TokenizedMessage.h"
 #include "Animation/AnimInertializationSyncScope.h"
 #include "Animation/AnimNode_StateResult.h"
+#include "Animation/SkeletonRemapping.h"
+#include "Animation/SkeletonRemappingRegistry.h"
 
 #include UE_INLINE_GENERATED_CPP_BY_NAME(AnimNode_StateMachine)
 
@@ -990,20 +992,42 @@ void FAnimNode_StateMachine::EvaluateTransitionStandardBlendInternal(FPoseContex
 	const ScalarRegister VPreviousWeight(1.0f - Transition.Alpha);
 	const ScalarRegister VWeight(Transition.Alpha);
 
-	// If we have a blend profile we need to blend per bone
-	if(Transition.BlendProfile)
+	// If we have a blend profile we need to blend per bone.
+	if (Transition.BlendProfile)
 	{
 		const FBoneContainer& RequiredBones = Output.AnimInstanceProxy->GetRequiredBones();
 		TSharedPtr<IInterpolationIndexProvider::FPerBoneInterpolationData> Data = Transition.BlendProfile->GetPerBoneInterpolationData(Output.AnimInstanceProxy->GetSkeleton());
-		for(FCompactPoseBoneIndex BoneIndex : Output.Pose.ForEachBoneIndex())
+
+		// If we have some skeleton remapping and the source data comes from another skeleton.
+		// This is a slightly slower path, so we made two branches, one with remapping and one without.
+		const FSkeletonRemapping& Remapping = UE::Anim::FSkeletonRemappingRegistry::Get().GetRemapping(Transition.BlendProfile->OwningSkeleton, Output.AnimInstanceProxy->GetSkeleton());
+		if (Remapping.IsValid())
+		{
+			for (const FCompactPoseBoneIndex TargetBoneIndex : Output.Pose.ForEachBoneIndex())
+			{
+				const FSkeletonPoseBoneIndex SourceSkelBoneIndex(Remapping.GetSourceSkeletonBoneIndex(TargetBoneIndex.GetInt()));
+				const FCompactPoseBoneIndex SourceBoneIndex = FCompactPoseBoneIndex(RequiredBones.GetCompactPoseIndexFromSkeletonPoseIndex(SourceSkelBoneIndex));
+				const int32 PerBoneIndex = (SourceBoneIndex != INDEX_NONE) ? Transition.BlendProfile->GetPerBoneInterpolationIndex(SourceBoneIndex, RequiredBones, Data.Get()) : INDEX_NONE;
+
+				// Use defined per-bone scale if the bone has a scale specified in the blend profile.
+				const ScalarRegister FirstWeight = (PerBoneIndex != INDEX_NONE) ? ScalarRegister(Transition.StateBlendData[1].PerBoneBlendData[PerBoneIndex]) : ScalarRegister(VPreviousWeight);
+				const ScalarRegister SecondWeight = (PerBoneIndex != INDEX_NONE) ? ScalarRegister(Transition.StateBlendData[0].PerBoneBlendData[PerBoneIndex]) : ScalarRegister(VWeight);
+				Output.Pose[TargetBoneIndex] = PreviousStateResult.Pose[TargetBoneIndex] * FirstWeight;
+				Output.Pose[TargetBoneIndex].AccumulateWithShortestRotation(NextStateResult.Pose[TargetBoneIndex], SecondWeight);
+			}
+		}
+		else // There is no skeleton remapping or we are using the same skeleton as the source.
 		{
-			const int32 PerBoneIndex = Transition.BlendProfile->GetPerBoneInterpolationIndex(BoneIndex, RequiredBones, Data.Get());
+			for (const FCompactPoseBoneIndex TargetBoneIndex : Output.Pose.ForEachBoneIndex())
+			{
+				const int32 PerBoneIndex = Transition.BlendProfile->GetPerBoneInterpolationIndex(TargetBoneIndex, RequiredBones, Data.Get());
 
-			// Use defined per-bone scale if the bone has a scale specified in the blend profile
-			ScalarRegister FirstWeight = PerBoneIndex != INDEX_NONE ? ScalarRegister(Transition.StateBlendData[1].PerBoneBlendData[PerBoneIndex]) : ScalarRegister(VPreviousWeight);
-			ScalarRegister SecondWeight = PerBoneIndex != INDEX_NONE ? ScalarRegister(Transition.StateBlendData[0].PerBoneBlendData[PerBoneIndex]) : ScalarRegister(VWeight);
-			Output.Pose[BoneIndex] = PreviousStateResult.Pose[BoneIndex] * FirstWeight;
-			Output.Pose[BoneIndex].AccumulateWithShortestRotation(NextStateResult.Pose[BoneIndex], SecondWeight);
+				// Use defined per-bone scale if the bone has a scale specified in the blend profile.
+				const ScalarRegister FirstWeight = (PerBoneIndex != INDEX_NONE) ? ScalarRegister(Transition.StateBlendData[1].PerBoneBlendData[PerBoneIndex]) : ScalarRegister(VPreviousWeight);
+				const ScalarRegister SecondWeight = (PerBoneIndex != INDEX_NONE) ? ScalarRegister(Transition.StateBlendData[0].PerBoneBlendData[PerBoneIndex]) : ScalarRegister(VWeight);
+				Output.Pose[TargetBoneIndex] = PreviousStateResult.Pose[TargetBoneIndex] * FirstWeight;
+				Output.Pose[TargetBoneIndex].AccumulateWithShortestRotation(NextStateResult.Pose[TargetBoneIndex], SecondWeight);
+			}
 		}
 	}
 	else
diff --git a/Engine/Source/Runtime/Engine/Private/Animation/AnimationRuntime.cpp b/Engine/Source/Runtime/Engine/Private/Animation/AnimationRuntime.cpp
--- a/Engine/Source/Runtime/Engine/Private/Animation/AnimationRuntime.cpp
+++ b/Engine/Source/Runtime/Engine/Private/Animation/AnimationRuntime.cpp
@@ -721,6 +721,55 @@ void FAnimationRuntime::BlendPosesTogetherPerBone(TArrayView<const FCompactPose>
 	}
 }
 
+void FAnimationRuntime::BlendPosesTogetherPerBoneRemapped(TArrayView<const FCompactPose> SourcePoses, TArrayView<const FBlendedCurve> SourceCurves, TArrayView<const UE::Anim::FStackAttributeContainer> SourceAttributes, const IInterpolationIndexProvider* InterpolationIndexProvider, TArrayView<const FBlendSampleData> BlendSampleDataCache, TArrayView<const int32> BlendSampleDataCacheIndices, const FSkeletonRemapping& SkeletonRemapping, /*out*/ FAnimationPoseData& OutAnimationPoseData)
+{
+	check(SourcePoses.Num() > 0);
+	check(SkeletonRemapping.IsValid());	// If this fails, you most likely want to use BlendPosesTogetherPerBone instead.
+
+	FCompactPose& OutPose = OutAnimationPoseData.GetPose();
+	FBlendedCurve& OutCurve = OutAnimationPoseData.GetCurve();
+	UE::Anim::FStackAttributeContainer& OutAttributes = OutAnimationPoseData.GetAttributes();
+
+	const FBoneContainer& RequiredBones = OutPose.GetBoneContainer();
+	TArray<int32> PerBoneIndices;
+	PerBoneIndices.AddUninitialized(OutPose.GetNumBones());
+	TSharedPtr<IInterpolationIndexProvider::FPerBoneInterpolationData> Data = InterpolationIndexProvider->GetPerBoneInterpolationData(OutPose.GetBoneContainer().GetSkeletonAsset());
+	for (FCompactPoseBoneIndex BoneIndex : OutPose.ForEachBoneIndex())
+	{
+		const FSkeletonPoseBoneIndex SourceSkelBoneIndex(SkeletonRemapping.GetSourceSkeletonBoneIndex(BoneIndex.GetInt()));
+		const FCompactPoseBoneIndex SourceBoneIndex = FCompactPoseBoneIndex(RequiredBones.GetCompactPoseIndexFromSkeletonPoseIndex(SourceSkelBoneIndex));
+		PerBoneIndices[BoneIndex.GetInt()] = (SourceBoneIndex != INDEX_NONE) ? InterpolationIndexProvider->GetPerBoneInterpolationIndex(SourceBoneIndex, RequiredBones, Data.Get()) : INDEX_NONE;
+	}
+
+	BlendPosePerBone<ETransformBlendMode::Overwrite>(PerBoneIndices, BlendSampleDataCache[BlendSampleDataCacheIndices[0]], OutPose, SourcePoses[0]);
+
+	for (int32 i = 1; i < SourcePoses.Num(); ++i)
+	{
+		BlendPosePerBone<ETransformBlendMode::Accumulate>(PerBoneIndices, BlendSampleDataCache[BlendSampleDataCacheIndices[i]], OutPose, SourcePoses[i]);
+	}
+
+	// Ensure that all of the resulting rotations are normalized
+	OutPose.NormalizeRotations();
+
+	if (SourceCurves.Num() > 0)
+	{
+		TArray<float, TInlineAllocator<16>> SourceWeights;
+		SourceWeights.AddUninitialized(BlendSampleDataCacheIndices.Num());
+		for (int32 CacheIndex = 0; CacheIndex < BlendSampleDataCacheIndices.Num(); ++CacheIndex)
+		{
+			SourceWeights[CacheIndex] = BlendSampleDataCache[BlendSampleDataCacheIndices[CacheIndex]].TotalWeight;
+		}
+
+		BlendCurves(SourceCurves, SourceWeights, OutCurve);		
+	}	
+
+	if (SourceAttributes.Num() > 0)
+	{
+		UE::Anim::Attributes::BlendAttributesPerBone(SourceAttributes, PerBoneIndices, BlendSampleDataCache, BlendSampleDataCacheIndices, OutAttributes);
+	}
+}
+
+
 void FAnimationRuntime::BlendPosesTogetherPerBoneInMeshSpace(
 	const TArrayView<FCompactPose>           SourcePoses,
 	const TArrayView<const FBlendedCurve>    SourceCurves,
diff --git a/Engine/Source/Runtime/Engine/Private/Animation/BlendProfile.cpp b/Engine/Source/Runtime/Engine/Private/Animation/BlendProfile.cpp
--- a/Engine/Source/Runtime/Engine/Private/Animation/BlendProfile.cpp
+++ b/Engine/Source/Runtime/Engine/Private/Animation/BlendProfile.cpp
@@ -3,6 +3,8 @@
 #include "Animation/BlendProfile.h"
 #include "AlphaBlend.h"
 #include "Animation/AnimationAsset.h"
+#include "Animation/SkeletonRemapping.h"
+#include "Animation/SkeletonRemappingRegistry.h"
 
 #include UE_INLINE_GENERATED_CPP_BY_NAME(BlendProfile)
 
@@ -249,28 +251,49 @@ void UBlendProfile::FillBoneScalesArray(TArray<float>& OutBoneBlendProfileFactor
 
 	// Overwrite the values of the bones that are inside the blend profile.
 	// Since the bones in the blend profile are stored as skeleton indices we need to remap them into our compact pose.
-	for (int32 Index = 0; Index < ProfileEntries.Num(); Index++)
+	const FSkeletonRemapping& SkeletonRemapping = UE::Anim::FSkeletonRemappingRegistry::Get().GetRemapping(OwningSkeleton, BoneContainer.GetSkeletonAsset());
+	if (SkeletonRemapping.IsValid())
 	{
-		const int32 SkeletonBoneIndex = ProfileEntries[Index].BoneReference.BoneIndex;
-		const FCompactPoseBoneIndex PoseBoneIndex = BoneContainer.GetCompactPoseIndexFromSkeletonIndex(SkeletonBoneIndex);
-		if (PoseBoneIndex.IsValid())
+		for (int32 Index = 0; Index < ProfileEntries.Num(); Index++)
 		{
-			OutBoneBlendProfileFactors[PoseBoneIndex.GetInt()] = GetEntryBlendScale(Index);
+			const int32 SkeletonBoneIndex = ProfileEntries[Index].BoneReference.BoneIndex;
+			const int32 TargetBoneIndex = SkeletonRemapping.GetTargetSkeletonBoneIndex(SkeletonBoneIndex);
+			const FCompactPoseBoneIndex PoseBoneIndex = BoneContainer.GetCompactPoseIndexFromSkeletonIndex(TargetBoneIndex);
+			if (PoseBoneIndex.IsValid())
+			{
+				OutBoneBlendProfileFactors[PoseBoneIndex.GetInt()] = GetEntryBlendScale(Index);
+			}
+		}
+	}
+	else // We don't use skeleton remapping, slightly more optimized.
+	{
+		for (int32 Index = 0; Index < ProfileEntries.Num(); Index++)
+		{
+			const int32 SkeletonBoneIndex = ProfileEntries[Index].BoneReference.BoneIndex;
+			const FCompactPoseBoneIndex PoseBoneIndex = BoneContainer.GetCompactPoseIndexFromSkeletonIndex(SkeletonBoneIndex);
+			if (PoseBoneIndex.IsValid())
+			{
+				OutBoneBlendProfileFactors[PoseBoneIndex.GetInt()] = GetEntryBlendScale(Index);
+			}
 		}
 	}
 }
 
 void UBlendProfile::FillSkeletonBoneDurationsArray(TCustomBoneIndexArrayView<float, FSkeletonPoseBoneIndex> OutDurationPerBone, float Duration) const
+{
+	FillSkeletonBoneDurationsArray(OutDurationPerBone, Duration, OwningSkeleton);
+}
+
+void UBlendProfile::FillSkeletonBoneDurationsArray(TCustomBoneIndexArrayView<float, FSkeletonPoseBoneIndex> OutDurationPerBone, float Duration, const USkeleton* TargetSkeleton) const
 {
 	check(OwningSkeleton != nullptr);
-	const FReferenceSkeleton& RefSkeleton = OwningSkeleton->GetReferenceSkeleton();
-	const int32 NumSkeletonBones = RefSkeleton.GetNum();
-	if (!ensureMsgf(OutDurationPerBone.Num() == NumSkeletonBones, 
-		TEXT("UBlendProfile::FillSkeletonBoneDurationsArray received OutDurationPerBone(%d) vs NumSkeletonBones(%d) for %s"), OutDurationPerBone.Num(), NumSkeletonBones, *GetNameSafe(OwningSkeleton)))
+	if (TargetSkeleton == nullptr)
 	{
-		return;
+		TargetSkeleton = OwningSkeleton;
 	}
 
+	const FSkeletonRemapping& SkeletonRemapping = UE::Anim::FSkeletonRemappingRegistry::Get().GetRemapping(OwningSkeleton, TargetSkeleton);
+
 	for(float& BoneDuration: OutDurationPerBone)
 	{
 		BoneDuration = Duration;
@@ -280,22 +303,47 @@ void UBlendProfile::FillSkeletonBoneDurationsArray(TCustomBoneIndexArrayView<flo
 	{
 		case EBlendProfileMode::TimeFactor:
 		{
-			for (const FBlendProfileBoneEntry& Entry : ProfileEntries)
+			if (SkeletonRemapping.IsValid())
 			{
-				const FSkeletonPoseBoneIndex SkeletonBoneIndex(Entry.BoneReference.BoneIndex);
-				OutDurationPerBone[SkeletonBoneIndex] *= Entry.BlendScale;
+				for (const FBlendProfileBoneEntry& Entry : ProfileEntries)
+				{
+					const FSkeletonPoseBoneIndex SkeletonBoneIndex(SkeletonRemapping.GetTargetSkeletonBoneIndex(Entry.BoneReference.BoneIndex));
+					OutDurationPerBone[SkeletonBoneIndex] *= Entry.BlendScale;
+				}
+			}
+			else
+			{
+				for (const FBlendProfileBoneEntry& Entry : ProfileEntries)
+				{
+					const FSkeletonPoseBoneIndex SkeletonBoneIndex(Entry.BoneReference.BoneIndex);
+					OutDurationPerBone[SkeletonBoneIndex] *= Entry.BlendScale;
+				}
 			}
 		}
 		break;
 
 		case EBlendProfileMode::WeightFactor:
 		{
-			for (const FBlendProfileBoneEntry& Entry : ProfileEntries)
+			if (SkeletonRemapping.IsValid())
+			{
+				for (const FBlendProfileBoneEntry& Entry : ProfileEntries)
+				{
+					const FSkeletonPoseBoneIndex SkeletonBoneIndex(SkeletonRemapping.GetTargetSkeletonBoneIndex(Entry.BoneReference.BoneIndex));
+					if (Entry.BlendScale > UE_SMALL_NUMBER)
+					{
+						OutDurationPerBone[SkeletonBoneIndex] /= Entry.BlendScale;
+					}
+				}
+			}
+			else
 			{
-				const FSkeletonPoseBoneIndex SkeletonBoneIndex(Entry.BoneReference.BoneIndex);
-				if (Entry.BlendScale > UE_SMALL_NUMBER)
+				for (const FBlendProfileBoneEntry& Entry : ProfileEntries)
 				{
-					OutDurationPerBone[SkeletonBoneIndex] /= Entry.BlendScale;
+					const FSkeletonPoseBoneIndex SkeletonBoneIndex(Entry.BoneReference.BoneIndex);
+					if (Entry.BlendScale > UE_SMALL_NUMBER)
+					{
+						OutDurationPerBone[SkeletonBoneIndex] /= Entry.BlendScale;
+					}
 				}
 			}
 		}
diff --git a/Engine/Source/Runtime/Engine/Private/Animation/BlendSpace.cpp b/Engine/Source/Runtime/Engine/Private/Animation/BlendSpace.cpp
--- a/Engine/Source/Runtime/Engine/Private/Animation/BlendSpace.cpp
+++ b/Engine/Source/Runtime/Engine/Private/Animation/BlendSpace.cpp
@@ -12,6 +12,8 @@
 #include "Animation/AnimationPoseData.h"
 #include "Animation/BlendSpaceHelpers.h"
 #include "Animation/BlendSpace1DHelpers.h"
+#include "Animation/SkeletonRemapping.h"
+#include "Animation/SkeletonRemappingRegistry.h"
 #include "Math/UnrealMathUtility.h"
 #include "UObject/FrameworkObjectVersion.h"
 #include "UObject/UE5MainStreamObjectVersion.h"
@@ -799,17 +801,30 @@ int32 UBlendSpace::GetPerBoneInterpolationIndex(
 	}
 	const TArray<FSortedPerBoneInterpolation>& SortedData = static_cast<const FSortedPerBoneInterpolationData*>(Data)->Data;
 
+	const USkeleton* SourceSkeleton = GetSkeleton();
+	const FSkeletonRemapping& SkeletonRemapping = UE::Anim::FSkeletonRemappingRegistry().Get().GetRemapping(SourceSkeleton, RequiredBones.GetSkeletonAsset());
+
 	for (int32 Iter = 0; Iter < SortedData.Num(); ++Iter)
 	{
 		const FPerBoneInterpolation& PerBoneInterpolation = SortedData[Iter].PerBoneBlend;
-		int32 OriginalIndex = SortedData[Iter].OriginalIndex;
+		const int32 OriginalIndex = SortedData[Iter].OriginalIndex;
 		const FBoneReference& SmoothedBone = PerBoneInterpolation.BoneReference;
-		FCompactPoseBoneIndex SmoothedBoneCompactPoseIndex = 
-			RequiredBones.GetCompactPoseIndexFromSkeletonPoseIndex(SmoothedBone.GetSkeletonPoseIndex(RequiredBones));
+
+		FSkeletonPoseBoneIndex SkelBoneIndex = SmoothedBone.GetSkeletonPoseIndex(RequiredBones);
+
+		// Remap to the target skeleton, using skeleton remapping, as we might be applying this blend space onto another skeleton than the asset was created for.
+		if (SkeletonRemapping.IsValid())
+		{
+			const int32 RemappedSkelBoneIndex = SkeletonRemapping.GetTargetSkeletonBoneIndex(SkelBoneIndex.GetInt());
+			SkelBoneIndex = FSkeletonPoseBoneIndex(RemappedSkelBoneIndex);
+		}
+
+		const FCompactPoseBoneIndex SmoothedBoneCompactPoseIndex = RequiredBones.GetCompactPoseIndexFromSkeletonPoseIndex(SkelBoneIndex);
 		if (SmoothedBoneCompactPoseIndex == InCompactPoseBoneIndex)
 		{
 			return OriginalIndex;
 		}
+
 		// BoneIsChildOf returns true if InCompactPoseBoneIndex is a child of SmoothedBoneCompactPoseIndex. 
 		if (SmoothedBoneCompactPoseIndex != INDEX_NONE && 
 			RequiredBones.BoneIsChildOf(InCompactPoseBoneIndex, SmoothedBoneCompactPoseIndex))
@@ -1591,10 +1606,10 @@ void UBlendSpace::InitializePerBoneBlend()
 	{
 		PerBoneBlendValues.Empty();
 
-		UBlendProfile* BlendProfile = PerBoneBlendProfile.BlendProfile.Get();
+		const UBlendProfile* BlendProfile = PerBoneBlendProfile.BlendProfile.Get();
 		if (BlendProfile)
 		{
-			int32 NumBlendEntries = BlendProfile->GetNumBlendEntries();
+			const int32 NumBlendEntries = BlendProfile->GetNumBlendEntries();
 			for (int32 EntryIndex = 0; EntryIndex < NumBlendEntries; ++EntryIndex)
 			{
 				const FBlendProfileBoneEntry& BoneEntry = BlendProfile->GetEntry(EntryIndex);
diff --git a/Engine/Source/Runtime/Engine/Public/AnimationRuntime.h b/Engine/Source/Runtime/Engine/Public/AnimationRuntime.h
--- a/Engine/Source/Runtime/Engine/Public/AnimationRuntime.h
+++ b/Engine/Source/Runtime/Engine/Public/AnimationRuntime.h
@@ -14,6 +14,7 @@
 #include "Animation/AnimationAsset.h"
 #include "Animation/AnimCurveTypes.h"
 #include "Animation/AnimSequenceBase.h"
+#include "Animation/SkeletonRemapping.h"
 #include "Components/SkinnedMeshComponent.h"
 #include "BonePose.h"
 #include "Containers/ArrayView.h"
@@ -255,6 +256,17 @@ public:
 		TArrayView<const FBlendSampleData> BlendSampleDataCache,
 		/*out*/ FAnimationPoseData& OutAnimationPoseData);
 
+	static ENGINE_API void BlendPosesTogetherPerBoneRemapped(
+		TArrayView<const FCompactPose> SourcePoses, 
+		TArrayView<const FBlendedCurve> SourceCurves, 
+		TArrayView<const UE::Anim::FStackAttributeContainer> SourceAttributes, 
+		const IInterpolationIndexProvider* InterpolationIndexProvider,
+		TArrayView<const FBlendSampleData> BlendSampleDataCache,
+		TArrayView<const int32> BlendSampleDataCacheIndices, 
+		const FSkeletonRemapping& SkeletonRemapping,
+		/*out*/ FAnimationPoseData& OutAnimationPoseData);
+
+
 	/**
 	* Blends together a set of poses, each with a given weight.
 	* This function is for BlendSpace per bone blending. BlendSampleDataCache contains the weight information 
