World Partition: merged similar code into a function to register actor descriptors editor streaming.

diff --git a/Engine/Source/Runtime/Engine/Private/WorldPartition/WorldPartition.cpp b/Engine/Source/Runtime/Engine/Private/WorldPartition/WorldPartition.cpp
--- a/Engine/Source/Runtime/Engine/Private/WorldPartition/WorldPartition.cpp
+++ b/Engine/Source/Runtime/Engine/Private/WorldPartition/WorldPartition.cpp
@@ -615,10 +615,6 @@ void UWorldPartition::Initialize(UWorld* InWorld, const FTransform& InTransform)
 		FString SourceWorldPath, RemappedWorldPath;
 		const bool bIsInstanced = OuterWorld->GetSoftObjectPathMapping(SourceWorldPath, RemappedWorldPath);
 
-		// Follow the world's streaming enabled value most of the times, except:
-		//	- World is instanced and from a Level Instance that supports partial loading
-		const bool bIsStreamingEnabled = IsStreamingEnabledInEditor();
-
 		if (bIsInstanced)
 		{
 			InstancingContext.AddPackageMapping(PackageName, LevelPackage->GetFName());
@@ -633,43 +629,31 @@ void UWorldPartition::Initialize(UWorld* InWorld, const FTransform& InTransform)
 		FContainerRegistrationParams ContainerInitParams(PackageName);
 		ActorDescContainer = RegisterActorDescContainer(ContainerInitParams);
 
+		// If a Valid Actor references an Invalid Actor:
+		// Make sure Invalid Actors do not load their imports (ex: outer non instanced world).
+		if (bIsInstanced)
 		{
-			// If a Valid Actor references an Invalid Actor:
-			// Make sure Invalid Actors do not load their imports (ex: outer non instanced world).
-			if (bIsInstanced)
+			TRACE_CPUPROFILER_EVENT_SCOPE(ContainerInstancing);
+			for (const FAssetData& InvalidActor : ActorDescContainer->InvalidActors)
 			{
-				for (const FAssetData& InvalidActor : ActorDescContainer->InvalidActors)
-				{
-					InstancingContext.AddPackageMapping(InvalidActor.PackageName, NAME_None);
-				}
+				InstancingContext.AddPackageMapping(InvalidActor.PackageName, NAME_None);
 			}
 
-			TRACE_CPUPROFILER_EVENT_SCOPE(UActorDescContainer::Hash);
 			for (FActorDescContainerCollection::TIterator<> ActorDescIterator(this); ActorDescIterator; ++ActorDescIterator)
 			{
-				if (bIsInstanced)
-				{
-					const FString LongActorPackageName = ActorDescIterator->GetActorPackage().ToString();
-					const FString InstancedName = ULevel::GetExternalActorPackageInstanceName(LevelPackage->GetName(), LongActorPackageName);
-
-					InstancingContext.AddPackageMapping(*LongActorPackageName, *InstancedName);
+				const FString LongActorPackageName = ActorDescIterator->GetActorPackage().ToString();
+				const FString InstancedName = ULevel::GetExternalActorPackageInstanceName(LevelPackage->GetName(), LongActorPackageName);
 
-					ActorDescIterator->TransformInstance(SourceWorldPath, RemappedWorldPath);
-				}
-
-				ActorDescIterator->bIsForcedNonSpatiallyLoaded = !bIsStreamingEnabled;
-
-				if (ForceLoadedActors)
-				{
-					ForceLoadedActorGuids.Add(ActorDescIterator->GetGuid());
-				}
+				InstancingContext.AddPackageMapping(*LongActorPackageName, *InstancedName);
 
-				if (bIsEditor && !bIsCooking)
-				{
-					HashActorDesc(*ActorDescIterator);
-				}
+				ActorDescIterator->TransformInstance(SourceWorldPath, RemappedWorldPath);
 			}
 		}
+
+		ForEachActorDescContainer([this, bIsEditor, bIsCooking](UActorDescContainer* ActorDescContainer)
+		{
+			InitializeActorDescContainerEditorStreaming(ActorDescContainer, bIsEditor && !bIsCooking);
+		});
 	}
 #endif
 
@@ -1358,6 +1342,35 @@ bool UWorldPartition::GetInstancingContext(const FLinkerInstancingContext*& OutI
 	}
 	return false;
 }
+
+void UWorldPartition::InitializeActorDescContainerEditorStreaming(UActorDescContainer* InActorDescContainer, bool bInHashActorDescs)
+{
+	TRACE_CPUPROFILER_EVENT_SCOPE(InitializeActorDescContainerEditorStreaming);
+
+	const bool bIsStreamingEnabled = IsStreamingEnabledInEditor();
+
+	TArray<FGuid> ForceLoadedActorGuids;
+	for (UActorDescContainer::TIterator<> It(InActorDescContainer); It; ++It)
+	{
+		It->bIsForcedNonSpatiallyLoaded = !bIsStreamingEnabled;
+
+		if (ForceLoadedActors)
+		{
+			ForceLoadedActorGuids.Add(It->GetGuid());
+		}
+
+		if (bInHashActorDescs)
+		{
+			HashActorDesc(*It);
+		}
+	}
+
+	if (ForceLoadedActorGuids.Num())
+	{
+		check(ForceLoadedActors);
+		ForceLoadedActors->AddActors(ForceLoadedActorGuids);
+	}
+}
 #endif
 
 const FTransform& UWorldPartition::GetInstanceTransform() const
@@ -2034,6 +2047,8 @@ void UWorldPartition::AppendAssetRegistryTags(TArray<FAssetRegistryTag>& OutTags
 
 UActorDescContainer* UWorldPartition::RegisterActorDescContainer(const FContainerRegistrationParams& InRegistrationParameters)
 {
+	TRACE_CPUPROFILER_EVENT_SCOPE(UWorldPartition::RegisterActorDescContainer);
+
 	if (!Contains(InRegistrationParameters.PackageName))
 	{	
 		UActorDescContainer::FInitializeParams ContainerInitParams(GetWorld(), InRegistrationParameters.PackageName);
@@ -2068,9 +2083,6 @@ UActorDescContainer* UWorldPartition::RegisterActorDescContainer(const FContaine
 
 		AddContainer(ContainerToRegister);
 
-		// @todo_ow: there is some redundancy here with UWorldPartition::Initialize, at some point we will also need to support instancing of CB containers to 
-		// support CB/EDL in Level Instances so this code should probably be shared with the code in UWorldPartition::Initialize
-		// See https://jira.it.epicgames.com/browse/UE-195953
 		if (IsInitialized() && EditorHash != nullptr)
 		{
 			FWorldPartitionReference WDLReference;
@@ -2083,25 +2095,7 @@ UActorDescContainer* UWorldPartition::RegisterActorDescContainer(const FContaine
 				}
 			}
 
-			const bool bIsStreamingEnabled = IsStreamingEnabledInEditor();
-			TArray<FGuid> ForceLoadedActorGuids;
-
-			for (UActorDescContainer::TIterator<> It(ContainerToRegister); It; ++It)
-			{
-				It->bIsForcedNonSpatiallyLoaded = !bIsStreamingEnabled;
-
-				if (ForceLoadedActors)
-				{
-					ForceLoadedActorGuids.Add(It->GetGuid());
-				}
-			
-				HashActorDesc(*It);
-			}
-
-			if (ForceLoadedActors && ForceLoadedActorGuids.Num() > 0)
-			{
-				ForceLoadedActors->AddActors(ForceLoadedActorGuids);
-			}
+			InitializeActorDescContainerEditorStreaming(ContainerToRegister, true);
 		}
 
 		OnActorDescContainerRegistered.Broadcast(ContainerToRegister);
diff --git a/Engine/Source/Runtime/Engine/Public/WorldPartition/WorldPartition.h b/Engine/Source/Runtime/Engine/Public/WorldPartition/WorldPartition.h
--- a/Engine/Source/Runtime/Engine/Public/WorldPartition/WorldPartition.h
+++ b/Engine/Source/Runtime/Engine/Public/WorldPartition/WorldPartition.h
@@ -198,6 +198,8 @@ private:
 	ENGINE_API void OnActorDescUpdated(FWorldPartitionActorDesc* ActorDesc);
 
 	ENGINE_API bool GetInstancingContext(const FLinkerInstancingContext*& OutInstancingContext) const;
+
+	ENGINE_API void InitializeActorDescContainerEditorStreaming(UActorDescContainer* InActorDescContainer, bool bInHashActorDescs);
 #endif
 
 public:
