Fix incremental mac/linux builds.

diff --git a/Engine/Source/Runtime/Engine/Private/WorldPartition/WorldPartition.cpp b/Engine/Source/Runtime/Engine/Private/WorldPartition/WorldPartition.cpp
--- a/Engine/Source/Runtime/Engine/Private/WorldPartition/WorldPartition.cpp
+++ b/Engine/Source/Runtime/Engine/Private/WorldPartition/WorldPartition.cpp
@@ -650,9 +650,9 @@ void UWorldPartition::Initialize(UWorld* InWorld, const FTransform& InTransform)
 			}
 		}
 
-		ForEachActorDescContainer([this, bIsEditor, bIsCooking](UActorDescContainer* ActorDescContainer)
+		ForEachActorDescContainer([this, bIsEditor, bIsCooking](UActorDescContainer* InActorDescContainer)
 		{
-			InitializeActorDescContainerEditorStreaming(ActorDescContainer, bIsEditor && !bIsCooking);
+			InitializeActorDescContainerEditorStreaming(InActorDescContainer, bIsEditor && !bIsCooking);
 		});
 	}
 #endif
