Fix a crash when a static mesh doesn't have render data

diff --git a/Engine/Source/Runtime/Engine/Private/StaticMesh.cpp b/Engine/Source/Runtime/Engine/Private/StaticMesh.cpp
--- a/Engine/Source/Runtime/Engine/Private/StaticMesh.cpp
+++ b/Engine/Source/Runtime/Engine/Private/StaticMesh.cpp
@@ -6515,44 +6515,47 @@ void UStaticMesh::ExecutePostLoadInternal(FStaticMeshPostLoadContext& Context)
 	//This will reuse everything valid from the just serialize SectionInfoMap.
 	FMeshSectionInfoMap TempOldSectionInfoMap = GetSectionInfoMap();
 	GetSectionInfoMap().Clear();
-	for (int32 LODResourceIndex = 0; LODResourceIndex < GetRenderData()->LODResources.Num(); ++LODResourceIndex)
+	if (GetRenderData())
 	{
-		FStaticMeshLODResources& LOD = GetRenderData()->LODResources[LODResourceIndex];
-		for (int32 SectionIndex = 0; SectionIndex < LOD.Sections.Num(); ++SectionIndex)
+		for (int32 LODResourceIndex = 0; LODResourceIndex < GetRenderData()->LODResources.Num(); ++LODResourceIndex)
 		{
-			if (TempOldSectionInfoMap.IsValidSection(LODResourceIndex, SectionIndex))
+			FStaticMeshLODResources& LOD = GetRenderData()->LODResources[LODResourceIndex];
+			for (int32 SectionIndex = 0; SectionIndex < LOD.Sections.Num(); ++SectionIndex)
 			{
-				FMeshSectionInfo Info = TempOldSectionInfoMap.Get(LODResourceIndex, SectionIndex);
-				if (GetStaticMaterials().IsValidIndex(Info.MaterialIndex))
+				if (TempOldSectionInfoMap.IsValidSection(LODResourceIndex, SectionIndex))
 				{
-					//Reuse the valid data that come from the serialize
-					GetSectionInfoMap().Set(LODResourceIndex, SectionIndex, Info);
+					FMeshSectionInfo Info = TempOldSectionInfoMap.Get(LODResourceIndex, SectionIndex);
+					if (GetStaticMaterials().IsValidIndex(Info.MaterialIndex))
+					{
+						//Reuse the valid data that come from the serialize
+						GetSectionInfoMap().Set(LODResourceIndex, SectionIndex, Info);
+					}
+					else
+					{
+						//Use the render data material index, but keep the flags (collision, shadow...)
+						const int32 MaterialIndex = LOD.Sections[SectionIndex].MaterialIndex;
+						if (GetStaticMaterials().IsValidIndex(MaterialIndex))
+						{
+							Info.MaterialIndex = MaterialIndex;
+							GetSectionInfoMap().Set(LODResourceIndex, SectionIndex, Info);
+						}
+					}
 				}
 				else
 				{
-					//Use the render data material index, but keep the flags (collision, shadow...)
+					//Create a new SectionInfoMap from the render data
 					const int32 MaterialIndex = LOD.Sections[SectionIndex].MaterialIndex;
 					if (GetStaticMaterials().IsValidIndex(MaterialIndex))
 					{
-						Info.MaterialIndex = MaterialIndex;
-						GetSectionInfoMap().Set(LODResourceIndex, SectionIndex, Info);
+						GetSectionInfoMap().Set(LODResourceIndex, SectionIndex, FMeshSectionInfo(MaterialIndex));
 					}
 				}
-			}
-			else
-			{
-				//Create a new SectionInfoMap from the render data
-				const int32 MaterialIndex = LOD.Sections[SectionIndex].MaterialIndex;
-				if (GetStaticMaterials().IsValidIndex(MaterialIndex))
+				//Make sure the OriginalSectionInfoMap has some information, the post load only add missing slot, this data should be set when importing/re-importing the asset
+				if (!GetOriginalSectionInfoMap().IsValidSection(LODResourceIndex, SectionIndex))
 				{
-					GetSectionInfoMap().Set(LODResourceIndex, SectionIndex, FMeshSectionInfo(MaterialIndex));
+					GetOriginalSectionInfoMap().Set(LODResourceIndex, SectionIndex, GetSectionInfoMap().Get(LODResourceIndex, SectionIndex));
 				}
 			}
-			//Make sure the OriginalSectionInfoMap has some information, the post load only add missing slot, this data should be set when importing/re-importing the asset
-			if (!GetOriginalSectionInfoMap().IsValidSection(LODResourceIndex, SectionIndex))
-			{
-				GetOriginalSectionInfoMap().Set(LODResourceIndex, SectionIndex, GetSectionInfoMap().Get(LODResourceIndex, SectionIndex));
-			}
 		}
 	}
 
