Fix invalid warning with Substrate material having several shading models.

diff --git a/Engine/Source/Runtime/Engine/Private/Materials/Material.cpp b/Engine/Source/Runtime/Engine/Private/Materials/Material.cpp
--- a/Engine/Source/Runtime/Engine/Private/Materials/Material.cpp
+++ b/Engine/Source/Runtime/Engine/Private/Materials/Material.cpp
@@ -4927,6 +4927,19 @@ void UMaterial::RebuildShadingModelField()
 					bSanitizeMaterial = true;
 				}
 			}
+			else if (SubstrateMaterialInfo.CountShadingModels() > 1 && MaterialDomain == MD_Surface)
+			{
+				// Case with SSS Profile or SSS MFP are already been handled by above cases. Simply fallback onto DefaultLit
+				if (BlendMode == EBlendMode::BLEND_Opaque || BlendMode == EBlendMode::BLEND_Masked)
+				{
+					SubstrateMaterialInfo.SetSingleShadingModel(SSM_DefaultLit);
+				}
+				else
+				{
+					// For transparent, we will fall back to use DefaultLit worst case with simple volumetric.
+					bSanitizeMaterial = true;
+				}
+			}
 			else if (SubstrateMaterialInfo.CountShadingModels() == 2 && SubstrateMaterialInfo.HasShadingModel(ESubstrateShadingModel::SSM_Decal))
 			{
 				// If material has SSM_Decal it has to have 'decal' domain and DefaultLit shading model
