Fix false positive warning during Substrate material conversion.

diff --git a/Engine/Source/Runtime/Engine/Private/Materials/Material.cpp b/Engine/Source/Runtime/Engine/Private/Materials/Material.cpp
--- a/Engine/Source/Runtime/Engine/Private/Materials/Material.cpp
+++ b/Engine/Source/Runtime/Engine/Private/Materials/Material.cpp
@@ -5053,11 +5053,7 @@ void UMaterial::RebuildShadingModelField()
 			{
 				MaterialDomain = EMaterialDomain::MD_Surface;
 				ShadingModel = MSM_Subsurface;
-				if (BlendMode != EBlendMode::BLEND_Opaque && BlendMode != EBlendMode::BLEND_Masked)
-				{
-					UE_LOG(LogMaterial, Error, TEXT("%s: Material has subsurface data, and its blending mode is not set to Opaque or Masked. Forcing blend mode to Opaque."), *GetName());
-					BlendMode = EBlendMode::BLEND_Opaque;
-				}
+				// This is a valid shading mode for both opaque/masked/translucent blend modes
 			}
 			else if (SubstrateMaterialInfo.HasOnlyShadingModel(SSM_SubsurfaceThinTwoSided))
 			{
