[NO NON-PII COMMIT MESSAGE]

diff --git a/Engine/Source/Runtime/Engine/Classes/DataDrivenCVars/DataDrivenCVars.h b/Engine/Source/Runtime/Engine/Classes/DataDrivenCVars/DataDrivenCVars.h
--- a/Engine/Source/Runtime/Engine/Classes/DataDrivenCVars/DataDrivenCVars.h
+++ b/Engine/Source/Runtime/Engine/Classes/DataDrivenCVars/DataDrivenCVars.h
@@ -74,6 +74,7 @@ struct FDataDrivenConsoleVariable
 	ENGINE_API void UnRegister(bool bUseShadowName = false);
 
 	FString ShadowName;
+	FString ShadowToolTip;
 	FDataDrivenCVarType ShadowType;
 };
 
diff --git a/Engine/Source/Runtime/Engine/Private/DataDrivenCVars/DataDrivenCVars.cpp b/Engine/Source/Runtime/Engine/Private/DataDrivenCVars/DataDrivenCVars.cpp
--- a/Engine/Source/Runtime/Engine/Private/DataDrivenCVars/DataDrivenCVars.cpp
+++ b/Engine/Source/Runtime/Engine/Private/DataDrivenCVars/DataDrivenCVars.cpp
@@ -37,6 +37,7 @@ void FDataDrivenConsoleVariable::Register()
 		}
 		CVarToAdd->SetOnChangedCallback(FConsoleVariableDelegate::CreateStatic(UDataDrivenConsoleVariableSettings::OnDataDrivenChange));
 		ShadowName = Name;
+		ShadowToolTip = ToolTip;
 		ShadowType = Type;
 	}
 }
@@ -68,6 +69,11 @@ void FDataDrivenConsoleVariable::Refresh()
 		}
 		ShadowName = Name;
 	}
+	else if (ShadowToolTip != ToolTip)
+	{
+		UnRegister(true);
+		ShadowToolTip = ToolTip;
+	}
 	else if (ShadowType != Type)
 	{
 		UnRegister(true);
