Wrapped material proxy invalidation / evaluation in a mutex. Prerequisite for parallel GDME work.

diff --git a/Engine/Source/Runtime/Engine/Private/Materials/MaterialRenderProxy.cpp b/Engine/Source/Runtime/Engine/Private/Materials/MaterialRenderProxy.cpp
--- a/Engine/Source/Runtime/Engine/Private/Materials/MaterialRenderProxy.cpp
+++ b/Engine/Source/Runtime/Engine/Private/Materials/MaterialRenderProxy.cpp
@@ -499,6 +499,8 @@ void FMaterialRenderProxy::InvalidateUniformExpressionCache(bool bRecreateUnifor
 	FStaticLightingSystemInterface::OnMaterialInvalidated.Broadcast(this);
 #endif
 
+	UE::TScopeLock Lock(Mutex);
+
 	if (HasVirtualTextureCallbacks)
 	{
 		GetRendererModule().RemoveAllVirtualTextureProducerDestroyedCallbacks(this);
@@ -527,6 +529,7 @@ void FMaterialRenderProxy::UpdateUniformExpressionCacheIfNeeded(FRHICommandListB
 	const FMaterial* Material = GetMaterialNoFallback(InFeatureLevel);
 	// Note: We would actually need to compare the FMaterialShaderMapId of both shader maps but a simple pointer compare also works
 	// because shader maps are currently swapped out whenever they are modified.
+	UE::TScopeLock Lock(Mutex);
 	if (Material && Material->GetRenderingThreadShaderMap() != UniformExpressionCache[InFeatureLevel].CachedUniformExpressionShaderMap)
 	{
 		FMaterialRenderContext MaterialRenderContext(this, *Material, nullptr);
@@ -689,7 +692,11 @@ void FMaterialRenderProxy::UpdateDeferredCachedUniformExpressions(FRHICommandLis
 		GUniformExpressionCacheAsyncUpdateTask.Wait();
 	}
 
-	auto EvaluateUniformExpressionsLambda = [RHICmdListTask, bAllowAsyncUpdate]
+	DeferredUniformExpressionCacheRequestsMutex.Lock();
+	TSet<FMaterialRenderProxy*> UniformExpressions = MoveTemp(DeferredUniformExpressionCacheRequests);
+	DeferredUniformExpressionCacheRequestsMutex.Unlock();
+
+	auto EvaluateUniformExpressionsLambda = [RHICmdListTask, bAllowAsyncUpdate, UniformExpressions = MoveTemp(UniformExpressions)]
 	{
 		TRACE_CPUPROFILER_EVENT_SCOPE(FMaterialRenderProxy::UpdateDeferredCachedUniformExpressions);
 		CSV_SCOPED_TIMING_STAT_EXCLUSIVE(Material_UpdateDeferredCachedUniformExpressions);
@@ -699,32 +706,27 @@ void FMaterialRenderProxy::UpdateDeferredCachedUniformExpressions(FRHICommandLis
 		FUniformExpressionCacheAsyncUpdater Updater;
 		FUniformExpressionCacheAsyncUpdater* UpdaterIfEnabled = bAllowAsyncUpdate ? &Updater : nullptr;
 
+		for (TSet<FMaterialRenderProxy*>::TConstIterator It(UniformExpressions); It; ++It)
 		{
-			UE::TScopeLock Lock(DeferredUniformExpressionCacheRequestsMutex);
+			FMaterialRenderProxy* MaterialProxy = *It;
+			if (MaterialProxy->IsDeleted())
+			{
+				UE_LOG(LogMaterial, Fatal, TEXT("FMaterialRenderProxy deleted and GC mark was: %i"), MaterialProxy->IsMarkedForGarbageCollection());
+			}
 
-			for (TSet<FMaterialRenderProxy*>::TConstIterator It(DeferredUniformExpressionCacheRequests); It; ++It)
+			UMaterialInterface::IterateOverActiveFeatureLevels([&](ERHIFeatureLevel::Type InFeatureLevel)
 			{
-				FMaterialRenderProxy* MaterialProxy = *It;
-				if (MaterialProxy->IsDeleted())
+				// Don't bother caching if we'll be falling back to a different FMaterialRenderProxy for rendering anyway
+				const FMaterial* Material = MaterialProxy->GetMaterialNoFallback(InFeatureLevel);
+				if (Material && Material->GetRenderingThreadShaderMap())
 				{
-					UE_LOG(LogMaterial, Fatal, TEXT("FMaterialRenderProxy deleted and GC mark was: %i"), MaterialProxy->IsMarkedForGarbageCollection());
+					FMaterialRenderContext MaterialRenderContext(MaterialProxy, *Material, nullptr);
+					MaterialRenderContext.bShowSelection = GIsEditor;
+					MaterialProxy->EvaluateUniformExpressions(*RHICmdListTask, MaterialProxy->UniformExpressionCache[(int32)InFeatureLevel], MaterialRenderContext, UpdaterIfEnabled);
 				}
+			});
 
-				UMaterialInterface::IterateOverActiveFeatureLevels([&](ERHIFeatureLevel::Type InFeatureLevel)
-				{
-					// Don't bother caching if we'll be falling back to a different FMaterialRenderProxy for rendering anyway
-					const FMaterial* Material = MaterialProxy->GetMaterialNoFallback(InFeatureLevel);
-					if (Material && Material->GetRenderingThreadShaderMap())
-					{
-						FMaterialRenderContext MaterialRenderContext(MaterialProxy, *Material, nullptr);
-						MaterialRenderContext.bShowSelection = GIsEditor;
-						MaterialProxy->EvaluateUniformExpressions(*RHICmdListTask, MaterialProxy->UniformExpressionCache[(int32)InFeatureLevel], MaterialRenderContext, UpdaterIfEnabled);
-					}
-				});
-
-				MaterialProxy->FinishCacheUniformExpressions();
-			}
-			DeferredUniformExpressionCacheRequests.Reset();
+			MaterialProxy->FinishCacheUniformExpressions();
 		}
 
 		if (UpdaterIfEnabled)
diff --git a/Engine/Source/Runtime/Engine/Public/Materials/MaterialRenderProxy.h b/Engine/Source/Runtime/Engine/Public/Materials/MaterialRenderProxy.h
--- a/Engine/Source/Runtime/Engine/Public/Materials/MaterialRenderProxy.h
+++ b/Engine/Source/Runtime/Engine/Public/Materials/MaterialRenderProxy.h
@@ -244,6 +244,9 @@ private:
 	/** If any VT producer destroyed callbacks have been registered */
 	mutable uint8 HasVirtualTextureCallbacks : 1;
 
+	/** Mutex for locking uniform expression invalidation / evaluation for this material. */
+	mutable UE::FMutex Mutex;
+
 #if WITH_EDITOR
 	/**
 	 * Tracks all material render proxies in all scenes.
