Refactored batched elements and GDME mesh collector for parallelization.
 - Added a commit mode which defers GPU scene updates and material updates until the collector is Committed (used by main view visibility). Shadows commit immediately.
 - Added support for coalescing simple primitive collectors and batched elements.

diff --git a/Engine/Source/Runtime/Engine/Private/BatchedElements.cpp b/Engine/Source/Runtime/Engine/Private/BatchedElements.cpp
--- a/Engine/Source/Runtime/Engine/Private/BatchedElements.cpp
+++ b/Engine/Source/Runtime/Engine/Private/BatchedElements.cpp
@@ -130,7 +130,6 @@ FBatchedElements::FBatchedElements()
 	, MaxMeshIndicesAllowed(GDrawUPIndexCheckCount / sizeof(int32))
 	// the index buffer is 2 bytes, so make sure we only address 0xFFFF vertices in the index buffer
 	, MaxMeshVerticesAllowed(FMath::Min<uint32>(0xFFFF, GDrawUPVertexCheckCount / sizeof(FSimpleElementVertex)))
-	, bEnableHDREncoding(true)
 {
 }
 
@@ -1309,10 +1308,45 @@ bool FBatchedElements::Draw(FRHICommandList& RHICmdList, const FMeshPassProcesso
 
 void FBatchedElements::Clear()
 {
-	LineVertices.Empty();
 	Points.Empty();
+	WireTris.Empty();
+	WireTriVerts.Empty();
+	ThickLines.Empty();
 	Sprites.Empty();
 	MeshElements.Empty();
-	ThickLines.Empty();
 	MeshVertices.Empty();
 }
+
+void FBatchedElements::AddAllocationInfo(FAllocationInfo& AllocationInfo) const
+{
+	AllocationInfo.NumPoints += Points.Num();
+	AllocationInfo.NumWireTris += WireTris.Num();
+	AllocationInfo.NumWireTriVerts += WireTriVerts.Num();
+	AllocationInfo.NumThickLines += ThickLines.Num();
+	AllocationInfo.NumSprites += Sprites.Num();
+	AllocationInfo.NumMeshElements += MeshElements.Num();
+	AllocationInfo.NumMeshVertices += MeshVertices.Num();
+}
+
+void FBatchedElements::Reserve(const FAllocationInfo& AllocationInfo)
+{
+	Points.Reserve(AllocationInfo.NumPoints);
+	WireTris.Reserve(AllocationInfo.NumWireTris);
+	WireTriVerts.Reserve(AllocationInfo.NumWireTriVerts);
+	ThickLines.Reserve(AllocationInfo.NumThickLines);
+	Sprites.Reserve(AllocationInfo.NumSprites);
+	MeshElements.Reserve(AllocationInfo.NumMeshElements);
+	MeshVertices.Reserve(AllocationInfo.NumMeshVertices);
+}
+
+void FBatchedElements::Append(FBatchedElements& Other)
+{
+	Points.Append(Other.Points);
+	WireTris.Append(Other.WireTris);
+	WireTriVerts.Append(Other.WireTriVerts);
+	ThickLines.Append(Other.ThickLines);
+	Sprites.Append(Other.Sprites);
+	MeshElements.Append(Other.MeshElements);
+	MeshVertices.Append(Other.MeshVertices);
+	Other.Clear();
+}
\ No newline at end of file
diff --git a/Engine/Source/Runtime/Engine/Private/SceneManagement.cpp b/Engine/Source/Runtime/Engine/Private/SceneManagement.cpp
--- a/Engine/Source/Runtime/Engine/Private/SceneManagement.cpp
+++ b/Engine/Source/Runtime/Engine/Private/SceneManagement.cpp
@@ -156,10 +156,7 @@ TGlobalResource<FDefaultWorkingColorSpaceUniformBuffer> GDefaultWorkingColorSpac
 
 FSimpleElementCollector::FSimpleElementCollector() :
 	FPrimitiveDrawInterface(nullptr)
-{
-	static auto* MobileHDRCvar = IConsoleManager::Get().FindTConsoleVariableDataInt(TEXT("r.MobileHDR"));
-	bIsMobileHDR = (MobileHDRCvar->GetValueOnAnyThread() == 1);
-}
+{}
 
 FSimpleElementCollector::~FSimpleElementCollector()
 {
@@ -311,6 +308,28 @@ void FSimpleElementCollector::DrawBatchedElements(FRHICommandList& RHICmdList, c
 		);
 }
 
+void FSimpleElementCollector::AddAllocationInfo(FAllocationInfo& AllocationInfo) const
+{
+	BatchedElements.AddAllocationInfo(AllocationInfo.BatchedElements);
+	TopBatchedElements.AddAllocationInfo(AllocationInfo.TopBatchedElements);
+	AllocationInfo.NumDynamicResources += DynamicResources.Num();
+}
+
+void FSimpleElementCollector::Reserve(const FAllocationInfo& AllocationInfo)
+{
+	BatchedElements.Reserve(AllocationInfo.BatchedElements);
+	TopBatchedElements.Reserve(AllocationInfo.TopBatchedElements);
+	DynamicResources.Reserve(AllocationInfo.NumDynamicResources);
+}
+
+void FSimpleElementCollector::Append(FSimpleElementCollector& Other)
+{
+	BatchedElements.Append(Other.BatchedElements);
+	TopBatchedElements.Append(Other.TopBatchedElements);
+	DynamicResources.Append(Other.DynamicResources);
+	Other.DynamicResources.Empty();
+}
+
 FMeshBatchAndRelevance::FMeshBatchAndRelevance(const FMeshBatch& InMesh, const FPrimitiveSceneProxy* InPrimitiveSceneProxy, ERHIFeatureLevel::Type FeatureLevel) :
 	Mesh(&InMesh),
 	PrimitiveSceneProxy(InPrimitiveSceneProxy)
@@ -349,11 +368,12 @@ FRayTracingMaterialGatheringContext::~FRayTracingMaterialGatheringContext()
 
 #endif
 
-FMeshElementCollector::FMeshElementCollector(ERHIFeatureLevel::Type InFeatureLevel, FSceneRenderingBulkObjectAllocator& InBulkAllocator) :
+FMeshElementCollector::FMeshElementCollector(ERHIFeatureLevel::Type InFeatureLevel, FSceneRenderingBulkObjectAllocator& InBulkAllocator, ECommitFlags InCommitFlags) :
 	OneFrameResources(InBulkAllocator),
 	PrimitiveSceneProxy(NULL),
 	DynamicReadBuffer(nullptr),
 	FeatureLevel(InFeatureLevel),
+	CommitFlags(InCommitFlags),
 	bUseGPUScene(UseGPUScene(GMaxRHIShaderPlatform, FeatureLevel))
 {
 }
@@ -375,18 +395,17 @@ void FMeshElementCollector::SetPrimitive(const FPrimitiveSceneProxy* InPrimitive
 	for (int32 ViewIndex = 0; ViewIndex < SimpleElementCollectors.Num(); ViewIndex++)
 	{
 		SimpleElementCollectors[ViewIndex]->HitProxyId = DefaultHitProxyId;
-		SimpleElementCollectors[ViewIndex]->PrimitiveMeshId = 0;
 	}
 
 	for (int32 ViewIndex = 0; ViewIndex < MeshIdInPrimitivePerView.Num(); ++ViewIndex)
 	{
 		MeshIdInPrimitivePerView[ViewIndex] = 0;
 	}
+
 #if UE_ENABLE_DEBUG_DRAWING
 	for (int32 ViewIndex = 0; ViewIndex < DebugSimpleElementCollectors.Num(); ViewIndex++)
 	{
 		DebugSimpleElementCollectors[ViewIndex]->HitProxyId = DefaultHitProxyId;
-		DebugSimpleElementCollectors[ViewIndex]->PrimitiveMeshId = 0;
 	}
 #endif
 }
@@ -405,7 +424,7 @@ void FMeshElementCollector::Start(
 }
 
 void FMeshElementCollector::AddViewMeshArrays(
-	FSceneView* InView,
+	const FSceneView* InView,
 	TArray<FMeshBatchAndRelevance, SceneRenderingAllocator>* ViewMeshes,
 	FSimpleElementCollector* ViewSimpleElementCollector,
 	FGPUScenePrimitiveCollector* DynamicPrimitiveCollector
@@ -482,13 +501,26 @@ void FMeshElementCollector::Finish()
 	RHICmdList = nullptr;
 }
 
+void FMeshElementCollector::CacheUniformExpressions(FMaterialRenderProxy* Proxy, bool bRecreateUniformBuffer)
+{
+	check(Proxy);
+	if (EnumHasAnyFlags(CommitFlags, ECommitFlags::DeferMaterials))
+	{
+		MaterialProxiesToInvalidate.Emplace(Proxy, bRecreateUniformBuffer);
+	}
+	else
+	{
+		Proxy->InvalidateUniformExpressionCache(bRecreateUniformBuffer);
+	}
+}
+
 void FMeshElementCollector::AddMesh(int32 ViewIndex, FMeshBatch& MeshBatch)
 {
 	DEFINE_LOG_CATEGORY_STATIC(FMeshElementCollector_AddMesh, Warning, All);
 
 	if (MeshBatch.bCanApplyViewModeOverrides)
 	{
-		FSceneView* View = Views[ViewIndex];
+		const FSceneView* View = Views[ViewIndex];
 
 		ApplyViewModeOverrides(
 			ViewIndex,
@@ -509,10 +541,24 @@ void FMeshElementCollector::AddMesh(int32 ViewIndex, FMeshBatch& MeshBatch)
 
 	if (bUseGPUScene && MeshBatch.VertexFactory->GetPrimitiveIdStreamIndex(FeatureLevel, EVertexInputStreamType::Default) >= 0)
 	{
-		MeshBatchesForGPUScene.Emplace(DynamicPrimitiveCollectorPerView[ViewIndex], &MeshBatch);
+		if (EnumHasAnyFlags(CommitFlags, ECommitFlags::DeferGPUScene))
+		{
+			MeshBatchesForGPUScene.Emplace(DynamicPrimitiveCollectorPerView[ViewIndex], &MeshBatch);
+		}
+		else
+		{
+			GetRendererModule().AddMeshBatchToGPUScene(DynamicPrimitiveCollectorPerView[ViewIndex], MeshBatch);
+		}
 	}
 
-	MaterialProxiesToUpdate.Emplace(MeshBatch.MaterialRenderProxy);
+	if (EnumHasAnyFlags(CommitFlags, ECommitFlags::DeferMaterials))
+	{
+		MaterialProxiesToUpdate.Emplace(MeshBatch.MaterialRenderProxy);
+	}
+	else
+	{
+		MeshBatch.MaterialRenderProxy->UpdateUniformExpressionCacheIfNeeded(*RHICmdList, FeatureLevel);
+	}
 
 	MeshBatch.MeshIdInPrimitive = MeshIdInPrimitivePerView[ViewIndex];
 	++MeshIdInPrimitivePerView[ViewIndex];
diff --git a/Engine/Source/Runtime/Engine/Public/BatchedElements.h b/Engine/Source/Runtime/Engine/Public/BatchedElements.h
--- a/Engine/Source/Runtime/Engine/Public/BatchedElements.h
+++ b/Engine/Source/Runtime/Engine/Public/BatchedElements.h
@@ -237,10 +237,34 @@ public:
 			+ Sprites.GetAllocatedSize() + MeshElements.GetAllocatedSize() + MeshVertices.GetAllocatedSize();
 	}
 
-	void EnableMobileHDREncoding(bool bInEnableHDREncoding)
+	UE_DEPRECATED(5.4, "EnableMobileHDREncoding is no longer supported")
+	void EnableMobileHDREncoding(bool bInEnableHDREncoding) {}
+
+	class FAllocationInfo
 	{
-		bEnableHDREncoding = bInEnableHDREncoding;
-	}
+	public:
+		FAllocationInfo() = default;
+
+	private:
+		int32 NumPoints = 0;
+		int32 NumWireTris = 0;
+		int32 NumWireTriVerts = 0;
+		int32 NumThickLines = 0;
+		int32 NumSprites = 0;
+		int32 NumMeshElements = 0;
+		int32 NumMeshVertices = 0;
+
+		friend FBatchedElements;
+	};
+
+	/** Accumulates allocation info for use calling Reserve. */
+	ENGINE_API void AddAllocationInfo(FAllocationInfo& AllocationInfo) const;
+
+	/** Reserves memory for all containers. */
+	ENGINE_API void Reserve(const FAllocationInfo& AllocationInfo);
+
+	/** Appends contents of another batched elements into this one and clears the other one. */
+	ENGINE_API void Append(FBatchedElements& Other);
 
 private:
 
@@ -347,8 +371,5 @@ private:
 		const FSceneView* View = nullptr,
 		float OpacityMaskRefVal = .5f
 		) const;
-
-	/** if false then prevent the use of HDR encoded shaders. */
-	bool bEnableHDREncoding;
 };
 
diff --git a/Engine/Source/Runtime/Engine/Public/SceneManagement.h b/Engine/Source/Runtime/Engine/Public/SceneManagement.h
--- a/Engine/Source/Runtime/Engine/Public/SceneManagement.h
+++ b/Engine/Source/Runtime/Engine/Public/SceneManagement.h
@@ -1585,7 +1585,7 @@ class FPrimitiveDrawInterface
 {
 public:
 
-	const FSceneView* const View;
+	const FSceneView* View;
 
 	/** Initialization constructor. */
 	FPrimitiveDrawInterface(const FSceneView* InView):
@@ -1776,6 +1776,28 @@ public:
 
 	ENGINE_API void DrawBatchedElements(FRHICommandList& RHICmdList, const FMeshPassProcessorRenderState& DrawRenderState, const FSceneView& InView, EBlendModeFilter::Type Filter, ESceneDepthPriorityGroup DPG) const;
 
+	class FAllocationInfo
+	{
+	public:
+		FAllocationInfo() = default;
+
+	private:
+		FBatchedElements::FAllocationInfo BatchedElements;
+		FBatchedElements::FAllocationInfo TopBatchedElements;
+		uint32 NumDynamicResources = 0;
+
+		friend FSimpleElementCollector;
+	};
+
+	/** Accumulates allocation info for use calling Reserve. */
+	ENGINE_API void AddAllocationInfo(FAllocationInfo& AllocationInfo) const;
+
+	/** Reserves memory for all containers. */
+	ENGINE_API void Reserve(const FAllocationInfo& AllocationInfo);
+
+	/** Appends contents of another batched elements into this one and clears the other one. */
+	ENGINE_API void Append(FSimpleElementCollector& Other);
+
 	bool HasAnyPrimitives() const
 	{
 		return BatchedElements.HasPrimsToDraw() || TopBatchedElements.HasPrimsToDraw();
@@ -1796,11 +1818,7 @@ public:
 	FBatchedElements TopBatchedElements;
 
 private:
-
 	FHitProxyId HitProxyId;
-	uint16 PrimitiveMeshId;
-
-	bool bIsMobileHDR;
 
 	/** The dynamic resources which have been registered with this drawer. */
 	TArray<FDynamicPrimitiveResource*,SceneRenderingAllocator> DynamicResources;
@@ -1931,11 +1949,7 @@ public:
 	}
 
 	/** Adds a request to force caching of uniform expressions for a material render proxy. */
-	void CacheUniformExpressions(FMaterialRenderProxy* Proxy, bool bRecreateUniformBuffer)
-	{
-		check(Proxy);
-		MaterialProxiesToInvalidate.Emplace(Proxy, bRecreateUniformBuffer);
-	}
+	ENGINE_API void CacheUniformExpressions(FMaterialRenderProxy* Proxy, bool bRecreateUniformBuffer);
 
 	/** Allocates a temporary resource that is safe to be referenced by an FMeshBatch added to the collector. */
 	template<typename T, typename... ARGS>
@@ -1965,8 +1979,21 @@ public:
 	}
 
 protected:
+	enum class ECommitFlags
+	{
+		None = 0,
+
+		// Defers material uniform expression updates until Commit or Finish is called.
+		DeferMaterials = 1 << 0,
 
-	ENGINE_API FMeshElementCollector(ERHIFeatureLevel::Type InFeatureLevel, FSceneRenderingBulkObjectAllocator& InBulkAllocator);
+		// Defers GPU scene updates until Commit or Finish is called.
+		DeferGPUScene  = 1 << 1,
+
+		DeferAll = DeferMaterials | DeferGPUScene
+	};
+	FRIEND_ENUM_CLASS_FLAGS(ECommitFlags);
+
+	ENGINE_API FMeshElementCollector(ERHIFeatureLevel::Type InFeatureLevel, FSceneRenderingBulkObjectAllocator& InBulkAllocator, ECommitFlags CommitFlags = ECommitFlags::None);
 
 	ENGINE_API ~FMeshElementCollector();
 
@@ -1979,7 +2006,7 @@ protected:
 		FGlobalDynamicReadBuffer& DynamicReadBuffer);
 
 	ENGINE_API void AddViewMeshArrays(
-		FSceneView* InView,
+		const FSceneView* InView,
 		TArray<FMeshBatchAndRelevance, SceneRenderingAllocator>* ViewMeshes,
 		FSimpleElementCollector* ViewSimpleElementCollector,
 		FGPUScenePrimitiveCollector* DynamicPrimitiveCollector
@@ -1998,7 +2025,7 @@ protected:
 	 * Using TChunkedArray which will never realloc as new elements are added
 	 * @todo - use mem stack
 	 */
-	TChunkedArray<FMeshBatch> MeshBatchStorage;
+	TChunkedArray<FMeshBatch, 1024, FConcurrentLinearArrayAllocator> MeshBatchStorage;
 
 	/** Meshes to render */
 	TArray<TArray<FMeshBatchAndRelevance, SceneRenderingAllocator>*, TInlineAllocator<2, SceneRenderingAllocator> > MeshBatches;
@@ -2014,7 +2041,7 @@ protected:
 #endif
 
 	/** Views being collected for */
-	TArray<FSceneView*, TInlineAllocator<2, SceneRenderingAllocator>> Views;
+	TArray<const FSceneView*, TInlineAllocator<2, SceneRenderingAllocator>> Views;
 
 	/** Current Mesh Id In Primitive per view */
 	TArray<uint16, TInlineAllocator<2, SceneRenderingAllocator>> MeshIdInPrimitivePerView;
@@ -2045,6 +2072,7 @@ protected:
 	FRHICommandList* RHICmdList = nullptr;
 
 	const ERHIFeatureLevel::Type FeatureLevel;
+	const ECommitFlags CommitFlags;
 	const bool bUseGPUScene;
 
 	/** Tracks dynamic primitive data for upload to GPU Scene for every view, when enabled. */
@@ -2057,9 +2085,13 @@ protected:
 	friend class FCardPageRenderData;
 	friend class FViewFamilyInfo;
 	friend class FShadowMeshCollector;
+	friend class FDynamicMeshElementContext;
 	friend struct FRayTracingMaterialGatheringContext;
+	friend FSceneRenderingBulkObjectAllocator;
 };
 
+ENUM_CLASS_FLAGS(FMeshElementCollector::ECommitFlags);
+
 #if RHI_RAYTRACING
 /**
  * Collector used to gather resources for the material mesh batches.
diff --git a/Engine/Source/Runtime/Renderer/Private/SceneRendering.cpp b/Engine/Source/Runtime/Renderer/Private/SceneRendering.cpp
--- a/Engine/Source/Runtime/Renderer/Private/SceneRendering.cpp
+++ b/Engine/Source/Runtime/Renderer/Private/SceneRendering.cpp
@@ -943,9 +943,6 @@ void FViewInfo::Init()
 	NumSphereReflectionCaptures = 0;
 	FurthestReflectionCaptureDistance = 0;
 
-	// Disable HDR encoding for editor elements.
-	EditorSimpleElementCollector.BatchedElements.EnableMobileHDREncoding(false);
-	
 	TemporalJitterSequenceLength = 1;
 	TemporalJitterIndex = 0;
 	TemporalJitterPixels = FVector2D::ZeroVector;
