[shaders]
- fix preprocess time stat - this was not accounting for time spent preprocessing jobs which subsequently hit in DDC or the job cache; required some rework to how stats are recorded (now all jobs will register with the shader compiler stats object, instead of only those that actually ran compilation)
- add option to skip dumping debug info for cache hits. having it enabled results in a deterministic set of debug information between runs, while having it disabled means only shaders which actually compile will output debug information (this is useful when used in conjunction with the shaderanalyzer tool)

diff --git a/Engine/Source/Runtime/Engine/Private/ShaderCompiler/ShaderCompiler.cpp b/Engine/Source/Runtime/Engine/Private/ShaderCompiler/ShaderCompiler.cpp
--- a/Engine/Source/Runtime/Engine/Private/ShaderCompiler/ShaderCompiler.cpp
+++ b/Engine/Source/Runtime/Engine/Private/ShaderCompiler/ShaderCompiler.cpp
@@ -1754,17 +1754,16 @@ void FShaderCompileJobCollection::HandlePrintStats()
 	GShaderCompilingManager->PrintStats();
 }
 
-TRACE_DECLARE_INT_COUNTER(Shaders_Compiled, TEXT("Shaders/Compiled"));
 double FShaderJobCache::ProcessFinishedJob(FShaderCommonCompileJob* FinishedJob, bool bWasCached)
 {
 	double StallTime;
 
 	FinishedJob->OnComplete();
 
+	GShaderCompilerStats->RegisterFinishedJob(*FinishedJob, bWasCached);
+
 	if (!bWasCached)
 	{
-		GShaderCompilerStats->RegisterFinishedJob(*FinishedJob);
-		TRACE_COUNTER_ADD(Shaders_Compiled, 1);
 	}
 
 	{
@@ -4992,21 +4991,30 @@ void FShaderCompilerStats::RegisterAssignedJob(FShaderCommonCompileJob& Job)
 	Counters.MaxPendingTime = FMath::Max(TimeSpendPending, Counters.MaxPendingTime);
 }
 
-void FShaderCompilerStats::RegisterFinishedJob(FShaderCommonCompileJob& Job)
+TRACE_DECLARE_INT_COUNTER(Shaders_Compiled, TEXT("Shaders/Compiled"));
+void FShaderCompilerStats::RegisterFinishedJob(FShaderCommonCompileJob& Job, bool bWasCached)
 {
-	ensure(Job.TimeAssignedToExecution != 0.0);
-	Job.TimeExecutionCompleted = FPlatformTime::Seconds();
-
 	FScopeLock Lock(&CompileStatsLock);
-	Counters.JobsCompleted++;
 	
-	double ExecutionTime = (Job.TimeExecutionCompleted - Job.TimeAssignedToExecution);
-	Counters.AccumulatedJobExecutionTime += ExecutionTime;
-	Counters.MaxJobExecutionTime = FMath::Max(ExecutionTime, Counters.MaxJobExecutionTime);
+	if (!bWasCached)
+	{
+		ensure(Job.TimeAssignedToExecution != 0.0);
+		Job.TimeExecutionCompleted = FPlatformTime::Seconds();
+		TRACE_COUNTER_ADD(Shaders_Compiled, 1);
+		Counters.JobsCompleted++;
+
+		double ExecutionTime = (Job.TimeExecutionCompleted - Job.TimeAssignedToExecution);
+		Counters.AccumulatedJobExecutionTime += ExecutionTime;
+		Counters.MaxJobExecutionTime = FMath::Max(ExecutionTime, Counters.MaxJobExecutionTime);
 
-	double LifeTime = (Job.TimeExecutionCompleted - Job.TimeAddedToPendingQueue);
-	Counters.AccumulatedJobLifeTime += LifeTime;
-	Counters.MaxJobLifeTime = FMath::Max(LifeTime, Counters.MaxJobLifeTime);
+		double LifeTime = (Job.TimeExecutionCompleted - Job.TimeAddedToPendingQueue);
+		Counters.AccumulatedJobLifeTime += LifeTime;
+		Counters.MaxJobLifeTime = FMath::Max(LifeTime, Counters.MaxJobLifeTime);
+		
+		// estimate lifetime without an overlap
+		ensure(Job.TimeAddedToPendingQueue != 0.0 && Job.TimeAddedToPendingQueue <= Job.TimeExecutionCompleted);
+		AddToInterval(JobLifeTimeIntervals, TInterval<double>(Job.TimeAddedToPendingQueue, Job.TimeExecutionCompleted));
+	}
 
 	if (Job.TimeTaskSubmitJobs)
 	{
@@ -5014,15 +5022,11 @@ void FShaderCompilerStats::RegisterFinishedJob(FShaderCommonCompileJob& Job)
 		Counters.AccumulatedTaskSubmitJobsStall += Job.TimeTaskSubmitJobsStall;
 	}
 
-	// estimate lifetime without an overlap
-	ensure(Job.TimeAddedToPendingQueue != 0.0 && Job.TimeAddedToPendingQueue <= Job.TimeExecutionCompleted);
-	AddToInterval(JobLifeTimeIntervals, TInterval<double>(Job.TimeAddedToPendingQueue, Job.TimeExecutionCompleted));
-
-	auto RegisterStatsFromSingleJob = [this](const FShaderCompileJob& SingleJob)
+	auto RegisterStatsFromSingleJob = [this, bWasCached](const FShaderCompileJob& SingleJob)
 	{
 		// Register min/max/average shader code sizes for single job output
 		const int32 ShaderCodeSize = SingleJob.Output.ShaderCode.GetShaderCodeSize();
-		if (ShaderCodeSize > 0)
+		if (!bWasCached && ShaderCodeSize > 0)
 		{
 			Counters.MinShaderCodeSize = (Counters.MinShaderCodeSize > 0 ? FMath::Min(Counters.MinShaderCodeSize, ShaderCodeSize) : ShaderCodeSize);
 			Counters.MaxShaderCodeSize = (Counters.MaxShaderCodeSize > 0 ? FMath::Max(Counters.MaxShaderCodeSize, ShaderCodeSize) : ShaderCodeSize);
@@ -5030,16 +5034,27 @@ void FShaderCompilerStats::RegisterFinishedJob(FShaderCommonCompileJob& Job)
 			++Counters.NumAccumulatedShaderCodes;
 		}
 
+		// Sanity check; compile time should be 0 for cache hits
+		check(!bWasCached || SingleJob.Output.CompileTime == 0.0f);
+		// Preprocess time should always be non-zero if preprocessed job cache is enabled
+		check(!SingleJob.Input.bCachePreprocessed || SingleJob.Output.PreprocessTime > 0.0f);
+
 		const FString ShaderName(SingleJob.Key.ShaderType->GetName());
 		if (FShaderTimings* Existing = ShaderTimings.Find(ShaderName))
 		{
-			Existing->MinCompileTime = FMath::Min(Existing->MinCompileTime, static_cast<float>(SingleJob.Output.CompileTime));
-			Existing->MaxCompileTime = FMath::Max(Existing->MaxCompileTime, static_cast<float>(SingleJob.Output.CompileTime));
-			Existing->TotalCompileTime += SingleJob.Output.CompileTime;
+			// Always want to log preprocess time, in case preprocessed cache is enabled and preprocessing ran in the cooker prior to compilation
+			// (PreprocessTime will be 0 if preprocessed cache is disabled)
 			Existing->TotalPreprocessTime += SingleJob.Output.PreprocessTime;
-			Existing->NumCompiled++;
-			// calculate as an optimization to make sorting later faster
-			Existing->AverageCompileTime = Existing->TotalCompileTime / static_cast<float>(Existing->NumCompiled);
+			if (!bWasCached)
+			{
+				// If no actual compiles have been logged yet, min compile time is just the compile time of this job (first to actually run)
+				Existing->MinCompileTime = Existing->NumCompiled ? FMath::Min(Existing->MinCompileTime, static_cast<float>(SingleJob.Output.CompileTime)) : SingleJob.Output.CompileTime;
+				Existing->MaxCompileTime = FMath::Max(Existing->MaxCompileTime, static_cast<float>(SingleJob.Output.CompileTime));
+				Existing->TotalCompileTime += SingleJob.Output.CompileTime;
+				Existing->NumCompiled++;
+				// calculate as an optimization to make sorting later faster
+				Existing->AverageCompileTime = Existing->TotalCompileTime / static_cast<float>(Existing->NumCompiled);
+			}
 		}
 		else
 		{
@@ -5047,9 +5062,12 @@ void FShaderCompilerStats::RegisterFinishedJob(FShaderCommonCompileJob& Job)
 			New.MinCompileTime = SingleJob.Output.CompileTime;
 			New.MaxCompileTime = New.MinCompileTime;
 			New.TotalCompileTime = New.MinCompileTime;
-			New.TotalPreprocessTime += SingleJob.Output.PreprocessTime;
 			New.AverageCompileTime = New.MinCompileTime;
-			New.NumCompiled = 1;
+			// It's possible the first entry for a given shader didn't actually compile (i.e. hit in DDC)
+			// so we need to account for that in the stats
+			New.NumCompiled = bWasCached ? 0 : 1;
+			New.TotalPreprocessTime += SingleJob.Output.PreprocessTime;
+
 			ShaderTimings.Add(ShaderName, New);
 		}
 	};
diff --git a/Engine/Source/Runtime/Engine/Public/ShaderCompiler.h b/Engine/Source/Runtime/Engine/Public/ShaderCompiler.h
--- a/Engine/Source/Runtime/Engine/Public/ShaderCompiler.h
+++ b/Engine/Source/Runtime/Engine/Public/ShaderCompiler.h
@@ -354,8 +354,10 @@ public:
 			TotalCompileTime += Other.TotalCompileTime;
 			TotalPreprocessTime += Other.TotalPreprocessTime;
 			NumCompiled += Other.NumCompiled;
-			AverageCompileTime = (TotalCompileTime + Other.TotalCompileTime) / static_cast<float>(NumCompiled + Other.NumCompiled);
-
+			if (NumCompiled)
+			{
+				AverageCompileTime = TotalCompileTime / static_cast<float>(NumCompiled);
+			}
 			return *this;
 		}
 	};
@@ -394,7 +396,7 @@ public:
 	void RegisterAssignedJob(FShaderCommonCompileJob& InOutJob);
 
 	/** Marks the job as finished for the stats purpose. Job will be modified to include the current timestamp. */
-	void RegisterFinishedJob(FShaderCommonCompileJob& InOutJob);
+	void RegisterFinishedJob(FShaderCommonCompileJob& InOutJob, bool bWasCached);
 
 	/** Informs statistics about a new job batch, so we can tally up batches. */
 	void RegisterJobBatch(int32 NumJobs, EExecutionType ExecType);
diff --git a/Engine/Source/Runtime/RenderCore/Private/ShaderCore.cpp b/Engine/Source/Runtime/RenderCore/Private/ShaderCore.cpp
--- a/Engine/Source/Runtime/RenderCore/Private/ShaderCore.cpp
+++ b/Engine/Source/Runtime/RenderCore/Private/ShaderCore.cpp
@@ -55,6 +55,12 @@ static TAutoConsoleVariable<int32> CVarShaderDevelopmentMode(
 	TEXT("0: Default, 1: Enable various shader development utilities, such as the ability to retry on failed shader compile, and extra logging as shaders are compiled."),
 	ECVF_Default);
 
+static TAutoConsoleVariable<bool> CVarDumpDebugInfoForCacheHits(
+	TEXT("r.ShaderCompiler.DumpDebugInfoForCacheHits"),
+	true,
+	TEXT("If true, debug info (via IShaderFormat::OutputDebugData) will be output for all jobs including duplicates and cache/DDC hits. If false, only jobs that actually executed compilation will dump debug info."),
+	ECVF_Default);
+
 void UpdateShaderDevelopmentMode()
 {
 	// Keep LogShaders verbosity in sync with r.ShaderDevelopmentMode
@@ -3446,13 +3452,14 @@ FShaderCommonCompileJob::FInputHash FShaderCompileJob::GetInputHash()
 void FShaderCompileJob::SerializeOutput(FArchive& Ar)
 {
 	double ActualCompileTime = 0.0;
-	double ActualPreprocessTime = 0.0;
+	// Save the preprocess time as set in the job regardless of whether saving or loading - if loading from the cache and the preprocessed job
+	// cache is enabled this job will have already run its own preprocessing and we want to track/aggregate this time properly.
+	double ActualPreprocessTime = Output.PreprocessTime;
 	if (Ar.IsSaving())
 	{
-		// Cached jobs won't have accurate results anyway, so reduce the storage requirements by setting those fields to a known value.
-		// This significantly reduces the memory needed to store the outputs (by more than a half)
+		// Clear preprocess time and compile time when storing a job in the cache. This reduces storage requirements since these objects are
+		// deduplicated based on a hash (and otherwise duplicate jobs will still differ in these values).
 		ActualCompileTime = Output.CompileTime;
-		ActualPreprocessTime = Output.PreprocessTime;
 		Output.CompileTime = 0.0;
 		Output.PreprocessTime = 0.0;
 	}
@@ -3468,10 +3475,13 @@ void FShaderCompileJob::SerializeOutput(FArchive& Ar)
 	}
 	else
 	{
-		// restore the compile time for this jobs. Jobs that will be deserialized from the cache will have a compile time of 0.0
+		// Restore the compile time for this job if we're saving to the cache.
+		// Jobs that will be deserialized from the cache will have a compile time of 0.0
 		Output.CompileTime = ActualCompileTime;
-		Output.PreprocessTime = ActualPreprocessTime;
 	}
+
+	// Unconditionally restore the preprocess time for this job after saving to or loading from the cache.
+	Output.PreprocessTime = ActualPreprocessTime;
 }
 
 void FShaderCompileJob::OnComplete()
@@ -3489,8 +3499,13 @@ void FShaderCompileJob::OnComplete()
 		// dump debug info for the job at this point if the preprocessed cache is enabled
 		// this ensures we get debug output for all jobs, including those that were found in the job cache,
 		// or matched another in-flight job's hash and so could share its results
-		if (Input.bCachePreprocessed && Input.DumpDebugInfoEnabled())
+		if (Input.bCachePreprocessed 
+			&& Input.DumpDebugInfoEnabled()
+			// if we only want debug info for jobs which actually compiled, check the CompileTime
+			// (jobs deserialized from the cache/wait list/ddc will have a compiletime of 0.0)
+			&& (CVarDumpDebugInfoForCacheHits.GetValueOnAnyThread() || Output.CompileTime > 0.0f))
 		{
+			
 			if (SecondaryPreprocessOutput.IsValid() && SecondaryOutput.IsValid())
 			{
 				ShaderFormat->OutputDebugData(Input, PreprocessOutput, *SecondaryPreprocessOutput, Output, *SecondaryOutput);
