Fix mishandling of Static Mesh's renderdata in headless editor.

#rnx

diff --git a/Engine/Source/Runtime/Engine/Private/StaticMesh.cpp b/Engine/Source/Runtime/Engine/Private/StaticMesh.cpp
--- a/Engine/Source/Runtime/Engine/Private/StaticMesh.cpp
+++ b/Engine/Source/Runtime/Engine/Private/StaticMesh.cpp
@@ -1896,7 +1896,7 @@ void FStaticMeshRenderData::InitResources(ERHIFeatureLevel::Type InFeatureLevel,
 	}
 #endif // #if WITH_EDITOR
 
-	checkf(FApp::CanEverRender(), TEXT("RenderData should not initialize resources in headless runs"));
+	checkf(FApp::CanEverRender() || !FPlatformProperties::RequiresCookedData(), TEXT("RenderData should not initialize resources in headless cooked runs"));
 
 	for (int32 LODIndex = 0; LODIndex < LODResources.Num(); ++LODIndex)
 	{
@@ -5725,7 +5725,7 @@ void UStaticMesh::Serialize(FArchive& Ar)
 			SCOPE_MS_ACCUMULATOR(STAT_StaticMesh_RenderData);
 			TUniquePtr<class FStaticMeshRenderData> LocalRenderData = MakeUnique<FStaticMeshRenderData>();
 			LocalRenderData->Serialize(Ar, this, bCooked);
-			if (FApp::CanEverRender())
+			if (FApp::CanEverRender() || !FPlatformProperties::RequiresCookedData())	// cooked assets can be loaded also in the headless editor commandlets
 			{
 				SetRenderData(MoveTemp(LocalRenderData));
 			}
