Cook Indeterminism: UCurveTable: Make RemoveRedundantKeys idempotent (running it multiple times is a noop) to remove indeterminism based on whether it gets called multiple times in the cook.
#rnx

diff --git a/Engine/Source/Runtime/Engine/Private/Curves/SimpleCurve.cpp b/Engine/Source/Runtime/Engine/Private/Curves/SimpleCurve.cpp
--- a/Engine/Source/Runtime/Engine/Private/Curves/SimpleCurve.cpp
+++ b/Engine/Source/Runtime/Engine/Private/Curves/SimpleCurve.cpp
@@ -510,13 +510,14 @@ float FSimpleCurve::EvalForTwoKeys(const FSimpleCurveKey& Key1, const FSimpleCur
 
 void FSimpleCurve::RemoveRedundantKeysInternal(float Tolerance, int32 InStartKeepKey, int32 InEndKeepKey)
 {
-	if (Keys.Num() < 3) // Will always keep first and last key
+	int32 NumKeys = Keys.Num();
+	if (NumKeys < 3) // Will always keep first and last key
 	{
 		return;
 	}
 
 	const int32 ActualStartKeepKey = FMath::Max(InStartKeepKey, 0); // Will always keep first and last key
-	const int32 ActualEndKeepKey = FMath::Min(InEndKeepKey, Keys.Num()-1);
+	const int32 ActualEndKeepKey = FMath::Min(InEndKeepKey, NumKeys - 1);
 
 	check(ActualStartKeepKey < ActualEndKeepKey); // Make sure we are doing something sane
 	if ((ActualEndKeepKey - ActualStartKeepKey) < 2)
@@ -525,67 +526,120 @@ void FSimpleCurve::RemoveRedundantKeysInternal(float Tolerance, int32 InStartKee
 		return;
 	}
 
-	//Build some helper data for managing the HandleTokey map
-	TArray<FKeyHandle> AllHandlesByIndex;
-	TArray<FKeyHandle> KeepHandles;
+	// Cook Determinism note:
+	// RemoveRedundantKeysInternal can be called a variable number of times during a cook, due to
+	// CompositeTables which might call it again after their Parent table has already called it.
+	// To prevent that variable number of calls from non-deterministically removing a variable number
+	// of keys, we need to make sure that calling RemoveRedundantKeysInternal a second time after it
+	// was called once is a noop that does not chose to remove any more keys.
 
-	if (KeyHandlesToIndices.Num() != 0)
-	{
-		check(KeyHandlesToIndices.Num() == Keys.Num());
-		AllHandlesByIndex.AddZeroed(Keys.Num());
-		KeepHandles.Reserve(Keys.Num());
-
-		for (const TPair<FKeyHandle, int32>& HandleIndexPair : KeyHandlesToIndices.GetMap())
-		{
-			AllHandlesByIndex[HandleIndexPair.Value] = HandleIndexPair.Key;
-		}
-	}
-	else
+	// Keep all keys up to and including ActualStartKeepKey even if they are redundant
+	TArray<int32> KeepIndices;
+	KeepIndices.Reserve(NumKeys);
+	int32 NextEvalIndex = 0;
+	for (; NextEvalIndex <= ActualStartKeepKey; ++NextEvalIndex)
 	{
-		AllHandlesByIndex.AddDefaulted(Keys.Num());
+		KeepIndices.Add(NextEvalIndex);
 	}
-	
 
+	// Consider keys up to ActualEndKeepKey and remove them if they are redundant.
+	auto ShouldKeepMiddleKey = [this, Tolerance](const FSimpleCurveKey* Start, const FSimpleCurveKey* Middle, const FSimpleCurveKey* End)
 	{
-		TArray<FSimpleCurveKey> NewKeys;
-		NewKeys.Reserve(Keys.Num());
+		const float KeyValue = Middle->Value;
+		const float ValueWithoutKey = EvalForTwoKeys(*Start, *End, Middle->Time);
+		return FMath::Abs(ValueWithoutKey - KeyValue) > Tolerance;
+	};
 
-		//Add all the keys we are keeping from the start
-		for(int32 StartKeepIndex = 0; StartKeepIndex <= ActualStartKeepKey; ++StartKeepIndex)
+	check(KeepIndices.Num() >= 1); // ActualStartKeepKey is in the list
+	for (;;)
+	{
+		// Find the next index we will keep
+		bool bSkippedKeys = false;
+		const FSimpleCurveKey* EvalStartKey = &Keys[KeepIndices.Last()];
+		while (NextEvalIndex < ActualEndKeepKey)
 		{
-			NewKeys.Add(Keys[StartKeepIndex]);
-			KeepHandles.Add(AllHandlesByIndex[StartKeepIndex]);
+			const FSimpleCurveKey* EvalMiddleKey = &Keys[NextEvalIndex]; // Exists and is not ActualEndKeepKey
+			const FSimpleCurveKey* EvalEndKey = &Keys[NextEvalIndex + 1]; // Exists because ActualEndKeepKey is later in the list
+			if (ShouldKeepMiddleKey(EvalStartKey, EvalMiddleKey, EvalEndKey))
+			{
+				break;
+			}
+			++NextEvalIndex;
+			bSkippedKeys = true;
 		}
 
-		//Add keys up to the first end keep key if they are not redundant
-		int32 MostRecentKeepKeyIndex = 0;
-		for (int32 TestIndex = ActualStartKeepKey+1; TestIndex < ActualEndKeepKey; ++TestIndex) //Loop within the bounds of the first and last key
+		if (bSkippedKeys)
 		{
-			const float KeyValue = Keys[TestIndex].Value;
-			const float ValueWithoutKey = EvalForTwoKeys(Keys[MostRecentKeepKeyIndex], Keys[TestIndex + 1], Keys[TestIndex].Time);
-			if (FMath::Abs(ValueWithoutKey - KeyValue) > Tolerance) // Is this key needed
+			bool bRemovedKeys = false;
+			// If we skipped any keys, reevaluate the last keys added popping them until we find one to keep.
+			// If we remove any then back up one and reevalue the NextEvalIndex we decided to keep.
+			while (KeepIndices.Last() > ActualStartKeepKey) // ActualStartKeepKey and earlier are not removable
 			{
-				MostRecentKeepKeyIndex = TestIndex;
-				NewKeys.Add(Keys[TestIndex]);
-				KeepHandles.Add(AllHandlesByIndex[TestIndex]);
+				int32 NumKeepIndices = KeepIndices.Num();
+				check(NumKeepIndices >= 2); // ActualStartKeepKey and the one we are going to evaluate are both in the list
+				EvalStartKey = &Keys[KeepIndices[NumKeepIndices - 2]];
+				const FSimpleCurveKey* EvalMiddleKey = &Keys[KeepIndices[NumKeepIndices - 1]];
+				const FSimpleCurveKey* EvalEndKey = &Keys[NextEvalIndex];
+				if (ShouldKeepMiddleKey(EvalStartKey, EvalMiddleKey, EvalEndKey))
+				{
+					break;
+				}
+				KeepIndices.Pop(false /* bAllowShrink */);
+				check(KeepIndices.Num() >= 1); // ActualStartKeepKey is in the list
+				bRemovedKeys = true;
+			}
+			if (bRemovedKeys)
+			{
+				continue; // Reevaluate NextEvalIndex
 			}
 		}
+		if (NextEvalIndex >= ActualEndKeepKey)
+		{
+			// No further removable indices, we are done evaluating
+			break;
+		}
+
+		// Mark NextEvalIndex for keeping and move to the next key
+		KeepIndices.Add(NextEvalIndex);
+		++NextEvalIndex;
+	}
+
+	// Add end keys that we are keeping
+	for (; NextEvalIndex < NumKeys; ++NextEvalIndex)
+	{
+		KeepIndices.Add(NextEvalIndex);
+	}
+
+	// Build some helper data for managing the KeyHandlesToIndices map
+	TArray<FKeyHandle> AllHandlesByIndex;
+	if (KeyHandlesToIndices.Num() != 0)
+	{
+		// Do not use AddDefaulted as every FKeyHandle constructor allocates a new global identifier
+		AllHandlesByIndex.AddZeroed(NumKeys);
 
-		//Add end keys that we are keeping
-		for (int32 EndKeepIndex = ActualEndKeepKey; EndKeepIndex < Keys.Num(); ++EndKeepIndex)
+		check(KeyHandlesToIndices.Num() == NumKeys);
+		for (const TPair<FKeyHandle, int32>& HandleIndexPair : KeyHandlesToIndices.GetMap())
 		{
-			NewKeys.Add(Keys[EndKeepIndex]);
-			KeepHandles.Add(AllHandlesByIndex[EndKeepIndex]);
+			AllHandlesByIndex[HandleIndexPair.Value] = HandleIndexPair.Key;
 		}
-		Keys = MoveTemp(NewKeys); //Do this at the end of scope, guaranteed that NewKeys is going away
+	}
+	else
+	{
+		AllHandlesByIndex.AddDefaulted(NumKeys);
 	}
 
-	// Rebuild KeyHandlesToIndices
+	// Copy keys and KeyHandlesToIndices from all the indices we decided to keep
+	int32 NumKeepKeys = KeepIndices.Num();
+	TArray<FSimpleCurveKey> NewKeys;
+	NewKeys.Reserve(NumKeepKeys);
 	KeyHandlesToIndices.Empty();
-	for (int32 KeyIndex = 0; KeyIndex < Keys.Num(); ++KeyIndex)
+	for (int32 NewIndex = 0; NewIndex < NumKeepKeys; ++NewIndex)
 	{
-		KeyHandlesToIndices.Add(KeepHandles[KeyIndex], KeyIndex);
+		int32 OldIndex = KeepIndices[NewIndex];
+		NewKeys.Add(MoveTemp(Keys[OldIndex]));
+		KeyHandlesToIndices.Add(AllHandlesByIndex[OldIndex], NewIndex);
 	}
+	Keys = MoveTemp(NewKeys);
 }
 
 void FSimpleCurve::RemapTimeValue(float& InTime, float& CycleValueOffset) const
