Add methods to account for force inline sound waves in the streaming cache budget.

Adds virtual methods to the base class for Audio Streaming

diff --git a/Engine/Source/Runtime/Engine/Private/AudioStreamingCache.cpp b/Engine/Source/Runtime/Engine/Private/AudioStreamingCache.cpp
--- a/Engine/Source/Runtime/Engine/Private/AudioStreamingCache.cpp
+++ b/Engine/Source/Runtime/Engine/Private/AudioStreamingCache.cpp
@@ -425,6 +425,24 @@ void FCachedAudioStreamingManager::RemoveStreamingSoundWave(const FSoundWaveProx
 	// Unused.
 }
 
+void FCachedAudioStreamingManager::AddForceInlineSoundWave(const FSoundWaveProxyPtr& SoundWave)
+{
+	// add the sound wave to the first cache
+	if (ensure(CacheArray.Num() > 0))
+	{
+		CacheArray[0].AddForceInlineSoundWave(SoundWave);
+	}
+}
+
+void FCachedAudioStreamingManager::RemoveForceInlineSoundWave(const FSoundWaveProxyPtr& SoundWave)
+{
+	// remove the sound wave from the first cache
+	if (ensure(CacheArray.Num() > 0))
+	{
+		CacheArray[0].RemoveForecInlineSoundWave(SoundWave);
+	}
+}
+
 void FCachedAudioStreamingManager::AddDecoder(ICompressedAudioInfo* InCompressedAudioInfo)
 {
 	// Unused.
@@ -701,6 +719,7 @@ FAudioChunkCache::FAudioChunkCache(uint32 InMaxChunkSize, uint32 NumChunks, uint
 	, ChunksInUse(0)
 	, MemoryCounterBytes(0)
 	, MemoryLimitBytes(InMemoryLimitInBytes)
+	, ForceInlineMemoryCounterBytes(0)
 	, bLogCacheMisses(false)
 {
 	check(NumChunks > 0);
@@ -778,7 +797,8 @@ uint64 FAudioChunkCache::AddOrTouchChunk(const FChunkKey& InKey, const TSharedPt
 		{
 			int32 ChunkDataSize = Chunk->AudioDataSize;
 
-			if (TrimCacheWhenOverBudgetCVar != 0 && (MemoryCounterBytes + ChunkDataSize) > MemoryLimitBytes)
+			const uint64 MemoryUsageBytes = MemoryCounterBytes + ForceInlineMemoryCounterBytes + ChunkDataSize;
+			if (TrimCacheWhenOverBudgetCVar != 0 && MemoryUsageBytes > MemoryLimitBytes)
 			{
 				uint64 MemoryToTrim = 0;
 				if (MemoryLimitTrimPercentageCVar > 0.0f)
@@ -787,7 +807,7 @@ uint64 FAudioChunkCache::AddOrTouchChunk(const FChunkKey& InKey, const TSharedPt
 				}
 				else
 				{
-					MemoryToTrim = MemoryCounterBytes + ChunkDataSize - MemoryLimitBytes;
+					MemoryToTrim = MemoryUsageBytes - MemoryLimitBytes;
 				}
 
 				TrimMemory(MemoryToTrim, true);
@@ -971,6 +991,34 @@ void FAudioChunkCache::ClearCache()
 	ChunksInUse = 0;
 }
 
+void FAudioChunkCache::AddForceInlineSoundWave(const FSoundWaveProxyPtr& SoundWave)
+{
+	check(SoundWave->GetLoadingBehavior() == ESoundWaveLoadingBehavior::ForceInline);
+	ForceInlineMemoryCounterBytes += SoundWave->GetResourceSize();
+
+	const uint64 MemoryUsageBytes = MemoryCounterBytes + ForceInlineMemoryCounterBytes;
+	if (TrimCacheWhenOverBudgetCVar != 0 && MemoryUsageBytes > MemoryLimitBytes)
+	{
+		uint64 MemoryToTrim = 0;
+		if (MemoryLimitTrimPercentageCVar > 0.0f)
+		{
+			MemoryToTrim = MemoryLimitBytes * FMath::Min(MemoryLimitTrimPercentageCVar, 1.0f);
+		}
+		else
+		{
+			MemoryToTrim = MemoryUsageBytes - MemoryLimitBytes;
+		}
+
+		TrimMemory(MemoryToTrim, true);
+	}
+}
+
+void FAudioChunkCache::RemoveForecInlineSoundWave(const FSoundWaveProxyPtr& SoundWave)
+{
+	check(SoundWave->GetLoadingBehavior() == ESoundWaveLoadingBehavior::ForceInline);
+	ForceInlineMemoryCounterBytes -= SoundWave->GetResourceSize();
+}
+
 uint64 FAudioChunkCache::TrimMemory(uint64 BytesToFree, bool bInAllowRetainedChunkTrimming)
 {
 	FScopeLock ScopeLock(&CacheMutationCriticalSection);
@@ -1359,7 +1407,7 @@ void FAudioChunkCache::TouchElement(FCacheElement* InElement)
 
 bool FAudioChunkCache::ShouldAddNewChunk() const
 {
-	return (ChunksInUse < CachePool.Num()) && (MemoryCounterBytes.Load() < MemoryLimitBytes);
+	return (ChunksInUse < CachePool.Num()) && (MemoryCounterBytes.Load() + ForceInlineMemoryCounterBytes.Load() < MemoryLimitBytes);
 }
 
 FAudioChunkCache::FCacheElement* FAudioChunkCache::InsertChunk(const FChunkKey& InKey, const TSharedPtr<FSoundWaveData>& InSoundWavePtr)
@@ -2116,16 +2164,21 @@ FString FAudioChunkCache::DebugPrint()
 		CurrentElement = CurrentElement->LessRecentElement;
 	}
 
+	// Num bytes in use should include Force Inline data!
+	NumBytesCounter += ForceInlineMemoryCounterBytes;
+
 	// Convert to megabytes and print the total size:
 	const double NumMegabytesInUse = (double)NumBytesCounter / (1024 * 1024);
+	const double NumMegabytesForceInline = (double)ForceInlineMemoryCounterBytes / (1024 * 1024);
 	const double NumMegabytesRetained = (double)NumBytesRetained / (1024 * 1024);
 
 	const double MaxCacheSizeMB = ((double)MemoryLimitBytes) / (1024 * 1024);
 	const double PercentageOfCacheRetained = NumMegabytesRetained / MaxCacheSizeMB;
+	const double PercentageOfCacheForceInlined = NumMegabytesForceInline / MaxCacheSizeMB;
 
-	FString CacheMemoryHeader = *FString::Printf(TEXT("Retaining:\t, Loaded:\t, Max Potential Usage:\t, \n"));
-	FString CacheMemoryUsage = *FString::Printf(TEXT("%.4f Megabytes (%.3f of total capacity)\t,  %.4f Megabytes (%lu bytes)\t, %.4f Megabytes\t, \n"), NumMegabytesRetained, PercentageOfCacheRetained, NumMegabytesInUse, MemoryCounterBytes.Load(), MaxCacheSizeMB);
-
+	FString CacheMemoryHeader = *FString::Printf(TEXT("Force Inline:\t, Retaining:\t, Loaded:\t, Max Potential Usage:\t, \n"));
+	FString CacheMemoryUsage = *FString::Printf(TEXT("%.4f Megabytes (%.3f of total capacity)\t %.4f Megabytes (%.3f of total capacity)\t,  %.4f Megabytes (%lu bytes)\t, %.4f Megabytes\t, \n"), 
+		NumMegabytesForceInline, PercentageOfCacheForceInlined, NumMegabytesRetained, PercentageOfCacheRetained, NumMegabytesInUse, MemoryCounterBytes.Load(), MaxCacheSizeMB);
 	OutputString += CacheMemoryHeader + CacheMemoryUsage + TEXT("\n");
 
 	// Second Pass: We're going to list the actual chunks in the cache.
@@ -2216,6 +2269,7 @@ static const FLinearColor ColorLoadInProgress = FLinearColor::Black;
 static const FLinearColor ColorTrimmed = FLinearColor::Red;
 static const FLinearColor ColorCacheMiss = ColorLOD;
 static const FLinearColor ColorOther = FLinearColor::Gray;
+static const FLinearColor ColorForceInline(255 / ColorMax, 0, 255 / ColorMax); // Magenta
 
 
 
@@ -2259,6 +2313,8 @@ TPair<int, int> FAudioChunkCache::DebugDisplay(UWorld* World, FViewport* Viewpor
 	int32 NumTrimmed = 0;
 	int32 NumLoadInProgress = 0;
 	int32 NumOther = 0;
+	// approximate how many chunks the force inline memory is using for debug purposes
+	int32 NumForceInline = ForceInlineMemoryCounterBytes.Load() / MaxChunkSize;
 
 
 	for (int i = 0; i < ChunksInUse; ++i)
@@ -2343,11 +2399,13 @@ TPair<int, int> FAudioChunkCache::DebugDisplay(UWorld* World, FViewport* Viewpor
 		}
 	}
 
+	NumBytesCounter += ForceInlineMemoryCounterBytes.Load();
 	// Convert to megabytes and print the total size:
 	const double NumMegabytesInUse = (double)NumBytesCounter / (1024 * 1024);
 	const double MaxCacheSizeMB = ((double)MemoryLimitBytes) / (1024 * 1024);
 
-	FString CacheMemoryUsage = *FString::Printf(TEXT("Using: %.4f Megabytes (%lu bytes). Max Potential Usage: %.4f Megabytes."), NumMegabytesInUse, MemoryCounterBytes.Load(), MaxCacheSizeMB);
+	FString CacheMemoryUsage = *FString::Printf(TEXT("Using: %.4f Megabytes (%lu bytes). Max Potential Usage: %.4f Megabytes."), 
+		NumMegabytesInUse, MemoryCounterBytes.Load() + ForceInlineMemoryCounterBytes.Load(), MaxCacheSizeMB);
 
 	// We're going to align this horizontally with the number of elements right above it.
 	Canvas->DrawShadowedString(X, Y, *CacheMemoryUsage, UEngine::GetMediumFont(), FLinearColor::White);
@@ -2366,7 +2424,8 @@ TPair<int, int> FAudioChunkCache::DebugDisplay(UWorld* World, FViewport* Viewpor
 		+ NumLOD
 		+ NumTrimmed
 		+ NumLoadInProgress
-		+ NumOther;
+		+ NumOther
+		+ NumForceInline;
 
 	if (NumChunks == 0)
 	{
@@ -2391,6 +2450,7 @@ TPair<int, int> FAudioChunkCache::DebugDisplay(UWorld* World, FViewport* Viewpor
 	const float PercentageTrimmed = (NumTrimmed / NumChunks);
 	const float PercentageLoadInProgress = (NumLoadInProgress / NumChunks);
 	const float PercentageOther = (NumOther / NumChunks);
+	const float PercentageForceInline = (NumForceInline / NumChunks);
 
 	const int32 BarWidthRetainedAndPlaying = PercentageRetainedAndPlaying * BarWidth;
 	const int32 BarWidthRetained = PercentageRetained * BarWidth;
@@ -2405,6 +2465,7 @@ TPair<int, int> FAudioChunkCache::DebugDisplay(UWorld* World, FViewport* Viewpor
 	const int32 BarWidthTrimmed = PercentageTrimmed * BarWidth;
 	const int32 BarWidthLoadInProgress = PercentageLoadInProgress * BarWidth;
 	const int32 BarWidthOther = PercentageOther * BarWidth;
+	const int32 BarWidthForceInline = PercentageForceInline * BarWidth;
 
 
 	// Draw color key
@@ -2431,6 +2492,10 @@ TPair<int, int> FAudioChunkCache::DebugDisplay(UWorld* World, FViewport* Viewpor
 	Canvas->DrawShadowedString(X, Y, *TempString, UEngine::GetSmallFont(), ColorLoadInProgress);
 	Y += 15;
 
+	TempString = *FString::Printf(TEXT("Force Inline: %.2f %%"), 100.f * PercentageForceInline);
+	Canvas->DrawShadowedString(X, Y, *TempString, UEngine::GetSmallFont(), ColorForceInline);
+	Y += 25;
+
 	TempString = *FString::Printf(TEXT("Other: %.2f %%"), 100.f * PercentageOther);
 	Canvas->DrawShadowedString(X, Y, *TempString, UEngine::GetSmallFont(), ColorOther);
 	Y += 24;
@@ -2502,6 +2567,18 @@ TPair<int, int> FAudioChunkCache::DebugDisplay(UWorld* World, FViewport* Viewpor
 	Canvas->DrawTile(CurrHorzOffset, CurrVertOffset, BarWidthOther, BarHeight, 0, 0, 0, 0, ColorOther);
 	CurrHorzOffset += BarWidthOther;
 
+	if (BarWidthForceInline > 0)
+	{
+		// (|| divider between cache and force inline)
+		const int32 DividerWidth = 5;
+		Canvas->DrawTile(CurrHorzOffset, CurrVertOffset, DividerWidth, BarHeight, 0, 0, 0, 0, FLinearColor::Black);
+		CurrHorzOffset += DividerWidth;
+
+		// (Force Inline)
+		Canvas->DrawTile(CurrHorzOffset, CurrVertOffset, BarWidthForceInline - DividerWidth, BarHeight, 0, 0, 0, 0, ColorForceInline);
+		CurrHorzOffset += BarWidthForceInline;	
+	}
+
 	Y = (CurrVertOffset + 24);
 
 	// Draw the body of our display depending on the CVAR
diff --git a/Engine/Source/Runtime/Engine/Public/AudioStreamingCache.h b/Engine/Source/Runtime/Engine/Public/AudioStreamingCache.h
--- a/Engine/Source/Runtime/Engine/Public/AudioStreamingCache.h
+++ b/Engine/Source/Runtime/Engine/Public/AudioStreamingCache.h
@@ -122,9 +122,13 @@ public:
 	void AddNewReferenceToChunk(const FChunkKey& InKey);
 	void RemoveReferenceToChunk(const FChunkKey& InKey);
 
-	// Evict all sounds from the cache.
+	// Evict all sounds from the cache. Does not account for force inline sounds
 	void ClearCache();
 
+	void AddForceInlineSoundWave(const FSoundWaveProxyPtr&);
+
+	void RemoveForecInlineSoundWave(const FSoundWaveProxyPtr&);
+
 	// This function will reclaim memory by freeing as many chunks as needed to free BytesToFree.
 	// returns the amount of bytes we were actually able to free.
 	// It's important to note that this will block any chunk requests.
@@ -396,6 +400,8 @@ private:
 	TAtomic<uint64> MemoryCounterBytes;
 	uint64 MemoryLimitBytes;
 
+	TAtomic<uint64> ForceInlineMemoryCounterBytes;
+
 	// Number of async load operations we have currently in flight.
 	FThreadSafeCounter NumberOfLoadsInFlight;
 
@@ -510,6 +516,8 @@ public:
 	// IAudioStreamingManager interface (unused functions)
 	virtual void AddStreamingSoundWave(const FSoundWaveProxyPtr& SoundWave) override;
 	virtual void RemoveStreamingSoundWave(const FSoundWaveProxyPtr& SoundWave) override;
+	virtual void AddForceInlineSoundWave(const FSoundWaveProxyPtr& SoundWave) override;
+	virtual void RemoveForceInlineSoundWave(const FSoundWaveProxyPtr& SoundWave) override;
 	virtual void AddDecoder(ICompressedAudioInfo* CompressedAudioInfo) override;
 	virtual void RemoveDecoder(ICompressedAudioInfo* CompressedAudioInfo) override;
 	virtual bool IsManagedStreamingSoundWave(const FSoundWaveProxyPtr&  SoundWave) const override;
diff --git a/Engine/Source/Runtime/Engine/Public/ContentStreaming.h b/Engine/Source/Runtime/Engine/Public/ContentStreaming.h
--- a/Engine/Source/Runtime/Engine/Public/ContentStreaming.h
+++ b/Engine/Source/Runtime/Engine/Public/ContentStreaming.h
@@ -520,6 +520,12 @@ struct IAudioStreamingManager : public IStreamingManager
 	/** Removes a Sound Wave from the streaming manager. */
 	virtual void RemoveStreamingSoundWave(const FSoundWaveProxyPtr& SoundWave) = 0;
 
+	/** Adds the memory usage of the force inline sound to the streaming cache budget */
+	virtual void AddForceInlineSoundWave(const FSoundWaveProxyPtr& SoundWave) { };
+
+	/** Removes the memory usage of the force inline sound from the streaming cache budget */
+	virtual void RemoveForceInlineSoundWave(const FSoundWaveProxyPtr& SoundWave) { };
+
 	/** Adds the decoder to the streaming manager to prevent stream chunks from getting reaped from underneath it */
 	virtual void AddDecoder(ICompressedAudioInfo* CompressedAudioInfo) = 0;
 
