Add ability to setup VT support on a per-platform basis

diff --git a/Engine/Plugins/Experimental/VirtualHeightfieldMesh/Source/VirtualHeightfieldMesh/Private/VirtualHeightfieldMeshEnable.cpp b/Engine/Plugins/Experimental/VirtualHeightfieldMesh/Source/VirtualHeightfieldMesh/Private/VirtualHeightfieldMeshEnable.cpp
--- a/Engine/Plugins/Experimental/VirtualHeightfieldMesh/Source/VirtualHeightfieldMesh/Private/VirtualHeightfieldMeshEnable.cpp
+++ b/Engine/Plugins/Experimental/VirtualHeightfieldMesh/Source/VirtualHeightfieldMesh/Private/VirtualHeightfieldMeshEnable.cpp
@@ -67,6 +67,6 @@ namespace VirtualHeightfieldMesh
 
 	bool IsEnabled(FStaticFeatureLevel InFeatureLevel)
 	{
-		return CVarVHMEnable.GetValueOnAnyThread() != 0 && (InFeatureLevel >= ERHIFeatureLevel::SM5) && UseVirtualTexturing(InFeatureLevel);
+		return CVarVHMEnable.GetValueOnAnyThread() != 0 && (InFeatureLevel >= ERHIFeatureLevel::SM5) && UseVirtualTexturing(GetFeatureLevelShaderPlatform(InFeatureLevel));
 	}
 }
diff --git a/Engine/Plugins/FX/Niagara/Source/Niagara/Private/NiagaraDataInterfaceLandscape.cpp b/Engine/Plugins/FX/Niagara/Source/Niagara/Private/NiagaraDataInterfaceLandscape.cpp
--- a/Engine/Plugins/FX/Niagara/Source/Niagara/Private/NiagaraDataInterfaceLandscape.cpp
+++ b/Engine/Plugins/FX/Niagara/Source/Niagara/Private/NiagaraDataInterfaceLandscape.cpp
@@ -1529,7 +1529,7 @@ void UNiagaraDataInterfaceLandscape::ApplyLandscape(const FNiagaraSystemInstance
 	InstanceData.NormalVirtualTextureIndex = INDEX_NONE;
 
 	// only worry about virtual textures if our current platform supports them
-	if (bVirtualTexturesSupported && UseVirtualTexturing(SystemInstance.GetFeatureLevel()))
+	if (bVirtualTexturesSupported && UseVirtualTexturing(GetFeatureLevelShaderPlatform(SystemInstance.GetFeatureLevel())))
 	{
 		const int32 RuntimeVirtualTextureCount = InstanceData.Landscape->RuntimeVirtualTextures.Num();
 		for (int32 TextureIt = 0; TextureIt < RuntimeVirtualTextureCount; ++TextureIt)
diff --git a/Engine/Plugins/Importers/USDImporter/Source/USDUtilities/Private/USDShadeConversion.cpp b/Engine/Plugins/Importers/USDImporter/Source/USDUtilities/Private/USDShadeConversion.cpp
--- a/Engine/Plugins/Importers/USDImporter/Source/USDUtilities/Private/USDShadeConversion.cpp
+++ b/Engine/Plugins/Importers/USDImporter/Source/USDUtilities/Private/USDShadeConversion.cpp
@@ -2793,7 +2793,7 @@ void UsdUtils::NotifyIfVirtualTexturesNeeded( UTexture* Texture )
 		}
 	}
 
-	if ( !UseVirtualTexturing( GMaxRHIFeatureLevel ) )
+	if ( !UseVirtualTexturing( GMaxRHIShaderPlatform ) )
 	{
 		FUsdLogManager::LogMessage(
 			EMessageSeverity::Warning,
diff --git a/Engine/Plugins/Runtime/MeshModelingToolset/Source/ModelingComponents/Private/AssetUtils/Texture2DUtil.cpp b/Engine/Plugins/Runtime/MeshModelingToolset/Source/ModelingComponents/Private/AssetUtils/Texture2DUtil.cpp
--- a/Engine/Plugins/Runtime/MeshModelingToolset/Source/ModelingComponents/Private/AssetUtils/Texture2DUtil.cpp
+++ b/Engine/Plugins/Runtime/MeshModelingToolset/Source/ModelingComponents/Private/AssetUtils/Texture2DUtil.cpp
@@ -318,7 +318,7 @@ bool UE::AssetUtils::ConvertToSingleChannel(UTexture2D* TextureMap)
 bool UE::AssetUtils::ForceVirtualTexturePrefetch(FImageDimensions ScreenSpaceDimensions, bool bWaitForPrefetchToComplete)
 {
 	// Prefetch all virtual textures so that we have content available
-	if (UseVirtualTexturing(GMaxRHIFeatureLevel))
+	if (UseVirtualTexturing(GMaxRHIShaderPlatform))
 	{
 		const FVector2D ScreenSpaceSize(ScreenSpaceDimensions.GetWidth(), ScreenSpaceDimensions.GetHeight());
 
diff --git a/Engine/Source/Developer/Android/AndroidTargetPlatform/Private/AndroidTargetPlatform.cpp b/Engine/Source/Developer/Android/AndroidTargetPlatform/Private/AndroidTargetPlatform.cpp
--- a/Engine/Source/Developer/Android/AndroidTargetPlatform/Private/AndroidTargetPlatform.cpp
+++ b/Engine/Source/Developer/Android/AndroidTargetPlatform/Private/AndroidTargetPlatform.cpp
@@ -265,6 +265,7 @@ FAndroidTargetPlatform::FAndroidTargetPlatform(bool bInIsClient, const TCHAR* Fl
 	, MobileShadingPath(0)
 	, bDistanceField(false)
 	, bMobileForwardEnableClusteredReflections(false)
+	, bMobileVirtualTextures(false)
 
 {
 #if WITH_ENGINE
@@ -273,6 +274,7 @@ FAndroidTargetPlatform::FAndroidTargetPlatform(bool bInIsClient, const TCHAR* Fl
 	GetConfigSystem()->GetBool(TEXT("/Script/Engine.RendererSettings"), TEXT("r.DistanceFields"), bDistanceField, GEngineIni);
 	GetConfigSystem()->GetInt(TEXT("/Script/Engine.RendererSettings"), TEXT("r.Mobile.ShadingPath"), MobileShadingPath, GEngineIni);
 	GetConfigSystem()->GetBool(TEXT("/Script/Engine.RendererSettings"), TEXT("r.Mobile.Forward.EnableClusteredReflections"), bMobileForwardEnableClusteredReflections, GEngineIni);
+	GetConfigSystem()->GetBool(TEXT("/Script/Engine.RendererSettings"), TEXT("r.Mobile.VirtualTextures"), bMobileVirtualTextures, GEngineIni);
 #endif
 
 	TickDelegate = FTickerDelegate::CreateRaw(this, &FAndroidTargetPlatform::HandleTicker);
@@ -290,12 +292,6 @@ FAndroidTargetDevicePtr FAndroidTargetPlatform::CreateTargetDevice(const ITarget
 	return MakeShareable(new FAndroidTargetDevice(InTargetPlatform, InSerialNumber, InAndroidVariant));
 }
 
-static bool UsesVirtualTextures()
-{
-	static auto* CVarMobileVirtualTextures = IConsoleManager::Get().FindTConsoleVariableDataInt(TEXT("r.Mobile.VirtualTextures"));
-	return CVarMobileVirtualTextures->GetValueOnAnyThread() != 0;
-}
-
 bool FAndroidTargetPlatform::SupportsES31() const
 {
 	// default no support for ES31
@@ -409,7 +405,8 @@ bool FAndroidTargetPlatform::SupportsFeature( ETargetPlatformFeatures Feature )
 			return SupportsVulkanSM5();
 
 		case ETargetPlatformFeatures::VirtualTextureStreaming:
-			return UsesVirtualTextures();
+			// TODO: should it check r.VirtualTextures for SM5 renderer ?
+			return bMobileVirtualTextures;
 
 		case ETargetPlatformFeatures::DistanceFieldAO:
 			return UsesDistanceFields();
diff --git a/Engine/Source/Developer/Android/AndroidTargetPlatform/Private/AndroidTargetPlatform.h b/Engine/Source/Developer/Android/AndroidTargetPlatform/Private/AndroidTargetPlatform.h
--- a/Engine/Source/Developer/Android/AndroidTargetPlatform/Private/AndroidTargetPlatform.h
+++ b/Engine/Source/Developer/Android/AndroidTargetPlatform/Private/AndroidTargetPlatform.h
@@ -219,6 +219,9 @@ protected:
 
 	// r.Mobile.Forward.EnableClusteredReflections value
 	bool bMobileForwardEnableClusteredReflections;
+
+	// r.Mobile.VirtualTextures value
+	bool bMobileVirtualTextures;
 	
 #if WITH_ENGINE
 	// Holds a cache of the target LOD settings.
diff --git a/Engine/Source/Developer/IOS/IOSTargetPlatform/Private/IOSTargetPlatform.cpp b/Engine/Source/Developer/IOS/IOSTargetPlatform/Private/IOSTargetPlatform.cpp
--- a/Engine/Source/Developer/IOS/IOSTargetPlatform/Private/IOSTargetPlatform.cpp
+++ b/Engine/Source/Developer/IOS/IOSTargetPlatform/Private/IOSTargetPlatform.cpp
@@ -36,6 +36,7 @@ FIOSTargetPlatform::FIOSTargetPlatform(bool bInIsTVOS, bool bInIsVisionOS, bool
 	, MobileShadingPath(0)
 	, bDistanceField(false)
 	, bMobileForwardEnableClusteredReflections(false)
+	, bMobileVirtualTextures(false)
 {
 #if WITH_ENGINE
 	TextureLODSettings = nullptr; // TextureLODSettings are registered by the device profile.
@@ -43,6 +44,7 @@ FIOSTargetPlatform::FIOSTargetPlatform(bool bInIsTVOS, bool bInIsVisionOS, bool
 	GetConfigSystem()->GetBool(TEXT("/Script/Engine.RendererSettings"), TEXT("r.DistanceFields"), bDistanceField, GEngineIni);
 	GetConfigSystem()->GetInt(TEXT("/Script/Engine.RendererSettings"), TEXT("r.Mobile.ShadingPath"), MobileShadingPath, GEngineIni);
 	GetConfigSystem()->GetBool(TEXT("/Script/Engine.RendererSettings"), TEXT("r.Mobile.Forward.EnableClusteredReflections"), bMobileForwardEnableClusteredReflections, GEngineIni);
+	GetConfigSystem()->GetBool(TEXT("/Script/Engine.RendererSettings"), TEXT("r.Mobile.VirtualTextures"), bMobileVirtualTextures, GEngineIni);
 #endif // #if WITH_ENGINE
 
 	// initialize the connected device detector
@@ -440,12 +442,6 @@ void FIOSTargetPlatform::HandleDeviceDisconnected(const FIOSLaunchDaemonPong& Me
 
 /* ITargetPlatform interface
  *****************************************************************************/
-static bool UsesVirtualTextures()
-{
-	static auto* CVarMobileVirtualTextures = IConsoleManager::Get().FindTConsoleVariableDataInt(TEXT("r.Mobile.VirtualTextures"));
-	return CVarMobileVirtualTextures->GetValueOnAnyThread() != 0;
-}
-
 static bool SupportsMetal()
 {
 	// default to NOT supporting metal
@@ -495,7 +491,8 @@ bool FIOSTargetPlatform::SupportsFeature( ETargetPlatformFeatures Feature ) cons
 			return SupportsMetalMRT();
 
 		case ETargetPlatformFeatures::VirtualTextureStreaming:
-			return UsesVirtualTextures();
+			// TODO: should it check r.VirtualTextures for SM5 renderer?
+			return bMobileVirtualTextures;
 
 		case ETargetPlatformFeatures::DistanceFieldAO:
 			return UsesDistanceFields();
diff --git a/Engine/Source/Developer/IOS/IOSTargetPlatform/Private/IOSTargetPlatform.h b/Engine/Source/Developer/IOS/IOSTargetPlatform/Private/IOSTargetPlatform.h
--- a/Engine/Source/Developer/IOS/IOSTargetPlatform/Private/IOSTargetPlatform.h
+++ b/Engine/Source/Developer/IOS/IOSTargetPlatform/Private/IOSTargetPlatform.h
@@ -142,6 +142,9 @@ private:
 	// r.Mobile.Forward.EnableClusteredReflections value
 	bool bMobileForwardEnableClusteredReflections;
 
+	// r.Mobile.VirtualTextures value
+	bool bMobileVirtualTextures;
+
 #if WITH_ENGINE
 	// Holds the cache of the target LOD settings.
 	const UTextureLODSettings* TextureLODSettings;
diff --git a/Engine/Source/Developer/MaterialBaking/Private/MaterialBakingModule.cpp b/Engine/Source/Developer/MaterialBaking/Private/MaterialBakingModule.cpp
--- a/Engine/Source/Developer/MaterialBaking/Private/MaterialBakingModule.cpp
+++ b/Engine/Source/Developer/MaterialBaking/Private/MaterialBakingModule.cpp
@@ -728,7 +728,7 @@ private:
 
 				// Virtual textures may require repeated rendering to warm up.
 				int32 WarmupIterationCount = 1;
-				if (UseVirtualTexturing(ViewFamily.GetFeatureLevel()))
+				if (UseVirtualTexturing(ViewFamily.GetShaderPlatform()))
 				{
 					const FMaterial& MeshMaterial = ExportMaterialProxy->GetIncompleteMaterialWithFallback(ViewFamily.GetFeatureLevel());
 					if (!MeshMaterial.GetUniformVirtualTextureExpressions().IsEmpty())
diff --git a/Engine/Source/Editor/UnrealEd/Private/Lightmass/LightmassRender.cpp b/Engine/Source/Editor/UnrealEd/Private/Lightmass/LightmassRender.cpp
--- a/Engine/Source/Editor/UnrealEd/Private/Lightmass/LightmassRender.cpp
+++ b/Engine/Source/Editor/UnrealEd/Private/Lightmass/LightmassRender.cpp
@@ -1311,7 +1311,7 @@ bool FLightmassMaterialRenderer::GenerateMaterialPropertyData(
 
 				// Prefetch all virtual textures so that we have content available
 				//todo[vt]: Move this to calling function to avoid multiple prefetches
-				if (UseVirtualTexturing(GMaxRHIFeatureLevel))
+				if (UseVirtualTexturing(GMaxRHIShaderPlatform))
 				{					
 					const ERHIFeatureLevel::Type FeatureLevel = GMaxRHIFeatureLevel;
 					const FVector2D ScreenSpaceSize(InOutSizeX, InOutSizeY);
diff --git a/Engine/Source/Runtime/Engine/Classes/Materials/MaterialExpressionTextureBase.h b/Engine/Source/Runtime/Engine/Classes/Materials/MaterialExpressionTextureBase.h
--- a/Engine/Source/Runtime/Engine/Classes/Materials/MaterialExpressionTextureBase.h
+++ b/Engine/Source/Runtime/Engine/Classes/Materials/MaterialExpressionTextureBase.h
@@ -74,7 +74,7 @@ class UMaterialExpressionTextureBase : public UMaterialExpression
 	 * @param OutErrorMessage - If 'false' is returned, will contain a message describing the error
 	 */
 	static ENGINE_API bool VerifySamplerType(
-		ERHIFeatureLevel::Type FeatureLevel,
+		EShaderPlatform ShaderPlatform,
 		const ITargetPlatform* TargetPlatform,
 		const UTexture* Texture,
 		EMaterialSamplerType SamplerType,
diff --git a/Engine/Source/Runtime/Engine/Private/Components/RuntimeVirtualTextureComponent.cpp b/Engine/Source/Runtime/Engine/Private/Components/RuntimeVirtualTextureComponent.cpp
--- a/Engine/Source/Runtime/Engine/Private/Components/RuntimeVirtualTextureComponent.cpp
+++ b/Engine/Source/Runtime/Engine/Private/Components/RuntimeVirtualTextureComponent.cpp
@@ -97,7 +97,7 @@ void URuntimeVirtualTextureComponent::GetHidePrimitiveSettings(bool& OutHidePrim
 bool URuntimeVirtualTextureComponent::IsVisible() const
 {
 	// Make sure to have the component do nothing if VT is disabled or if the world is not compatible with RVT
-	return Super::IsVisible() && IsActiveInWorld() && CVarRVTEnableVolumes.GetValueOnGameThread() != 0 && UseVirtualTexturing(GetScene()->GetFeatureLevel());
+	return Super::IsVisible() && IsActiveInWorld() && CVarRVTEnableVolumes.GetValueOnGameThread() != 0 && UseVirtualTexturing(GetScene()->GetShaderPlatform());
 }
 
 void URuntimeVirtualTextureComponent::ApplyWorldOffset(const FVector& InOffset, bool bWorldShift)
diff --git a/Engine/Source/Runtime/Engine/Private/HLSLTree/HLSLTreeCommon.cpp b/Engine/Source/Runtime/Engine/Private/HLSLTree/HLSLTreeCommon.cpp
--- a/Engine/Source/Runtime/Engine/Private/HLSLTree/HLSLTreeCommon.cpp
+++ b/Engine/Source/Runtime/Engine/Private/HLSLTree/HLSLTreeCommon.cpp
@@ -1219,7 +1219,7 @@ bool FExpressionVirtualTextureFeatureSwitch::IsInputActive(const FEmitContext& C
 		return true;
 	}
 
-	if (UseVirtualTexturing(Context.TargetParameters.FeatureLevel, Context.TargetParameters.TargetPlatform))
+	if (UseVirtualTexturing(Context.TargetParameters.ShaderPlatform))
 	{
 		return Index == 0;
 	}
diff --git a/Engine/Source/Runtime/Engine/Private/InstancedStaticMesh.cpp b/Engine/Source/Runtime/Engine/Private/InstancedStaticMesh.cpp
--- a/Engine/Source/Runtime/Engine/Private/InstancedStaticMesh.cpp
+++ b/Engine/Source/Runtime/Engine/Private/InstancedStaticMesh.cpp
@@ -3273,7 +3273,7 @@ void UInstancedStaticMeshComponent::GetStaticLightingInfo(FStaticLightingPrimiti
 void UInstancedStaticMeshComponent::ApplyLightMapping(FStaticLightingTextureMapping_InstancedStaticMesh* InMapping, ULevel* LightingScenario)
 {
 	static const auto CVar = IConsoleManager::Get().FindTConsoleVariableDataInt(TEXT("r.VirtualTexturedLightmaps"));
-	const bool bUseVirtualTextures = (CVar->GetValueOnAnyThread() != 0) && UseVirtualTexturing(GMaxRHIFeatureLevel);
+	const bool bUseVirtualTextures = (CVar->GetValueOnAnyThread() != 0) && UseVirtualTexturing(GMaxRHIShaderPlatform);
 
 	NumPendingLightmaps--;
 
diff --git a/Engine/Source/Runtime/Engine/Private/LightMap.cpp b/Engine/Source/Runtime/Engine/Private/LightMap.cpp
--- a/Engine/Source/Runtime/Engine/Private/LightMap.cpp
+++ b/Engine/Source/Runtime/Engine/Private/LightMap.cpp
@@ -1194,7 +1194,7 @@ void FLightMapPendingTexture::CreateUObjects()
 	++GLightmapCounter;
 	
 	// Only build VT lightmaps if they are enabled
-	const bool bUseVirtualTextures = (CVarVirtualTexturedLightMaps.GetValueOnAnyThread() != 0) && UseVirtualTexturing(GMaxRHIFeatureLevel);
+	const bool bUseVirtualTextures = (CVarVirtualTexturedLightMaps.GetValueOnAnyThread() != 0) && UseVirtualTexturing(GMaxRHIShaderPlatform);
 	const bool bIncludeNonVirtualTextures = !bUseVirtualTextures || (CVarIncludeNonVirtualTexturedLightMaps.GetValueOnAnyThread() != 0);
 	
 	if (bIncludeNonVirtualTextures)
@@ -2395,7 +2395,7 @@ void FLightMap2D::EncodeTextures( UWorld* InWorld, ULevel* LightingScenario, boo
 #if WITH_EDITOR
 	if (bLightingSuccessful)
 	{
-		const bool bUseVirtualTextures = (CVarVirtualTexturedLightMaps.GetValueOnAnyThread() != 0) && UseVirtualTexturing(GMaxRHIFeatureLevel);
+		const bool bUseVirtualTextures = (CVarVirtualTexturedLightMaps.GetValueOnAnyThread() != 0) && UseVirtualTexturing(GMaxRHIShaderPlatform);
 		const bool bIncludeNonVirtualTextures = !bUseVirtualTextures || (CVarIncludeNonVirtualTexturedLightMaps.GetValueOnAnyThread() != 0);
 
 		GWarn->BeginSlowTask( NSLOCTEXT("LightMap2D", "BeginEncodingLightMapsTask", "Encoding light-maps"), false );
@@ -2974,7 +2974,7 @@ void FLightMap2D::Serialize(FArchive& Ar)
 
 	FLightMap::Serialize(Ar);
 
-	const bool bUsingVTLightmaps = (CVarVirtualTexturedLightMaps.GetValueOnAnyThread() != 0) && UseVirtualTexturing(GMaxRHIFeatureLevel, Ar.CookingTarget());
+	const bool bUsingVTLightmaps = (CVarVirtualTexturedLightMaps.GetValueOnAnyThread() != 0) && UseVirtualTexturing(GMaxRHIShaderPlatform, Ar.CookingTarget());
 
 	if( Ar.IsLoading() && Ar.UEVer() < VER_UE4_LOW_QUALITY_DIRECTIONAL_LIGHTMAPS )
 	{
@@ -3186,7 +3186,7 @@ FLightMapInteraction FLightMap2D::GetInteraction(ERHIFeatureLevel::Type InFeatur
 
 	int32 LightmapIndex = bHighQuality ? 0 : 1;
 
-	const bool bUseVirtualTextures = (CVarVirtualTexturedLightMaps.GetValueOnAnyThread() != 0) && UseVirtualTexturing(InFeatureLevel);
+	const bool bUseVirtualTextures = (CVarVirtualTexturedLightMaps.GetValueOnAnyThread() != 0) && UseVirtualTexturing(GetFeatureLevelShaderPlatform(InFeatureLevel));
 	if (!bUseVirtualTextures)
 	{
 		bool bValidTextures = Textures[LightmapIndex] && Textures[LightmapIndex]->GetResource();
@@ -3216,7 +3216,7 @@ FShadowMapInteraction FLightMap2D::GetShadowInteraction(ERHIFeatureLevel::Type I
 
 	int32 LightmapIndex = bHighQuality ? 0 : 1;
 
-	const bool bUseVirtualTextures = (CVarVirtualTexturedLightMaps.GetValueOnAnyThread() != 0) && UseVirtualTexturing(InFeatureLevel);
+	const bool bUseVirtualTextures = (CVarVirtualTexturedLightMaps.GetValueOnAnyThread() != 0) && UseVirtualTexturing(GetFeatureLevelShaderPlatform(InFeatureLevel));
 	if (bUseVirtualTextures)
 	{
 		// Preview lightmaps don't stream from disk, thus no FVirtualTexture2DResource
@@ -3368,7 +3368,7 @@ FLightmapResourceCluster::~FLightmapResourceCluster()
 
 bool FLightmapResourceCluster::GetUseVirtualTexturing() const
 {
-	return (CVarVirtualTexturedLightMaps.GetValueOnRenderThread() != 0) && UseVirtualTexturing(GetFeatureLevel());
+	return (CVarVirtualTexturedLightMaps.GetValueOnRenderThread() != 0) && UseVirtualTexturing(GetFeatureLevelShaderPlatform(GetFeatureLevel()));
 }
 
 // Two stage initialization of FLightmapResourceCluster
diff --git a/Engine/Source/Runtime/Engine/Private/LightmapUniformShaderParameters.cpp b/Engine/Source/Runtime/Engine/Private/LightmapUniformShaderParameters.cpp
--- a/Engine/Source/Runtime/Engine/Private/LightmapUniformShaderParameters.cpp
+++ b/Engine/Source/Runtime/Engine/Private/LightmapUniformShaderParameters.cpp
@@ -65,7 +65,7 @@ void GetPrecomputedLightingParameters(
 )
 {
 	static const auto CVar = IConsoleManager::Get().FindTConsoleVariableDataInt(TEXT("r.VirtualTexturedLightmaps"));
-	const bool bUseVirtualTextures = (CVar->GetValueOnRenderThread() != 0) && UseVirtualTexturing(FeatureLevel);
+	const bool bUseVirtualTextures = (CVar->GetValueOnRenderThread() != 0) && UseVirtualTexturing(GetFeatureLevelShaderPlatform(FeatureLevel));
 
 	// TDistanceFieldShadowsAndLightMapPolicy
 	const FShadowMapInteraction ShadowMapInteraction = LCI ? LCI->GetShadowMapInteraction(FeatureLevel) : FShadowMapInteraction();
diff --git a/Engine/Source/Runtime/Engine/Private/MapBuildData.cpp b/Engine/Source/Runtime/Engine/Private/MapBuildData.cpp
--- a/Engine/Source/Runtime/Engine/Private/MapBuildData.cpp
+++ b/Engine/Source/Runtime/Engine/Private/MapBuildData.cpp
@@ -973,7 +973,7 @@ bool UMapBuildDataRegistry::IsLightingValid(ERHIFeatureLevel::Type InFeatureLeve
 	}
 	else
 	{
-		const bool bUsingVTLightmaps = UseVirtualTextureLightmap(InFeatureLevel);
+		const bool bUsingVTLightmaps = UseVirtualTextureLightmap(GetFeatureLevelShaderPlatform(InFeatureLevel));
 
 		// this code checks if AT LEAST 1 virtual textures is valid. 
 		for (auto MeshBuildDataPair : MeshBuildData)
diff --git a/Engine/Source/Runtime/Engine/Private/Materials/HLSLMaterialTranslator.cpp b/Engine/Source/Runtime/Engine/Private/Materials/HLSLMaterialTranslator.cpp
--- a/Engine/Source/Runtime/Engine/Private/Materials/HLSLMaterialTranslator.cpp
+++ b/Engine/Source/Runtime/Engine/Private/Materials/HLSLMaterialTranslator.cpp
@@ -1038,7 +1038,7 @@ bool FHLSLMaterialTranslator::Translate()
 		// Validate some things on the VT system. Since generated code for expressions shared between multiple properties
 		// (e.g. a texture sample connected to both diffuse and opacity mask) is reused we can't check based on the MaterialProperty
 		// variable inside the actual code generation pass. So we do a pre-pass over it here.
-		if (UseVirtualTexturing(FeatureLevel, TargetPlatform))
+		if (UseVirtualTexturing(Platform))
 		{
 			ValidateVtPropertyLimits();
 		}
@@ -7640,7 +7640,7 @@ int32 FHLSLMaterialTranslator::Texture(UTexture* InTexture, int32& TextureRefere
 #endif
 	checkf(TextureReferenceIndex != INDEX_NONE, TEXT("Material expression called Compiler->Texture() without implementing UMaterialExpression::GetReferencedTexture properly"));
 
-	const bool bVirtualTexturesEnabeled = UseVirtualTexturing(FeatureLevel, TargetPlatform);
+	const bool bVirtualTexturesEnabeled = UseVirtualTexturing(Platform);
 	bool bVirtual = ShaderType == MCT_TextureVirtual;
 	if (bVirtualTexturesEnabeled == false && ShaderType == MCT_TextureVirtual)
 	{
@@ -7674,7 +7674,7 @@ int32 FHLSLMaterialTranslator::TextureParameter(FName ParameterName, UTexture* I
 	FMaterialParameterInfo ParameterInfo = GetParameterAssociationInfo();
 	ParameterInfo.Name = ParameterName;
 
-	const bool bVirtualTexturesEnabled = UseVirtualTexturing(FeatureLevel, TargetPlatform);
+	const bool bVirtualTexturesEnabled = UseVirtualTexturing(Platform);
 	bool bVirtual = ShaderType == MCT_TextureVirtual;
 	if (bVirtualTexturesEnabled == false && ShaderType == MCT_TextureVirtual)
 	{
@@ -7686,7 +7686,7 @@ int32 FHLSLMaterialTranslator::TextureParameter(FName ParameterName, UTexture* I
 
 int32 FHLSLMaterialTranslator::VirtualTexture(URuntimeVirtualTexture* InTexture, int32 TextureLayerIndex, int32 PageTableLayerIndex, int32& TextureReferenceIndex, EMaterialSamplerType SamplerType) 
 {
-	if (!UseVirtualTexturing(FeatureLevel, TargetPlatform))
+	if (!UseVirtualTexturing(Platform))
 	{
 		return INDEX_NONE;
 	}
@@ -7699,7 +7699,7 @@ int32 FHLSLMaterialTranslator::VirtualTexture(URuntimeVirtualTexture* InTexture,
 
 int32 FHLSLMaterialTranslator::VirtualTextureParameter(FName ParameterName, URuntimeVirtualTexture* InDefaultValue, int32 TextureLayerIndex, int32 PageTableLayerIndex, int32& TextureReferenceIndex, EMaterialSamplerType SamplerType)
 {
-	if (!UseVirtualTexturing(FeatureLevel, TargetPlatform))
+	if (!UseVirtualTexturing(Platform))
 	{
 		return INDEX_NONE;
 	}
@@ -7737,7 +7737,7 @@ int32 FHLSLMaterialTranslator::VirtualTextureUniform(FName ParameterName, int32
 
 int32 FHLSLMaterialTranslator::VirtualTextureWorldToUV(int32 WorldPositionIndex, int32 P0, int32 P1, int32 P2)
 {
-	if (!UseVirtualTexturing(FeatureLevel, TargetPlatform) || WorldPositionIndex == INDEX_NONE || P0 == INDEX_NONE || P1 == INDEX_NONE || P2 == INDEX_NONE)
+	if (!UseVirtualTexturing(Platform) || WorldPositionIndex == INDEX_NONE || P0 == INDEX_NONE || P1 == INDEX_NONE || P2 == INDEX_NONE)
 	{
 		return INDEX_NONE;
 	}
diff --git a/Engine/Source/Runtime/Engine/Private/Materials/Material.cpp b/Engine/Source/Runtime/Engine/Private/Materials/Material.cpp
--- a/Engine/Source/Runtime/Engine/Private/Materials/Material.cpp
+++ b/Engine/Source/Runtime/Engine/Private/Materials/Material.cpp
@@ -5910,7 +5910,9 @@ bool UMaterial::GetAllReferencedExpressions(TArray<UMaterialExpression*>& OutExp
 			}
 		}
 
-		bool bMobileUseVirtualTexturing = UseVirtualTexturing(ERHIFeatureLevel::ES3_1);
+		// TODO: Need an actual ShaderPlatform for a more precise result
+		EShaderPlatform ShaderPlatform = GetFeatureLevelShaderPlatform(ERHIFeatureLevel::ES3_1);
+		bool bMobileUseVirtualTexturing = UseVirtualTexturing(ShaderPlatform);
 		if (bMobileUseVirtualTexturing)
 		{
 			TArray<class UMaterialExpressionCustomOutput*> CustomOutputExpressions;
diff --git a/Engine/Source/Runtime/Engine/Private/Materials/MaterialExpressions.cpp b/Engine/Source/Runtime/Engine/Private/Materials/MaterialExpressions.cpp
--- a/Engine/Source/Runtime/Engine/Private/Materials/MaterialExpressions.cpp
+++ b/Engine/Source/Runtime/Engine/Private/Materials/MaterialExpressions.cpp
@@ -2695,7 +2695,7 @@ FName UMaterialExpressionTextureSample::GetInputName(int32 InputIndex) const
 #undef IF_INPUT_RETURN
 
 bool UMaterialExpressionTextureBase::VerifySamplerType(
-	ERHIFeatureLevel::Type FeatureLevel,
+	EShaderPlatform ShaderPlatform,
 	const ITargetPlatform* TargetPlatform,
 	const UTexture* Texture,
 	EMaterialSamplerType SamplerType,
@@ -2705,7 +2705,7 @@ bool UMaterialExpressionTextureBase::VerifySamplerType(
 	{
 		EMaterialSamplerType CorrectSamplerType = UMaterialExpressionTextureBase::GetSamplerTypeForTexture( Texture );
 		bool bIsVirtualTextured = IsVirtualSamplerType(SamplerType);
-		if (bIsVirtualTextured && !UseVirtualTexturing(FeatureLevel, TargetPlatform))
+		if (bIsVirtualTextured && !UseVirtualTexturing(ShaderPlatform))
 		{
 			SamplerType = UMaterialExpressionTextureBase::GetSamplerTypeForTexture(Texture, !bIsVirtualTextured);
 		}
@@ -2795,7 +2795,7 @@ int32 UMaterialExpressionTextureSample::Compile(class FMaterialCompiler* Compile
 		}
 
 		FString SamplerTypeError;
-		if (EffectiveTexture && VerifySamplerType(Compiler->GetFeatureLevel(), Compiler->GetTargetPlatform(), EffectiveTexture, EffectiveSamplerType, SamplerTypeError))
+		if (EffectiveTexture && VerifySamplerType(Compiler->GetShaderPlatform(), Compiler->GetTargetPlatform(), EffectiveTexture, EffectiveSamplerType, SamplerTypeError))
 		{
 			if (TextureCodeIndex != INDEX_NONE)
 			{
@@ -3680,7 +3680,7 @@ int32 UMaterialExpressionVirtualTextureFeatureSwitch::Compile(class FMaterialCom
 		return Compiler->Errorf(TEXT("Missing VirtualTextureFeatureSwitch input 'No'"));
 	}
 
-	if (UseVirtualTexturing(Compiler->GetFeatureLevel(), Compiler->GetTargetPlatform()))
+	if (UseVirtualTexturing(Compiler->GetShaderPlatform()))
 	{
 		return Yes.Compile(Compiler);
 	}
@@ -3765,7 +3765,7 @@ int32 UMaterialExpressionTextureSampleParameter::Compile(class FMaterialCompiler
 	}
 
 	FString SamplerTypeError;
-	if (!VerifySamplerType(Compiler->GetFeatureLevel(), Compiler->GetTargetPlatform(), Texture, SamplerType, SamplerTypeError))
+	if (!VerifySamplerType(Compiler->GetShaderPlatform(), Compiler->GetTargetPlatform(), Texture, SamplerType, SamplerTypeError))
 	{
 		return Compiler->Errorf(TEXT("%s"), *SamplerTypeError);
 	}
@@ -4567,7 +4567,7 @@ int32 UMaterialExpressionTextureSampleParameterSubUV::Compile(class FMaterialCom
 	}
 
 	FString SamplerTypeError;
-	if (!VerifySamplerType(Compiler->GetFeatureLevel(), Compiler->GetTargetPlatform(), Texture, SamplerType, SamplerTypeError))
+	if (!VerifySamplerType(Compiler->GetShaderPlatform(), Compiler->GetTargetPlatform(), Texture, SamplerType, SamplerTypeError))
 	{
 		return Compiler->Errorf(TEXT("%s"), *SamplerTypeError);
 	}
@@ -10601,7 +10601,7 @@ int32 UMaterialExpressionParticleSubUV::Compile(class FMaterialCompiler* Compile
 	if (TextureToCompile)
 	{
 		FString SamplerTypeError;
-		if (!VerifySamplerType(Compiler->GetFeatureLevel(), Compiler->GetTargetPlatform(), TextureToCompile, SamplerTypeToUse, SamplerTypeError))
+		if (!VerifySamplerType(Compiler->GetShaderPlatform(), Compiler->GetTargetPlatform(), TextureToCompile, SamplerTypeToUse, SamplerTypeError))
 		{
 			return Compiler->Errorf(TEXT("%s"), *SamplerTypeError);
 		}
@@ -12831,7 +12831,7 @@ int32 UMaterialExpressionFontSample::Compile(class FMaterialCompiler* Compiler,
 		}
 
 		FString SamplerTypeError;
-		if (!UMaterialExpressionTextureBase::VerifySamplerType(Compiler->GetFeatureLevel(), Compiler->GetTargetPlatform(), Texture, ExpectedSamplerType, SamplerTypeError))
+		if (!UMaterialExpressionTextureBase::VerifySamplerType(Compiler->GetShaderPlatform(), Compiler->GetTargetPlatform(), Texture, ExpectedSamplerType, SamplerTypeError))
 		{
 			return Compiler->Errorf(TEXT("%s"), *SamplerTypeError);
 		}
@@ -12937,7 +12937,7 @@ int32 UMaterialExpressionFontSampleParameter::Compile(class FMaterialCompiler* C
 		}
 
 		FString SamplerTypeError;
-		if (!UMaterialExpressionTextureBase::VerifySamplerType(Compiler->GetFeatureLevel(), Compiler->GetTargetPlatform(), Texture, ExpectedSamplerType, SamplerTypeError))
+		if (!UMaterialExpressionTextureBase::VerifySamplerType(Compiler->GetShaderPlatform(), Compiler->GetTargetPlatform(), Texture, ExpectedSamplerType, SamplerTypeError))
 		{
 			return Compiler->Errorf(TEXT("%s"), *SamplerTypeError);
 		}
@@ -20374,7 +20374,7 @@ int32 UMaterialExpressionAntialiasedTextureMask::Compile(class FMaterialCompiler
 	}
 
 	FString SamplerTypeError;
-	if (!VerifySamplerType(Compiler->GetFeatureLevel(), Compiler->GetTargetPlatform(), Texture, SamplerType, SamplerTypeError))
+	if (!VerifySamplerType(Compiler->GetShaderPlatform(), Compiler->GetTargetPlatform(), Texture, SamplerType, SamplerTypeError))
 	{
 		return Compiler->Errorf(TEXT("%s"), *SamplerTypeError);
 	}
diff --git a/Engine/Source/Runtime/Engine/Private/Materials/MaterialInstance.cpp b/Engine/Source/Runtime/Engine/Private/Materials/MaterialInstance.cpp
--- a/Engine/Source/Runtime/Engine/Private/Materials/MaterialInstance.cpp
+++ b/Engine/Source/Runtime/Engine/Private/Materials/MaterialInstance.cpp
@@ -1292,7 +1292,7 @@ void UMaterialInstance::ValidateTextureOverrides(ERHIFeatureLevel::Type InFeatur
 
 	const UMaterial* Material = GetMaterial();
 	const FMaterialResource* CurrentResource = Material->GetMaterialResource(InFeatureLevel);
-	const bool bShouldValidateVTUsage = UseVirtualTexturing(GMaxRHIFeatureLevel);
+	const bool bShouldValidateVTUsage = UseVirtualTexturing(GMaxRHIShaderPlatform);
 
 	if (!CurrentResource)
 	{
diff --git a/Engine/Source/Runtime/Engine/Private/ModelLight.cpp b/Engine/Source/Runtime/Engine/Private/ModelLight.cpp
--- a/Engine/Source/Runtime/Engine/Private/ModelLight.cpp
+++ b/Engine/Source/Runtime/Engine/Private/ModelLight.cpp
@@ -998,7 +998,7 @@ void UModel::ApplyStaticLighting(ULevel* LightingScenario)
 {
 #if WITH_EDITOR
 	static const auto CVar = IConsoleManager::Get().FindTConsoleVariableDataInt(TEXT("r.VirtualTexturedLightmaps"));
-	const bool bUseVirtualTextures = (CVar->GetValueOnAnyThread() != 0) && UseVirtualTexturing(GMaxRHIFeatureLevel);
+	const bool bUseVirtualTextures = (CVar->GetValueOnAnyThread() != 0) && UseVirtualTexturing(GMaxRHIShaderPlatform);
 
 	check(CachedMappings[0]->QuantizedData);
 
diff --git a/Engine/Source/Runtime/Engine/Private/PrimitiveSceneProxy.cpp b/Engine/Source/Runtime/Engine/Private/PrimitiveSceneProxy.cpp
--- a/Engine/Source/Runtime/Engine/Private/PrimitiveSceneProxy.cpp
+++ b/Engine/Source/Runtime/Engine/Private/PrimitiveSceneProxy.cpp
@@ -596,7 +596,7 @@ FPrimitiveSceneProxy::FPrimitiveSceneProxy(const FPrimitiveSceneProxyDesc& InPro
 	}	
 
 	// Setup the runtime virtual texture information
-	if (UseVirtualTexturing(GetScene().GetFeatureLevel()))
+	if (UseVirtualTexturing(GetScene().GetShaderPlatform()))
 	{
 		for (URuntimeVirtualTexture* VirtualTexture : InProxyDesc.GetRuntimeVirtualTextures())
 		{
diff --git a/Engine/Source/Runtime/Engine/Private/SceneManagement.cpp b/Engine/Source/Runtime/Engine/Private/SceneManagement.cpp
--- a/Engine/Source/Runtime/Engine/Private/SceneManagement.cpp
+++ b/Engine/Source/Runtime/Engine/Private/SceneManagement.cpp
@@ -1329,7 +1329,7 @@ void GetLightmapClusterResourceParameters(
 	const bool bAllowHighQualityLightMaps = AllowHighQualityLightmaps(FeatureLevel);
 
 	static const auto CVar = IConsoleManager::Get().FindTConsoleVariableDataInt(TEXT("r.VirtualTexturedLightmaps"));
-	const bool bUseVirtualTextures = (CVar->GetValueOnRenderThread() != 0) && UseVirtualTexturing(FeatureLevel);
+	const bool bUseVirtualTextures = (CVar->GetValueOnRenderThread() != 0) && UseVirtualTexturing(GetFeatureLevelShaderPlatform(FeatureLevel));
 
 	Parameters.LightMapTexture = GBlackTexture->TextureRHI;
 	Parameters.SkyOcclusionTexture = GWhiteTexture->TextureRHI;
diff --git a/Engine/Source/Runtime/Engine/Private/StaticMeshLight.cpp b/Engine/Source/Runtime/Engine/Private/StaticMeshLight.cpp
--- a/Engine/Source/Runtime/Engine/Private/StaticMeshLight.cpp
+++ b/Engine/Source/Runtime/Engine/Private/StaticMeshLight.cpp
@@ -212,7 +212,7 @@ FStaticMeshStaticLightingTextureMapping::FStaticMeshStaticLightingTextureMapping
 void FStaticMeshStaticLightingTextureMapping::Apply(FQuantizedLightmapData* QuantizedData, const TMap<ULightComponent*,FShadowMapData2D*>& ShadowMapData, ULevel* LightingScenario)
 {
 	static const auto CVar = IConsoleManager::Get().FindTConsoleVariableDataInt(TEXT("r.VirtualTexturedLightmaps"));
-	const bool bUseVirtualTextures = (CVar->GetValueOnAnyThread() != 0) && UseVirtualTexturing(GMaxRHIFeatureLevel);
+	const bool bUseVirtualTextures = (CVar->GetValueOnAnyThread() != 0) && UseVirtualTexturing(GMaxRHIShaderPlatform);
 
 	UStaticMeshComponent* StaticMeshComponent = Primitive.Get();
 
diff --git a/Engine/Source/Runtime/Engine/Private/Texture.cpp b/Engine/Source/Runtime/Engine/Private/Texture.cpp
--- a/Engine/Source/Runtime/Engine/Private/Texture.cpp
+++ b/Engine/Source/Runtime/Engine/Private/Texture.cpp
@@ -977,7 +977,7 @@ void UTexture::Serialize(FArchive& Ar)
 
 	if (Ar.IsCooking() && VirtualTextureStreaming)
 	{
-		if (UseVirtualTexturing(GMaxRHIFeatureLevel, Ar.CookingTarget()) == false)
+		if (UseVirtualTexturing(GMaxRHIShaderPlatform, Ar.CookingTarget()) == false)
 		{
 			UE_LOG(LogTexture, Display, TEXT("%s is marked for virtual streaming but virtual texture streaming is not available."), *GetPathName());
 		}
diff --git a/Engine/Source/Runtime/Engine/Private/VT/VirtualTexture.cpp b/Engine/Source/Runtime/Engine/Private/VT/VirtualTexture.cpp
--- a/Engine/Source/Runtime/Engine/Private/VT/VirtualTexture.cpp
+++ b/Engine/Source/Runtime/Engine/Private/VT/VirtualTexture.cpp
@@ -196,7 +196,7 @@ void UVirtualTexture2D::BeginCacheForCookedPlatformData(const ITargetPlatform* T
 	// Even though we skip the cook of this object for non VT platforms in URuntimeVirtualTexture::Serialize()
 	// we still load the object at cook time and kick off the DDC build. This will trigger an error in the texture DDC code.
 	// Either we need to make the DDC code more robust for non VT platforms or we can skip the process here...
-	if (!UseVirtualTexturing(GMaxRHIFeatureLevel, TargetPlatform))
+	if (!UseVirtualTexturing(GMaxRHIShaderPlatform, TargetPlatform))
 	{
 		return;
 	}
@@ -206,7 +206,7 @@ void UVirtualTexture2D::BeginCacheForCookedPlatformData(const ITargetPlatform* T
 
 bool UVirtualTexture2D::IsCachedCookedPlatformDataLoaded(const ITargetPlatform* TargetPlatform)
 {
-	if (!UseVirtualTexturing(GMaxRHIFeatureLevel, TargetPlatform))
+	if (!UseVirtualTexturing(GMaxRHIShaderPlatform, TargetPlatform))
 	{
 		return true;
 	}
@@ -216,7 +216,7 @@ bool UVirtualTexture2D::IsCachedCookedPlatformDataLoaded(const ITargetPlatform*
 
 void UVirtualTexture2D::ClearCachedCookedPlatformData(const ITargetPlatform* TargetPlatform)
 {
-	if (!UseVirtualTexturing(GMaxRHIFeatureLevel, TargetPlatform))
+	if (!UseVirtualTexturing(GMaxRHIShaderPlatform, TargetPlatform))
 	{
 		return;
 	}
diff --git a/Engine/Source/Runtime/Engine/Private/VT/VirtualTextureBuilder.cpp b/Engine/Source/Runtime/Engine/Private/VT/VirtualTextureBuilder.cpp
--- a/Engine/Source/Runtime/Engine/Private/VT/VirtualTextureBuilder.cpp
+++ b/Engine/Source/Runtime/Engine/Private/VT/VirtualTextureBuilder.cpp
@@ -31,7 +31,7 @@ void UVirtualTextureBuilder::Serialize(FArchive& Ar)
 		UVirtualTexture2D* TextureMobileBackup = TextureMobile;
 		
 		// Clear Texture during cook for platforms that don't support virtual texturing
-		if (!UseVirtualTexturing(GMaxRHIFeatureLevel, Ar.CookingTarget()))
+		if (!UseVirtualTexturing(GMaxRHIShaderPlatform, Ar.CookingTarget()))
 		{
 			Texture = nullptr;
 			TextureMobile = nullptr;
diff --git a/Engine/Source/Runtime/Engine/Private/VT/VirtualTextureChunkManager.cpp b/Engine/Source/Runtime/Engine/Private/VT/VirtualTextureChunkManager.cpp
--- a/Engine/Source/Runtime/Engine/Private/VT/VirtualTextureChunkManager.cpp
+++ b/Engine/Source/Runtime/Engine/Private/VT/VirtualTextureChunkManager.cpp
@@ -33,7 +33,7 @@ FVirtualTextureChunkStreamingManager::~FVirtualTextureChunkStreamingManager()
 void FVirtualTextureChunkStreamingManager::UpdateResourceStreaming(float DeltaTime, bool bProcessEverything /*= false*/)
 {
 	// It's OK to execute update if virtual texturing is disabled but we early out anyway to avoid an unnecessary render thread task.
-	if (UseVirtualTexturing(GMaxRHIFeatureLevel))
+	if (UseVirtualTexturing(GMaxRHIShaderPlatform))
 	{
 		FVirtualTextureChunkStreamingManager* StreamingManager = this;
 
diff --git a/Engine/Source/Runtime/Engine/Private/WorldPartition/WorldPartitionMiniMap.cpp b/Engine/Source/Runtime/Engine/Private/WorldPartition/WorldPartitionMiniMap.cpp
--- a/Engine/Source/Runtime/Engine/Private/WorldPartition/WorldPartitionMiniMap.cpp
+++ b/Engine/Source/Runtime/Engine/Private/WorldPartition/WorldPartitionMiniMap.cpp
@@ -62,7 +62,7 @@ void AWorldPartitionMiniMap::CheckForErrors()
 	{
 		const int32 MaxTextureDimension = GetMax2DTextureDimension();
 		const bool bExceedMaxDimension = MiniMapTexture->GetImportedSize().GetMax() > MaxTextureDimension;
-		if (!UseVirtualTexturing(GMaxRHIFeatureLevel) && bExceedMaxDimension)
+		if (!UseVirtualTexturing(GMaxRHIShaderPlatform) && bExceedMaxDimension)
 		{
 			FFormatNamedArguments Arguments;
 			Arguments.Add(TEXT("ActorName"), FText::FromString(GetName()));
diff --git a/Engine/Source/Runtime/Landscape/Private/Landscape.cpp b/Engine/Source/Runtime/Landscape/Private/Landscape.cpp
--- a/Engine/Source/Runtime/Landscape/Private/Landscape.cpp
+++ b/Engine/Source/Runtime/Landscape/Private/Landscape.cpp
@@ -1086,8 +1086,7 @@ void ULandscapeComponent::PostLoad()
 		}
 
 		// we need the fixed grid vertex factory for both virtual texturing and grass
-		ERHIFeatureLevel::Type FeatureLevel = GMaxRHIFeatureLevel;
-		bool bNeedsFixedGridVertexFactory = UseVirtualTexturing(FeatureLevel);
+		bool bNeedsFixedGridVertexFactory = UseVirtualTexturing(GMaxRHIShaderPlatform);
 		// This cvar is defined in the water plugin and searching for it should return nullptr if the plugin is not loaded
 		const bool bWaterPluginLoaded = IConsoleManager::Get().FindConsoleVariable(TEXT("r.Water.WaterInfo.RenderMethod")) != nullptr;
 		bNeedsFixedGridVertexFactory |= bWaterPluginLoaded;
diff --git a/Engine/Source/Runtime/Landscape/Private/LandscapeLight.cpp b/Engine/Source/Runtime/Landscape/Private/LandscapeLight.cpp
--- a/Engine/Source/Runtime/Landscape/Private/LandscapeLight.cpp
+++ b/Engine/Source/Runtime/Landscape/Private/LandscapeLight.cpp
@@ -47,7 +47,7 @@ FStaticLightingTextureMapping(
 void FLandscapeStaticLightingTextureMapping::Apply(FQuantizedLightmapData* QuantizedData, const TMap<ULightComponent*,FShadowMapData2D*>& ShadowMapData, ULevel* LightingScenario)
 {
 	static const auto CVar = IConsoleManager::Get().FindTConsoleVariableDataInt(TEXT("r.VirtualTexturedLightmaps"));
-	const bool bUseVirtualTextures = (CVar->GetValueOnAnyThread() != 0) && UseVirtualTexturing(GMaxRHIFeatureLevel);
+	const bool bUseVirtualTextures = (CVar->GetValueOnAnyThread() != 0) && UseVirtualTexturing(GMaxRHIShaderPlatform);
 
 	//ELightMapPaddingType PaddingType = GAllowLightmapPadding ? LMPT_NormalPadding : LMPT_NoPadding;
 	ELightMapPaddingType PaddingType = LMPT_NoPadding;
diff --git a/Engine/Source/Runtime/Landscape/Private/LandscapeRender.cpp b/Engine/Source/Runtime/Landscape/Private/LandscapeRender.cpp
--- a/Engine/Source/Runtime/Landscape/Private/LandscapeRender.cpp
+++ b/Engine/Source/Runtime/Landscape/Private/LandscapeRender.cpp
@@ -1435,7 +1435,7 @@ void FLandscapeComponentSceneProxy::CreateRenderThreadResources(FRHICommandListB
 		}
 
 		// We need the fixed grid vertex factory for virtual texturing, grass and for rendering the water info texture and for rendering Landscape into Lumen Surface Cache: 
-		bool bNeedsFixedGridVertexFactory = UseVirtualTexturing(FeatureLevel) || DoesPlatformSupportLumenGI(GMaxRHIShaderPlatform);
+		bool bNeedsFixedGridVertexFactory = UseVirtualTexturing(GetScene().GetShaderPlatform()) || DoesPlatformSupportLumenGI(GMaxRHIShaderPlatform);
 
 		// This cvar is defined in the water plugin and searching for it should return nullptr if the plugin is not loaded
 		const bool bWaterPluginLoaded = IConsoleManager::Get().FindConsoleVariable(TEXT("r.Water.WaterInfo.RenderMethod")) != nullptr;
diff --git a/Engine/Source/Runtime/RenderCore/Private/RenderUtils.cpp b/Engine/Source/Runtime/RenderCore/Private/RenderUtils.cpp
--- a/Engine/Source/Runtime/RenderCore/Private/RenderUtils.cpp
+++ b/Engine/Source/Runtime/RenderCore/Private/RenderUtils.cpp
@@ -1144,7 +1144,7 @@ RENDERCORE_API void QuantizeSceneBufferSize(const FIntPoint& InBufferSize, FIntP
 	OutBufferSize.Y = (InBufferSize.Y + DividableBy - 1) & Mask;
 }
 
-bool UseVirtualTexturing(bool bIsMobilePlatform, const ITargetPlatform* TargetPlatform)
+RENDERCORE_API bool UseVirtualTexturing(const FStaticShaderPlatform ShaderPlatform, const ITargetPlatform* TargetPlatform)
 {
 #if PLATFORM_SUPPORTS_VIRTUAL_TEXTURE_STREAMING
 	if (!FPlatformProperties::SupportsVirtualTextureStreaming())
@@ -1160,38 +1160,48 @@ bool UseVirtualTexturing(bool bIsMobilePlatform, const ITargetPlatform* TargetPl
 		return false;
 	}
 
+	bool bIsMobilePlatform = IsMobilePlatform(ShaderPlatform);
 #if WITH_EDITOR
 	// in case we are cooking for a platform that supports only mobile rendering
 	bIsMobilePlatform |= (TargetPlatform && !TargetPlatform->SupportsFeature(ETargetPlatformFeatures::DeferredRendering));
 #endif
 	// mobile needs an additional switch to enable VT		
-	static const auto CVarMobileVirtualTexture = IConsoleManager::Get().FindTConsoleVariableDataInt(TEXT("r.Mobile.VirtualTextures"));
-	if (bIsMobilePlatform && CVarMobileVirtualTexture->GetValueOnAnyThread() == 0)
+	if (bIsMobilePlatform)
 	{
-		return false;
+	#if WITH_EDITOR		
+		if (TargetPlatform)
+		{
+			// VT can be toggled on and off for specific platform
+			return TargetPlatform->SupportsFeature(ETargetPlatformFeatures::VirtualTextureStreaming);
+		}
+		else
+	#endif
+		{ 
+			static FShaderPlatformCachedIniValue<bool> MobileVirtualTexturesIniValue(TEXT("r.Mobile.VirtualTextures"));
+			return (MobileVirtualTexturesIniValue.Get(ShaderPlatform) != false);
+		}
 	}
-
 	return true;
 #else
 	return false;
 #endif
 }
 
-RENDERCORE_API bool UseVirtualTexturing(const EShaderPlatform InShaderPlatform, const ITargetPlatform* TargetPlatform)
+RENDERCORE_API bool UseVirtualTexturing(const FStaticFeatureLevel InFeatureLevel, const ITargetPlatform* TargetPlatform)
 {
-	return UseVirtualTexturing(IsMobilePlatform(InShaderPlatform), TargetPlatform);
+	return UseVirtualTexturing(GetFeatureLevelShaderPlatform(InFeatureLevel), TargetPlatform);
 }
 
-RENDERCORE_API bool UseVirtualTexturing(const FStaticFeatureLevel InFeatureLevel, const ITargetPlatform* TargetPlatform)
+RENDERCORE_API bool UseVirtualTextureLightmap(const FStaticShaderPlatform InShaderPlatform, const ITargetPlatform* TargetPlatform)
 {
-	return UseVirtualTexturing(InFeatureLevel == ERHIFeatureLevel::ES3_1, TargetPlatform);
+	static const auto CVar = IConsoleManager::Get().FindTConsoleVariableDataInt(TEXT("r.VirtualTexturedLightmaps"));
+	const bool bUseVirtualTextureLightmap = (CVar->GetValueOnAnyThread() != 0) && UseVirtualTexturing(InShaderPlatform, TargetPlatform);
+	return bUseVirtualTextureLightmap;
 }
 
 RENDERCORE_API bool UseVirtualTextureLightmap(const FStaticFeatureLevel InFeatureLevel, const ITargetPlatform* TargetPlatform)
 {
-	static const auto CVar = IConsoleManager::Get().FindTConsoleVariableDataInt(TEXT("r.VirtualTexturedLightmaps"));
-	const bool bUseVirtualTextureLightmap = (CVar->GetValueOnAnyThread() != 0) && UseVirtualTexturing(InFeatureLevel, TargetPlatform);
-	return bUseVirtualTextureLightmap;
+	return UseVirtualTextureLightmap(GetFeatureLevelShaderPlatform(InFeatureLevel), TargetPlatform);
 }
 
 RENDERCORE_API bool UseNaniteLandscapeMesh(EShaderPlatform ShaderPlatform)
diff --git a/Engine/Source/Runtime/RenderCore/Private/Shader.cpp b/Engine/Source/Runtime/RenderCore/Private/Shader.cpp
--- a/Engine/Source/Runtime/RenderCore/Private/Shader.cpp
+++ b/Engine/Source/Runtime/RenderCore/Private/Shader.cpp
@@ -2115,13 +2115,16 @@ void ShaderMapAppendKeyString(EShaderPlatform Platform, FString& KeyString)
 		static const auto CVarVTAnisotropic = IConsoleManager::Get().FindTConsoleVariableDataInt(TEXT("r.VT.AnisotropicFiltering"));
 		int32 VTFiltering = CVarVTAnisotropic && CVarVTAnisotropic->GetValueOnAnyThread() != 0 ? 1 : 0;
 
-		static const auto CVarMobileVirtualTexture = IConsoleManager::Get().FindTConsoleVariableDataInt(TEXT("r.Mobile.VirtualTextures"));
 		if (IsMobilePlatform(Platform) && VTTextures)
 		{
-			VTTextures = (CVarMobileVirtualTexture->GetValueOnAnyThread() != 0);
+			static FShaderPlatformCachedIniValue<bool> MobileVirtualTexturesIniValue(TEXT("r.Mobile.VirtualTextures"));
+			VTTextures = (MobileVirtualTexturesIniValue.Get(Platform) != false);
 
-			static FShaderPlatformCachedIniValue<bool> CVarVTMobileManualTrilinearFiltering(TEXT("r.VT.Mobile.ManualTrilinearFiltering"));
-			VTFiltering += (CVarVTMobileManualTrilinearFiltering.Get(Platform) ? 2 : 0);
+			if (VTTextures)
+			{
+				static FShaderPlatformCachedIniValue<bool> CVarVTMobileManualTrilinearFiltering(TEXT("r.VT.Mobile.ManualTrilinearFiltering"));
+				VTFiltering += (CVarVTMobileManualTrilinearFiltering.Get(Platform) ? 2 : 0);
+			}
 		}
 
 		const bool VTSupported = TargetPlatform != nullptr && TargetPlatform->SupportsFeature(ETargetPlatformFeatures::VirtualTextureStreaming);
diff --git a/Engine/Source/Runtime/RenderCore/Public/RenderUtils.h b/Engine/Source/Runtime/RenderCore/Public/RenderUtils.h
--- a/Engine/Source/Runtime/RenderCore/Public/RenderUtils.h
+++ b/Engine/Source/Runtime/RenderCore/Public/RenderUtils.h
@@ -363,9 +363,10 @@ RENDERCORE_API void QuantizeSceneBufferSize(const FIntPoint& InBufferSize, FIntP
 
 /**
 * Checks if virtual texturing enabled and supported
-* todo: Deprecate the version of the function that takes FStaticFeatureLevel
 */
-RENDERCORE_API bool UseVirtualTexturing(const EShaderPlatform InShaderPlatform, const ITargetPlatform* TargetPlatform = nullptr);
+RENDERCORE_API bool UseVirtualTexturing(const FStaticShaderPlatform InShaderPlatform, const ITargetPlatform* TargetPlatform = nullptr);
+
+UE_DEPRECATED(5.4, "Use version that takes FStaticShaderPlatform instead")
 RENDERCORE_API bool UseVirtualTexturing(const FStaticFeatureLevel InFeatureLevel, const ITargetPlatform* TargetPlatform = nullptr);
 
 RENDERCORE_API bool NaniteAtomicsSupported();
@@ -418,6 +419,9 @@ RENDERCORE_API bool IsSingleLayerWaterDepthPrepassEnabled(const FStaticShaderPla
 /**
 *	Checks if virtual texturing lightmap enabled and supported
 */
+RENDERCORE_API bool UseVirtualTextureLightmap(const FStaticShaderPlatform Platform, const ITargetPlatform* TargetPlatform = nullptr);
+
+UE_DEPRECATED(5.4, "Use version that takes FStaticShaderPlatform instead")
 RENDERCORE_API bool UseVirtualTextureLightmap(const FStaticFeatureLevel InFeatureLevel, const ITargetPlatform* TargetPlatform = nullptr);
 
 /**
diff --git a/Engine/Source/Runtime/Renderer/Private/DeferredShadingRenderer.cpp b/Engine/Source/Runtime/Renderer/Private/DeferredShadingRenderer.cpp
--- a/Engine/Source/Runtime/Renderer/Private/DeferredShadingRenderer.cpp
+++ b/Engine/Source/Runtime/Renderer/Private/DeferredShadingRenderer.cpp
@@ -2711,7 +2711,7 @@ void FDeferredShadingSceneRenderer::Render(FRDGBuilder& GraphBuilder)
 
 	FGPUSceneScopeBeginEndHelper GPUSceneScopeBeginEndHelper(Scene->GPUScene, GPUSceneDynamicContext, Scene);
 
-	const bool bUseVirtualTexturing = UseVirtualTexturing(FeatureLevel);
+	const bool bUseVirtualTexturing = UseVirtualTexturing(ShaderPlatform);
 
 	if (bUseVirtualTexturing && RendererOutput == ERendererOutput::FinalSceneColor)
 	{
diff --git a/Engine/Source/Runtime/Renderer/Private/LightMapRendering.cpp b/Engine/Source/Runtime/Renderer/Private/LightMapRendering.cpp
--- a/Engine/Source/Runtime/Renderer/Private/LightMapRendering.cpp
+++ b/Engine/Source/Runtime/Renderer/Private/LightMapRendering.cpp
@@ -42,7 +42,7 @@ void LightMapPolicyImpl::ModifyCompilationEnvironment(ELightmapQuality LightmapQ
 	OutEnvironment.SetDefine(TEXT("NUM_LIGHTMAP_COEFFICIENTS"), GNumLightmapCoefficients[LightmapQuality]);
 
 	static const auto CVar = IConsoleManager::Get().FindTConsoleVariableDataInt(TEXT("r.VirtualTexturedLightmaps"));
-	const bool VirtualTextureLightmaps = (CVar->GetValueOnAnyThread() != 0) && UseVirtualTexturing(GMaxRHIFeatureLevel, OutEnvironment.TargetPlatform);
+	const bool VirtualTextureLightmaps = (CVar->GetValueOnAnyThread() != 0) && UseVirtualTexturing(Parameters.Platform);
 	OutEnvironment.SetDefine(TEXT("LIGHTMAP_VT_ENABLED"), VirtualTextureLightmaps);
 }
 
diff --git a/Engine/Source/Runtime/Renderer/Private/MobileShadingRenderer.cpp b/Engine/Source/Runtime/Renderer/Private/MobileShadingRenderer.cpp
--- a/Engine/Source/Runtime/Renderer/Private/MobileShadingRenderer.cpp
+++ b/Engine/Source/Runtime/Renderer/Private/MobileShadingRenderer.cpp
@@ -267,7 +267,7 @@ FMobileSceneRenderer::FMobileSceneRenderer(const FSceneViewFamily* InViewFamily,
 	: FSceneRenderer(InViewFamily, HitProxyConsumer)
 	, bGammaSpace(!IsMobileHDR())
 	, bDeferredShading(IsMobileDeferredShadingEnabled(ShaderPlatform))
-	, bUseVirtualTexturing(UseVirtualTexturing(FeatureLevel) && GetRendererOutput() == FSceneRenderer::ERendererOutput::FinalSceneColor)
+	, bUseVirtualTexturing(UseVirtualTexturing(ShaderPlatform) && GetRendererOutput() == FSceneRenderer::ERendererOutput::FinalSceneColor)
 {
 	bRenderToSceneColor = false;
 	bRequiresMultiPass = false;
diff --git a/Engine/Source/Runtime/Renderer/Private/RayTracing/RayTracingMaterialHitShaders.cpp b/Engine/Source/Runtime/Renderer/Private/RayTracing/RayTracingMaterialHitShaders.cpp
--- a/Engine/Source/Runtime/Renderer/Private/RayTracing/RayTracingMaterialHitShaders.cpp
+++ b/Engine/Source/Runtime/Renderer/Private/RayTracing/RayTracingMaterialHitShaders.cpp
@@ -157,7 +157,7 @@ public:
 		LightMapPolicyType::ModifyCompilationEnvironment(Parameters, OutEnvironment);
 		FMeshMaterialShader::ModifyCompilationEnvironment(Parameters, OutEnvironment);
 		static const auto CVar = IConsoleManager::Get().FindTConsoleVariableDataInt(TEXT("r.VirtualTexturedLightmaps"));
-		const bool VirtualTextureLightmaps = (CVar->GetValueOnAnyThread() != 0) && UseVirtualTexturing(GMaxRHIFeatureLevel, OutEnvironment.TargetPlatform);
+		const bool VirtualTextureLightmaps = (CVar->GetValueOnAnyThread() != 0) && UseVirtualTexturing(Parameters.Platform);
 		OutEnvironment.SetDefine(TEXT("LIGHTMAP_VT_ENABLED"), VirtualTextureLightmaps);
 	}
 
diff --git a/Engine/Source/Runtime/Renderer/Private/ReflectionEnvironmentCapture.cpp b/Engine/Source/Runtime/Renderer/Private/ReflectionEnvironmentCapture.cpp
--- a/Engine/Source/Runtime/Renderer/Private/ReflectionEnvironmentCapture.cpp
+++ b/Engine/Source/Runtime/Renderer/Private/ReflectionEnvironmentCapture.cpp
@@ -1440,7 +1440,7 @@ void FScene::CaptureOrUploadReflectionCapture(UReflectionCaptureComponent* Captu
 			UE::RenderCommandPipe::FSyncScope SyncScope;
 
 			// Prefetch all virtual textures so that we have content available
-			if (UseVirtualTexturing(GetFeatureLevel()))
+			if (UseVirtualTexturing(GetShaderPlatform()))
 			{
 				const ERHIFeatureLevel::Type InFeatureLevel = FeatureLevel;
 				const FVector2D ScreenSpaceSize(ReflectionCaptureSize, ReflectionCaptureSize);
diff --git a/Engine/Source/Runtime/Renderer/Private/Renderer.cpp b/Engine/Source/Runtime/Renderer/Private/Renderer.cpp
--- a/Engine/Source/Runtime/Renderer/Private/Renderer.cpp
+++ b/Engine/Source/Runtime/Renderer/Private/Renderer.cpp
@@ -310,7 +310,7 @@ void FRendererModule::DrawTileMesh(FCanvasRenderContext& RenderContext, FMeshPas
 
 		// Materials sampling VTs need FVirtualTextureSystem to be updated before being rendered
 		const FMaterial& MeshMaterial = Mesh.MaterialRenderProxy->GetIncompleteMaterialWithFallback(FeatureLevel);
-		const bool bUseVirtualTexturing = UseVirtualTexturing(FeatureLevel) && !MeshMaterial.GetUniformVirtualTextureExpressions().IsEmpty();
+		const bool bUseVirtualTexturing = UseVirtualTexturing(View.GetShaderPlatform()) && !MeshMaterial.GetUniformVirtualTextureExpressions().IsEmpty();
 		if (bUseVirtualTexturing)
 		{
 			FVirtualTextureUpdateSettings Settings;
diff --git a/Engine/Source/Runtime/Renderer/Private/SceneHitProxyRendering.cpp b/Engine/Source/Runtime/Renderer/Private/SceneHitProxyRendering.cpp
--- a/Engine/Source/Runtime/Renderer/Private/SceneHitProxyRendering.cpp
+++ b/Engine/Source/Runtime/Renderer/Private/SceneHitProxyRendering.cpp
@@ -182,7 +182,7 @@ void InitHitProxyRender(FRDGBuilder& GraphBuilder, FSceneRenderer* SceneRenderer
 	auto FeatureLevel = ViewFamily.Scene->GetFeatureLevel();
 
 	// Ensure VirtualTexture resources are allocated
-	if (UseVirtualTexturing(FeatureLevel))
+	if (UseVirtualTexturing(ViewFamily.Scene->GetShaderPlatform()))
 	{
 		FVirtualTextureUpdateSettings Settings;
 		Settings.EnablePageRequests(false);
diff --git a/Engine/Source/Runtime/Renderer/Private/VT/RuntimeVirtualTextureRender.cpp b/Engine/Source/Runtime/Renderer/Private/VT/RuntimeVirtualTextureRender.cpp
--- a/Engine/Source/Runtime/Renderer/Private/VT/RuntimeVirtualTextureRender.cpp
+++ b/Engine/Source/Runtime/Renderer/Private/VT/RuntimeVirtualTextureRender.cpp
@@ -131,7 +131,7 @@ namespace RuntimeVirtualTexture
 
 		static bool ShouldCompilePermutation(const FMeshMaterialShaderPermutationParameters& Parameters)
 		{
-			return UseVirtualTexturing(GetMaxSupportedFeatureLevel(Parameters.Platform)) &&
+			return UseVirtualTexturing(Parameters.Platform) &&
 				(Parameters.MaterialParameters.bHasRuntimeVirtualTextureOutput || Parameters.MaterialParameters.bIsDefaultMaterial);
 		}
 
