Fixed: NonUnity compile errors

diff --git a/Engine/Source/Runtime/Engine/Classes/Materials/MaterialExpressionTextureBase.h b/Engine/Source/Runtime/Engine/Classes/Materials/MaterialExpressionTextureBase.h
--- a/Engine/Source/Runtime/Engine/Classes/Materials/MaterialExpressionTextureBase.h
+++ b/Engine/Source/Runtime/Engine/Classes/Materials/MaterialExpressionTextureBase.h
@@ -15,6 +15,7 @@
 class UTexture;
 struct FPropertyChangedEvent;
 enum EMaterialSamplerType : int;
+enum EShaderPlatform : uint16;
 
 UCLASS(abstract, hidecategories=Object, MinimalAPI)
 class UMaterialExpressionTextureBase : public UMaterialExpression 
