Fix VRPIE crashing when mobile multi view is enabled and editor window is not minimized.

diff --git a/Engine/Source/Runtime/Engine/Private/GameViewportClient.cpp b/Engine/Source/Runtime/Engine/Private/GameViewportClient.cpp
--- a/Engine/Source/Runtime/Engine/Private/GameViewportClient.cpp
+++ b/Engine/Source/Runtime/Engine/Private/GameViewportClient.cpp
@@ -1373,11 +1373,21 @@ void UGameViewportClient::Draw(FViewport* InViewport, FCanvas* SceneCanvas)
 	}
 
 	// create the view family for rendering the world scene to the viewport's render target
+	bool bRequireMultiView = false;
+	if (GEngine && GEngine->IsStereoscopic3D())
+	{
+		static const auto MobileMultiViewCVar = IConsoleManager::Get().FindTConsoleVariableDataInt(TEXT("vr.MobileMultiView"));
+		const bool bSkipPostprocessing = !IsMobileHDR();
+		const bool bUsingMobileRenderer = FSceneInterface::GetShadingPath(MyWorld->Scene->GetFeatureLevel()) == EShadingPath::Mobile;
+		bRequireMultiView = (GSupportsMobileMultiView || GRHISupportsArrayIndexFromAnyShader) && bUsingMobileRenderer && bSkipPostprocessing && (MobileMultiViewCVar && MobileMultiViewCVar->GetValueOnAnyThread() != 0);
+	}
+
 	FSceneViewFamilyContext ViewFamily(FSceneViewFamily::ConstructionValues(
 		InViewport,
 		MyWorld->Scene,
 		EngineShowFlags)
-		.SetRealtimeUpdate(true));
+		.SetRealtimeUpdate(true)
+		.SetRequireMobileMultiView(bRequireMultiView));
 
 	ViewFamily.DebugDPIScale = GetDPIScale();
 
diff --git a/Engine/Source/Runtime/Engine/Private/SceneView.cpp b/Engine/Source/Runtime/Engine/Private/SceneView.cpp
--- a/Engine/Source/Runtime/Engine/Private/SceneView.cpp
+++ b/Engine/Source/Runtime/Engine/Private/SceneView.cpp
@@ -2721,6 +2721,7 @@ PRAGMA_ENABLE_DEPRECATION_WARNINGS
 ,	bDeferClear(false)
 ,	bResolveScene(true)			
 ,	bTimesSet(false)
+,	bRequireMultiView(false)
 {
 	if( InScene != NULL )			
 	{
@@ -2753,7 +2754,7 @@ FSceneViewFamily::FSceneViewFamily(const ConstructionValues& CVS)
 	SceneCaptureCompositeMode(SCCM_Overwrite),
 	bWorldIsPaused(false),
 	bIsHDR(false),
-	bRequireMultiView(false),
+	bRequireMultiView(CVS.bRequireMultiView),
 PRAGMA_DISABLE_DEPRECATION_WARNINGS
 	GammaCorrection(CVS.GammaCorrection),
 PRAGMA_ENABLE_DEPRECATION_WARNINGS
@@ -2828,14 +2829,6 @@ PRAGMA_ENABLE_DEPRECATION_WARNINGS
 		EngineShowFlags.ScreenPercentage = false;
 	}
 
-	if (GEngine && GEngine->IsStereoscopic3D())
-	{
-		static const auto MobileMultiViewCVar = IConsoleManager::Get().FindTConsoleVariableDataInt(TEXT("vr.MobileMultiView"));
-		const bool bSkipPostprocessing = !IsMobileHDR();
-		const bool bUsingMobileRenderer = FSceneInterface::GetShadingPath(GetFeatureLevel()) == EShadingPath::Mobile;
-		bRequireMultiView = (GSupportsMobileMultiView || GRHISupportsArrayIndexFromAnyShader) && bUsingMobileRenderer && bSkipPostprocessing && (MobileMultiViewCVar && MobileMultiViewCVar->GetValueOnAnyThread() != 0);
-	}
-
 	// Check if the translucency are allowed to be rendered after DOF, if not, translucency after DOF will be rendered in standard translucency.
 	{
 		bool SeparateTranslucencyEnabled = EngineShowFlags.PostProcessing // Used for reflection captures.
diff --git a/Engine/Source/Runtime/Engine/Public/SceneView.h b/Engine/Source/Runtime/Engine/Public/SceneView.h
--- a/Engine/Source/Runtime/Engine/Public/SceneView.h
+++ b/Engine/Source/Runtime/Engine/Public/SceneView.h
@@ -1981,6 +1981,9 @@ public:
 		/** Safety check to ensure valid times are set either from a valid world/scene pointer or via the SetWorldTimes function */
 		uint32 bTimesSet:1;
 
+		/** True if scene color and depth should be multiview-allocated */
+		uint32 bRequireMultiView:1;
+
 		/** Set the world time and real time independently to handle time dilation. */
 		ConstructionValues& SetTime(const FGameTime& InTime)
 		{
@@ -2006,6 +2009,9 @@ public:
 		
 		/** Setting to if true then results of scene rendering are copied/resolved to the RenderTarget. */
 		ConstructionValues& SetResolveScene(const bool Value) { bResolveScene = Value; return *this; }
+
+		/** Setting to true results in scene color and depth being multiview-allocated. */
+		ConstructionValues& SetRequireMobileMultiView(const bool Value) { bRequireMultiView = Value; return *this; }
 		
 		/** Set Gamma correction used when rendering this family. */
 		UE_DEPRECATED(5.4, "Unused gamma correction.")
