Plugging in Instanced Actors to engine-level collision results (via generalized LWI interface)

Disabled by default, controlled by [REDACTED_PERSON]

diff --git a/Engine/Source/Runtime/Engine/Classes/Engine/ActorInstanceHandle.h b/Engine/Source/Runtime/Engine/Classes/Engine/ActorInstanceHandle.h
--- a/Engine/Source/Runtime/Engine/Classes/Engine/ActorInstanceHandle.h
+++ b/Engine/Source/Runtime/Engine/Classes/Engine/ActorInstanceHandle.h
@@ -61,7 +61,7 @@ struct FActorInstanceHandle
 	T* FetchActor() const;
 
 	AActor* GetCachedActor() const { return Actor.Get(); }
-	void SetCachedActor(AActor* InActor) const;
+	ENGINE_API void SetCachedActor(AActor* InActor) const;
 
 	/* Returns the index used internally by the manager */
 	FORCEINLINE int32 GetInstanceIndex() const { return InstanceIndex; }
