ActorDescContainer: added a fast name actor lookup when querying with raw actor names (don't support fully qualified actor outer chains).

diff --git a/Engine/Plugins/Editor/DataValidation/Source/DataValidation/Private/WorldPartitionChangelistValidator.cpp b/Engine/Plugins/Editor/DataValidation/Source/DataValidation/Private/WorldPartitionChangelistValidator.cpp
--- a/Engine/Plugins/Editor/DataValidation/Source/DataValidation/Private/WorldPartitionChangelistValidator.cpp
+++ b/Engine/Plugins/Editor/DataValidation/Source/DataValidation/Private/WorldPartitionChangelistValidator.cpp
@@ -212,7 +212,7 @@ EDataValidationResult UWorldPartitionChangelistValidator::ValidateActorsAndDataL
 		for (const FAssetData& ActorData : ActorsData)
 		{
 			// Get the actor descriptor
-			if (const FWorldPartitionActorDesc* ActorDesc = ContainersToValidate.GetActorDescByName(ActorData.AssetName.ToString()))
+			if (const FWorldPartitionActorDesc* ActorDesc = ContainersToValidate.GetActorDescByPath(ActorData.AssetName.ToString()))
 			{
 				RelevantActorGuids.Add(ActorDesc->GetGuid());
 			}
diff --git a/Engine/Source/Editor/PropertyEditor/Private/UserInterface/PropertyEditor/SPropertyEditorAsset.cpp b/Engine/Source/Editor/PropertyEditor/Private/UserInterface/PropertyEditor/SPropertyEditorAsset.cpp
--- a/Engine/Source/Editor/PropertyEditor/Private/UserInterface/PropertyEditor/SPropertyEditorAsset.cpp
+++ b/Engine/Source/Editor/PropertyEditor/Private/UserInterface/PropertyEditor/SPropertyEditorAsset.cpp
@@ -1264,7 +1264,7 @@ void SPropertyEditorAsset::OnBrowse()
 				{
 					if (UWorld* World = Cast<UWorld>(MapObject); World && World->IsPartitionedWorld())
 					{
-						if (const FWorldPartitionActorDesc* ActorDesc = World->GetWorldPartition()->GetActorDescByName(Value.ObjectPath))
+						if (const FWorldPartitionActorDesc* ActorDesc = World->GetWorldPartition()->GetActorDescByPath(Value.ObjectPath))
 						{
 							World->GetWorldPartition()->PinActors({ ActorDesc->GetGuid() });
 							GetValue(Value);
diff --git a/Engine/Source/Editor/UnrealEd/Private/WorldPartition/WorldPartitionFoliageBuilder.cpp b/Engine/Source/Editor/UnrealEd/Private/WorldPartition/WorldPartitionFoliageBuilder.cpp
--- a/Engine/Source/Editor/UnrealEd/Private/WorldPartition/WorldPartitionFoliageBuilder.cpp
+++ b/Engine/Source/Editor/UnrealEd/Private/WorldPartition/WorldPartitionFoliageBuilder.cpp
@@ -137,7 +137,7 @@ bool UWorldPartitionFoliageBuilder::RunInternal(UWorld* World, const FCellInfo&
 
 				const FString ExpectedActorPath = FString::Printf(TEXT("%s.%s"), *World->PersistentLevel->GetPathName(), *ExpectedActorName);
 				
-				if (const FWorldPartitionActorDesc* DupActorDesc = WorldPartition->GetActorDescByName(ExpectedActorPath))
+				if (const FWorldPartitionActorDesc* DupActorDesc = WorldPartition->GetActorDescByPath(ExpectedActorPath))
 				{
 					// Merge with existing
 					FWorldPartitionReference DupActorRef(WorldPartition, DupActorDesc->GetGuid());
diff --git a/Engine/Source/Runtime/Engine/Private/WorldPartition/ActorDescContainer.cpp b/Engine/Source/Runtime/Engine/Private/WorldPartition/ActorDescContainer.cpp
--- a/Engine/Source/Runtime/Engine/Private/WorldPartition/ActorDescContainer.cpp
+++ b/Engine/Source/Runtime/Engine/Private/WorldPartition/ActorDescContainer.cpp
@@ -282,7 +282,7 @@ bool UActorDescContainer::ShouldHandleActorEvent(const AActor* Actor)
 	return Actor && IsActorDescHandled(Actor) && Actor->IsMainPackageActor() && Actor->GetLevel();
 }
 
-const FWorldPartitionActorDesc* UActorDescContainer::GetActorDescByName(const FString& ActorPath) const
+const FWorldPartitionActorDesc* UActorDescContainer::GetActorDescByPath(const FString& ActorPath) const
 {
 	FString ActorName;
 	FString ActorContext;
@@ -291,7 +291,17 @@ const FWorldPartitionActorDesc* UActorDescContainer::GetActorDescByName(const FS
 		ActorName = ActorPath;
 	}
 
-	if (const TUniquePtr<FWorldPartitionActorDesc>* const* ActorDesc = ActorsByName.Find(*ActorName))
+	return GetActorDescByName(FName(*ActorName));
+}
+
+const FWorldPartitionActorDesc* UActorDescContainer::GetActorDescByPath(const FSoftObjectPath& ActorPath) const
+{
+	return GetActorDescByPath(ActorPath.ToString());
+}
+
+const FWorldPartitionActorDesc* UActorDescContainer::GetActorDescByName(FName ActorName) const
+{
+	if (const TUniquePtr<FWorldPartitionActorDesc>* const* ActorDesc = ActorsByName.Find(ActorName))
 	{
 		return (*ActorDesc)->Get();
 	}
@@ -299,11 +309,6 @@ const FWorldPartitionActorDesc* UActorDescContainer::GetActorDescByName(const FS
 	return nullptr;
 }
 
-const FWorldPartitionActorDesc* UActorDescContainer::GetActorDescByName(const FSoftObjectPath& ActorPath) const
-{
-	return GetActorDescByName(ActorPath.ToString());
-}
-
 void UActorDescContainer::OnObjectPreSave(UObject* Object, FObjectPreSaveContext SaveContext)
 {
 	if (!SaveContext.IsProceduralSave() && !(SaveContext.GetSaveFlags() & SAVE_FromAutosave))
diff --git a/Engine/Source/Runtime/Engine/Private/WorldPartition/WorldPartition.cpp b/Engine/Source/Runtime/Engine/Private/WorldPartition/WorldPartition.cpp
--- a/Engine/Source/Runtime/Engine/Private/WorldPartition/WorldPartition.cpp
+++ b/Engine/Source/Runtime/Engine/Private/WorldPartition/WorldPartition.cpp
@@ -199,7 +199,7 @@ static FAutoConsoleCommand DumpActorDesc(
 						TMap<FName, FString> DataLayersDumpString = GetDataLayersDumpString(WorldPartition);
 						for (const FString& ActorPath : ActorPaths)
 						{
-							if (const FWorldPartitionActorDesc* ActorDesc = WorldPartition->GetActorDescByName(ActorPath))
+							if (const FWorldPartitionActorDesc* ActorDesc = WorldPartition->GetActorDescByPath(ActorPath))
 							{
 								UE_LOG(LogWorldPartition, Log, TEXT("%s"), *GetActorDescDumpString(ActorDesc, DataLayersDumpString));
 							}
@@ -1576,7 +1576,7 @@ bool UWorldPartition::ResolveSubobject(const TCHAR* SubObjectPath, UObject*& Out
 				SubObjectName = SubObjectPath;
 			}
 
-			if (const FWorldPartitionActorDesc* ActorDesc = GetActorDescByName(SubObjectName))
+			if (const FWorldPartitionActorDesc* ActorDesc = GetActorDescByPath(SubObjectName))
 			{
 				if (bLoadIfExists)
 				{
diff --git a/Engine/Source/Runtime/Engine/Public/WorldPartition/ActorDescContainer.h b/Engine/Source/Runtime/Engine/Public/WorldPartition/ActorDescContainer.h
--- a/Engine/Source/Runtime/Engine/Public/WorldPartition/ActorDescContainer.h
+++ b/Engine/Source/Runtime/Engine/Public/WorldPartition/ActorDescContainer.h
@@ -101,8 +101,9 @@ public:
 
 	ENGINE_API bool ShouldHandleActorEvent(const AActor* Actor);
 
-	ENGINE_API const FWorldPartitionActorDesc* GetActorDescByName(const FString& ActorPath) const;
-	ENGINE_API const FWorldPartitionActorDesc* GetActorDescByName(const FSoftObjectPath& InActorPath) const;
+	ENGINE_API const FWorldPartitionActorDesc* GetActorDescByPath(const FString& ActorPath) const;
+	ENGINE_API const FWorldPartitionActorDesc* GetActorDescByPath(const FSoftObjectPath& ActorPath) const;
+	ENGINE_API const FWorldPartitionActorDesc* GetActorDescByName(FName ActorName) const;
 
 	bool bContainerInitialized;
 
diff --git a/Engine/Source/Runtime/Engine/Public/WorldPartition/ActorDescContainerCollection.h b/Engine/Source/Runtime/Engine/Public/WorldPartition/ActorDescContainerCollection.h
--- a/Engine/Source/Runtime/Engine/Public/WorldPartition/ActorDescContainerCollection.h
+++ b/Engine/Source/Runtime/Engine/Public/WorldPartition/ActorDescContainerCollection.h
@@ -40,8 +40,9 @@ public:
 	FWorldPartitionActorDesc& GetActorDescChecked(const FGuid& Guid);
 	const FWorldPartitionActorDesc& GetActorDescChecked(const FGuid& Guid) const;
 
-	const FWorldPartitionActorDesc* GetActorDescByName(const FString& ActorPath) const;
-	const FWorldPartitionActorDesc* GetActorDescByName(const FSoftObjectPath& ActorPath) const;
+	const FWorldPartitionActorDesc* GetActorDescByPath(const FString& ActorPath) const;
+	const FWorldPartitionActorDesc* GetActorDescByPath(const FSoftObjectPath& ActorPath) const;
+	const FWorldPartitionActorDesc* GetActorDescByName(FName ActorName) const;
 
 	template<typename Dummy = void, typename = typename TEnableIf<!TIsConst<TRemovePointer<ActorDescContPtrType>>::Value, Dummy>::Type>
 	bool RemoveActor(const FGuid& ActorGuid);
@@ -328,12 +329,12 @@ FWorldPartitionActorDesc& TActorDescContainerCollection<ActorDescContPtrType>::G
 }
 
 template<class ActorDescContPtrType>
-const FWorldPartitionActorDesc* TActorDescContainerCollection<ActorDescContPtrType>::GetActorDescByName(const FString& ActorPath) const
+const FWorldPartitionActorDesc* TActorDescContainerCollection<ActorDescContPtrType>::GetActorDescByPath(const FString& ActorPath) const
 {
 	const FWorldPartitionActorDesc* ActorDesc = nullptr;
 	ForEachActorDescContainerBreakable([&ActorPath, &ActorDesc](ActorDescContPtrType ActorDescContainer)
 	{
-		ActorDesc = ActorDescContainer->GetActorDescByName(ActorPath);
+		ActorDesc = ActorDescContainer->GetActorDescByPath(ActorPath);
 		return ActorDesc == nullptr;
 	});
 
@@ -341,12 +342,25 @@ const FWorldPartitionActorDesc* TActorDescContainerCollection<ActorDescContPtrTy
 }
 
 template<class ActorDescContPtrType>
-const FWorldPartitionActorDesc* TActorDescContainerCollection<ActorDescContPtrType>::GetActorDescByName(const FSoftObjectPath& ActorPath) const
+const FWorldPartitionActorDesc* TActorDescContainerCollection<ActorDescContPtrType>::GetActorDescByPath(const FSoftObjectPath& ActorPath) const
 {
 	const FWorldPartitionActorDesc* ActorDesc = nullptr;
 	ForEachActorDescContainerBreakable([&ActorPath, &ActorDesc](ActorDescContPtrType ActorDescContainer)
 	{
-		ActorDesc = ActorDescContainer->GetActorDescByName(ActorPath);
+		ActorDesc = ActorDescContainer->GetActorDescByPath(ActorPath);
+		return ActorDesc == nullptr;
+	});
+
+	return ActorDesc;
+}
+
+template<class ActorDescContPtrType>
+const FWorldPartitionActorDesc* TActorDescContainerCollection<ActorDescContPtrType>::GetActorDescByName(FName ActorName) const
+{
+	const FWorldPartitionActorDesc* ActorDesc = nullptr;
+	ForEachActorDescContainerBreakable([&ActorName, &ActorDesc](ActorDescContPtrType ActorDescContainer)
+	{
+		ActorDesc = ActorDescContainer->GetActorDescByName(ActorName);
 		return ActorDesc == nullptr;
 	});
 
