Make sure most common VFs do not add vertex streams that will never be used. This helps to avoid exra heap allocations on MDC caching

diff --git a/Engine/Source/Runtime/Engine/Private/GeometryCollectionRendering.cpp b/Engine/Source/Runtime/Engine/Private/GeometryCollectionRendering.cpp
--- a/Engine/Source/Runtime/Engine/Private/GeometryCollectionRendering.cpp
+++ b/Engine/Source/Runtime/Engine/Private/GeometryCollectionRendering.cpp
@@ -203,13 +203,16 @@ void FGeometryCollectionVertexFactory::InitRHI(FRHICommandListBase& RHICmdList)
 		}
 	}
 
-	if (Data.LightMapCoordinateComponent.VertexBuffer)
+	if (FReadOnlyCVARCache::Get().bAllowStaticLighting)
 	{
-		Elements.Add(AccessStreamComponent(Data.LightMapCoordinateComponent, 15));
-	}
-	else if (Data.TextureCoordinates.Num())
-	{
-		Elements.Add(AccessStreamComponent(Data.TextureCoordinates[0], 15));
+		if (Data.LightMapCoordinateComponent.VertexBuffer)
+		{
+			Elements.Add(AccessStreamComponent(Data.LightMapCoordinateComponent, 15));
+		}
+		else if (Data.TextureCoordinates.Num())
+		{
+			Elements.Add(AccessStreamComponent(Data.TextureCoordinates[0], 15));
+		}
 	}
 
 	check(Streams.Num() > 0);
diff --git a/Engine/Source/Runtime/Engine/Private/InstancedStaticMesh.cpp b/Engine/Source/Runtime/Engine/Private/InstancedStaticMesh.cpp
--- a/Engine/Source/Runtime/Engine/Private/InstancedStaticMesh.cpp
+++ b/Engine/Source/Runtime/Engine/Private/InstancedStaticMesh.cpp
@@ -1089,13 +1089,16 @@ void FInstancedStaticMeshVertexFactory::GetVertexElements(
 			Elements.Add(AccessStreamComponent(NullComponent, 14, Streams));
 		}
 
-		if (Data.LightMapCoordinateComponent.VertexBuffer)
+		if (FReadOnlyCVARCache::Get().bAllowStaticLighting)
 		{
-			Elements.Add(AccessStreamComponent(Data.LightMapCoordinateComponent, 15, Streams));
-		}
-		else if (Data.TextureCoordinates.Num())
-		{
-			Elements.Add(AccessStreamComponent(Data.TextureCoordinates[0], 15, Streams));
+			if (Data.LightMapCoordinateComponent.VertexBuffer)
+			{
+				Elements.Add(AccessStreamComponent(Data.LightMapCoordinateComponent, 15, Streams));
+			}
+			else if (Data.TextureCoordinates.Num())
+			{
+				Elements.Add(AccessStreamComponent(Data.TextureCoordinates[0], 15, Streams));
+			}
 		}
 	}
 
diff --git a/Engine/Source/Runtime/Engine/Private/LocalVertexFactory.cpp b/Engine/Source/Runtime/Engine/Private/LocalVertexFactory.cpp
--- a/Engine/Source/Runtime/Engine/Private/LocalVertexFactory.cpp
+++ b/Engine/Source/Runtime/Engine/Private/LocalVertexFactory.cpp
@@ -333,7 +333,8 @@ void FLocalVertexFactory::GetPSOPrecacheVertexFetchElements(EVertexInputStreamTy
 		Elements.Add(FVertexElement(1, 4, VET_PackedNormal, 2, sizeof(FPackedNormal)*2u, false));
 	}
 
-	if (UseGPUScene(GMaxRHIShaderPlatform, GMaxRHIFeatureLevel))
+	if (UseGPUScene(GMaxRHIShaderPlatform, GMaxRHIFeatureLevel) 
+		&& !PlatformGPUSceneUsesUniformBufferView(GMaxRHIShaderPlatform))
 	{
 		switch (VertexInputStreamType)
 		{
@@ -364,7 +365,8 @@ void FLocalVertexFactory::GetVertexElements(ERHIFeatureLevel::Type FeatureLevel,
 	int32 ColorStreamIndex;
 	GetVertexElements(FeatureLevel, InputStreamType, bSupportsManualVertexFetch, Data, Elements, VertexStreams, ColorStreamIndex);
 
-	if (UseGPUScene(GMaxRHIShaderPlatform, GMaxRHIFeatureLevel))
+	if (UseGPUScene(GMaxRHIShaderPlatform, GMaxRHIFeatureLevel) 
+		&& !PlatformGPUSceneUsesUniformBufferView(GMaxRHIShaderPlatform))
 	{
 		Elements.Add(FVertexElement(VertexStreams.Num(), 0, VET_UInt, 13, sizeof(uint32), true));
 	}
@@ -538,17 +540,24 @@ void FLocalVertexFactory::GetVertexElements(
 			}
 		}
 
-		// Fill PreSkinPosition slot for GPUSkinPassThrough vertex factory, or else use a dummy buffer.
-		FVertexStreamComponent NullComponent(&GNullVertexBuffer, 0, 0, VET_Float4);
-		Elements.Add(AccessStreamComponent(Data.PreSkinPositionComponent.VertexBuffer ? Data.PreSkinPositionComponent : NullComponent, 14, InOutStreams));
-
-		if (Data.LightMapCoordinateComponent.VertexBuffer)
+		// TODO: should also check if VFType supports 'SupportsGPUSkinPassThrough'
+		if (IsGPUSkinPassThroughSupported(GMaxRHIShaderPlatform))
 		{
-			Elements.Add(AccessStreamComponent(Data.LightMapCoordinateComponent, 15, InOutStreams));
+			// Fill PreSkinPosition slot for GPUSkinPassThrough vertex factory, or else use a dummy buffer.
+			FVertexStreamComponent NullComponent(&GNullVertexBuffer, 0, 0, VET_Float4);
+			Elements.Add(AccessStreamComponent(Data.PreSkinPositionComponent.VertexBuffer ? Data.PreSkinPositionComponent : NullComponent, 14, InOutStreams));
 		}
-		else if (Data.TextureCoordinates.Num())
+
+		if (FReadOnlyCVARCache::Get().bAllowStaticLighting)
 		{
-			Elements.Add(AccessStreamComponent(Data.TextureCoordinates[0], 15, InOutStreams));
+			if (Data.LightMapCoordinateComponent.VertexBuffer)
+			{
+				Elements.Add(AccessStreamComponent(Data.LightMapCoordinateComponent, 15, InOutStreams));
+			}
+			else if (Data.TextureCoordinates.Num())
+			{
+				Elements.Add(AccessStreamComponent(Data.TextureCoordinates[0], 15, InOutStreams));
+			}
 		}
 	}
 }
diff --git a/Engine/Source/Runtime/RenderCore/Private/VertexFactory.cpp b/Engine/Source/Runtime/RenderCore/Private/VertexFactory.cpp
--- a/Engine/Source/Runtime/RenderCore/Private/VertexFactory.cpp
+++ b/Engine/Source/Runtime/RenderCore/Private/VertexFactory.cpp
@@ -290,13 +290,20 @@ bool FVertexFactory::AddPrimitiveIdStreamElement(EVertexInputStreamType InputStr
 
 		if (GIsEditor || GMaxRHIFeatureLevel > ERHIFeatureLevel::ES3_1 || AttributeIndex_Mobile != 0xff)
 		{
-			// When the VF is used for rendering in normal mesh passes, this vertex buffer and offset will be overridden
-			Elements.Add(AccessStreamComponent(FVertexStreamComponent(&GPrimitiveIdDummy, 0, 0, 0u, VET_UInt, EVertexStreamUsage::Instancing), AttributeIndex, InputStreamType));
-			SetPrimitiveIdStreamIndex(GMaxRHIFeatureLevel, InputStreamType, Elements.Last().StreamIndex);
+			// UniformView path does not use PrimitiveId stream, we still need to set it to a non-negative index
+			int32 AddedStreamIndex = 0;
+			if (!PlatformGPUSceneUsesUniformBufferView(GMaxRHIShaderPlatform))
+			{
+				// When the VF is used for rendering in normal mesh passes, this vertex buffer and offset will be overridden
+				Elements.Add(AccessStreamComponent(FVertexStreamComponent(&GPrimitiveIdDummy, 0, 0, 0u, VET_UInt, EVertexStreamUsage::Instancing), AttributeIndex, InputStreamType));
+				AddedStreamIndex = Elements.Last().StreamIndex;
+			}
+			
+			SetPrimitiveIdStreamIndex(GMaxRHIFeatureLevel, InputStreamType, AddedStreamIndex);
 			
 			if (GIsEditor && (AttributeIndex_Mobile != 0xff && GMaxRHIFeatureLevel != ERHIFeatureLevel::ES3_1))
 			{
-				SetPrimitiveIdStreamIndex(ERHIFeatureLevel::ES3_1, InputStreamType, Elements.Last().StreamIndex);
+				SetPrimitiveIdStreamIndex(ERHIFeatureLevel::ES3_1, InputStreamType, AddedStreamIndex);
 			}
 		}
 
