Fixed NonUnity compile errors

diff --git a/Engine/Source/Runtime/Engine/Private/GeometryCollectionRendering.cpp b/Engine/Source/Runtime/Engine/Private/GeometryCollectionRendering.cpp
--- a/Engine/Source/Runtime/Engine/Private/GeometryCollectionRendering.cpp
+++ b/Engine/Source/Runtime/Engine/Private/GeometryCollectionRendering.cpp
@@ -10,6 +10,7 @@
 #include "ProfilingDebugging/LoadTimeTracker.h"
 #include "PrimitiveUniformShaderParameters.h"
 #include "RenderUtils.h"
+#include "SceneManagement.h"
 
 IMPLEMENT_TYPE_LAYOUT(FGeometryCollectionVertexFactoryShaderParameters);
 
diff --git a/Engine/Source/Runtime/Engine/Private/InstancedStaticMesh.cpp b/Engine/Source/Runtime/Engine/Private/InstancedStaticMesh.cpp
--- a/Engine/Source/Runtime/Engine/Private/InstancedStaticMesh.cpp
+++ b/Engine/Source/Runtime/Engine/Private/InstancedStaticMesh.cpp
@@ -45,6 +45,7 @@
 #include "UObject/UObjectIterator.h"
 #include "GenericPlatform/ICursor.h"
 #include "Rendering/RenderCommandPipes.h"
+#include "SceneManagement.h"
 
 #include "InstancedStaticMeshSceneProxyDesc.h"
 
diff --git a/Engine/Source/Runtime/Engine/Private/LocalVertexFactory.cpp b/Engine/Source/Runtime/Engine/Private/LocalVertexFactory.cpp
--- a/Engine/Source/Runtime/Engine/Private/LocalVertexFactory.cpp
+++ b/Engine/Source/Runtime/Engine/Private/LocalVertexFactory.cpp
@@ -20,6 +20,7 @@
 #include "GPUSkinVertexFactory.h"
 #include "RenderUtils.h"
 #include "SceneInterface.h"
+#include "SceneManagement.h"
 
 IMPLEMENT_TYPE_LAYOUT(FLocalVertexFactoryShaderParametersBase);
 IMPLEMENT_TYPE_LAYOUT(FLocalVertexFactoryShaderParameters);
