Add API to seek to a specific `frame` so we can be way more accurate about seeking.

- Avoids converting Frames to Seconds then back to Frames again.
- Allows for more precise (frame specific) seeking

diff --git a/Engine/Source/Runtime/AudioCodecEngine/Private/SoundWaveProxyReader.cpp b/Engine/Source/Runtime/AudioCodecEngine/Private/SoundWaveProxyReader.cpp
--- a/Engine/Source/Runtime/AudioCodecEngine/Private/SoundWaveProxyReader.cpp
+++ b/Engine/Source/Runtime/AudioCodecEngine/Private/SoundWaveProxyReader.cpp
@@ -146,7 +146,7 @@ bool FSoundWaveProxyReader::SeekToTime(float InSeconds)
 
 	if (WaveProxy->IsSeekable() && CompressedAudioInfo)
 	{
-		CompressedAudioInfo->SeekToTime(InSeconds);
+		CompressedAudioInfo->SeekToFrame(InFrameIndex);
 		CurrentFrameIndex = InFrameIndex;
 		DecoderOutput.SetNum(0);
 		NumDecodeSamplesToDiscard = 0;
@@ -158,6 +158,30 @@ bool FSoundWaveProxyReader::SeekToTime(float InSeconds)
 	return bIsDecoderValid;
 }
 
+AUDIOCODECENGINE_API bool FSoundWaveProxyReader::SeekToFrame(uint32 InFrameNum)
+{
+	// ignore seek request if we're already at the specified time
+	if (InFrameNum == CurrentFrameIndex)
+	{
+		return bIsDecoderValid;
+	}
+
+	if (WaveProxy->IsSeekable() && CompressedAudioInfo)
+	{
+		CompressedAudioInfo->SeekToFrame(InFrameNum);
+		CurrentFrameIndex = InFrameNum;
+		DecoderOutput.SetNum(0);
+		NumDecodeSamplesToDiscard = 0;
+		DecodeResult = EDecodeResult::MoreDataRemaining;
+		return bIsDecoderValid;
+	}
+
+	// Direct seeking is not supported. A new decoder must be created.
+	float Seconds = GetSampleRate() * InFrameNum;
+	bIsDecoderValid = InitializeDecoder(Seconds);
+	return bIsDecoderValid;
+}
+
 /** Copies audio into OutBuffer. It returns the number of samples copied.
  * Samples not written to will be set to zero.
  */
@@ -224,7 +248,6 @@ int32 FSoundWaveProxyReader::PopAudio(Audio::FAlignedFloatBuffer& OutBuffer)
 			// we can produce more audio, but we were unable to
 			// this is likely due to the streaming data not being available yet
 			// let's early out to avoid a hitch and hope that it's ready on the next read
-			UE_LOG(LogAudio, Verbose, TEXT("FSoundWaveProxyReader experienced an underrun while decoding samples: %s"), *WaveProxy->GetPackageName().ToString());
 			break;
 		}
 	}
diff --git a/Engine/Source/Runtime/AudioCodecEngine/Public/SoundWaveProxyReader.h b/Engine/Source/Runtime/AudioCodecEngine/Public/SoundWaveProxyReader.h
--- a/Engine/Source/Runtime/AudioCodecEngine/Public/SoundWaveProxyReader.h
+++ b/Engine/Source/Runtime/AudioCodecEngine/Public/SoundWaveProxyReader.h
@@ -135,6 +135,15 @@ public:
 	 */
 	AUDIOCODECENGINE_API bool SeekToTime(float InSeconds);
 
+	/** Seeks to position in wave at a specific frame.
+	 *
+	 * @param InFrameNum - The specific frame to seek the playhead
+	 *
+	 * @return true on success, false on failure.
+	 */
+	AUDIOCODECENGINE_API bool SeekToFrame(uint32 InFrameNum);
+
+
 	/** Pops audio from reader and copies audio into OutBuffer. It returns the number of samples copied.
 	 * Samples not written to will be set to zero.
 	 */
diff --git a/Engine/Source/Runtime/BinkAudioDecoder/Module/Private/BinkAudioInfo.cpp b/Engine/Source/Runtime/BinkAudioDecoder/Module/Private/BinkAudioInfo.cpp
--- a/Engine/Source/Runtime/BinkAudioDecoder/Module/Private/BinkAudioInfo.cpp
+++ b/Engine/Source/Runtime/BinkAudioDecoder/Module/Private/BinkAudioInfo.cpp
@@ -121,14 +121,28 @@ void FBinkAudioInfo::SeekToTime(const float SeekTimeSeconds)
 		return;
 	}
 
+	// convert seconds to frames and call SeekToFrame
 	uint32 SeekTimeFrames = 0;
-	uint32 SeekTimeSamples = 0;
 	if (SeekTimeSeconds > 0)
 	{
 		SeekTimeFrames = (uint32)(SeekTimeSeconds * SampleRate);
-		SeekTimeSamples = SeekTimeFrames * NumChannels;
 	}
 
+	SeekToFrame(SeekTimeFrames);
+}
+
+void FBinkAudioInfo::SeekToFrame(const uint32 InFrameNum)
+{
+	// If there's no seek table on the header, fall-back to Super implementation.
+	if (Decoder->SeekTableCount == 0)
+	{
+		Super::SeekToFrame(InFrameNum);
+		return;
+	}
+
+	uint32 SeekTimeFrames = InFrameNum;
+	uint32 SeekTimeSamples = SeekTimeFrames * NumChannels;
+
 	uint32 SamplesInFrame = GetMaxFrameSizeSamples();
 	if (SeekTimeSamples > this->TrueSampleCount)
 	{
diff --git a/Engine/Source/Runtime/BinkAudioDecoder/Module/Public/BinkAudioInfo.h b/Engine/Source/Runtime/BinkAudioDecoder/Module/Public/BinkAudioInfo.h
--- a/Engine/Source/Runtime/BinkAudioDecoder/Module/Public/BinkAudioInfo.h
+++ b/Engine/Source/Runtime/BinkAudioDecoder/Module/Public/BinkAudioInfo.h
@@ -21,6 +21,7 @@ public:
 	virtual uint32 GetMaxFrameSizeSamples() const override;
 	virtual bool CreateDecoder() override;
 	virtual void SeekToTime(const float SeekToTimeSeconds) override;
+	virtual void SeekToFrame(const uint32 SeekTimeFrames) override;
 	virtual void PrepareToLoop() override;
 	virtual FDecodeResult Decode(const uint8* CompressedData, const int32 CompressedDataSize, uint8* OutPCMData, const int32 OutputPCMDataSize) override;
 	virtual bool HasError() const override;
diff --git a/Engine/Source/Runtime/Engine/Private/ADPCMAudioInfo.cpp b/Engine/Source/Runtime/Engine/Private/ADPCMAudioInfo.cpp
--- a/Engine/Source/Runtime/Engine/Private/ADPCMAudioInfo.cpp
+++ b/Engine/Source/Runtime/Engine/Private/ADPCMAudioInfo.cpp
@@ -66,7 +66,7 @@ FADPCMAudioInfo::FADPCMAudioInfo(void)
 	, bNewSeekRequest(false)
 	, bSeekPendingRead(false)
 	, bSeekedFowardToNextChunk(false)
-	, TargetSeekTime(0.0f)
+	, TargetSeekFrame(0)
 {
 }
 
@@ -86,18 +86,31 @@ void FADPCMAudioInfo::SeekToTime(const float InSeekTime)
 		return;
 	}
 
-	TargetSeekTime = InSeekTime;
+	const uint32 SamplesPerSec = *WaveInfo.pSamplesPerSec;
+	TargetSeekFrame = static_cast<uint32>(InSeekTime * SamplesPerSec);
 	bNewSeekRequest = true;
 }
 
-void FADPCMAudioInfo::SeekToTimeInternal(const float InSeekTime)
+void FADPCMAudioInfo::SeekToFrame(const uint32 InSeekFrame)
+{
+	if (bDisableADPCMSeekingCVar)
+	{
+		return;
+	}
+
+	TargetSeekFrame = InSeekFrame;
+	bNewSeekRequest = true;
+}
+
+void FADPCMAudioInfo::SeekToFrameInternal(const uint32 InSeekFrame)
 {
 	// Reset chunk handle in preperation for a new chunk.
 	CurCompressedChunkData = nullptr;
 
-	UE_LOG(LogAudio, Verbose, TEXT("Seeking ADPCM source to %.3f sec"), InSeekTime);
+	float SeekTimeSeconds = InSeekFrame / static_cast<float>(*WaveInfo.pSamplesPerSec);
+	UE_LOG(LogAudio, Verbose, TEXT("Seeking ADPCM source to %.3f sec"), SeekTimeSeconds);
 
-	if (InSeekTime <= 0.0f)
+	if (InSeekFrame == 0)
 	{
 		CurrentCompressedBlockIndex = 0;
 		CurrentUncompressedBlockSampleIndex = 0;
@@ -111,8 +124,7 @@ void FADPCMAudioInfo::SeekToTimeInternal(const float InSeekTime)
 
 	// Calculate block index & force SeekTime to be in bounds.
 	check(WaveInfo.pSamplesPerSec != nullptr);
-	const uint32 SamplesPerSec = *WaveInfo.pSamplesPerSec;
-	uint32 SeekedSamples = static_cast<uint32>(InSeekTime * SamplesPerSec);
+	uint32 SeekedSamples = InSeekFrame;
 	TotalSamplesStreamed = FMath::Min<uint32>(SeekedSamples, TotalSamplesPerChannel - 1);
 
 	const uint32 HeaderOffset = static_cast<uint32>(WaveInfo.SampleDataStart - SrcBufferData);
@@ -447,16 +459,9 @@ int FADPCMAudioInfo::GetStreamBufferSize() const
 
 void FADPCMAudioInfo::ProcessSeekRequest()
 {
-	float NewSeekTime = -1.0f;
 	if (bNewSeekRequest)
 	{
-		NewSeekTime = TargetSeekTime;
-	}
-
-
-	if (NewSeekTime >= 0.0f)
-	{
-		SeekToTimeInternal(NewSeekTime);
+		SeekToFrameInternal(TargetSeekFrame);
 	}
 }
 
diff --git a/Engine/Source/Runtime/Engine/Private/AudioDecompress.cpp b/Engine/Source/Runtime/Engine/Private/AudioDecompress.cpp
--- a/Engine/Source/Runtime/Engine/Private/AudioDecompress.cpp
+++ b/Engine/Source/Runtime/Engine/Private/AudioDecompress.cpp
@@ -151,6 +151,25 @@ void IStreamedCompressedInfo::SeekToTime(const float InSeekTimeSeconds)
 	}
 }
 
+void IStreamedCompressedInfo::SeekToFrame(const uint32 InSeekTimeFrames)
+{
+	if (StreamingSoundWave.IsValid())
+	{
+		const FSoundWavePtr WaveData = StreamingSoundWave->GetSoundWaveData();
+		if (WaveData.IsValid())
+		{
+			// If we have a chunk setup to contain a seek-table it will return a value other than INDEX_NONE here.
+			const int32 ChunkIndexToSeekTo = WaveData->FindChunkIndexForSeeking(IntCastChecked<uint32>(InSeekTimeFrames));
+			if (ChunkIndexToSeekTo >= 0)
+			{
+				StreamSeekBlockIndex = ChunkIndexToSeekTo;
+				StreamSeekBlockOffset = 0;								// We don't know this until the seek-table is loaded, so we leave it 0 for now.
+				StreamSeekToAudioFrames = InSeekTimeFrames;			// Store the time in samples so when we load the chunks table loads we can find the offset.
+			}
+		}
+	}
+}
+
 void IStreamedCompressedInfo::ExpandFile(uint8* DstBuffer, struct FSoundQualityInfo* QualityInfo)
 {
 	check(DstBuffer);
diff --git a/Engine/Source/Runtime/Engine/Private/OpusAudioInfo.cpp b/Engine/Source/Runtime/Engine/Private/OpusAudioInfo.cpp
--- a/Engine/Source/Runtime/Engine/Private/OpusAudioInfo.cpp
+++ b/Engine/Source/Runtime/Engine/Private/OpusAudioInfo.cpp
@@ -330,16 +330,29 @@ void FOpusAudioInfo::SeekToTime(const float InSeekTime)
 	}
 	else if (SrcBufferData && SrcBufferDataSize)
 	{
-		uint32 SeekSampleNum = 0;
+		uint32 SeekFrameNum = 0;
 		if (InSeekTime > 0.0f)
 		{
-			SeekSampleNum = (uint32)(InSeekTime * Header.SampleRate * NumChannels);
+			SeekFrameNum = (uint32)(InSeekTime * Header.SampleRate);
 		}
 
+		SeekToFrame(SeekFrameNum);
+	}
+}
+
+void FOpusAudioInfo::SeekToFrame(const uint32 InSeekFrame)
+{
+	if (GetStreamingSoundWave().IsValid())
+	{
+		IStreamedCompressedInfo::SeekToFrame(InSeekFrame);
+	}
+	else if (SrcBufferData && SrcBufferDataSize)
+	{
+		uint32 SeekSampleNum = InSeekFrame * NumChannels;
 		const uint8* ChunkPtr = SrcBufferData + AudioDataOffset;
 		const uint8* const EndPtr = SrcBufferData + SrcBufferDataSize;
 		uint32 CurrentChunkSampleNum = 0;
-		while(ChunkPtr < EndPtr)
+		while (ChunkPtr < EndPtr)
 		{
 			uint32 ChunkSize = (uint32)ChunkPtr[0] + ((uint32)ChunkPtr[1] << 8);
 			int32 ExpectedFrames = opus_packet_get_nb_frames(ChunkPtr + 2, ChunkSize);
@@ -365,5 +378,4 @@ void FOpusAudioInfo::SeekToTime(const float InSeekTime)
 
 	NumRemainingSamplesToSkip = Header.NumSilentSamplesAtBeginning + Header.NumPreSkipSamples;
 	PreviousDecodedUnusedSamples.Empty();
-}
-
+}
\ No newline at end of file
diff --git a/Engine/Source/Runtime/Engine/Private/VorbisAudioInfo.cpp b/Engine/Source/Runtime/Engine/Private/VorbisAudioInfo.cpp
--- a/Engine/Source/Runtime/Engine/Private/VorbisAudioInfo.cpp
+++ b/Engine/Source/Runtime/Engine/Private/VorbisAudioInfo.cpp
@@ -553,6 +553,28 @@ void FVorbisAudioInfo::SeekToTime( const float SeekTime )
 	ov_time_seek( &VFWrapper->vf, TargetTime );
 }
 
+void FVorbisAudioInfo::SeekToFrame(const uint32 SeekFrames)
+{
+	if (!bDllLoaded)
+	{
+		UE_LOG(LogAudio, Error, TEXT("FVorbisAudioInfo::SeekToTime failed due to vorbis DLL not being loaded."));
+		return;
+	}
+
+	FScopeLock ScopeLock(&VorbisCriticalSection);
+
+	if (!bHeaderParsed)
+	{
+		UE_LOG(LogAudio, Error, TEXT("FVorbisAudioInfo::SeekToTime failed due to not parsing header first."));
+		return;
+	}
+
+	vorbis_info* vi = ov_info(&VFWrapper->vf, -1);
+	int64 SeekSamples = SeekFrames * vi->channels;
+	const int64 TargetSample = FMath::Min(SeekSamples, (int64)ov_pcm_total(&VFWrapper->vf, -1));
+	ov_pcm_seek(&VFWrapper->vf, TargetSample);
+}
+
 void FVorbisAudioInfo::EnableHalfRate( bool HalfRate )
 {
 	if (!bDllLoaded)
diff --git a/Engine/Source/Runtime/Engine/Public/ADPCMAudioInfo.h b/Engine/Source/Runtime/Engine/Public/ADPCMAudioInfo.h
--- a/Engine/Source/Runtime/Engine/Public/ADPCMAudioInfo.h
+++ b/Engine/Source/Runtime/Engine/Public/ADPCMAudioInfo.h
@@ -98,6 +98,7 @@ public:
 	ENGINE_API virtual bool ReadCompressedInfo(const uint8* InSrcBufferData, uint32 InSrcBufferDataSize, struct FSoundQualityInfo* QualityInfo);
 	ENGINE_API virtual bool ReadCompressedData(uint8* Destination, bool bLooping, uint32 BufferSize);
 	ENGINE_API virtual void SeekToTime(const float SeekTime);
+	ENGINE_API virtual void SeekToFrame(const uint32 SeekFrame);
 	ENGINE_API virtual void ExpandFile(uint8* DstBuffer, struct FSoundQualityInfo* QualityInfo);
 	virtual void EnableHalfRate(bool HalfRate) {};
 	virtual uint32 GetSourceBufferSize() const { return SrcBufferDataSize; }
@@ -179,10 +180,10 @@ private:
 
 private:
 	void ProcessSeekRequest();
-	void SeekToTimeInternal(const float InSeekTime);
+	void SeekToFrameInternal(const uint32 InSeekFrame);
 	void ResetSeekState();
 
-	std::atomic<float> TargetSeekTime;
+	std::atomic<uint32> TargetSeekFrame;
 
 	FCriticalSection StreamSeekCriticalSection;
 };
diff --git a/Engine/Source/Runtime/Engine/Public/AudioDecompress.h b/Engine/Source/Runtime/Engine/Public/AudioDecompress.h
--- a/Engine/Source/Runtime/Engine/Public/AudioDecompress.h
+++ b/Engine/Source/Runtime/Engine/Public/AudioDecompress.h
@@ -64,6 +64,11 @@ public:
 	 */
 	virtual void SeekToTime(const float SeekTime) = 0;
 
+	/**
+	* Seeks to specific frame in the audio (Some formats might not be seekable)
+	*/
+	virtual void SeekToFrame(const uint32 Frame) = 0;
+
 	/**
 	* Decompress an entire data file to a TArray
 	*/
@@ -199,7 +204,8 @@ public:
 	//~ Begin ICompressedInfo Interface
 	ENGINE_API virtual bool ReadCompressedInfo(const uint8* InSrcBufferData, uint32 InSrcBufferDataSize, FSoundQualityInfo* QualityInfo) override;
 	ENGINE_API virtual bool ReadCompressedData(uint8* Destination, bool bLooping, uint32 BufferSize) override;
-	ENGINE_API virtual void SeekToTime(const float SeekTime) override;;
+	ENGINE_API virtual void SeekToTime(const float SeekTime) override;
+	ENGINE_API virtual void SeekToFrame(const uint32 SeekFrame) override;
 	ENGINE_API virtual void ExpandFile(uint8* DstBuffer, struct FSoundQualityInfo* QualityInfo) override;
 	virtual void EnableHalfRate(bool HalfRate) override {};
 	virtual uint32 GetSourceBufferSize() const override { return SrcBufferDataSize; }
diff --git a/Engine/Source/Runtime/Engine/Public/OpusAudioInfo.h b/Engine/Source/Runtime/Engine/Public/OpusAudioInfo.h
--- a/Engine/Source/Runtime/Engine/Public/OpusAudioInfo.h
+++ b/Engine/Source/Runtime/Engine/Public/OpusAudioInfo.h
@@ -29,6 +29,7 @@ public:
 	FDecodeResult Decode(const uint8* CompressedData, const int32 CompressedDataSize, uint8* OutPCMData, const int32 OutputPCMDataSize) override;
 	void PrepareToLoop() override;
 	void SeekToTime(const float SeekTime) override;
+	void SeekToFrame(const uint32 SeekFrame) override;
 	//~ End IStreamedCompressedInfo Interface
 
 	struct ENGINE_API FHeader
diff --git a/Engine/Source/Runtime/Engine/Public/VorbisAudioInfo.h b/Engine/Source/Runtime/Engine/Public/VorbisAudioInfo.h
--- a/Engine/Source/Runtime/Engine/Public/VorbisAudioInfo.h
+++ b/Engine/Source/Runtime/Engine/Public/VorbisAudioInfo.h
@@ -80,6 +80,8 @@ public:
 
 	ENGINE_API virtual void SeekToTime( const float SeekTime ) override;
 
+	ENGINE_API virtual void SeekToFrame(const uint32 SeekFrame) override;
+
 	/** 
 	 * Decompress an entire ogg data file to a TArray
 	 */
