[shaders] add debug dump option to output "detailed" shader source, including the output of the preprocessing step unmodified, the stripped version used for hashing, and whatever the final source passed to the compiler is (can differ from the stripped version if further modifications are applied as part of the compile step). this is only done if the preprocessed job cache is enabled; otherwise (and if this cvar is disabled) the single output will either be the compiler-step-modified source if it exists, or the output of the preprocessing step.

diff --git a/Engine/Source/Developer/ShaderCompilerCommon/Private/ShaderCompilerCommon.cpp b/Engine/Source/Developer/ShaderCompilerCommon/Private/ShaderCompilerCommon.cpp
--- a/Engine/Source/Developer/ShaderCompilerCommon/Private/ShaderCompilerCommon.cpp
+++ b/Engine/Source/Developer/ShaderCompilerCommon/Private/ShaderCompilerCommon.cpp
@@ -1625,7 +1625,26 @@ namespace UE::ShaderCompilerCommon
 		const FShaderCompilerOutput& Output,
 		const FDebugShaderDataOptions& Options)
 	{
-		const FString& SourceToDump = Output.ModifiedShaderSource.IsEmpty() ? PreprocessOutput.GetUnstrippedSource() : Output.ModifiedShaderSource;
+		if (Input.bCachePreprocessed && EnumHasAnyFlags(Input.DebugInfoFlags, EShaderDebugInfoFlags::DetailedSource))
+		{
+			FDebugShaderDataOptions PrefixedOptions(Options);
+			uint32 SlackLen = Options.FilenamePrefix ? FCString::Strlen(Options.FilenamePrefix) : 0;
+			FString StrippedPrefix(TEXT("Stripped_"), SlackLen);
+			FString PreprocessedPrefix(TEXT("Preprocessed_"), SlackLen);
+			if (Options.FilenamePrefix)
+			{
+				StrippedPrefix += Options.FilenamePrefix;
+				PreprocessedPrefix += Options.FilenamePrefix;
+			}
+			
+			PrefixedOptions.FilenamePrefix = *StrippedPrefix;
+			FFileHelper::SaveStringToFile(PreprocessOutput.GetSource(), *PrefixedOptions.GetDebugShaderPath(Input));
+
+			PrefixedOptions.FilenamePrefix = *PreprocessedPrefix;
+			FFileHelper::SaveStringToFile(PreprocessOutput.GetUnstrippedSource(), *PrefixedOptions.GetDebugShaderPath(Input));
+		}
+			
+		const FString& SourceToDump = Output.ModifiedShaderSource.IsEmpty() ? PreprocessOutput.GetSource() : Output.ModifiedShaderSource;
 		DumpDebugShaderData(Input, SourceToDump, Options);
 		FFileHelper::SaveStringToFile(Output.OutputHash.ToString(), *GetDebugFileName(Input, Options, TEXT("OutputHash.txt")), FFileHelper::EEncodingOptions::ForceAnsi);
 
diff --git a/Engine/Source/Developer/ShaderCompilerCommon/Public/ShaderCompilerCommon.h b/Engine/Source/Developer/ShaderCompilerCommon/Public/ShaderCompilerCommon.h
--- a/Engine/Source/Developer/ShaderCompilerCommon/Public/ShaderCompilerCommon.h
+++ b/Engine/Source/Developer/ShaderCompilerCommon/Public/ShaderCompilerCommon.h
@@ -120,6 +120,15 @@ namespace UE::ShaderCompilerCommon
 			const TCHAR* BaseFileName;
 			const TCHAR* Data;
 		};
+		
+		PRAGMA_DISABLE_DEPRECATION_WARNINGS
+		// Explicitly-defaulted ctors are needed temporarily due to deprecation of the HlslCCFlags field.
+		// These can be removed once the deprecation window for said field ends.
+		FDebugShaderDataOptions() = default;
+		FDebugShaderDataOptions(FDebugShaderDataOptions&&) = default;
+		FDebugShaderDataOptions(const FDebugShaderDataOptions&) = default;
+		PRAGMA_ENABLE_DEPRECATION_WARNINGS
+		
 		UE_DEPRECATED(5.4, "HlslCCFlags field is no longer used in debug output; please remove any usage.")
 		uint32 HlslCCFlags = 0;
 		const TCHAR* OverrideBaseFilename = nullptr;
diff --git a/Engine/Source/Runtime/Engine/Private/ShaderCompiler/ShaderCompiler.cpp b/Engine/Source/Runtime/Engine/Private/ShaderCompiler/ShaderCompiler.cpp
--- a/Engine/Source/Runtime/Engine/Private/ShaderCompiler/ShaderCompiler.cpp
+++ b/Engine/Source/Runtime/Engine/Private/ShaderCompiler/ShaderCompiler.cpp
@@ -163,6 +163,16 @@ static TAutoConsoleVariable<bool> CVarDebugDumpShaderCode(
 	TEXT("If true, each shader job will dump a ShaderCode.bin containing the contents of the output shader code object (the contents of this can differ for each shader format; note that this is the data that is hashed to produce the OutputHash.txt file)"),
 	ECVF_ReadOnly);
 
+static TAutoConsoleVariable<bool> CVarDebugDumpDetailedShaderSource(
+	TEXT("r.ShaderCompiler.DebugDumpDetailedShaderSource"),
+	false,
+	TEXT("If true, and if the preprocessed job cache is enabled, this will dump multiple copies of the shader source for any job which has debug output enabled:\n")
+	TEXT("\t1. The unmodified output of the preprocessing step as constructed by the PreprocessShader implementation of the IShaderFormat (Preprocessed_<shader>.usf\n")
+	TEXT("\t2. The stripped version of the above (with comments, line directives, and whitespace-only lines removed), which is the version hashed for inclusion in the job input hash when the preprocessed job cache is enabled (Stripped_<shader>.usf)")
+	TEXT("\t3. The final source as passed to the platform compiler (this will differ if the IShaderFormat compile function applies further modifications to the source after preprocessing; otherwise this will be the same as 2 above (<shader>.usf)\n")
+	TEXT("If false, or the preprocessed job cache is disabled, this will simply dump whatever source is passed to the compiler (equivalent to either 1 or 3 depending on if the IShaderFormat implementation modifies the source in the compile step."),
+	ECVF_ReadOnly);
+
 static TAutoConsoleVariable<bool> CVarCompileParallelInProcess(
 	TEXT("r.ShaderCompiler.ParallelInProcess"),
 	false,
@@ -5669,6 +5679,11 @@ EShaderDebugInfoFlags FShaderCompilingManager::GetDumpShaderDebugInfoFlags() con
 		Flags |= EShaderDebugInfoFlags::ShaderCodeBinary;
 	}
 
+	if (CVarDebugDumpDetailedShaderSource.GetValueOnAnyThread())
+	{
+		Flags |= EShaderDebugInfoFlags::DetailedSource;
+	}
+
 	return Flags;
 }
 
diff --git a/Engine/Source/Runtime/RenderCore/Public/ShaderCompilerCore.h b/Engine/Source/Runtime/RenderCore/Public/ShaderCompilerCore.h
--- a/Engine/Source/Runtime/RenderCore/Public/ShaderCompilerCore.h
+++ b/Engine/Source/Runtime/RenderCore/Public/ShaderCompilerCore.h
@@ -219,6 +219,7 @@ enum class EShaderDebugInfoFlags : uint8
 	InputHash = 1 << 1,
 	Diagnostics = 1 << 2,
 	ShaderCodeBinary = 1 << 3,
+	DetailedSource = 1 << 4,
 };
 ENUM_CLASS_FLAGS(EShaderDebugInfoFlags)
 
