Removed optimization in GPU skin cache to move locks to the RHI thread; as the work is now offloaded to an async task.

diff --git a/Engine/Source/Runtime/Engine/Private/GPUSkinVertexFactory.cpp b/Engine/Source/Runtime/Engine/Private/GPUSkinVertexFactory.cpp
--- a/Engine/Source/Runtime/Engine/Private/GPUSkinVertexFactory.cpp
+++ b/Engine/Source/Runtime/Engine/Private/GPUSkinVertexFactory.cpp
@@ -236,17 +236,7 @@ TStatId FClothBufferPool::GetStatId() const
 TConsoleVariableData<int32>* FGPUBaseSkinVertexFactory::FShaderDataType::MaxBonesVar = NULL;
 uint32 FGPUBaseSkinVertexFactory::FShaderDataType::MaxGPUSkinBones = 0;
 
-static TAutoConsoleVariable<int32> CVarRHICmdDeferSkeletalLockAndFillToRHIThread(
-	TEXT("r.RHICmdDeferSkeletalLockAndFillToRHIThread"),
-	0,
-	TEXT("If > 0, then do the bone and cloth copies on the RHI thread. Experimental option."));
-
-static bool DeferSkeletalLockAndFillToRHIThread()
-{
-	return IsRunningRHIInSeparateThread() && CVarRHICmdDeferSkeletalLockAndFillToRHIThread.GetValueOnRenderThread() > 0 && !UE::RenderCommandPipe::SkeletalMesh.IsReplaying();
-}
-
-bool FGPUBaseSkinVertexFactory::FShaderDataType::UpdateBoneData(FRHICommandList& RHICmdList, const TArray<FMatrix44f>& ReferenceToLocalMatrices,
+void FGPUBaseSkinVertexFactory::FShaderDataType::UpdateBoneData(FRHICommandList& RHICmdList, const TArray<FMatrix44f>& ReferenceToLocalMatrices,
 	const TArray<FBoneIndexType>& BoneMap, uint32 RevisionNumber, bool bPrevious, ERHIFeatureLevel::Type InFeatureLevel, bool bUseSkinCache, bool bForceUpdateImmediately, const FName& AssetPathName)
 {
 	// stat disabled by default due to low-value/high-frequency
@@ -285,38 +275,6 @@ bool FGPUBaseSkinVertexFactory::FShaderDataType::UpdateBoneData(FRHICommandList&
 		}
 		if(NumBones)
 		{
-			if (!bUseSkinCache && !bForceUpdateImmediately && DeferSkeletalLockAndFillToRHIThread())
-			{
-				FRHIBuffer* VertexBuffer = CurrentBoneBuffer->VertexBufferRHI;
-				RHICmdList.EnqueueLambda([VertexBuffer, VectorArraySize, &ReferenceToLocalMatrices, &BoneMap](FRHICommandList& InRHICmdList)
-				{
-					QUICK_SCOPE_CYCLE_COUNTER(STAT_FRHICommandUpdateBoneBuffer_Execute);
-					FMatrix3x4* LambdaChunkMatrices = (FMatrix3x4*)InRHICmdList.LockBuffer(VertexBuffer, 0, VectorArraySize, RLM_WriteOnly);
-					//FMatrix3x4 is sizeof() == 48
-					// PLATFORM_CACHE_LINE_SIZE (128) / 48 = 2.6
-					//  sizeof(FMatrix) == 64
-					// PLATFORM_CACHE_LINE_SIZE (128) / 64 = 2
-					const uint32 LocalNumBones = BoneMap.Num();
-					check(LocalNumBones > 0 && LocalNumBones < 256); // otherwise maybe some bad threading on BoneMap, maybe we need to copy that
-					const int32 PreFetchStride = 2; // FPlatformMisc::Prefetch stride
-					for (uint32 BoneIdx = 0; BoneIdx < LocalNumBones; BoneIdx++)
-					{
-						const FBoneIndexType RefToLocalIdx = BoneMap[BoneIdx];
-						check(ReferenceToLocalMatrices.IsValidIndex(RefToLocalIdx)); // otherwise maybe some bad threading on BoneMap, maybe we need to copy that
-						FPlatformMisc::Prefetch(ReferenceToLocalMatrices.GetData() + RefToLocalIdx + PreFetchStride);
-						FPlatformMisc::Prefetch(ReferenceToLocalMatrices.GetData() + RefToLocalIdx + PreFetchStride, PLATFORM_CACHE_LINE_SIZE);
-
-						FMatrix3x4& BoneMat = LambdaChunkMatrices[BoneIdx];
-						const FMatrix44f& RefToLocal = ReferenceToLocalMatrices[RefToLocalIdx];
-						RefToLocal.To3x4MatrixTranspose((float*)BoneMat.M);
-					}
-					InRHICmdList.UnlockBuffer(VertexBuffer);
-				});
-
-				RHICmdList.RHIThreadFence(true);
-
-				return true;
-			}
 			ChunkMatrices = (FMatrix3x4*)RHICmdList.LockBuffer(CurrentBoneBuffer->VertexBufferRHI, 0, VectorArraySize, RLM_WriteOnly);
 		}
 	}
@@ -373,7 +331,6 @@ bool FGPUBaseSkinVertexFactory::FShaderDataType::UpdateBoneData(FRHICommandList&
 	{
 		UniformBuffer = RHICreateUniformBuffer(&GBoneUniformStruct, FBoneMatricesUniformShaderParameters::FTypeInfo::GetStructMetadata()->GetLayoutPtr(), UniformBuffer_MultiFrame);
 	}
-	return false;
 }
 
 int32 FGPUBaseSkinVertexFactory::GetMinimumPerPlatformMaxGPUSkinBonesValue()
@@ -1054,7 +1011,7 @@ IMPLEMENT_TYPE_LAYOUT(TGPUSkinAPEXClothVertexFactoryShaderParameters);
 	TGPUSkinAPEXClothVertexFactory::ClothShaderType
 -----------------------------------------------------------------------------*/
 
-bool FGPUBaseSkinAPEXClothVertexFactory::ClothShaderType::UpdateClothSimulData(FRHICommandList& RHICmdList, const TArray<FVector3f>& InSimulPositions,
+void FGPUBaseSkinAPEXClothVertexFactory::ClothShaderType::UpdateClothSimulData(FRHICommandList& RHICmdList, const TArray<FVector3f>& InSimulPositions,
 	const TArray<FVector3f>& InSimulNormals, uint32 RevisionNumber, ERHIFeatureLevel::Type FeatureLevel, bool bForceUpdateImmediately, const FName& AssetPathName)
 {
 	QUICK_SCOPE_CYCLE_COUNTER(STAT_FGPUBaseSkinAPEXClothVertexFactory_UpdateClothSimulData);
@@ -1083,35 +1040,6 @@ bool FGPUBaseSkinAPEXClothVertexFactory::ClothShaderType::UpdateClothSimulData(F
 
 	if(NumSimulVerts)
 	{
-		if (!bForceUpdateImmediately && DeferSkeletalLockAndFillToRHIThread())
-		{
-			FRHIBuffer* VertexBuffer = CurrentClothBuffer->VertexBufferRHI;
-			RHICmdList.EnqueueLambda([VertexBuffer, VectorArraySize, &InSimulPositions, &InSimulNormals](FRHICommandList& InRHICmdList)
-			{
-				QUICK_SCOPE_CYCLE_COUNTER(STAT_FRHICommandUpdateBoneBuffer_Execute);
-				float* RESTRICT Data = (float* RESTRICT)InRHICmdList.LockBuffer(VertexBuffer, 0, VectorArraySize, RLM_WriteOnly);
-				uint32 LambdaNumSimulVerts = InSimulPositions.Num();
-				check(LambdaNumSimulVerts > 0 && LambdaNumSimulVerts <= MAX_APEXCLOTH_VERTICES_FOR_VB);
-				float* RESTRICT Pos = (float* RESTRICT) &InSimulPositions[0].X;
-				float* RESTRICT Normal = (float* RESTRICT) &InSimulNormals[0].X;
-				for (uint32 Index = 0; Index < LambdaNumSimulVerts; Index++)
-				{
-					FPlatformMisc::Prefetch(Pos + PLATFORM_CACHE_LINE_SIZE);
-					FPlatformMisc::Prefetch(Normal + PLATFORM_CACHE_LINE_SIZE);
-
-					FMemory::Memcpy(Data, Pos, sizeof(float) * 3);
-					FMemory::Memcpy(Data + 3, Normal, sizeof(float) * 3);
-					Data += 6;
-					Pos += 3;
-					Normal += 3;
-				}
-				InRHICmdList.UnlockBuffer(VertexBuffer);
-			});
-
-			RHICmdList.RHIThreadFence(true);
-
-			return true;
-		}
 		float* RESTRICT Data = (float* RESTRICT)RHICmdList.LockBuffer(CurrentClothBuffer->VertexBufferRHI, 0, VectorArraySize, RLM_WriteOnly);
 		{
 			QUICK_SCOPE_CYCLE_COUNTER(STAT_FGPUBaseSkinAPEXClothVertexFactory_UpdateClothSimulData_CopyData);
@@ -1131,8 +1059,6 @@ bool FGPUBaseSkinAPEXClothVertexFactory::ClothShaderType::UpdateClothSimulData(F
 		}
 		RHICmdList.UnlockBuffer(CurrentClothBuffer->VertexBufferRHI);
 	}
-	
-	return false;
 }
 
 void FGPUBaseSkinAPEXClothVertexFactory::ClothShaderType::SetCurrentRevisionNumber(uint32 RevisionNumber)
diff --git a/Engine/Source/Runtime/Engine/Private/SkeletalRenderGPUSkin.cpp b/Engine/Source/Runtime/Engine/Private/SkeletalRenderGPUSkin.cpp
--- a/Engine/Source/Runtime/Engine/Private/SkeletalRenderGPUSkin.cpp
+++ b/Engine/Source/Runtime/Engine/Private/SkeletalRenderGPUSkin.cpp
@@ -248,7 +248,6 @@ FSkeletalMeshObjectGPUSkin::FSkeletalMeshObjectGPUSkin(USkinnedMeshComponent* In
 
 FSkeletalMeshObjectGPUSkin::~FSkeletalMeshObjectGPUSkin()
 {
-	check(!RHIThreadFenceForDynamicData.GetReference());
 	if (DynamicData)
 	{
 		FDynamicSkelMeshObjectDataGPUSkin::FreeDynamicSkelMeshObjectDataGPUSkin(DynamicData);
@@ -296,7 +295,7 @@ void FSkeletalMeshObjectGPUSkin::ReleaseResources()
 	ReleaseMorphResources();
 	FSkeletalMeshObjectGPUSkin* MeshObject = this;
 	FGPUSkinCacheEntry** PtrSkinCacheEntry = &SkinCacheEntry;
-	ENQUEUE_RENDER_COMMAND(WaitRHIThreadFenceForDynamicData)(UE::RenderCommandPipe::SkeletalMesh,
+	ENQUEUE_RENDER_COMMAND(ReleaseSkeletalMeshSkinCacheResources)(UE::RenderCommandPipe::SkeletalMesh,
 		[MeshObject, PtrSkinCacheEntry, &SkinCacheEntryForRayTracing = SkinCacheEntryForRayTracing](FRHICommandList& RHICmdList)
 		{
 			FGPUSkinCacheEntry*& LocalSkinCacheEntry = *PtrSkinCacheEntry;
@@ -304,7 +303,6 @@ void FSkeletalMeshObjectGPUSkin::ReleaseResources()
 			FGPUSkinCacheEntry* LocalSkinCacheEntryForRayTracing = SkinCacheEntryForRayTracing;
 			FGPUSkinCache::Release(LocalSkinCacheEntryForRayTracing);
 
-			MeshObject->WaitForRHIThreadFenceForDynamicData();
 			*PtrSkinCacheEntry = nullptr;
 			SkinCacheEntryForRayTracing = nullptr;
 		}
@@ -505,7 +503,6 @@ void FSkeletalMeshObjectGPUSkin::UpdateDynamicData_RenderThread(FGPUSkinCache* G
 	}
 #endif
 
-	WaitForRHIThreadFenceForDynamicData();
 	if (DynamicData)
 	{
 		FDynamicSkelMeshObjectDataGPUSkin::FreeDynamicSkelMeshObjectDataGPUSkin(DynamicData);
@@ -555,17 +552,6 @@ void FSkeletalMeshObjectGPUSkin::PreGDMECallback(FRHICommandList& RHICmdList, FG
 	}
 }
 
-void FSkeletalMeshObjectGPUSkin::WaitForRHIThreadFenceForDynamicData()
-{
-	// we should be done with the old data at this point
-	if (RHIThreadFenceForDynamicData.GetReference())
-	{
-		FScopeCycleCounter Context(GetStatId());
-		RHIThreadFenceForDynamicData->Wait();
-		RHIThreadFenceForDynamicData = nullptr;
-	}
-}
-
 void FSkeletalMeshObjectGPUSkin::ProcessUpdatedDynamicData(EGPUSkinCacheEntryMode Mode, FGPUSkinCache* GPUSkinCache, FRHICommandList& RHICmdList, uint32 FrameNumberToPrepare, uint32 RevisionNumber, bool bMorphNeedsUpdate, int32 LODIndex)
 {
 	QUICK_SCOPE_CYCLE_COUNTER(STAT_FSkeletalMeshObjectGPUSkin_ProcessUpdatedDynamicData);
@@ -703,7 +689,7 @@ void FSkeletalMeshObjectGPUSkin::ProcessUpdatedDynamicData(EGPUSkinCacheEntryMod
 
 			// Create a uniform buffer from the bone transforms.
 			TArray<FMatrix44f>& ReferenceToLocalMatrices = bShouldUseSeparateMatricesForRayTracing ? DynamicData->ReferenceToLocalForRayTracing : DynamicData->ReferenceToLocal;
-			bool bNeedFence = ShaderData.UpdateBoneData(RHICmdList, ReferenceToLocalMatrices, Section.BoneMap, RevisionNumber, false, FeatureLevel, bUseSkinCache, DynamicData->bForceUpdateDynamicDataImmediately, OwnerName);
+			ShaderData.UpdateBoneData(RHICmdList, ReferenceToLocalMatrices, Section.BoneMap, RevisionNumber, false, FeatureLevel, bUseSkinCache, DynamicData->bForceUpdateDynamicDataImmediately, OwnerName);
 			ShaderData.UpdatedFrameNumber = FrameNumberToPrepare;
 
 			// Update uniform buffer for APEX cloth simulation mesh positions and normals
@@ -713,7 +699,7 @@ void FSkeletalMeshObjectGPUSkin::ProcessUpdatedDynamicData(EGPUSkinCacheEntryMod
 				ClothShaderData.ClothBlendWeight = DynamicData->ClothBlendWeight;
 				ClothShaderData.WorldScale = (FVector3f)WorldScale;
 
-				bNeedFence = ClothShaderData.UpdateClothSimulData(RHICmdList, SimData->Positions, SimData->Normals, RevisionNumber, FeatureLevel, DynamicData->bForceUpdateDynamicDataImmediately, OwnerName) || bNeedFence;
+				ClothShaderData.UpdateClothSimulData(RHICmdList, SimData->Positions, SimData->Normals, RevisionNumber, FeatureLevel, DynamicData->bForceUpdateDynamicDataImmediately, OwnerName);
 				// Transform from cloth space to local space. Cloth space is relative to cloth root bone, local space is component space.
 				ClothShaderData.GetClothToLocalForWriting() = FMatrix44f(SimData->ComponentRelativeTransform.ToMatrixWithScale());
 			}
@@ -754,11 +740,6 @@ void FSkeletalMeshObjectGPUSkin::ProcessUpdatedDynamicData(EGPUSkinCacheEntryMod
 				FGPUBaseSkinAPEXClothVertexFactory::ClothShaderType& ClothShaderData = VertexFactoryData.ClothVertexFactories[SectionIdx]->GetClothShaderData();
 				ClothShaderData.EnableDoubleBuffer();
 			}
-
-			if (bNeedFence)
-			{
-				RHIThreadFenceForDynamicData = RHICmdList.RHIThreadFence(true);
-			}
 		}
 	}
 
diff --git a/Engine/Source/Runtime/Engine/Private/SkeletalRenderGPUSkin.h b/Engine/Source/Runtime/Engine/Private/SkeletalRenderGPUSkin.h
--- a/Engine/Source/Runtime/Engine/Private/SkeletalRenderGPUSkin.h
+++ b/Engine/Source/Runtime/Engine/Private/SkeletalRenderGPUSkin.h
@@ -670,17 +670,12 @@ protected:
 
 	ENGINE_API virtual void UpdateMorphVertexBuffer(FRHICommandList& RHICmdList, EGPUSkinCacheEntryMode Mode, FSkeletalMeshObjectLOD& LOD, const FSkeletalMeshLODRenderData& LODData, bool bGPUSkinCacheEnabled, FMorphVertexBuffer& MorphVertexBuffer);
 
-	ENGINE_API void WaitForRHIThreadFenceForDynamicData();
-
 	/** Render data for each LOD */
 	TArray<struct FSkeletalMeshObjectLOD> LODs;
 
 	/** Data that is updated dynamically and is needed for rendering */
 	FDynamicSkelMeshObjectDataGPUSkin* DynamicData;
 
-	/** Fence for dynamic Data */
-	FGraphEventRef RHIThreadFenceForDynamicData;
-
 	/** True if we are doing a deferred update later in GDME. */
 	bool bNeedsUpdateDeferred;
 
diff --git a/Engine/Source/Runtime/Engine/Public/GPUSkinVertexFactory.h b/Engine/Source/Runtime/Engine/Public/GPUSkinVertexFactory.h
--- a/Engine/Source/Runtime/Engine/Public/GPUSkinVertexFactory.h
+++ b/Engine/Source/Runtime/Engine/Public/GPUSkinVertexFactory.h
@@ -211,7 +211,7 @@ public:
 		}
 
 		// @param FrameTime from GFrameTime
-		bool UpdateBoneData(FRHICommandList& RHICmdList, const TArray<FMatrix44f>& ReferenceToLocalMatrices,
+		void UpdateBoneData(FRHICommandList& RHICmdList, const TArray<FMatrix44f>& ReferenceToLocalMatrices,
 			const TArray<FBoneIndexType>& BoneMap, uint32 RevisionNumber, bool bPrevious, ERHIFeatureLevel::Type FeatureLevel, 
 			bool bUseSkinCache, bool bForceUpdateImmediately, const FName& AssetPathName);
 
@@ -502,7 +502,7 @@ public:
 			Reset();
 		}
 
-		bool UpdateClothSimulData(FRHICommandList& RHICmdList, const TArray<FVector3f>& InSimulPositions, const TArray<FVector3f>& InSimulNormals, uint32 RevisionNumber, 
+		void UpdateClothSimulData(FRHICommandList& RHICmdList, const TArray<FVector3f>& InSimulPositions, const TArray<FVector3f>& InSimulNormals, uint32 RevisionNumber, 
 									ERHIFeatureLevel::Type FeatureLevel, bool bForceUpdateImmediately, const FName& AssetPathName);
 
 		void ReleaseClothSimulData()
