Add an IsDebugCameraActive method to Cheat Manager

diff --git a/Engine/Source/Runtime/Engine/Classes/GameFramework/CheatManager.h b/Engine/Source/Runtime/Engine/Classes/GameFramework/CheatManager.h
--- a/Engine/Source/Runtime/Engine/Classes/GameFramework/CheatManager.h
+++ b/Engine/Source/Runtime/Engine/Classes/GameFramework/CheatManager.h
@@ -246,6 +246,9 @@ class UCheatManager : public UObject
 	UFUNCTION(exec)
 	ENGINE_API virtual void ToggleDebugCamera();
 
+	/** Returns true if the player currently has a debug camera */
+	ENGINE_API virtual bool IsDebugCameraActive() const;
+
 	/** toggles AI logging */
 	UFUNCTION(exec)
 	ENGINE_API virtual void ToggleAILogging();
diff --git a/Engine/Source/Runtime/Engine/Private/CheatManager.cpp b/Engine/Source/Runtime/Engine/Private/CheatManager.cpp
--- a/Engine/Source/Runtime/Engine/Private/CheatManager.cpp
+++ b/Engine/Source/Runtime/Engine/Private/CheatManager.cpp
@@ -631,8 +631,7 @@ void UCheatManager::StreamLevelOut(FName PackageName)
 
 void UCheatManager::ToggleDebugCamera()
 {
-	ADebugCameraController* const DCC = Cast<ADebugCameraController>(GetOuter());
-	if (DCC)
+	if (IsDebugCameraActive())
 	{
 		DisableDebugCamera();
 	}
@@ -642,6 +641,11 @@ void UCheatManager::ToggleDebugCamera()
 	}
 }
 
+bool UCheatManager::IsDebugCameraActive() const
+{
+	return GetOuter() ? GetOuter()->IsA<ADebugCameraController>() : false;
+}
+
 void UCheatManager::EnableDebugCamera()
 {
 	APlayerController* const PC = GetOuterAPlayerController();
