Simplification of LevelInstanceLevelStreaming to do Actor fixup code (no more FActorInstanceGuidMapper), register events earlier and use TArray versions

- OnLevelLoadedChanged handles Level referenced actors (Actors that aren't loaded via LoaderAdapter) : InstanceGuid, Object flags
- OnLoadedActorsAddedToLevelPreEvent handles actors loaded via LoaderAdapter : InstanceGuid, Object flags
- OnLoadedActorsAddedToLevelPostEvent handles actors loaded and registered via LoaderAdapter : Attach to Editor Actor, Reflect Selection/Level Instance state
- OnLoadedActorsRemovedFromLevelPostEvent : Detach from Editor Actor
- LoadLevelInstance : Handles creation of Editor Actor for Attachment and Reflect Selection (after BlockTill for actors that are registered)

#rnx

diff --git a/Engine/Source/Runtime/Engine/Classes/GameFramework/Actor.h b/Engine/Source/Runtime/Engine/Classes/GameFramework/Actor.h
--- a/Engine/Source/Runtime/Engine/Classes/GameFramework/Actor.h
+++ b/Engine/Source/Runtime/Engine/Classes/GameFramework/Actor.h
@@ -137,22 +137,6 @@ DECLARE_EVENT_TwoParams(AActor, FActorOnPackagingModeChanged, AActor*, bool /* b
 DECLARE_DELEGATE_RetVal_ThreeParams(bool, FOnProcessEvent, AActor*, UFunction*, void*);
 #endif
 
-#if WITH_EDITOR
-class FActorInstanceGuidMapper
-{
-	using FGuidMapper = TFunction<FGuid(const FGuid&)>;
-
-public:
-	void RegisterGuidMapper(FName InPackageName, const FGuidMapper& InGuidMapper);
-	void UnregisterGuidMapper(FName InPackageName);
-
-	FGuid MapGuid(FName InPackageName, const FGuid& InGuid);
-
-private:
-	TMap<FName, FGuidMapper> GuidMappers;
-};
-#endif
-
 /**
  * TInlineComponentArray is simply a TArray that reserves a fixed amount of space on the stack
  * to try to avoid heap allocation when there are fewer than a specified number of elements expected in the result.
@@ -947,7 +931,7 @@ protected:
 	/**
 	 * The instance GUID for this actor; this guid will be unique for actors from instanced streaming levels.
 	 * @see		ActorGuid
-	 * @note	This is not guaranteed to be valid during PostLoad in all situations, but safe to access from RegisterAllComponents.
+	 * @note	This is not guaranteed to be valid during PostLoad, but safe to access from RegisterAllComponents.
 	 */
 	UPROPERTY(BluePrintReadOnly, AdvancedDisplay, Category=Actor, Transient, NonPIEDuplicateTransient, TextExportTransient, NonTransactional)
 	FGuid ActorInstanceGuid;
diff --git a/Engine/Source/Runtime/Engine/Private/Actor.cpp b/Engine/Source/Runtime/Engine/Private/Actor.cpp
--- a/Engine/Source/Runtime/Engine/Private/Actor.cpp
+++ b/Engine/Source/Runtime/Engine/Private/Actor.cpp
@@ -104,27 +104,6 @@ FCriticalSection CSVActorClassNameToCountMapLock;
 #endif // (CSV_PROFILER && !UE_BUILD_SHIPPING)
 
 #if WITH_EDITOR
-void FActorInstanceGuidMapper::RegisterGuidMapper(FName InPackageName, const FGuidMapper& InGuidMapper)
-{
-	check(!GuidMappers.Contains(InPackageName));
-	GuidMappers.Add(InPackageName, InGuidMapper);
-}
-
-void FActorInstanceGuidMapper::UnregisterGuidMapper(FName InPackageName)
-{
-	check(GuidMappers.Contains(InPackageName));
-	GuidMappers.Remove(InPackageName);
-}
-
-FGuid FActorInstanceGuidMapper::MapGuid(FName InPackageName, const FGuid& InGuid)
-{
-	if (FGuidMapper* GuidMapper = GuidMappers.Find(InPackageName))
-	{
-		return (*GuidMapper)(InGuid);
-	}
-	return InGuid;
-}
-
 AActor::FDuplicationSeedInterface::FDuplicationSeedInterface(TMap<UObject*, UObject*>& InDuplicationSeed)
 	: DuplicationSeed(InDuplicationSeed)
 {
@@ -889,16 +868,6 @@ void AActor::Serialize(FArchive& Ar)
 			ActorGuid = FGuid::NewGuid();
 		}
 
-		if (!IsTemplate())
-		{
-			const FGuid NewActorInstanceGuid = TLazySingleton<FActorInstanceGuidMapper>::Get().MapGuid(GetOuter()->GetPackage()->GetFName(), ActorGuid);
-			if (NewActorInstanceGuid != ActorGuid)
-			{
-				check(!ActorInstanceGuid.IsValid());
-				ActorInstanceGuid = NewActorInstanceGuid;
-			}
-		}
-
 		if (!CanChangeIsSpatiallyLoadedFlag() && (Ar.CustomVer(FUE5ReleaseStreamObjectVersion::GUID) < FUE5ReleaseStreamObjectVersion::ActorGridPlacementDeprecateDefaultValueFixup))
 		{
 			bIsSpatiallyLoaded = GetClass()->GetDefaultObject<AActor>()->bIsSpatiallyLoaded;
diff --git a/Engine/Source/Runtime/Engine/Private/LevelInstance/LevelInstanceEditorInstanceActor.cpp b/Engine/Source/Runtime/Engine/Private/LevelInstance/LevelInstanceEditorInstanceActor.cpp
--- a/Engine/Source/Runtime/Engine/Private/LevelInstance/LevelInstanceEditorInstanceActor.cpp
+++ b/Engine/Source/Runtime/Engine/Private/LevelInstance/LevelInstanceEditorInstanceActor.cpp
@@ -59,6 +59,8 @@ ALevelInstanceEditorInstanceActor* ALevelInstanceEditorInstanceActor::Create(ILe
 		}
 	}
 
+	InstanceActor->PushSelectionToProxies();
+
 	return InstanceActor;
 }
 
diff --git a/Engine/Source/Runtime/Engine/Private/LevelInstance/LevelInstanceLevelStreaming.cpp b/Engine/Source/Runtime/Engine/Private/LevelInstance/LevelInstanceLevelStreaming.cpp
--- a/Engine/Source/Runtime/Engine/Private/LevelInstance/LevelInstanceLevelStreaming.cpp
+++ b/Engine/Source/Runtime/Engine/Private/LevelInstance/LevelInstanceLevelStreaming.cpp
@@ -107,25 +107,64 @@ FBox ULevelStreamingLevelInstance::GetBounds() const
 	return CachedBounds;
 }
 
-void ULevelStreamingLevelInstance::OnLoadedActorPreAddedToLevel(const TArray<AActor*>& InActors)
+void ULevelStreamingLevelInstance::OnLoadedActorsAddedToLevelPreEvent(const TArray<AActor*>& InActors)
 {
-	check(LevelInstanceEditorInstanceActor.IsValid());
-	if (ILevelInstanceInterface* LevelInstance = GetLevelInstance())
+	if (!GetWorld()->IsGameWorld())
 	{
-		for (AActor* Actor : InActors)
+		if (ILevelInstanceInterface* LevelInstance = GetLevelInstance())
 		{
-			// Must happen before the actors are registered with the world, which is the case for this delegate.
-			const FActorContainerID& ContainerID = LevelInstance->GetLevelInstanceID().GetContainerID();
-			FSetActorInstanceGuid SetActorInstanceGuid(Actor, ContainerID.GetActorGuid(Actor->GetActorGuid()));
+			for (AActor* Actor : InActors)
+			{
+				if (Actor->IsPackageExternal())
+				{
+					if (bResetLoadersCalled)
+					{
+						ResetLoaders(Actor->GetExternalPackage());
+					}
+
+					FLevelInstanceLevelStreamingUtils::MarkObjectsInPackageAsTransientAndNonTransactional(Actor->GetExternalPackage());
+				}
+
+				// Must happen before the actors are registered with the world, which is the case for this delegate.
+				const FActorContainerID& ContainerID = LevelInstance->GetLevelInstanceID().GetContainerID();
+				FSetActorInstanceGuid SetActorInstanceGuid(Actor, ContainerID.GetActorGuid(Actor->GetActorGuid()));
+
+				FSetActorIsInLevelInstance SetIsInLevelInstance(Actor);
+			}
 		}
 	}
 }
 
-void ULevelStreamingLevelInstance::OnLoadedActorAddedToLevel(AActor& InActor)
+void ULevelStreamingLevelInstance::OnLoadedActorsAddedToLevelPostEvent(const TArray<AActor*>& InActors)
 {
-	check(LevelInstanceEditorInstanceActor.IsValid());
-	// If bResetLoadersCalled is false, wait for ResetLevelInstanceLoaders call to do the ResetLoaders calls
-	PrepareLevelInstanceLoadedActor(InActor, GetLevelInstance(), bResetLoadersCalled);
+	if (!GetWorld()->IsGameWorld())
+	{
+		if (GetLoadedLevel()->bAreComponentsCurrentlyRegistered)
+		{
+			if (ILevelInstanceInterface* LevelInstance = GetLevelInstance())
+			{
+				for (AActor* Actor : InActors)
+				{
+					if (IsValid(Actor))
+					{
+						Actor->PushSelectionToProxies();
+						if (LevelInstance)
+						{
+							Actor->PushLevelInstanceEditingStateToProxies(CastChecked<AActor>(LevelInstance)->IsInEditingLevelInstance());
+						}
+
+						if (LevelInstanceEditorInstanceActor.IsValid())
+						{
+							if (Actor->GetAttachParentActor() == nullptr && !Actor->IsChildActor())
+							{
+								Actor->AttachToActor(LevelInstanceEditorInstanceActor.Get(), FAttachmentTransformRules::KeepWorldTransform);
+							}
+						}
+					}
+				}
+			}
+		}
+	}
 }
 
 void ULevelStreamingLevelInstance::ResetLevelInstanceLoaders()
@@ -173,43 +212,21 @@ void ULevelStreamingLevelInstance::ResetLevelInstanceLoaders()
 	}
 }
 
-void ULevelStreamingLevelInstance::PrepareLevelInstanceLoadedActor(AActor& InActor, ILevelInstanceInterface* InLevelInstance, bool bResetLoaders)
+void ULevelStreamingLevelInstance::OnLoadedActorsRemovedFromLevelPostEvent(const TArray<AActor*>& InActors)
 {
-	if (InActor.IsPackageExternal())
-	{
-		if (bResetLoaders)
-		{
-			ResetLoaders(InActor.GetExternalPackage());
-		}
-
-		FLevelInstanceLevelStreamingUtils::MarkObjectsInPackageAsTransientAndNonTransactional(InActor.GetExternalPackage());
-	}
-
-	InActor.PushSelectionToProxies();
-	if (InLevelInstance)
-	{
-		InActor.PushLevelInstanceEditingStateToProxies(CastChecked<AActor>(InLevelInstance)->IsInEditingLevelInstance());
-	}
+	// Detach actor or else it will keep it alive (attachement keeps a reference to the actor)
+	check(LevelInstanceEditorInstanceActor.IsValid());
 
-	if (LevelInstanceEditorInstanceActor.IsValid())
+	for (AActor* Actor : InActors)
 	{
-		if (InActor.GetAttachParentActor() == nullptr && !InActor.IsChildActor())
+		if (IsValid(Actor))
 		{
-			InActor.AttachToActor(LevelInstanceEditorInstanceActor.Get(), FAttachmentTransformRules::KeepWorldTransform);
+			if (Actor->GetAttachParentActor() == LevelInstanceEditorInstanceActor.Get())
+			{
+				Actor->DetachFromActor(FDetachmentTransformRules::KeepWorldTransform);
+			}
 		}
 	}
-
-	FSetActorIsInLevelInstance SetIsInLevelInstance(&InActor);
-}
-
-void ULevelStreamingLevelInstance::OnLoadedActorRemovedFromLevel(AActor& InActor)
-{
-	// Detach actor or else it will keep it alive (attachement keeps a reference to the actor)
-	check(LevelInstanceEditorInstanceActor.IsValid());
-	if (InActor.GetAttachParentActor() == LevelInstanceEditorInstanceActor.Get())
-	{
-		InActor.DetachFromActor(FDetachmentTransformRules::KeepWorldTransform);
-	}
 }
 
 #endif
@@ -257,12 +274,6 @@ ULevelStreamingLevelInstance* ULevelStreamingLevelInstance::LoadInstance(ILevelI
 
 	const FString LevelInstancePackageName = ULevelStreamingDynamic::GetLevelInstancePackageName(Params);
 
-#if WITH_EDITOR
-	FActorInstanceGuidMapper& ActorInstanceGuidMapper = TLazySingleton<FActorInstanceGuidMapper>::Get();
-	ActorInstanceGuidMapper.RegisterGuidMapper(*LevelInstancePackageName, [LevelInstance](const FGuid& InActorGuid) { return LevelInstance->GetLevelInstanceID().GetContainerID().GetActorGuid(InActorGuid); });
-	ON_SCOPE_EXIT { ActorInstanceGuidMapper.UnregisterGuidMapper(*LevelInstancePackageName); };
-#endif
-
 	ULevelStreamingLevelInstance* LevelStreaming = Cast<ULevelStreamingLevelInstance>(ULevelStreamingDynamic::LoadLevelInstance(Params, bOutSuccess));
 	if (bOutSuccess)
 	{
@@ -273,39 +284,15 @@ ULevelStreamingLevelInstance* ULevelStreamingLevelInstance::LoadInstance(ILevelI
 		{
 			GEngine->BlockTillLevelStreamingCompleted(LevelInstanceActor->GetWorld());
 
-			// Most of the code here is meant to allow partial support for undo/redo of LevelInstance Instance Loading:
-			// by setting the objects RF_Transient and !RF_Transactional we can check when unloading if those flags
-			// have been changed and figure out if we need to clear the transaction buffer or not.
-			// It might not be the final solution to support Undo/Redo in LevelInstances but it handles most of the non-editing part
 			if (ULevel* Level = LevelStreaming->GetLoadedLevel())
 			{
 				check(LevelStreaming->GetLevelStreamingState() == ELevelStreamingState::LoadedVisible);
 
-				Level->OnLoadedActorAddedToLevelPreEvent.AddUObject(LevelStreaming, &ULevelStreamingLevelInstance::OnLoadedActorPreAddedToLevel);
-				Level->OnLoadedActorAddedToLevelEvent.AddUObject(LevelStreaming, &ULevelStreamingLevelInstance::OnLoadedActorAddedToLevel);
-				Level->OnLoadedActorRemovedFromLevelEvent.AddUObject(LevelStreaming, &ULevelStreamingLevelInstance::OnLoadedActorRemovedFromLevel);
-
-				FLevelInstanceLevelStreamingUtils::MarkObjectsInPackageAsTransientAndNonTransactional(Level->GetPackage());
-
-				for (AActor* LevelActor : Level->Actors)
-				{
-					if (LevelActor)
-					{
-						LevelStreaming->PrepareLevelInstanceLoadedActor(*LevelActor, LevelInstance, false);
-					}
-				}
-
-				Level->ForEachActorFolder([](UActorFolder* ActorFolder)
-				{
-					if (ActorFolder->IsPackageExternal())
-					{
-						FLevelInstanceLevelStreamingUtils::MarkObjectsInPackageAsTransientAndNonTransactional(ActorFolder->GetPackage());
-					}
-					return true;
-				});
-
 				// Create special actor that will handle selection and transform
 				LevelStreaming->LevelInstanceEditorInstanceActor = ALevelInstanceEditorInstanceActor::Create(LevelInstance, Level);
+
+				// Push editing state to child actors
+				LevelInstanceActor->PushLevelInstanceEditingStateToProxies(LevelInstanceActor->IsInEditingLevelInstance());
 			}
 			else
 			{
@@ -333,8 +320,8 @@ void ULevelStreamingLevelInstance::UnloadInstance(ULevelStreamingLevelInstance*
 	{
 		ULevel* LoadedLevel = LevelStreaming->GetLoadedLevel();
 		LoadedLevel->OnLoadedActorAddedToLevelPreEvent.RemoveAll(LevelStreaming);
-		LoadedLevel->OnLoadedActorAddedToLevelEvent.RemoveAll(LevelStreaming);
-		LoadedLevel->OnLoadedActorRemovedFromLevelEvent.RemoveAll(LevelStreaming);
+		LoadedLevel->OnLoadedActorAddedToLevelPostEvent.RemoveAll(LevelStreaming);
+		LoadedLevel->OnLoadedActorRemovedFromLevelPreEvent.RemoveAll(LevelStreaming);
 		LevelStreaming->LevelInstanceEditorInstanceActor.Reset();
 
 		// Check if we need to flush the Trans buffer...
@@ -362,6 +349,33 @@ void ULevelStreamingLevelInstance::OnLevelLoadedChanged(ULevel* InLevel)
 
 	if (ULevel* NewLoadedLevel = GetLoadedLevel())
 	{
+#if WITH_EDITOR
+		if (!GetWorld()->IsGameWorld())
+		{
+			// Most of the code here is meant to allow partial support for undo/redo of LevelInstance Instance Loading:
+			// by setting the objects RF_Transient and !RF_Transactional we can check when unloading if those flags
+			// have been changed and figure out if we need to clear the transaction buffer or not.
+			// It might not be the final solution to support Undo/Redo in LevelInstances but it handles most of the non-editing part
+			check(!InLevel->bAreComponentsCurrentlyRegistered);
+			FLevelInstanceLevelStreamingUtils::MarkObjectsInPackageAsTransientAndNonTransactional(NewLoadedLevel->GetPackage());
+
+			NewLoadedLevel->ForEachActorFolder([](UActorFolder* ActorFolder)
+			{
+				if (ActorFolder->IsPackageExternal())
+				{
+					FLevelInstanceLevelStreamingUtils::MarkObjectsInPackageAsTransientAndNonTransactional(ActorFolder->GetPackage());
+				}
+				return true;
+			});
+
+			OnLoadedActorsAddedToLevelPreEvent(InLevel->Actors);
+
+			NewLoadedLevel->OnLoadedActorAddedToLevelPreEvent.AddUObject(this, &ULevelStreamingLevelInstance::OnLoadedActorsAddedToLevelPreEvent);
+			NewLoadedLevel->OnLoadedActorAddedToLevelPostEvent.AddUObject(this, &ULevelStreamingLevelInstance::OnLoadedActorsAddedToLevelPostEvent);
+			NewLoadedLevel->OnLoadedActorRemovedFromLevelPreEvent.AddUObject(this, &ULevelStreamingLevelInstance::OnLoadedActorsRemovedFromLevelPostEvent);
+		}
+#endif
+
 		check(InLevel == NewLoadedLevel);
 		if (!NewLoadedLevel->bAlreadyMovedActors)
 		{
diff --git a/Engine/Source/Runtime/Engine/Public/LevelInstance/LevelInstanceLevelStreaming.h b/Engine/Source/Runtime/Engine/Public/LevelInstance/LevelInstanceLevelStreaming.h
--- a/Engine/Source/Runtime/Engine/Public/LevelInstance/LevelInstanceLevelStreaming.h
+++ b/Engine/Source/Runtime/Engine/Public/LevelInstance/LevelInstanceLevelStreaming.h
@@ -40,10 +40,9 @@ protected:
 private:
 #if WITH_EDITOR
 	ENGINE_API void ResetLevelInstanceLoaders();
-	ENGINE_API void PrepareLevelInstanceLoadedActor(AActor& InActor, ILevelInstanceInterface* InLevelInstance, bool bResetLoaders);
-	ENGINE_API void OnLoadedActorPreAddedToLevel(const TArray<AActor*>& InActors);
-	ENGINE_API void OnLoadedActorAddedToLevel(AActor& InActor);
-	ENGINE_API void OnLoadedActorRemovedFromLevel(AActor& InActor);
+	ENGINE_API virtual void OnLoadedActorsAddedToLevelPreEvent(const TArray<AActor*>& InActors);
+	ENGINE_API virtual void OnLoadedActorsAddedToLevelPostEvent(const TArray<AActor*>& InActors);
+	ENGINE_API virtual void OnLoadedActorsRemovedFromLevelPostEvent(const TArray<AActor*>& InActors);
 
 	TWeakObjectPtr<ALevelInstanceEditorInstanceActor> LevelInstanceEditorInstanceActor;
 
