World Partition Streaming Generation: hendle the potential case where the generation state log can't be created.

diff --git a/Engine/Source/Runtime/Engine/Private/WorldPartition/WorldPartitionStreamingGeneration.cpp b/Engine/Source/Runtime/Engine/Private/WorldPartition/WorldPartitionStreamingGeneration.cpp
--- a/Engine/Source/Runtime/Engine/Private/WorldPartition/WorldPartitionStreamingGeneration.cpp
+++ b/Engine/Source/Runtime/Engine/Private/WorldPartition/WorldPartitionStreamingGeneration.cpp
@@ -1383,8 +1383,11 @@ bool UWorldPartition::GenerateContainerStreaming(const FGenerateStreamingParams&
 		StateLogSuffix += ContainerShortName;
 		LogFileAr = FWorldPartitionStreamingGenerator::CreateDumpStateLogArchive(*StateLogSuffix, !InParams.OutputLogPath);
 		
-		InContext.OutputLogFilename = LogFileAr->GetArchiveName();
-		HierarchicalLogAr = MakeUnique<FHierarchicalLogArchive>(*LogFileAr);
+		if (LogFileAr.IsValid())
+		{
+			InContext.OutputLogFilename = LogFileAr->GetArchiveName();
+			HierarchicalLogAr = MakeUnique<FHierarchicalLogArchive>(*LogFileAr);
+		}
 	}
 
 	FWorldPartitionStreamingGenerator::FWorldPartitionStreamingGeneratorParams StreamingGeneratorParams = FWorldPartitionStreamingGenerator::FWorldPartitionStreamingGeneratorParams()
@@ -1463,8 +1466,12 @@ void UWorldPartition::SetupHLODActors(const FSetupHLODActorsParams& Params)
 		StreamingGenerator.PreparationPhase(ActorDescCollection);
 
 		TUniquePtr<FArchive> LogFileAr = FWorldPartitionStreamingGenerator::CreateDumpStateLogArchive(TEXT("HLOD"));
-		FHierarchicalLogArchive HierarchicalLogAr(*LogFileAr);
-		StreamingGenerator.DumpStateLog(HierarchicalLogAr);
+
+		if (LogFileAr.IsValid())
+		{
+			FHierarchicalLogArchive HierarchicalLogAr(*LogFileAr);
+			StreamingGenerator.DumpStateLog(HierarchicalLogAr);
+		}
 
 		RuntimeHash->SetupHLODActors(StreamingGenerator.GetStreamingGenerationContext(ActorDescCollection), Params);
 	});	
