Easy optimization for some reallocation costs I spotted in SetMaterials, preallocate for known-size array.

#rnx

diff --git a/Engine/Source/Runtime/Engine/Private/PhysicsEngine/Experimental/PhysInterface_Chaos.cpp b/Engine/Source/Runtime/Engine/Private/PhysicsEngine/Experimental/PhysInterface_Chaos.cpp
--- a/Engine/Source/Runtime/Engine/Private/PhysicsEngine/Experimental/PhysInterface_Chaos.cpp
+++ b/Engine/Source/Runtime/Engine/Private/PhysicsEngine/Experimental/PhysInterface_Chaos.cpp
@@ -764,6 +764,7 @@ void FPhysInterface_Chaos::SetMaterials(const FPhysicsShapeHandle& InShape, cons
 		TArray<Chaos::FMaterialHandle> NewMaterialMaskMaterialHandles;
 
 		NewMaterialMaskHandles.Reserve(InMaterialMasks.Num());
+		NewMaterialMaskMaps.Reserve(InMaterialMasks.Num() * EPhysicalMaterialMaskColor::MAX);
 
 		int MaskMapMatIdx = 0;
 
