- Add property to enable breaking in the debugger rather than showing the retry dialog
- bDebugBreakOnPromptToRetryShaderCompile is disabled by default

diff --git a/Engine/Config/BaseEngine.ini b/Engine/Config/BaseEngine.ini
--- a/Engine/Config/BaseEngine.ini
+++ b/Engine/Config/BaseEngine.ini
@@ -1904,6 +1904,7 @@ PercentageUnusedShaderCompilingThreads=50
 ; Batching multiple jobs to reduce file overhead, but not so many that latency of blocking compiles is hurt
 MaxShaderJobBatchSize=10
 bPromptToRetryFailedShaderCompiles=False
+bDebugBreakOnPromptToRetryShaderCompile=False
 bLogJobCompletionTimes=False
 ; Only using 10ms of game thread time per frame to process async shader maps
 ProcessGameThreadTargetTime=.01
diff --git a/Engine/Source/Runtime/Engine/Private/ShaderCompiler/ShaderCompiler.cpp b/Engine/Source/Runtime/Engine/Private/ShaderCompiler/ShaderCompiler.cpp
--- a/Engine/Source/Runtime/Engine/Private/ShaderCompiler/ShaderCompiler.cpp
+++ b/Engine/Source/Runtime/Engine/Private/ShaderCompiler/ShaderCompiler.cpp
@@ -5342,6 +5342,7 @@ FShaderCompilingManager::FShaderCompilingManager() :
 
 	verify(GConfig->GetInt( TEXT("DevOptions.Shaders"), TEXT("MaxShaderJobBatchSize"), MaxShaderJobBatchSize, GEngineIni ));
 	verify(GConfig->GetBool( TEXT("DevOptions.Shaders"), TEXT("bPromptToRetryFailedShaderCompiles"), bPromptToRetryFailedShaderCompiles, GEngineIni ));
+	verify(GConfig->GetBool(TEXT("DevOptions.Shaders"), TEXT("bDebugBreakOnPromptToRetryShaderCompile"), bDebugBreakOnPromptToRetryShaderCompile, GEngineIni));
 	verify(GConfig->GetBool( TEXT("DevOptions.Shaders"), TEXT("bLogJobCompletionTimes"), bLogJobCompletionTimes, GEngineIni ));
 	GConfig->GetFloat(TEXT("DevOptions.Shaders"), TEXT("WorkerTimeToLive"), GRegularWorkerTimeToLive, GEngineIni);
 	GConfig->GetFloat(TEXT("DevOptions.Shaders"), TEXT("BuildWorkerTimeToLive"), GBuildWorkerTimeToLive, GEngineIni);
@@ -6721,9 +6722,8 @@ bool FShaderCompilingManager::HandlePotentialRetryOnError(TMap<int32, FShaderMap
 
 				if (UE_LOG_ACTIVE(LogShaders, Log) && (bAnyErrorLikelyToBeCodeError || bPromptToRetryFailedShaderCompiles || bSpecialEngineMaterial))
 				{
-#if UE_BUILD_DEBUG
 					// Use debug break in debug with the debugger attached, otherwise message box
-					if (FPlatformMisc::IsDebuggerPresent())
+					if (bDebugBreakOnPromptToRetryShaderCompile && FPlatformMisc::IsDebuggerPresent())
 					{
 						// A shader compile error has occurred, see the debug output for information.
 						// Double click the errors in the VS.NET output window and the IDE will take you directly to the file and line of the error.
@@ -6734,7 +6734,6 @@ bool FShaderCompilingManager::HandlePotentialRetryOnError(TMap<int32, FShaderMap
 						bRetryCompile = GRetryShaderCompilation;
 					}
 					else
-#endif	//UE_BUILD_DEBUG
 					{
 						if (FPlatformMisc::MessageBoxExt(EAppMsgType::YesNo, *FText::Format(NSLOCTEXT("UnrealEd", "Error_RetryShaderCompilation", "{0}\r\n\r\nRetry compilation?"),
 							FText::FromString(ErrorString)).ToString(), TEXT("Error")) == EAppReturnType::Type::Yes)
diff --git a/Engine/Source/Runtime/Engine/Public/ShaderCompiler.h b/Engine/Source/Runtime/Engine/Public/ShaderCompiler.h
--- a/Engine/Source/Runtime/Engine/Public/ShaderCompiler.h
+++ b/Engine/Source/Runtime/Engine/Public/ShaderCompiler.h
@@ -638,6 +638,8 @@ private:
 	bool bAllowAsynchronousShaderCompiling;
 	/** Whether to ask to retry a failed shader compile error. */
 	bool bPromptToRetryFailedShaderCompiles;
+	/** If enabled when we enter the prompt to retry we will break in the debugger if one is attached rather than prompting. */
+	bool bDebugBreakOnPromptToRetryShaderCompile = false;
 	/** Whether to log out shader job completion times on the worker thread.  Useful for tracking down which global shader is taking a long time. */
 	bool bLogJobCompletionTimes;
 	/** Target execution time for ProcessAsyncResults.  Larger values speed up async shader map processing but cause more hitchiness while async compiling is happening. */
