Fix crash in D3D11 when FMaterialRenderProxy::UpdateDeferredCachedUniformExpressions is using a non-immediate command list.

#rnx

diff --git a/Engine/Source/Runtime/Engine/Private/Materials/MaterialRenderProxy.cpp b/Engine/Source/Runtime/Engine/Private/Materials/MaterialRenderProxy.cpp
--- a/Engine/Source/Runtime/Engine/Private/Materials/MaterialRenderProxy.cpp
+++ b/Engine/Source/Runtime/Engine/Private/Materials/MaterialRenderProxy.cpp
@@ -684,8 +684,6 @@ void FMaterialRenderProxy::UpdateDeferredCachedUniformExpressions(FRHICommandLis
 	{
 		RHICmdListTask = new FRHICommandList(FRHIGPUMask::All());
 		RHICmdListTask->SwitchPipeline(ERHIPipeline::Graphics);
-
-		FRHICommandListImmediate::Get(RHICmdList).QueueAsyncCommandListSubmit(RHICmdListTask);
 	}
 	else
 	{
@@ -743,6 +741,11 @@ void FMaterialRenderProxy::UpdateDeferredCachedUniformExpressions(FRHICommandLis
 	{
 		EvaluateUniformExpressionsLambda();
 	}
+
+	if (bAllowAsyncUpdate)
+	{
+		FRHICommandListImmediate::Get(RHICmdList).QueueAsyncCommandListSubmit(RHICmdListTask);
+	}
 }
 
 void FMaterialRenderProxy::UpdateDeferredCachedUniformExpressions()
diff --git a/Engine/Source/Runtime/Windows/D3D11RHI/Private/D3D11UniformBuffer.cpp b/Engine/Source/Runtime/Windows/D3D11RHI/Private/D3D11UniformBuffer.cpp
--- a/Engine/Source/Runtime/Windows/D3D11RHI/Private/D3D11UniformBuffer.cpp
+++ b/Engine/Source/Runtime/Windows/D3D11RHI/Private/D3D11UniformBuffer.cpp
@@ -187,7 +187,7 @@ FUniformBufferRHIRef FD3D11DynamicRHI::RHICreateUniformBuffer(const void* Conten
 		SCOPE_CYCLE_COUNTER(STAT_D3D11UpdateUniformBufferTime);
 
 PRAGMA_DISABLE_DEPRECATION_WARNINGS
-		if (IsPoolingEnabled() && (IsInRenderingThread() || IsInRHIThread()) && !UE::Tasks::Private::IsThreadRetractingTask())
+		if (IsPoolingEnabled() && (IsInActualRenderingThread() || IsInRHIThread()) && !UE::Tasks::Private::IsThreadRetractingTask())
 PRAGMA_ENABLE_DEPRECATION_WARNINGS
 		{
 			const bool bAllocatedFromPool = true;
