Optimizing ComponentSweepMultiByChannel by moving the TArray of hits outside the for loop and calling Reset within it to prevent as many allocations and frees as possible, and moving the Sort on the Hits to happen just before we exit the function, instead of within the for loops so we only have to sort each element once.

diff --git a/Engine/Source/Runtime/Engine/Private/Collision/WorldCollision.cpp b/Engine/Source/Runtime/Engine/Private/Collision/WorldCollision.cpp
--- a/Engine/Source/Runtime/Engine/Private/Collision/WorldCollision.cpp
+++ b/Engine/Source/Runtime/Engine/Private/Collision/WorldCollision.cpp
@@ -460,7 +460,7 @@ bool UWorld::ComponentSweepMultiByChannel(TArray<struct FHitResult>& OutHits, cl
 
 	const FTransform WorldToComponent = PrimComp->GetComponentToWorld().Inverse();
 
-	TArray<struct FOverlapResult> TempOverlaps;
+	TArray<FHitResult> TempHits;
 	for (Chaos::FPhysicsObjectHandle Object : PhysicsObjects)
 	{
 		FComponentQueryParams ParamsCopy{ Params };
@@ -498,16 +498,16 @@ bool UWorld::ComponentSweepMultiByChannel(TArray<struct FHitResult>& OutHits, cl
 			}
 
 			FPhysicsGeometryCollection GeomCollection = FPhysicsInterface::GetGeometryCollection(ShapeHandle);
-			TArray<FHitResult> TmpHits;
-			if (FPhysicsInterface::GeomSweepMulti(this, GeomCollection, Rot, TmpHits, Start, End, TraceChannel, ParamsCopy, FCollisionResponseParams(PrimComp->GetCollisionResponseToChannels())))
+			TempHits.Reset();
+			if (FPhysicsInterface::GeomSweepMulti(this, GeomCollection, Rot, TempHits, Start, End, TraceChannel, ParamsCopy, FCollisionResponseParams(PrimComp->GetCollisionResponseToChannels())))
 			{
 				bHaveBlockingHit = true;
 			}
-			OutHits.Append(TmpHits);	//todo: should these be made unique?
-			OutHits.Sort(FCompareFHitResultTime());
+			OutHits.Append(TempHits);	//todo: should these be made unique?
 		}
 	}
 
+	OutHits.Sort(FCompareFHitResultTime());
 	return bHaveBlockingHit;
 }
 
