cluster union modification of geometry needs to grab a write scene lock on the GT

diff --git a/Engine/Source/Runtime/Engine/Private/PhysicsEngine/ClusterUnionComponent.cpp b/Engine/Source/Runtime/Engine/Private/PhysicsEngine/ClusterUnionComponent.cpp
--- a/Engine/Source/Runtime/Engine/Private/PhysicsEngine/ClusterUnionComponent.cpp
+++ b/Engine/Source/Runtime/Engine/Private/PhysicsEngine/ClusterUnionComponent.cpp
@@ -238,7 +238,7 @@ void UClusterUnionComponent::RemoveComponentFromCluster(UPrimitiveComponent* InC
 
 	TSet<Chaos::FPhysicsObjectHandle> PhysicsObjectsToRemove;
 
-	FLockedReadPhysicsObjectExternalInterface Interface = FPhysicsObjectExternalInterface::LockRead(GetWorld()->GetPhysicsScene());
+	FLockedWritePhysicsObjectExternalInterface Interface = FPhysicsObjectExternalInterface::LockWrite(GetWorld()->GetPhysicsScene());
 
 	if (FClusteredComponentData* ComponentData = PerComponentData.Find(InComponent))
 	{
@@ -398,7 +398,7 @@ void UClusterUnionComponent::AddGTParticleGeometry(const TArray<Chaos::FPhysicsO
 	}
 
 	// This is an assumption that all the physics objects are part of this scene and thus this is the right thing to lock.
-	FLockedReadPhysicsObjectExternalInterface Interface = FPhysicsObjectExternalInterface::LockRead(GetWorld()->GetPhysicsScene());
+	FLockedWritePhysicsObjectExternalInterface Interface = FPhysicsObjectExternalInterface::LockWrite(GetWorld()->GetPhysicsScene());
 
 	TArray<Chaos::FImplicitObjectPtr> ImplicitObjects;
 	TArray<Chaos::FPBDRigidParticle*> ShapeParticles;
@@ -453,7 +453,7 @@ void UClusterUnionComponent::RemoveGTParticleGeometry(const TSet<Chaos::FPhysics
 	}
 
 	// This is an assumption that all the physics objects are part of this scene and thus this is the right thing to lock.
-	FLockedReadPhysicsObjectExternalInterface Interface = FPhysicsObjectExternalInterface::LockRead(GetWorld()->GetPhysicsScene());
+	FLockedWritePhysicsObjectExternalInterface Interface = FPhysicsObjectExternalInterface::LockWrite(GetWorld()->GetPhysicsScene());
 
 	TArray<Chaos::FPBDRigidParticle*> ShapeParticles;
 	for (Chaos::FPhysicsObjectHandle PhysicsObject : PhysicsObjects)
