Undoing CLs 28079915, 28081103 and 2827483 to restore PerfB replay deserialization.

#rnx

diff --git a/Engine/Plugins/Runtime/GameplayAbilities/Source/GameplayAbilities/Private/Abilities/GameplayAbilityTargetActor_SingleLineTrace.cpp b/Engine/Plugins/Runtime/GameplayAbilities/Source/GameplayAbilities/Private/Abilities/GameplayAbilityTargetActor_SingleLineTrace.cpp
--- a/Engine/Plugins/Runtime/GameplayAbilities/Source/GameplayAbilities/Private/Abilities/GameplayAbilityTargetActor_SingleLineTrace.cpp
+++ b/Engine/Plugins/Runtime/GameplayAbilities/Source/GameplayAbilities/Private/Abilities/GameplayAbilityTargetActor_SingleLineTrace.cpp
@@ -2,6 +2,7 @@
 
 #include "Abilities/GameplayAbilityTargetActor_SingleLineTrace.h"
 #include "DrawDebugHelpers.h"
+#include "GameFramework/LightWeightInstanceSubsystem.h"
 
 #include UE_INLINE_GENERATED_CPP_BY_NAME(GameplayAbilityTargetActor_SingleLineTrace)
 
@@ -43,9 +44,7 @@ FHitResult AGameplayAbilityTargetActor_SingleLineTrace::PerformTrace(AActor* InS
 	if (AGameplayAbilityWorldReticle* LocalReticleActor = ReticleActor.Get())
 	{
 		const bool bHitActor = (ReturnHitResult.bBlockingHit && (ReturnHitResult.HitObjectHandle.IsValid()));
-		const FVector ReticleLocation = (bHitActor && LocalReticleActor->bSnapToTargetedActor) 
-			? ReturnHitResult.HitObjectHandle.GetLocation()
-			: ReturnHitResult.Location;
+		const FVector ReticleLocation = (bHitActor && LocalReticleActor->bSnapToTargetedActor) ? FLightWeightInstanceSubsystem::Get().GetLocation(ReturnHitResult.HitObjectHandle) : ReturnHitResult.Location;
 
 		LocalReticleActor->SetActorLocation(ReticleLocation);
 		LocalReticleActor->SetIsTargetAnActor(bHitActor);
diff --git a/Engine/Source/Runtime/Engine/Classes/Engine/ActorInstanceHandle.h b/Engine/Source/Runtime/Engine/Classes/Engine/ActorInstanceHandle.h
--- a/Engine/Source/Runtime/Engine/Classes/Engine/ActorInstanceHandle.h
+++ b/Engine/Source/Runtime/Engine/Classes/Engine/ActorInstanceHandle.h
@@ -2,35 +2,25 @@
 
 #pragma once
 
-#include "UObject/WeakInterfacePtr.h"
+#include "CoreMinimal.h"
 #include "ActorInstanceHandle.generated.h"
 
-class USceneComponent;
-class AActor;
-class UActorInstanceManager;
-class IActorInstanceManagerInterface;
-class ULevel;
-
-using FActorInstanceManagerInterface = TWeakInterfacePtr<IActorInstanceManagerInterface>;
-
-// Handle to a unique object. This may specify a full weigh actor or it may only specify the actor instance that represents the same object.
+// Handle to a unique object. This may specify a full weigh actor or it may only specify the light weight instance that represents the same object.
 USTRUCT(BlueprintType)
 struct FActorInstanceHandle
 {
 	GENERATED_BODY()
 
+	friend struct FLightWeightInstanceSubsystem;
+	friend class ALightWeightInstanceManager;
+	friend class UActorInstanceHandleInterface;
+
 	ENGINE_API FActorInstanceHandle();
 
 	ENGINE_API explicit FActorInstanceHandle(AActor* InActor);
-	ENGINE_API FActorInstanceHandle(UObject* InManager, const UPrimitiveComponent* RelevantComponent, int32 CollisionInstanceIndex);
-	ENGINE_API FActorInstanceHandle(FActorInstanceManagerInterface InManagerInterface, int32 InstanceIndex);
-	ENGINE_API FActorInstanceHandle(const FActorInstanceHandle& Other);
+	ENGINE_API explicit FActorInstanceHandle(class ALightWeightInstanceManager* Manager, int32 InInstanceIndex);
 
-	/** 
-	 * A path dedicated to creation of handles while converting actor to a dehydrated representation. This path ensures
-	 * an actor won't be spawned as a side effect of looking for the actor given Manager/Index represents
-	 */
-	static FActorInstanceHandle MakeDehydratedActorHandle(UObject& Manager, int32 InInstanceIndex);
+	ENGINE_API FActorInstanceHandle(const FActorInstanceHandle& Other);
 
 	ENGINE_API bool IsValid() const;
 
@@ -40,7 +30,7 @@ struct FActorInstanceHandle
 	bool DoesRepresent() const;
 
 	ENGINE_API UClass* GetRepresentedClass() const;
-	ENGINE_API ULevel* GetLevel() const;
+
 	ENGINE_API FVector GetLocation() const;
 	ENGINE_API FRotator GetRotation() const;
 	ENGINE_API FTransform GetTransform() const;
@@ -56,16 +46,15 @@ struct FActorInstanceHandle
 
 	/** Returns the actor specified by this handle. This may require loading and creating the actor object. */
 	ENGINE_API AActor* FetchActor() const;
-
 	template <typename T>
 	T* FetchActor() const;
 
-	AActor* GetCachedActor() const { return Actor.Get(); }
-	ENGINE_API void SetCachedActor(AActor* InActor) const;
-
 	/* Returns the index used internally by the manager */
 	FORCEINLINE int32 GetInstanceIndex() const { return InstanceIndex; }
 
+	/* Returns the index used by rendering and collision */
+	ENGINE_API int32 GetRenderingInstanceIndex() const;
+
 	FActorInstanceHandle& operator=(const FActorInstanceHandle& Other) = default;
 	FActorInstanceHandle& operator=(FActorInstanceHandle&& Other) = default;
 	ENGINE_API FActorInstanceHandle& operator=(AActor* OtherActor);
@@ -80,17 +69,9 @@ struct FActorInstanceHandle
 
 	friend ENGINE_API FArchive& operator<<(FArchive& Ar, FActorInstanceHandle& Handle);
 
-	FActorInstanceManagerInterface GetManagerInterface() const { return ManagerInterface; }
-
-	template<typename T>
-	T* GetManager() const 
-	{
-		return Cast<T>(ManagerInterface.GetObject());
-	}
+	uint32 GetInstanceUID() const { return InstanceUID; }
 
 private:
-	void SetInternal(IActorInstanceManagerInterface& InManagerInterface, const UPrimitiveComponent* RelevantComponent, int32 CollisionInstanceIndex);
-
 	/**
 	 * helper functions that let us treat the actor pointer as a UObject in templated functions
 	 * these do NOT fetch the actor so they will return nullptr if we don't have a full actor representation
@@ -103,13 +84,16 @@ private:
 
 	/** this is cached here for convenience */
 	UPROPERTY()
-	mutable TWeakObjectPtr<AActor> Actor;
+		mutable TWeakObjectPtr<AActor> Actor;
 
-	/** Identifies the actor instance manager to use */
-	FActorInstanceManagerInterface ManagerInterface;
+	/** Identifies the light weight instance manager to use */
+	TWeakObjectPtr<ALightWeightInstanceManager> Manager;
 
 	/** Identifies the instance within the manager */
 	int32 InstanceIndex;
+
+	/** Unique identifier for instances represented by the handle */
+	uint32 InstanceUID;
 };
 
 template<typename T>
diff --git a/Engine/Source/Runtime/Engine/Classes/Engine/ActorInstanceManagerInterface.h b/Engine/Source/Runtime/Engine/Classes/Engine/ActorInstanceManagerInterface.h
deleted file mode 100644
diff --git a/Engine/Source/Runtime/Engine/Classes/Engine/World.h b/Engine/Source/Runtime/Engine/Classes/Engine/World.h
--- a/Engine/Source/Runtime/Engine/Classes/Engine/World.h
+++ b/Engine/Source/Runtime/Engine/Classes/Engine/World.h
@@ -1510,6 +1510,10 @@ public:
 	UPROPERTY(Transient)
 	TObjectPtr<class UPhysicsFieldComponent> PhysicsField;
 
+	/** Tracks the last assigned unique id for light weight instances in this world. */
+	UPROPERTY()
+	uint32 LWILastAssignedUID;
+
 private:
 
 	/** Array of components that need to wait on tasks before end of frame updates */
diff --git a/Engine/Source/Runtime/Engine/Classes/GameFramework/LightWeightInstanceManager.h b/Engine/Source/Runtime/Engine/Classes/GameFramework/LightWeightInstanceManager.h
--- a/Engine/Source/Runtime/Engine/Classes/GameFramework/LightWeightInstanceManager.h
+++ b/Engine/Source/Runtime/Engine/Classes/GameFramework/LightWeightInstanceManager.h
@@ -11,14 +11,13 @@
 #include "Engine/EngineBaseTypes.h"
 #include "Templates/SharedPointer.h"
 
-#include "GameFramework/Actor.h"
-#include "Engine/ActorInstanceManagerInterface.h"
+#include "Actor.h"
 
 #include "LightWeightInstanceManager.generated.h"
 
 
-class AActor;
 struct FActorSpawnParameters;
+//DECLARE_DYNAMIC_DELEGATE_RetVal_OneParam(FActorInstanceHandle, FOnActorReady, FActorInstanceHandle, InHandle);
 
 // Used for initializing light weight instances.
 struct FLWIData
@@ -27,8 +26,31 @@ struct FLWIData
 };
 
 
+// Base class for interfaces for each handle
+UCLASS(Experimental, MinimalAPI)
+class UActorInstanceHandleInterface : public UObject
+{
+	GENERATED_UCLASS_BODY()
+
+	friend class ALightWeightInstanceManager;
+
+	// Returns a pointer to the actor for this handle. This will create a new copy of the actor if it was being stored as a light weight instance.
+	template<typename T>
+	T* FetchActor();
+
+protected:
+
+	// Handle to the actor or instance that is associated with this interface
+	FActorInstanceHandle Handle;
+
+	// A cached pointer to the manager associated with Handle so we don't need to find it every time
+	ALightWeightInstanceManager* Manager;
+};
+
+
+
 UCLASS(BlueprintType, Blueprintable, Experimental, MinimalAPI)
-class ALightWeightInstanceManager : public AActor, public IActorInstanceManagerInterface
+class ALightWeightInstanceManager : public AActor
 {
 	GENERATED_UCLASS_BODY()
 
@@ -47,6 +69,9 @@ public:
 	// Returns the rotation of the instance specified by Handle
 	ENGINE_API FRotator GetRotation(const FActorInstanceHandle& Handle) const;
 
+	// Returns the transform of the instance specified by Handle
+	ENGINE_API FTransform GetTransform(const FActorInstanceHandle& Handle) const;
+
 	// Returns the name of the instance specified by Handle
 	ENGINE_API FString GetName(const FActorInstanceHandle& Handle) const;
 
@@ -56,18 +81,30 @@ public:
 	// Returns true if this manager is capable of representing objects of type OtherClass
 	ENGINE_API bool DoesAcceptClass(const UClass* OtherClass) const;
 
+	// Returns the specific class that this manages
+	ENGINE_API UClass* GetRepresentedClass() const;
+
 	// Returns the base class of types that this can manage
 	ENGINE_API UClass* GetAcceptedClass() const;
 
 	// Sets the specific class that this manages
 	ENGINE_API virtual void SetRepresentedClass(UClass* ActorClass);
 
+	// Returns the actor associated with Handle if one exists
+	ENGINE_API AActor* FetchActorFromHandle(const FActorInstanceHandle& Handle);
+
 	// Returns the index of the light weight instance associated with InActor if one exists; otherwise we return INDEX_NONE
 	ENGINE_API int32 FindIndexForActor(const AActor* InActor) const;
 
 	// Returns a handle to a light weight instance representing the same object as InActor and calls destroy on InActor if successful.
 	ENGINE_API FActorInstanceHandle ConvertActorToLightWeightInstance(AActor* InActor);
 
+	// Returns the index used internally by the light weight instance manager that is associated with the instance referred to by InIndex used by collision and rendering
+	ENGINE_API virtual int32 ConvertCollisionIndexToLightWeightIndex(int32 InIndex) const;
+
+	// Returns the index used by collision and rendering that is associated with the instance referred to by InIndex
+	ENGINE_API virtual int32 ConvertLightWeightIndexToCollisionIndex(int32 InIndex) const;
+
 	// LWI grid size to use for this manager.
 	ENGINE_API virtual int32 GetGridSize() const;
 
@@ -77,23 +114,42 @@ public:
 	// Helper that retrieve the world space bounds of the LWI grid cell encompassing the provided coordinate.
 	ENGINE_API FBox ConvertPositionToGridBounds(const FVector& InPosition) const;
 
-	ENGINE_API bool HasAnyValidInstancesOrManagedActors() const;
+	template<typename U>
+	bool IsInterfaceSupported() const
+	{
+		const UClass* const InterfaceClass = GetInterfaceClass();
+		return InterfaceClass && InterfaceClass->ImplementsInterface(U::StaticClass());
+	}
+
+	template<typename I>
+	I* FetchInterfaceObject(const FActorInstanceHandle& Handle)
+	{
+		// if we have a valid actor, use it
+		if (Handle.Actor.IsValid())
+		{
+			return Cast<I>(Handle.Actor.Get());
+		}
+
+		// TODO: once we have a better idea of how we're going to use this we should add some kind of caching so we don't always need to create a new uobject.
+		// if we can support the interface using just the handle then do that
+		if (I* InterfaceObj = Cast<I>(CreateInterfaceObject(Handle)))
+		{
+			return InterfaceObj;
+		}
+
+		// fallback to creating the actor
+		return Cast<I>(FetchActorFromHandle(Handle));
+	}
+
+	ENGINE_API AActor* FindActorForInstanceIndex(const int32 InstanceIndex);
 
-	// IActorInstanceManagerInterface begin
-	ENGINE_API virtual int32 ConvertCollisionIndexToInstanceIndex(int32 InIndex, const UPrimitiveComponent* RelevantComponent) const;
-	ENGINE_API virtual AActor* FindActor(const FActorInstanceHandle& Handle) override;
-	ENGINE_API virtual AActor* FindOrCreateActor(const FActorInstanceHandle& Handle) override;
-	ENGINE_API virtual UClass* GetRepresentedClass(const int32 InstanceIndex) const override { return GetRepresentedClassInternal(); }
-	ENGINE_API virtual ULevel* GetLevelForInstance(const int32 InstanceIndex) const override { return GetLevel(); }
-	ENGINE_API virtual FTransform GetTransform(const FActorInstanceHandle& Handle) const;
-	// IActorInstanceManagerInterface end
+	ENGINE_API bool HasAnyValidInstancesOrManagedActors() const;
 
 protected:
-	UClass* GetRepresentedClassInternal() const;
-
 	// Creates an actor to replace the instance specified by Handle
 	ENGINE_API AActor* ConvertInstanceToActor(const FActorInstanceHandle& Handle);
 
+	// Takes a polymorphic struct to set the initial data for a new instance
 	ENGINE_API int32 AddNewInstance(FLWIData* InitData);
 
 	// Adds a new instance at the specified index. This function should only be called by AddNewInstance.
@@ -124,6 +180,18 @@ protected:
 	// Called after a spawned actor is destroyed
 	ENGINE_API virtual void OnSpawnedActorDestroyed(AActor* DestroyedActor, const int32 DestroyedActorInstanceIndex);
 
+	// Create an object that implements interfaces for the light weight instance specified by Handle
+	UObject* CreateInterfaceObject(const FActorInstanceHandle& Handle)
+	{
+		UActorInstanceHandleInterface* InterfaceObj = NewObject<UActorInstanceHandleInterface>(this, GetInterfaceClass());
+		InterfaceObj->Handle = Handle;
+		InterfaceObj->Manager = this;
+		return InterfaceObj;
+	}
+
+	// Gets the class that implements interfaces for light weight instances owned by this manager
+	ENGINE_API virtual UClass* GetInterfaceClass() const;
+
 	// Helper functions for converting between our internal storage indices and the indices used by external bookkeeping
 	virtual int32 ConvertInternalIndexToHandleIndex(int32 InInternalIndex) const { return InInternalIndex; }
 	virtual int32 ConvertHandleIndexToInternalIndex(int32 InHandleIndex) const { return InHandleIndex; }
diff --git a/Engine/Source/Runtime/Engine/Classes/GameFramework/LightWeightInstanceStaticMeshManager.h b/Engine/Source/Runtime/Engine/Classes/GameFramework/LightWeightInstanceStaticMeshManager.h
--- a/Engine/Source/Runtime/Engine/Classes/GameFramework/LightWeightInstanceStaticMeshManager.h
+++ b/Engine/Source/Runtime/Engine/Classes/GameFramework/LightWeightInstanceStaticMeshManager.h
@@ -4,7 +4,7 @@
 
 #include "CoreMinimal.h"
 #include "Engine/StaticMesh.h"
-#include "LightWeightInstanceManager.h"
+#include "GameFramework/LightWeightInstanceManager.h"
 #include "Elements/SMInstance/SMInstanceManager.h"
 
 #include "LightWeightInstanceStaticMeshManager.generated.h"
@@ -24,7 +24,9 @@ class ALightWeightInstanceStaticMeshManager : public ALightWeightInstanceManager
 	// Clears the static mesh used for rendering instances
 	ENGINE_API void ClearStaticMesh();
 
-	ENGINE_API virtual int32 ConvertCollisionIndexToInstanceIndex(int32 InIndex, const UPrimitiveComponent* RelevantComponent) const override;
+	ENGINE_API virtual int32 ConvertCollisionIndexToLightWeightIndex(int32 InIndex) const override;
+
+	ENGINE_API virtual int32 ConvertLightWeightIndexToCollisionIndex(int32 InIndex) const override;
 
 protected:
 	ENGINE_API virtual void AddNewInstanceAt(FLWIData* InitData, int32 Index) override;
diff --git a/Engine/Source/Runtime/Engine/Classes/GameFramework/LightWeightInstanceSubsystem.h b/Engine/Source/Runtime/Engine/Classes/GameFramework/LightWeightInstanceSubsystem.h
--- a/Engine/Source/Runtime/Engine/Classes/GameFramework/LightWeightInstanceSubsystem.h
+++ b/Engine/Source/Runtime/Engine/Classes/GameFramework/LightWeightInstanceSubsystem.h
@@ -6,6 +6,7 @@
 #include "GameFramework/LightWeightInstanceManager.h"
 #include "Templates/SharedPointer.h"
 
+#include "Actor.h"
 
 ENGINE_API DECLARE_LOG_CATEGORY_EXTERN(LogLightWeightInstance, Log, All);
 
@@ -48,6 +49,9 @@ struct FLightWeightInstanceSubsystem
 	// Returns the instance manager that handles instances of type Class that live in Level
 	ENGINE_API ALightWeightInstanceManager* FindOrAddLightWeightInstanceManager(UClass& ActorClass, UWorld& World, const FVector& InPos, const UDataLayerInstance* DataLayer = nullptr);
 
+	// Returns the actor specified by Handle. This may require loading and creating the actor object.
+	ENGINE_API AActor* FetchActor(const FActorInstanceHandle& Handle);
+
 	// Returns the actor specified by Handle if it exists. Returns nullptr if it doesn't
 	ENGINE_API AActor* GetActor_NoCreate(const FActorInstanceHandle& Handle) const;
 
@@ -69,6 +73,30 @@ struct FLightWeightInstanceSubsystem
 	// deletes the instance identified by Handle
 	ENGINE_API void DeleteInstance(const FActorInstanceHandle& Handle);
 
+	// Returns true if the handle can return an object that implements the interface U
+	template<typename U>
+	bool IsInterfaceSupported(const FActorInstanceHandle& Handle) const
+	{
+		if (ALightWeightInstanceManager* InstanceManager = FindLightWeightInstanceManager(Handle))
+		{
+			return InstanceManager->IsInterfaceSupported<U>();
+		}
+
+		return false;
+	}
+
+	// Returns an object that implements the interface I for Handle
+	template<typename I>
+	I* FetchInterfaceObject(const FActorInstanceHandle& Handle)
+	{
+		if (ALightWeightInstanceManager* InstanceManager = FindLightWeightInstanceManager(Handle))
+		{
+			return InstanceManager->FetchInterfaceObject<I>(Handle);
+		}
+
+		return nullptr;
+	}
+
 	// Helper that converts a position (world space) into a coordinate for the LWI grid.
 	UE_DEPRECATED(5.3, "Use LWI Managers version of ConvertPositionToCoord()")
 	static ENGINE_API FInt32Vector3 ConvertPositionToCoord(const FVector& InPosition);
diff --git a/Engine/Source/Runtime/Engine/Classes/GameFramework/LightWeightInstanceUtils.h b/Engine/Source/Runtime/Engine/Classes/GameFramework/LightWeightInstanceUtils.h
new file mode 100644
--- /dev/null
+++ b/Engine/Source/Runtime/Engine/Classes/GameFramework/LightWeightInstanceUtils.h
@@ -0,0 +1,56 @@
+// Copyright Epic Games, Inc. All Rights Reserved.
+
+#pragma once
+
+// HEADER_UNIT_UNSUPPORTED - I don't think this header is used. dll export of templated function below indicates that.
+
+#include "Kismet/BlueprintFunctionLibrary.h"
+#include "Engine/ActorInstanceHandle.h"
+#include "Engine/EngineTypes.h"
+#include "GameFramework/LightWeightInstanceSubsystem.h"
+
+#include "LightWeightInstanceBlueprintFunctionLibrary.generated.h"
+
+namespace LWIUtils
+{
+	/**
+	 * Returns true if the object this handle represents supports the interface of type U
+	 */
+	template<typename U>
+	static bool DoesHandleSupportInterface(const FActorInstanceHandle& Handle);
+
+	/**
+	 * Returns an object implementing the interface I.
+	 */
+	template<typename I>
+	static I* FetchInterfaceFromHandle(const FActorInstanceHandle& Handle);
+
+
+	/**
+	 * Implementations of inline or template functions go below here
+	 */
+
+	template<typename U>
+	static bool DoesHandleSupportInterface(const FActorInstanceHandle& Handle)
+	{
+		// if we have a valid actor, see if it supports the interface
+		if (const UClass* ObjClass = Handle.GetRepresentedClass())
+		{
+			return ObjClass->ImplementsInterface(U::StaticClass());
+		}
+
+		// no valid actor, ask the instance manager instead
+		return FLightWeightInstanceSubsystem::Get().IsInterfaceSupported<U>(Handle);
+	}
+
+	template<typename I>
+	static I* FetchInterfaceFromHandle(const FActorInstanceHandle& Handle)
+	{
+		if (AActor* Actor = Handle.FetchActor())
+		{
+			return Cast<I>(Actor);
+		}
+
+		return FLightWeightInstanceSubsystem::Get().FetchInterfaceObject<I>(Handle);
+	}
+};
\ No newline at end of file
diff --git a/Engine/Source/Runtime/Engine/Private/Collision/CollisionConversions.cpp b/Engine/Source/Runtime/Engine/Private/Collision/CollisionConversions.cpp
--- a/Engine/Source/Runtime/Engine/Private/Collision/CollisionConversions.cpp
+++ b/Engine/Source/Runtime/Engine/Private/Collision/CollisionConversions.cpp
@@ -14,7 +14,7 @@
 #include "Physics/Experimental/ChaosInterfaceWrapper.h"
 
 #include "PhysicsEngine/CollisionQueryFilterCallback.h"
-#include "Engine/ActorInstanceManagerInterface.h"
+#include "GameFramework/LightWeightInstanceManager.h"
 #include "PhysicsProxy/SingleParticlePhysicsProxy.h"
 #include "PhysicsProxy/GeometryCollectionPhysicsProxy.h"
 
@@ -186,7 +186,14 @@ static void SetHitResultFromShapeAndFaceIndex(const FPhysicsShape& Shape,  const
 	{
 		OutResult.Component = OwningComponent;
 		AActor* Owner = OwningComponent->GetOwner();
-		OutResult.HitObjectHandle = FActorInstanceHandle(Owner);
+		if (ALightWeightInstanceManager* LWIManager = Cast<ALightWeightInstanceManager>(Owner))
+		{
+			OutResult.HitObjectHandle = FActorInstanceHandle(LWIManager, OutResult.Item);
+		}
+		else
+		{
+			OutResult.HitObjectHandle = FActorInstanceHandle(OwningComponent->GetOwner());
+		}
 
 		if (bReturnPhysMat)
 		{
@@ -603,10 +610,9 @@ void ConvertQueryOverlap(const FPhysicsShape& Shape, const FPhysicsActor& Actor,
         BodyInst = FPhysicsInterface::ShapeToOriginalBodyInstance(BodyInst, &Shape);
 		if (const UPrimitiveComponent* OwnerComponent = BodyInst->OwnerComponent.Get())
 		{
-			UObject* Outer = OwnerComponent->GetOuter();
-			if (Outer && Outer->Implements<UActorInstanceManagerInterface>())
+			if (ALightWeightInstanceManager* LWIManager = Cast<ALightWeightInstanceManager>(OwnerComponent->GetOwner()))
 			{
-				OutOverlap.OverlapObjectHandle = FActorInstanceHandle(Outer, OwnerComponent, BodyInst->InstanceBodyIndex);
+				OutOverlap.OverlapObjectHandle = FActorInstanceHandle(LWIManager, BodyInst->InstanceBodyIndex);
 			}
 			else
 			{
diff --git a/Engine/Source/Runtime/Engine/Private/DebugCameraHUD.cpp b/Engine/Source/Runtime/Engine/Private/DebugCameraHUD.cpp
--- a/Engine/Source/Runtime/Engine/Private/DebugCameraHUD.cpp
+++ b/Engine/Source/Runtime/Engine/Private/DebugCameraHUD.cpp
@@ -12,6 +12,7 @@
 #include "Engine/World.h"
 #include "GameFramework/PawnMovementComponent.h"
 #include "GameFramework/SpectatorPawn.h"
+#include "GameFramework/LightWeightInstanceSubsystem.h"
 #include "Engine/Canvas.h"
 #include "Engine/DebugCameraController.h"
 #include "Materials/MaterialInterface.h"
@@ -154,7 +155,7 @@ void ADebugCameraHUD::PostRender()
 				yl += Y;
 				Canvas->DrawText(RenderFont, FString::Printf(TEXT("HitDist: %f"), Hit.Distance), X, yl, 1.f, 1.f, FontRenderInfo);
 				yl += Y;
-				Canvas->DrawText(RenderFont, FString::Printf(TEXT("HitObject: '%s'"), *HitHandle.GetName()), X, yl, 1.f, 1.f, FontRenderInfo);
+				Canvas->DrawText(RenderFont, FString::Printf(TEXT("HitObject: '%s'"), *FLightWeightInstanceSubsystem::Get().GetName(HitHandle)), X, yl, 1.f, 1.f, FontRenderInfo);
 				yl += Y;
 				Canvas->DrawText(RenderFont, FString::Printf(TEXT("HitComponent: '%s'"), Hit.Component.Get() ? *Hit.Component.Get()->GetFName().ToString() : TEXT("<NULL>")), X, yl, 1.f, 1.f, FontRenderInfo);
 				yl += Y;
diff --git a/Engine/Source/Runtime/Engine/Private/Engine/ActorInstanceHandle.cpp b/Engine/Source/Runtime/Engine/Private/Engine/ActorInstanceHandle.cpp
--- a/Engine/Source/Runtime/Engine/Private/Engine/ActorInstanceHandle.cpp
+++ b/Engine/Source/Runtime/Engine/Private/Engine/ActorInstanceHandle.cpp
@@ -1,19 +1,15 @@
 // Copyright Epic Games, Inc. All Rights Reserved.
 
 #include "Engine/ActorInstanceHandle.h"
+#include "GameFramework/LightWeightInstanceSubsystem.h"
 #include "Engine/World.h"
-#include "Engine/ActorInstanceManagerInterface.h"
-#include "GameFramework/Actor.h"
 
 #include UE_INLINE_GENERATED_CPP_BY_NAME(ActorInstanceHandle)
 
-
-//-----------------------------------------------------------------------------
-// FActorInstanceHandle
-//-----------------------------------------------------------------------------
 FActorInstanceHandle::FActorInstanceHandle()
 	: Actor(nullptr)
 	, InstanceIndex(INDEX_NONE)
+	, InstanceUID(0)
 {
 	// do nothing
 }
@@ -21,28 +17,26 @@ FActorInstanceHandle::FActorInstanceHandle()
 FActorInstanceHandle::FActorInstanceHandle(AActor* InActor)
 	: Actor(InActor)
 	, InstanceIndex(INDEX_NONE)
+	, InstanceUID(0)
 {
 }
 
-FActorInstanceHandle::FActorInstanceHandle(UObject* InManager, const UPrimitiveComponent* RelevantComponent, int32 CollisionInstanceIndex)
-	: Actor(nullptr)
-	, ManagerInterface(InManager)
-	, InstanceIndex(CollisionInstanceIndex)
-{
-	if (IActorInstanceManagerInterface* ManagerInterfacePtr = ManagerInterface.Get())
-	{
-		SetInternal(*ManagerInterfacePtr, RelevantComponent, CollisionInstanceIndex);
-	}
-}
-
-FActorInstanceHandle::FActorInstanceHandle(FActorInstanceManagerInterface InManagerInterface, int32 CollisionInstanceIndex)
+FActorInstanceHandle::FActorInstanceHandle(ALightWeightInstanceManager* InManager, int32 InInstanceIndex)
 	: Actor(nullptr)
-	, ManagerInterface(InManagerInterface)
-	, InstanceIndex(CollisionInstanceIndex)
+	, InstanceIndex(InInstanceIndex)
+	, InstanceUID(0)
 {
-	if (IActorInstanceManagerInterface* ManagerInterfacePtr = ManagerInterface.Get())
+	Manager = InManager;
+	if (Manager.IsValid())
 	{
-		SetInternal(*ManagerInterfacePtr, /*RelevantComponent=*/nullptr, CollisionInstanceIndex);
+		InstanceIndex = Manager->ConvertCollisionIndexToLightWeightIndex(InInstanceIndex);
+		Actor = Manager->FindActorForInstanceIndex(InstanceIndex);
+
+		UWorld* World = Manager->GetWorld();
+		if (ensure(World))
+		{
+			InstanceUID = World->LWILastAssignedUID++;
+		}
 	}
 }
 
@@ -50,30 +44,14 @@ FActorInstanceHandle::FActorInstanceHandle(const FActorInstanceHandle& Other)
 {
 	Actor = Other.Actor;
 
-	ManagerInterface = Other.ManagerInterface;
+	Manager = Other.Manager;
 	InstanceIndex = Other.InstanceIndex;
-}
-
-FActorInstanceHandle FActorInstanceHandle::MakeDehydratedActorHandle(UObject& Manager, int32 InInstanceIndex)
-{
-	FActorInstanceHandle ReturnHandle;
-	ReturnHandle.ManagerInterface = FActorInstanceManagerInterface(&Manager);
-	ReturnHandle.InstanceIndex = InInstanceIndex;
-
-	return ReturnHandle;
-}
-
-void FActorInstanceHandle::SetInternal(IActorInstanceManagerInterface& InManagerInterface, const UPrimitiveComponent* RelevantComponent, int32 CollisionInstanceIndex)
-{
-	check(ManagerInterface.Get() == &InManagerInterface);
-
-	InstanceIndex = InManagerInterface.ConvertCollisionIndexToInstanceIndex(CollisionInstanceIndex, RelevantComponent);
-	Actor = InManagerInterface.FindActor(*this);
+	InstanceUID = Other.InstanceUID;
 }
 
 bool FActorInstanceHandle::IsValid() const
 {
-	return (ManagerInterface.IsValid() && InstanceIndex != INDEX_NONE) || IsActorValid();
+	return (Manager.IsValid() && InstanceIndex != INDEX_NONE) || IsActorValid();
 }
 
 bool FActorInstanceHandle::DoesRepresentClass(const UClass* OtherClass) const
@@ -87,28 +65,22 @@ bool FActorInstanceHandle::DoesRepresentClass(const UClass* OtherClass) const
 
 UClass* FActorInstanceHandle::GetRepresentedClass() const
 {
+	if (!IsValid())
+	{
+		return nullptr;
+	}
+
 	if (IsActorValid())
 	{
 		return Actor->GetClass();
 	}
 
-	IActorInstanceManagerInterface* ManagerInterfacePtr = ManagerInterface.Get();
-	return ManagerInterfacePtr && (InstanceIndex != INDEX_NONE)
-		? ManagerInterfacePtr->GetRepresentedClass(InstanceIndex)
-		: nullptr;
-}
-
-ULevel* FActorInstanceHandle::GetLevel() const
-{
-	if (IsActorValid())
+	if (Manager.IsValid())
 	{
-		return Actor->GetLevel();
+		return Manager->GetRepresentedClass();
 	}
 
-	IActorInstanceManagerInterface* ManagerInterfacePtr = ManagerInterface.Get();
-	return ManagerInterfacePtr && (InstanceIndex != INDEX_NONE)
-		? ManagerInterfacePtr->GetLevelForInstance(InstanceIndex)
-		: nullptr;
+	return nullptr;
 }
 
 FVector FActorInstanceHandle::GetLocation() const
@@ -118,10 +90,12 @@ FVector FActorInstanceHandle::GetLocation() const
 		return Actor->GetActorLocation();
 	}
 
-	IActorInstanceManagerInterface* ManagerInterfacePtr = ManagerInterface.Get();
-	return ManagerInterfacePtr && (InstanceIndex != INDEX_NONE)
-		? ManagerInterfacePtr->GetTransform(*this).GetLocation()
-		: FVector();
+	if (Manager.IsValid())
+	{
+		return Manager->GetLocation(*this);
+	}
+
+	return FVector();
 }
 
 FRotator FActorInstanceHandle::GetRotation() const
@@ -131,10 +105,12 @@ FRotator FActorInstanceHandle::GetRotation() const
 		return Actor->GetActorRotation();
 	}
 
-	IActorInstanceManagerInterface* ManagerInterfacePtr = ManagerInterface.Get();
-	return ManagerInterfacePtr && (InstanceIndex != INDEX_NONE)
-		? ManagerInterfacePtr->GetTransform(*this).GetRotation().Rotator()
-		: FRotator();
+	if (Manager.IsValid())
+	{
+		return Manager->GetRotation(*this);
+	}
+
+	return FRotator();
 }
 
 FTransform FActorInstanceHandle::GetTransform() const
@@ -144,10 +120,12 @@ FTransform FActorInstanceHandle::GetTransform() const
 		return Actor->GetActorTransform();
 	}
 
-	IActorInstanceManagerInterface* ManagerInterfacePtr = ManagerInterface.Get();
-	return ManagerInterfacePtr && (InstanceIndex != INDEX_NONE)
-		? ManagerInterfacePtr->GetTransform(*this)
-		: FTransform();
+	if (Manager.IsValid())
+	{
+		return Manager->GetTransform(*this);
+	}
+
+	return FTransform();
 }
 
 FName FActorInstanceHandle::GetFName() const
@@ -167,12 +145,12 @@ FString FActorInstanceHandle::GetName() const
 		return Actor->GetName();
 	}
 
-	if (ManagerInterface.IsValid())
+	if (Manager.IsValid())
 	{
-		return FString::Printf(TEXT("%s:d"), *GetNameSafe(ManagerInterface.GetObject()), InstanceIndex);
+		return Manager->GetName(*this);
 	}
 
-	return TEXT("Invalid");
+	return FString();
 }
 
 AActor* FActorInstanceHandle::GetManagingActor() const
@@ -182,7 +160,7 @@ AActor* FActorInstanceHandle::GetManagingActor() const
 		return Actor.Get();
 	}
 
-	return Cast<AActor>(ManagerInterface.GetObject());
+	return Manager.Get();
 }
 
 USceneComponent* FActorInstanceHandle::GetRootComponent() const
@@ -192,8 +170,7 @@ USceneComponent* FActorInstanceHandle::GetRootComponent() const
 		return Actor->GetRootComponent();
 	}
 
-	AActor* AsActor = Cast<AActor>(ManagerInterface.GetObject());
-	return AsActor ? AsActor->GetRootComponent() : nullptr;
+	return Manager.IsValid() ? Manager->GetRootComponent() : nullptr;
 }
 
 AActor* FActorInstanceHandle::FetchActor() const
@@ -203,13 +180,22 @@ AActor* FActorInstanceHandle::FetchActor() const
 		return Actor.Get();
 	}
 
-	return ManagerInterface.IsValid() ? ManagerInterface->FindOrCreateActor(*this) : nullptr;
+	return FLightWeightInstanceSubsystem::Get().FetchActor(*this);
+}
+
+int32 FActorInstanceHandle::GetRenderingInstanceIndex() const
+{
+	return Manager.IsValid() ? Manager->ConvertLightWeightIndexToCollisionIndex(InstanceIndex) : INDEX_NONE;
 }
 
 UObject* FActorInstanceHandle::GetActorAsUObject()
 {
-	// 
-	return Cast<UObject>(Actor.Get());
+	if (IsActorValid())
+	{
+		return Cast<UObject>(Actor.Get());
+	}
+
+	return nullptr;
 }
 
 const UObject* FActorInstanceHandle::GetActorAsUObject() const
@@ -227,16 +213,10 @@ bool FActorInstanceHandle::IsActorValid() const
 	return Actor.IsValid();
 }
 
-void FActorInstanceHandle::SetCachedActor(AActor* InActor) const
-{
-	check(Actor.IsValid() == false);
-	Actor = InActor;
-}
-
 FActorInstanceHandle& FActorInstanceHandle::operator=(AActor* OtherActor)
 {
 	Actor = OtherActor;
-	ManagerInterface.Reset();
+	Manager.Reset();
 	InstanceIndex = INDEX_NONE;
 
 	return *this;
@@ -245,9 +225,9 @@ FActorInstanceHandle& FActorInstanceHandle::operator=(AActor* OtherActor)
 bool FActorInstanceHandle::operator==(const FActorInstanceHandle& Other) const
 {
 	// try to compare managers and indices first if we have them
-	if (ManagerInterface.IsValid() && Other.ManagerInterface.IsValid() && InstanceIndex != INDEX_NONE && Other.InstanceIndex != INDEX_NONE)
+	if (Manager.IsValid() && Other.Manager.IsValid() && InstanceIndex != INDEX_NONE && Other.InstanceIndex != INDEX_NONE)
 	{
-		return ManagerInterface == Other.ManagerInterface && InstanceIndex == Other.InstanceIndex;
+		return Manager == Other.Manager && InstanceIndex == Other.InstanceIndex;
 	}
 
 	// try to compare the actors
@@ -265,19 +245,38 @@ bool FActorInstanceHandle::operator!=(const FActorInstanceHandle& Other) const
 bool FActorInstanceHandle::operator==(const AActor* OtherActor) const
 {
 	// if we have an actor, compare the two actors
-	if (AActor* AsActor = Actor.Get())
+	if (Actor.IsValid())
 	{
-		return AsActor == OtherActor;
+		return Actor.Get() == OtherActor;
 	}
 
 	// if OtherActor is null then we're only equal if this doesn't refer to a valid instance
 	if (OtherActor == nullptr)
 	{
-		return !ManagerInterface.IsValid() && InstanceIndex == INDEX_NONE;
+		return !Manager.IsValid() && InstanceIndex == INDEX_NONE;
+	}
+
+	// we don't have an actor so see if we can look up an instance associated with OtherActor and see if we refer to the same instance
+
+#if WITH_EDITOR
+	// use the first layer the actor is in if it's in multiple layers
+	TArray<const UDataLayerInstance*> DataLayerInstances = OtherActor->GetDataLayerInstances();
+	const UDataLayerInstance* DataLayerInstance = DataLayerInstances.Num() > 0 ? DataLayerInstances[0] : nullptr;
+#else
+	const UDataLayerInstance* DataLayerInstance = nullptr;
+#endif // WITH_EDITOR
+
+	if (ALightWeightInstanceManager* LWIManager = FLightWeightInstanceSubsystem::Get().FindLightWeightInstanceManager(*OtherActor->GetClass(), *OtherActor->GetWorld(), OtherActor->GetActorLocation(), DataLayerInstance))
+	{
+		if (Manager.Get() != LWIManager)
+		{
+			return false;
+		}
+
+		return Manager->FindIndexForActor(OtherActor) == InstanceIndex;
 	}
 
-	IActorInstanceManagerInterface* ManagerInterfacePtr = ManagerInterface.Get();
-	return ManagerInterfacePtr && (ManagerInterfacePtr->FindActor(*this) == OtherActor);
+	return false;
 }
 
 bool FActorInstanceHandle::operator!=(const AActor* OtherActor) const
@@ -288,13 +287,13 @@ bool FActorInstanceHandle::operator!=(const AActor* OtherActor) const
 uint32 GetTypeHash(const FActorInstanceHandle& Handle)
 {
 	uint32 Hash = 0;
-	if (AActor* Actor = Handle.Actor.Get())
+	if (Handle.Actor.IsValid())
 	{
-		FCrc::StrCrc32(*(Actor->GetPathName()), Hash);
+		FCrc::StrCrc32(*(Handle.Actor->GetPathName()), Hash);
 	}
-	if (UObject* ManagerInterfaceObject = Handle.ManagerInterface.GetObject())
+	if (Handle.Manager.IsValid())
 	{
-		Hash = HashCombine(Hash, GetTypeHash(ManagerInterfaceObject));
+		Hash = HashCombine(Hash, GetTypeHash(Handle.Manager.Get()));
 	}
 	Hash = HashCombine(Hash, Handle.InstanceIndex);
 
@@ -304,15 +303,9 @@ uint32 GetTypeHash(const FActorInstanceHandle& Handle)
 FArchive& operator<<(FArchive& Ar, FActorInstanceHandle& Handle)
 {
 	Ar << Handle.Actor;
+	Ar << Handle.Manager;
 	Ar << Handle.InstanceIndex;
 
-	TSoftObjectPtr<UObject> SoftObject(Handle.ManagerInterface.GetObject());
-	Ar << SoftObject;
-
-	if (Ar.IsLoading())
-	{
-		Handle.ManagerInterface = FActorInstanceManagerInterface(SoftObject.Get());
-	}
-
 	return Ar;
 }
+
diff --git a/Engine/Source/Runtime/Engine/Private/Engine/HitResult.cpp b/Engine/Source/Runtime/Engine/Private/Engine/HitResult.cpp
--- a/Engine/Source/Runtime/Engine/Private/Engine/HitResult.cpp
+++ b/Engine/Source/Runtime/Engine/Private/Engine/HitResult.cpp
@@ -2,6 +2,7 @@
 
 #include "Engine/HitResult.h"
 #include "PhysicalMaterials/PhysicalMaterial.h"
+#include "GameFramework/LightWeightInstanceSubsystem.h"
 #include "Components/PrimitiveComponent.h"
 
 #include UE_INLINE_GENERATED_CPP_BY_NAME(HitResult)
@@ -157,7 +158,7 @@ FString FHitResult::ToString() const
 		PenetrationDepth,
 		Item,
 		PhysMaterial.IsValid() ? *PhysMaterial->GetName() : TEXT("None"),
-		HitObjectHandle.IsValid() ? *HitObjectHandle.GetName() : TEXT("None"),
+		HitObjectHandle.IsValid() ? *FLightWeightInstanceSubsystem::Get().GetName(HitObjectHandle) : TEXT("None"),
 		Component.IsValid() ? *Component->GetName() : TEXT("None"),
 		BoneName.IsValid() ? *BoneName.ToString() : TEXT("None"),
 		FaceIndex);
diff --git a/Engine/Source/Runtime/Engine/Private/EngineUtils.cpp b/Engine/Source/Runtime/Engine/Private/EngineUtils.cpp
--- a/Engine/Source/Runtime/Engine/Private/EngineUtils.cpp
+++ b/Engine/Source/Runtime/Engine/Private/EngineUtils.cpp
@@ -13,6 +13,7 @@
 #include "Misc/PackageName.h"
 #include "Misc/EngineVersion.h"
 #include "GameFramework/PlayerController.h"
+#include "GameFramework/LightWeightInstanceSubsystem.h"
 #include "GenericPlatform/ICursor.h"
 #include "Elements/Framework/EngineElementsLibrary.h"
 #include "Components/PrimitiveComponent.h"
@@ -519,7 +520,7 @@ TArray<FSubLevelStatus> GetSubLevelsStatus( UWorld* World, bool SortByActorCount
 
 				if (Hit.HitObjectHandle.IsValid())
 				{
-					LevelPlayerIsIn = Hit.HitObjectHandle.GetLevel();
+					LevelPlayerIsIn = FLightWeightInstanceSubsystem::Get().GetLevel(Hit.HitObjectHandle);
 				}
 				else if (UPrimitiveComponent* HitComponent = Hit.Component.Get())
 				{
diff --git a/Engine/Source/Runtime/Engine/Private/LightWeightInstanceManager.cpp b/Engine/Source/Runtime/Engine/Private/LightWeightInstanceManager.cpp
--- a/Engine/Source/Runtime/Engine/Private/LightWeightInstanceManager.cpp
+++ b/Engine/Source/Runtime/Engine/Private/LightWeightInstanceManager.cpp
@@ -20,6 +20,12 @@ static FAutoConsoleVariableRef CVarLWIGridSize
 );
 
 
+UActorInstanceHandleInterface::UActorInstanceHandleInterface(const FObjectInitializer& ObjectInitializer)
+	:Super(ObjectInitializer)
+{
+	// do nothing
+}
+
 ALightWeightInstanceManager::ALightWeightInstanceManager(const FObjectInitializer& ObjectInitializer)
 	:Super(ObjectInitializer)
 {
@@ -41,22 +47,20 @@ void ALightWeightInstanceManager::SetRepresentedClass(UClass* ActorClass)
 	RepresentedClass = ActorClass;
 }
 
-void ALightWeightInstanceManager::Tick(float DeltaSeconds)
+UClass* ALightWeightInstanceManager::GetInterfaceClass() const
 {
-	// do nothing
+	return UActorInstanceHandleInterface::StaticClass();
 }
 
-AActor* ALightWeightInstanceManager::FindActor(const FActorInstanceHandle& Handle)
+void ALightWeightInstanceManager::Tick(float DeltaSeconds)
 {
-	TObjectPtr<AActor>* FoundActor = Actors.Find(Handle.GetInstanceIndex());
-
-	return FoundActor ? *FoundActor : nullptr;
+	// do nothing
 }
 
-AActor* ALightWeightInstanceManager::FindOrCreateActor(const FActorInstanceHandle& Handle)
+AActor* ALightWeightInstanceManager::FetchActorFromHandle(const FActorInstanceHandle& Handle)
 {
 	// make sure the handle doesn't have an actor already
-	if (ensure(!Handle.GetCachedActor()))
+	if (ensure(!Handle.Actor.IsValid()))
 	{
 		// check if we already have an actor for this handle
 		TObjectPtr<AActor>* FoundActor = nullptr;
@@ -78,14 +82,14 @@ AActor* ALightWeightInstanceManager::FindOrCreateActor(const FActorInstanceHandl
 			}
 #endif //!UE_BUILD_SHIPPING
 
-			Handle.SetCachedActor(*FoundActor);
+			Handle.Actor = *FoundActor;
 		}
 	}
 
 	// Unless we are on the server or the actor has been spawned, this ensure will fail.
 	// Commented out until there is more robust replication.
 	// ensure(Handle.Actor.IsValid());
-	return Handle.GetCachedActor();
+	return Handle.Actor.Get();
 }
 
 AActor* ALightWeightInstanceManager::ConvertInstanceToActor(const FActorInstanceHandle& Handle)
@@ -136,7 +140,7 @@ AActor* ALightWeightInstanceManager::ConvertInstanceToActor(const FActorInstance
 		check(NewActor);
 
 		//should have been assigned in CustomPreSpawnInitialization
-		check(Handle.GetCachedActor() == NewActor);
+		check(Handle.Actor == NewActor);
 		check(NewActor == Actors.FindRef(Handle.GetInstanceIndex()));
 
 		NewActor->OnDestroyed.AddUniqueDynamic(this, &ALightWeightInstanceManager::OnSpawnedActorDestroyed);
@@ -194,9 +198,10 @@ int32 ALightWeightInstanceManager::FindIndexForActor(const AActor* InActor) cons
 
 FActorInstanceHandle ALightWeightInstanceManager::ConvertActorToLightWeightInstance(AActor* InActor)
 {
+	FActorInstanceHandle ReturnHandle;
 	if (!InActor)
 	{
-		return FActorInstanceHandle();
+		return ReturnHandle;
 	}
 
 	if (FLWIData* Data = AllocateInitData())
@@ -223,18 +228,23 @@ FActorInstanceHandle ALightWeightInstanceManager::ConvertActorToLightWeightInsta
 			UpdateDataAtIndex(Data, Idx);
 		}
 
+		// Update our handle
+		ReturnHandle.Manager = this;
+		ReturnHandle.InstanceIndex = Idx;
+		ReturnHandle.InstanceUID = InActor->GetWorld()->LWILastAssignedUID++;
+
 		// cleanup
 		InActor->Destroy();
 		delete Data;
-
-		return FActorInstanceHandle::MakeDehydratedActorHandle(*this, Idx);
 	}
 	else
 	{
 		// something went wrong and we can't manage this actor
 		// just return a handle to the actor
-		return FActorInstanceHandle(InActor);
+		ReturnHandle.Actor = InActor;
 	}
+
+	return ReturnHandle;
 }
 
 FLWIData* ALightWeightInstanceManager::AllocateInitData() const
@@ -255,7 +265,12 @@ bool ALightWeightInstanceManager::SetDataFromActor(FLWIData* InData, AActor* InA
 	return true;
 }
 
-int32 ALightWeightInstanceManager::ConvertCollisionIndexToInstanceIndex(int32 InIndex, const UPrimitiveComponent* RelevantComponent) const
+int32 ALightWeightInstanceManager::ConvertCollisionIndexToLightWeightIndex(int32 InIndex) const
+{
+	return InIndex;
+}
+
+int32 ALightWeightInstanceManager::ConvertLightWeightIndexToCollisionIndex(int32 InIndex) const
 {
 	return InIndex;
 }
@@ -310,7 +325,7 @@ UClass* ALightWeightInstanceManager::GetActorClassToSpawn(const FActorInstanceHa
 
 void ALightWeightInstanceManager::PreSpawnInitalization(const FActorInstanceHandle& Handle, AActor* SpawnedActor)
 {
-	Handle.SetCachedActor(SpawnedActor);
+	Handle.Actor = SpawnedActor;
 	Actors.Add(Handle.GetInstanceIndex(), SpawnedActor);
 }
 
@@ -331,15 +346,19 @@ bool ALightWeightInstanceManager::IsIndexValid(int32 Index) const
 
 bool ALightWeightInstanceManager::FindActorForHandle(const FActorInstanceHandle& Handle) const
 {
-	if (Handle.GetCachedActor())
+	if (Handle.Actor.IsValid())
 	{
 		return true;
 	}
 
 	const TObjectPtr<AActor>* FoundActor = Actors.Find(Handle.GetInstanceIndex());
-	AActor* ResolvedActor = FoundActor ? *FoundActor : nullptr;
-	Handle.SetCachedActor(ResolvedActor);
-	return ResolvedActor != nullptr;
+	Handle.Actor = FoundActor ? *FoundActor : nullptr;
+	return Handle.Actor != nullptr;
+}
+
+AActor* ALightWeightInstanceManager::FindActorForInstanceIndex(const int32 InstanceIndex)
+{
+	return Actors.FindRef(InstanceIndex);
 }
 
 FVector ALightWeightInstanceManager::GetLocation(const FActorInstanceHandle& Handle) const
@@ -356,8 +375,7 @@ FTransform ALightWeightInstanceManager::GetTransform(const FActorInstanceHandle&
 {
 	if (FindActorForHandle(Handle))
 	{
-		check(Handle.GetCachedActor());
-		return Handle.GetCachedActor()->GetActorTransform();
+		return Handle.Actor->GetActorTransform();
 	}
 
 	if (ensure(IsIndexValid(Handle.GetInstanceIndex())))
@@ -373,8 +391,7 @@ FString ALightWeightInstanceManager::GetName(const FActorInstanceHandle& Handle)
 {
 	if (FindActorForHandle(Handle))
 	{
-		check(Handle.GetCachedActor());
-		return Handle.GetCachedActor()->GetName();
+		return Handle.Actor->GetName();
 	}
 
 	return FString::Printf(TEXT("%s_%u"), *BaseInstanceName, Handle.GetInstanceIndex());
@@ -382,7 +399,7 @@ FString ALightWeightInstanceManager::GetName(const FActorInstanceHandle& Handle)
 
 bool ALightWeightInstanceManager::DoesRepresentClass(const UClass* OtherClass) const
 {
-	return OtherClass ? OtherClass->IsChildOf(GetRepresentedClassInternal()) : false;
+	return OtherClass ? OtherClass->IsChildOf(GetRepresentedClass()) : false;
 }
 
 bool ALightWeightInstanceManager::DoesAcceptClass(const UClass* OtherClass) const
@@ -390,7 +407,7 @@ bool ALightWeightInstanceManager::DoesAcceptClass(const UClass* OtherClass) cons
 	return OtherClass ? OtherClass->IsChildOf(GetAcceptedClass()) : false;
 }
 
-UClass* ALightWeightInstanceManager::GetRepresentedClassInternal() const
+UClass* ALightWeightInstanceManager::GetRepresentedClass() const
 {
 	return IsValid(RepresentedClass) ? RepresentedClass : nullptr;
 }
diff --git a/Engine/Source/Runtime/Engine/Private/LightWeightInstanceStaticMeshManager.cpp b/Engine/Source/Runtime/Engine/Private/LightWeightInstanceStaticMeshManager.cpp
--- a/Engine/Source/Runtime/Engine/Private/LightWeightInstanceStaticMeshManager.cpp
+++ b/Engine/Source/Runtime/Engine/Private/LightWeightInstanceStaticMeshManager.cpp
@@ -116,7 +116,7 @@ PRAGMA_ENABLE_DEPRECATION_WARNINGS
 	}
 }
 
-int32 ALightWeightInstanceStaticMeshManager::ConvertCollisionIndexToInstanceIndex(int32 InIndex, const UPrimitiveComponent* RelevantComponent) const
+int32 ALightWeightInstanceStaticMeshManager::ConvertCollisionIndexToLightWeightIndex(int32 InIndex) const
 {
 	if (ensureMsgf(RenderingIndicesToDataIndices.IsValidIndex(InIndex), TEXT("Invalid index [ %d ]"), InIndex))
 	{
@@ -126,6 +126,16 @@ int32 ALightWeightInstanceStaticMeshManager::ConvertCollisionIndexToInstanceInde
 	return InIndex;
 }
 
+int32 ALightWeightInstanceStaticMeshManager::ConvertLightWeightIndexToCollisionIndex(int32 InIndex) const
+{
+	if (ensureMsgf(DataIndicesToRenderingIndices.IsValidIndex(InIndex), TEXT("Invalid index [ %d ]"), InIndex))
+	{
+		return DataIndicesToRenderingIndices[InIndex];
+	}
+
+	return InIndex;
+}
+
 void ALightWeightInstanceStaticMeshManager::AddNewInstanceAt(FLWIData* InitData, int32 Index)
 {
 	Super::AddNewInstanceAt(InitData, Index);
diff --git a/Engine/Source/Runtime/Engine/Private/LightWeightInstanceSubsystem.cpp b/Engine/Source/Runtime/Engine/Private/LightWeightInstanceSubsystem.cpp
--- a/Engine/Source/Runtime/Engine/Private/LightWeightInstanceSubsystem.cpp
+++ b/Engine/Source/Runtime/Engine/Private/LightWeightInstanceSubsystem.cpp
@@ -79,21 +79,18 @@ bool FLightWeightInstanceSubsystem::RemoveManager(ALightWeightInstanceManager* M
 
 ALightWeightInstanceManager* FLightWeightInstanceSubsystem::FindLightWeightInstanceManager(const FActorInstanceHandle& Handle) const
 {
-	if (Handle.GetManagerInterface().IsValid())
+	if (Handle.Manager.IsValid())
 	{
-		return Handle.GetManager<ALightWeightInstanceManager>();
+		return Handle.Manager.Get();
 	}
 
-	if (AActor* CachedActor = Handle.GetCachedActor())
+	if (Handle.Actor.IsValid())
 	{
-		const UClass* CachedActorClass = CachedActor->GetClass();
-		CA_ASSUME(CachedActorClass);
-
 		FReadScopeLock Lock(LWIManagersRWLock);
 		// see if we already have a match
 		for (ALightWeightInstanceManager* LWInstance : LWInstanceManagers)
 		{
-			if (CachedActorClass == LWInstance->GetRepresentedClass(Handle.GetInstanceIndex()))
+			if (Handle.Actor->GetClass() == LWInstance->GetRepresentedClass())
 			{
 				const FInt32Vector3 GridCoord = LWInstance->ConvertPositionToCoord(Handle.GetLocation());
 				const FInt32Vector3 ManagerGridCoord = LWInstance->ConvertPositionToCoord(LWInstance->GetActorLocation());
@@ -142,7 +139,7 @@ ALightWeightInstanceManager* FLightWeightInstanceSubsystem::FindLightWeightInsta
 	{
 		if (IsValid(InstanceManager))
 		{
-			if (InstanceManager->GetRepresentedClassInternal() == &ActorClass)
+			if (InstanceManager->GetRepresentedClass() == &ActorClass)
 			{
 				const FInt32Vector3 GridCoord = InstanceManager->ConvertPositionToCoord(InPos);
 				const FInt32Vector3 ManagerGridCoord = InstanceManager->ConvertPositionToCoord(InstanceManager->GetActorLocation());
@@ -243,7 +240,7 @@ UClass* FLightWeightInstanceSubsystem::FindBestInstanceManagerClass(const UClass
 	{
 		if (ManagerClass->GetDefaultObject<ALightWeightInstanceManager>()->DoesAcceptClass(InActorClass))
 		{
-			const UClass* HandledClass = ManagerClass->GetDefaultObject<ALightWeightInstanceManager>()->GetRepresentedClassInternal();
+			const UClass* HandledClass = ManagerClass->GetDefaultObject<ALightWeightInstanceManager>()->GetRepresentedClass();
 			if (!HandledClass)
 			{
 				HandledClass = ManagerClass->GetDefaultObject<ALightWeightInstanceManager>()->GetAcceptedClass();
@@ -271,21 +268,37 @@ UClass* FLightWeightInstanceSubsystem::FindBestInstanceManagerClass(const UClass
 	return BestManagerClass;
 }
 
+AActor* FLightWeightInstanceSubsystem::FetchActor(const FActorInstanceHandle& Handle)
+{
+	// if the actor is valid return it
+	if (Handle.Actor.IsValid())
+	{
+		return Handle.Actor.Get();
+	}
+
+	if (ALightWeightInstanceManager* LWIManager = FindLightWeightInstanceManager(Handle))
+	{
+		return LWIManager->FetchActorFromHandle(Handle);
+	}
+
+	return nullptr;
+}
+
 AActor* FLightWeightInstanceSubsystem::GetActor_NoCreate(const FActorInstanceHandle& Handle) const
 {
-	return Handle.GetCachedActor();
+	return Handle.Actor.Get();
 }
 
 UClass* FLightWeightInstanceSubsystem::GetActorClass(const FActorInstanceHandle& Handle)
 {
-	if (const AActor* CachedActor = Handle.GetCachedActor())
+	if (Handle.Actor.IsValid())
 	{
-		return CachedActor->StaticClass();
+		return Handle.Actor->StaticClass();
 	}
 
 	if (ALightWeightInstanceManager* LWIManager = FindLightWeightInstanceManager(Handle))
 	{
-		return LWIManager->GetRepresentedClass(Handle.GetInstanceIndex());
+		return LWIManager->GetRepresentedClass();
 	}
 
 	return nullptr;
@@ -295,9 +308,9 @@ FVector FLightWeightInstanceSubsystem::GetLocation(const FActorInstanceHandle& H
 {
 	ensure(Handle.IsValid());
 
-	if (const AActor* CachedActor = Handle.GetCachedActor())
+	if (Handle.Actor.IsValid())
 	{
-		return CachedActor->GetActorLocation();
+		return Handle.Actor->GetActorLocation();
 	}
 
 	if (ALightWeightInstanceManager* InstanceManager = FindLightWeightInstanceManager(Handle))
@@ -312,9 +325,9 @@ FString FLightWeightInstanceSubsystem::GetName(const FActorInstanceHandle& Handl
 {
 	ensure(Handle.IsValid());
 
-	if (const AActor* CachedActor = Handle.GetCachedActor())
+	if (Handle.Actor.IsValid())
 	{
-		return CachedActor->GetName();
+		return Handle.Actor->GetName();
 	}
 
 	if (ALightWeightInstanceManager* InstanceManager = FindLightWeightInstanceManager(Handle))
@@ -329,9 +342,9 @@ ULevel* FLightWeightInstanceSubsystem::GetLevel(const FActorInstanceHandle& Hand
 {
 	ensure(Handle.IsValid());
 
-	if (AActor* CachedActor = Handle.GetCachedActor())
+	if (Handle.Actor.IsValid())
 	{
-		return CachedActor->GetLevel();
+		return Handle.Actor->GetLevel();
 	}
 
 	if (ALightWeightInstanceManager* InstanceManager = FindLightWeightInstanceManager(Handle))
@@ -346,9 +359,9 @@ bool FLightWeightInstanceSubsystem::IsInLevel(const FActorInstanceHandle& Handle
 {
 	ensure(Handle.IsValid());
 
-	if (AActor* CachedActor = Handle.GetCachedActor())
+	if (Handle.Actor.IsValid())
 	{
-		return CachedActor->IsInLevel(InLevel);
+		return Handle.Actor->IsInLevel(InLevel);
 	}
 
 	if (ALightWeightInstanceManager* InstanceManager = FindLightWeightInstanceManager(Handle))
@@ -369,7 +382,7 @@ FActorInstanceHandle FLightWeightInstanceSubsystem::CreateNewLightWeightInstance
 			// create an instance with the given data
 			int32 InstanceIdx = LWIManager->AddNewInstance(InitData);
 			InstanceIdx = LWIManager->ConvertInternalIndexToHandleIndex(InstanceIdx);
-			return FActorInstanceHandle(FActorInstanceManagerInterface(LWIManager), InstanceIdx);
+			return FActorInstanceHandle(LWIManager, InstanceIdx);
 		}
 	}
 
diff --git a/Engine/Source/Runtime/Engine/Private/World.cpp b/Engine/Source/Runtime/Engine/Private/World.cpp
--- a/Engine/Source/Runtime/Engine/Private/World.cpp
+++ b/Engine/Source/Runtime/Engine/Private/World.cpp
@@ -601,6 +601,7 @@ UWorld::UWorld( const FObjectInitializer& ObjectInitializer )
 #if WITH_EDITOR
 , HierarchicalLODBuilder(new FHierarchicalLODBuilder(this))
 #endif
+, LWILastAssignedUID(0)
 , URL(FURL(NULL))
 ,	FXSystem(NULL)
 ,	TickTaskLevel(FTickTaskManagerInterface::Get().AllocateTickTaskLevel())
