[PhysicsFormats] Remove extraneous GetPhysicsFormat() overrides and make the default "Chaos".

diff --git a/Engine/Source/Runtime/Core/Public/Android/AndroidPlatformProperties.h b/Engine/Source/Runtime/Core/Public/Android/AndroidPlatformProperties.h
--- a/Engine/Source/Runtime/Core/Public/Android/AndroidPlatformProperties.h
+++ b/Engine/Source/Runtime/Core/Public/Android/AndroidPlatformProperties.h
@@ -19,11 +19,6 @@
 struct FAndroidPlatformProperties
 	: public FGenericPlatformProperties
 {
-	static FORCEINLINE const char* GetPhysicsFormat( )
-	{
-		return "PhysXGeneric";		//@todo android: physx format
-	}
-
 	static FORCEINLINE bool HasEditorOnlyData( )
 	{
 		return false;
diff --git a/Engine/Source/Runtime/Core/Public/GenericPlatform/GenericPlatformProperties.h b/Engine/Source/Runtime/Core/Public/GenericPlatform/GenericPlatformProperties.h
--- a/Engine/Source/Runtime/Core/Public/GenericPlatform/GenericPlatformProperties.h
+++ b/Engine/Source/Runtime/Core/Public/GenericPlatform/GenericPlatformProperties.h
@@ -22,7 +22,7 @@ struct FGenericPlatformProperties
 	 */
 	static FORCEINLINE const char* GetPhysicsFormat()
 	{
-		return "PhysXGeneric";
+		return "Chaos";
 	}
 
 	/**
diff --git a/Engine/Source/Runtime/Core/Public/Windows/WindowsPlatformProperties.h b/Engine/Source/Runtime/Core/Public/Windows/WindowsPlatformProperties.h
--- a/Engine/Source/Runtime/Core/Public/Windows/WindowsPlatformProperties.h
+++ b/Engine/Source/Runtime/Core/Public/Windows/WindowsPlatformProperties.h
@@ -28,11 +28,6 @@ struct FWindowsPlatformProperties
 		return TEXT("/Script/WindowsTargetPlatform.WindowsTargetSettings");
 	}
 
-	static FORCEINLINE const char* GetPhysicsFormat()
-	{
-		return "PhysXPC";
-	}
-
 	static FORCEINLINE bool IsGameOnly()
 	{
 		return UE_GAME;
diff --git a/Engine/Source/Runtime/Engine/Private/PhysicsEngine/BodySetup.cpp b/Engine/Source/Runtime/Engine/Private/PhysicsEngine/BodySetup.cpp
--- a/Engine/Source/Runtime/Engine/Private/PhysicsEngine/BodySetup.cpp
+++ b/Engine/Source/Runtime/Engine/Private/PhysicsEngine/BodySetup.cpp
@@ -388,16 +388,7 @@ FByteBulkData* UBodySetup::GetCookedFormatData()
 {
 	// Find or create cooked physics data
 	static FName PhysicsFormatName(FPlatformProperties::GetPhysicsFormat());
-
-	FByteBulkData* FormatData = GetCookedData(PhysicsFormatName);
-
-	// On dedicated servers we may be cooking generic data and sharing it
-	if (FormatData == nullptr && IsRunningDedicatedServer())
-	{
-		FormatData = GetCookedData(FGenericPlatformProperties::GetPhysicsFormat());
-	}
-
-	return FormatData;
+	return GetCookedData(PhysicsFormatName);
 }
 
 void UBodySetup::CreatePhysicsMeshes()
