Adds new timing regions for easier UnrealInsights trace analysis:
- Adds Level::AddToWorld region

#rnx
#tests pie, replay, win64

diff --git a/Engine/Source/Runtime/Engine/Private/World.cpp b/Engine/Source/Runtime/Engine/Private/World.cpp
--- a/Engine/Source/Runtime/Engine/Private/World.cpp
+++ b/Engine/Source/Runtime/Engine/Private/World.cpp
@@ -3167,6 +3167,8 @@ void UWorld::AddToWorld( ULevel* Level, const FTransform& LevelTransform, bool b
 	// Don't make this level visible if it's currently being made invisible
 	if( bExecuteNextStep && CurrentLevelPendingVisibility == NULL && CurrentLevelPendingInvisibility != Level )
 	{
+		TRACE_BEGIN_REGION(*WriteToString<256>(TEXT("AddToWorld: "), Level->GetOutermost()->GetName()));
+
 		Level->OwningWorld = this;
 
 		if (OwningLevelStreaming)
@@ -3436,6 +3438,8 @@ void UWorld::AddToWorld( ULevel* Level, const FTransform& LevelTransform, bool b
 
 			ULevelStreaming::BroadcastLevelVisibleStatus(this, Level->GetOutermost()->GetFName(), true);
 		}
+
+		TRACE_END_REGION(*WriteToString<256>(TEXT("AddToWorld: "), Level->GetOutermost()->GetName()));
 	}
 #if CSV_PROFILER
 	else
