When previewing an anim notify, sound is not played anymore, if animation is on pause.

diff --git a/Engine/Source/Runtime/Engine/Private/Animation/AnimNotify_PlaySound.cpp b/Engine/Source/Runtime/Engine/Private/Animation/AnimNotify_PlaySound.cpp
--- a/Engine/Source/Runtime/Engine/Private/Animation/AnimNotify_PlaySound.cpp
+++ b/Engine/Source/Runtime/Engine/Private/Animation/AnimNotify_PlaySound.cpp
@@ -52,7 +52,10 @@ void UAnimNotify_PlaySound::Notify(class USkeletalMeshComponent* MeshComp, class
 		UWorld* World = MeshComp->GetWorld();
 		if (bPreviewIgnoreAttenuation && World && World->WorldType == EWorldType::EditorPreview)
 		{
-			UGameplayStatics::PlaySound2D(World, Sound, VolumeMultiplier, PitchMultiplier);
+			if (MeshComp->IsPlaying())
+			{
+				UGameplayStatics::PlaySound2D(World, Sound, VolumeMultiplier, PitchMultiplier);
+			}
 		}
 		else
 #endif
