[NO NON-PII COMMIT MESSAGE]

diff --git a/Engine/Source/Runtime/Engine/Classes/Interfaces/Interface_AssetUserData.h b/Engine/Source/Runtime/Engine/Classes/Interfaces/Interface_AssetUserData.h
--- a/Engine/Source/Runtime/Engine/Classes/Interfaces/Interface_AssetUserData.h
+++ b/Engine/Source/Runtime/Engine/Classes/Interfaces/Interface_AssetUserData.h
@@ -22,6 +22,14 @@ class IInterface_AssetUserData
 
 	virtual void AddAssetUserData(UAssetUserData* InUserData) {}
 	
+	/**
+	* Returns an instance of the provided AssetUserData class if it's contained in the target asset.
+	*
+	* @param	InUserDataClass		UAssetUserData sub class to get
+	*
+	* @return	The instance of the UAssetUserData class contained, or null if it doesn't exist
+	*/
+	UFUNCTION(BlueprintCallable, Category = AssetUserData)
 	virtual UAssetUserData* GetAssetUserDataOfClass(TSubclassOf<UAssetUserData> InUserDataClass)
 	{
 		return nullptr;
