Remove "AllowStaticLighting" test from VF, as cvar cache is initialized too late.
The test is not really needed as it does not add new streams when project has static lighting disabled

diff --git a/Engine/Source/Runtime/Engine/Private/GeometryCollectionRendering.cpp b/Engine/Source/Runtime/Engine/Private/GeometryCollectionRendering.cpp
--- a/Engine/Source/Runtime/Engine/Private/GeometryCollectionRendering.cpp
+++ b/Engine/Source/Runtime/Engine/Private/GeometryCollectionRendering.cpp
@@ -10,7 +10,6 @@
 #include "ProfilingDebugging/LoadTimeTracker.h"
 #include "PrimitiveUniformShaderParameters.h"
 #include "RenderUtils.h"
-#include "SceneManagement.h"
 
 IMPLEMENT_TYPE_LAYOUT(FGeometryCollectionVertexFactoryShaderParameters);
 
@@ -204,16 +203,13 @@ void FGeometryCollectionVertexFactory::InitRHI(FRHICommandListBase& RHICmdList)
 		}
 	}
 
-	if (FReadOnlyCVARCache::Get().bAllowStaticLighting)
+	if (Data.LightMapCoordinateComponent.VertexBuffer)
 	{
-		if (Data.LightMapCoordinateComponent.VertexBuffer)
-		{
-			Elements.Add(AccessStreamComponent(Data.LightMapCoordinateComponent, 15));
-		}
-		else if (Data.TextureCoordinates.Num())
-		{
-			Elements.Add(AccessStreamComponent(Data.TextureCoordinates[0], 15));
-		}
+		Elements.Add(AccessStreamComponent(Data.LightMapCoordinateComponent, 15));
+	}
+	else if (Data.TextureCoordinates.Num())
+	{
+		Elements.Add(AccessStreamComponent(Data.TextureCoordinates[0], 15));
 	}
 
 	check(Streams.Num() > 0);
diff --git a/Engine/Source/Runtime/Engine/Private/InstancedStaticMesh.cpp b/Engine/Source/Runtime/Engine/Private/InstancedStaticMesh.cpp
--- a/Engine/Source/Runtime/Engine/Private/InstancedStaticMesh.cpp
+++ b/Engine/Source/Runtime/Engine/Private/InstancedStaticMesh.cpp
@@ -45,7 +45,6 @@
 #include "UObject/UObjectIterator.h"
 #include "GenericPlatform/ICursor.h"
 #include "Rendering/RenderCommandPipes.h"
-#include "SceneManagement.h"
 
 #include "InstancedStaticMeshSceneProxyDesc.h"
 
@@ -1090,16 +1089,13 @@ void FInstancedStaticMeshVertexFactory::GetVertexElements(
 			Elements.Add(AccessStreamComponent(NullComponent, 14, Streams));
 		}
 
-		if (FReadOnlyCVARCache::Get().bAllowStaticLighting)
+		if (Data.LightMapCoordinateComponent.VertexBuffer)
 		{
-			if (Data.LightMapCoordinateComponent.VertexBuffer)
-			{
-				Elements.Add(AccessStreamComponent(Data.LightMapCoordinateComponent, 15, Streams));
-			}
-			else if (Data.TextureCoordinates.Num())
-			{
-				Elements.Add(AccessStreamComponent(Data.TextureCoordinates[0], 15, Streams));
-			}
+			Elements.Add(AccessStreamComponent(Data.LightMapCoordinateComponent, 15, Streams));
+		}
+		else if (Data.TextureCoordinates.Num())
+		{
+			Elements.Add(AccessStreamComponent(Data.TextureCoordinates[0], 15, Streams));
 		}
 	}
 
diff --git a/Engine/Source/Runtime/Engine/Private/LocalVertexFactory.cpp b/Engine/Source/Runtime/Engine/Private/LocalVertexFactory.cpp
--- a/Engine/Source/Runtime/Engine/Private/LocalVertexFactory.cpp
+++ b/Engine/Source/Runtime/Engine/Private/LocalVertexFactory.cpp
@@ -20,7 +20,6 @@
 #include "GPUSkinVertexFactory.h"
 #include "RenderUtils.h"
 #include "SceneInterface.h"
-#include "SceneManagement.h"
 
 IMPLEMENT_TYPE_LAYOUT(FLocalVertexFactoryShaderParametersBase);
 IMPLEMENT_TYPE_LAYOUT(FLocalVertexFactoryShaderParameters);
@@ -549,16 +548,13 @@ void FLocalVertexFactory::GetVertexElements(
 			Elements.Add(AccessStreamComponent(Data.PreSkinPositionComponent.VertexBuffer ? Data.PreSkinPositionComponent : NullComponent, 14, InOutStreams));
 		}
 
-		if (FReadOnlyCVARCache::Get().bAllowStaticLighting)
+		if (Data.LightMapCoordinateComponent.VertexBuffer)
 		{
-			if (Data.LightMapCoordinateComponent.VertexBuffer)
-			{
-				Elements.Add(AccessStreamComponent(Data.LightMapCoordinateComponent, 15, InOutStreams));
-			}
-			else if (Data.TextureCoordinates.Num())
-			{
-				Elements.Add(AccessStreamComponent(Data.TextureCoordinates[0], 15, InOutStreams));
-			}
+			Elements.Add(AccessStreamComponent(Data.LightMapCoordinateComponent, 15, InOutStreams));
+		}
+		else if (Data.TextureCoordinates.Num())
+		{
+			Elements.Add(AccessStreamComponent(Data.TextureCoordinates[0], 15, InOutStreams));
 		}
 	}
 }
