LFV - changed the base size to something reasonable  so that scaling LFVs is also easier.

#prelfight none

diff --git a/Engine/Source/Editor/ComponentVisualizers/Private/LocalFogVolumeComponentVisualizer.cpp b/Engine/Source/Editor/ComponentVisualizers/Private/LocalFogVolumeComponentVisualizer.cpp
--- a/Engine/Source/Editor/ComponentVisualizers/Private/LocalFogVolumeComponentVisualizer.cpp
+++ b/Engine/Source/Editor/ComponentVisualizers/Private/LocalFogVolumeComponentVisualizer.cpp
@@ -16,7 +16,7 @@ void FLocalFogVolumeComponentVisualizer::DrawVisualization( const UActorComponen
 		{
 			FTransform LocalFogVolumeTransform = LocalFogVolume->GetComponentTransform();
 
-			const float MaximumAxisScale = LocalFogVolumeTransform.GetMaximumAxisScale();
+			const float MaximumAxisScale = LocalFogVolumeTransform.GetMaximumAxisScale() * ULocalFogVolumeComponent::GetBaseVolumeSize();
 			LocalFogVolumeTransform.SetScale3D(FVector(MaximumAxisScale, MaximumAxisScale, MaximumAxisScale));
 
 			// Draw local fog volume spherical shape radius
diff --git a/Engine/Source/Runtime/Engine/Classes/Components/LocalFogVolumeComponent.h b/Engine/Source/Runtime/Engine/Classes/Components/LocalFogVolumeComponent.h
--- a/Engine/Source/Runtime/Engine/Classes/Components/LocalFogVolumeComponent.h
+++ b/Engine/Source/Runtime/Engine/Classes/Components/LocalFogVolumeComponent.h
@@ -50,6 +50,11 @@ class ULocalFogVolumeComponent : public USceneComponent
 
 public:
 
+	static float GetBaseVolumeSize()
+	{
+		return 500.0f; // This is immutable and cannot be changed without transform data conversion of the component.
+	}
+
 	//~ Begin UObject Interface
 #if WITH_EDITOR
 	virtual bool CanEditChange(const FProperty* InProperty) const override;
diff --git a/Engine/Source/Runtime/Engine/Private/Components/LocalFogVolumeComponent.cpp b/Engine/Source/Runtime/Engine/Private/Components/LocalFogVolumeComponent.cpp
--- a/Engine/Source/Runtime/Engine/Private/Components/LocalFogVolumeComponent.cpp
+++ b/Engine/Source/Runtime/Engine/Private/Components/LocalFogVolumeComponent.cpp
@@ -20,10 +20,6 @@ ULocalFogVolumeComponent::ULocalFogVolumeComponent(const FObjectInitializer& Obj
 	, LocalFogVolumeSceneProxy(nullptr)
 {
 	Mobility = EComponentMobility::Movable;
-
-	// Set a size by default to be visible when droped in the scene
-	const float Size = 500.0f;
-	SetRelativeScale3D(FVector(Size, Size, Size));
 }
 
 ULocalFogVolumeComponent::~ULocalFogVolumeComponent()
@@ -169,7 +165,7 @@ ALocalFogVolume::ALocalFogVolume(const FObjectInitializer& ObjectInitializer)
 		if (GetSpriteComponent())
 		{
 			GetSpriteComponent()->Sprite = ConstructorStatics.LocalFogVolumeTextureObject.Get();
-			GetSpriteComponent()->SetRelativeScale3D(FVector(0.5f, 0.5f, 0.5f));
+			GetSpriteComponent()->SetRelativeScale3D(FVector(1.0f, 1.0f, 1.0f));
 			GetSpriteComponent()->SpriteInfo.Category = ConstructorStatics.ID_LocalFogVolume;
 			GetSpriteComponent()->SpriteInfo.DisplayName = ConstructorStatics.NAME_LocalFogVolume;
 			GetSpriteComponent()->SetupAttachment(LocalFogVolumeVolume);
diff --git a/Engine/Source/Runtime/Engine/Private/LocalFogVolumeSceneProxy.cpp b/Engine/Source/Runtime/Engine/Private/LocalFogVolumeSceneProxy.cpp
--- a/Engine/Source/Runtime/Engine/Private/LocalFogVolumeSceneProxy.cpp
+++ b/Engine/Source/Runtime/Engine/Private/LocalFogVolumeSceneProxy.cpp
@@ -30,7 +30,7 @@ FLocalFogVolumeSceneProxy::~FLocalFogVolumeSceneProxy()
 void FLocalFogVolumeSceneProxy::UpdateComponentTransform(const FTransform& Transform)
 {
 	FogTransform = Transform;
-	const float MaximumAxisScale = FogTransform.GetMaximumAxisScale();
+	const float MaximumAxisScale = FogTransform.GetMaximumAxisScale() * ULocalFogVolumeComponent::GetBaseVolumeSize();
 	FogTransform.SetScale3D(FVector(MaximumAxisScale, MaximumAxisScale, MaximumAxisScale));
 
 	FogUniformScale = MaximumAxisScale;
