World Partition Streaming Generation: move the actor container resolver out of the streaming generation context into the streaming generator.
- This is required for an upcoming change where we implement a new streaming generation context for HashSet HLOD generation.
- It was impossible for a streaming generation context to return a null container resolver.
- Don't generate the container resolver when validating changelists or checking map for errors.

diff --git a/Engine/Source/Runtime/Engine/Private/WorldPartition/WorldPartitionStreamingGeneration.cpp b/Engine/Source/Runtime/Engine/Private/WorldPartition/WorldPartitionStreamingGeneration.cpp
--- a/Engine/Source/Runtime/Engine/Private/WorldPartition/WorldPartitionStreamingGeneration.cpp
+++ b/Engine/Source/Runtime/Engine/Private/WorldPartition/WorldPartitionStreamingGeneration.cpp
@@ -96,12 +96,6 @@ class FWorldPartitionStreamingGenerator
 			MainWorldActorSetContainerIndex = INDEX_NONE;
 			ActorSetContainers.Empty(StreamingGenerator->ContainerCollectionDescriptorsMap.Num());
 			
-			const bool bBuildContainerResolver = FEditorPathHelper::IsEnabled();
-			if (bBuildContainerResolver)
-			{
-				ContainerResolver.SetMainContainerPackage(TopLevelActorDescCollection.GetMainContainerPackageName());
-			}
-
 			TMap<TWeakPtr<FStreamingGenerationActorDescCollection>, int32> ActorSetContainerMap;
 			for (const auto& [LevelName, ContainerDescriptor] : StreamingGenerator->ContainerCollectionDescriptorsMap)
 			{
@@ -124,20 +118,6 @@ class FWorldPartitionStreamingGenerator
 					check(MainWorldActorSetContainerIndex == INDEX_NONE);
 					MainWorldActorSetContainerIndex = ContainerIndex;
 				}
-
-				if (bBuildContainerResolver)
-				{
-					auto& Container = ContainerResolver.AddContainer(ContainerDescriptor.ActorDescCollection->GetMainContainerPackageName());
-					for (const FWorldPartitionActorDescView& ActorDescView : ContainerDescriptor.ContainerCollectionInstanceViews)
-					{
-						Container.AddContainerInstance(ActorDescView.GetActorName(), { ActorDescView.GetGuid(), ActorDescView.GetContainerPackage() });
-					}
-				}
-			}
-
-			if (bBuildContainerResolver)
-			{
-				ContainerResolver.BuildContainerIDToEditorPathMap();
 			}
 
 			ActorSetInstances.Empty();
@@ -228,8 +208,6 @@ class FWorldPartitionStreamingGenerator
 				Func(ActorSetContainer);
 			}
 		}
-
-		virtual const FWorldPartitionRuntimeContainerResolver& GetContainerResolver() const override { return ContainerResolver; }
 		//~End IStreamingGenerationContext interface};
 
 	private:
@@ -237,7 +215,6 @@ class FWorldPartitionStreamingGenerator
 		int32 MainWorldActorSetContainerIndex;
 		TArray<FActorSetContainer> ActorSetContainers;
 		TArray<FActorSetInstance> ActorSetInstances;
-		FWorldPartitionRuntimeContainerResolver ContainerResolver;
 	};
 
 	/** 
@@ -661,6 +638,26 @@ class FWorldPartitionStreamingGenerator
 		CreateActorDescriptorViewsRecursive(MainContainerCollectionInstance);
 	}
 
+	/** 
+	 * Creates the actor descriptor container resolver
+	 */
+	void CreateContainerResolver(const FStreamingGenerationActorDescCollection& InActorDescCollection)
+	{
+		ContainerResolver.SetMainContainerPackage(InActorDescCollection.GetMainContainerPackageName());
+
+		for (const auto& [LevelName, ContainerDescriptor] : ContainerCollectionDescriptorsMap)
+		{
+			FWorldPartitionRuntimeContainer& Container = ContainerResolver.AddContainer(ContainerDescriptor.ActorDescCollection->GetMainContainerPackageName());
+
+			for (const FWorldPartitionActorDescView& ActorDescView : ContainerDescriptor.ContainerCollectionInstanceViews)
+			{
+				Container.AddContainerInstance(ActorDescView.GetActorName(), { ActorDescView.GetGuid(), ActorDescView.GetContainerPackage() });
+			}
+		}
+
+		ContainerResolver.BuildContainerIDToEditorPathMap();
+	}
+
 	/** 
 	 * Perform various validations on the container descriptor, and adjust it based on different requirements. This needs to happen before updating
 	 * containers bounds because some actor descriptor views might change grid placement, etc.
@@ -1115,12 +1112,14 @@ public:
 			, ModifiedActorsDescList(nullptr)
 			, ErrorHandler(&NullErrorHandler)
 			, bEnableStreaming(false)
+			, bCreateContainerResolver(false)
 		{}
 
 		const UWorldPartition* WorldPartitionContext;
 		FActorDescList* ModifiedActorsDescList;
 		IStreamingGenerationErrorHandler* ErrorHandler;
 		bool bEnableStreaming;
+		bool bCreateContainerResolver;
 		TMap<FGuid, const UActorDescContainer*> ActorGuidsToContainerMap;
 		TArray<TSubclassOf<AActor>> FilteredClasses;
 		TFunction<bool(FName)> IsValidGrid;
@@ -1130,6 +1129,7 @@ public:
 		FWorldPartitionStreamingGeneratorParams& SetModifiedActorsDescList(FActorDescList* InModifiedActorsDescList) { ModifiedActorsDescList = InModifiedActorsDescList; return *this; }
 		FWorldPartitionStreamingGeneratorParams& SetErrorHandler(IStreamingGenerationErrorHandler* InErrorHandler) { ErrorHandler = InErrorHandler; return *this; }
 		FWorldPartitionStreamingGeneratorParams& SetEnableStreaming(bool bInEnableStreaming) { bEnableStreaming = bInEnableStreaming; return *this; }
+		FWorldPartitionStreamingGeneratorParams& SetCreateContainerResolver(bool bInCreateContainerResolver) { bCreateContainerResolver = bInCreateContainerResolver; return *this; }
 		FWorldPartitionStreamingGeneratorParams& SetActorGuidsToContainerMap(const TMap<FGuid, const UActorDescContainer*>& InActorGuidsToContainerMap) { ActorGuidsToContainerMap = InActorGuidsToContainerMap; return *this; }
 		FWorldPartitionStreamingGeneratorParams& SetFilteredClasses(const TArray<TSubclassOf<AActor>>& InFilteredClasses) { FilteredClasses = InFilteredClasses; return *this; }
 		FWorldPartitionStreamingGeneratorParams& SetIsValidGrid(TFunction<bool(FName)> InIsValidGrid) { IsValidGrid = InIsValidGrid; return *this; }
@@ -1141,6 +1141,7 @@ public:
 	FWorldPartitionStreamingGenerator(const FWorldPartitionStreamingGeneratorParams& Params)
 		: WorldPartitionContext(Params.WorldPartitionContext)
 		, bEnableStreaming(Params.bEnableStreaming)
+		, bCreateContainerResolver(Params.bCreateContainerResolver)
 		, ModifiedActorsDescList(Params.ModifiedActorsDescList)
 		, FilteredClasses(Params.FilteredClasses)
 		, IsValidGrid(Params.IsValidGrid)
@@ -1156,6 +1157,12 @@ public:
 	void PreparationPhase(const FStreamingGenerationActorDescCollection& ActorDescCollection)
 	{
 		CreateActorContainers(ActorDescCollection);
+
+		if (bCreateContainerResolver)
+		{
+			CreateContainerResolver(ActorDescCollection);
+		}
+
 		ValidateInternalState();
 	}
 
@@ -1282,16 +1289,23 @@ public:
 		return DataLayerManager ? DataLayerManager->GetRuntimeDataLayerInstances(RuntimeDataLayers) : TArray<const UDataLayerInstance*>();
 	}
 
+	const FWorldPartitionRuntimeContainerResolver& GetContainerResolver() const
+	{
+		return ContainerResolver;
+	}
+
 private:
 	const UWorldPartition* WorldPartitionContext;
 	UWorldPartitionSubsystem* WorldPartitionSubsystem;
 	const UDataLayerManager* DataLayerManager;
 	bool bEnableStreaming;
+	bool bCreateContainerResolver;
 	FActorDescList* ModifiedActorsDescList;
 	TArray<TSubclassOf<AActor>> FilteredClasses;
 	TFunction<bool(FName)> IsValidGrid;
 	TFunction<bool(FName, const FSoftObjectPath&)> IsValidHLODLayer;
 	IStreamingGenerationErrorHandler* ErrorHandler;
+	FWorldPartitionRuntimeContainerResolver ContainerResolver;
 
 	// Maps a level to its descriptor
 	TMap<FName, FContainerCollectionDescriptor> ContainerCollectionDescriptorsMap;
@@ -1396,7 +1410,8 @@ bool UWorldPartition::GenerateContainerStreaming(const FGenerateStreamingParams&
 		.SetIsValidGrid([this](FName GridName) { return RuntimeHash->IsValidGrid(GridName); })
 		.SetIsValidHLODLayer([this](FName GridName, const FSoftObjectPath& HLODLayerPath) { return RuntimeHash->IsValidHLODLayer(GridName, HLODLayerPath); })
 		.SetErrorHandler(StreamingGenerationErrorHandlerOverride ? (*StreamingGenerationErrorHandlerOverride)(ErrorHandler) : ErrorHandler)
-		.SetEnableStreaming(IsStreamingEnabled());
+		.SetEnableStreaming(IsStreamingEnabled())
+		.SetCreateContainerResolver(FEditorPathHelper::IsEnabled());
 
 	FWorldPartitionStreamingGenerator StreamingGenerator(StreamingGeneratorParams);
 
@@ -1422,7 +1437,7 @@ bool UWorldPartition::GenerateContainerStreaming(const FGenerateStreamingParams&
 			RuntimeHash->DumpStateLog(*HierarchicalLogAr);
 		}
 
-		StreamingPolicy->SetContainerResolver(StreamingGenerationContext->GetContainerResolver());
+		StreamingPolicy->SetContainerResolver(StreamingGenerator.GetContainerResolver());
 		StreamingPolicy->PrepareActorToCellRemapping();
 		return true;
 	}
diff --git a/Engine/Source/Runtime/Engine/Public/WorldPartition/WorldPartitionStreamingGenerationContext.h b/Engine/Source/Runtime/Engine/Public/WorldPartition/WorldPartitionStreamingGenerationContext.h
--- a/Engine/Source/Runtime/Engine/Public/WorldPartition/WorldPartitionStreamingGenerationContext.h
+++ b/Engine/Source/Runtime/Engine/Public/WorldPartition/WorldPartitionStreamingGenerationContext.h
@@ -89,6 +89,5 @@ public:
 	virtual const FActorSetContainer* GetMainWorldContainer() const = 0;
 	virtual void ForEachActorSetInstance(TFunctionRef<void(const FActorSetInstance&)> Func) const = 0;
 	virtual void ForEachActorSetContainer(TFunctionRef<void(const FActorSetContainer&)> Func) const = 0;
-	virtual const FWorldPartitionRuntimeContainerResolver& GetContainerResolver() const = 0;
 };
 #endif
\ No newline at end of file
