Add new GameFeatureAction to add dynamic submixes

#tests PIE

diff --git a/Engine/Source/Runtime/Engine/Classes/Sound/SoundSubmix.h b/Engine/Source/Runtime/Engine/Classes/Sound/SoundSubmix.h
--- a/Engine/Source/Runtime/Engine/Classes/Sound/SoundSubmix.h
+++ b/Engine/Source/Runtime/Engine/Classes/Sound/SoundSubmix.h
@@ -3,17 +3,18 @@
 
 #include "CoreMinimal.h"
 
-#include "ISoundfieldFormat.h"
+#include "AudioDeviceHandle.h"
+#include "AudioLinkSettingsAbstract.h"
+#include "DSP/SpectrumAnalyzer.h"
 #include "IAudioEndpoint.h"
 #include "ISoundfieldEndpoint.h"
+#include "ISoundfieldFormat.h"
 #include "SampleBufferIO.h"
 #include "SoundEffectSubmix.h"
 #include "SoundModulationDestination.h"
 #include "SoundSubmixSend.h"
 #include "UObject/Object.h"
 #include "UObject/ObjectMacros.h"
-#include "DSP/SpectrumAnalyzer.h"
-#include "AudioLinkSettingsAbstract.h"
 
 #include "SoundSubmix.generated.h"
 
@@ -256,11 +257,13 @@ public:
 	TMap<uint32, TObjectPtr<USoundSubmixBase>> DynamicParentSubmix;
 
 	UFUNCTION(BlueprintCallable, Category = "Audio", meta = (WorldContext = "WorldContextObject", DisplayName = "Connect"))
-	ENGINE_API bool DynamicConnect(const UObject* WorldContextObject, USoundSubmixBase* Parent);
-	
+	ENGINE_API bool DynamicConnect(const UObject* WorldContextObject, USoundSubmixBase* InParent);
+	ENGINE_API bool DynamicConnect(const FAudioDeviceHandle& Handle, USoundSubmixBase* InParent);
+
 	UFUNCTION(BlueprintCallable, Category = "Audio", meta = (WorldContext = "WorldContextObject", DisplayName = "Disconnect"))
 	ENGINE_API bool DynamicDisconnect(const UObject* WorldContextObject);
-	
+	ENGINE_API bool DynamicDisconnect(const FAudioDeviceHandle& Handle);
+
 	/**
 	* Set the parent submix of this SoundSubmix, removing it as a child from its previous owner
 	*
diff --git a/Engine/Source/Runtime/Engine/Private/SoundSubmix.cpp b/Engine/Source/Runtime/Engine/Private/SoundSubmix.cpp
--- a/Engine/Source/Runtime/Engine/Private/SoundSubmix.cpp
+++ b/Engine/Source/Runtime/Engine/Private/SoundSubmix.cpp
@@ -647,15 +647,19 @@ bool USoundSubmixWithParentBase::DynamicConnect(const UObject* WorldContextObjec
 		return false;
 	}
 
-	FAudioDeviceHandle Handle = World->GetAudioDevice();
+	return DynamicConnect(World->GetAudioDevice(), InParent);
+}
+
+bool USoundSubmixWithParentBase::DynamicConnect(const FAudioDeviceHandle& Handle, USoundSubmixBase* InParent)
+{
 	if (!Handle.IsValid())
 	{
-		UE_LOG(LogAudio, Warning, TEXT("Submix (DynamicConnect): No valid audio device in this world for [%s]"), *GetName());				
+		UE_LOG(LogAudio, Warning, TEXT("Submix (DynamicConnect): No valid audio device in this world for [%s]"), *GetName());
 		return false;
 	}
 
 	const Audio::DeviceID Id = Handle.GetDeviceID();
-	
+
 	TObjectPtr<USoundSubmixBase>& CurrentParent = DynamicParentSubmix.FindOrAdd(Id);
 
 	if (CurrentParent != InParent)
@@ -674,15 +678,20 @@ bool USoundSubmixWithParentBase::DynamicConnect(const UObject* WorldContextObjec
 			UE_LOG(LogAudio, Verbose, TEXT("Submix (DynamicConnect): Reparenting [%s] and adding to it's new parent [%s]"), *GetName(), *CurrentParent->GetName());
 		}
 
-		Handle->RegisterSoundSubmix(this, true);
-		UE_LOG(LogAudio, Verbose, TEXT("Submix (DynamicConnect): Registering [%s] with World's [%s] AudioDevice [%u]"), *GetName(), *World->GetName(), Handle.GetDeviceID());
+		if (FAudioDevice* RawDevice = Handle.GetAudioDevice())
+		{
+			RawDevice->RegisterSoundSubmix(this, true);
+			UE_LOG(LogAudio, Verbose, TEXT("Submix (DynamicConnect): Registering [%s] with AudioDevice [%u]"), *GetName(), Id);
 
-		return CurrentParent != nullptr;
+			return CurrentParent != nullptr;
+		}
 	}
 
 	UE_LOG(LogAudio, Warning, TEXT("Submix (DynamicConnect): Submix [%s] was already connected to [%s]"), *GetName(), *GetNameSafe(CurrentParent));
 	return false;
+
 }
+
 bool USoundSubmixWithParentBase::DynamicDisconnect(const UObject* WorldContextObject)
 {
 	if (!WorldContextObject)
@@ -698,10 +707,14 @@ bool USoundSubmixWithParentBase::DynamicDisconnect(const UObject* WorldContextOb
 		return false;
 	}
 
-	FAudioDeviceHandle Handle = World->GetAudioDevice();
+	return DynamicDisconnect(World->GetAudioDevice());
+}
+
+bool USoundSubmixWithParentBase::DynamicDisconnect(const FAudioDeviceHandle& Handle)
+{
 	if (!Handle.IsValid())
 	{
-		UE_LOG(LogAudio, Warning, TEXT("Submix (DynamicDisconnect): No valid audio device in this world for [%s]"), *GetName());				
+		UE_LOG(LogAudio, Warning, TEXT("Submix (DynamicDisconnect): No valid audio device in this world for [%s]"), *GetName());
 		return false;
 	}
 
@@ -712,16 +725,17 @@ bool USoundSubmixWithParentBase::DynamicDisconnect(const UObject* WorldContextOb
 	if (CurrentParent)
 	{
 		UE_LOG(LogAudio, Verbose, TEXT("Submix (DynamicDisconnect): Removing [%s] from it's parent [%s]"), *GetName(), *CurrentParent->GetName());
-		
+
 		CurrentParent->DynamicChildSubmixes.FindOrAdd(Id).ChildSubmixes.Remove(this);
 		CurrentParent = nullptr;
 
-		Handle->UnregisterSoundSubmix(this);
-
-		UE_LOG(LogAudio, Verbose, TEXT("Submix (DynamicDisconnect): Unregistering [%s] with World's [%s] AudioDevice [%u]"), 
-		       *GetName(), *World->GetName(), Handle.GetDeviceID());
+		if (FAudioDevice* RawDevice = Handle.GetAudioDevice())
+		{
+			RawDevice->UnregisterSoundSubmix(this);
+			UE_LOG(LogAudio, Verbose, TEXT("Submix (DynamicDisconnect): Unregistering [%s] with AudioDevice [%u]"), *GetName(), Id);
 
-		return true;
+			return true;
+		}
 	}
 
 	UE_LOG(LogAudio, Warning, TEXT("Submix (DynamicDisconnect): Submix was not connected to any dynamic parent [%s]"), *GetName());
