Control Rig: Add editor only AssetUserData array to ActorComponent, SkeeletalMesh, Skeleton and RigVMHost

diff --git a/Engine/Plugins/Animation/ControlRig/Source/ControlRig/Private/ControlRig.cpp b/Engine/Plugins/Animation/ControlRig/Source/ControlRig/Private/ControlRig.cpp
--- a/Engine/Plugins/Animation/ControlRig/Source/ControlRig/Private/ControlRig.cpp
+++ b/Engine/Plugins/Animation/ControlRig/Source/ControlRig/Private/ControlRig.cpp
@@ -2779,18 +2779,26 @@ const TArray<UAssetUserData*>* UControlRig::GetAssetUserDataArray() const
 {
 	if(HasAnyFlags(RF_ClassDefaultObject))
 	{
+#if WITH_EDITOR
+		return &ToRawPtrTArrayUnsafe(CombinedAssetUserData);
+#else
 		return &ToRawPtrTArrayUnsafe(AssetUserData);
+#endif
 	}
 
 	CombinedAssetUserData.Reset();
 
 	if (UControlRig* CDO = Cast<UControlRig>(GetClass()->GetDefaultObject(false)))
 	{
-		CombinedAssetUserData.Append(CDO->AssetUserData);
+		CombinedAssetUserData.Append(*CDO->GetAssetUserDataArray());
 	}
 	else
 	{
+#if WITH_EDITOR
+		CombinedAssetUserData.Append(CombinedAssetUserData);
+#else
 		CombinedAssetUserData.Append(AssetUserData);
+#endif
 	}
 	
 	if(GetExternalAssetUserDataDelegate.IsBound())
diff --git a/Engine/Plugins/Runtime/RigVM/Source/RigVM/Private/RigVMHost.cpp b/Engine/Plugins/Runtime/RigVM/Source/RigVM/Private/RigVMHost.cpp
--- a/Engine/Plugins/Runtime/RigVM/Source/RigVM/Private/RigVMHost.cpp
+++ b/Engine/Plugins/Runtime/RigVM/Source/RigVM/Private/RigVMHost.cpp
@@ -161,6 +161,8 @@ void URigVMHost::PostLoad()
 		VMSnapshotBeforeExecution->SetFlags(VMSnapshotBeforeExecution->GetFlags() | RF_Transient);
 	}
 #endif
+
+	RecomputeCombinedAssetUserData();
 }
 
 void URigVMHost::PreSave(FObjectPreSaveContext SaveContext)
@@ -618,6 +620,13 @@ bool URigVMHost::DisableExecution()
 	return CVarRigVMDisableExecutionAll->GetInt() == 1;
 }
 
+#if WITH_EDITOR
+void URigVMHost::PostEditChangeProperty(FPropertyChangedEvent& PropertyChangedEvent)
+{
+	RecomputeCombinedAssetUserData();
+}
+#endif
+
 bool URigVMHost::InitializeCDOVM()
 {
 	check(VM != nullptr);
@@ -1394,13 +1403,24 @@ void URigVMHost::AddAssetUserData(UAssetUserData* InUserData)
 {
 	if (InUserData != NULL)
 	{
-		UAssetUserData* ExistingData = GetAssetUserDataOfClass(InUserData->GetClass());
-		if (ExistingData != NULL)
+		RemoveUserDataOfClass(InUserData->GetClass());
+		AssetUserData.Add(InUserData);
+		RecomputeCombinedAssetUserData();
+	}
+}
+
+UAssetUserData* URigVMHost::GetAssetUserDataOfClass(TSubclassOf<UAssetUserData> InUserDataClass)
+{
+	const TArray<UAssetUserData*>* ArrayPtr = GetAssetUserDataArray();
+	for (int32 DataIdx = 0; DataIdx < ArrayPtr->Num(); DataIdx++)
+	{
+		UAssetUserData* Datum = (*ArrayPtr)[DataIdx];
+		if (Datum != NULL && Datum->IsA(InUserDataClass))
 		{
-			AssetUserData.Remove(ExistingData);
+			return Datum;
 		}
-		AssetUserData.Add(InUserData);
 	}
+	return NULL;
 }
 
 void URigVMHost::RemoveUserDataOfClass(TSubclassOf<UAssetUserData> InUserDataClass)
@@ -1411,27 +1431,40 @@ void URigVMHost::RemoveUserDataOfClass(TSubclassOf<UAssetUserData> InUserDataCla
 		if (Datum != NULL && Datum->IsA(InUserDataClass))
 		{
 			AssetUserData.RemoveAt(DataIdx);
+			RecomputeCombinedAssetUserData();
 			return;
 		}
 	}
-}
-
-UAssetUserData* URigVMHost::GetAssetUserDataOfClass(TSubclassOf<UAssetUserData> InUserDataClass)
-{
-	for (int32 DataIdx = 0; DataIdx < AssetUserData.Num(); DataIdx++)
+#if WITH_EDITOR
+	for (int32 DataIdx = 0; DataIdx < AssetUserDataEditorOnly.Num(); DataIdx++)
 	{
-		UAssetUserData* Datum = AssetUserData[DataIdx];
+		UAssetUserData* Datum = AssetUserDataEditorOnly[DataIdx];
 		if (Datum != NULL && Datum->IsA(InUserDataClass))
 		{
-			return Datum;
+			AssetUserDataEditorOnly.RemoveAt(DataIdx);
+			RecomputeCombinedAssetUserData();
+			return;
 		}
 	}
-	return NULL;
+#endif
 }
 
 const TArray<UAssetUserData*>* URigVMHost::GetAssetUserDataArray() const
 {
+#if WITH_EDITOR
+	return &ToRawPtrTArrayUnsafe(CombinedAssetUserData);
+#else
 	return &ToRawPtrTArrayUnsafe(AssetUserData);
+#endif
+}
+
+void URigVMHost::RecomputeCombinedAssetUserData()
+{
+	CombinedAssetUserData.Reset();
+	CombinedAssetUserData.Append(AssetUserData);
+#if WITH_EDITOR
+	CombinedAssetUserData.Append(AssetUserDataEditorOnly);
+#endif
 }
 
 #if WITH_EDITOR	
diff --git a/Engine/Plugins/Runtime/RigVM/Source/RigVM/Public/RigVMHost.h b/Engine/Plugins/Runtime/RigVM/Source/RigVM/Public/RigVMHost.h
--- a/Engine/Plugins/Runtime/RigVM/Source/RigVM/Public/RigVMHost.h
+++ b/Engine/Plugins/Runtime/RigVM/Source/RigVM/Public/RigVMHost.h
@@ -159,6 +159,8 @@ public:
 	void BroadCastEndLoadPackage() { EndLoadPackageEvent.Broadcast(this); }
 	FOnEndLoadPackage& OnEndLoadPackage() { return EndLoadPackageEvent; }
 
+	virtual void PostEditChangeProperty(FPropertyChangedEvent& PropertyChangedEvent) override;
+
 private:
 	FOnEndLoadPackage EndLoadPackageEvent;
 
@@ -397,6 +399,16 @@ protected:
 	UPROPERTY(EditAnywhere, AdvancedDisplay, Instanced, Category = "Default")
 	TArray<TObjectPtr<UAssetUserData>> AssetUserData;
 
+#if WITH_EDITORONLY_DATA
+	/** Array of user data stored with the asset */
+	UPROPERTY(EditAnywhere, AdvancedDisplay, Instanced, Category = Skeleton)
+	TArray<TObjectPtr<UAssetUserData>> AssetUserDataEditorOnly;
+#endif
+
+	void RecomputeCombinedAssetUserData();
+	mutable TArray<TObjectPtr<UAssetUserData>> CombinedAssetUserData;
+
+
 protected:
 	bool bRequiresInitExecution;
 
diff --git a/Engine/Source/Runtime/Engine/Classes/Animation/Skeleton.h b/Engine/Source/Runtime/Engine/Classes/Animation/Skeleton.h
--- a/Engine/Source/Runtime/Engine/Classes/Animation/Skeleton.h
+++ b/Engine/Source/Runtime/Engine/Classes/Animation/Skeleton.h
@@ -976,6 +976,8 @@ protected:
 	/** Call this when the skeleton has changed to fix dependent assets */
 	ENGINE_API void HandleSkeletonHierarchyChange(bool bShowProgress = true);
 
+	ENGINE_API virtual void PostEditChangeProperty(FPropertyChangedEvent& PropertyChangedEvent) override;
+
 public:
 	typedef FOnSkeletonHierarchyChangedMulticaster::FDelegate FOnSkeletonHierarchyChanged;
 
@@ -1028,6 +1030,15 @@ protected:
 	UPROPERTY(EditAnywhere, AdvancedDisplay, Instanced, Category = Skeleton)
 	TArray<TObjectPtr<UAssetUserData>> AssetUserData;
 
+#if WITH_EDITORONLY_DATA
+	/** Array of user data stored with the asset */
+	UPROPERTY(EditAnywhere, AdvancedDisplay, Instanced, Category = Skeleton)
+	TArray<TObjectPtr<UAssetUserData>> AssetUserDataEditorOnly;
+#endif
+
+	void RecomputeCombinedAssetUserData();
+	mutable TArray<TObjectPtr<UAssetUserData>> CombinedAssetUserData;
+
 	friend struct FReferenceSkeletonModifier;
 	friend class FEditableSkeleton;
 };
diff --git a/Engine/Source/Runtime/Engine/Classes/Components/ActorComponent.h b/Engine/Source/Runtime/Engine/Classes/Components/ActorComponent.h
--- a/Engine/Source/Runtime/Engine/Classes/Components/ActorComponent.h
+++ b/Engine/Source/Runtime/Engine/Classes/Components/ActorComponent.h
@@ -157,6 +157,15 @@ protected:
 	UPROPERTY(EditAnywhere, AdvancedDisplay, Instanced, Category = AssetUserData)
 	TArray<TObjectPtr<UAssetUserData>> AssetUserData;
 
+#if WITH_EDITORONLY_DATA
+	/** Array of user data stored with the component */
+	UPROPERTY(EditAnywhere, AdvancedDisplay, Instanced, Category = AssetUserData)
+	TArray<TObjectPtr<UAssetUserData>> AssetUserDataEditorOnly;
+#endif
+
+	void RecomputeCombinedAssetUserData();
+	mutable TArray<TObjectPtr<UAssetUserData>> CombinedAssetUserData;
+
 private:
 	/** Used for fast removal of end of frame update */
 	int32 MarkedForEndOfFrameUpdateArrayIndex;
diff --git a/Engine/Source/Runtime/Engine/Classes/Engine/SkeletalMesh.h b/Engine/Source/Runtime/Engine/Classes/Engine/SkeletalMesh.h
--- a/Engine/Source/Runtime/Engine/Classes/Engine/SkeletalMesh.h
+++ b/Engine/Source/Runtime/Engine/Classes/Engine/SkeletalMesh.h
@@ -2222,10 +2222,18 @@ protected:
 	UPROPERTY(EditAnywhere, AdvancedDisplay, Instanced, Category=SkeletalMesh)
 	TArray<TObjectPtr<UAssetUserData>> AssetUserData;
 
+#if WITH_EDITORONLY_DATA
+	/** Array of user data stored with the asset */
+	UPROPERTY(EditAnywhere, AdvancedDisplay, Instanced, Category=SkeletalMesh)
+	TArray<TObjectPtr<UAssetUserData>> AssetUserDataEditorOnly;
+#endif
 #if WITH_EDITOR
 	FOnMeshChanged OnMeshChanged;
 #endif
 
+	void RecomputeCombinedAssetUserData();
+	mutable TArray<TObjectPtr<UAssetUserData>> CombinedAssetUserData;
+
 	friend struct FSkeletalMeshUpdateContext;
 	friend class FSkeletalMeshUpdate;
 
diff --git a/Engine/Source/Runtime/Engine/Private/Animation/Skeleton.cpp b/Engine/Source/Runtime/Engine/Private/Animation/Skeleton.cpp
--- a/Engine/Source/Runtime/Engine/Private/Animation/Skeleton.cpp
+++ b/Engine/Source/Runtime/Engine/Private/Animation/Skeleton.cpp
@@ -396,6 +396,8 @@ void USkeleton::PostLoad()
 
 	// refresh linked bone indices
 	RefreshSkeletonMetaData();
+
+	RecomputeCombinedAssetUserData();
 }
 
 void USkeleton::PostDuplicate(bool bDuplicateForPIE)
@@ -1482,6 +1484,11 @@ void USkeleton::HandleSkeletonHierarchyChange(bool bShowProgress /*= true*/)
 	OnSkeletonHierarchyChanged.Broadcast();
 }
 
+void USkeleton::PostEditChangeProperty(FPropertyChangedEvent& PropertyChangedEvent)
+{
+	RecomputeCombinedAssetUserData();
+}
+
 void USkeleton::RegisterOnSkeletonHierarchyChanged(const FOnSkeletonHierarchyChanged& Delegate)
 {
 	OnSkeletonHierarchyChanged.Add(Delegate);
@@ -2010,7 +2017,7 @@ void USkeleton::GetAssetRegistryTags(TArray<FAssetRegistryTag>& OutTags) const
 	OutTags.Add(FAssetRegistryTag(USkeleton::AnimSyncMarkerTag, SyncMarkersBuilder.ToString(), FAssetRegistryTag::TT_Hidden));
 	
 	// Allow asset user data to output tags
-	for(UAssetUserData* AssetUserDataItem : AssetUserData)
+	for(UAssetUserData* AssetUserDataItem : *GetAssetUserDataArray())
 	{
 		AssetUserDataItem->GetAssetRegistryTags(OutTags);
 	}
@@ -2069,24 +2076,22 @@ USkeletalMeshSocket* USkeleton::FindSocketAndIndex(FName InSocketName, int32& Ou
 }
 
 
-void USkeleton::AddAssetUserData(UAssetUserData* InUserData)
+void USkeleton::AddAssetUserData( UAssetUserData* InUserData)
 {
 	if (InUserData != NULL)
 	{
-		UAssetUserData* ExistingData = GetAssetUserDataOfClass(InUserData->GetClass());
-		if (ExistingData != NULL)
-		{
-			AssetUserData.Remove(ExistingData);
-		}
+		RemoveUserDataOfClass(InUserData->GetClass());
 		AssetUserData.Add(InUserData);
+		RecomputeCombinedAssetUserData();
 	}
 }
 
 UAssetUserData* USkeleton::GetAssetUserDataOfClass(TSubclassOf<UAssetUserData> InUserDataClass)
 {
-	for (int32 DataIdx = 0; DataIdx < AssetUserData.Num(); DataIdx++)
+	const TArray<UAssetUserData*>* ArrayPtr = GetAssetUserDataArray();
+	for (int32 DataIdx = 0; DataIdx < ArrayPtr->Num(); DataIdx++)
 	{
-		UAssetUserData* Datum = AssetUserData[DataIdx];
+		UAssetUserData* Datum = (*ArrayPtr)[DataIdx];
 		if (Datum != NULL && Datum->IsA(InUserDataClass))
 		{
 			return Datum;
@@ -2103,14 +2108,41 @@ void USkeleton::RemoveUserDataOfClass(TSubclassOf<UAssetUserData> InUserDataClas
 		if (Datum != NULL && Datum->IsA(InUserDataClass))
 		{
 			AssetUserData.RemoveAt(DataIdx);
+			RecomputeCombinedAssetUserData();
 			return;
 		}
 	}
+
+#if WITH_EDITOR
+	for (int32 DataIdx = 0; DataIdx < AssetUserDataEditorOnly.Num(); DataIdx++)
+	{
+		UAssetUserData* Datum = AssetUserDataEditorOnly[DataIdx];
+		if (Datum != NULL && Datum->IsA(InUserDataClass))
+		{
+			AssetUserDataEditorOnly.RemoveAt(DataIdx);
+			RecomputeCombinedAssetUserData();
+			return;
+		}
+	}
+#endif
 }
 
 const TArray<UAssetUserData*>* USkeleton::GetAssetUserDataArray() const
 {
+#if WITH_EDITOR
+	return &ToRawPtrTArrayUnsafe(CombinedAssetUserData);
+#else
 	return &ToRawPtrTArrayUnsafe(AssetUserData);
+#endif
+}
+
+void USkeleton::RecomputeCombinedAssetUserData()
+{
+	CombinedAssetUserData.Reset();
+	CombinedAssetUserData.Append(AssetUserData);
+#if WITH_EDITOR
+	CombinedAssetUserData.Append(AssetUserDataEditorOnly);
+#endif
 }
 
 void USkeleton::HandlePackageReloaded(const EPackageReloadPhase InPackageReloadPhase, FPackageReloadedEvent* InPackageReloadedEvent)
diff --git a/Engine/Source/Runtime/Engine/Private/Components/ActorComponent.cpp b/Engine/Source/Runtime/Engine/Private/Components/ActorComponent.cpp
--- a/Engine/Source/Runtime/Engine/Private/Components/ActorComponent.cpp
+++ b/Engine/Source/Runtime/Engine/Private/Components/ActorComponent.cpp
@@ -391,6 +391,15 @@ UActorComponent::UActorComponent(const FObjectInitializer& ObjectInitializer /*=
 	bReplicateUsingRegisteredSubObjectList = GDefaultUseSubObjectReplicationList;
 }
 
+void UActorComponent::RecomputeCombinedAssetUserData()
+{
+	CombinedAssetUserData.Reset();
+	CombinedAssetUserData.Append(AssetUserData);
+#if WITH_EDITOR
+	CombinedAssetUserData.Append(AssetUserDataEditorOnly);
+#endif
+}
+
 void UActorComponent::PostInitProperties()
 {
 	Super::PostInitProperties();
@@ -489,6 +498,8 @@ void UActorComponent::PostLoad()
 		bMarkPendingKillOnPostLoad = false;
 	}
 #endif // WITH_EDITOR
+
+	RecomputeCombinedAssetUserData();
 }
 
 bool UActorComponent::Rename( const TCHAR* InName, UObject* NewOuter, ERenameFlags Flags )
@@ -1015,6 +1026,14 @@ void UActorComponent::ConsolidatedPostEditChange(const FPropertyChangedEvent& Pr
 			Datum->PostEditChangeOwner();
 		}
 	}
+	for (UAssetUserData* Datum : AssetUserDataEditorOnly)
+	{
+		if (Datum != nullptr)
+		{
+			Datum->PostEditChangeOwner();
+		}
+	}
+	RecomputeCombinedAssetUserData();
 }
 
 void UActorComponent::PostEditChangeProperty(FPropertyChangedEvent& PropertyChangedEvent)
@@ -2055,24 +2074,22 @@ bool UActorComponent::IsOwnerRunningUserConstructionScript() const
 	return (MyOwner && MyOwner->IsRunningUserConstructionScript());
 }
 
-void UActorComponent::AddAssetUserData(UAssetUserData* InUserData)
+void UActorComponent::AddAssetUserData( UAssetUserData* InUserData)
 {
 	if (InUserData != NULL)
 	{
-		UAssetUserData* ExistingData = GetAssetUserDataOfClass(InUserData->GetClass());
-		if (ExistingData != NULL)
-		{
-			AssetUserData.Remove(ExistingData);
-		}
+		RemoveUserDataOfClass(InUserData->GetClass());
 		AssetUserData.Add(InUserData);
+		RecomputeCombinedAssetUserData();
 	}
 }
 
 UAssetUserData* UActorComponent::GetAssetUserDataOfClass(TSubclassOf<UAssetUserData> InUserDataClass)
 {
-	for (int32 DataIdx = 0; DataIdx < AssetUserData.Num(); DataIdx++)
+	const TArray<UAssetUserData*>* ArrayPtr = GetAssetUserDataArray();
+	for (int32 DataIdx = 0; DataIdx < ArrayPtr->Num(); DataIdx++)
 	{
-		UAssetUserData* Datum = AssetUserData[DataIdx];
+		UAssetUserData* Datum = (*ArrayPtr)[DataIdx];
 		if (Datum != NULL && Datum->IsA(InUserDataClass))
 		{
 			return Datum;
@@ -2083,7 +2100,11 @@ UAssetUserData* UActorComponent::GetAssetUserDataOfClass(TSubclassOf<UAssetUserD
 
 const TArray<UAssetUserData*>* UActorComponent::GetAssetUserDataArray() const
 {
+#if WITH_EDITOR
+	return &ToRawPtrTArrayUnsafe(CombinedAssetUserData);
+#else
 	return &ToRawPtrTArrayUnsafe(AssetUserData);
+#endif
 }
 
 void UActorComponent::RemoveUserDataOfClass(TSubclassOf<UAssetUserData> InUserDataClass)
@@ -2097,6 +2118,17 @@ void UActorComponent::RemoveUserDataOfClass(TSubclassOf<UAssetUserData> InUserDa
 			return;
 		}
 	}
+#if WITH_EDITOR
+	for (int32 DataIdx = 0; DataIdx < AssetUserDataEditorOnly.Num(); DataIdx++)
+	{
+		UAssetUserData* Datum = AssetUserDataEditorOnly[DataIdx];
+		if (Datum != NULL && Datum->IsA(InUserDataClass))
+		{
+			AssetUserDataEditorOnly.RemoveAt(DataIdx);
+			return;
+		}
+	}
+#endif
 }
 
 void UActorComponent::OnCreatedFromReplication()
diff --git a/Engine/Source/Runtime/Engine/Private/SkeletalMesh.cpp b/Engine/Source/Runtime/Engine/Private/SkeletalMesh.cpp
--- a/Engine/Source/Runtime/Engine/Private/SkeletalMesh.cpp
+++ b/Engine/Source/Runtime/Engine/Private/SkeletalMesh.cpp
@@ -818,6 +818,15 @@ bool USkeletalMesh::NeedCPUData(int32 LODIndex) const
 	return GetSamplingInfo().IsSamplingEnabled(this, LODIndex);
 }
 
+void USkeletalMesh::RecomputeCombinedAssetUserData()
+{
+	CombinedAssetUserData.Reset();
+	CombinedAssetUserData.Append(AssetUserData);
+#if WITH_EDITOR
+	CombinedAssetUserData.Append(AssetUserDataEditorOnly);
+#endif
+}
+
 void USkeletalMesh::InitResources()
 {
 	LLM_SCOPE_BYNAME(TEXT("SkeletalMesh/InitResources")); // This is an important test case for SCOPE_BYNAME without a matching LLM_DEFINE_TAG
@@ -1299,6 +1308,14 @@ void USkeletalMesh::PostEditChangeProperty(FPropertyChangedEvent& PropertyChange
 			Datum->PostEditChangeOwner();
 		}
 	}
+	for (UAssetUserData* Datum : AssetUserDataEditorOnly)
+	{
+		if (Datum != nullptr)
+		{
+			Datum->PostEditChangeOwner();
+		}
+	}
+	RecomputeCombinedAssetUserData();
 
 	Super::PostEditChangeProperty(PropertyChangedEvent);
 
@@ -3145,6 +3162,8 @@ void USkeletalMesh::PostLoad()
 {
 	LLM_SCOPE(ELLMTag::SkeletalMesh);
 	Super::PostLoad();
+
+	RecomputeCombinedAssetUserData();
 }
 
 void USkeletalMesh::ExecutePostLoadInternal(FSkinnedAssetPostLoadContext& Context)
@@ -3508,7 +3527,7 @@ void USkeletalMesh::GetAssetRegistryTags(TArray<FAssetRegistryTag>& OutTags) con
 	OutTags.Add(FAssetRegistryTag("MaxBoneInfluences", MaxBoneInfluencesString, FAssetRegistryTag::TT_Numerical));
 
 	// Allow asset user data to output tags
-	for(UAssetUserData* AssetUserDataItem : AssetUserData)
+	for(UAssetUserData* AssetUserDataItem : *GetAssetUserDataArray())
 	{
 		if(AssetUserDataItem)
 		{
@@ -4710,24 +4729,22 @@ bool USkeletalMesh::GetPhysicsTriMeshData(FTriMeshCollisionData* CollisionData,
 #endif // #if WITH_EDITORONLY_DATA
 }
 
-void USkeletalMesh::AddAssetUserData(UAssetUserData* InUserData)
+void USkeletalMesh::AddAssetUserData( UAssetUserData* InUserData)
 {
 	if (InUserData != NULL)
 	{
-		UAssetUserData* ExistingData = GetAssetUserDataOfClass(InUserData->GetClass());
-		if (ExistingData != NULL)
-		{
-			AssetUserData.Remove(ExistingData);
-		}
+		RemoveUserDataOfClass(InUserData->GetClass());
 		AssetUserData.Add(InUserData);
+		RecomputeCombinedAssetUserData();
 	}
 }
 
 UAssetUserData* USkeletalMesh::GetAssetUserDataOfClass(TSubclassOf<UAssetUserData> InUserDataClass)
 {
-	for (int32 DataIdx = 0; DataIdx < AssetUserData.Num(); DataIdx++)
+	const TArray<UAssetUserData*>* ArrayPtr = GetAssetUserDataArray();
+	for (int32 DataIdx = 0; DataIdx < ArrayPtr->Num(); DataIdx++)
 	{
-		UAssetUserData* Datum = AssetUserData[DataIdx];
+		UAssetUserData* Datum = (*ArrayPtr)[DataIdx];
 		if (Datum != NULL && Datum->IsA(InUserDataClass))
 		{
 			return Datum;
@@ -4744,14 +4761,31 @@ void USkeletalMesh::RemoveUserDataOfClass(TSubclassOf<UAssetUserData> InUserData
 		if (Datum != NULL && Datum->IsA(InUserDataClass))
 		{
 			AssetUserData.RemoveAt(DataIdx);
+			RecomputeCombinedAssetUserData();
 			return;
 		}
 	}
+#if WITH_EDITOR
+	for (int32 DataIdx = 0; DataIdx < AssetUserDataEditorOnly.Num(); DataIdx++)
+	{
+		UAssetUserData* Datum = AssetUserDataEditorOnly[DataIdx];
+		if (Datum != NULL && Datum->IsA(InUserDataClass))
+		{
+			AssetUserDataEditorOnly.RemoveAt(DataIdx);
+			RecomputeCombinedAssetUserData();
+			return;
+		}
+	}
+#endif
 }
 
 const TArray<UAssetUserData*>* USkeletalMesh::GetAssetUserDataArray() const
 {
+#if WITH_EDITOR
+	return &ToRawPtrTArrayUnsafe(CombinedAssetUserData);
+#else
 	return &ToRawPtrTArrayUnsafe(AssetUserData);
+#endif
 }
 
 ////// SKELETAL MESH THUMBNAIL SUPPORT ////////
