ShaderCompiler - Change the scoped slow task to use a delayed dialog so that dialogs aren't shown for compiles less than 1 second.

diff --git a/Engine/Source/Runtime/Engine/Private/ShaderCompiler/ShaderCompiler.cpp b/Engine/Source/Runtime/Engine/Private/ShaderCompiler/ShaderCompiler.cpp
--- a/Engine/Source/Runtime/Engine/Private/ShaderCompiler/ShaderCompiler.cpp
+++ b/Engine/Source/Runtime/Engine/Private/ShaderCompiler/ShaderCompiler.cpp
@@ -6019,7 +6019,7 @@ void FShaderCompilingManager::BlockOnShaderMapCompletion(const TArray<int32>& Sh
 		FScopedSlowTask SlowTask(NumJobs, FText::Format(LOCTEXT("BlockOnShaderMapCompletion", "Compiling Shaders ({0})"), NumJobs), GIsEditor && !IsRunningCommandlet() && GPlayInEditorID == INDEX_NONE);
 		if (NumJobs > 0)
 		{
-			SlowTask.MakeDialog(false, true);
+			SlowTask.MakeDialogDelayed(1.0f);
 		}
 
 		int32 NumPendingJobs = 0;
