Add additional spline->polypath conversion functions to support converting within a specified time or distance range, and to also get distance values associated with each point.

diff --git a/Engine/Source/Runtime/Engine/Classes/Components/SplineComponent.h b/Engine/Source/Runtime/Engine/Classes/Components/SplineComponent.h
--- a/Engine/Source/Runtime/Engine/Classes/Components/SplineComponent.h
+++ b/Engine/Source/Runtime/Engine/Classes/Components/SplineComponent.h
@@ -806,6 +806,18 @@ public:
 	UFUNCTION(BlueprintCallable, Category = Spline)
 	ENGINE_API bool ConvertSplineToPolyLine(ESplineCoordinateSpace::Type CoordinateSpace, const float MaxSquareDistanceFromSpline, TArray<FVector>& OutPoints) const;
 
+	/** Given a threshold, returns a list of vertices along the spline that, treated as a list of segments (polyline), matches the spline shape. Also fills a list of corresponding distances along the spline for each point. */
+	UFUNCTION(BlueprintPure = false, Category = Spline)
+	ENGINE_API bool ConvertSplineToPolyLineWithDistances(ESplineCoordinateSpace::Type CoordinateSpace, const float MaxSquareDistanceFromSpline, TArray<FVector>& OutPoints, TArray<double>& OutDistancesAlongSpline) const;
+
+	/** Given a threshold and a start and end distance range, returns a list of vertices along the spline that, treated as a list of segments (polyline), matches the spline shape in that range. Also fills a list of corresponding distances along the spline for each point. */
+	UFUNCTION(BlueprintPure = false, Category = Spline)
+	ENGINE_API bool ConvertSplineToPolyline_InDistanceRange(ESplineCoordinateSpace::Type CoordinateSpace, const float MaxSquareDistanceFromSpline, float StartDistAlongSpline, float EndDistAlongSpline, TArray<FVector>& OutPoints, TArray<double>& OutDistancesAlongSpline, bool bAllowWrappingIfClosed = true) const;
+
+	/** Given a threshold and start and end time range, returns a list of vertices along the spline that, treated as a list of segments (polyline), matches the spline shape in that range. Also fills a list of corresponding distances along the spline for each point. */
+	UFUNCTION(BlueprintPure = false, Category = Spline)
+	ENGINE_API bool ConvertSplineToPolyline_InTimeRange(ESplineCoordinateSpace::Type CoordinateSpace, const float MaxSquareDistanceFromSpline, float StartTimeAlongSpline, float EndTimeAlongSpline, bool bUseConstantVelocity, TArray<FVector>& OutPoints, TArray<double>& OutDistancesAlongSpline, bool bAllowWrappingIfClosed = true) const;
+
 private:
 	/** The dummy value used for queries when there are no point in a spline */
 	static ENGINE_API const FInterpCurvePointVector DummyPointPosition;
diff --git a/Engine/Source/Runtime/Engine/Private/Components/SplineComponent.cpp b/Engine/Source/Runtime/Engine/Private/Components/SplineComponent.cpp
--- a/Engine/Source/Runtime/Engine/Private/Components/SplineComponent.cpp
+++ b/Engine/Source/Runtime/Engine/Private/Components/SplineComponent.cpp
@@ -1728,6 +1728,200 @@ bool USplineComponent::ConvertSplineToPolyLine(ESplineCoordinateSpace::Type Coor
 	return (OutPoints.Num() > 0);
 }
 
+bool USplineComponent::ConvertSplineToPolyLineWithDistances(ESplineCoordinateSpace::Type CoordinateSpace, const float MaxSquareDistanceFromSpline, TArray<FVector>& OutPoints, TArray<double>& OutDistancesAlongSpline) const
+{
+	return ConvertSplineToPolyline_InDistanceRange(CoordinateSpace, MaxSquareDistanceFromSpline, 0, GetSplineLength(), OutPoints, OutDistancesAlongSpline, false);
+}
+
+bool USplineComponent::ConvertSplineToPolyline_InDistanceRange(ESplineCoordinateSpace::Type CoordinateSpace, const float InMaxSquareDistanceFromSpline, float RangeStart, float RangeEnd, TArray<FVector>& OutPoints, TArray<double>& OutDistancesAlongSpline, bool bAllowWrappingIfClosed) const
+{
+	const int32 NumPoints = SplineCurves.Position.Points.Num();
+	if (NumPoints == 0)
+	{
+		return false;
+	}
+	const int32 NumSegments = GetNumberOfSplineSegments();
+
+	float SplineLength = GetSplineLength();
+	if (SplineLength <= 0)
+	{
+		OutPoints.Add(GetLocationAtDistanceAlongSpline(0, CoordinateSpace));
+		OutDistancesAlongSpline.Add(0);
+		return false;
+	}
+
+	// Sanitize the sampling tolerance
+	const float MaxSquareDistanceFromSpline = FMath::Max(UE_SMALL_NUMBER, InMaxSquareDistanceFromSpline);
+
+	// Sanitize range and mark whether the range wraps through 0
+	bool bNeedsWrap = false;
+	if (!bClosedLoop || !bAllowWrappingIfClosed)
+	{
+		RangeStart = FMath::Clamp(RangeStart, 0, SplineLength);
+		RangeEnd = FMath::Clamp(RangeEnd, 0, SplineLength);
+	}
+	else if (RangeStart < 0 || RangeEnd > SplineLength)
+	{
+		bNeedsWrap = true;
+	}
+	if (RangeStart > RangeEnd)
+	{
+		return false;
+	}
+
+	// expect at least 2 points per segment covered
+	int32 EstimatedPoints = 2 * NumSegments * static_cast<int32>((RangeEnd - RangeStart) / SplineLength);
+	OutPoints.Empty();
+	OutPoints.Reserve(EstimatedPoints);
+	OutDistancesAlongSpline.Empty();
+	OutDistancesAlongSpline.Reserve(EstimatedPoints);
+
+	if (RangeStart == RangeEnd)
+	{
+		OutPoints.Add(GetLocationAtDistanceAlongSpline(RangeStart, CoordinateSpace));
+		OutDistancesAlongSpline.Add(RangeStart);
+		return true;
+	}
+
+	// If we need to wrap around, break the wrapped segments into non-wrapped parts and add each part separately
+	if (bNeedsWrap)
+	{
+		float TotalRange = RangeEnd - RangeStart;
+		auto WrapDistance = [SplineLength](float Distance, int32& LoopIdx) -> float
+		{
+			LoopIdx = FMath::FloorToInt32(Distance / SplineLength);
+			float WrappedDistance = FMath::Fmod(Distance, SplineLength);
+			if (WrappedDistance < 0)
+			{
+				WrappedDistance += SplineLength;
+			}
+			return WrappedDistance;
+		};
+		int32 StartLoopIdx, EndLoopIdx;
+		float WrappedStart = WrapDistance(RangeStart, StartLoopIdx);
+		float WrappedEnd = WrapDistance(RangeEnd, EndLoopIdx);
+		float WrappedLoc = WrappedStart;
+		bool bHasAdded = false;
+		for (int32 LoopIdx = StartLoopIdx; LoopIdx <= EndLoopIdx; ++LoopIdx)
+		{
+			if (bHasAdded && ensure(OutPoints.Num()))
+			{
+				OutPoints.RemoveAt(OutPoints.Num() - 1, 1, false);
+				OutDistancesAlongSpline.RemoveAt(OutDistancesAlongSpline.Num() - 1, 1, false);
+			}
+			float EndLoc = LoopIdx == EndLoopIdx ? WrappedEnd : SplineLength;
+
+			TArray<FVector> Points;
+			TArray<double> Distances;
+			ConvertSplineToPolyline_InDistanceRange(CoordinateSpace, MaxSquareDistanceFromSpline, WrappedLoc, EndLoc, Points, Distances, false);
+			OutPoints.Append(Points);
+			OutDistancesAlongSpline.Append(Distances);
+
+			bHasAdded = true;
+			WrappedLoc = 0;
+		}
+		return bHasAdded;
+	} // end of the wrap-around case, after this values will be in the normal range
+	
+	int32 SegmentStart = SplineCurves.ReparamTable.GetPointIndexForInputValue(RangeStart) / ReparamStepsPerSegment;
+	int32 SegmentEnd = FMath::Min(NumSegments, 1 + SplineCurves.ReparamTable.GetPointIndexForInputValue(RangeEnd) / ReparamStepsPerSegment);
+
+	TArray<FVector> NewPoints;
+	TArray<double> NewDistances;
+	for (int32 SegmentIndex = SegmentStart; SegmentIndex < SegmentEnd; ++SegmentIndex)
+	{
+		// Get the segment range as distances, clipped with the input range
+		double StartDist = FMath::Max(RangeStart, GetDistanceAlongSplineAtSplinePoint(SegmentIndex));
+		double StopDist = FMath::Min(RangeEnd, GetDistanceAlongSplineAtSplinePoint(SegmentIndex + 1));
+		bool bIsLast = SegmentIndex + 1 == SegmentEnd;
+
+		const int32 NumLines = 2; // Dichotomic subdivision of the spline segment
+		double Dist = StopDist - StartDist;
+		double SubstepSize = Dist / NumLines;
+		if (SubstepSize == 0.0)
+		{
+			// There is no distance to cover, so handle the segment with a single point (or nothing, if this isn't the very last point)
+			if (bIsLast)
+			{
+				OutPoints.Add(GetLocationAtDistanceAlongSpline(StopDist, CoordinateSpace));
+				OutDistancesAlongSpline.Add(StopDist);
+			}
+			continue;
+		}
+
+		double SubstepStartDist = StartDist;
+		for (int32 i = 0; i < NumLines; ++i)
+		{
+			double SubstepEndDist = SubstepStartDist + SubstepSize;
+			NewPoints.Reset();
+			NewDistances.Reset();
+			// Recursively sub-divide each segment until the requested precision is reached :
+			if (DivideSplineIntoPolylineRecursiveWithDistances(SubstepStartDist, SubstepEndDist, CoordinateSpace, MaxSquareDistanceFromSpline, NewPoints, NewDistances))
+			{
+				if (OutPoints.Num() > 0)
+				{
+					check(OutPoints.Last() == NewPoints[0]); // our last point must be the same as the new segment's first
+					OutPoints.RemoveAt(OutPoints.Num() - 1);
+					OutDistancesAlongSpline.RemoveAt(OutDistancesAlongSpline.Num() - 1);
+				}
+				OutPoints.Append(NewPoints);
+				OutDistancesAlongSpline.Append(NewDistances);
+			}
+
+			SubstepStartDist = SubstepEndDist;
+		}
+	}
+
+	return !OutPoints.IsEmpty();
+}
+
+bool USplineComponent::ConvertSplineToPolyline_InTimeRange(ESplineCoordinateSpace::Type CoordinateSpace, const float MaxSquareDistanceFromSpline, float StartTimeAlongSpline, float EndTimeAlongSpline, bool bUseConstantVelocity, TArray<FVector>& OutPoints, TArray<double>& OutDistancesAlongSpline, bool bAllowWrappingIfClosed) const
+{
+	if (SplineCurves.Position.Points.Num() == 0)
+	{
+		return false;
+	}
+
+	// Helper to convert times to distances, so we can call the distance-based version of this function
+	auto TimeToDistance = [this, bUseConstantVelocity, bAllowWrappingIfClosed](float Time) -> float
+	{
+		float TimeFrac = Time / Duration; // fraction of spline travelled
+		if (bUseConstantVelocity)
+		{
+			return TimeFrac * GetSplineLength();
+		}
+		else
+		{
+			const int32 NumPoints = SplineCurves.Position.Points.Num();
+			const int32 NumSegments = bClosedLoop ? NumPoints : NumPoints - 1;
+			// Note: 'InputKey' values correspond to the spline in parameter space, in the range of 0 to NumSegments
+			float InputKey = TimeFrac * NumSegments;
+			if (bClosedLoop && bAllowWrappingIfClosed)
+			{
+				// Note the GetDistanceAlongSplineAtSplineInputKey() requires values in the 0-NumSegments range
+				// So we wrap (modulus) into that range, find the distance, and then translate back to the original un-wrapped range.
+				float DistanceAtStartOfLoop = FMath::Floor(TimeFrac) * GetSplineLength();
+				float InRangeInputKey = FMath::Fmod(InputKey, NumSegments);
+				if (InRangeInputKey < 0)
+				{
+					InRangeInputKey += NumSegments;
+				}
+				float DistanceWrapped = GetDistanceAlongSplineAtSplineInputKey(InRangeInputKey);
+				return DistanceWrapped + DistanceAtStartOfLoop;
+			}
+			else
+			{
+				// If wrapping is not allowed, clamp to the valid range
+				float ClampedInputKey = FMath::Clamp(InputKey, 0, NumSegments);
+				return GetDistanceAlongSplineAtSplineInputKey(InputKey);
+			}
+		}
+	};
+
+	return ConvertSplineToPolyline_InDistanceRange(CoordinateSpace, MaxSquareDistanceFromSpline, TimeToDistance(StartTimeAlongSpline), TimeToDistance(EndTimeAlongSpline), OutPoints, OutDistancesAlongSpline, bAllowWrappingIfClosed);
+}
+
+
 template<class T>
 T GetPropertyValueAtSplinePoint(const USplineMetadata* Metadata, int32 Index, FName PropertyName)
 {
