WorldPartitionCookPackageSplitter:
- Added logs to track down an intermittent cooking issue.

diff --git a/Engine/Source/Runtime/Engine/Private/WorldPartition/WorldPartitionCookPackageSplitter.cpp b/Engine/Source/Runtime/Engine/Private/WorldPartition/WorldPartitionCookPackageSplitter.cpp
--- a/Engine/Source/Runtime/Engine/Private/WorldPartition/WorldPartitionCookPackageSplitter.cpp
+++ b/Engine/Source/Runtime/Engine/Private/WorldPartition/WorldPartitionCookPackageSplitter.cpp
@@ -32,7 +32,7 @@ FWorldPartitionCookPackageSplitter::~FWorldPartitionCookPackageSplitter()
 
 void FWorldPartitionCookPackageSplitter::Teardown(ETeardown Status)
 {
-	UE_LOG(LogWorldPartition, Log, TEXT("[Cook] Debug(TearDown): OwnerObject=%s bForceInitializedWorld=%d bInitializedPhysicsSceneForSave=%d"), *GetFullNameSafe(ReferencedWorld),  bForceInitializedWorld ? 1 : 0, bInitializedPhysicsSceneForSave ? 1 : 0);
+	UE_LOG(LogWorldPartition, Display, TEXT("[Cook] Debug(TearDown): OwnerObject=%s bForceInitializedWorld=%d bInitializedPhysicsSceneForSave=%d"), *GetFullNameSafe(ReferencedWorld),  bForceInitializedWorld ? 1 : 0, bInitializedPhysicsSceneForSave ? 1 : 0);
 
 	FWorldDelegates::OnWorldCleanup.RemoveAll(this);
 
@@ -55,7 +55,7 @@ void FWorldPartitionCookPackageSplitter::Teardown(ETeardown Status)
 
 void FWorldPartitionCookPackageSplitter::OnWorldCleanup(UWorld* InWorld, bool bSessionEnded, bool bCleanupResources)
 {
-	check(InWorld != ReferencedWorld);
+	checkf(InWorld != ReferencedWorld, TEXT("[Cook] %s is being cleaned up while still referenced by a package splitter."), *GetFullNameSafe(InWorld));
 }
 
 void FWorldPartitionCookPackageSplitter::AddReferencedObjects(FReferenceCollector& Collector)
@@ -92,13 +92,13 @@ TArray<ICookPackageSplitter::FGeneratedPackage> FWorldPartitionCookPackageSplitt
 	bool bIsSuccess = CookContext.GatherPackagesToCook();
 	UE_CLOG(!bIsSuccess, LogWorldPartition, Warning, TEXT("[Cook] Errors while gathering packages to took from generators for owner object %s."), *GetFullNameSafe(OwnerObject));
 
-	UE_LOG(LogWorldPartition, Log, TEXT("[Cook] Gathered %u packages to generate from %u Generators."), CookContext.NumPackageToGenerate(), CookContext.NumGenerators());
+	UE_LOG(LogWorldPartition, Display, TEXT("[Cook] Gathered %u packages to generate from %u Generators."), CookContext.NumPackageToGenerate(), CookContext.NumGenerators());
 
 	TArray<ICookPackageSplitter::FGeneratedPackage> PackagesToGenerate;
 	BuildPackagesToGenerateList(PackagesToGenerate);
 
-	UE_LOG(LogWorldPartition, Log, TEXT("[Cook] Sending %u packages to be generated."), PackagesToGenerate.Num());
-	UE_LOG(LogWorldPartition, Log, TEXT("[Cook] Debug(GetGenerateList) : OwnerObject=%s bForceInitializedWorld=%d bInitializedPhysicsSceneForSave=%d"), *GetFullNameSafe(OwnerObject), bForceInitializedWorld ? 1 : 0, bInitializedPhysicsSceneForSave ? 1 : 0);
+	UE_LOG(LogWorldPartition, Display, TEXT("[Cook] Sending %u packages to be generated."), PackagesToGenerate.Num());
+	UE_LOG(LogWorldPartition, Display, TEXT("[Cook] Debug(GetGenerateList) : OwnerObject=%s bForceInitializedWorld=%d bInitializedPhysicsSceneForSave=%d"), *GetFullNameSafe(OwnerObject), bForceInitializedWorld ? 1 : 0, bInitializedPhysicsSceneForSave ? 1 : 0);
 
 	FWorldDelegates::OnWorldCleanup.AddRaw(this, &FWorldPartitionCookPackageSplitter::OnWorldCleanup);
 
@@ -136,7 +136,7 @@ bool FWorldPartitionCookPackageSplitter::PopulateGeneratorPackage(UPackage* Owne
 	const TArray<ICookPackageSplitter::FGeneratedPackageForPreSave>& GeneratedPackages, TArray<UObject*>& OutObjectsToMove,
 	TArray<UPackage*>& OutModifiedPackages)
 {
-	UE_LOG(LogWorldPartition, Log, TEXT("[Cook][PopulateGeneratorPackage] Processing %u packages"), GeneratedPackages.Num());
+	UE_LOG(LogWorldPartition, Display, TEXT("[Cook][PopulateGeneratorPackage] Processing %u packages"), GeneratedPackages.Num());
 
 	bool bIsSuccess = true;
 	if (GeneratedPackages.Num() != CookContext.NumPackageToGenerate())
@@ -155,7 +155,7 @@ bool FWorldPartitionCookPackageSplitter::PopulateGeneratorPackage(UPackage* Owne
 		}
 	}
 
-	UE_LOG(LogWorldPartition, Log, TEXT("[Cook][PopulateGeneratorPackage] Gathered %u modified packages"), ModifiedPackages.Num());
+	UE_LOG(LogWorldPartition, Display, TEXT("[Cook][PopulateGeneratorPackage] Gathered %u modified packages"), ModifiedPackages.Num());
 	OutModifiedPackages = MoveTemp(ModifiedPackages);
 
 	return bIsSuccess;
