Insights scope

diff --git a/Engine/Source/Runtime/Engine/Private/AssetCompilingManager.cpp b/Engine/Source/Runtime/Engine/Private/AssetCompilingManager.cpp
--- a/Engine/Source/Runtime/Engine/Private/AssetCompilingManager.cpp
+++ b/Engine/Source/Runtime/Engine/Private/AssetCompilingManager.cpp
@@ -475,6 +475,7 @@ int32 FAssetCompilingManager::GetNumRemainingAssets() const
  */
 void FAssetCompilingManager::FinishAllCompilation()
 {
+	TRACE_CPUPROFILER_EVENT_SCOPE_STR("FAssetCompilingManager::FinishAllCompilation");
 	for (IAssetCompilingManager* AssetCompilingManager : AssetCompilingManagers)
 	{
 		AssetCompilingManager->FinishAllCompilation();
