Instanced Static Mesh Component - Show bounds for this type of component when using the Bounds ShowFlag

diff --git a/Engine/Source/Runtime/Engine/Private/InstancedStaticMesh.cpp b/Engine/Source/Runtime/Engine/Private/InstancedStaticMesh.cpp
--- a/Engine/Source/Runtime/Engine/Private/InstancedStaticMesh.cpp
+++ b/Engine/Source/Runtime/Engine/Private/InstancedStaticMesh.cpp
@@ -1492,6 +1492,7 @@ void FInstancedStaticMeshSceneProxy::GetDynamicMeshElements(const TArray<const F
 		if (VisibilityMap & (1 << ViewIndex))
 		{
 			const FSceneView* View = Views[ViewIndex];
+			FPrimitiveDrawInterface* PDI = Collector.GetPDI(ViewIndex);
 
 			for (int32 SelectionGroupIndex = 0; SelectionGroupIndex < NumSelectionGroups; SelectionGroupIndex++)
 			{
@@ -1547,6 +1548,9 @@ void FInstancedStaticMeshSceneProxy::GetDynamicMeshElements(const TArray<const F
 					}
 				}
 			}
+			
+			// Draw the bounds
+			RenderBounds(PDI, View->Family->EngineShowFlags, GetBounds(), IsSelected());
 
 #if !(UE_BUILD_SHIPPING || UE_BUILD_TEST)
 			if (View->Family->EngineShowFlags.VisualizeInstanceUpdates && HasInstanceDebugData())
@@ -1557,11 +1561,11 @@ void FInstancedStaticMeshSceneProxy::GetDynamicMeshElements(const TArray<const F
 					const FInstanceSceneData& Instance = InstanceSceneData[i];
 
 					FRenderTransform InstanceToWorld = Instance.ComputeLocalToWorld(PrimitiveToWorld);
-					DrawWireStar(Collector.GetPDI(ViewIndex), (FVector)InstanceToWorld.Origin, 40.0f, WasInstanceXFormUpdatedThisFrame(i) ? FColor::Red : FColor::Green, View->Family->EngineShowFlags.Game ? SDPG_World : SDPG_Foreground);
+					DrawWireStar(PDI, (FVector)InstanceToWorld.Origin, 40.0f, WasInstanceXFormUpdatedThisFrame(i) ? FColor::Red : FColor::Green, View->Family->EngineShowFlags.Game ? SDPG_World : SDPG_Foreground);
 
 					if (WasInstanceCustomDataUpdatedThisFrame(i))
 					{
-						DrawCircle(Collector.GetPDI(ViewIndex), (FVector)InstanceToWorld.Origin, FVector(1, 0, 0), FVector(0, 1, 0), FColor::Orange, 40.0f, 32, View->Family->EngineShowFlags.Game ? SDPG_World : SDPG_Foreground);
+						DrawCircle(PDI, (FVector)InstanceToWorld.Origin, FVector(1, 0, 0), FVector(0, 1, 0), FColor::Orange, 40.0f, 32, View->Family->EngineShowFlags.Game ? SDPG_World : SDPG_Foreground);
 					}
 				}
 			}
