World Partition - Added the ability to iterate over loaded (referenced) actors from a loader adapter with ILoaderAdapter::ForEachReferencedActor

diff --git a/Engine/Source/Runtime/Engine/Private/WorldPartition/WorldPartitionActorLoaderInterface.cpp b/Engine/Source/Runtime/Engine/Private/WorldPartition/WorldPartitionActorLoaderInterface.cpp
--- a/Engine/Source/Runtime/Engine/Private/WorldPartition/WorldPartitionActorLoaderInterface.cpp
+++ b/Engine/Source/Runtime/Engine/Private/WorldPartition/WorldPartitionActorLoaderInterface.cpp
@@ -240,6 +240,20 @@ void IWorldPartitionActorLoaderInterface::ILoaderAdapter::OnActorDescContainerUn
 	}
 }
 
+void IWorldPartitionActorLoaderInterface::ILoaderAdapter::ForEachReferencedActor(TFunctionRef<void(const FWorldPartitionReference&)> InFunc) const
+{
+	for (const auto& [Container, ActorReferenceMap] : ContainerActorReferences)
+	{
+		for (const auto& [ReferenceGuid, ReferenceMap] : ActorReferenceMap)
+		{
+			for (const auto& [Guid, ActorReference] : ReferenceMap)
+			{
+				InFunc(ActorReference);
+			}
+		}
+	}
+}
+
 bool IWorldPartitionActorLoaderInterface::ILoaderAdapter::ShouldActorBeLoaded(const FWorldPartitionHandle& Actor) const
 {
 	check(Actor.IsValid());
diff --git a/Engine/Source/Runtime/Engine/Public/WorldPartition/WorldPartitionActorLoaderInterface.h b/Engine/Source/Runtime/Engine/Public/WorldPartition/WorldPartitionActorLoaderInterface.h
--- a/Engine/Source/Runtime/Engine/Public/WorldPartition/WorldPartitionActorLoaderInterface.h
+++ b/Engine/Source/Runtime/Engine/Public/WorldPartition/WorldPartitionActorLoaderInterface.h
@@ -48,6 +48,8 @@ public:
 		ENGINE_API void OnActorDescContainerInitialize(UActorDescContainer* Container);
 		ENGINE_API void OnActorDescContainerUninitialize(UActorDescContainer* Container);
 
+		ENGINE_API void ForEachReferencedActor(TFunctionRef<void(const FWorldPartitionReference&)>) const;
+
 	protected:
 		// Private interface
 		virtual void ForEachActor(TFunctionRef<void(const FWorldPartitionHandle&)> InOperation) const =0;
