World Partition - HLOD: Made some virtual methods from FHLODActorDesc public to given them the same visibility as the ones in the base class

diff --git a/Engine/Source/Runtime/Engine/Public/WorldPartition/HLOD/HLODActorDesc.h b/Engine/Source/Runtime/Engine/Public/WorldPartition/HLOD/HLODActorDesc.h
--- a/Engine/Source/Runtime/Engine/Public/WorldPartition/HLOD/HLODActorDesc.h
+++ b/Engine/Source/Runtime/Engine/Public/WorldPartition/HLOD/HLODActorDesc.h
@@ -34,6 +34,12 @@ public:
 	ENGINE_API int64 GetPackageSize() const;
 	static ENGINE_API int64 GetPackageSize(const AWorldPartitionHLOD* InHLODActor);
 
+	//~ Begin FWorldPartitionActorDesc Interface.
+	virtual bool IsRuntimeRelevant(const FActorContainerID& InContainerID) const override { return !bIsForcedNonSpatiallyLoaded; }
+	virtual bool ShouldValidateRuntimeGrid() const override { return false; }
+	virtual FBox GetEditorBounds() const override { return EditorBounds; }
+	//~ End FWorldPartitionActorDesc Interface.
+
 protected:
 	ENGINE_API FHLODActorDesc();
 
@@ -42,9 +48,6 @@ protected:
 	ENGINE_API virtual bool Equals(const FWorldPartitionActorDesc* Other) const override;
 	virtual uint32 GetSizeOf() const override { return sizeof(FHLODActorDesc); }
 	ENGINE_API virtual void Serialize(FArchive& Ar) override;
-	virtual bool IsRuntimeRelevant(const FActorContainerID& InContainerID) const override { return !bIsForcedNonSpatiallyLoaded; }
-	virtual bool ShouldValidateRuntimeGrid() const override { return false; }
-	virtual FBox GetEditorBounds() const override { return EditorBounds; }
 	//~ End FWorldPartitionActorDesc Interface.
 
 	TArray<FGuid> ChildHLODActors;
