PR #10773: Stereo rendering: Make SkeletalMeshes use the same LOD for all views

diff --git a/Engine/Source/Runtime/Engine/Private/SceneManagement.cpp b/Engine/Source/Runtime/Engine/Private/SceneManagement.cpp
--- a/Engine/Source/Runtime/Engine/Private/SceneManagement.cpp
+++ b/Engine/Source/Runtime/Engine/Private/SceneManagement.cpp
@@ -902,7 +902,7 @@ int8 ComputeTemporalStaticMeshLOD( const FStaticMeshRenderData* RenderData, cons
 // Ensure we always use the left eye when selecting lods to avoid divergent selections in stereo
 const FSceneView& GetLODView(const FSceneView& InView)
 {
-	if (IStereoRendering::IsStereoEyeView(InView) && GEngine->StereoRenderingDevice.IsValid())
+	if (UNLIKELY(IStereoRendering::IsStereoEyeView(InView) && GEngine->StereoRenderingDevice.IsValid()))
 	{
 		uint32 LODViewIndex = GEngine->StereoRenderingDevice->GetLODViewIndex();
 		if (InView.Family && InView.Family->Views.IsValidIndex(LODViewIndex))
diff --git a/Engine/Source/Runtime/Engine/Private/SkeletalMesh.cpp b/Engine/Source/Runtime/Engine/Private/SkeletalMesh.cpp
--- a/Engine/Source/Runtime/Engine/Private/SkeletalMesh.cpp
+++ b/Engine/Source/Runtime/Engine/Private/SkeletalMesh.cpp
@@ -6123,12 +6123,20 @@ void FSkeletalMeshSceneProxy::GetMeshElementsConditionallySelectable(const TArra
 	}
 	else
 	{
-		for (int32 ViewIndex = 0; ViewIndex < Views.Num(); ViewIndex++)
+		if (UNLIKELY(!Views.IsEmpty() && IStereoRendering::IsStereoEyeView(*Views[0])))
 		{
-			if (VisibilityMap & (1 << ViewIndex))
+			const FSceneView& View = GetLODView(*Views[0]);
+			MeshObject->UpdateMinDesiredLODLevel(&View, GetBounds(), ViewFamily.FrameNumber, FirstLODIdx);
+		}
+		else
+		{
+			for (int32 ViewIndex = 0; ViewIndex < Views.Num(); ViewIndex++)
 			{
-				const FSceneView* View = Views[ViewIndex];
-				MeshObject->UpdateMinDesiredLODLevel(View, GetBounds(), ViewFamily.FrameNumber, FirstLODIdx);
+				if (VisibilityMap & (1 << ViewIndex))
+				{
+					const FSceneView* View = Views[ViewIndex];
+					MeshObject->UpdateMinDesiredLODLevel(View, GetBounds(), ViewFamily.FrameNumber, FirstLODIdx);
+				}
 			}
 		}
 
