Do not check for built mipmaps when building an empty texture array, as it generates unnecessary warning in the log.

diff --git a/Engine/Source/Runtime/Engine/Classes/Engine/Texture2DArray.cpp b/Engine/Source/Runtime/Engine/Classes/Engine/Texture2DArray.cpp
--- a/Engine/Source/Runtime/Engine/Classes/Engine/Texture2DArray.cpp
+++ b/Engine/Source/Runtime/Engine/Classes/Engine/Texture2DArray.cpp
@@ -174,6 +174,13 @@ FTextureResource* UTexture2DArray::CreateResource()
 	{
 		return new FTexture2DArrayResource(this, GetResourcePostInitState(GetPlatformData(), GSupportsTexture2DArrayStreaming));
 	}
+#if WITH_EDITORONLY_DATA
+	else if (!SourceTextures.Num())
+	{
+		// empty arrays don't have built mips
+		return nullptr;
+	}
+#endif
 	else if (GetNumMips() == 0)
 	{
 		UE_LOG(LogTexture, Warning, TEXT("%s contains no miplevels! Please delete."), *GetFullName());
diff --git a/Engine/Source/Runtime/Engine/Private/TextureCubeArray.cpp b/Engine/Source/Runtime/Engine/Private/TextureCubeArray.cpp
--- a/Engine/Source/Runtime/Engine/Private/TextureCubeArray.cpp
+++ b/Engine/Source/Runtime/Engine/Private/TextureCubeArray.cpp
@@ -436,6 +436,13 @@ FTextureResource* UTextureCubeArray::CreateResource()
 	{
 		return new FTextureCubeArrayResource(this);
 	}
+#if WITH_EDITORONLY_DATA
+	else if (!SourceTextures.Num())
+	{
+		// empty arrays don't have built mips
+		return nullptr;
+	}
+#endif
 	else if (GetNumMips() == 0)
 	{
 		UE_LOG(LogTexture, Warning, TEXT("%s contains no miplevels! Please delete."), *GetFullName());
