Control Rig: Fix break on user asset links defined in default class
#rnx

diff --git a/Engine/Plugins/Animation/ControlRig/Source/ControlRig/Private/ControlRig.cpp b/Engine/Plugins/Animation/ControlRig/Source/ControlRig/Private/ControlRig.cpp
--- a/Engine/Plugins/Animation/ControlRig/Source/ControlRig/Private/ControlRig.cpp
+++ b/Engine/Plugins/Animation/ControlRig/Source/ControlRig/Private/ControlRig.cpp
@@ -2780,7 +2780,10 @@ const TArray<UAssetUserData*>* UControlRig::GetAssetUserDataArray() const
 	if(HasAnyFlags(RF_ClassDefaultObject))
 	{
 #if WITH_EDITOR
-		return &ToRawPtrTArrayUnsafe(CombinedAssetUserData);
+		CachedAssetUserData.Reset();
+		CachedAssetUserData.Append(AssetUserData);
+		CachedAssetUserData.Append(AssetUserDataEditorOnly);
+		return &ToRawPtrTArrayUnsafe(CachedAssetUserData);
 #else
 		return &ToRawPtrTArrayUnsafe(AssetUserData);
 #endif
@@ -2794,10 +2797,9 @@ const TArray<UAssetUserData*>* UControlRig::GetAssetUserDataArray() const
 	}
 	else
 	{
-#if WITH_EDITOR
-		CombinedAssetUserData.Append(CombinedAssetUserData);
-#else
 		CombinedAssetUserData.Append(AssetUserData);
+#if WITH_EDITOR
+		CombinedAssetUserData.Append(AssetUserDataEditorOnly);
 #endif
 	}
 	
diff --git a/Engine/Plugins/Runtime/RigVM/Source/RigVM/Private/RigVMHost.cpp b/Engine/Plugins/Runtime/RigVM/Source/RigVM/Private/RigVMHost.cpp
--- a/Engine/Plugins/Runtime/RigVM/Source/RigVM/Private/RigVMHost.cpp
+++ b/Engine/Plugins/Runtime/RigVM/Source/RigVM/Private/RigVMHost.cpp
@@ -162,7 +162,6 @@ void URigVMHost::PostLoad()
 	}
 #endif
 
-	RecomputeCombinedAssetUserData();
 }
 
 void URigVMHost::PreSave(FObjectPreSaveContext SaveContext)
@@ -620,13 +619,6 @@ bool URigVMHost::DisableExecution()
 	return CVarRigVMDisableExecutionAll->GetInt() == 1;
 }
 
-#if WITH_EDITOR
-void URigVMHost::PostEditChangeProperty(FPropertyChangedEvent& PropertyChangedEvent)
-{
-	RecomputeCombinedAssetUserData();
-}
-#endif
-
 bool URigVMHost::InitializeCDOVM()
 {
 	check(VM != nullptr);
@@ -1405,7 +1397,6 @@ void URigVMHost::AddAssetUserData(UAssetUserData* InUserData)
 	{
 		RemoveUserDataOfClass(InUserData->GetClass());
 		AssetUserData.Add(InUserData);
-		RecomputeCombinedAssetUserData();
 	}
 }
 
@@ -1431,7 +1422,6 @@ void URigVMHost::RemoveUserDataOfClass(TSubclassOf<UAssetUserData> InUserDataCla
 		if (Datum != NULL && Datum->IsA(InUserDataClass))
 		{
 			AssetUserData.RemoveAt(DataIdx);
-			RecomputeCombinedAssetUserData();
 			return;
 		}
 	}
@@ -1442,7 +1432,6 @@ void URigVMHost::RemoveUserDataOfClass(TSubclassOf<UAssetUserData> InUserDataCla
 		if (Datum != NULL && Datum->IsA(InUserDataClass))
 		{
 			AssetUserDataEditorOnly.RemoveAt(DataIdx);
-			RecomputeCombinedAssetUserData();
 			return;
 		}
 	}
@@ -1452,21 +1441,15 @@ void URigVMHost::RemoveUserDataOfClass(TSubclassOf<UAssetUserData> InUserDataCla
 const TArray<UAssetUserData*>* URigVMHost::GetAssetUserDataArray() const
 {
 #if WITH_EDITOR
-	return &ToRawPtrTArrayUnsafe(CombinedAssetUserData);
+	CachedAssetUserData.Reset();
+	CachedAssetUserData.Append(AssetUserData);
+	CachedAssetUserData.Append(AssetUserDataEditorOnly);
+	return &ToRawPtrTArrayUnsafe(CachedAssetUserData);
 #else
 	return &ToRawPtrTArrayUnsafe(AssetUserData);
 #endif
 }
 
-void URigVMHost::RecomputeCombinedAssetUserData()
-{
-	CombinedAssetUserData.Reset();
-	CombinedAssetUserData.Append(AssetUserData);
-#if WITH_EDITOR
-	CombinedAssetUserData.Append(AssetUserDataEditorOnly);
-#endif
-}
-
 #if WITH_EDITOR	
 
 void URigVMHost::LogOnce(EMessageSeverity::Type InSeverity, int32 InInstructionIndex, const FString& InMessage)
diff --git a/Engine/Plugins/Runtime/RigVM/Source/RigVM/Public/RigVMHost.h b/Engine/Plugins/Runtime/RigVM/Source/RigVM/Public/RigVMHost.h
--- a/Engine/Plugins/Runtime/RigVM/Source/RigVM/Public/RigVMHost.h
+++ b/Engine/Plugins/Runtime/RigVM/Source/RigVM/Public/RigVMHost.h
@@ -159,8 +159,6 @@ public:
 	void BroadCastEndLoadPackage() { EndLoadPackageEvent.Broadcast(this); }
 	FOnEndLoadPackage& OnEndLoadPackage() { return EndLoadPackageEvent; }
 
-	virtual void PostEditChangeProperty(FPropertyChangedEvent& PropertyChangedEvent) override;
-
 private:
 	FOnEndLoadPackage EndLoadPackageEvent;
 
@@ -401,12 +399,11 @@ protected:
 
 #if WITH_EDITORONLY_DATA
 	/** Array of user data stored with the asset */
-	UPROPERTY(EditAnywhere, AdvancedDisplay, Instanced, Category = Skeleton)
+	UPROPERTY(EditAnywhere, AdvancedDisplay, Instanced, Category = "Default")
 	TArray<TObjectPtr<UAssetUserData>> AssetUserDataEditorOnly;
 #endif
 
-	void RecomputeCombinedAssetUserData();
-	mutable TArray<TObjectPtr<UAssetUserData>> CombinedAssetUserData;
+	mutable TArray<TObjectPtr<UAssetUserData>> CachedAssetUserData;
 
 
 protected:
diff --git a/Engine/Source/Runtime/Engine/Classes/Animation/Skeleton.h b/Engine/Source/Runtime/Engine/Classes/Animation/Skeleton.h
--- a/Engine/Source/Runtime/Engine/Classes/Animation/Skeleton.h
+++ b/Engine/Source/Runtime/Engine/Classes/Animation/Skeleton.h
@@ -976,8 +976,6 @@ protected:
 	/** Call this when the skeleton has changed to fix dependent assets */
 	ENGINE_API void HandleSkeletonHierarchyChange(bool bShowProgress = true);
 
-	ENGINE_API virtual void PostEditChangeProperty(FPropertyChangedEvent& PropertyChangedEvent) override;
-
 public:
 	typedef FOnSkeletonHierarchyChangedMulticaster::FDelegate FOnSkeletonHierarchyChanged;
 
@@ -1036,8 +1034,7 @@ protected:
 	TArray<TObjectPtr<UAssetUserData>> AssetUserDataEditorOnly;
 #endif
 
-	void RecomputeCombinedAssetUserData();
-	mutable TArray<TObjectPtr<UAssetUserData>> CombinedAssetUserData;
+	mutable TArray<TObjectPtr<UAssetUserData>> CachedAssetUserData;
 
 	friend struct FReferenceSkeletonModifier;
 	friend class FEditableSkeleton;
diff --git a/Engine/Source/Runtime/Engine/Classes/Components/ActorComponent.h b/Engine/Source/Runtime/Engine/Classes/Components/ActorComponent.h
--- a/Engine/Source/Runtime/Engine/Classes/Components/ActorComponent.h
+++ b/Engine/Source/Runtime/Engine/Classes/Components/ActorComponent.h
@@ -163,8 +163,7 @@ protected:
 	TArray<TObjectPtr<UAssetUserData>> AssetUserDataEditorOnly;
 #endif
 
-	void RecomputeCombinedAssetUserData();
-	mutable TArray<TObjectPtr<UAssetUserData>> CombinedAssetUserData;
+	mutable TArray<TObjectPtr<UAssetUserData>> CachedAssetUserData;
 
 private:
 	/** Used for fast removal of end of frame update */
diff --git a/Engine/Source/Runtime/Engine/Classes/Engine/SkeletalMesh.h b/Engine/Source/Runtime/Engine/Classes/Engine/SkeletalMesh.h
--- a/Engine/Source/Runtime/Engine/Classes/Engine/SkeletalMesh.h
+++ b/Engine/Source/Runtime/Engine/Classes/Engine/SkeletalMesh.h
@@ -2231,8 +2231,7 @@ protected:
 	FOnMeshChanged OnMeshChanged;
 #endif
 
-	void RecomputeCombinedAssetUserData();
-	mutable TArray<TObjectPtr<UAssetUserData>> CombinedAssetUserData;
+	mutable TArray<TObjectPtr<UAssetUserData>> CachedAssetUserData;
 
 	friend struct FSkeletalMeshUpdateContext;
 	friend class FSkeletalMeshUpdate;
diff --git a/Engine/Source/Runtime/Engine/Private/Animation/Skeleton.cpp b/Engine/Source/Runtime/Engine/Private/Animation/Skeleton.cpp
--- a/Engine/Source/Runtime/Engine/Private/Animation/Skeleton.cpp
+++ b/Engine/Source/Runtime/Engine/Private/Animation/Skeleton.cpp
@@ -397,7 +397,6 @@ void USkeleton::PostLoad()
 	// refresh linked bone indices
 	RefreshSkeletonMetaData();
 
-	RecomputeCombinedAssetUserData();
 }
 
 void USkeleton::PostDuplicate(bool bDuplicateForPIE)
@@ -1484,11 +1483,6 @@ void USkeleton::HandleSkeletonHierarchyChange(bool bShowProgress /*= true*/)
 	OnSkeletonHierarchyChanged.Broadcast();
 }
 
-void USkeleton::PostEditChangeProperty(FPropertyChangedEvent& PropertyChangedEvent)
-{
-	RecomputeCombinedAssetUserData();
-}
-
 void USkeleton::RegisterOnSkeletonHierarchyChanged(const FOnSkeletonHierarchyChanged& Delegate)
 {
 	OnSkeletonHierarchyChanged.Add(Delegate);
@@ -2082,7 +2076,6 @@ void USkeleton::AddAssetUserData( UAssetUserData* InUserData)
 	{
 		RemoveUserDataOfClass(InUserData->GetClass());
 		AssetUserData.Add(InUserData);
-		RecomputeCombinedAssetUserData();
 	}
 }
 
@@ -2108,7 +2101,6 @@ void USkeleton::RemoveUserDataOfClass(TSubclassOf<UAssetUserData> InUserDataClas
 		if (Datum != NULL && Datum->IsA(InUserDataClass))
 		{
 			AssetUserData.RemoveAt(DataIdx);
-			RecomputeCombinedAssetUserData();
 			return;
 		}
 	}
@@ -2120,7 +2112,6 @@ void USkeleton::RemoveUserDataOfClass(TSubclassOf<UAssetUserData> InUserDataClas
 		if (Datum != NULL && Datum->IsA(InUserDataClass))
 		{
 			AssetUserDataEditorOnly.RemoveAt(DataIdx);
-			RecomputeCombinedAssetUserData();
 			return;
 		}
 	}
@@ -2130,21 +2121,15 @@ void USkeleton::RemoveUserDataOfClass(TSubclassOf<UAssetUserData> InUserDataClas
 const TArray<UAssetUserData*>* USkeleton::GetAssetUserDataArray() const
 {
 #if WITH_EDITOR
-	return &ToRawPtrTArrayUnsafe(CombinedAssetUserData);
+	CachedAssetUserData.Reset();
+	CachedAssetUserData.Append(AssetUserData);
+	CachedAssetUserData.Append(AssetUserDataEditorOnly);
+	return &ToRawPtrTArrayUnsafe(CachedAssetUserData);
 #else
 	return &ToRawPtrTArrayUnsafe(AssetUserData);
 #endif
 }
 
-void USkeleton::RecomputeCombinedAssetUserData()
-{
-	CombinedAssetUserData.Reset();
-	CombinedAssetUserData.Append(AssetUserData);
-#if WITH_EDITOR
-	CombinedAssetUserData.Append(AssetUserDataEditorOnly);
-#endif
-}
-
 void USkeleton::HandlePackageReloaded(const EPackageReloadPhase InPackageReloadPhase, FPackageReloadedEvent* InPackageReloadedEvent)
 {
 	if (InPackageReloadPhase == EPackageReloadPhase::PostPackageFixup)
diff --git a/Engine/Source/Runtime/Engine/Private/Components/ActorComponent.cpp b/Engine/Source/Runtime/Engine/Private/Components/ActorComponent.cpp
--- a/Engine/Source/Runtime/Engine/Private/Components/ActorComponent.cpp
+++ b/Engine/Source/Runtime/Engine/Private/Components/ActorComponent.cpp
@@ -391,15 +391,6 @@ UActorComponent::UActorComponent(const FObjectInitializer& ObjectInitializer /*=
 	bReplicateUsingRegisteredSubObjectList = GDefaultUseSubObjectReplicationList;
 }
 
-void UActorComponent::RecomputeCombinedAssetUserData()
-{
-	CombinedAssetUserData.Reset();
-	CombinedAssetUserData.Append(AssetUserData);
-#if WITH_EDITOR
-	CombinedAssetUserData.Append(AssetUserDataEditorOnly);
-#endif
-}
-
 void UActorComponent::PostInitProperties()
 {
 	Super::PostInitProperties();
@@ -498,8 +489,6 @@ void UActorComponent::PostLoad()
 		bMarkPendingKillOnPostLoad = false;
 	}
 #endif // WITH_EDITOR
-
-	RecomputeCombinedAssetUserData();
 }
 
 bool UActorComponent::Rename( const TCHAR* InName, UObject* NewOuter, ERenameFlags Flags )
@@ -1033,7 +1022,6 @@ void UActorComponent::ConsolidatedPostEditChange(const FPropertyChangedEvent& Pr
 			Datum->PostEditChangeOwner();
 		}
 	}
-	RecomputeCombinedAssetUserData();
 }
 
 void UActorComponent::PostEditChangeProperty(FPropertyChangedEvent& PropertyChangedEvent)
@@ -2080,7 +2068,6 @@ void UActorComponent::AddAssetUserData( UAssetUserData* InUserData)
 	{
 		RemoveUserDataOfClass(InUserData->GetClass());
 		AssetUserData.Add(InUserData);
-		RecomputeCombinedAssetUserData();
 	}
 }
 
@@ -2101,7 +2088,10 @@ UAssetUserData* UActorComponent::GetAssetUserDataOfClass(TSubclassOf<UAssetUserD
 const TArray<UAssetUserData*>* UActorComponent::GetAssetUserDataArray() const
 {
 #if WITH_EDITOR
-	return &ToRawPtrTArrayUnsafe(CombinedAssetUserData);
+	CachedAssetUserData.Reset();
+	CachedAssetUserData.Append(AssetUserData);
+	CachedAssetUserData.Append(AssetUserDataEditorOnly);
+	return &ToRawPtrTArrayUnsafe(CachedAssetUserData);
 #else
 	return &ToRawPtrTArrayUnsafe(AssetUserData);
 #endif
diff --git a/Engine/Source/Runtime/Engine/Private/SkeletalMesh.cpp b/Engine/Source/Runtime/Engine/Private/SkeletalMesh.cpp
--- a/Engine/Source/Runtime/Engine/Private/SkeletalMesh.cpp
+++ b/Engine/Source/Runtime/Engine/Private/SkeletalMesh.cpp
@@ -818,15 +818,6 @@ bool USkeletalMesh::NeedCPUData(int32 LODIndex) const
 	return GetSamplingInfo().IsSamplingEnabled(this, LODIndex);
 }
 
-void USkeletalMesh::RecomputeCombinedAssetUserData()
-{
-	CombinedAssetUserData.Reset();
-	CombinedAssetUserData.Append(AssetUserData);
-#if WITH_EDITOR
-	CombinedAssetUserData.Append(AssetUserDataEditorOnly);
-#endif
-}
-
 void USkeletalMesh::InitResources()
 {
 	LLM_SCOPE_BYNAME(TEXT("SkeletalMesh/InitResources")); // This is an important test case for SCOPE_BYNAME without a matching LLM_DEFINE_TAG
@@ -1315,7 +1306,6 @@ void USkeletalMesh::PostEditChangeProperty(FPropertyChangedEvent& PropertyChange
 			Datum->PostEditChangeOwner();
 		}
 	}
-	RecomputeCombinedAssetUserData();
 
 	Super::PostEditChangeProperty(PropertyChangedEvent);
 
@@ -3163,7 +3153,6 @@ void USkeletalMesh::PostLoad()
 	LLM_SCOPE(ELLMTag::SkeletalMesh);
 	Super::PostLoad();
 
-	RecomputeCombinedAssetUserData();
 }
 
 void USkeletalMesh::ExecutePostLoadInternal(FSkinnedAssetPostLoadContext& Context)
@@ -4735,7 +4724,6 @@ void USkeletalMesh::AddAssetUserData( UAssetUserData* InUserData)
 	{
 		RemoveUserDataOfClass(InUserData->GetClass());
 		AssetUserData.Add(InUserData);
-		RecomputeCombinedAssetUserData();
 	}
 }
 
@@ -4761,7 +4749,6 @@ void USkeletalMesh::RemoveUserDataOfClass(TSubclassOf<UAssetUserData> InUserData
 		if (Datum != NULL && Datum->IsA(InUserDataClass))
 		{
 			AssetUserData.RemoveAt(DataIdx);
-			RecomputeCombinedAssetUserData();
 			return;
 		}
 	}
@@ -4772,7 +4759,6 @@ void USkeletalMesh::RemoveUserDataOfClass(TSubclassOf<UAssetUserData> InUserData
 		if (Datum != NULL && Datum->IsA(InUserDataClass))
 		{
 			AssetUserDataEditorOnly.RemoveAt(DataIdx);
-			RecomputeCombinedAssetUserData();
 			return;
 		}
 	}
@@ -4782,7 +4768,10 @@ void USkeletalMesh::RemoveUserDataOfClass(TSubclassOf<UAssetUserData> InUserData
 const TArray<UAssetUserData*>* USkeletalMesh::GetAssetUserDataArray() const
 {
 #if WITH_EDITOR
-	return &ToRawPtrTArrayUnsafe(CombinedAssetUserData);
+	CachedAssetUserData.Reset();
+	CachedAssetUserData.Append(AssetUserData);
+	CachedAssetUserData.Append(AssetUserDataEditorOnly);
+	return &ToRawPtrTArrayUnsafe(CachedAssetUserData);
 #else
 	return &ToRawPtrTArrayUnsafe(AssetUserData);
 #endif
