Remove non SRV bone buffer path since this is no longer used.

diff --git a/Engine/Shaders/Private/GpuSkinVertexFactory.ush b/Engine/Shaders/Private/GpuSkinVertexFactory.ush
--- a/Engine/Shaders/Private/GpuSkinVertexFactory.ush
+++ b/Engine/Shaders/Private/GpuSkinVertexFactory.ush
@@ -28,10 +28,6 @@
 #define GPUSKIN_UNLIMITED_BONE_INFLUENCE 0
 #endif
 
-#ifndef GPUSKIN_USE_BONES_SRV_BUFFER
-#define GPUSKIN_USE_BONES_SRV_BUFFER	0	// Use a Buffer<float4> if 1, otherwise a uniform buffer   
-#endif
-
 #define MAX_INFLUENCES_PER_STREAM	4
 
 #define FIXED_VERTEX_INDEX 0xFFFF
@@ -281,15 +277,11 @@ FMaterialPixelParameters GetMaterialPixelParameters(FVertexFactoryInterpolantsVS
 
 #define FBoneMatrix float3x4
 
-#if GPUSKIN_USE_BONES_SRV_BUFFER
-
 // The bone matrix buffer stored as 4x3 (3 float4 behind each other), all chunks of a skeletal mesh in one, each skeletal mesh has it's own buffer
 STRONG_TYPE Buffer<float4> BoneMatrices;
 // The previous bone matrix buffer stored as 4x3 (3 float4 behind each other), all chunks of a skeletal mesh in one, each skeletal mesh has it's own buffer
 STRONG_TYPE Buffer<float4> PreviousBoneMatrices;
 
-#endif
-
 // Cache data to avoid multiple calculation 
 struct FVertexFactoryIntermediates
 {
@@ -529,14 +521,10 @@ float3 MorphPreviousPosition( FVertexFactoryInput Input, FVertexFactoryIntermedi
 
 FBoneMatrix GetBoneMatrix(int Index)
 {
-#if GPUSKIN_USE_BONES_SRV_BUFFER
 	float4 A = BoneMatrices[Index * 3];
 	float4 B = BoneMatrices[Index * 3 + 1];
 	float4 C = BoneMatrices[Index * 3 + 2];
 	return FBoneMatrix(A,B,C);
-#else
-	return Bones.BoneMatrices[Index];
-#endif
 }
 
 FBoneMatrix CalcBoneMatrix( FVertexFactoryInput Input )
@@ -606,14 +594,10 @@ FBoneMatrix CalcBoneMatrix( FVertexFactoryInput Input )
 
 FBoneMatrix GetPreviousBoneMatrix(int Index)
 {
-#if GPUSKIN_USE_BONES_SRV_BUFFER
 	float4 A = PreviousBoneMatrices[Index * 3 + 0];
 	float4 B = PreviousBoneMatrices[Index * 3 + 1];
 	float4 C = PreviousBoneMatrices[Index * 3 + 2];
 	return FBoneMatrix(A,B,C);
-#else
-	return Bones.BoneMatrices[Index];
-#endif
 }
 
 FBoneMatrix CalcPreviousBoneMatrix( FVertexFactoryInput Input )
diff --git a/Engine/Source/Runtime/Engine/Private/GPUSkinVertexFactory.cpp b/Engine/Source/Runtime/Engine/Private/GPUSkinVertexFactory.cpp
--- a/Engine/Source/Runtime/Engine/Private/GPUSkinVertexFactory.cpp
+++ b/Engine/Source/Runtime/Engine/Private/GPUSkinVertexFactory.cpp
@@ -69,10 +69,6 @@ static TAutoConsoleVariable<bool> CVarMobileEnableCloth(
 
 IMPLEMENT_GLOBAL_SHADER_PARAMETER_STRUCT(FAPEXClothUniformShaderParameters,"APEXClothParam");
 
-IMPLEMENT_GLOBAL_SHADER_PARAMETER_STRUCT(FBoneMatricesUniformShaderParameters,"Bones");
-
-static FBoneMatricesUniformShaderParameters GBoneUniformStruct;
-
 #define IMPLEMENT_GPUSKINNING_VERTEX_FACTORY_TYPE_INTERNAL(FactoryClass, ShaderFilename, Flags) \
 	template <GPUSkinBoneInfluenceType BoneInfluenceType> FVertexFactoryType FactoryClass<BoneInfluenceType>::StaticType( \
 	BoneInfluenceType == DefaultBoneInfluence ? TEXT(#FactoryClass) TEXT("Default") : TEXT(#FactoryClass) TEXT("Unlimited"), \
@@ -119,22 +115,6 @@ static_assert(sizeof(ispc::FMatrix44f) == sizeof(FMatrix44f), "sizeof(ispc::FMat
 static_assert(sizeof(ispc::FMatrix3x4) == sizeof(FMatrix3x4), "sizeof(ispc::FMatrix3x4) != sizeof(FMatrix3x4)");
 #endif
 
-// ---
-// These should match USE_BONES_SRV_BUFFER
-static inline bool SupportsBonesBufferSRV(EShaderPlatform Platform)
-{
-	// at some point we might switch GL to uniform buffers
-	return true;
-}
-
-static inline bool SupportsBonesBufferSRV(ERHIFeatureLevel::Type InFeatureLevel)
-{
-	// at some point we might switch GL to uniform buffers
-	return true;
-}
-// ---
-
-
 /*-----------------------------------------------------------------------------
  FSharedPoolPolicyData
  -----------------------------------------------------------------------------*/
@@ -247,8 +227,6 @@ void FGPUBaseSkinVertexFactory::FShaderDataType::UpdateBoneData(FRHICommandList&
 	FMatrix3x4* ChunkMatrices = nullptr;
 
 	FVertexBufferAndSRV* CurrentBoneBuffer = 0;
-
-	if (SupportsBonesBufferSRV(InFeatureLevel))
 	{
 		check(IsInParallelRenderingThread());
 
@@ -278,14 +256,6 @@ void FGPUBaseSkinVertexFactory::FShaderDataType::UpdateBoneData(FRHICommandList&
 			ChunkMatrices = (FMatrix3x4*)RHICmdList.LockBuffer(CurrentBoneBuffer->VertexBufferRHI, 0, VectorArraySize, RLM_WriteOnly);
 		}
 	}
-	else
-	{
-		if(NumBones)
-		{
-			check(NumBones * sizeof(FMatrix3x4) <= sizeof(GBoneUniformStruct));
-			ChunkMatrices = (FMatrix3x4*)&GBoneUniformStruct;
-		}
-	}
 
 	{
 		QUICK_SCOPE_CYCLE_COUNTER(STAT_FGPUBaseSkinVertexFactory_ShaderDataType_UpdateBoneData_CopyBones);
@@ -319,7 +289,6 @@ void FGPUBaseSkinVertexFactory::FShaderDataType::UpdateBoneData(FRHICommandList&
 			}
 		}
 	}
-	if (SupportsBonesBufferSRV(InFeatureLevel))
 	{
 		if (NumBones)
 		{
@@ -327,10 +296,6 @@ void FGPUBaseSkinVertexFactory::FShaderDataType::UpdateBoneData(FRHICommandList&
 			RHICmdList.UnlockBuffer(CurrentBoneBuffer->VertexBufferRHI);
 		}
 	}
-	else
-	{
-		UniformBuffer = RHICreateUniformBuffer(&GBoneUniformStruct, FBoneMatricesUniformShaderParameters::FTypeInfo::GetStructMetadata()->GetLayoutPtr(), UniformBuffer_MultiFrame);
-	}
 }
 
 int32 FGPUBaseSkinVertexFactory::GetMinimumPerPlatformMaxGPUSkinBonesValue()
@@ -550,7 +515,6 @@ void TGPUSkinVertexFactory<BoneInfluenceType>::ModifyCompilationEnvironment(cons
 		OutEnvironment.SetDefine(TEXT("GPUSKIN_LIMIT_2BONE_INFLUENCES"), (bLimit2BoneInfluences ? 1 : 0));
 	}
 
-	OutEnvironment.SetDefine(TEXT("GPUSKIN_USE_BONES_SRV_BUFFER"), SupportsBonesBufferSRV(Parameters.Platform) ? 1 : 0);
 	OutEnvironment.SetDefine(TEXT("GPUSKIN_UNLIMITED_BONE_INFLUENCE"), BoneInfluenceType == UnlimitedBoneInfluence ? 1 : 0);
 
 	OutEnvironment.SetDefine(TEXT("GPU_SKINNED_MESH_FACTORY"), 1);
@@ -830,8 +794,6 @@ public:
 		const FGPUBaseSkinVertexFactory::FShaderDataType& ShaderData = ((const FGPUBaseSkinVertexFactory*)VertexFactory)->GetShaderData();
 
 		bool bLocalPerBoneMotionBlur = false;
-
-		if (SupportsBonesBufferSRV(FeatureLevel))
 		{
 			if (BoneMatrices.IsBound())
 			{
@@ -853,10 +815,6 @@ public:
 				ShaderBindings.Add(PreviousBoneMatrices, PreviousData);
 			}
 		}
-		else
-		{
-			ShaderBindings.Add(Shader->GetUniformBufferParameter<FBoneMatricesUniformShaderParameters>(), ShaderData.GetUniformBuffer());
-		}
 
 		ShaderBindings.Add(PerBoneMotionBlur, (uint32)(bLocalPerBoneMotionBlur ? 1 : 0));
 
diff --git a/Engine/Source/Runtime/Engine/Public/GPUSkinVertexFactory.h b/Engine/Source/Runtime/Engine/Public/GPUSkinVertexFactory.h
--- a/Engine/Source/Runtime/Engine/Public/GPUSkinVertexFactory.h
+++ b/Engine/Source/Runtime/Engine/Public/GPUSkinVertexFactory.h
@@ -25,16 +25,6 @@ template <class T> class TConsoleVariableData;
 BEGIN_GLOBAL_SHADER_PARAMETER_STRUCT(FAPEXClothUniformShaderParameters,)
 END_GLOBAL_SHADER_PARAMETER_STRUCT()
 
-enum
-{
-	// 256 works for real uniform buffers, emulated UB can support up to 75 
-	MAX_GPU_BONE_MATRICES_UNIFORMBUFFER = 256,
-};
-
-BEGIN_GLOBAL_SHADER_PARAMETER_STRUCT(FBoneMatricesUniformShaderParameters,)
-	SHADER_PARAMETER_ARRAY(FMatrix3x4, BoneMatrices, [MAX_GPU_BONE_MATRICES_UNIFORMBUFFER])
-END_GLOBAL_SHADER_PARAMETER_STRUCT()
-
 #define SET_BONE_DATA(B, X) B.SetMatrixTranspose(X)
 
 /** Shared data & implementation for the different types of pool */
