World Partition: add support for RuntimeHashSet HLODs generation.

diff --git a/Engine/Plugins/Editor/WorldPartitionHLODUtilities/Source/Private/WorldPartition/HLOD/Utilities/WorldPartitionHLODUtilities.cpp b/Engine/Plugins/Editor/WorldPartitionHLODUtilities/Source/Private/WorldPartition/HLOD/Utilities/WorldPartitionHLODUtilities.cpp
--- a/Engine/Plugins/Editor/WorldPartitionHLODUtilities/Source/Private/WorldPartition/HLOD/Utilities/WorldPartitionHLODUtilities.cpp
+++ b/Engine/Plugins/Editor/WorldPartitionHLODUtilities/Source/Private/WorldPartition/HLOD/Utilities/WorldPartitionHLODUtilities.cpp
@@ -84,7 +84,8 @@ void AddSubActor(const FWorldPartitionActorDescView& ActorDescView, const IStrea
 
 	// Add the actor
 	bool bIsAlreadyInSet = false;
-	FHLODSubActor SubActor(ActorDescView.GetGuid(), ActorDescView.GetActorPackage(), ActorPath, ActorInstance.GetContainerID(), ActorInstance.GetActorDescContainer()->GetContainerPackage(), ActorInstance.GetTransform());
+	const UActorDescContainer* ActorDescContainer = ActorDescView.GetActorDesc()->GetContainer();
+	FHLODSubActor SubActor(ActorDescView.GetGuid(), ActorDescView.GetActorPackage(), ActorPath, ActorInstance.GetContainerID(), ActorDescContainer->GetContainerPackage(), ActorInstance.GetTransform());
 	SubActors.Add(SubActor, &bIsAlreadyInSet);
 
 	if (!bIsAlreadyInSet)
diff --git a/Engine/Source/Runtime/Engine/Private/WorldPartition/RuntimeHashSet/RuntimePartition.cpp b/Engine/Source/Runtime/Engine/Private/WorldPartition/RuntimeHashSet/RuntimePartition.cpp
--- a/Engine/Source/Runtime/Engine/Private/WorldPartition/RuntimeHashSet/RuntimePartition.cpp
+++ b/Engine/Source/Runtime/Engine/Private/WorldPartition/RuntimeHashSet/RuntimePartition.cpp
@@ -1,6 +1,7 @@
 // Copyright Epic Games, Inc. All Rights Reserved.
 
 #include "WorldPartition/RuntimeHashSet/RuntimePartition.h"
+#include "WorldPartition/RuntimeHashSet/WorldPartitionRuntimeHashSet.h"
 
 #if WITH_EDITOR
 void URuntimePartition::SetDefaultValues()
@@ -25,13 +26,7 @@ void URuntimePartition::PostEditChangeProperty(FPropertyChangedEvent& InProperty
 	Super::PostEditChangeProperty(InPropertyChangedEvent);
 }
 
-bool URuntimePartition::PopulateCellActorInstances(const TArray<const IStreamingGenerationContext::FActorSetInstance*>& InActorSetInstances, bool bInIsMainWorldPartition, bool bInIsCellAlwaysLoaded, TArray<IStreamingGenerationContext::FActorInstance>& OutCellActorInstances)
-{
-	UWorldPartitionRuntimeHash* RuntimeHash = GetTypedOuter<UWorldPartitionRuntimeHash>();
-	return RuntimeHash->PopulateCellActorInstances(InActorSetInstances, bInIsMainWorldPartition, bInIsCellAlwaysLoaded, OutCellActorInstances);
-}
-
-URuntimePartition::FCellDesc URuntimePartition::CreateCellDesc(const FString& InName, bool bInIsSpatiallyLoaded, const FGuid& InContentBundleID, int32 InLevel, const TArray<IStreamingGenerationContext::FActorInstance>& InActorInstances)
+URuntimePartition::FCellDesc URuntimePartition::CreateCellDesc(const FString& InName, bool bInIsSpatiallyLoaded, int32 InLevel, const TArray<const IStreamingGenerationContext::FActorSetInstance*>& InActorSetInstances)
 {
 	FCellDesc CellDesc;
 
@@ -49,20 +44,15 @@ URuntimePartition::FCellDesc URuntimePartition::CreateCellDesc(const FString& In
 
 	// Set provided input values
 	CellDesc.bIsSpatiallyLoaded = bInIsSpatiallyLoaded;
-	CellDesc.ContentBundleID = InContentBundleID;
 	CellDesc.Level = InLevel;
 
-	// Add actor instances and update bounds
-	CellDesc.ActorInstances = InActorInstances;
-	for (const IStreamingGenerationContext::FActorInstance& ActorInstance : CellDesc.ActorInstances)
-	{
-		const FWorldPartitionActorDescView& ActorDescView = ActorInstance.GetActorDescView();
-		const FBox RuntimeBounds = ActorDescView.GetRuntimeBounds();
+	// Add actor set instances
+	CellDesc.ActorSetInstances = InActorSetInstances;
 
-		if (RuntimeBounds.IsValid)
-		{
-			CellDesc.Bounds += RuntimeBounds.TransformBy(ActorInstance.GetTransform());
-		}
+	// Update cell bounds
+	for (const IStreamingGenerationContext::FActorSetInstance* ActorSetInstance : InActorSetInstances)
+	{
+		CellDesc.Bounds += ActorSetInstance->Bounds;
 	}
 
 	return CellDesc;
diff --git a/Engine/Source/Runtime/Engine/Private/WorldPartition/RuntimeHashSet/RuntimePartitionLHGrid.cpp b/Engine/Source/Runtime/Engine/Private/WorldPartition/RuntimeHashSet/RuntimePartitionLHGrid.cpp
--- a/Engine/Source/Runtime/Engine/Private/WorldPartition/RuntimeHashSet/RuntimePartitionLHGrid.cpp
+++ b/Engine/Source/Runtime/Engine/Private/WorldPartition/RuntimeHashSet/RuntimePartitionLHGrid.cpp
@@ -1,6 +1,7 @@
 // Copyright Epic Games, Inc. All Rights Reserved.
 
 #include "WorldPartition/RuntimeHashSet/RuntimePartitionLHGrid.h"
+#include "WorldPartition/RuntimeHashSet/WorldPartitionRuntimeHashSet.h"
 #include "Misc/HashBuilder.h"
 
 #if WITH_EDITOR
@@ -73,15 +74,9 @@ void URuntimePartitionLHGrid::SetDefaultValues()
 	CellSize = LoadingRange / 2;
 }
 
-bool URuntimePartitionLHGrid::SupportsHLODs() const
+bool URuntimePartitionLHGrid::IsValidPartitionTokens(const TArray<FName> InPartitionTokens) const
 {
-	return true;
-}
-
-bool URuntimePartitionLHGrid::IsValidGrid(FName GridName) const
-{
-	const TArray<FName> GridNameList = UWorldPartitionRuntimeHashSet::ParseGridName(GridName);
-	return GridNameList.Num() == 1;
+	return InPartitionTokens.Num() == 1;
 }
 
 bool URuntimePartitionLHGrid::GenerateStreaming(const FGenerateStreamingParams& InParams, FGenerateStreamingResult& OutResult)
@@ -91,21 +86,17 @@ bool URuntimePartitionLHGrid::GenerateStreaming(const FGenerateStreamingParams&
 	UWorld* OuterWorld = GetTypedOuter<UWorld>();
 	const bool bIsMainWorldPartition = (World == OuterWorld);
 
-	TArray<IStreamingGenerationContext::FActorInstance> CellActorInstances;
-	if (PopulateCellActorInstances(*InParams.ActorSetInstances, bIsMainWorldPartition, false, CellActorInstances))
+	TMap<FCellCoord, TArray<const IStreamingGenerationContext::FActorSetInstance*>> CellsActorSetInstances;
+	for (const IStreamingGenerationContext::FActorSetInstance* ActorSetInstance : *InParams.ActorSetInstances)
+	{
+		const int32 GridLevel = FCellCoord::GetLevelForBox(ActorSetInstance->Bounds, CellSize);
+		const FCellCoord CellCoord = FCellCoord::GetCellCoords(ActorSetInstance->Bounds.GetCenter(), CellSize, GridLevel);
+		CellsActorSetInstances.FindOrAdd(CellCoord).Add(ActorSetInstance);
+	}
+
+	for (auto& [CellCoord, CellActorSetInstances] : CellsActorSetInstances)
 	{
-		TMap<FCellCoord, TArray<IStreamingGenerationContext::FActorInstance>> SubLevelsActorInstances;
-		for (const IStreamingGenerationContext::FActorInstance& ActorInstance : CellActorInstances)
-		{
-			const int32 GridLevel = FCellCoord::GetLevelForBox(ActorInstance.GetBounds(), CellSize);
-			const FCellCoord CellCoord = FCellCoord::GetCellCoords(ActorInstance.GetBounds().GetCenter(), CellSize, GridLevel);
-			SubLevelsActorInstances.FindOrAdd(CellCoord).Add(ActorInstance);
-		}
-
-		for (auto& [CellCoord, SubLevelActorSetInstances] : SubLevelsActorInstances)
-		{
-			OutResult.RuntimeCellDescs.Emplace(CreateCellDesc(CellCoord.ToString(), true, SubLevelActorSetInstances[0].ActorSetInstance->ContentBundleID, CellCoord.Level, SubLevelActorSetInstances));
-		}
+		OutResult.RuntimeCellDescs.Emplace(CreateCellDesc(CellCoord.ToString(), true, CellCoord.Level, CellActorSetInstances));
 	}
 
 	return true;
diff --git a/Engine/Source/Runtime/Engine/Private/WorldPartition/RuntimeHashSet/RuntimePartitionLevelStreaming.cpp b/Engine/Source/Runtime/Engine/Private/WorldPartition/RuntimeHashSet/RuntimePartitionLevelStreaming.cpp
--- a/Engine/Source/Runtime/Engine/Private/WorldPartition/RuntimeHashSet/RuntimePartitionLevelStreaming.cpp
+++ b/Engine/Source/Runtime/Engine/Private/WorldPartition/RuntimeHashSet/RuntimePartitionLevelStreaming.cpp
@@ -1,18 +1,13 @@
 // Copyright Epic Games, Inc. All Rights Reserved.
 
 #include "WorldPartition/RuntimeHashSet/RuntimePartitionLevelStreaming.h"
+#include "WorldPartition/RuntimeHashSet/WorldPartitionRuntimeHashSet.h"
 #include "WorldPartition/WorldPartitionStreamingGenerationContext.h"
 
 #if WITH_EDITOR
-bool URuntimePartitionLevelStreaming::SupportsHLODs() const
+bool URuntimePartitionLevelStreaming::IsValidPartitionTokens(const TArray<FName> InPartitionTokens) const
 {
-	return false;
-}
-
-bool URuntimePartitionLevelStreaming::IsValidGrid(FName GridName) const
-{
-	const TArray<FName> GridNameList = UWorldPartitionRuntimeHashSet::ParseGridName(GridName);
-	return GridNameList.Num() && (GridNameList.Num() <= 2);
+	return InPartitionTokens.Num() && (InPartitionTokens.Num() <= 2);
 }
 
 bool URuntimePartitionLevelStreaming::GenerateStreaming(const FGenerateStreamingParams& InParams, FGenerateStreamingResult& OutResult)
@@ -22,45 +17,38 @@ bool URuntimePartitionLevelStreaming::GenerateStreaming(const FGenerateStreaming
 	UWorld* OuterWorld = GetTypedOuter<UWorld>();
 	const bool bIsMainWorldPartition = (World == OuterWorld);
 
-	TArray<IStreamingGenerationContext::FActorInstance> CellActorInstances;
-	if (PopulateCellActorInstances(*InParams.ActorSetInstances, bIsMainWorldPartition, false, CellActorInstances))
+	TMap<FName, TArray<const IStreamingGenerationContext::FActorSetInstance*>> CellsActorSetInstances;
+	for (const IStreamingGenerationContext::FActorSetInstance* ActorSetInstance : *InParams.ActorSetInstances)
 	{
-		TMap<FName, TArray<IStreamingGenerationContext::FActorInstance>> SubLevelsActorInstances;
-		for (const IStreamingGenerationContext::FActorInstance& ActorInstance : CellActorInstances)
-		{
-			TArray<FName> ActorSetGridNameList;
-					
-			if (!ActorInstance.ActorSetInstance->RuntimeGrid.IsNone())
-			{
-				ActorSetGridNameList = UWorldPartitionRuntimeHashSet::ParseGridName(ActorInstance.ActorSetInstance->RuntimeGrid);
-			}
-			else
-			{
-				ActorSetGridNameList.Add(NAME_Default);
-			}
-
-			if (bOneLevelPerActorContainer && !ActorInstance.GetContainerID().IsMainContainer())
-			{
-				ActorSetGridNameList.Add(*ActorInstance.GetContainerID().ToString());
-			}
+		TArray<FName> ActorSetGridNameList;
 
-			TStringBuilder<512> StringBuilder;
-			for (FName GridName : ActorSetGridNameList)
+		FName LevelName = NAME_Default;
+		if (!ActorSetInstance->RuntimeGrid.IsNone())
+		{
+			TArray<FName> MainPartitionTokens;
+			TArray<FName> HLODPartitionTokens;
+			if (UWorldPartitionRuntimeHashSet::ParseGridName(ActorSetInstance->RuntimeGrid, MainPartitionTokens, HLODPartitionTokens))
 			{
-				StringBuilder += GridName.ToString();
-				StringBuilder += TEXT("_");
+				if (MainPartitionTokens.Num() == 2)
+				{
+					LevelName = MainPartitionTokens[1];
+				}
 			}
-			StringBuilder.RemoveSuffix(1);
+		}
+		ActorSetGridNameList.Add(LevelName);
 
-			FName SubLevelName = *StringBuilder;
+		TStringBuilder<512> StringBuilder;
+		StringBuilder += Name.ToString();
+		StringBuilder += TEXT("_");
+		StringBuilder += LevelName.ToString();
+		FName CellName = *StringBuilder;
 
-			SubLevelsActorInstances.FindOrAdd(SubLevelName).Add(ActorInstance);
-		}
+		CellsActorSetInstances.FindOrAdd(CellName).Add(ActorSetInstance);
+	}
 
-		for (auto& [SubLevelName, SubLevelActorSetInstances] : SubLevelsActorInstances)
-		{
-			OutResult.RuntimeCellDescs.Emplace(CreateCellDesc(SubLevelName.ToString(), true, SubLevelActorSetInstances[0].ActorSetInstance->ContentBundleID, 0, SubLevelActorSetInstances));
-		}
+	for (auto& [CellName, CellActorSetInstances] : CellsActorSetInstances)
+	{
+		OutResult.RuntimeCellDescs.Emplace(CreateCellDesc(CellName.ToString(), true, 0, CellActorSetInstances));
 	}
 
 	return true;
diff --git a/Engine/Source/Runtime/Engine/Private/WorldPartition/RuntimeHashSet/RuntimePartitionPersistent.cpp b/Engine/Source/Runtime/Engine/Private/WorldPartition/RuntimeHashSet/RuntimePartitionPersistent.cpp
--- a/Engine/Source/Runtime/Engine/Private/WorldPartition/RuntimeHashSet/RuntimePartitionPersistent.cpp
+++ b/Engine/Source/Runtime/Engine/Private/WorldPartition/RuntimeHashSet/RuntimePartitionPersistent.cpp
@@ -6,18 +6,8 @@
 #if WITH_EDITOR
 bool URuntimePartitionPersistent::GenerateStreaming(const FGenerateStreamingParams& InParams, FGenerateStreamingResult& OutResult)
 {
-	UWorldPartition* WorldPartition = GetTypedOuter<UWorldPartition>();
-	UWorld* World = WorldPartition->GetWorld();
-	UWorld* OuterWorld = GetTypedOuter<UWorld>();
-	const bool bIsMainWorldPartition = (World == OuterWorld);
-
-	TArray<IStreamingGenerationContext::FActorInstance> CellActorInstances;
-	if (PopulateCellActorInstances(*InParams.ActorSetInstances, bIsMainWorldPartition, true, CellActorInstances))
-	{
-		const FString CellName(TEXT("Persistent"));
-		OutResult.RuntimeCellDescs.Emplace(CreateCellDesc(CellName, false, CellActorInstances[0].ActorSetInstance->ContentBundleID, 0, CellActorInstances));
-	}
-
+	const FString CellName(TEXT("Persistent"));
+	OutResult.RuntimeCellDescs.Emplace(CreateCellDesc(CellName, false, 0, *InParams.ActorSetInstances));
 	return true;
 }
 #endif
\ No newline at end of file
diff --git a/Engine/Source/Runtime/Engine/Private/WorldPartition/RuntimeHashSet/WorldPartitionRuntimeHashSet.cpp b/Engine/Source/Runtime/Engine/Private/WorldPartition/RuntimeHashSet/WorldPartitionRuntimeHashSet.cpp
--- a/Engine/Source/Runtime/Engine/Private/WorldPartition/RuntimeHashSet/WorldPartitionRuntimeHashSet.cpp
+++ b/Engine/Source/Runtime/Engine/Private/WorldPartition/RuntimeHashSet/WorldPartitionRuntimeHashSet.cpp
@@ -4,6 +4,9 @@
 #include "WorldPartition/RuntimeHashSet/RuntimePartition.h"
 #include "WorldPartition/RuntimeHashSet/RuntimePartitionPersistent.h"
 #include "WorldPartition/HLOD/HLODLayer.h"
+#include "WorldPartition/ContentBundle/ContentBundleDescriptor.h"
+#include "WorldPartition/DataLayer/DataLayersID.h"
+#include "Misc/ArchiveMD5.h"
 
 FAutoConsoleCommand WorldPartitionRuntimeHashSetEnable(
 	TEXT("wp.Editor.WorldPartitionRuntimeHashSet.Enable"),
@@ -47,7 +50,7 @@ void FRuntimePartitionDesc::UpdateHLODPartitionLayers()
 
 			FRuntimePartitionHLODSetupLayer& HLODSetupLayer = HLODSetup.PartitionLayers[HLODSetupIndex];
 
-			const bool bHLODLayerMatches = HLODSetup.HLODLayer == CurHLODLayer;
+			const bool bHLODLayerMatches = HLODSetupLayer.HLODLayer == CurHLODLayer;
 			const UClass* ExpectedHLODPartitionClass = CurHLODLayer->IsSpatiallyLoaded() ? MainLayer->GetClass() : URuntimePartitionPersistent::StaticClass();
 			const bool bHasValidPartitionLayer = HLODSetupLayer.PartitionLayer && (HLODSetupLayer.PartitionLayer->GetClass() == ExpectedHLODPartitionClass);
 			
@@ -155,29 +158,40 @@ void UWorldPartitionRuntimeHashSet::FlushStreaming()
 	PersistentPartitionDesc.Name = NAME_None;
 	PersistentPartitionDesc.MainLayer = nullptr;
 
+	check(!RuntimeStreamingData.IsEmpty());
 	RuntimeStreamingData.Empty();
 }
 
 bool UWorldPartitionRuntimeHashSet::IsValidGrid(FName GridName) const
 {
-	if (!RuntimePartitions.Num())
-	{
-		return false;
-	}
-
-	// The None grid name will always map to the first runtime partition in the list
-	if (GridName.IsNone())
-	{
-		return true;
-	}
+	TArray<FName> MainPartitionTokens;
+	TArray<FName> HLODPartitionTokens;
 
 	// Parse the potentially dot separated grid name to identiy the associated runtime partition
-	const TArray<FName> GridNameList = ParseGridName(GridName);
-	for (const FRuntimePartitionDesc& RuntimePartitionDesc : RuntimePartitions)
+	if (ParseGridName(GridName, MainPartitionTokens, HLODPartitionTokens))
 	{
-		if (RuntimePartitionDesc.Name == GridNameList[0])
+		// The None grid name will always map to the first runtime partition in the list
+		if (MainPartitionTokens[0].IsNone())
+		{
+			return true;
+		}
+
+		for (const FRuntimePartitionDesc& RuntimePartitionDesc : RuntimePartitions)
 		{
-			return RuntimePartitionDesc.MainLayer && RuntimePartitionDesc.MainLayer->IsValidGrid(GridName);
+			if (RuntimePartitionDesc.Name == MainPartitionTokens[0])
+			{
+				if (!RuntimePartitionDesc.MainLayer)
+				{
+					return false;
+				}
+
+				if (!RuntimePartitionDesc.MainLayer->IsValidPartitionTokens(MainPartitionTokens))
+				{
+					return false;
+				}
+
+				return true;
+			}
 		}
 	}
 
@@ -198,15 +212,23 @@ bool UWorldPartitionRuntimeHashSet::IsValidHLODLayer(FName GridName, const FSoft
 		
 		if (RuntimePartitionIndex == INDEX_NONE)
 		{
+			TArray<FName> PartitionTokens;
+			TArray<FName> HLODPartitionTokens;
+
 			// Parse the potentially dot separated grid name to identiy the associated runtime partition
-			const TArray<FName> GridNameList = ParseGridName(GridName);
-			for (const FRuntimePartitionDesc& RuntimePartitionDesc : RuntimePartitions)
+			if (ParseGridName(GridName, PartitionTokens, HLODPartitionTokens))
 			{
-				if (RuntimePartitionDesc.Name == GridNameList[0])
+				int32 RuntimePartitionIndexLookup = 0;
+				for (const FRuntimePartitionDesc& RuntimePartitionDesc : RuntimePartitions)
 				{
-					break;
+					if (RuntimePartitionDesc.Name == PartitionTokens[0])
+					{
+						RuntimePartitionIndex = RuntimePartitionIndexLookup;
+						break;
+					}
+
+					RuntimePartitionIndexLookup++;
 				}
-				RuntimePartitionIndex++;
 			}
 		}
 
@@ -237,17 +259,42 @@ TArray<UWorldPartitionRuntimeCell*> UWorldPartitionRuntimeHashSet::GetAlwaysLoad
 	return AlwaysLoadedCells;
 }
 
-TArray<FName> UWorldPartitionRuntimeHashSet::ParseGridName(FName GridName)
+bool UWorldPartitionRuntimeHashSet::ParseGridName(FName GridName, TArray<FName>& MainPartitionTokens, TArray<FName>& HLODPartitionTokens)
 {
-	TArray<FString> GridNameList;
-	const FString GridNameStr = GridName.ToString();
-	if (GridNameStr.ParseIntoArray(GridNameList, TEXT(".")))
+	// If the grid name is none, it directly maps to the main partition
+	if (GridName.IsNone())
+	{
+		MainPartitionTokens.Add(NAME_None);
+		return true;
+	}
+
+	// Split grid name into its partition and HLOD parts
+	TArray<FString> GridNameTokens;
+	if (!GridName.ToString().ParseIntoArray(GridNameTokens, TEXT(":")))
+	{
+		GridNameTokens.Add(GridName.ToString());
+	}
+
+	// Parsed grid names token should be either "RuntimeHash" or "RuntimeHash:HLODLayer"
+	if (GridNameTokens.Num() > 2)
+	{
+		return false;
+	}
+
+	// Parse the target main partition
+	TArray<FString> MainPartitionTokensStr;
+	if (GridNameTokens[0].ParseIntoArray(MainPartitionTokensStr, TEXT(".")))
 	{
-		TArray<FName> Result;
-		Algo::Transform(GridNameList, Result, [](const FString& GridName) { return *GridName; });
-		return MoveTemp(Result);
+		Algo::Transform(MainPartitionTokensStr, MainPartitionTokens, [](const FString& GridName) { return *GridName; });
 	}
-	return { GridName };
+
+	// Parse the target HLOD partition
+	if (GridNameTokens.IsValidIndex(1))
+	{
+		HLODPartitionTokens.Add(*GridNameTokens[1]);
+	}
+
+	return true;
 }
 
 URuntimeHashExternalStreamingObjectBase* UWorldPartitionRuntimeHashSet::StoreToExternalStreamingObject(UObject* StreamingObjectOuter, FName StreamingObjectName)
@@ -491,6 +538,11 @@ void UWorldPartitionRuntimeHashSet::PostEditChangeChainProperty(FPropertyChanged
 			RuntimePartitionDesc.MainLayer = NewObject<URuntimePartition>(this, RuntimePartitionDesc.Class, NAME_None);
 			RuntimePartitionDesc.MainLayer->SetDefaultValues();
 
+			if (UHLODLayer* DefaultHLODLayer = GetTypedOuter<UWorldPartition>()->GetDefaultHLODLayer())
+			{
+				RuntimePartitionDesc.HLODSetups.Emplace_GetRef().HLODLayer = DefaultHLODLayer;
+			}
+
 			RuntimePartitionDesc.UpdateHLODPartitionLayers();
 		}
 	}
@@ -538,6 +590,61 @@ void UWorldPartitionRuntimeHashSet::UpdateHLODPartitionLayers()
 		RuntimePartitionDesc.UpdateHLODPartitionLayers();
 	}
 }
+
+UWorldPartitionRuntimeHashSet::FCellUniqueId UWorldPartitionRuntimeHashSet::GetCellUniqueId(const URuntimePartition::FCellDescInstance& InCellDescInstance) const
+{
+	TStringBuilder<128> CellNameBuilder;
+
+	UWorld* OuterWorld = GetTypedOuter<UWorld>();
+	check(OuterWorld);
+
+	FString WorldName = FPackageName::GetShortName(OuterWorld->GetPackage());
+
+	CellNameBuilder.Appendf(TEXT("%s_%s"), *WorldName, *InCellDescInstance.Name.ToString());
+
+	const FDataLayersID DataLayersID(InCellDescInstance.DataLayerInstances);
+	if (DataLayersID.GetHash())
+	{
+		CellNameBuilder.Appendf(TEXT("_d%X"), DataLayersID.GetHash());
+	}
+
+	if (InCellDescInstance.ContentBundleID.IsValid())
+	{
+		CellNameBuilder.Appendf(TEXT("_c%s"), *UContentBundleDescriptor::GetContentBundleCompactString(InCellDescInstance.ContentBundleID));
+	}
+
+	if (!IsRunningCookCommandlet() && OuterWorld->IsGameWorld())
+	{
+		FString SourceWorldPath;
+		FString InstancedWorldPath;
+		if (OuterWorld->GetSoftObjectPathMapping(SourceWorldPath, InstancedWorldPath))
+		{
+			const FTopLevelAssetPath SourceAssetPath(SourceWorldPath);
+			WorldName = FPackageName::GetShortName(SourceAssetPath.GetPackageName());
+						
+			InstancedWorldPath = UWorld::RemovePIEPrefix(InstancedWorldPath);
+
+			const FString SourcePackageName = SourceAssetPath.GetPackageName().ToString();
+			const FTopLevelAssetPath InstanceAssetPath(InstancedWorldPath);
+			const FString InstancePackageName = InstanceAssetPath.GetPackageName().ToString();
+
+			if (int32 Index = InstancePackageName.Find(SourcePackageName); Index != INDEX_NONE)
+			{
+				CellNameBuilder.Appendf(TEXT("_i%s"), *InstancePackageName.Mid(Index + SourcePackageName.Len()));
+			}
+		}
+	}
+
+	FCellUniqueId CellUniqueId;
+	CellUniqueId.Name = CellNameBuilder.ToString();
+
+	FArchiveMD5 ArMD5;
+	ArMD5 << CellUniqueId.Name;
+	CellUniqueId.Guid = ArMD5.GetGuidFromHash();
+	check(CellUniqueId.Guid.IsValid());
+
+	return CellUniqueId;
+}
 #endif
 
 void UWorldPartitionRuntimeHashSet::ForEachStreamingData(TFunctionRef<void(const FRuntimePartitionStreamingData&)> Func) const
diff --git a/Engine/Source/Runtime/Engine/Private/WorldPartition/RuntimeHashSet/WorldPartitionRuntimeHashSetHLODGeneration.cpp b/Engine/Source/Runtime/Engine/Private/WorldPartition/RuntimeHashSet/WorldPartitionRuntimeHashSetHLODGeneration.cpp
--- a/Engine/Source/Runtime/Engine/Private/WorldPartition/RuntimeHashSet/WorldPartitionRuntimeHashSetHLODGeneration.cpp
+++ b/Engine/Source/Runtime/Engine/Private/WorldPartition/RuntimeHashSet/WorldPartitionRuntimeHashSetHLODGeneration.cpp
@@ -3,7 +3,145 @@
 #include "WorldPartition/RuntimeHashSet/WorldPartitionRuntimeHashSet.h"
 #include "WorldPartition/RuntimeHashSet/RuntimePartition.h"
 
+#include "WorldPartition/WorldPartition.h"
+#include "WorldPartition/WorldPartitionLog.h"
+#include "WorldPartition/WorldPartitionHelpers.h"
+
+#include "WorldPartition/HLOD/HLODLayer.h"
+#include "WorldPartition/HLOD/HLODActor.h"
+#include "WorldPartition/HLOD/HLODActorDesc.h"
+#include "WorldPartition/HLOD/HLODModifier.h"
+#include "WorldPartition/HLOD/IWorldPartitionHLODUtilities.h"
+#include "WorldPartition/HLOD/IWorldPartitionHLODUtilitiesModule.h"
+
+#include "WorldPartition/DataLayer/DataLayerManager.h"
+#include "WorldPartition/DataLayer/DataLayerUtils.h"
+
+#include "UObject/SavePackage.h"
+#include "HAL/PlatformFileManager.h"
+
 #if WITH_EDITOR
+class FHLODStreamingGenerationContext : public IStreamingGenerationContext
+{
+	using FActorSetInstanceList = TArray<IStreamingGenerationContext::FActorSetInstance>;
+
+public:
+	FHLODStreamingGenerationContext()
+		: WorldBounds(ForceInit)
+	{}
+
+	virtual FBox GetWorldBounds() const override
+	{
+		return WorldBounds;
+	}
+
+	virtual const FActorSetContainer* GetMainWorldContainer() const override
+	{
+		return &ActorSetContainer;
+	}
+
+	virtual void ForEachActorSetInstance(TFunctionRef<void(const FActorSetInstance&)> Func) const override
+	{
+		for (const FActorSetInstance& ActorSetInstance : ActorSetInstanceList)
+		{
+			Func(ActorSetInstance);
+		}
+	}
+
+	virtual void ForEachActorSetContainer(TFunctionRef<void(const FActorSetContainer&)> Func) const override
+	{
+		Func(ActorSetContainer);
+	}
+
+	FBox WorldBounds;
+	FActorSetContainer ActorSetContainer;
+	FActorDescViewMap ActorDescViewMap;
+	FActorSetInstanceList ActorSetInstanceList;
+};
+
+namespace Private
+{
+	static void GameTick(UWorld* InWorld)
+	{
+		static int32 TickRendering = 0;
+		static const int32 FlushRenderingFrequency = 256;
+
+		// Perform a GC when memory usage exceeds a given threshold
+		if (FWorldPartitionHelpers::ShouldCollectGarbage())
+		{
+			FWorldPartitionHelpers::DoCollectGarbage();
+		}
+
+		// When running with -AllowCommandletRendering we want to flush
+		if (((++TickRendering % FlushRenderingFrequency) == 0) && IsAllowCommandletRendering())
+		{
+			FWorldPartitionHelpers::FakeEngineTick(InWorld);
+		}
+	}
+	static void SavePackage(UPackage* Package, ISourceControlHelper* SourceControlHelper)
+	{
+		if (SourceControlHelper)
+		{
+			SourceControlHelper->Save(Package);
+		}
+		else
+		{
+			Package->MarkAsFullyLoaded();
+
+			FPackagePath PackagePath = FPackagePath::FromPackageNameChecked(Package->GetName());
+			const FString PackageFileName = PackagePath.GetLocalFullPath();
+			FSavePackageArgs SaveArgs;
+			SaveArgs.TopLevelFlags = RF_Standalone;
+			if (!UPackage::SavePackage(Package, nullptr, *PackageFileName, SaveArgs))
+			{
+				UE_LOG(LogWorldPartition, Error, TEXT("Error saving package %s."), *Package->GetName());
+				check(0);
+			}
+		}
+	}
+
+	static void DeletePackage(const FString& PackageName, ISourceControlHelper* SourceControlHelper)
+	{
+		FPackagePath PackagePath = FPackagePath::FromPackageNameChecked(PackageName);
+		const FString PackageFileName = PackagePath.GetLocalFullPath();
+
+		if (SourceControlHelper)
+		{
+			SourceControlHelper->Delete(PackageFileName);
+		}
+		else
+		{
+			FPlatformFileManager::Get().GetPlatformFile().DeleteFile(*PackageFileName);
+		}
+	}
+
+	static void DeletePackage(UPackage* Package, ISourceControlHelper* SourceControlHelper)
+	{
+		if (SourceControlHelper)
+		{
+			SourceControlHelper->Delete(Package);
+		}
+		else
+		{
+			DeletePackage(Package->GetName(), SourceControlHelper);
+		}
+	}
+
+	static void DeletePackage(UWorldPartition* WorldPartition, FWorldPartitionActorDesc* ActorDesc, ISourceControlHelper* SourceControlHelper)
+	{
+		if (ActorDesc->IsLoaded())
+		{
+			DeletePackage(ActorDesc->GetActor()->GetPackage(), SourceControlHelper);
+			WorldPartition->OnPackageDeleted(ActorDesc->GetActor()->GetPackage());
+		}
+		else
+		{
+			DeletePackage(ActorDesc->GetActorPackage().ToString(), SourceControlHelper);
+			WorldPartition->RemoveActor(ActorDesc->GetGuid());
+		}
+	}
+}
+
 bool UWorldPartitionRuntimeHashSet::SupportsHLODs() const
 {
 	for (const FRuntimePartitionDesc& RuntimePartitionDesc : RuntimePartitions)
@@ -22,6 +160,181 @@ bool UWorldPartitionRuntimeHashSet::SupportsHLODs() const
 
 bool UWorldPartitionRuntimeHashSet::SetupHLODActors(const IStreamingGenerationContext* StreamingGenerationContext, const UWorldPartition::FSetupHLODActorsParams& Params) const
 {
+	auto GetRuntimePartitionDescHLODSetup = [](const FRuntimePartitionDesc* RuntimePartitionDesc, const UHLODLayer* HLODLayer) -> const FRuntimePartitionHLODSetup*
+	{
+		for (const FRuntimePartitionHLODSetup& HLODSetup : RuntimePartitionDesc->HLODSetups)
+		{
+			if (HLODLayer == HLODSetup.HLODLayer)
+			{
+				return &HLODSetup;
+			}
+		}
+
+		return nullptr;
+	};
+
+	UWorldPartition* WorldPartition = GetOuterUWorldPartition();
+	const UDataLayerManager* DataLayerManager = WorldPartition->GetDataLayerManager();
+	IStreamingGenerationContext::FActorSetContainer* MainActorSetContainer = const_cast<IStreamingGenerationContext::FActorSetContainer*>(StreamingGenerationContext->GetMainWorldContainer());
+	const FStreamingGenerationActorDescCollection* MainContainerCollection = MainActorSetContainer->ActorDescCollection;
+
+	// Create the HLOD creation context
+	FHLODCreationContext HLODCreationContext;
+	for (FActorDescList::TConstIterator<AWorldPartitionHLOD> HLODIterator(MainContainerCollection->GetMainActorDescContainer()); HLODIterator; ++HLODIterator)
+	{
+		FWorldPartitionHandle HLODActorHandle(WorldPartition, HLODIterator->GetGuid());
+		HLODCreationContext.HLODActorDescs.Emplace(HLODIterator->GetActorName(), MoveTemp(HLODActorHandle));
+	}
+
+	TUniquePtr<IStreamingGenerationContext> CurrentHLODStreamingGenerationContext = MakeUnique<FStreamingGenerationContextProxy>(StreamingGenerationContext);
+
+	int32 HLODLevel = 0;
+	while (CurrentHLODStreamingGenerationContext)
+	{
+		TMap<URuntimePartition*, TArray<URuntimePartition::FCellDescInstance>> RuntimePartitionsStreamingDescs;
+		GenerateRuntimePartitionsStreamingDescs(CurrentHLODStreamingGenerationContext.Get(), RuntimePartitionsStreamingDescs);
+
+		int32 NumNextLayerHLODActors = 0;
+		TArray<FGuid> HLODActorGuids;
+		for (auto& [RuntimePartition, CellDescInstances] : RuntimePartitionsStreamingDescs)
+		{
+			for (URuntimePartition::FCellDescInstance& CellDescInstance : CellDescInstances)
+			{
+				// Here we split actors into their respective HLOD layers because we want to provide a specific runtime grid name for our HLOD setup to FWorldPartitionHLODUtilities::CreateHLODActors. 
+				TMap<const UHLODLayer*, TArray<IStreamingGenerationContext::FActorInstance>> HLODLayerActorInstances;
+				for (const IStreamingGenerationContext::FActorSetInstance* ActorSetInstance : CellDescInstance.ActorSetInstances)
+				{
+					ActorSetInstance->ForEachActor([this, ActorSetInstance, &HLODLayerActorInstances, &GetRuntimePartitionDescHLODSetup](const FGuid& ActorGuid)
+					{
+						IStreamingGenerationContext::FActorInstance ActorInstance(ActorGuid, ActorSetInstance);
+						const FWorldPartitionActorDescView& ActorDescView = ActorInstance.GetActorDescView();
+			
+						if (const UHLODLayer* ActorInstanceHLODLayer = Cast<const UHLODLayer>(ActorDescView.GetHLODLayer().ResolveObject()))
+						{
+							HLODLayerActorInstances.FindOrAdd(ActorInstanceHLODLayer).Add(ActorInstance);
+						}
+					});
+				}
+
+				for (auto& [HLODLayer, ActorInstances] : HLODLayerActorInstances)
+				{
+					// Fake tick
+					Private::GameTick(WorldPartition->GetWorld());
+
+					const FCellUniqueId CellUniqueId = GetCellUniqueId(CellDescInstance);
+
+					TArray<FName> MainPartitionTokens;
+					TArray<FName> HLODPartitionTokens;
+					verify(ParseGridName(ActorInstances[0].ActorSetInstance->RuntimeGrid, MainPartitionTokens, HLODPartitionTokens));
+					
+					FHLODCreationParams HLODCreationParams;
+					HLODCreationParams.WorldPartition = WorldPartition;
+					HLODCreationParams.CellName = CellUniqueId.Name;
+					HLODCreationParams.CellGuid = CellUniqueId.Guid;			
+					HLODCreationParams.CellBounds = CellDescInstance.Bounds;
+					HLODCreationParams.RuntimeGrid = *FString::Printf(TEXT("%s:%s"), *FString::JoinBy(MainPartitionTokens, TEXT("."), [](const FName Token) { return Token.ToString(); }), *HLODLayer->GetName());
+					HLODCreationParams.HLODLevel = HLODLevel;
+					HLODCreationParams.MinVisibleDistance = RuntimePartition->LoadingRange;
+					HLODCreationParams.ContentBundleGuid = CellDescInstance.ContentBundleID;
+					HLODCreationParams.DataLayerInstances = CellDescInstance.DataLayerInstances;
+
+					IWorldPartitionHLODUtilities* WPHLODUtilities = FModuleManager::Get().LoadModuleChecked<IWorldPartitionHLODUtilitiesModule>("WorldPartitionHLODUtilities").GetUtilities();
+					TArray<AWorldPartitionHLOD*> CellHLODActors = WPHLODUtilities->CreateHLODActors(HLODCreationContext, HLODCreationParams, ActorInstances);
+
+					if (!CellHLODActors.IsEmpty())
+					{
+						for (AWorldPartitionHLOD* CellHLODActor : CellHLODActors)
+						{
+							FGuid ActorGuid = CellHLODActor->GetActorGuid();
+
+							UPackage* CellHLODActorPackage = CellHLODActor->GetPackage();
+							if (CellHLODActorPackage->HasAnyPackageFlags(PKG_NewlyCreated))
+							{
+								// Get a reference to newly create actors so they get unloaded when we release the references
+								HLODCreationContext.ActorReferences.Emplace(WorldPartition, CellHLODActor->GetActorGuid());
+							}
+
+							HLODActorGuids.Add(ActorGuid);
+							NumNextLayerHLODActors += CellHLODActor->GetHLODLayer() ? 1 : 0;
+						}
+
+						if (!Params.bReportOnly)
+						{
+							for (AWorldPartitionHLOD* CellHLODActor : CellHLODActors)
+							{
+								if (CellHLODActor->GetPackage()->IsDirty())
+								{
+									Private::SavePackage(CellHLODActor->GetPackage(), Params.SourceControlHelper);
+								}
+							}
+						}
+					}
+
+					// Unload actors
+					HLODCreationContext.ActorReferences.Empty();
+				}
+			}
+		}
+
+		CurrentHLODStreamingGenerationContext.Reset();
+
+		// Build the next HLOD generation context
+		if (NumNextLayerHLODActors)
+		{
+			CurrentHLODStreamingGenerationContext = MakeUnique<FHLODStreamingGenerationContext>();
+			FHLODStreamingGenerationContext* HLODStreamingGenerationContext = (FHLODStreamingGenerationContext*)CurrentHLODStreamingGenerationContext.Get();
+
+			HLODStreamingGenerationContext->ActorSetInstanceList.Reserve(HLODActorGuids.Num());
+			HLODStreamingGenerationContext->ActorSetContainer.ActorDescViewMap = &HLODStreamingGenerationContext->ActorDescViewMap;
+
+			UE_LOG(LogWorldPartition, Log, TEXT("Creating HLOD context:"));
+			for (const FGuid& HLODActorGuid : HLODActorGuids)
+			{
+				FWorldPartitionActorDesc* HLODActorDesc = WorldPartition->GetActorDesc(HLODActorGuid);
+				check(HLODActorDesc);
+
+				FWorldPartitionActorDescView* HLODActorDescView = HLODStreamingGenerationContext->ActorDescViewMap.Emplace(HLODActorDesc);
+				HLODStreamingGenerationContext->WorldBounds += HLODActorDescView->GetRuntimeBounds();
+			
+				// Create actor set instances
+				IStreamingGenerationContext::FActorSet* ActorSet = HLODStreamingGenerationContext->ActorSetContainer.ActorSets.Emplace_GetRef(MakeUnique<IStreamingGenerationContext::FActorSet>()).Get();
+				ActorSet->Actors.Add(HLODActorDescView->GetGuid());
+
+				IStreamingGenerationContext::FActorSetInstance& ActorSetInstance = HLODStreamingGenerationContext->ActorSetInstanceList.Emplace_GetRef();
+
+				ActorSetInstance.Bounds = HLODActorDescView->GetRuntimeBounds();
+				ActorSetInstance.RuntimeGrid = HLODActorDescView->GetRuntimeGrid();
+				ActorSetInstance.bIsSpatiallyLoaded = HLODActorDescView->GetIsSpatiallyLoaded();
+				ActorSetInstance.ContentBundleID = MainContainerCollection->GetContentBundleGuid();
+				ActorSetInstance.ContainerInstance = &HLODStreamingGenerationContext->ActorSetContainer;
+				ActorSetInstance.ActorSet = ActorSet;
+
+				TArray<FName> RuntimeDataLayerInstanceNames;
+				if (FDataLayerUtils::ResolveRuntimeDataLayerInstanceNames(DataLayerManager, *HLODActorDescView, *MainActorSetContainer->ActorDescViewMap, RuntimeDataLayerInstanceNames))
+				{
+					HLODActorDescView->SetRuntimeDataLayerInstanceNames(RuntimeDataLayerInstanceNames);
+					ActorSetInstance.DataLayers = DataLayerManager->GetRuntimeDataLayerInstances(RuntimeDataLayerInstanceNames);
+				}
+
+				UE_LOG(LogWorldPartition, Log, TEXT("\t- %s"), *HLODActorDesc->ToString());
+			}
+
+			HLODLevel++;
+		}
+	}
+
+	// Destroy all unreferenced HLOD actors
+	if (!Params.bReportOnly)
+	{
+		for (const auto& HLODActorPair : HLODCreationContext.HLODActorDescs)
+		{
+			FWorldPartitionActorDesc* HLODActorDesc = HLODActorPair.Value.Get();
+			check(HLODActorDesc);
+
+			Private::DeletePackage(WorldPartition, HLODActorDesc, Params.SourceControlHelper);
+		}
+	}
+
 	return false;
 }
 #endif
\ No newline at end of file
diff --git a/Engine/Source/Runtime/Engine/Private/WorldPartition/RuntimeHashSet/WorldPartitionRuntimeHashSetStreamingGeneration.cpp b/Engine/Source/Runtime/Engine/Private/WorldPartition/RuntimeHashSet/WorldPartitionRuntimeHashSetStreamingGeneration.cpp
--- a/Engine/Source/Runtime/Engine/Private/WorldPartition/RuntimeHashSet/WorldPartitionRuntimeHashSetStreamingGeneration.cpp
+++ b/Engine/Source/Runtime/Engine/Private/WorldPartition/RuntimeHashSet/WorldPartitionRuntimeHashSetStreamingGeneration.cpp
@@ -6,31 +6,28 @@
 #include "WorldPartition/ContentBundle/ContentBundleDescriptor.h"
 #include "WorldPartition/WorldPartitionStreamingPolicy.h"
 #include "WorldPartition/DataLayer/DataLayersID.h"
-#include "Misc/ArchiveMD5.h"
 
 #if WITH_EDITOR
-bool UWorldPartitionRuntimeHashSet::GenerateStreaming(UWorldPartitionStreamingPolicy* StreamingPolicy, const IStreamingGenerationContext* StreamingGenerationContext, TArray<FString>* OutPackagesToGenerate)
+bool UWorldPartitionRuntimeHashSet::GenerateRuntimePartitionsStreamingDescs(const IStreamingGenerationContext* StreamingGenerationContext, TMap<URuntimePartition*, TArray<URuntimePartition::FCellDescInstance>>& OutRuntimeCellDescs) const
 {
-	verify(Super::GenerateStreaming(StreamingPolicy, StreamingGenerationContext, OutPackagesToGenerate));
-
 	UWorldPartition* WorldPartition = GetOuterUWorldPartition();
 	UWorld* World = WorldPartition->GetWorld();
 	UWorld* OuterWorld = GetTypedOuter<UWorld>();
 	const bool bIsMainWorldPartition = (World == OuterWorld);
 
-	check(!PersistentPartitionDesc.Class);
-	PersistentPartitionDesc.Class = URuntimePartitionPersistent::StaticClass();
-	PersistentPartitionDesc.Name = NAME_PersistentLevel;
-	PersistentPartitionDesc.MainLayer = NewObject<URuntimePartition>(this, URuntimePartitionPersistent::StaticClass(), NAME_None);
-	PersistentPartitionDesc.MainLayer->Name = NAME_PersistentLevel;
-
 	//
 	// Split actor sets into their corresponding runtime partition implementation
 	//
 	TMap<FName, const FRuntimePartitionDesc*> NameToRuntimePartitionDescMap;
 
-	NameToRuntimePartitionDescMap.Add(NAME_None, &RuntimePartitions[0]);				// Actors with RuntimeGrid=None will be assigned to the default partition
-	NameToRuntimePartitionDescMap.Add(NAME_PersistentLevel, &PersistentPartitionDesc);	// Non-spatially loaded actors will be assigned to the persistent partition
+	// Actors with RuntimeGrid set to None will be assigned to the default partition
+	NameToRuntimePartitionDescMap.Add(NAME_None, &RuntimePartitions[0]);
+
+	// Non-spatially loaded actors will be assigned to the persistent partition
+	if (PersistentPartitionDesc.Class)
+	{
+		NameToRuntimePartitionDescMap.Add(NAME_PersistentLevel, &PersistentPartitionDesc);
+	}
 
 	for (const FRuntimePartitionDesc& RuntimePartitionDesc : RuntimePartitions)
 	{
@@ -40,180 +37,199 @@ bool UWorldPartitionRuntimeHashSet::GenerateStreaming(UWorldPartitionStreamingPo
 	TMap<URuntimePartition*, TArray<const IStreamingGenerationContext::FActorSetInstance*>> RuntimePartitionsToActorSetMap;
 	StreamingGenerationContext->ForEachActorSetInstance([this, &NameToRuntimePartitionDescMap, &RuntimePartitionsToActorSetMap](const IStreamingGenerationContext::FActorSetInstance& ActorSetInstance)
 	{
-		const TArray<FName> ActorSetRuntimeGrid = ActorSetInstance.bIsSpatiallyLoaded ? ParseGridName(ActorSetInstance.RuntimeGrid) : TArray<FName>({ NAME_PersistentLevel });
+		TArray<FName> MainPartitionTokens;
+		TArray<FName> HLODPartitionTokens;
 
-		if (const FRuntimePartitionDesc** RuntimePartitionDesc = NameToRuntimePartitionDescMap.Find(ActorSetRuntimeGrid[0]))
+		if (!ActorSetInstance.bIsSpatiallyLoaded)
 		{
-			RuntimePartitionsToActorSetMap.FindOrAdd((*RuntimePartitionDesc)->MainLayer).Add(&ActorSetInstance);
+			MainPartitionTokens.Add(NAME_PersistentLevel);
 		}
-	});
-
-	struct FCellDescInstance : public URuntimePartition::FCellDesc
-	{
-		FCellDescInstance(const URuntimePartition::FCellDesc& InCellDesc, URuntimePartition* InSourcePartition, const TArray<const UDataLayerInstance*>& InDataLayers)
-			: URuntimePartition::FCellDesc(InCellDesc)
-			, SourcePartition(InSourcePartition)
-			, DataLayers(InDataLayers)
-		{}
-
-		URuntimePartition* SourcePartition;
-		TArray<const UDataLayerInstance*> DataLayers;
-	};
-
-	//
-	// Generate runtime partitions streaming data
-	//
-	TArray<FCellDescInstance> RuntimeCellDescsInstances;
-	{
-		TSet<FName> CellDescsNames;
-		for (auto [RuntimePartition, ActorSetInstances] : RuntimePartitionsToActorSetMap)
+		else
 		{
-			// Gather runtime partition cell descs
-			URuntimePartition::FGenerateStreamingParams GenerateStreamingParams;
-			GenerateStreamingParams.ActorSetInstances = &ActorSetInstances;
-
-			URuntimePartition::FGenerateStreamingResult GenerateStreamingResult;
-
-			if (!RuntimePartition->GenerateStreaming(GenerateStreamingParams, GenerateStreamingResult))
-			{
-				return false;
-			}
+			verify(ParseGridName(ActorSetInstance.RuntimeGrid, MainPartitionTokens, HLODPartitionTokens));
+		}
 
-			// Split cell descs into data layers
-			for (const URuntimePartition::FCellDesc& RuntimeCellDesc : GenerateStreamingResult.RuntimeCellDescs)
+		check(!MainPartitionTokens.IsEmpty());
+		if (const FRuntimePartitionDesc** RuntimePartitionDesc = NameToRuntimePartitionDescMap.Find(MainPartitionTokens[0]))
+		{
+			if (!HLODPartitionTokens.IsEmpty())
 			{
-				TMap<FDataLayersID, FCellDescInstance> RuntimeCellDescsInstancesSet;
+				bool bFoundHLODPartitionLayer = false;
 
-				bool bCellNameExists;
-				CellDescsNames.Add(RuntimeCellDesc.Name, &bCellNameExists);
-				check(!bCellNameExists);
-
-				for (const IStreamingGenerationContext::FActorInstance& ActorInstance : RuntimeCellDesc.ActorInstances)
+				for (const FRuntimePartitionHLODSetup& HLODSetup : (*RuntimePartitionDesc)->HLODSetups)
 				{
-					const FDataLayersID DataLayersID(ActorInstance.ActorSetInstance->DataLayers);
-					FCellDescInstance* CellDescInstance = RuntimeCellDescsInstancesSet.Find(DataLayersID);
-
-					if (!CellDescInstance)
+					for (const FRuntimePartitionHLODSetupLayer& HLODPartitionLayer : HLODSetup.PartitionLayers)
 					{
-						CellDescInstance = &RuntimeCellDescsInstancesSet.Emplace(DataLayersID, FCellDescInstance(RuntimeCellDesc, RuntimePartition, ActorInstance.ActorSetInstance->DataLayers));
-						CellDescInstance->ActorInstances.Empty();
+						if (HLODPartitionLayer.HLODLayer.GetName() == HLODPartitionTokens[0])
+						{
+							RuntimePartitionsToActorSetMap.FindOrAdd(HLODPartitionLayer.PartitionLayer).Add(&ActorSetInstance);
+							bFoundHLODPartitionLayer = true;
+							break;
+						}
 					}
 
-					CellDescInstance->ActorInstances.Add(ActorInstance);
+					if (bFoundHLODPartitionLayer)
+					{
+						break;
+					}
 				}
-
-				Algo::Transform(RuntimeCellDescsInstancesSet, RuntimeCellDescsInstances, [](const auto& Value) { return Value.Value; });
+			}
+			else
+			{
+				RuntimePartitionsToActorSetMap.FindOrAdd((*RuntimePartitionDesc)->MainLayer).Add(&ActorSetInstance);
 			}
 		}
-	}
+	});
 
 	//
-	// Create and populate streaming object
+	// Generate runtime partitions streaming data
 	//
-	check(RuntimeStreamingData.IsEmpty());
-
-	// Generate runtime cells
-	auto CreateRuntimeCellFromCellDesc = [this](const URuntimePartition::FCellDesc& CellDesc, const TArray<const UDataLayerInstance*>& DataLayers, TSubclassOf<UWorldPartitionRuntimeCell> CellClass, TSubclassOf<UWorldPartitionRuntimeCellData> CellDataClass)
+	TMap<URuntimePartition*, TArray<URuntimePartition::FCellDesc>> RuntimePartitionsStreamingDescs;
+	for (auto [RuntimePartition, ActorSetInstances] : RuntimePartitionsToActorSetMap)
 	{
-		FString CellObjectName;
-		FGuid CellGuid;
+		// Gather runtime partition cell descs
+		URuntimePartition::FGenerateStreamingParams GenerateStreamingParams;
+		GenerateStreamingParams.ActorSetInstances = &ActorSetInstances;
+
+		URuntimePartition::FGenerateStreamingResult GenerateStreamingResult;
+
+		if (!RuntimePartition->GenerateStreaming(GenerateStreamingParams, GenerateStreamingResult))
 		{
-			UWorld* OuterWorld = GetTypedOuter<UWorld>();
-			check(OuterWorld);
+			return false;
+		}
 
-			FString WorldName = FPackageName::GetShortName(OuterWorld->GetPackage());
+		RuntimePartitionsStreamingDescs.Add(RuntimePartition, GenerateStreamingResult.RuntimeCellDescs);
+	}
 
-			CellObjectName = FString::Printf(TEXT("%s_%s"), *WorldName, *CellDesc.Name.ToString());
+	//
+	// Generate runtime partitions streaming data
+	//
+	TSet<FName> CellDescsNames;
+	for (auto& [RuntimePartition, RuntimeCellDescs] : RuntimePartitionsStreamingDescs)
+	{
+		for (const URuntimePartition::FCellDesc& RuntimeCellDesc : RuntimeCellDescs)
+		{
+			TMap<FDataLayersID, URuntimePartition::FCellDescInstance> RuntimeCellDescsInstancesSet;
+
+			bool bCellNameExists;
+			CellDescsNames.Add(RuntimeCellDesc.Name, &bCellNameExists);
+			check(!bCellNameExists);
 
-			const FDataLayersID DataLayersID(DataLayers);
-			if (DataLayersID.GetHash())
+			// Split cell descs into data layers
+			for (const IStreamingGenerationContext::FActorSetInstance* ActorSetInstance : RuntimeCellDesc.ActorSetInstances)
 			{
-				CellObjectName += FString::Printf(TEXT("_d%X"), DataLayersID.GetHash());
+				const FDataLayersID DataLayersID(ActorSetInstance->DataLayers);
+				URuntimePartition::FCellDescInstance* CellDescInstance = RuntimeCellDescsInstancesSet.Find(DataLayersID);
+
+				if (!CellDescInstance)
+				{
+					CellDescInstance = &RuntimeCellDescsInstancesSet.Emplace(DataLayersID, URuntimePartition::FCellDescInstance(RuntimeCellDesc, RuntimePartition, ActorSetInstance->DataLayers, ActorSetInstance->ContentBundleID));
+					CellDescInstance->ActorSetInstances.Empty();
+				}
+
+				CellDescInstance->ActorSetInstances.Add(ActorSetInstance);
 			}
 
-			if (CellDesc.ContentBundleID.IsValid())
+			for (auto& [DataLayersID, CellDescInstance] : RuntimeCellDescsInstancesSet)
 			{
-				CellObjectName += FString::Printf(TEXT("_c%s"), *UContentBundleDescriptor::GetContentBundleCompactString(CellDesc.ContentBundleID));
+				OutRuntimeCellDescs.FindOrAdd(RuntimePartition).Add(CellDescInstance);
 			}
+		}
+	}
 
-			if (!IsRunningCookCommandlet() && OuterWorld->IsGameWorld())
-			{
-				FString SourceWorldPath;
-				FString InstancedWorldPath;
-				if (OuterWorld->GetSoftObjectPathMapping(SourceWorldPath, InstancedWorldPath))
-				{
-					const FTopLevelAssetPath SourceAssetPath(SourceWorldPath);
-					WorldName = FPackageName::GetShortName(SourceAssetPath.GetPackageName());
-						
-					InstancedWorldPath = UWorld::RemovePIEPrefix(InstancedWorldPath);
+	return true;
+}
 
-					const FString SourcePackageName = SourceAssetPath.GetPackageName().ToString();
-					const FTopLevelAssetPath InstanceAssetPath(InstancedWorldPath);
-					const FString InstancePackageName = InstanceAssetPath.GetPackageName().ToString();
+bool UWorldPartitionRuntimeHashSet::GenerateStreaming(UWorldPartitionStreamingPolicy* StreamingPolicy, const IStreamingGenerationContext* StreamingGenerationContext, TArray<FString>* OutPackagesToGenerate)
+{
+	verify(Super::GenerateStreaming(StreamingPolicy, StreamingGenerationContext, OutPackagesToGenerate));
 
-					if (int32 Index = InstancePackageName.Find(SourcePackageName); Index != INDEX_NONE)
-					{
-						CellObjectName += FString::Printf(TEXT("_i%s"), *InstancePackageName.Mid(Index + SourcePackageName.Len()));
-					}
-				}
-			}
+	UWorldPartition* WorldPartition = GetOuterUWorldPartition();
+	UWorld* World = WorldPartition->GetWorld();
+	UWorld* OuterWorld = GetTypedOuter<UWorld>();
+	const bool bIsMainWorldPartition = (World == OuterWorld);
 
-			FArchiveMD5 ArMD5;
-			ArMD5 << CellObjectName;
-			CellGuid = ArMD5.GetGuidFromHash();
-			check(CellGuid.IsValid());
-		}
+	check(!PersistentPartitionDesc.Class);
+	PersistentPartitionDesc.Class = URuntimePartitionPersistent::StaticClass();
+	PersistentPartitionDesc.Name = NAME_PersistentLevel;
+	PersistentPartitionDesc.MainLayer = NewObject<URuntimePartition>(this, URuntimePartitionPersistent::StaticClass(), NAME_None);
+	PersistentPartitionDesc.MainLayer->Name = NAME_PersistentLevel;
 
-		UWorldPartitionRuntimeCell* RuntimeCell = Super::CreateRuntimeCell(CellClass, CellDataClass, CellObjectName, TEXT(""));
+	//
+	// Generate runtime partitions streaming cell desccriptors
+	//
+	TMap<URuntimePartition*, TArray<URuntimePartition::FCellDescInstance>> RuntimePartitionsStreamingDescs;
+	GenerateRuntimePartitionsStreamingDescs(StreamingGenerationContext, RuntimePartitionsStreamingDescs);
+
+	//
+	// Create and populate streaming object
+	//
+
+	// Generate runtime cells
+	auto CreateRuntimeCellFromCellDesc = [this](const URuntimePartition::FCellDescInstance& CellDescInstance, TSubclassOf<UWorldPartitionRuntimeCell> CellClass, TSubclassOf<UWorldPartitionRuntimeCellData> CellDataClass)
+	{
+		const FCellUniqueId CellUniqueId = GetCellUniqueId(CellDescInstance);
 
-		RuntimeCell->SetIsAlwaysLoaded(!CellDesc.bIsSpatiallyLoaded);
-		RuntimeCell->SetDataLayers(DataLayers);
-		RuntimeCell->SetContentBundleUID(CellDesc.ContentBundleID);
-		RuntimeCell->SetPriority(CellDesc.Priority);
-		RuntimeCell->SetClientOnlyVisible(CellDesc.bClientOnlyVisible);
-		RuntimeCell->SetBlockOnSlowLoading(CellDesc.bBlockOnSlowStreaming);
+		UWorldPartitionRuntimeCell* RuntimeCell = Super::CreateRuntimeCell(CellClass, CellDataClass, CellUniqueId.Name, TEXT(""));
+
+		RuntimeCell->SetIsAlwaysLoaded(!CellDescInstance.bIsSpatiallyLoaded);
+		RuntimeCell->SetDataLayers(CellDescInstance.DataLayerInstances);
+		RuntimeCell->SetContentBundleUID(CellDescInstance.ContentBundleID);
+		RuntimeCell->SetPriority(CellDescInstance.Priority);
+		RuntimeCell->SetClientOnlyVisible(CellDescInstance.bClientOnlyVisible);
+		RuntimeCell->SetBlockOnSlowLoading(CellDescInstance.bBlockOnSlowStreaming);
 		RuntimeCell->SetIsHLOD(false);
-		RuntimeCell->SetGuid(CellGuid);
+		RuntimeCell->SetGuid(CellUniqueId.Guid);
 
 		UWorldPartitionRuntimeCellDataSpatialHashSet* RuntimeCellDataHashSet = CastChecked<UWorldPartitionRuntimeCellDataSpatialHashSet>(RuntimeCell->RuntimeCellData);
-		RuntimeCellDataHashSet->DebugName = CellObjectName;
-		RuntimeCellDataHashSet->Level = CellDesc.Level;
+		RuntimeCellDataHashSet->DebugName = CellUniqueId.Name;
+		RuntimeCellDataHashSet->Level = CellDescInstance.Level;
 
 		return RuntimeCell;
 	};
 
 	TArray<UWorldPartitionRuntimeCell*> RuntimeCells;
 	TMap<URuntimePartition*, FRuntimePartitionStreamingData> RuntimePartitionsStreamingData;
-	for (const FCellDescInstance& CellDescInstance : RuntimeCellDescsInstances)
+	for (auto& [RuntimePartition, CellDescInstances] : RuntimePartitionsStreamingDescs)
 	{
-		UWorldPartitionRuntimeCell* RuntimeCell = RuntimeCells.Emplace_GetRef(CreateRuntimeCellFromCellDesc(CellDescInstance, CellDescInstance.DataLayers, StreamingPolicy->GetRuntimeCellClass(), UWorldPartitionRuntimeCellDataSpatialHashSet::StaticClass()));
-		PopulateRuntimeCell(RuntimeCell, CellDescInstance.ActorInstances, nullptr);
-
-		// Override the cell bounds if the runtime partition provided one
-		if (CellDescInstance.Bounds.IsValid)
+		for (URuntimePartition::FCellDescInstance& CellDescInstance : CellDescInstances)
 		{
-			RuntimeCell->RuntimeCellData->ContentBounds = CellDescInstance.Bounds;
-		}
+			const bool bIsCellAlwaysLoaded = !CellDescInstance.bIsSpatiallyLoaded && !CellDescInstance.DataLayerInstances.Num() && !CellDescInstance.ContentBundleID.IsValid();
 
-		// Create partition streaming data
-		FRuntimePartitionStreamingData& StreamingData = RuntimePartitionsStreamingData.FindOrAdd(CellDescInstance.SourcePartition);
+			TArray<IStreamingGenerationContext::FActorInstance> CellActorInstances;
+			if (PopulateCellActorInstances(CellDescInstance.ActorSetInstances, bIsMainWorldPartition, bIsCellAlwaysLoaded, CellActorInstances))
+			{
+				UWorldPartitionRuntimeCell* RuntimeCell = RuntimeCells.Emplace_GetRef(CreateRuntimeCellFromCellDesc(CellDescInstance,StreamingPolicy->GetRuntimeCellClass(), UWorldPartitionRuntimeCellDataSpatialHashSet::StaticClass()));
+				PopulateRuntimeCell(RuntimeCell, CellActorInstances, OutPackagesToGenerate);
 
-		StreamingData.Name = CellDescInstance.SourcePartition->Name;
-		StreamingData.LoadingRange = CellDescInstance.SourcePartition->LoadingRange;
+				// Override the cell bounds if the runtime partition provided one
+				if (CellDescInstance.Bounds.IsValid)
+				{
+					RuntimeCell->RuntimeCellData->ContentBounds = CellDescInstance.Bounds;
+				}
 
-		if (RuntimeCell->IsAlwaysLoaded())
-		{
-			StreamingData.NonStreamingCells.Add(RuntimeCell);
-		}
-		else
-		{
-			StreamingData.StreamingCells.Add(RuntimeCell);
+				// Create partition streaming data
+				FRuntimePartitionStreamingData& StreamingData = RuntimePartitionsStreamingData.FindOrAdd(CellDescInstance.SourcePartition);
+
+				StreamingData.Name = CellDescInstance.SourcePartition->Name;
+				StreamingData.LoadingRange = CellDescInstance.SourcePartition->LoadingRange;
+
+				if (RuntimeCell->IsAlwaysLoaded())
+				{
+					StreamingData.NonStreamingCells.Add(RuntimeCell);
+				}
+				else
+				{
+					StreamingData.StreamingCells.Add(RuntimeCell);
+				}
+			}
 		}
 	}
 
 	//
 	// Finalize streaming object
 	//
+	check(RuntimeStreamingData.IsEmpty());
 	for (auto& [Partition, StreamingData] : RuntimePartitionsStreamingData)
 	{
 		RuntimeStreamingData.Emplace(MoveTemp(StreamingData));
diff --git a/Engine/Source/Runtime/Engine/Private/WorldPartition/WorldPartitionRuntimeHash.cpp b/Engine/Source/Runtime/Engine/Private/WorldPartition/WorldPartitionRuntimeHash.cpp
--- a/Engine/Source/Runtime/Engine/Private/WorldPartition/WorldPartitionRuntimeHash.cpp
+++ b/Engine/Source/Runtime/Engine/Private/WorldPartition/WorldPartitionRuntimeHash.cpp
@@ -228,7 +228,8 @@ void UWorldPartitionRuntimeHash::PopulateRuntimeCell(UWorldPartitionRuntimeCell*
 	for (const IStreamingGenerationContext::FActorInstance& ActorInstance : ActorInstances)
 	{
 		const FWorldPartitionActorDescView& ActorDescView = ActorInstance.GetActorDescView();
-		RuntimeCell->AddActorToCell(ActorDescView, ActorInstance.GetContainerID(), ActorInstance.GetTransform(), ActorInstance.GetActorDescContainer());
+		const UActorDescContainer* ActorDescContainer = ActorDescView.GetActorDesc()->GetContainer();
+		RuntimeCell->AddActorToCell(ActorDescView, ActorInstance.GetContainerID(), ActorInstance.GetTransform(), ActorDescContainer);
 		const FBox RuntimeBounds = ActorDescView.GetRuntimeBounds();
 		if (RuntimeBounds.IsValid)
 		{
diff --git a/Engine/Source/Runtime/Engine/Private/WorldPartition/WorldPartitionStreamingGeneration.cpp b/Engine/Source/Runtime/Engine/Private/WorldPartition/WorldPartitionStreamingGeneration.cpp
--- a/Engine/Source/Runtime/Engine/Private/WorldPartition/WorldPartitionStreamingGeneration.cpp
+++ b/Engine/Source/Runtime/Engine/Private/WorldPartition/WorldPartitionStreamingGeneration.cpp
@@ -79,6 +79,11 @@ FWorldPartitionActorDescView* FActorDescViewMap::Emplace(const FGuid& InGuid, co
 	return NewActorDescView;
 }
 
+FWorldPartitionActorDescView* FActorDescViewMap::Emplace(const FWorldPartitionActorDesc* InActorDesc)
+{
+	return Emplace(InActorDesc->GetGuid(), FWorldPartitionActorDescView(InActorDesc));
+}
+
 UWorldPartition::FCheckForErrorsParams::FCheckForErrorsParams()
 	: ErrorHandler(nullptr)
 	, ActorDescContainer(nullptr)
diff --git a/Engine/Source/Runtime/Engine/Private/WorldPartition/WorldPartitionStreamingGenerationContext.cpp b/Engine/Source/Runtime/Engine/Private/WorldPartition/WorldPartitionStreamingGenerationContext.cpp
--- a/Engine/Source/Runtime/Engine/Private/WorldPartition/WorldPartitionStreamingGenerationContext.cpp
+++ b/Engine/Source/Runtime/Engine/Private/WorldPartition/WorldPartitionStreamingGenerationContext.cpp
@@ -19,11 +19,6 @@ const FTransform& IStreamingGenerationContext::FActorInstance::GetTransform() co
 	return ActorSetInstance->Transform;
 }
 
-const UActorDescContainer* IStreamingGenerationContext::FActorInstance::GetActorDescContainer() const
-{
-	return GetActorDescView().GetActorDesc()->GetContainer();
-}
-
 const FBox IStreamingGenerationContext::FActorInstance::GetBounds() const
 {
 	return ActorSetInstance->Bounds;
diff --git a/Engine/Source/Runtime/Engine/Public/WorldPartition/DataLayer/DataLayerManager.h b/Engine/Source/Runtime/Engine/Public/WorldPartition/DataLayer/DataLayerManager.h
--- a/Engine/Source/Runtime/Engine/Public/WorldPartition/DataLayer/DataLayerManager.h
+++ b/Engine/Source/Runtime/Engine/Public/WorldPartition/DataLayer/DataLayerManager.h
@@ -189,6 +189,7 @@ private:
 	friend class FWorldPartitionStreamingGenerator;
 	friend class FWorldPartitionActorDesc;
 	friend class UWorldPartitionRuntimeSpatialHash;
+	friend class UWorldPartitionRuntimeHashSet;
 	friend class UWorldPartitionConvertCommandlet;
 	friend class FDataLayersBroadcast;
 	friend class FDataLayerUtils;
diff --git a/Engine/Source/Runtime/Engine/Public/WorldPartition/RuntimeHashSet/RuntimePartition.h b/Engine/Source/Runtime/Engine/Public/WorldPartition/RuntimeHashSet/RuntimePartition.h
--- a/Engine/Source/Runtime/Engine/Public/WorldPartition/RuntimeHashSet/RuntimePartition.h
+++ b/Engine/Source/Runtime/Engine/Public/WorldPartition/RuntimeHashSet/RuntimePartition.h
@@ -2,7 +2,6 @@
 #pragma once
 
 #include "CoreMinimal.h"
-#include "WorldPartition/RuntimeHashSet/WorldPartitionRuntimeHashSet.h"
 #include "WorldPartition/WorldPartitionStreamingGenerationContext.h"
 #include "RuntimePartition.generated.h"
 
@@ -17,6 +16,9 @@ public:
 	virtual void PostEditChangeProperty(FPropertyChangedEvent& InPropertyChangedEvent) override;
 	//~ End UObject Interface.
 
+	/*
+	 * Represents a cell descriptor, generated by runtime partitions. This is a streaming cell containing actors, without taking into account data layers, content bundles, etc.
+	 */
 	struct FCellDesc
 	{
 		FName Name;
@@ -24,13 +26,29 @@ public:
 		bool bIsSpatiallyLoaded;
 		bool bBlockOnSlowStreaming;
 		bool bClientOnlyVisible;
-		FGuid ContentBundleID;
 		int32 Priority;
 
 		/** Optional level value that can be used to filter debug display */
 		int32 Level;
 
-		TArray<IStreamingGenerationContext::FActorInstance> ActorInstances;
+		TArray<const IStreamingGenerationContext::FActorSetInstance*> ActorSetInstances;
+	};
+
+	/*
+	 * Represents a cell descriptor instance, which is an instance of a cell after being split into data layers, etc. and ready to be converted into a streaming level.
+	 */
+	struct FCellDescInstance : public FCellDesc
+	{
+		FCellDescInstance(const FCellDesc& InCellDesc, URuntimePartition* InSourcePartition, const TArray<const UDataLayerInstance*>& InDataLayerInstances, const FGuid& InContentBundleID)
+			: FCellDesc(InCellDesc)
+			, SourcePartition(InSourcePartition)
+			, DataLayerInstances(InDataLayerInstances)
+			, ContentBundleID(InContentBundleID)
+		{}
+
+		URuntimePartition* SourcePartition;
+		TArray<const UDataLayerInstance*> DataLayerInstances;
+		FGuid ContentBundleID;
 	};
 
 	struct FGenerateStreamingParams
@@ -44,8 +62,8 @@ public:
 	};
 
 	virtual void SetDefaultValues();
-	virtual bool SupportsHLODs() const PURE_VIRTUAL(URuntimePartition::SupportsHLODs, return false;);
-	virtual bool IsValidGrid(FName InGridName) const PURE_VIRTUAL(URuntimePartition::IsValidGrid, return false;);
+	virtual bool SupportsHLODs() const PURE_VIRTUAL(URuntimePartition::SupportsHLODs, return true;);
+	virtual bool IsValidPartitionTokens(const TArray<FName> InPartitionTokens) const PURE_VIRTUAL(URuntimePartition::IsValidPartitionTokens, return false;);
 	virtual bool GenerateStreaming(const FGenerateStreamingParams& InParams, FGenerateStreamingResult& OutResult) PURE_VIRTUAL(URuntimePartition::GenerateStreaming, return false;);
 #endif
 
@@ -71,7 +89,6 @@ public:
 
 protected:
 #if WITH_EDITOR
-	bool PopulateCellActorInstances(const TArray<const IStreamingGenerationContext::FActorSetInstance*>& InActorSetInstances, bool bInIsMainWorldPartition, bool bInIsCellAlwaysLoaded, TArray<IStreamingGenerationContext::FActorInstance>& OutCellActorInstances);
-	FCellDesc CreateCellDesc(const FString& InName, bool bInIsSpatiallyLoaded, const FGuid& InContentBundleID, int32 InLevel, const TArray<IStreamingGenerationContext::FActorInstance>& InActorInstances);
+	FCellDesc CreateCellDesc(const FString& InName, bool bInIsSpatiallyLoaded, int32 InLevel, const TArray<const IStreamingGenerationContext::FActorSetInstance*>& InActorSetInstances);
 #endif
 };
\ No newline at end of file
diff --git a/Engine/Source/Runtime/Engine/Public/WorldPartition/RuntimeHashSet/RuntimePartitionLHGrid.h b/Engine/Source/Runtime/Engine/Public/WorldPartition/RuntimeHashSet/RuntimePartitionLHGrid.h
--- a/Engine/Source/Runtime/Engine/Public/WorldPartition/RuntimeHashSet/RuntimePartitionLHGrid.h
+++ b/Engine/Source/Runtime/Engine/Public/WorldPartition/RuntimeHashSet/RuntimePartitionLHGrid.h
@@ -17,9 +17,9 @@ public:
 	//~ End UObject Interface.
 
 	//~ Begin URuntimePartition interface
+	virtual bool SupportsHLODs() const override { return true; }
 	virtual void SetDefaultValues() override;
-	virtual bool SupportsHLODs() const override;
-	virtual bool IsValidGrid(FName GridName) const override;
+	virtual bool IsValidPartitionTokens(const TArray<FName> InPartitionTokens) const override;
 	virtual bool GenerateStreaming(const FGenerateStreamingParams& InParams, FGenerateStreamingResult& OutResult) override;
 	//~ End URuntimePartition interface
 #endif
diff --git a/Engine/Source/Runtime/Engine/Public/WorldPartition/RuntimeHashSet/RuntimePartitionLevelStreaming.h b/Engine/Source/Runtime/Engine/Public/WorldPartition/RuntimeHashSet/RuntimePartitionLevelStreaming.h
--- a/Engine/Source/Runtime/Engine/Public/WorldPartition/RuntimeHashSet/RuntimePartitionLevelStreaming.h
+++ b/Engine/Source/Runtime/Engine/Public/WorldPartition/RuntimeHashSet/RuntimePartitionLevelStreaming.h
@@ -13,15 +13,9 @@ class URuntimePartitionLevelStreaming : public URuntimePartition
 public:
 #if WITH_EDITOR
 	//~ Begin URuntimePartition interface
-	virtual bool SupportsHLODs() const override;
-	virtual bool IsValidGrid(FName GridName) const override;
+	virtual bool SupportsHLODs() const override { return true; }
+	virtual bool IsValidPartitionTokens(const TArray<FName> InPartitionTokens) const override;
 	virtual bool GenerateStreaming(const FGenerateStreamingParams& InParams, FGenerateStreamingResult& OutResult) override;
 	//~ End URuntimePartition interface
 #endif
-
-#if WITH_EDITORONLY_DATA
-	/** Splits actors into different sublevels based on their actor containers (level instance) */
-	UPROPERTY(EditAnywhere, Category = RuntimeSettings, Meta = (EditCondition = "!bIsHLODSetup", EditConditionHides, HideEditConditionToggle))
-	bool bOneLevelPerActorContainer;
-#endif
 };
\ No newline at end of file
diff --git a/Engine/Source/Runtime/Engine/Public/WorldPartition/RuntimeHashSet/RuntimePartitionPersistent.h b/Engine/Source/Runtime/Engine/Public/WorldPartition/RuntimeHashSet/RuntimePartitionPersistent.h
--- a/Engine/Source/Runtime/Engine/Public/WorldPartition/RuntimeHashSet/RuntimePartitionPersistent.h
+++ b/Engine/Source/Runtime/Engine/Public/WorldPartition/RuntimeHashSet/RuntimePartitionPersistent.h
@@ -14,7 +14,7 @@ public:
 #if WITH_EDITOR
 	//~ Begin URuntimePartition interface
 	virtual bool SupportsHLODs() const override { return false; }
-	virtual bool IsValidGrid(FName GridName) const override { return true; }
+	virtual bool IsValidPartitionTokens(const TArray<FName> InPartitionTokens) const override { return true; }
 	virtual bool GenerateStreaming(const FGenerateStreamingParams& InParams, FGenerateStreamingResult& OutResult) override;
 	//~ End URuntimePartition interface
 #endif
diff --git a/Engine/Source/Runtime/Engine/Public/WorldPartition/RuntimeHashSet/WorldPartitionRuntimeHashSet.h b/Engine/Source/Runtime/Engine/Public/WorldPartition/RuntimeHashSet/WorldPartitionRuntimeHashSet.h
--- a/Engine/Source/Runtime/Engine/Public/WorldPartition/RuntimeHashSet/WorldPartitionRuntimeHashSet.h
+++ b/Engine/Source/Runtime/Engine/Public/WorldPartition/RuntimeHashSet/WorldPartitionRuntimeHashSet.h
@@ -3,11 +3,11 @@
 
 #include "CoreMinimal.h"
 #include "WorldPartition/WorldPartitionRuntimeHash.h"
+#include "WorldPartition/RuntimeHashSet/RuntimePartition.h"
 #include "WorldPartition/RuntimeHashSet/StaticSpatialIndex.h"
 #include "WorldPartitionRuntimeHashSet.generated.h"
 
 class UHLODLayer;
-class URuntimePartition;
 class URuntimePartitionPersistent;
 struct FPropertyChangedChainEvent;
 
@@ -162,7 +162,7 @@ public:
 	ENGINE_API virtual void DumpStateLog(FHierarchicalLogArchive& Ar) const override;
 
 	// Helpers
-	static ENGINE_API TArray<FName> ParseGridName(FName GridName);
+	static ENGINE_API bool ParseGridName(FName GridName, TArray<FName>& MainPartitionTokens, TArray<FName>& HLODPartitionTokens);
 #endif
 
 	// External streaming object interface
@@ -179,8 +179,19 @@ public:
 
 private:
 #if WITH_EDITOR
+	/** Generate the runtime partitions streaming descs. */
+	bool GenerateRuntimePartitionsStreamingDescs(const IStreamingGenerationContext* StreamingGenerationContext, TMap<URuntimePartition*, TArray<URuntimePartition::FCellDescInstance>>& OutRuntimeCellDescs) const;
+
 	/** Update the partition layers to reflect the curent HLOD setups. */
 	void UpdateHLODPartitionLayers();
+
+	struct FCellUniqueId
+	{
+		FString Name;
+		FGuid Guid;
+	};
+
+	FCellUniqueId GetCellUniqueId(const URuntimePartition::FCellDescInstance& InCellDescInstance) const;
 #endif
 
 	ENGINE_API void ForEachStreamingData(TFunctionRef<void(const FRuntimePartitionStreamingData&)> Func) const;
diff --git a/Engine/Source/Runtime/Engine/Public/WorldPartition/WorldPartitionStreamingGeneration.h b/Engine/Source/Runtime/Engine/Public/WorldPartition/WorldPartitionStreamingGeneration.h
--- a/Engine/Source/Runtime/Engine/Public/WorldPartition/WorldPartitionStreamingGeneration.h
+++ b/Engine/Source/Runtime/Engine/Public/WorldPartition/WorldPartitionStreamingGeneration.h
@@ -50,6 +50,7 @@ public:
 	FActorDescViewMap& operator=(FActorDescViewMap&&) = default;
 
 	ENGINE_API FWorldPartitionActorDescView* Emplace(const FGuid& InActorGuid, const FWorldPartitionActorDescView& InActorDescView);
+	ENGINE_API FWorldPartitionActorDescView* Emplace(const FWorldPartitionActorDesc* InActorDesc);
 
 	FORCEINLINE int32 Num() const
 	{
diff --git a/Engine/Source/Runtime/Engine/Public/WorldPartition/WorldPartitionStreamingGenerationContext.h b/Engine/Source/Runtime/Engine/Public/WorldPartition/WorldPartitionStreamingGenerationContext.h
--- a/Engine/Source/Runtime/Engine/Public/WorldPartition/WorldPartitionStreamingGenerationContext.h
+++ b/Engine/Source/Runtime/Engine/Public/WorldPartition/WorldPartitionStreamingGenerationContext.h
@@ -6,6 +6,7 @@
 #include "WorldPartition/WorldPartitionActorContainerID.h"
 
 class FActorDescViewMap;
+class FWorldPartitionStreamingGenerator;
 class FStreamingGenerationActorDescCollection;
 class FWorldPartitionActorDescView;
 class UActorDescContainer;
@@ -16,11 +17,20 @@ struct FWorldPartitionRuntimeContainerResolver;
 class IStreamingGenerationContext
 {
 public:
+	virtual ~IStreamingGenerationContext()
+	{}
+
+	/**
+	 * An actor set represents a group of actors that needs to be part of the same streaming cell, because they have hard references between them.
+	 */
 	struct FActorSet
 	{
 		TArray<FGuid> Actors;
 	};
 
+	/**
+	 * An actor set container represents the list of actor sets in an actor container, e.g. a level instance.
+	 */
 	struct FActorSetContainer
 	{
 		FActorSetContainer()
@@ -33,10 +43,13 @@ public:
 		FActorSetContainer& operator=(const FActorSetContainer&) = delete;
 
 		const FActorDescViewMap* ActorDescViewMap;
-		const FStreamingGenerationActorDescCollection* ActorDescCollection;
+		const FStreamingGenerationActorDescCollection* ActorDescCollection; // Only used by UWorldPartitionRuntimeSpatialHash::SetupHLODActors
 		TArray<TUniquePtr<FActorSet>> ActorSets;
 	};
 
+	/**
+	 * An actor set instance is an actual intance of an actor set in the world.
+	 */
 	struct FActorSetInstance
 	{
 		FBox Bounds;
@@ -68,6 +81,9 @@ public:
 		}
 	};
 
+	/**
+	 * An actor instance represents a single instanced actor in the world.
+	 */
 	struct FActorInstance
 	{
 		FActorInstance(const FGuid& InActorGuid, const FActorSetInstance* InActorSetInstance)
@@ -81,7 +97,6 @@ public:
 		ENGINE_API const FWorldPartitionActorDescView& GetActorDescView() const;
 		ENGINE_API const FActorContainerID& GetContainerID() const;
 		ENGINE_API const FTransform& GetTransform() const;
-		ENGINE_API const UActorDescContainer* GetActorDescContainer() const;
 		ENGINE_API const FBox GetBounds() const;
 	};
 
@@ -90,4 +105,35 @@ public:
 	virtual void ForEachActorSetInstance(TFunctionRef<void(const FActorSetInstance&)> Func) const = 0;
 	virtual void ForEachActorSetContainer(TFunctionRef<void(const FActorSetContainer&)> Func) const = 0;
 };
+
+class FStreamingGenerationContextProxy : public IStreamingGenerationContext
+{
+public:
+	FStreamingGenerationContextProxy(const IStreamingGenerationContext* InSourceContext)
+		: SourceContext(InSourceContext)
+	{}
+
+	virtual FBox GetWorldBounds() const override
+	{
+		return SourceContext->GetWorldBounds();
+	}
+
+	virtual const FActorSetContainer* GetMainWorldContainer() const override
+	{
+		return SourceContext->GetMainWorldContainer();
+	}
+
+	virtual void ForEachActorSetInstance(TFunctionRef<void(const FActorSetInstance&)> Func) const override
+	{
+		SourceContext->ForEachActorSetInstance(Func);
+	}
+
+	virtual void ForEachActorSetContainer(TFunctionRef<void(const FActorSetContainer&)> Func) const override
+	{
+		SourceContext->ForEachActorSetContainer(Func);
+	}
+
+private:
+	const IStreamingGenerationContext* SourceContext;
+};
 #endif
\ No newline at end of file
