Fix Compile UnrealEditor StressTestUnity Win64

diff --git a/Engine/Source/Runtime/Engine/Private/WorldPartition/RuntimeHashSet/WorldPartitionRuntimeHashSetHLODGeneration.cpp b/Engine/Source/Runtime/Engine/Private/WorldPartition/RuntimeHashSet/WorldPartitionRuntimeHashSetHLODGeneration.cpp
--- a/Engine/Source/Runtime/Engine/Private/WorldPartition/RuntimeHashSet/WorldPartitionRuntimeHashSetHLODGeneration.cpp
+++ b/Engine/Source/Runtime/Engine/Private/WorldPartition/RuntimeHashSet/WorldPartitionRuntimeHashSetHLODGeneration.cpp
@@ -59,7 +59,7 @@ public:
 	FActorSetInstanceList ActorSetInstanceList;
 };
 
-namespace Private
+namespace PrivateUtils
 {
 	static void GameTick(UWorld* InWorld)
 	{
@@ -219,7 +219,7 @@ bool UWorldPartitionRuntimeHashSet::SetupHLODActors(const IStreamingGenerationCo
 				for (auto& [HLODLayer, ActorInstances] : HLODLayerActorInstances)
 				{
 					// Fake tick
-					Private::GameTick(WorldPartition->GetWorld());
+					PrivateUtils::GameTick(WorldPartition->GetWorld());
 
 					const FCellUniqueId CellUniqueId = GetCellUniqueId(CellDescInstance);
 
@@ -264,7 +264,7 @@ bool UWorldPartitionRuntimeHashSet::SetupHLODActors(const IStreamingGenerationCo
 							{
 								if (CellHLODActor->GetPackage()->IsDirty())
 								{
-									Private::SavePackage(CellHLODActor->GetPackage(), Params.SourceControlHelper);
+									PrivateUtils::SavePackage(CellHLODActor->GetPackage(), Params.SourceControlHelper);
 								}
 							}
 						}
@@ -331,7 +331,7 @@ bool UWorldPartitionRuntimeHashSet::SetupHLODActors(const IStreamingGenerationCo
 			FWorldPartitionActorDesc* HLODActorDesc = HLODActorPair.Value.Get();
 			check(HLODActorDesc);
 
-			Private::DeletePackage(WorldPartition, HLODActorDesc, Params.SourceControlHelper);
+			PrivateUtils::DeletePackage(WorldPartition, HLODActorDesc, Params.SourceControlHelper);
 		}
 	}
 
