Misc. fixes for the WIP FMaterialHLSLGenerator:

- Do not merge scope evaluations when preparing a LocalPHI if we know for certain that the local variable has the same value in all live scopes. This avoids adding the statement owning the value scopes as a dependency of the LocalPHI, which can form a circular dependency when all scope values are the same constant. When all scope values are the same constant, the LocalPHI becomes an inline constant but the constant still depends on the statement if the statement is not constant. If the statement directly or indirectly depends on the constant, we get a circular dependency.

- Use as key the scope address and requested type in addition to the LocalPHI expression address when looking up EmitLocalPHIMap to avoid incorrectly reusing PHI nodes from a different scope.

- FExpressionSelect now emits LWCSelect when at least one of the inputs is LWC.

- Truncate numeric vectors via swizzles instead of type casts which seems to cause dxil validation errors.

- Reset EmitValueMap after emitting shader code because some emitted values are different based on shader frequencies.

- Set bIsLWC to true for the float3 * FLWCMatrix -> FLWCVector3 case so LWCMultiply will be used instead of mul. Courtesy to tom.holmes.

- Cast all smoothstep operands to LWC if one of them is LWC. This is needed because LWCSmoothStep only takes inputs of the same type.

- float3 * float4x4 now appends 1.f to the 3-vector because EOperation::VecMulMatrix3 implies trasforming points.

- Prevent the TStringBuilder dynamic array in FMaterialCachedHLSLTree::GenerateTree from reallocating because TStringBuilder is not relocatable.

- FMaterialCachedHLSLTree::SetRequestedFields now only requests custom UVs equal to the number of pixel shader texture coordinates.

- Custom material output getters now emit both LWC and non-LWC versions.

- FMaterialCachedHLSLTree::GetMaterialAttributesDefaultValue is now private. FMaterialHLSLGenerator::GetMaterialAttributesDefaultExpression should be used instead because the defaults of some material properties are not constant values.

- Properly set shader define NEEDS_WORLD_POSITION_EXCLUDING_SHADER_OFFSETS.

- Set INTERPOLATE_VERTEX_COLOR should be set based on whether vertex color is used in pixel shaders not vertex shaders.

- Some EExternalInputs now fallbacks to different external inputs when they are not available in the given shader frequencies.

- FExpressionStaticTerrainLayerWeight now only prepares texture coordinates expression when any terrain layer is actually sampled. This avoids activating extra tex coord channels unnecessarily.

- FExpressionRuntimeVirtualTextureOutput now sets bHasRuntimeVirtualTextureOutput in cached expression data when it is used.

diff --git a/Engine/Source/Editor/UnrealEd/Private/MaterialGraph.cpp b/Engine/Source/Editor/UnrealEd/Private/MaterialGraph.cpp
--- a/Engine/Source/Editor/UnrealEd/Private/MaterialGraph.cpp
+++ b/Engine/Source/Editor/UnrealEd/Private/MaterialGraph.cpp
@@ -121,7 +121,7 @@ void PrepareHLSLTree(UE::HLSLTree::FEmitContext& EmitContext,
 	FEmitScope* EmitResultScope = EmitContext.PrepareScope(CachedTree.GetResultScope());
 
 	FRequestedType RequestedAttributesType(CachedTree.GetMaterialAttributesType(), false);
-	CachedTree.SetRequestedFields(ShaderFrequency, RequestedAttributesType);
+	CachedTree.SetRequestedFields(EmitContext, RequestedAttributesType);
 
 	FOwnedNodeMaterial MaterialOwner(Material);
 	FEmitOwnerScope OwnerScope(EmitContext, &MaterialOwner);
diff --git a/Engine/Source/Runtime/Engine/Private/HLSLTree/HLSLTree.cpp b/Engine/Source/Runtime/Engine/Private/HLSLTree/HLSLTree.cpp
--- a/Engine/Source/Runtime/Engine/Private/HLSLTree/HLSLTree.cpp
+++ b/Engine/Source/Runtime/Engine/Private/HLSLTree/HLSLTree.cpp
@@ -253,7 +253,6 @@ bool FExpressionLocalPHI::PrepareValue(FEmitContext& Context, FEmitScope& Scope,
 							ValueType.GetName());
 					}
 					CurrentType = MergedType;
-					CurrentType.MergeEvaluation(EmitValueScopes[i]->Evaluation);
 					check(NumValidTypes < NumLiveScopes);
 					NumValidTypes++;
 				}
@@ -311,6 +310,51 @@ bool FExpressionLocalPHI::PrepareValue(FEmitContext& Context, FEmitScope& Scope,
 		}
 	}
 
+	// Only incorporate scope evaluations if we are not sure whether all paths produce the same value
+	bool bMergeScopeEvaluations = false;
+
+	for (int32 Index = 1; Index < NumLiveScopes; ++Index)
+	{
+		if (LiveValues[Index] != LiveValues[Index - 1])
+		{
+			bMergeScopeEvaluations = true;
+			break;
+		}
+	}
+
+	if (bMergeScopeEvaluations && IsConstantEvaluation(CurrentType.GetEvaluation(Scope, RequestedType)))
+	{
+		bool bFirst = true;
+		bool bAllValuesSame = true;
+		Shader::FValue TestValue;
+
+		for (int32 Index = 0; Index < NumLiveScopes; ++Index)
+		{
+			const Shader::FValue ValueConst = LiveValues[Index]->GetValueConstant(Context, *EmitValueScopes[Index], RequestedType, TypePerValue[Index]);
+			if (bFirst)
+			{
+				TestValue = ValueConst;
+				bFirst = false;
+			}
+			else if (ValueConst != TestValue)
+			{
+				bAllValuesSame = false;
+				break;
+			}
+		}
+
+		bMergeScopeEvaluations = !bAllValuesSame;
+	}
+
+	if (bMergeScopeEvaluations)
+	{
+		for (int32 Index = 0; Index < NumLiveScopes; ++Index)
+		{
+			CurrentType.MergeEvaluation(EmitValueScopes[Index]->Evaluation);
+		}
+		verify(OutResult.SetType(Context, RequestedType, CurrentType));
+	}
+
 	return true;
 }
 
@@ -325,7 +369,7 @@ struct FLocalPHILiveScopes
 	bool bCanForwardValue = true;
 };
 
-bool GetLiveScopes(FEmitContext& Context, const FExpressionLocalPHI& Expression, FLocalPHILiveScopes& OutLiveScopes)
+bool GetLiveScopes(FEmitContext& Context, const FExpressionLocalPHI& Expression, FLocalPHILiveScopes& OutLiveScopes, bool bPreparedTypeConstant)
 {
 	for (int32 ScopeIndex = 0; ScopeIndex < Expression.NumValues; ++ScopeIndex)
 	{
@@ -354,19 +398,32 @@ bool GetLiveScopes(FEmitContext& Context, const FExpressionLocalPHI& Expression,
 			}
 		}
 	}
+
+	// If this LocalPHI is constant, then either only one value scope is live and it is constant or all live scopes produce the same constant value
+	OutLiveScopes.bCanForwardValue |= bPreparedTypeConstant;
+
 	return OutLiveScopes.NumScopes > 0;
 }
 }
 
 void FExpressionLocalPHI::EmitValueShader(FEmitContext& Context, FEmitScope& Scope, const FRequestedType& RequestedType, FEmitValueShaderResult& OutResult) const
 {
-	FEmitShaderExpression* const* PrevEmitExpression = Context.EmitLocalPHIMap.Find(this);
+	FXxHash64 Hash;
+	{
+		FHasher Hasher;
+		AppendHash(Hasher, this);
+		AppendHash(Hasher, &Scope);
+		AppendHash(Hasher, RequestedType);
+		Hash = Hasher.Finalize();
+	}
+
+	FEmitShaderExpression* const* PrevEmitExpression = Context.EmitLocalPHIMap.Find(Hash);
 	FEmitShaderExpression* EmitExpression = PrevEmitExpression ? *PrevEmitExpression : nullptr;
 	if (!EmitExpression)
 	{
 		// Find the outermost scope to declare our local variable
 		Private::FLocalPHILiveScopes LiveScopes;
-		if (!Private::GetLiveScopes(Context, *this, LiveScopes))
+		if (!Private::GetLiveScopes(Context, *this, LiveScopes, false))
 		{
 			return;
 		}
@@ -374,7 +431,7 @@ void FExpressionLocalPHI::EmitValueShader(FEmitContext& Context, FEmitScope& Sco
 		if (LiveScopes.bCanForwardValue)
 		{
 			EmitExpression = LiveScopes.LiveValues[0]->GetValueShader(Context, Scope, RequestedType);
-			Context.EmitLocalPHIMap.Add(this, EmitExpression);
+			Context.EmitLocalPHIMap.Add(Hash, EmitExpression);
 		}
 		else
 		{
@@ -386,7 +443,7 @@ void FExpressionLocalPHI::EmitValueShader(FEmitContext& Context, FEmitScope& Sco
 			EmitExpression = OutResult.Code = Context.EmitInlineExpression(Scope,
 				LocalType,
 				TEXT("LocalPHI%"), LocalPHIIndex);
-			Context.EmitLocalPHIMap.Add(this, EmitExpression);
+			Context.EmitLocalPHIMap.Add(Hash, EmitExpression);
 
 			FEmitShaderStatement* EmitDeclaration = nullptr;
 			for (int32 i = 0; i < LiveScopes.NumScopes; ++i)
@@ -445,11 +502,13 @@ void FExpressionLocalPHI::EmitValuePreshader(FEmitContext& Context, FEmitScope&
 		}
 	}
 
-	OutResult.Type = Context.GetResultType(this, RequestedType);
+	const FPreparedType& PreparedType = Context.GetPreparedType(this, RequestedType);
+	OutResult.Type = PreparedType.GetResultType();
+
 	if (ValueStackPosition == INDEX_NONE)
 	{
 		Private::FLocalPHILiveScopes LiveScopes;
-		if (!Private::GetLiveScopes(Context, *this, LiveScopes))
+		if (!Private::GetLiveScopes(Context, *this, LiveScopes, IsConstantEvaluation(PreparedType.GetEvaluation(Scope, RequestedType))))
 		{
 			return;
 		}
@@ -494,7 +553,7 @@ void FExpressionLocalPHI::EmitValuePreshader(FEmitContext& Context, FEmitScope&
 bool FExpressionLocalPHI::EmitValueObject(FEmitContext& Context, FEmitScope& Scope, const FName& ObjectTypeName, void* OutObjectBase) const
 {
 	Private::FLocalPHILiveScopes LiveScopes;
-	if (!Private::GetLiveScopes(Context, *this, LiveScopes) || !LiveScopes.bCanForwardValue)
+	if (!Private::GetLiveScopes(Context, *this, LiveScopes, false) || !LiveScopes.bCanForwardValue)
 	{
 		// Cannot get if no live scope. Don't know which scope to use if there is more than one live
 		return false;
diff --git a/Engine/Source/Runtime/Engine/Private/HLSLTree/HLSLTreeCommon.cpp b/Engine/Source/Runtime/Engine/Private/HLSLTree/HLSLTreeCommon.cpp
--- a/Engine/Source/Runtime/Engine/Private/HLSLTree/HLSLTreeCommon.cpp
+++ b/Engine/Source/Runtime/Engine/Private/HLSLTree/HLSLTreeCommon.cpp
@@ -515,10 +515,12 @@ void FExpressionSelect::EmitValueShader(FEmitContext& Context, FEmitScope& Scope
 	else
 	{
 		const Shader::FType LocalType = Context.GetResultType(this, RequestedType);
+		const bool bIsLWC = LocalType.IsNumericLWC();
+
 		FEmitShaderExpression* TrueValue = TrueExpression->GetValueShader(Context, Scope, RequestedType, LocalType);
 		FEmitShaderExpression* FalseValue = FalseExpression->GetValueShader(Context, Scope, RequestedType, LocalType);
 
-		OutResult.Code = Context.EmitExpression(Scope, LocalType, TEXT("(% ? % : %)"),
+		OutResult.Code = Context.EmitExpression(Scope, LocalType, bIsLWC ? TEXT("LWCSelect(%, %, %)") : TEXT("(% ? % : %)"),
 			ConditionExpression->GetValueShader(Context, Scope, Shader::EValueType::Bool1),
 			TrueValue,
 			FalseValue);
diff --git a/Engine/Source/Runtime/Engine/Private/HLSLTree/HLSLTreeEmit.cpp b/Engine/Source/Runtime/Engine/Private/HLSLTree/HLSLTreeEmit.cpp
--- a/Engine/Source/Runtime/Engine/Private/HLSLTree/HLSLTreeEmit.cpp
+++ b/Engine/Source/Runtime/Engine/Private/HLSLTree/HLSLTreeEmit.cpp
@@ -872,11 +872,6 @@ FEmitShaderExpression* FEmitContext::InternalEmitExpression(FEmitScope& Scope, T
 {
 	FEmitShaderExpression* ShaderValue = nullptr;
 
-	if (Code.Contains(TEXT("((float2)Local21).xyz")))
-	{
-		int a = 0;
-	}
-
 	FXxHash64Builder Hasher;
 	Hasher.Update(Code.GetData(), Code.Len() * sizeof(TCHAR));
 	if (bInline)
@@ -1334,7 +1329,12 @@ FEmitShaderExpression* FEmitContext::EmitCast(FEmitScope& Scope, FEmitShaderExpr
 			}
 			else
 			{
-				if ((SourceTypeDesc.NumComponents == 1 && bReplicateScalar) || SourceTypeDesc.NumComponents == DestTypeDesc.NumComponents)
+				if (SourceTypeDesc.NumComponents == DestTypeDesc.NumComponents)
+				{
+					NumComponents = DestTypeDesc.NumComponents;
+					FormattedCode.Append(ShaderValue->Reference);
+				}
+				else if (bReplicateScalar)
 				{
 					NumComponents = DestTypeDesc.NumComponents;
 					// Cast the scalar to the correct type, HLSL language will replicate the scalar if needed when performing this cast
@@ -1348,16 +1348,16 @@ FEmitShaderExpression* FEmitContext::EmitCast(FEmitScope& Scope, FEmitShaderExpr
 						FormattedCode.Appendf(TEXT("%s("), DestTypeDesc.Name);
 						bNeedClosingParen = true;
 					}
-					if (NumComponents == SourceTypeDesc.NumComponents && SourceTypeDesc.ComponentType == DestTypeDesc.ComponentType)
+					if (NumComponents == SourceTypeDesc.NumComponents)
 					{
 						// If we're taking all the components from the source, can avoid adding a swizzle
 						FormattedCode.Append(ShaderValue->Reference);
 					}
 					else
 					{
-						// Use a cast to truncate the source to the correct number of types
-						const Shader::EValueType LocalType = Shader::MakeValueType(DestTypeDesc.ComponentType, NumComponents);
-						FormattedCode.Appendf(TEXT("((%s)%s)"), Shader::GetValueTypeDescription(LocalType).Name, ShaderValue->Reference);
+						// Truncate using a swizzle
+						static const TCHAR* Mask[] = { nullptr, TEXT("x"), TEXT("xy"), TEXT("xyz"), TEXT("xyzw") };
+						FormattedCode.Appendf(TEXT("%s.%s"), ShaderValue->Reference, Mask[NumComponents]);
 					}
 				}
 			}
@@ -1510,6 +1510,7 @@ void FEmitContext::Finalize()
 	EmitScopeMap.Reset();
 	EmitFunctionMap.Reset();
 	EmitLocalPHIMap.Reset();
+	EmitValueMap.Reset();
 
 	MaterialCompilationOutput->UniformExpressionSet.UniformPreshaderBufferSize = (UniformPreshaderOffset + 3u) / 4u;
 }
diff --git a/Engine/Source/Runtime/Engine/Private/HLSLTree/HLSLTreeOperation.cpp b/Engine/Source/Runtime/Engine/Private/HLSLTree/HLSLTreeOperation.cpp
--- a/Engine/Source/Runtime/Engine/Private/HLSLTree/HLSLTreeOperation.cpp
+++ b/Engine/Source/Runtime/Engine/Private/HLSLTree/HLSLTreeOperation.cpp
@@ -239,6 +239,7 @@ FOperationTypes GetOperationTypes(EOperation Op, TConstArrayView<FPreparedType>
 				Types.InputType[0] = Shader::EValueType::Float3;
 				Types.InputType[1] = Shader::EValueType::Double4x4;
 				Types.ResultType = FPreparedType(Shader::EValueType::Double3, IntermediateComponent);
+				Types.bIsLWC = true;
 				break;
 			case Shader::EValueType::DoubleInverse4x4:
 				// FLWCVector3 * FLWCInverseMatrix -> float3
@@ -326,6 +327,14 @@ FOperationTypes GetOperationTypes(EOperation Op, TConstArrayView<FPreparedType>
 		case EOperation::Frac:
 		case EOperation::Step:
 		case EOperation::SmoothStep:
+			if (Types.bIsLWC)
+			{
+				// LWCSmoothStep requires all inputs have the same type
+				for (int32 InputIndex = 0; InputIndex < InputPreparedType.Num(); ++InputIndex)
+				{
+					Types.InputType[InputIndex] = IntermediateType.Type.ValueType;
+				}
+			}
 			for (int32 Index = 0; Index < Types.ResultType.PreparedComponents.Num(); ++Index)
 			{
 				Types.ResultType.SetComponentBounds(Index, Shader::FComponentBounds(Shader::EComponentBound::Zero, Shader::EComponentBound::One));
@@ -838,7 +847,8 @@ void FExpressionOperation::EmitValueShader(FEmitContext& Context, FEmitScope& Sc
 		}
 		else
 		{
-			OutResult.Code = Context.EmitExpression(Scope, ResultType, TEXT("mul(%, %)"), InputValue[0], InputValue[1]);
+			// Append 1.f because VecMulMatrix4 implies position but we request float3 for the vector
+			OutResult.Code = Context.EmitExpression(Scope, ResultType, TEXT("mul(float4(%, 1.f), %).xyz"), InputValue[0], InputValue[1]);
 		}
 		break;
 	case EOperation::Matrix3MulVec:
diff --git a/Engine/Source/Runtime/Engine/Private/Materials/MaterialCachedData.cpp b/Engine/Source/Runtime/Engine/Private/Materials/MaterialCachedData.cpp
--- a/Engine/Source/Runtime/Engine/Private/Materials/MaterialCachedData.cpp
+++ b/Engine/Source/Runtime/Engine/Private/Materials/MaterialCachedData.cpp
@@ -938,7 +938,7 @@ void PrepareHLSLTree(UE::HLSLTree::FEmitContext& EmitContext,
 	FEmitScope* EmitResultScope = EmitContext.PrepareScope(CachedTree.GetResultScope());
 
 	FRequestedType RequestedAttributesType(CachedTree.GetMaterialAttributesType(), false);
-	CachedTree.SetRequestedFields(ShaderFrequency, RequestedAttributesType);
+	CachedTree.SetRequestedFields(EmitContext, RequestedAttributesType);
 
 	const FPreparedType& ResultType = EmitContext.PrepareExpression(CachedTree.GetResultExpression(), *EmitResultScope, RequestedAttributesType);
 	if (!ResultType.IsVoid())
diff --git a/Engine/Source/Runtime/Engine/Private/Materials/MaterialCachedHLSLTree.cpp b/Engine/Source/Runtime/Engine/Private/Materials/MaterialCachedHLSLTree.cpp
--- a/Engine/Source/Runtime/Engine/Private/Materials/MaterialCachedHLSLTree.cpp
+++ b/Engine/Source/Runtime/Engine/Private/Materials/MaterialCachedHLSLTree.cpp
@@ -9,6 +9,8 @@
 #include "Materials/MaterialAttributeDefinitionMap.h"
 #include "Materials/MaterialExpressionCustomOutput.h"
 #include "MaterialHLSLGenerator.h"
+#include "MaterialHLSLTree.h"
+#include "HLSLTree/HLSLTreeEmit.h"
 
 DECLARE_STATS_GROUP(TEXT("Material Memory"), STATGROUP_MaterialMemory, STATCAT_Advanced);
 DECLARE_MEMORY_STAT(TEXT("Material HLSLTree Memory"), STAT_MaterialMemory_HLSLTree, STATGROUP_MaterialMemory);
@@ -93,8 +95,16 @@ bool FMaterialCachedHLSLTree::GenerateTree(UMaterial* Material, const FMaterialL
 		}
 	}
 
-	TArray<TStringBuilder<256>, TInlineAllocator<4>> CustomOutputNames;
-	CustomOutputNames.Reserve(MaterialCustomOutputs.Num());
+	// Make sure this array won't resize because TStringBuilder holds pointers to its inline storage
+	TArray<TStringBuilder<64>> CustomOutputNames;
+	int32 TotalNumOutputs = 0;
+
+	for (UMaterialExpressionCustomOutput* CustomOutput : MaterialCustomOutputs)
+	{
+		TotalNumOutputs += CustomOutput->GetNumOutputs();
+	}
+	CustomOutputNames.Empty(TotalNumOutputs);
+
 	for (UMaterialExpressionCustomOutput* CustomOutput : MaterialCustomOutputs)
 	{
 		const int32 NumOutputs = CustomOutput->GetNumOutputs();
@@ -139,8 +149,12 @@ bool FMaterialCachedHLSLTree::GenerateTree(UMaterial* Material, const FMaterialL
 	return bResult;
 }
 
-void FMaterialCachedHLSLTree::SetRequestedFields(EShaderFrequency ShaderFrequency, UE::HLSLTree::FRequestedType& OutRequestedType) const
+void FMaterialCachedHLSLTree::SetRequestedFields(const UE::HLSLTree::FEmitContext& Context, UE::HLSLTree::FRequestedType& OutRequestedType) const
 {
+	using namespace UE::HLSLTree;
+
+	const EShaderFrequency ShaderFrequency = Context.ShaderFrequency;
+
 	for (UMaterialExpressionCustomOutput* CustomOutput : MaterialCustomOutputs)
 	{
 		if (CustomOutput->GetShaderFrequency() != ShaderFrequency)
@@ -161,11 +175,31 @@ void FMaterialCachedHLSLTree::SetRequestedFields(EShaderFrequency ShaderFrequenc
 		}
 	}
 
+	int32 NumCustomUVs = 0;
+	if (ShaderFrequency == SF_Vertex)
+	{
+		const Material::FEmitData& EmitData = Context.FindData<Material::FEmitData>();
+		for (int32 TexCoordIndex = Material::MaxNumTexCoords - 1; TexCoordIndex >= 0; --TexCoordIndex)
+		{
+			if (EmitData.IsExternalInputUsed(SF_Pixel, Material::MakeInputTexCoord(TexCoordIndex)))
+			{
+				NumCustomUVs = TexCoordIndex + 1;
+				break;
+			}
+		}
+	}
+
 	const TArray<FGuid>& OrderedVisibleAttributes = FMaterialAttributeDefinitionMap::GetOrderedVisibleAttributeList();
 	for (const FGuid& AttributeID : OrderedVisibleAttributes)
 	{
 		if (FMaterialAttributeDefinitionMap::GetShaderFrequency(AttributeID) == ShaderFrequency)
 		{
+			const EMaterialProperty Property = FMaterialAttributeDefinitionMap::GetProperty(AttributeID);
+			if (Property >= MP_CustomizedUVs0 && Property <= MP_CustomizedUVs7 && Property - MP_CustomizedUVs0 >= NumCustomUVs)
+			{
+				continue;
+			}
+
 			const FString& FieldName = FMaterialAttributeDefinitionMap::GetAttributeName(AttributeID);
 			const UE::Shader::FStructField* Field = MaterialAttributesType->FindFieldByName(*FieldName);
 			if (Field)
@@ -197,14 +231,39 @@ void FMaterialCachedHLSLTree::EmitSharedCode(FStringBuilderBase& OutCode) const
 		{
 			const EValueType ValueType = CustomOutput->GetCustomOutputType(OutputIndex);
 			const FValueTypeDescription ValueTypeDesc = GetValueTypeDescription(ValueType);
+			const bool bIsLWC = ValueTypeDesc.ComponentType == EValueComponentType::Double;
 
 			OutCode.Appendf(TEXT("#define HAVE_%s%d 1\n"), *OutputName, OutputIndex);
 
-			OutCode.Appendf(TEXT("%s %s%d(FMaterial%sParameters Parameters) { return Parameters.MaterialAttributes.%s%d; }\n"),
+			OutCode.Appendf(TEXT("%s %s%d%s(FMaterial%sParameters Parameters) { return Parameters.MaterialAttributes.%s%d; }\n"),
 				ValueTypeDesc.Name,
 				*OutputName, OutputIndex,
+				bIsLWC ? TEXT("_LWC") : TEXT(""),
 				ShaderFrequency == SF_Pixel ? TEXT("Pixel") : TEXT("Vertex"),
 				*OutputName, OutputIndex);
+
+			if (bIsLWC)
+			{
+				// Add a wrapper with no suffix to return a non-LWC type
+				const EValueType NonLWCType = MakeValueType(EValueComponentType::Float, ValueTypeDesc.NumComponents);
+				const FValueTypeDescription NonLWCTypeDesc = GetValueTypeDescription(NonLWCType);
+				OutCode.Appendf(TEXT("%s %s%d(FMaterial%sParameters Parameters) { return LWCToFloat(%s%d_LWC(Parameters)); }\n"),
+					NonLWCTypeDesc.Name,
+					*OutputName, OutputIndex,
+					ShaderFrequency == SF_Vertex ? TEXT("Vertex") : TEXT("Pixel"),
+					*OutputName, OutputIndex);
+			}
+			else
+			{
+				// Add a wrapper with LWC suffix to return a LWC type
+				const EValueType LWCType = MakeValueType(EValueComponentType::Double, ValueTypeDesc.NumComponents);
+				const FValueTypeDescription LWCTypeDesc = GetValueTypeDescription(LWCType);
+				OutCode.Appendf(TEXT("%s %s%d_LWC(FMaterial%sParameters Parameters) { return LWCPromote(%s%d(Parameters)); }\n"),
+					LWCTypeDesc.Name,
+					*OutputName, OutputIndex,
+					ShaderFrequency == SF_Vertex ? TEXT("Vertex") : TEXT("Pixel"),
+					*OutputName, OutputIndex);
+			}
 		}
 		OutCode.Append(TEXT("\n"));
 	}
diff --git a/Engine/Source/Runtime/Engine/Private/Materials/MaterialExpressionHLSL.cpp b/Engine/Source/Runtime/Engine/Private/Materials/MaterialExpressionHLSL.cpp
--- a/Engine/Source/Runtime/Engine/Private/Materials/MaterialExpressionHLSL.cpp
+++ b/Engine/Source/Runtime/Engine/Private/Materials/MaterialExpressionHLSL.cpp
@@ -1743,10 +1743,11 @@ UE::Shader::EValueType UMaterialExpressionRuntimeVirtualTextureOutput::GetCustom
 		return EValueType::Float3;
 	case 1:
 	case 2:
-	case 4:
 	case 5:
 	case 6:
 		return EValueType::Float1;
+	case 4:
+		return EValueType::Double1;
 	default:
 		checkNoEntry();
 		return EValueType::Void;
@@ -1794,7 +1795,7 @@ bool UMaterialExpressionRuntimeVirtualTextureOutput::GenerateHLSLExpression(FMat
 	}
 	else if (OutputIndex == 4)
 	{
-		OutExpression = MakeOutputExpression(WorldHeight, 0.f, ERuntimeVirtualTextureAttributeType::WorldHeight);
+		OutExpression = MakeOutputExpression(WorldHeight, 0., ERuntimeVirtualTextureAttributeType::WorldHeight);
 	}
 	else if (OutputIndex == 5)
 	{
@@ -2683,8 +2684,7 @@ bool UMaterialExpressionComponentMask::GenerateHLSLExpression(FMaterialHLSLGener
 
 bool UMaterialExpressionGetMaterialAttributes::GenerateHLSLExpression(FMaterialHLSLGenerator& Generator, UE::HLSLTree::FScope& Scope, int32 OutputIndex, UE::HLSLTree::FExpression const*& OutExpression) const
 {
-	const UE::HLSLTree::FExpression* AttributesExpression = MaterialAttributes.AcquireHLSLExpressionOrConstant(Generator, Scope, Generator.GetMaterialAttributesDefaultValue());
-	check(AttributesExpression);
+	const UE::HLSLTree::FExpression* AttributesExpression = MaterialAttributes.AcquireHLSLExpressionOrDefaultExpression(Generator, Scope, Generator.GetMaterialAttributesDefaultExpression());
 
 	if (OutputIndex == 0)
 	{
@@ -2708,7 +2708,7 @@ bool UMaterialExpressionGetMaterialAttributes::GenerateHLSLExpression(FMaterialH
 
 bool UMaterialExpressionSetMaterialAttributes::GenerateHLSLExpression(FMaterialHLSLGenerator& Generator, UE::HLSLTree::FScope& Scope, int32 OutputIndex, UE::HLSLTree::FExpression const*& OutExpression) const
 {
-	const UE::HLSLTree::FExpression* AttributesExpression = Inputs[0].AcquireHLSLExpressionOrConstant(Generator, Scope, Generator.GetMaterialAttributesDefaultValue());
+	const UE::HLSLTree::FExpression* AttributesExpression = Inputs[0].AcquireHLSLExpressionOrDefaultExpression(Generator, Scope, Generator.GetMaterialAttributesDefaultExpression());
 	
 	for (int32 PinIndex = 0; PinIndex < AttributeSetTypes.Num(); ++PinIndex)
 	{
@@ -2760,7 +2760,7 @@ static const UE::HLSLTree::FExpression* SetAttribute(FMaterialHLSLGenerator& Gen
 bool UMaterialExpressionMakeMaterialAttributes::GenerateHLSLExpression(FMaterialHLSLGenerator& Generator, UE::HLSLTree::FScope& Scope, int32 OutputIndex, UE::HLSLTree::FExpression const*& OutExpression) const
 {
 	using namespace UE::HLSLTree;
-	const FExpression* AttributesExpression = Generator.NewConstant(Generator.GetMaterialAttributesDefaultValue());
+	const FExpression* AttributesExpression = Generator.GetMaterialAttributesDefaultExpression();
 	AttributesExpression = SetAttribute(Generator, Scope, MP_BaseColor, BaseColor, AttributesExpression);
 	AttributesExpression = SetAttribute(Generator, Scope, MP_Metallic, Metallic, AttributesExpression);
 
@@ -2878,7 +2878,7 @@ bool UMaterialExpressionMaterialAttributeLayers::GenerateHLSLExpression(FMateria
 
 	TArray<FFunctionExpressionInput> FunctionInputs;
 	TArray<FFunctionExpressionOutput> FunctionOutputs;
-	const FExpression* ExpressionLayerInput = Input.AcquireHLSLExpressionOrConstant(Generator, Scope, Generator.GetMaterialAttributesDefaultValue());
+	const FExpression* ExpressionLayerInput = Input.AcquireHLSLExpressionOrDefaultExpression(Generator, Scope, Generator.GetMaterialAttributesDefaultExpression());
 
 	TArray<FMaterialHLSLGenerator::FConnectedInput, TInlineAllocator<1>> LayerInputExpressions;
 	TArray<const FExpression*, TInlineAllocator<16>> LayerExpressions;
@@ -2981,17 +2981,23 @@ bool UMaterialExpressionBlendMaterialAttributes::GenerateHLSLExpression(FMateria
 	using namespace UE::HLSLTree;
 	using namespace UE::Shader;
 
-	const FExpression* ExpressionA = A.AcquireHLSLExpressionOrConstant(Generator, Scope, Generator.GetMaterialAttributesDefaultValue());
-	const FExpression* ExpressionB = B.AcquireHLSLExpressionOrConstant(Generator, Scope, Generator.GetMaterialAttributesDefaultValue());
+	const FExpression* ExpressionA = A.AcquireHLSLExpressionOrDefaultExpression(Generator, Scope, Generator.GetMaterialAttributesDefaultExpression());
+	const FExpression* ExpressionB = B.AcquireHLSLExpressionOrDefaultExpression(Generator, Scope, Generator.GetMaterialAttributesDefaultExpression());
 	const FExpression* ExpressionAlpha = Alpha.AcquireHLSLExpression(Generator, Scope);
 	if (!ExpressionA || !ExpressionB || !ExpressionAlpha)
 	{
 		return false;
 	}
 
+	if (ExpressionA == ExpressionB)
+	{
+		OutExpression = ExpressionA;
+		return true;
+	}
+
 	const UE::Shader::FStructType* MaterialAttributesType = Generator.GetMaterialAttributesType();
 
-	const FExpression* ExpressionResult = Generator.GetTree().NewConstant(Generator.GetMaterialAttributesDefaultValue());
+	const FExpression* ExpressionResult = Generator.GetMaterialAttributesDefaultExpression();
 	const TArray<FGuid>& OrderedVisibleAttributes = FMaterialAttributeDefinitionMap::GetOrderedVisibleAttributeList();
 	for (const FGuid& AttributeID : OrderedVisibleAttributes)
 	{
@@ -3636,12 +3642,19 @@ bool UMaterialExpressionVolumetricAdvancedMaterialOutput::GenerateHLSLExpression
 	case 3: OutExpression = MultiScatteringContribution.AcquireHLSLExpressionOrConstant(Generator, Scope, ConstMultiScatteringContribution); break;
 	case 4: OutExpression = MultiScatteringOcclusion.AcquireHLSLExpressionOrConstant(Generator, Scope, ConstMultiScatteringOcclusion); break;
 	case 5: OutExpression = MultiScatteringEccentricity.AcquireHLSLExpressionOrConstant(Generator, Scope, ConstMultiScatteringEccentricity); break;
-	case 6: OutExpression = ConservativeDensity.AcquireHLSLExpressionOrConstant(Generator, Scope, FVector4f(1.0f, 1.0f, 1.0f, 1.0f)); break;
+	case 6: OutExpression = ConservativeDensity.AcquireHLSLExpressionOrConstant(Generator, Scope, FVector4f(1.0f, 1.0f, 1.0f, 0.0f)); break;
 	default: return Generator.Error(TEXT("Invalid output"));
 	}
 	return OutExpression != nullptr;
 }
 
+UE::Shader::EValueType UMaterialExpressionVolumetricAdvancedMaterialOutput::GetCustomOutputType(int32 OutputIndex) const
+{
+	if (OutputIndex >= 0 && OutputIndex < 6) return UE::Shader::EValueType::Float1;
+	else if (OutputIndex == 6) return UE::Shader::EValueType::Float4;
+	else return UE::Shader::EValueType::Void;
+}
+
 bool UMaterialExpressionThinTranslucentMaterialOutput::GenerateHLSLExpression(FMaterialHLSLGenerator& Generator, UE::HLSLTree::FScope& Scope, int32 OutputIndex, UE::HLSLTree::FExpression const*& OutExpression) const
 {
 	if (OutputIndex == 0)
@@ -3664,13 +3677,6 @@ UE::Shader::EValueType UMaterialExpressionThinTranslucentMaterialOutput::GetCust
 	return EValueType::Void;
 }
 
-UE::Shader::EValueType UMaterialExpressionVolumetricAdvancedMaterialOutput::GetCustomOutputType(int32 OutputIndex) const
-{
-	if (OutputIndex >= 0 && OutputIndex < 6) return UE::Shader::EValueType::Float1;
-	else if (OutputIndex == 6) return UE::Shader::EValueType::Float3;
-	else return UE::Shader::EValueType::Void;
-}
-
 bool UMaterialExpressionVolumetricCloudEmptySpaceSkippingInput::GenerateHLSLExpression(FMaterialHLSLGenerator& Generator, UE::HLSLTree::FScope& Scope, int32 OutputIndex, UE::HLSLTree::FExpression const*& OutExpression) const
 {
 	return Generator.Error(TEXT("UMaterialExpressionVolumetricCloudEmptySpaceSkippingInput::GenerateHLSLExpression unimplemented"));
diff --git a/Engine/Source/Runtime/Engine/Private/Materials/MaterialHLSLEmitter.cpp b/Engine/Source/Runtime/Engine/Private/Materials/MaterialHLSLEmitter.cpp
--- a/Engine/Source/Runtime/Engine/Private/Materials/MaterialHLSLEmitter.cpp
+++ b/Engine/Source/Runtime/Engine/Private/Materials/MaterialHLSLEmitter.cpp
@@ -409,7 +409,13 @@ static void GetMaterialEnvironment(EShaderPlatform InPlatform,
 		OutEnvironment.SetDefine(TEXT("USES_SPEEDTREE"), TEXT("1"));
 	}
 
-	if (false)//bNeedsWorldPositionExcludingShaderOffsets)
+	const bool bNeedsWorldPositionExcludingShaderOffsets =
+		EmitMaterialData.IsExternalInputUsed(Material::EExternalInput::WorldPosition_NoOffsets) ||
+		EmitMaterialData.IsExternalInputUsed(Material::EExternalInput::TranslatedWorldPosition_NoOffsets) ||
+		EmitMaterialData.IsExternalInputUsed(Material::EExternalInput::PrevWorldPosition_NoOffsets) ||
+		EmitMaterialData.IsExternalInputUsed(Material::EExternalInput::PrevTranslatedWorldPosition_NoOffsets);
+
+	if (bNeedsWorldPositionExcludingShaderOffsets)
 	{
 		OutEnvironment.SetDefine(TEXT("NEEDS_WORLD_POSITION_EXCLUDING_SHADER_OFFSETS"), TEXT("1"));
 	}
@@ -452,9 +458,9 @@ static void GetMaterialEnvironment(EShaderPlatform InPlatform,
 	OutEnvironment.SetDefine(TEXT("USES_PER_INSTANCE_FADE_AMOUNT"), bUsesPerInstanceFadeAmount&& InMaterial.IsUsedWithInstancedStaticMeshes());
 	OutEnvironment.SetDefine(TEXT("USES_VERTEX_INTERPOLATOR"), MaterialCompilationOutput.bUsesVertexInterpolator);
 
-	const bool bUsesVertexColor = EmitMaterialData.IsExternalInputUsed(SF_Vertex, Material::EExternalInput::VertexColor) ||
-		EmitMaterialData.IsExternalInputUsed(SF_Vertex, Material::EExternalInput::VertexColor_Ddx) ||
-		EmitMaterialData.IsExternalInputUsed(SF_Vertex, Material::EExternalInput::VertexColor_Ddy);
+	const bool bUsesVertexColor = EmitMaterialData.IsExternalInputUsed(SF_Pixel, Material::EExternalInput::VertexColor) ||
+		EmitMaterialData.IsExternalInputUsed(SF_Pixel, Material::EExternalInput::VertexColor_Ddx) ||
+		EmitMaterialData.IsExternalInputUsed(SF_Pixel, Material::EExternalInput::VertexColor_Ddy);
 
 	const bool bUsesParticleColor = EmitMaterialData.IsExternalInputUsed(SF_Pixel, Material::EExternalInput::ParticleColor);
 
@@ -815,7 +821,7 @@ bool MaterialEmitHLSL(const FMaterialCompileTargetParameters& InCompilerTarget,
 
 		// Prepare all fields *except* normal
 		FRequestedType RequestedPixelAttributesType(CachedTree->GetMaterialAttributesType(), false);
-		CachedTree->SetRequestedFields(SF_Pixel, RequestedPixelAttributesType);
+		CachedTree->SetRequestedFields(EmitContext, RequestedPixelAttributesType);
 		RequestedPixelAttributesType.SetFieldRequested(NormalField, false);
 
 		const FPreparedType PixelResultType0 = EmitContext.PrepareExpression(CachedTree->GetResultExpression(), *EmitResultScope, RequestedPixelAttributesType);
@@ -934,11 +940,12 @@ bool MaterialEmitHLSL(const FMaterialCompileTargetParameters& InCompilerTarget,
 	// Prepare vertex shader code
 	FStringBuilderMemstack VertexCode(Allocator, 128 * 1024);
 	{
+		EmitContext.ShaderFrequency = SF_Vertex;
+
 		FRequestedType RequestedVertexAttributesType(CachedTree->GetMaterialAttributesType(), false);
-		CachedTree->SetRequestedFields(SF_Vertex, RequestedVertexAttributesType);
+		CachedTree->SetRequestedFields(EmitContext, RequestedVertexAttributesType);
 		RequestedVertexAttributesType.SetFieldRequested(CachedTree->GetMaterialAttributesType()->FindFieldByName(TEXT("PrevWorldPositionOffset")));
 
-		EmitContext.ShaderFrequency = SF_Vertex;
 		EmitContext.bUseAnalyticDerivatives = false;
 		EmitContext.bMarkLiveValues = false;
 		FEmitScope* EmitResultScope = EmitContext.PrepareScope(CachedTree->GetResultScope());
diff --git a/Engine/Source/Runtime/Engine/Private/Materials/MaterialHLSLGenerator.cpp b/Engine/Source/Runtime/Engine/Private/Materials/MaterialHLSLGenerator.cpp
--- a/Engine/Source/Runtime/Engine/Private/Materials/MaterialHLSLGenerator.cpp
+++ b/Engine/Source/Runtime/Engine/Private/Materials/MaterialHLSLGenerator.cpp
@@ -50,9 +50,36 @@ const UE::Shader::FStructType* FMaterialHLSLGenerator::GetMaterialAttributesType
 	return CachedTree.GetMaterialAttributesType();
 }
 
-const UE::Shader::FValue& FMaterialHLSLGenerator::GetMaterialAttributesDefaultValue() const
+const UE::HLSLTree::FExpression* FMaterialHLSLGenerator::GetMaterialAttributesDefaultExpression() const
 {
-	return CachedTree.GetMaterialAttributesDefaultValue();
+	using namespace UE::HLSLTree;
+
+	FTree& Tree = GetTree();
+	const UE::Shader::FStructType* StructType = GetMaterialAttributesType();
+	const FExpression* OutExpression = Tree.NewConstant(CachedTree.GetMaterialAttributesDefaultValue());
+
+	// Some material attribute defaults aren't compile time constants
+	if (ensure(TargetMaterial) && !TargetMaterial->bTangentSpaceNormal)
+	{
+		const UE::Shader::FStructField* NormalField = StructType->FindFieldByName(*FMaterialAttributeDefinitionMap::GetAttributeName(MP_Normal));
+		const UE::Shader::FStructField* TangentField = StructType->FindFieldByName(*FMaterialAttributeDefinitionMap::GetAttributeName(MP_Tangent));
+
+		const FExpression* NormalExpression = Tree.NewExpression<Material::FExpressionExternalInput>(Material::EExternalInput::WorldVertexNormal);
+		const FExpression* TangentExpression = Tree.NewExpression<Material::FExpressionExternalInput>(Material::EExternalInput::WorldVertexTangent);
+
+		OutExpression = Tree.NewExpression<FExpressionSetStructField>(StructType, NormalField, OutExpression, NormalExpression);
+		OutExpression = Tree.NewExpression<FExpressionSetStructField>(StructType, TangentField, OutExpression, TangentExpression);
+	}
+
+	for (int32 CustomUVProperty = MP_CustomizedUVs0; CustomUVProperty <= MP_CustomizedUVs7; ++CustomUVProperty)
+	{
+		const UE::Shader::FStructField* CustomUVField = StructType->FindFieldByName(*FMaterialAttributeDefinitionMap::GetAttributeName((EMaterialProperty)CustomUVProperty));
+		const Material::EExternalInput TexCoordInput = Material::EExternalInput((uint8)Material::EExternalInput::TexCoord0 + uint8(CustomUVProperty - MP_CustomizedUVs0));
+		const FExpression* TexCoordExpression = Tree.NewExpression<Material::FExpressionExternalInput>(TexCoordInput);
+		OutExpression = Tree.NewExpression<FExpressionSetStructField>(StructType, CustomUVField, OutExpression, TexCoordExpression);
+	}
+
+	return OutExpression;
 }
 
 UMaterialExpression* FMaterialHLSLGenerator::GetCurrentExpression() const
@@ -213,7 +240,7 @@ bool FMaterialHLSLGenerator::GenerateResult(UE::HLSLTree::FScope& Scope)
 
 					if (AttributesExpression)
 					{
-						AttributesExpression = GetTree().NewExpression<FExpressionDefaultValue>(AttributesExpression, CachedTree.GetMaterialAttributesDefaultValue());
+						const FExpression* PrevAttributesExpression = GetTree().NewExpression<FExpressionDefaultValue>(AttributesExpression, CachedTree.GetMaterialAttributesDefaultValue());
 
 						const FString& WPOName = FMaterialAttributeDefinitionMap::GetAttributeName(MP_WorldPositionOffset);
 						const FStructField* WPOField = CachedTree.GetMaterialAttributesType()->FindFieldByName(*WPOName);
@@ -221,7 +248,7 @@ bool FMaterialHLSLGenerator::GenerateResult(UE::HLSLTree::FScope& Scope)
 						FRequestedType PrevRequestedType(CachedTree.GetMaterialAttributesType());
 						PrevRequestedType.SetFieldRequested(WPOField);
 
-						const FExpression* PrevAttributesExpression = GetTree().GetPreviousFrame(AttributesExpression, PrevRequestedType);
+						PrevAttributesExpression = GetTree().GetPreviousFrame(PrevAttributesExpression, PrevRequestedType);
 						ensure(PrevAttributesExpression);
 						const FExpression* PrevWPOExpression = GetTree().NewExpression<FExpressionGetStructField>(CachedTree.GetMaterialAttributesType(), WPOField, PrevAttributesExpression);
 						AttributesExpression = GetTree().NewExpression<FExpressionSetStructField>(CachedTree.GetMaterialAttributesType(), PrevWPOField, AttributesExpression, PrevWPOExpression);
@@ -230,7 +257,7 @@ bool FMaterialHLSLGenerator::GenerateResult(UE::HLSLTree::FScope& Scope)
 			}
 			else
 			{
-				AttributesExpression = GetTree().NewExpression<FExpressionConstant>(CachedTree.GetMaterialAttributesDefaultValue());
+				AttributesExpression = GetMaterialAttributesDefaultExpression();
 				for (int32 PropertyIndex = 0; PropertyIndex < MP_MAX; ++PropertyIndex)
 				{
 					const EMaterialProperty Property = (EMaterialProperty)PropertyIndex;
@@ -299,7 +326,7 @@ bool FMaterialHLSLGenerator::GenerateResult(UE::HLSLTree::FScope& Scope)
 				// Get back into gamma corrected space, as DrawTile does not do this adjustment.
 				ExpressionEmissive = GetTree().NewPowClamped(ExpressionEmissive, NewConstant(1.f / 2.2f));
 
-				AttributesExpression = GetTree().NewExpression<FExpressionConstant>(CachedTree.GetMaterialAttributesDefaultValue());
+				AttributesExpression = GetMaterialAttributesDefaultExpression();
 				AttributesExpression = GetTree().NewExpression<FExpressionSetStructField>(CachedTree.GetMaterialAttributesType(), EmissiveColorField, AttributesExpression, ExpressionEmissive);
 			}
 		}
@@ -427,6 +454,19 @@ const UE::HLSLTree::FExpression* FMaterialHLSLGenerator::NewDefaultInputExternal
 	return Expression;
 }
 
+const UE::HLSLTree::FExpression* FMaterialHLSLGenerator::NewDefaultInputExpression(int32 InputIndex, const UE::HLSLTree::FExpression* Expression)
+{
+	using namespace UE::HLSLTree;
+
+	UObject* InputOwner = GetTree().GetCurrentOwner();
+	if (InputOwner && InputIndex != INDEX_NONE)
+	{
+		const FMaterialConnectionKey Key{ InputOwner, nullptr, InputIndex, INDEX_NONE };
+		CachedTree.ConnectionMap.Add(Key, Expression);
+	}
+	return Expression;
+}
+
 const UE::HLSLTree::FExpression* FMaterialHLSLGenerator::AcquireExpression(UE::HLSLTree::FScope& Scope,
 	int32 InputIndex,
 	UMaterialExpression* MaterialExpression,
@@ -551,7 +591,7 @@ const UE::HLSLTree::FExpression* FMaterialHLSLGenerator::AcquireFunctionInputExp
 			case FunctionInput_Vector2: DefaultValue = FVector2f(PreviewValue.X, PreviewValue.Y); break;
 			case FunctionInput_Vector3: DefaultValue = FVector3f(PreviewValue.X, PreviewValue.Y, PreviewValue.Z); break;
 			case FunctionInput_Vector4: DefaultValue = PreviewValue; break;
-			case FunctionInput_MaterialAttributes: DefaultValue = CachedTree.GetMaterialAttributesDefaultValue(); break;
+			case FunctionInput_MaterialAttributes: InputExpression = GetMaterialAttributesDefaultExpression(); break;
 			case FunctionInput_Texture2D:
 			case FunctionInput_TextureCube:
 			case FunctionInput_Texture2DArray:
@@ -566,7 +606,10 @@ const UE::HLSLTree::FExpression* FMaterialHLSLGenerator::AcquireFunctionInputExp
 				return nullptr;
 			}
 
-			InputExpression = NewConstant(DefaultValue);
+			if (!InputExpression)
+			{
+				InputExpression = NewConstant(DefaultValue);
+			}
 		}
 	}
 
diff --git a/Engine/Source/Runtime/Engine/Private/Materials/MaterialHLSLTree.cpp b/Engine/Source/Runtime/Engine/Private/Materials/MaterialHLSLTree.cpp
--- a/Engine/Source/Runtime/Engine/Private/Materials/MaterialHLSLTree.cpp
+++ b/Engine/Source/Runtime/Engine/Private/Materials/MaterialHLSLTree.cpp
@@ -223,19 +223,68 @@ const FExpression* FExpressionExternalInput::ComputePreviousFrame(FTree& Tree, c
 	return nullptr;
 }
 
+EExternalInput FExpressionExternalInput::GetResolvedInputType(EShaderFrequency ShaderFrequency) const
+{
+	EExternalInput Result = InputType;
+
+	if (ShaderFrequency != SF_Vertex)
+	{
+		switch (Result)
+		{
+		case EExternalInput::PrevWorldPosition:
+			Result = EExternalInput::WorldPosition;
+			break;
+		case EExternalInput::PrevTranslatedWorldPosition:
+			Result = EExternalInput::TranslatedWorldPosition;
+			break;
+		case EExternalInput::PrevWorldPosition_NoOffsets:
+			Result = EExternalInput::WorldPosition_NoOffsets;
+			break;
+		case EExternalInput::PrevTranslatedWorldPosition_NoOffsets:
+			Result = EExternalInput::TranslatedWorldPosition_NoOffsets;
+		default:
+			break;
+		}
+	}
+
+	if (ShaderFrequency != SF_Pixel)
+	{
+		switch (Result)
+		{
+		case EExternalInput::WorldPosition_NoOffsets:
+			Result = EExternalInput::WorldPosition;
+			break;
+		case EExternalInput::TranslatedWorldPosition_NoOffsets:
+			Result = EExternalInput::TranslatedWorldPosition;
+			break;
+		case EExternalInput::PrevWorldPosition_NoOffsets:
+			Result = EExternalInput::PrevWorldPosition;
+			break;
+		case EExternalInput::PrevTranslatedWorldPosition_NoOffsets:
+			Result = EExternalInput::PrevTranslatedWorldPosition;
+			break;
+		default:
+			break;
+		}
+	}
+
+	return Result;
+}
+
 bool FExpressionExternalInput::PrepareValue(FEmitContext& Context, FEmitScope& Scope, const FRequestedType& RequestedType, FPrepareValueResult& OutResult) const
 {
-	const FExternalInputDescription InputDesc = GetExternalInputDescription(InputType);
+	const EExternalInput ResolvedInputType = GetResolvedInputType(Context.ShaderFrequency);
+	const FExternalInputDescription InputDesc = GetExternalInputDescription(ResolvedInputType);
 
 	if (Context.bMarkLiveValues)
 	{
 		FEmitData& EmitMaterialData = Context.FindData<FEmitData>();
-		const int32 TypeIndex = (int32)InputType;
+		const int32 TypeIndex = (int32)ResolvedInputType;
 		EmitMaterialData.ExternalInputMask[Context.ShaderFrequency][TypeIndex] = true;
 
 		if (Context.MaterialCompilationOutput)
 		{
-			switch (InputType)
+			switch (ResolvedInputType)
 			{
 			case EExternalInput::PerInstanceRandom:
 				Context.MaterialCompilationOutput->bUsesPerInstanceRandom = true;
@@ -258,30 +307,31 @@ bool FExpressionExternalInput::PrepareValue(FEmitContext& Context, FEmitScope& S
 
 void FExpressionExternalInput::EmitValueShader(FEmitContext& Context, FEmitScope& Scope, const FRequestedType& RequestedType, FEmitValueShaderResult& OutResult) const
 {
+	const EExternalInput ResolvedInputType = GetResolvedInputType(Context.ShaderFrequency);
 	FEmitData& EmitMaterialData = Context.FindData<FEmitData>();
 
-	const int32 TypeIndex = (int32)InputType;
+	const int32 TypeIndex = (int32)ResolvedInputType;
 	EmitMaterialData.ExternalInputMask[Context.ShaderFrequency][TypeIndex] = true;
-	if (IsTexCoord(InputType))
+	if (IsTexCoord(ResolvedInputType))
 	{
 		const int32 TexCoordIndex = TypeIndex - (int32)EExternalInput::TexCoord0;
 		OutResult.Code = Context.EmitInlineExpression(Scope, Shader::EValueType::Float2, TEXT("Parameters.TexCoords[%].xy"), TexCoordIndex);
 	}
-	else if (IsTexCoord_Ddx(InputType))
+	else if (IsTexCoord_Ddx(ResolvedInputType))
 	{
 		const int32 TexCoordIndex = TypeIndex - (int32)EExternalInput::TexCoord0_Ddx;
 		OutResult.Code = Context.EmitInlineExpression(Scope, Shader::EValueType::Float2, TEXT("Parameters.TexCoords_DDX[%].xy"), TexCoordIndex);
 	}
-	else if (IsTexCoord_Ddy(InputType))
+	else if (IsTexCoord_Ddy(ResolvedInputType))
 	{
 		const int32 TexCoordIndex = TypeIndex - (int32)EExternalInput::TexCoord0_Ddy;
 		OutResult.Code = Context.EmitInlineExpression(Scope, Shader::EValueType::Float2, TEXT("Parameters.TexCoords_DDY[%].xy"), TexCoordIndex);
 	}
 	else
 	{
-		const FExternalInputDescription InputDesc = GetExternalInputDescription(InputType);
+		const FExternalInputDescription InputDesc = GetExternalInputDescription(ResolvedInputType);
 		const TCHAR* Code = nullptr;
-		switch (InputType)
+		switch (ResolvedInputType)
 		{
 		case EExternalInput::LightmapTexCoord: Code = TEXT("GetLightmapUVs(Parameters)"); break;
 		case EExternalInput::LightmapTexCoord_Ddx: Code = TEXT("GetLightmapUVs_DDX(Parameters)"); break;
@@ -1344,18 +1394,10 @@ void FExpressionTextureSample::EmitValueShader(FEmitContext& Context, FEmitScope
 
 bool FExpressionStaticTerrainLayerWeight::PrepareValue(FEmitContext& Context, FEmitScope& Scope, const FRequestedType& RequestedType, FPrepareValueResult& OutResult) const
 {
-	const EMaterialValueType TextureMaterialType = GEngine->WeightMapPlaceholderTexture->GetMaterialType();
-	const FRequestedType TexCoordRequestedType = Private::GetTexCoordType(TextureMaterialType);
-	const FPreparedType& TexCoordPreparedType = Context.PrepareExpression(TexCoordExpression, Scope, TexCoordRequestedType);
-	if (TexCoordPreparedType.IsVoid())
-	{
-		return false;
-	}
-	
 	bool bFoundMatchingParameter = false;
 	FEmitData& EmitData = Context.FindData<FEmitData>();
 
-	if (Context.bMarkLiveValues && EmitData.CachedExpressionData && EmitData.StaticParameters)
+	if (EmitData.StaticParameters)
 	{
 		for (int32 ParameterIndex = 0; ParameterIndex < EmitData.StaticParameters->EditorOnly.TerrainLayerWeightParameters.Num(); ++ParameterIndex)
 		{
@@ -1379,6 +1421,11 @@ bool FExpressionStaticTerrainLayerWeight::PrepareValue(FEmitContext& Context, FE
 
 			bFoundMatchingParameter = true;
 
+			if (!Context.bMarkLiveValues || !EmitData.CachedExpressionData)
+			{
+				break;
+			}
+
 			FMaterialParameterInfo WeightmapParameterInfo = BaseParameterInfo;
 			WeightmapParameterInfo.Name = *FString::Printf(TEXT("Weightmap%d"), WeightmapIndex);
 
@@ -1397,15 +1444,45 @@ bool FExpressionStaticTerrainLayerWeight::PrepareValue(FEmitContext& Context, FE
 			EmitData.CachedExpressionData->AddParameter(LayerMaskParameterInfo, LayerMaskParameterMeta, UnusedReferencedTexture);
 		}
 
-		if (bFoundMatchingParameter)
+		if (bFoundMatchingParameter && EmitData.CachedExpressionData)
 		{
 			EmitData.CachedExpressionData->ReferencedTextures.AddUnique(GEngine->WeightMapPlaceholderTexture);
 		}
 	}
 
+	if (!bFoundMatchingParameter)
+	{
+		return OutResult.SetType(Context, RequestedType, UE::HLSLTree::Private::PrepareConstant(Context.Material && Context.Material->IsPreview() ? DefaultWeight : 0.f));
+	}
+
+	const EMaterialValueType TextureMaterialType = GEngine->WeightMapPlaceholderTexture->GetMaterialType();
+	const FRequestedType TexCoordRequestedType = Private::GetTexCoordType(TextureMaterialType);
+	const FPreparedType& TexCoordPreparedType = Context.PrepareExpression(TexCoordExpression, Scope, TexCoordRequestedType);
+	if (TexCoordPreparedType.IsVoid())
+	{
+		return false;
+	}
+
 	return OutResult.SetType(Context, RequestedType, EExpressionEvaluation::Shader, Shader::EValueType::Float4);
 }
 
+void FExpressionStaticTerrainLayerWeight::EmitValuePreshader(FEmitContext& Context, FEmitScope& Scope, const FRequestedType& RequestedType, FEmitValuePreshaderResult& OutResult) const
+{
+	Context.PreshaderStackPosition++;
+	if (Context.Material && Context.Material->IsPreview())
+	{
+		const Shader::FValue Value = DefaultWeight;
+		OutResult.Type = Value.Type;
+		OutResult.Preshader.WriteOpcode(Shader::EPreshaderOpcode::Constant).Write(Value);
+	}
+	else
+	{
+		const Shader::FValue Value = 0.f;
+		OutResult.Type = Value.Type;
+		OutResult.Preshader.WriteOpcode(Shader::EPreshaderOpcode::ConstantZero).Write(OutResult.Type);
+	}
+}
+
 void FExpressionStaticTerrainLayerWeight::EmitValueShader(FEmitContext& Context, FEmitScope& Scope, const FRequestedType& RequestedType, FEmitValueShaderResult& OutResult) const
 {
 	FEmitData& EmitData = Context.FindData<FEmitData>();
@@ -1501,14 +1578,8 @@ void FExpressionStaticTerrainLayerWeight::EmitValueShader(FEmitContext& Context,
 		++NumWeightmapParameters;
 	}
 
-	if (NumWeightmapParameters > 0)
-	{
-		OutResult.Code = NumWeightmapParameters > 1 ? Context.EmitExpression(Scope, Shader::EValueType::Float1, TEXT("saturate(%)"), EmitResult) : EmitResult;
-	}
-	else
-	{
-		OutResult.Code = Context.EmitInlineExpression(Scope, Shader::EValueType::Float1, TEXT("((float)%)"), Context.Material && Context.Material->IsPreview() ? DefaultWeight : 0.f);
-	}
+	check(NumWeightmapParameters > 0);
+	OutResult.Code = NumWeightmapParameters > 1 ? Context.EmitExpression(Scope, Shader::EValueType::Float1, TEXT("saturate(%)"), EmitResult) : EmitResult;
 }
 
 bool FExpressionTextureProperty::PrepareValue(FEmitContext& Context, FEmitScope& Scope, const FRequestedType& RequestedType, FPrepareValueResult& OutResult) const
@@ -1584,10 +1655,19 @@ void FExpressionRuntimeVirtualTextureUniform::EmitValuePreshader(FEmitContext& C
 
 bool FExpressionRuntimeVirtualTextureOutput::PrepareValue(FEmitContext& Context, FEmitScope& Scope, const FRequestedType& RequestedType, FPrepareValueResult& OutResult) const
 {
-	if (Context.bMarkLiveValues && Context.MaterialCompilationOutput)
+	if (Context.bMarkLiveValues)
 	{
-		Context.MaterialCompilationOutput->bHasRuntimeVirtualTextureOutputNode |= OutputAttributeMask != 0;
-		Context.MaterialCompilationOutput->RuntimeVirtualTextureOutputAttributeMask |= OutputAttributeMask;
+		if (Context.MaterialCompilationOutput)
+		{
+			Context.MaterialCompilationOutput->bHasRuntimeVirtualTextureOutputNode |= OutputAttributeMask != 0;
+			Context.MaterialCompilationOutput->RuntimeVirtualTextureOutputAttributeMask |= OutputAttributeMask;
+		}
+
+		FEmitData& EmitData = Context.FindData<FEmitData>();
+		if (EmitData.CachedExpressionData)
+		{
+			EmitData.CachedExpressionData->bHasRuntimeVirtualTextureOutput = true;
+		}
 	}
 
 	return FExpressionForward::PrepareValue(Context, Scope, RequestedType, OutResult);
diff --git a/Engine/Source/Runtime/Engine/Private/Materials/MaterialShared.cpp b/Engine/Source/Runtime/Engine/Private/Materials/MaterialShared.cpp
--- a/Engine/Source/Runtime/Engine/Private/Materials/MaterialShared.cpp
+++ b/Engine/Source/Runtime/Engine/Private/Materials/MaterialShared.cpp
@@ -424,6 +424,21 @@ const UE::HLSLTree::FExpression* FExpressionInput::AcquireHLSLExpressionOrExtern
 	return AcquireHLSLExpressionOrExternalInput(Generator, Scope, Input, Generator.FindInputIndex(this));
 }
 
+const UE::HLSLTree::FExpression* FExpressionInput::AcquireHLSLExpressionOrDefaultExpression(FMaterialHLSLGenerator& Generator, UE::HLSLTree::FScope& Scope, const UE::HLSLTree::FExpression* DefaultExpression, int32 InputIndex) const
+{
+	const FExpressionInput TracedInput = GetTracedInput();
+	if (!TracedInput.Expression)
+	{
+		return Generator.NewDefaultInputExpression(InputIndex, DefaultExpression);
+	}
+	return TryAcquireHLSLExpression(Generator, Scope, InputIndex);
+}
+
+const UE::HLSLTree::FExpression* FExpressionInput::AcquireHLSLExpressionOrDefaultExpression(FMaterialHLSLGenerator& Generator, UE::HLSLTree::FScope& Scope, const UE::HLSLTree::FExpression* DefaultExpression) const
+{
+	return AcquireHLSLExpressionOrDefaultExpression(Generator, Scope, DefaultExpression, Generator.FindInputIndex(this));
+}
+
 void FExpressionInput::Connect( int32 InOutputIndex, class UMaterialExpression* InExpression )
 {
 	InExpression->ConnectExpression(this, InOutputIndex);
diff --git a/Engine/Source/Runtime/Engine/Public/HLSLTree/HLSLTreeCommon.h b/Engine/Source/Runtime/Engine/Public/HLSLTree/HLSLTreeCommon.h
--- a/Engine/Source/Runtime/Engine/Public/HLSLTree/HLSLTreeCommon.h
+++ b/Engine/Source/Runtime/Engine/Public/HLSLTree/HLSLTreeCommon.h
@@ -58,6 +58,11 @@ public:
 	virtual const FExpression* ComputePreviousFrame(FTree& Tree, const FRequestedType& RequestedType) const override;
 };
 
+namespace Private
+{
+	FPreparedType PrepareConstant(const Shader::FValue& Value);
+}
+
 class FExpressionConstant : public FExpression
 {
 public:
diff --git a/Engine/Source/Runtime/Engine/Public/HLSLTree/HLSLTreeEmit.h b/Engine/Source/Runtime/Engine/Public/HLSLTree/HLSLTreeEmit.h
--- a/Engine/Source/Runtime/Engine/Public/HLSLTree/HLSLTreeEmit.h
+++ b/Engine/Source/Runtime/Engine/Public/HLSLTree/HLSLTreeEmit.h
@@ -593,8 +593,7 @@ public:
 	TMap<FXxHash64, FPrepareValueResult*> PrepareValueMap;
 	TMap<FXxHash64, FRequestedType*> RequestedTypeTracker;
 	TMap<const FExpression*, FEmitScope*> PrepareLocalPHIMap;
-	TMap<const FExpression*, FEmitShaderExpression*> EmitLocalPHIMap;
-	// TODO: investigate EmitExpressionMap and EmitPreshaderMap can be removed and only keep EmitValueMap
+	TMap<FXxHash64, FEmitShaderExpression*> EmitLocalPHIMap;
 	TMap<FXxHash64, FEmitShaderExpression*> EmitExpressionMap;
 	TMap<FXxHash64, FEmitShaderExpression*> EmitPreshaderMap;
 	TMap<FXxHash64, FEmitShaderExpression*> EmitValueMap;
diff --git a/Engine/Source/Runtime/Engine/Public/MaterialCachedHLSLTree.h b/Engine/Source/Runtime/Engine/Public/MaterialCachedHLSLTree.h
--- a/Engine/Source/Runtime/Engine/Public/MaterialCachedHLSLTree.h
+++ b/Engine/Source/Runtime/Engine/Public/MaterialCachedHLSLTree.h
@@ -15,6 +15,14 @@ class UMaterialExpressionCustomOutput;
 class UMaterialParameterCollection;
 struct FMaterialLayersFunctions;
 
+namespace UE
+{
+namespace HLSLTree
+{
+class FEmitContext;
+}
+}
+
 struct FMaterialConnectionKey
 {
 	const UObject* InputObject = nullptr;
@@ -61,11 +69,10 @@ public:
 
 	const UE::Shader::FStructType* GetMaterialAttributesType() const { return MaterialAttributesType; }
 	const UE::Shader::FStructType* GetVTPageTableResultType() const { return VTPageTableResultType; }
-	const UE::Shader::FValue& GetMaterialAttributesDefaultValue() const { return MaterialAttributesDefaultValue; }
 
 	const TMap<FMaterialConnectionKey, const UE::HLSLTree::FExpression*>& GetConnections() const { return ConnectionMap; }
 
-	ENGINE_API void SetRequestedFields(EShaderFrequency ShaderFrequency, UE::HLSLTree::FRequestedType& OutRequestedType) const;
+	ENGINE_API void SetRequestedFields(const UE::HLSLTree::FEmitContext& Context, UE::HLSLTree::FRequestedType& OutRequestedType) const;
 	void EmitSharedCode(FStringBuilderBase& OutCode) const;
 
 	ENGINE_API bool IsAttributeUsed(UE::HLSLTree::FEmitContext& Context,
@@ -74,6 +81,8 @@ public:
 		EMaterialProperty Property) const;
 
 private:
+	const UE::Shader::FValue& GetMaterialAttributesDefaultValue() const { return MaterialAttributesDefaultValue; }
+
 	FMemStackBase Allocator;
 	UE::Shader::FStructTypeRegistry TypeRegistry;
 	UE::HLSLTree::FTree* HLSLTree = nullptr;
diff --git a/Engine/Source/Runtime/Engine/Public/MaterialExpressionIO.h b/Engine/Source/Runtime/Engine/Public/MaterialExpressionIO.h
--- a/Engine/Source/Runtime/Engine/Public/MaterialExpressionIO.h
+++ b/Engine/Source/Runtime/Engine/Public/MaterialExpressionIO.h
@@ -74,11 +74,13 @@ struct FExpressionInput
 	ENGINE_API const UE::HLSLTree::FExpression* AcquireHLSLExpression(FMaterialHLSLGenerator& Generator, UE::HLSLTree::FScope& Scope, int32 InputIndex) const;
 	ENGINE_API const UE::HLSLTree::FExpression* AcquireHLSLExpressionOrConstant(FMaterialHLSLGenerator& Generator, UE::HLSLTree::FScope& Scope, const UE::Shader::FValue& ConstantValue, int32 InputIndex) const;
 	ENGINE_API const UE::HLSLTree::FExpression* AcquireHLSLExpressionOrExternalInput(FMaterialHLSLGenerator& Generator, UE::HLSLTree::FScope& Scope, UE::HLSLTree::Material::EExternalInput Input, int32 InputIndex) const;
+	ENGINE_API const UE::HLSLTree::FExpression* AcquireHLSLExpressionOrDefaultExpression(FMaterialHLSLGenerator& Generator, UE::HLSLTree::FScope& Scope, const UE::HLSLTree::FExpression* DefaultExpression, int32 InputIndex) const;
 
 	ENGINE_API const UE::HLSLTree::FExpression* TryAcquireHLSLExpression(FMaterialHLSLGenerator& Generator, UE::HLSLTree::FScope& Scope) const;
 	ENGINE_API const UE::HLSLTree::FExpression* AcquireHLSLExpression(FMaterialHLSLGenerator& Generator, UE::HLSLTree::FScope& Scope) const;
 	ENGINE_API const UE::HLSLTree::FExpression* AcquireHLSLExpressionOrConstant(FMaterialHLSLGenerator& Generator, UE::HLSLTree::FScope& Scope, const UE::Shader::FValue& ConstantValue) const;
 	ENGINE_API const UE::HLSLTree::FExpression* AcquireHLSLExpressionOrExternalInput(FMaterialHLSLGenerator& Generator, UE::HLSLTree::FScope& Scope, UE::HLSLTree::Material::EExternalInput Input) const;
+	ENGINE_API const UE::HLSLTree::FExpression* AcquireHLSLExpressionOrDefaultExpression(FMaterialHLSLGenerator& Generator, UE::HLSLTree::FScope& Scope, const UE::HLSLTree::FExpression* DefaultExpression) const;
 
 	/**
 	 * Tests if the input has a material expression connected to it
diff --git a/Engine/Source/Runtime/Engine/Public/MaterialHLSLGenerator.h b/Engine/Source/Runtime/Engine/Public/MaterialHLSLGenerator.h
--- a/Engine/Source/Runtime/Engine/Public/MaterialHLSLGenerator.h
+++ b/Engine/Source/Runtime/Engine/Public/MaterialHLSLGenerator.h
@@ -78,7 +78,7 @@ public:
 	ENGINE_API UE::HLSLTree::FTree& GetTree() const;
 	UE::Shader::FStructTypeRegistry& GetTypeRegistry() const;
 	const UE::Shader::FStructType* GetMaterialAttributesType() const;
-	const UE::Shader::FValue& GetMaterialAttributesDefaultValue() const;
+	const UE::HLSLTree::FExpression* GetMaterialAttributesDefaultExpression() const;
 
 	UMaterialExpression* GetCurrentExpression() const;
 
@@ -128,6 +128,7 @@ public:
 
 	const UE::HLSLTree::FExpression* NewDefaultInputConstant(int32 InputIndex, const UE::Shader::FValue& Value);
 	const UE::HLSLTree::FExpression* NewDefaultInputExternal(int32 InputIndex, UE::HLSLTree::Material::EExternalInput Input);
+	const UE::HLSLTree::FExpression* NewDefaultInputExpression(int32 InputIndex, const UE::HLSLTree::FExpression* Expression);
 
 	/**
 	 * Returns the appropriate HLSLNode representing the given UMaterialExpression.
diff --git a/Engine/Source/Runtime/Engine/Public/MaterialHLSLTree.h b/Engine/Source/Runtime/Engine/Public/MaterialHLSLTree.h
--- a/Engine/Source/Runtime/Engine/Public/MaterialHLSLTree.h
+++ b/Engine/Source/Runtime/Engine/Public/MaterialHLSLTree.h
@@ -253,6 +253,9 @@ public:
 	virtual const FExpression* ComputePreviousFrame(FTree& Tree, const FRequestedType& RequestedType) const override;
 	virtual bool PrepareValue(FEmitContext& Context, FEmitScope& Scope, const FRequestedType& RequestedType, FPrepareValueResult& OutResult) const override;
 	virtual void EmitValueShader(FEmitContext& Context, FEmitScope& Scope, const FRequestedType& RequestedType, FEmitValueShaderResult& OutResult) const override;
+
+private:
+	EExternalInput GetResolvedInputType(EShaderFrequency ShaderFrequency) const;
 };
 
 class FExpressionShadingModel : public FExpression
@@ -472,6 +475,7 @@ public:
 	float DefaultWeight;
 
 	virtual bool PrepareValue(FEmitContext& Context, FEmitScope& Scope, const FRequestedType& RequestedType, FPrepareValueResult& OutResult) const override;
+	virtual void EmitValuePreshader(FEmitContext& Context, FEmitScope& Scope, const FRequestedType& RequestedType, FEmitValuePreshaderResult& OutResult) const override;
 	virtual void EmitValueShader(FEmitContext& Context, FEmitScope& Scope, const FRequestedType& RequestedType, FEmitValueShaderResult& OutResult) const override;
 };
 
