Fix MSVC static analysis warnings

diff --git a/Engine/Source/Runtime/Engine/Private/HLSLTree/HLSLTree.cpp b/Engine/Source/Runtime/Engine/Private/HLSLTree/HLSLTree.cpp
--- a/Engine/Source/Runtime/Engine/Private/HLSLTree/HLSLTree.cpp
+++ b/Engine/Source/Runtime/Engine/Private/HLSLTree/HLSLTree.cpp
@@ -315,7 +315,14 @@ bool FExpressionLocalPHI::PrepareValue(FEmitContext& Context, FEmitScope& Scope,
 
 	for (int32 Index = 1; Index < NumLiveScopes; ++Index)
 	{
+#ifdef _MSC_VER
+#pragma warning(push)
+#pragma warning(disable : 6385) // The first NumLiveScopes entries in LiveValues are valid
+#endif
 		if (LiveValues[Index] != LiveValues[Index - 1])
+#ifdef _MSC_VER
+#pragma warning(pop)
+#endif
 		{
 			bMergeScopeEvaluations = true;
 			break;
@@ -350,7 +357,14 @@ bool FExpressionLocalPHI::PrepareValue(FEmitContext& Context, FEmitScope& Scope,
 	{
 		for (int32 Index = 0; Index < NumLiveScopes; ++Index)
 		{
+#ifdef _MSC_VER
+#pragma warning(push)
+#pragma warning(disable : 28182) // The first NumLiveScopes entries in EmitValueScopes are not null
+#endif
 			CurrentType.MergeEvaluation(EmitValueScopes[Index]->Evaluation);
+#ifdef _MSC_VER
+#pragma warning(pop)
+#endif
 		}
 		verify(OutResult.SetType(Context, RequestedType, CurrentType));
 	}
