Swap the uniform buffer invalidation check for material render proxies in editor to a buffer invalidation to avoid future unpredictable crashes. Non-editor retains the original check for QA reasons in cooked projects.

diff --git a/Engine/Source/Runtime/Engine/Private/Materials/MaterialRenderProxy.cpp b/Engine/Source/Runtime/Engine/Private/Materials/MaterialRenderProxy.cpp
--- a/Engine/Source/Runtime/Engine/Private/Materials/MaterialRenderProxy.cpp
+++ b/Engine/Source/Runtime/Engine/Private/Materials/MaterialRenderProxy.cpp
@@ -399,9 +399,22 @@ void FMaterialRenderProxy::EvaluateUniformExpressions(FRHICommandListBase& RHICm
 			UE_LOG(LogMaterial, Fatal, TEXT("The Uniformbuffer needs to be valid if it has been set"));
 		}
 
-		// The actual pointer may not match because there are cases (in the editor, during the shader compilation) when material's shader map gets updated without proxy's cache
-		// getting invalidated, but the layout contents must match.
+		/**
+		* The actual pointer may not match because there are cases(in the editor, during the shader compilation) when material's shader map gets updated without proxy's cache
+		* getting invalidated, but the layout contents must match.
+		*/
+#if WITH_EDITOR
+		/**
+		* If we are in the editor, this is likely occuring due to a mismatch of buffer invalidations and the type of shaders that need processing, so we can safely reset the
+		* Buffer layout to avoid any engine crashes. However we still want this check in project builds for safety when doing QA passes.
+		*/
+		if (OutUniformExpressionCache.UniformBuffer->GetLayoutPtr() != UniformBufferLayout || *OutUniformExpressionCache.UniformBuffer->GetLayoutPtr() != *UniformBufferLayout)
+		{
+			OutUniformExpressionCache.UniformBuffer = nullptr;
+		}
+#else
 		check(OutUniformExpressionCache.UniformBuffer->GetLayoutPtr() == UniformBufferLayout || *OutUniformExpressionCache.UniformBuffer->GetLayoutPtr() == *UniformBufferLayout);
+#endif
 	}
 
 	if (Updater)
