Added a CompilePreview function for UMaterialExpressionShadingModel to circumvent a necessary uniform expression being missing when recaching shaders in material editor, causing an assert. Previews aren't needed for the shading model node so this is acceptable for now.

diff --git a/Engine/Source/Runtime/Engine/Classes/Materials/MaterialExpressionShadingModel.h b/Engine/Source/Runtime/Engine/Classes/Materials/MaterialExpressionShadingModel.h
--- a/Engine/Source/Runtime/Engine/Classes/Materials/MaterialExpressionShadingModel.h
+++ b/Engine/Source/Runtime/Engine/Classes/Materials/MaterialExpressionShadingModel.h
@@ -28,6 +28,7 @@ class UMaterialExpressionShadingModel : public UMaterialExpression
 	//~ Begin UMaterialExpression Interface
 #if WITH_EDITOR
 	virtual int32 Compile(class FMaterialCompiler* Compiler, int32 OutputIndex) override;
+	virtual int32 CompilePreview(class FMaterialCompiler* Compiler, int32 OutputIndex) override;
 	virtual uint32 GetOutputType(int32 OutputIndex) override;
 	virtual void GetCaption(TArray<FString>& OutCaptions) const override;
 	virtual bool GenerateHLSLExpression(FMaterialHLSLGenerator& Generator, UE::HLSLTree::FScope& Scope, int32 OutputIndex, UE::HLSLTree::FExpression const*& OutExpression) const override;
@@ -36,4 +37,4 @@ public:
 	UPROPERTY(EditAnywhere, Category=ShadingModel,  meta=(ValidEnumValues="MSM_DefaultLit, MSM_Subsurface, MSM_PreintegratedSkin, MSM_ClearCoat, MSM_SubsurfaceProfile, MSM_TwoSidedFoliage, MSM_Hair, MSM_Cloth, MSM_Eye", ShowAsInputPin = "Primary"))
 	TEnumAsByte<enum EMaterialShadingModel> ShadingModel = MSM_DefaultLit;
 	//~ End UMaterialExpression Interface
-};
+};
\ No newline at end of file
diff --git a/Engine/Source/Runtime/Engine/Private/Materials/MaterialExpressions.cpp b/Engine/Source/Runtime/Engine/Private/Materials/MaterialExpressions.cpp
--- a/Engine/Source/Runtime/Engine/Private/Materials/MaterialExpressions.cpp
+++ b/Engine/Source/Runtime/Engine/Private/Materials/MaterialExpressions.cpp
@@ -22485,6 +22485,11 @@ int32 UMaterialExpressionShadingModel::Compile(class FMaterialCompiler* Compiler
 	return Compiler->ShadingModel(ShadingModel);
 }
 
+int32 UMaterialExpressionShadingModel::CompilePreview(class FMaterialCompiler* Compiler, int32 OutputIndex)
+{
+	return Compiler->Constant(0.f);
+}
+
 void UMaterialExpressionShadingModel::GetCaption(TArray<FString>& OutCaptions) const
 {
 	const UEnum* ShadingModelEnum = FindObject<UEnum>(nullptr, TEXT("Engine.EMaterialShadingModel"));
