Fixed unacceptable word usage in commentss (CIS error)

#rnx

diff --git a/Engine/Source/Runtime/Engine/Private/InstancedStaticMesh.cpp b/Engine/Source/Runtime/Engine/Private/InstancedStaticMesh.cpp
--- a/Engine/Source/Runtime/Engine/Private/InstancedStaticMesh.cpp
+++ b/Engine/Source/Runtime/Engine/Private/InstancedStaticMesh.cpp
@@ -3716,7 +3716,7 @@ bool UInstancedStaticMeshComponent::UpdateInstances(
 	int32 TotalSizeUpdateBytes = 0;
 #endif
 
-	// Note: this will nuke the custom data for all instaces meaning they will be updated below
+	// Note: this will reset the custom data for all instaces meaning they will be updated below
 	if (InNumCustomDataFloats != NumCustomDataFloats)
 	{
 		SetNumCustomDataFloats(InNumCustomDataFloats);
diff --git a/Engine/Source/Runtime/Engine/Private/InstancedStaticMesh/ISMInstanceDataManager.cpp b/Engine/Source/Runtime/Engine/Private/InstancedStaticMesh/ISMInstanceDataManager.cpp
--- a/Engine/Source/Runtime/Engine/Private/InstancedStaticMesh/ISMInstanceDataManager.cpp
+++ b/Engine/Source/Runtime/Engine/Private/InstancedStaticMesh/ISMInstanceDataManager.cpp
@@ -554,7 +554,7 @@ bool FPrimitiveInstanceDataManager::FlushChanges(FInstanceUpdateComponentDesc &&
 		return false;
 	}
 
-	// Can't update if there is no proxy, if it was nuked we must wait for a new scene proxy to come around and request a new one to be created.
+	// Can't update if there is no proxy, if it was destroyed we must wait for a new scene proxy to come around and request a new one to be created.
 	if (!Proxy)
 	{
 		return false;
@@ -647,7 +647,7 @@ bool FPrimitiveInstanceDataManager::FlushChanges(FInstanceUpdateComponentDesc &&
 		{
 			FISMCInstanceDataSceneProxy &ProxyRef = *Proxy;
 			
-			// Forcibly nuke any tracking state, if the external entity manages this it can track the data on its own.
+			// Forcibly destroy any tracking state, if the external entity manages this it can track the data on its own.
 			ProxyRef.ChangeMask.Reset();
 			ProxyRef.InstanceIdIndexMap = MoveTemp(InstanceIdIndexMap);
 #if WITH_EDITOR
