Fix build health warning on null GEngine pointer

diff --git a/Engine/Source/Runtime/Engine/Private/GameViewportClient.cpp b/Engine/Source/Runtime/Engine/Private/GameViewportClient.cpp
--- a/Engine/Source/Runtime/Engine/Private/GameViewportClient.cpp
+++ b/Engine/Source/Runtime/Engine/Private/GameViewportClient.cpp
@@ -1333,6 +1333,7 @@ void UGameViewportClient::Draw(FViewport* InViewport, FCanvas* SceneCanvas)
 {
 	//Valid SceneCanvas is required.  Make this explicit.
 	check(SceneCanvas);
+	check(GEngine);
 
 	BeginDrawDelegate.Broadcast();
 
