Speculative fix for cooker error fallout from ISM instance data refactor.

diff --git a/Engine/Source/Runtime/Engine/Private/InstancedStaticMesh.cpp b/Engine/Source/Runtime/Engine/Private/InstancedStaticMesh.cpp
--- a/Engine/Source/Runtime/Engine/Private/InstancedStaticMesh.cpp
+++ b/Engine/Source/Runtime/Engine/Private/InstancedStaticMesh.cpp
@@ -4288,16 +4288,15 @@ void UInstancedStaticMeshComponent::OnRegister()
 
 	if (FApp::CanEverRender() && !HasAnyFlags(RF_ClassDefaultObject | RF_ArchetypeObject))
 	{
-	// If we don't have a random seed for this instanced static mesh component yet, then go ahead and
-	// generate one now.  This will be saved with the static mesh component and used for future generation
-	// of random numbers for this component's instances. (Used by the PerInstanceRandom material expression)
-	while (InstancingRandomSeed == 0)
-	{
-		InstancingRandomSeed = FMath::Rand();
-	}
-
-	}
+		// If we don't have a random seed for this instanced static mesh component yet, then go ahead and
+		// generate one now.  This will be saved with the static mesh component and used for future generation
+		// of random numbers for this component's instances. (Used by the PerInstanceRandom material expression)
+		while (InstancingRandomSeed == 0)
+		{
+			InstancingRandomSeed = FMath::Rand();
+		}
 	}
+}
 		
 #if WITH_EDITOR
 
@@ -4605,6 +4604,17 @@ void UInstancedStaticMeshComponent::OnPostLoadPerInstanceData()
 {
 	TRACE_CPUPROFILER_EVENT_SCOPE(UInstancedStaticMeshComponent::OnPostLoadPerInstanceData);
 
+	if (!HasAnyFlags(RF_ClassDefaultObject|RF_ArchetypeObject))
+	{
+		// If we don't have a random seed for this instanced static mesh component yet, then go ahead and
+		// generate one now.  This will be saved with the static mesh component and used for future generation
+		// of random numbers for this component's instances. (Used by the PerInstanceRandom material expression)
+		while (InstancingRandomSeed == 0)
+		{
+			InstancingRandomSeed = FMath::Rand();
+		}		
+	}
+
 	if (AActor* Owner = GetOwner())
 	{
 		ULevel* OwnerLevel = Owner->GetLevel();
