Set default Local Exposure Contrast Scale in PPV to 0.8.
- required deprecating exiting members in FPostProcessSettings and adding new ones with new default value.

diff --git a/Engine/Source/Runtime/Core/Public/UObject/FortniteMainBranchObjectVersion.h b/Engine/Source/Runtime/Core/Public/UObject/FortniteMainBranchObjectVersion.h
--- a/Engine/Source/Runtime/Core/Public/UObject/FortniteMainBranchObjectVersion.h
+++ b/Engine/Source/Runtime/Core/Public/UObject/FortniteMainBranchObjectVersion.h
@@ -378,6 +378,9 @@ struct FFortniteMainBranchObjectVersion
 		// Serialize custom editor bounds for HLOD actor descriptors.
 		WorldPartitionHLODActorDescSerializeEditorBounds,
 
+		// Changed default Local Exposure Contrast from 1.0 to 0.8
+		LocalExposureDefaultChangeFrom1,
+
 		// -----<new versions can be added above this line>-------------------------------------------------
 		VersionPlusOne,
 		LatestVersion = VersionPlusOne - 1
diff --git a/Engine/Source/Runtime/Engine/Classes/Engine/RendererSettings.h b/Engine/Source/Runtime/Engine/Classes/Engine/RendererSettings.h
--- a/Engine/Source/Runtime/Engine/Classes/Engine/RendererSettings.h
+++ b/Engine/Source/Runtime/Engine/Classes/Engine/RendererSettings.h
@@ -673,13 +673,13 @@ class URendererSettings : public UDeveloperSettings
 
 	UPROPERTY(config, EditAnywhere, Category = DefaultSettings, meta = (
 		ConsoleVariable = "r.DefaultFeature.LocalExposure.HighlightContrastScale", DisplayName = "Local Exposure Highlight Contrast",
-		ToolTip = "Default Value for Local Exposure Highlight Contrast Scale.", ClampMin = "0.0", ClampMax = "1.0"))
-	float DefaultFeatureLocalExposureHighlightContrastScale;
+		ToolTip = "Default Value for Local Exposure Highlight Contrast.", ClampMin = "0.0", ClampMax = "1.0"))
+	float DefaultFeatureLocalExposureHighlightContrast;
 
 	UPROPERTY(config, EditAnywhere, Category = DefaultSettings, meta = (
 		ConsoleVariable = "r.DefaultFeature.LocalExposure.ShadowContrastScale", DisplayName = "Local Exposure Shadow Contrast",
-		ToolTip = "Default Value for Local Exposure Shadow Contrast Scale.", ClampMin = "0.0", ClampMax = "1.0"))
-	float DefaultFeatureLocalExposureShadowContrastScale;
+		ToolTip = "Default Value for Local Exposure Shadow Contrast.", ClampMin = "0.0", ClampMax = "1.0"))
+	float DefaultFeatureLocalExposureShadowContrast;
 
 	UPROPERTY(config, EditAnywhere, Category = DefaultSettings, meta = (
 		ConsoleVariable = "r.DefaultFeature.MotionBlur", DisplayName = "Motion Blur",
diff --git a/Engine/Source/Runtime/Engine/Classes/Engine/Scene.h b/Engine/Source/Runtime/Engine/Classes/Engine/Scene.h
--- a/Engine/Source/Runtime/Engine/Classes/Engine/Scene.h
+++ b/Engine/Source/Runtime/Engine/Classes/Engine/Scene.h
@@ -901,11 +901,17 @@ struct FPostProcessSettings
 	UPROPERTY()
 	uint8 bOverride_LocalExposureContrastScale_DEPRECATED:1;
 
+	UPROPERTY()
+	uint8 bOverride_LocalExposureHighlightContrastScale_DEPRECATED :1;
+
+	UPROPERTY()
+	uint8 bOverride_LocalExposureShadowContrastScale_DEPRECATED :1;
+
 	UPROPERTY(EditAnywhere, BlueprintReadWrite, Category=Overrides, meta=(PinHiddenByDefault, InlineEditConditionToggle))
-	uint8 bOverride_LocalExposureHighlightContrastScale:1;
+	uint8 bOverride_LocalExposureHighlightContrast:1;
 
 	UPROPERTY(EditAnywhere, BlueprintReadWrite, Category=Overrides, meta=(PinHiddenByDefault, InlineEditConditionToggle))
-	uint8 bOverride_LocalExposureShadowContrastScale:1;
+	uint8 bOverride_LocalExposureShadowContrast:1;
 
 	UPROPERTY(EditAnywhere, BlueprintReadWrite, Category=Overrides, meta=(PinHiddenByDefault, InlineEditConditionToggle))
 	uint8 bOverride_LocalExposureDetailStrength:1;
@@ -1823,23 +1829,23 @@ struct FPostProcessSettings
 	UPROPERTY()
 	float AutoExposureCalibrationConstant_DEPRECATED;
 
-	/** 
-	 * Local Exposure decomposes luminance of the frame into a base layer and a detail layer.
-	 * Contrast of the base layer is reduced based on this value.
-	 * Value less than 1 will enable local exposure.
-	 * Good values are usually in the range 0.6 .. 1.0.
-	*/
 	UPROPERTY()
 	float LocalExposureContrastScale_DEPRECATED;
 
+	UPROPERTY()
+	float LocalExposureHighlightContrastScale_DEPRECATED;
+
+	UPROPERTY()
+	float LocalExposureShadowContrastScale_DEPRECATED;
+
 	/** 
 	 * Local Exposure decomposes luminance of the frame into a base layer and a detail layer.
 	 * Contrast of the base layer is reduced based on this value.
 	 * Value less than 1 will enable local exposure.
 	 * Good values are usually in the range 0.6 .. 1.0.
 	*/
-	UPROPERTY(interp, BlueprintReadWrite, Category="Lens|Local Exposure", meta=(UIMin = "0.0", UIMax = "1.0", editcondition = "bOverride_LocalExposureHighlightContrastScale", DisplayName = "Highlight Contrast Scale"))
-	float LocalExposureHighlightContrastScale;
+	UPROPERTY(interp, BlueprintReadWrite, Category="Lens|Local Exposure", meta=(UIMin = "0.0", UIMax = "1.0", editcondition = "bOverride_LocalExposureHighlightContrast", DisplayName = "Highlight Contrast"))
+	float LocalExposureHighlightContrast;
 
 	/** 
 	 * Local Exposure decomposes luminance of the frame into a base layer and a detail layer.
@@ -1847,8 +1853,8 @@ struct FPostProcessSettings
 	 * Value less than 1 will enable local exposure.
 	 * Good values are usually in the range 0.6 .. 1.0.
 	*/
-	UPROPERTY(interp, BlueprintReadWrite, Category="Lens|Local Exposure", meta=(UIMin = "0.0", UIMax = "1.0", editcondition = "bOverride_LocalExposureShadowContrastScale", DisplayName = "Shadow Contrast Scale"))
-	float LocalExposureShadowContrastScale;
+	UPROPERTY(interp, BlueprintReadWrite, Category="Lens|Local Exposure", meta=(UIMin = "0.0", UIMax = "1.0", editcondition = "bOverride_LocalExposureShadowContrast", DisplayName = "Shadow Contrast"))
+	float LocalExposureShadowContrast;
 
 	/**
 	 * Local Exposure decomposes luminance of the frame into a base layer and a detail layer.
@@ -2283,17 +2289,6 @@ struct FPostProcessSettings
 		{
 			RayTracingGIType = (ERayTracingGlobalIlluminationType)(RayTracingGI_DEPRECATED == 1);
 		}
-
-		if (bOverride_LocalExposureContrastScale_DEPRECATED)
-		{
-			bOverride_LocalExposureHighlightContrastScale = bOverride_LocalExposureContrastScale_DEPRECATED;
-			bOverride_LocalExposureShadowContrastScale = bOverride_LocalExposureContrastScale_DEPRECATED;
-		}
-		if (LocalExposureContrastScale_DEPRECATED != 1.0f)
-		{
-			LocalExposureHighlightContrastScale = LocalExposureContrastScale_DEPRECATED;
-			LocalExposureShadowContrastScale = LocalExposureContrastScale_DEPRECATED;
-		}
 	}
 #endif
 
diff --git a/Engine/Source/Runtime/Engine/Private/Scene.cpp b/Engine/Source/Runtime/Engine/Private/Scene.cpp
--- a/Engine/Source/Runtime/Engine/Private/Scene.cpp
+++ b/Engine/Source/Runtime/Engine/Private/Scene.cpp
@@ -7,6 +7,7 @@
 #include "UObject/RenderingObjectVersion.h"
 #include "UObject/ReleaseObjectVersion.h"
 #include "UObject/UE5ReleaseStreamObjectVersion.h"
+#include "UObject/FortniteMainBranchObjectVersion.h"
 #include "DataDrivenShaderPlatformInfo.h"
 #include "UObject/UnrealType.h"
 
@@ -509,8 +510,10 @@ FPostProcessSettings::FPostProcessSettings()
 	AutoExposureApplyPhysicalCameraExposure = 1;
 
 	LocalExposureContrastScale_DEPRECATED = 1.0f;
-	LocalExposureHighlightContrastScale = 1.0f;
-	LocalExposureShadowContrastScale = 1.0f;
+	LocalExposureHighlightContrastScale_DEPRECATED = 1.0f;
+	LocalExposureShadowContrastScale_DEPRECATED = 1.0f;
+	LocalExposureHighlightContrast = 0.8f;
+	LocalExposureShadowContrast = 0.8f;
 	LocalExposureDetailStrength = 1.0f;
 	LocalExposureBlurredLuminanceBlend = 0.6f;
 	LocalExposureBlurredLuminanceKernelSizePercent = 50.0f;
@@ -901,6 +904,26 @@ void FPostProcessSettings::PostSerialize(const FArchive& Ar)
 			}
 		}
 
+		if (bOverride_LocalExposureContrastScale_DEPRECATED)
+		{
+			bOverride_LocalExposureHighlightContrastScale_DEPRECATED = bOverride_LocalExposureContrastScale_DEPRECATED;
+			bOverride_LocalExposureShadowContrastScale_DEPRECATED = bOverride_LocalExposureContrastScale_DEPRECATED;
+		}
+		if (LocalExposureContrastScale_DEPRECATED != 1.0f)
+		{
+			LocalExposureHighlightContrastScale_DEPRECATED = LocalExposureContrastScale_DEPRECATED;
+			LocalExposureShadowContrastScale_DEPRECATED = LocalExposureContrastScale_DEPRECATED;
+		}
+
+		if (Ar.CustomVer(FFortniteMainBranchObjectVersion::GUID) < FFortniteMainBranchObjectVersion::LocalExposureDefaultChangeFrom1)
+		{
+			bOverride_LocalExposureHighlightContrast = bOverride_LocalExposureHighlightContrastScale_DEPRECATED;
+			bOverride_LocalExposureShadowContrast = bOverride_LocalExposureShadowContrastScale_DEPRECATED;
+
+			LocalExposureHighlightContrast = LocalExposureHighlightContrastScale_DEPRECATED;
+			LocalExposureShadowContrast = LocalExposureShadowContrastScale_DEPRECATED;
+		}
+
 		if (Ar.CustomVer(FUE5ReleaseStreamObjectVersion::GUID) < FUE5ReleaseStreamObjectVersion::ReflectionMethodEnum)
 		{
 			if (bOverride_RayTracingGI && RayTracingGIType != ERayTracingGlobalIlluminationType::Disabled)
diff --git a/Engine/Source/Runtime/Engine/Private/SceneView.cpp b/Engine/Source/Runtime/Engine/Private/SceneView.cpp
--- a/Engine/Source/Runtime/Engine/Private/SceneView.cpp
+++ b/Engine/Source/Runtime/Engine/Private/SceneView.cpp
@@ -182,15 +182,15 @@ static TAutoConsoleVariable<int32> CVarDefaultAutoExposureExtendDefaultLuminance
 	TEXT(" 0: Legacy range (UE4 default)\n")
 	TEXT(" 1: Extended range (UE5 default)"));
 
-static TAutoConsoleVariable<float> CVarDefaultLocalExposureHighlightContrastScale(
+static TAutoConsoleVariable<float> CVarDefaultLocalExposureHighlightContrast(
 	TEXT("r.DefaultFeature.LocalExposure.HighlightContrastScale"),
 	1.0f,
-	TEXT("Engine default (project setting) for Local Exposure Highlight Contrast Scale (postprocess volume/camera/game setting still can override)\n"));
+	TEXT("Engine default (project setting) for Local Exposure Highlight Contrast (postprocess volume/camera/game setting still can override)\n"));
 
-static TAutoConsoleVariable<float> CVarDefaultLocalExposureShadowContrastScale(
+static TAutoConsoleVariable<float> CVarDefaultLocalExposureShadowContrast(
 	TEXT("r.DefaultFeature.LocalExposure.ShadowContrastScale"),
 	1.0f,
-	TEXT("Engine default (project setting) for Local Exposure Shadow Contrast Scale (postprocess volume/camera/game setting still can override)\n"));
+	TEXT("Engine default (project setting) for Local Exposure Shadow Contrast (postprocess volume/camera/game setting still can override)\n"));
 
 static TAutoConsoleVariable<int32> CVarDefaultMotionBlur(
 	TEXT("r.DefaultFeature.MotionBlur"),
@@ -1544,9 +1544,8 @@ void FSceneView::OverridePostProcessSettings(const FPostProcessSettings& Src, fl
 		LERP_PP(AutoExposureBias);
 		LERP_PP(HistogramLogMin);
 		LERP_PP(HistogramLogMax);
-		LERP_PP(LocalExposureContrastScale_DEPRECATED);
-		LERP_PP(LocalExposureHighlightContrastScale);
-		LERP_PP(LocalExposureShadowContrastScale);
+		LERP_PP(LocalExposureHighlightContrast);
+		LERP_PP(LocalExposureShadowContrast);
 		LERP_PP(LocalExposureDetailStrength);
 		LERP_PP(LocalExposureBlurredLuminanceBlend);
 		LERP_PP(LocalExposureBlurredLuminanceKernelSizePercent);
@@ -1844,15 +1843,15 @@ void FSceneView::StartFinalPostprocessSettings(FVector InViewLocation)
 		}
 
 		{
-			const float HighlightContrastScale = FMath::Clamp(CVarDefaultLocalExposureHighlightContrastScale.GetValueOnGameThread(), 0.0f, 1.0f);
+			const float HighlightContrast = FMath::Clamp(CVarDefaultLocalExposureHighlightContrast.GetValueOnGameThread(), 0.0f, 1.0f);
 
-			FinalPostProcessSettings.LocalExposureHighlightContrastScale = HighlightContrastScale;
+			FinalPostProcessSettings.LocalExposureHighlightContrast = HighlightContrast;
 		}
 
 		{
-			const float ShadowContrastScale = FMath::Clamp(CVarDefaultLocalExposureHighlightContrastScale.GetValueOnGameThread(), 0.0f, 1.0f);
+			const float ShadowContrast = FMath::Clamp(CVarDefaultLocalExposureHighlightContrast.GetValueOnGameThread(), 0.0f, 1.0f);
 
-			FinalPostProcessSettings.LocalExposureShadowContrastScale = ShadowContrastScale;
+			FinalPostProcessSettings.LocalExposureShadowContrast = ShadowContrast;
 		}
 
 		if (!CVarDefaultMotionBlur.GetValueOnGameThread())
@@ -1927,12 +1926,12 @@ void FSceneView::EndFinalPostprocessSettings(const FSceneViewInitOptions& ViewIn
 	{
 		static const auto LocalExposureCVar = IConsoleManager::Get().FindTConsoleVariableDataInt(TEXT("r.LocalExposure"));
 
-		int Value = LocalExposureCVar->GetValueOnGameThread();
+		const int LocalExposureCVarValue = LocalExposureCVar->GetValueOnGameThread();
 
-		if (Value <= 0)
+		if (LocalExposureCVarValue <= 0 || !Family->EngineShowFlags.LocalExposure)
 		{
-			FinalPostProcessSettings.LocalExposureHighlightContrastScale = 1.0f;
-			FinalPostProcessSettings.LocalExposureShadowContrastScale = 1.0f;
+			FinalPostProcessSettings.LocalExposureHighlightContrast = 1.0f;
+			FinalPostProcessSettings.LocalExposureShadowContrast = 1.0f;
 			FinalPostProcessSettings.LocalExposureDetailStrength = 1.0f;
 		}
 	}
@@ -1953,13 +1952,6 @@ void FSceneView::EndFinalPostprocessSettings(const FSceneViewInitOptions& ViewIn
 		FinalPostProcessSettings.BloomIntensity = 0.0f;
 	}
 
-	if (!Family->EngineShowFlags.LocalExposure)
-	{
-		FinalPostProcessSettings.LocalExposureHighlightContrastScale = 1.0f;
-		FinalPostProcessSettings.LocalExposureShadowContrastScale = 1.0f;
-		FinalPostProcessSettings.LocalExposureDetailStrength = 1.0f;
-	}
-
 	// scale down tone mapper shader permutation
 	{
 		int32 Quality = CVarTonemapperQuality.GetValueOnGameThread();
diff --git a/Engine/Source/Runtime/Renderer/Private/PostProcess/PostProcessLocalExposure.cpp b/Engine/Source/Runtime/Renderer/Private/PostProcess/PostProcessLocalExposure.cpp
--- a/Engine/Source/Runtime/Renderer/Private/PostProcess/PostProcessLocalExposure.cpp
+++ b/Engine/Source/Runtime/Renderer/Private/PostProcess/PostProcessLocalExposure.cpp
@@ -111,8 +111,8 @@ FLocalExposureParameters GetLocalExposureParameters(const FViewInfo& View, FIntP
 	const FVector2f LocalExposureBilateralGridUVScale = GetLocalExposureBilateralGridUVScale(ViewRectSize);
 
 	FLocalExposureParameters Parameters;
-	Parameters.HighlightContrastScale = Settings.LocalExposureHighlightContrastScale;
-	Parameters.ShadowContrastScale = Settings.LocalExposureShadowContrastScale;
+	Parameters.HighlightContrastScale = Settings.LocalExposureHighlightContrast;
+	Parameters.ShadowContrastScale = Settings.LocalExposureShadowContrast;
 	Parameters.DetailStrength = Settings.LocalExposureDetailStrength;
 	Parameters.BlurredLuminanceBlend = Settings.LocalExposureBlurredLuminanceBlend;
 	Parameters.MiddleGreyExposureCompensation = LocalExposureMiddleGreyExposureCompensation;
diff --git a/Engine/Source/Runtime/Renderer/Private/PostProcess/PostProcessing.cpp b/Engine/Source/Runtime/Renderer/Private/PostProcess/PostProcessing.cpp
--- a/Engine/Source/Runtime/Renderer/Private/PostProcess/PostProcessing.cpp
+++ b/Engine/Source/Runtime/Renderer/Private/PostProcess/PostProcessing.cpp
@@ -613,8 +613,8 @@ void AddPostProcessingPasses(
 
 		const bool bLocalExposureEnabled =
 			EngineShowFlags.VisualizeLocalExposure || 
-			!FMath::IsNearlyEqual(View.FinalPostProcessSettings.LocalExposureHighlightContrastScale, 1.0f) ||
-			!FMath::IsNearlyEqual(View.FinalPostProcessSettings.LocalExposureShadowContrastScale, 1.0f) ||
+			!FMath::IsNearlyEqual(View.FinalPostProcessSettings.LocalExposureHighlightContrast, 1.0f) ||
+			!FMath::IsNearlyEqual(View.FinalPostProcessSettings.LocalExposureShadowContrast, 1.0f) ||
 			!FMath::IsNearlyEqual(View.FinalPostProcessSettings.LocalExposureDetailStrength, 1.0f);
 
 		const bool bBloomEnabled = View.FinalPostProcessSettings.BloomIntensity > 0.0f && !bVisualizeTSR;
diff --git a/Engine/Source/Runtime/Renderer/Private/SceneRendering.cpp b/Engine/Source/Runtime/Renderer/Private/SceneRendering.cpp
--- a/Engine/Source/Runtime/Renderer/Private/SceneRendering.cpp
+++ b/Engine/Source/Runtime/Renderer/Private/SceneRendering.cpp
@@ -3542,8 +3542,8 @@ void FSceneRenderer::OnRenderFinish(FRDGBuilder& GraphBuilder, FRDGTextureRef Vi
 			bNaniteEnabledButDisabledInProject = bNaniteEnabledButDisabledInProject || (WouldRenderNanite(Scene, View, /*bCheckForAtomicSupport*/ false, /*bCheckForProjectSetting*/ false) && !WouldRenderNanite(Scene, View, /*bCheckForAtomicSupport*/ false, /*bCheckForProjectSetting*/ true));
 
 			if (IsPostProcessingEnabled(View)
-				&& (!FMath::IsNearlyEqual(View.FinalPostProcessSettings.LocalExposureHighlightContrastScale, 1.0f) 
-					|| !FMath::IsNearlyEqual(View.FinalPostProcessSettings.LocalExposureShadowContrastScale, 1.0f) 
+				&& (!FMath::IsNearlyEqual(View.FinalPostProcessSettings.LocalExposureHighlightContrast, 1.0f) 
+					|| !FMath::IsNearlyEqual(View.FinalPostProcessSettings.LocalExposureShadowContrast, 1.0f) 
 					|| !FMath::IsNearlyEqual(View.FinalPostProcessSettings.LocalExposureDetailStrength, 1.0f)))
 				bLocalExposureEnabledOnAnyView = true;
 		}
