Introduced material translation logging to MaterialTranslationLog.txt, containing the list of all translated materials since startup and their translation time.

diff --git a/Engine/Source/Runtime/Engine/Private/Materials/HLSLMaterialTranslator.cpp b/Engine/Source/Runtime/Engine/Private/Materials/HLSLMaterialTranslator.cpp
--- a/Engine/Source/Runtime/Engine/Private/Materials/HLSLMaterialTranslator.cpp
+++ b/Engine/Source/Runtime/Engine/Private/Materials/HLSLMaterialTranslator.cpp
@@ -47,12 +47,57 @@
 #include "VT/RuntimeVirtualTexture.h"
 #include <memory>
 #include <tuple>
+#include "ShaderCompiler.h"
+#include "Misc/FileHelper.h"
+#include "Misc/ScopeLock.h"
 
 #if WITH_EDITORONLY_DATA
 #include "Materials/MaterialExpressionSubstrate.h"
 #include "ShaderPlatformCachedIniValue.h"
 #endif
 
+
+/**
+ * Utility used to create a MaterialTranslationLog.txt file during cooks that contains the list of all translated materials
+ * and other info such as how long the translation took.
+ */
+struct FCsvLogFile
+{
+	FCriticalSection CriticalSection;
+	FString LogContent;
+
+	static FCsvLogFile& Get()
+	{
+		static FCsvLogFile Instance;
+		return Instance;
+	}
+
+	FCsvLogFile()
+	{
+		LogContent = TEXT("Date,MaterialName,TranslationTime\n");
+	}
+
+	void AddEntry(FStringView MaterialName, FDateTime DateTime, float TranslationTime)
+	{
+		FScopeLock Lock{ &CriticalSection };
+		FString DateTimeString = DateTime.ToString(TEXT("%Y-%m-%d %H:%M:%S"));
+		LogContent.Appendf(TEXT("%s,%s,%f\n"), *DateTimeString, MaterialName.GetData(), TranslationTime);
+	}
+
+	void Save()
+	{
+		uint32 MultiprocessId = 0;
+		FParse::Value(FCommandLine::Get(), TEXT("-MultiprocessId="), MultiprocessId);
+		FString FilePath = GShaderCompilingManager->GetAbsoluteShaderDebugInfoDirectory() / FString::Printf(TEXT("MaterialTranslationLog-%d.csv"), MultiprocessId);
+		UE_LOG(LogMaterial, Display, TEXT("Writing out MaterialTranslation log file '%s'"), *FilePath);
+		if (!FFileHelper::SaveStringToFile(LogContent, *FilePath))
+		{
+			UE_LOG(LogMaterial, Display, TEXT("Cannot open MaterialTranslation log file '%s'"), *FilePath);
+		}
+	}
+
+};
+
 #if ENABLE_COOK_STATS
 #include "ProfilingDebugging/ScopedTimers.h"
 namespace MaterialTranslatorCookStats
@@ -66,6 +111,7 @@ namespace MaterialTranslatorCookStats
 				TEXT("MaterialTranslateCalls"), MaterialTranslateCalls,
 				TEXT("MaterialTranslateTimeSec"), MaterialTranslateTimeSec
 			));
+			FCsvLogFile::Get().Save();
 		});
 }
 #endif
@@ -839,6 +885,7 @@ bool FHLSLMaterialTranslator::Translate()
 
 	COOK_STAT(MaterialTranslatorCookStats::MaterialTranslateCalls++);
 	COOK_STAT(FScopedDurationTimer DurationTimer(MaterialTranslatorCookStats::MaterialTranslateTimeSec));
+	COOK_STAT(FDateTime TranslationDateTime = FDateTime::Now();)
 
 	STAT(double HLSLTranslateTime = 0);
 	{
@@ -1927,6 +1974,12 @@ bool FHLSLMaterialTranslator::Translate()
 	ClearAllFunctionStacks();
 		
 	INC_FLOAT_STAT_BY(STAT_ShaderCompiling_HLSLTranslation,(float)HLSLTranslateTime);
+
+#if ENABLE_COOK_STATS
+	// Write out a CSV file MaterialTranslationLog.txt containing info about all material translations.
+	FCsvLogFile::Get().AddEntry(Material->GetMaterialInterface()->GetFullName(), TranslationDateTime, HLSLTranslateTime);
+#endif
+
 	return bSuccess;
 }
 
