Expose local gravity direction accessors to blueprints/other modules
#rnx

diff --git a/Engine/Source/Runtime/Engine/Classes/GameFramework/CharacterMovementComponent.h b/Engine/Source/Runtime/Engine/Classes/GameFramework/CharacterMovementComponent.h
--- a/Engine/Source/Runtime/Engine/Classes/GameFramework/CharacterMovementComponent.h
+++ b/Engine/Source/Runtime/Engine/Classes/GameFramework/CharacterMovementComponent.h
@@ -1616,13 +1616,16 @@ public:
 	 * from gravity relative space.
 	 * @param GravityDir		A non-zero vector representing the new gravity direction. The vector will be normalized.
 	 */
+	UFUNCTION(BlueprintCallable, Category="Pawn|Components|CharacterMovement")
 	ENGINE_API virtual void SetGravityDirection(const FVector& GravityDir);
 
 	/** Whether the gravity direction is different from UCharacterMovementComponent::DefaultGravityDirection. */
-	bool HasCustomGravity() const { return bHasCustomGravity; }
+	UFUNCTION(BlueprintPure, Category="Pawn|Components|CharacterMovement")
+	ENGINE_API bool HasCustomGravity() const { return bHasCustomGravity; }
 
 	/** Returns the current gravity direction. */
-	FVector GetGravityDirection() const { return GravityDirection; }
+	UFUNCTION(BlueprintPure, Category="Pawn|Components|CharacterMovement")
+	ENGINE_API FVector GetGravityDirection() const { return GravityDirection; }
 
 	/** Returns a quaternion transforming from world to gravity space. */
 	FQuat GetWorldToGravityTransform() const { return WorldToGravityTransform; }
