Add prioritization to the categories on UMaterials

The order is the following organized by most important and grouping similar areas together:

- Material
- Nanite
- Translucency
- Translucency Self Shadowing
- Refraction
- World Position Offset
- Post Process Material
- Mobile
- Forward Shading
- Physical Material
- Physical Material Mask
- Usage
- Lightmass
- Previewing
- Import Settings

diff --git a/Engine/Source/Runtime/Engine/Classes/Materials/Material.h b/Engine/Source/Runtime/Engine/Classes/Materials/Material.h
--- a/Engine/Source/Runtime/Engine/Classes/Materials/Material.h
+++ b/Engine/Source/Runtime/Engine/Classes/Materials/Material.h
@@ -411,7 +411,7 @@ public:
  *
  * Warning: Creating new materials directly increases shader compile times!  Consider creating a Material Instance off of an existing material instead.
  */
-UCLASS(hidecategories=Object, MinimalAPI, BlueprintType)
+UCLASS(hidecategories=Object, MinimalAPI, BlueprintType, PrioritizeCategories = ("Material", "Nanite", "Translucency", "TranslucencySelfShadowing", "Refraction", "WorldPositionOffset", "PostProcessMaterial",  "Mobile", "ForwardShading", "PhysicalMaterial", "PhysicalMaterialMask", "Usage", "Lightmass", "Previewing", "ImportSettings"))
 class UMaterial : public UMaterialInterface
 {
 	GENERATED_UCLASS_BODY()
