Speculative fix number 2 for cooker error fallout from ISM instance data refactor.

diff --git a/Engine/Source/Runtime/Engine/Private/InstancedStaticMesh.cpp b/Engine/Source/Runtime/Engine/Private/InstancedStaticMesh.cpp
--- a/Engine/Source/Runtime/Engine/Private/InstancedStaticMesh.cpp
+++ b/Engine/Source/Runtime/Engine/Private/InstancedStaticMesh.cpp
@@ -2884,7 +2884,7 @@ void UInstancedStaticMeshComponent::SerializeRenderData(FArchive& Ar)
 		uint64 RenderDataSizeBytes = 0;
 		Ar << RenderDataSizeBytes;
 
-		bool bSaveRenderData = NeedRenderDataForTargetPlatform(Ar.CookingTarget());
+		bool bSaveRenderData = NeedRenderDataForTargetPlatform(Ar.CookingTarget()) && InstancingRandomSeed != 0;
 		if (bSaveRenderData && !HasAnyFlags(RF_ClassDefaultObject | RF_ArchetypeObject))
 		{
 			uint64 RenderDataPos = Ar.Tell();
