Added missing include

#rnx

diff --git a/Engine/Source/Runtime/Engine/Public/CookedMetaData.h b/Engine/Source/Runtime/Engine/Public/CookedMetaData.h
--- a/Engine/Source/Runtime/Engine/Public/CookedMetaData.h
+++ b/Engine/Source/Runtime/Engine/Public/CookedMetaData.h
@@ -4,6 +4,7 @@
 
 #include "UObject/Object.h"
 #include "UObject/UObjectGlobals.h"
+#include "Templates/SubclassOf.h"
 #include "CookedMetaData.generated.h"
 
 /**
