Fix navmesh issues with rivers and river beds.
Temporary fix for regression introduced by CL 27753412. The change is to return to previous behavior for now.

diff --git a/Engine/Source/Runtime/Engine/Classes/Components/SplineMeshComponent.h b/Engine/Source/Runtime/Engine/Classes/Components/SplineMeshComponent.h
--- a/Engine/Source/Runtime/Engine/Classes/Components/SplineMeshComponent.h
+++ b/Engine/Source/Runtime/Engine/Classes/Components/SplineMeshComponent.h
@@ -227,6 +227,10 @@ public:
 	ENGINE_API virtual bool GetTriMeshSizeEstimates(struct FTriMeshCollisionDataEstimates& OutTriMeshEstimates, bool bInUseAllTriData) const override;
 	//~ End Interface_CollisionDataProvider Interface
 
+	//~ Begin INavRelevantInterface
+	ENGINE_API virtual FBox GetNavigationBounds() const override;
+	//~ End  INavRelevantInterface
+
 	/** Generates FSplineMeshShaderParams for the current state of the component */
 	ENGINE_API FSplineMeshShaderParams CalculateShaderParams() const;
 
diff --git a/Engine/Source/Runtime/Engine/Private/Components/SplineMeshComponent.cpp b/Engine/Source/Runtime/Engine/Private/Components/SplineMeshComponent.cpp
--- a/Engine/Source/Runtime/Engine/Private/Components/SplineMeshComponent.cpp
+++ b/Engine/Source/Runtime/Engine/Private/Components/SplineMeshComponent.cpp
@@ -1121,6 +1121,14 @@ bool USplineMeshComponent::GetTriMeshSizeEstimates(struct FTriMeshCollisionDataE
 	return false;
 }
 
+FBox USplineMeshComponent::GetNavigationBounds() const
+{
+	// @todo Revisit how to collect navigation bounds, but for now behave like UPrimitiveComponent::GetNavigationBounds(), same as before the recent addition of UStaticMeshComponent::GetNavigationBounds().
+	
+	// Return invalid box when retrieving NavigationBounds before they are being computed at component registration
+	return bRegistered ? Bounds.GetBox() : FBox(ForceInit);
+}
+
 void USplineMeshComponent::GetMeshId(FString& OutMeshId)
 {
 	// First get the base mesh id from the static mesh
@@ -1200,7 +1208,7 @@ bool USplineMeshComponent::DoCustomNavigableGeometryExport(FNavigableGeometryExp
 
 	if (GetStaticMesh() != nullptr && GetStaticMesh()->GetNavCollision() != nullptr)
 	{
-		UNavCollisionBase* NavCollision = GetStaticMesh()->GetNavCollision();
+		const UNavCollisionBase* NavCollision = GetStaticMesh()->GetNavCollision();
 
 		if (ensure(!NavCollision->IsDynamicObstacle()))
 		{
