Local Exposure Contrast Curves
- highlight and shadow contrast can be set using user-defined curves that are evaluated based on the scene EV100.

diff --git a/Engine/Source/Runtime/Engine/Classes/Engine/Scene.h b/Engine/Source/Runtime/Engine/Classes/Engine/Scene.h
--- a/Engine/Source/Runtime/Engine/Classes/Engine/Scene.h
+++ b/Engine/Source/Runtime/Engine/Classes/Engine/Scene.h
@@ -913,6 +913,12 @@ struct FPostProcessSettings
 	UPROPERTY(EditAnywhere, BlueprintReadWrite, Category=Overrides, meta=(PinHiddenByDefault, InlineEditConditionToggle))
 	uint8 bOverride_LocalExposureShadowContrast:1;
 
+	UPROPERTY(EditAnywhere, BlueprintReadWrite, Category=Overrides, meta=(PinHiddenByDefault, InlineEditConditionToggle))
+	uint8 bOverride_LocalExposureHighlightContrastCurve:1;
+
+	UPROPERTY(EditAnywhere, BlueprintReadWrite, Category=Overrides, meta=(PinHiddenByDefault, InlineEditConditionToggle))
+	uint8 bOverride_LocalExposureShadowContrastCurve:1;
+
 	UPROPERTY(EditAnywhere, BlueprintReadWrite, Category=Overrides, meta=(PinHiddenByDefault, InlineEditConditionToggle))
 	uint8 bOverride_LocalExposureDetailStrength:1;
 
@@ -1856,6 +1862,20 @@ struct FPostProcessSettings
 	UPROPERTY(interp, BlueprintReadWrite, Category="Lens|Local Exposure", meta=(UIMin = "0.0", UIMax = "1.0", editcondition = "bOverride_LocalExposureShadowContrast", DisplayName = "Shadow Contrast"))
 	float LocalExposureShadowContrast;
 
+	/**
+	 * Local Exposure Highlight Contrast based on the scene EV100.
+	 * Used to calibrate Local Exposure differently depending on the average scene luminance.
+	 */
+	UPROPERTY(EditAnywhere, BlueprintReadWrite, Category = "Lens|Local Exposure", meta = (editcondition = "bOverride_LocalExposureHighlightContrastCurve", DisplayName = "Highlight Contrast Curve"))
+	TObjectPtr<class UCurveFloat> LocalExposureHighlightContrastCurve = nullptr;
+
+	/**
+	 * Local Exposure Shadow Contrast based on the scene EV100.
+	 * Used to calibrate Local Exposure differently depending on the average scene luminance.
+	 */
+	UPROPERTY(EditAnywhere, BlueprintReadWrite, Category = "Lens|Local Exposure", meta = (editcondition = "bOverride_LocalExposureShadowContrastCurve", DisplayName = "Shadow Contrast Curve"))
+	TObjectPtr<class UCurveFloat> LocalExposureShadowContrastCurve = nullptr;
+
 	/**
 	 * Local Exposure decomposes luminance of the frame into a base layer and a detail layer.
 	 * Value different than 1 will enable local exposure.
diff --git a/Engine/Source/Runtime/Engine/Private/SceneView.cpp b/Engine/Source/Runtime/Engine/Private/SceneView.cpp
--- a/Engine/Source/Runtime/Engine/Private/SceneView.cpp
+++ b/Engine/Source/Runtime/Engine/Private/SceneView.cpp
@@ -1735,6 +1735,18 @@ void FSceneView::OverridePostProcessSettings(const FPostProcessSettings& Src, fl
 			Dest.AutoExposureMeterMask = Src.AutoExposureMeterMask;
 		}
 
+		// Curve assets cannot be blended.
+		IF_PP(LocalExposureHighlightContrastCurve)
+		{
+			Dest.LocalExposureHighlightContrastCurve = Src.LocalExposureHighlightContrastCurve;
+		}
+
+		// Curve assets cannot be blended.
+		IF_PP(LocalExposureShadowContrastCurve)
+		{
+			Dest.LocalExposureShadowContrastCurve = Src.LocalExposureShadowContrastCurve;
+		}
+
 		// actual texture cannot be blended but the intensity can be blended
 		IF_PP(LensFlareBokehShape)
 		{
@@ -1932,6 +1944,8 @@ void FSceneView::EndFinalPostprocessSettings(const FSceneViewInitOptions& ViewIn
 		{
 			FinalPostProcessSettings.LocalExposureHighlightContrast = 1.0f;
 			FinalPostProcessSettings.LocalExposureShadowContrast = 1.0f;
+			FinalPostProcessSettings.LocalExposureHighlightContrastCurve = nullptr;
+			FinalPostProcessSettings.LocalExposureShadowContrastCurve = nullptr;
 			FinalPostProcessSettings.LocalExposureDetailStrength = 1.0f;
 		}
 	}
diff --git a/Engine/Source/Runtime/Renderer/Private/PostProcess/PostProcessLocalExposure.cpp b/Engine/Source/Runtime/Renderer/Private/PostProcess/PostProcessLocalExposure.cpp
--- a/Engine/Source/Runtime/Renderer/Private/PostProcess/PostProcessLocalExposure.cpp
+++ b/Engine/Source/Runtime/Renderer/Private/PostProcess/PostProcessLocalExposure.cpp
@@ -7,6 +7,7 @@
 #include "PostProcess/PostProcessLocalExposure.h"
 #include "PostProcess/PostProcessEyeAdaptation.h"
 #include "PostProcess/PostProcessWeightedSampleSum.h"
+#include "Curves/CurveFloat.h"
 #include "SceneRendering.h"
 #include "ShaderCompilerCore.h"
 #include "DataDrivenShaderPlatformInfo.h"
@@ -110,9 +111,30 @@ FLocalExposureParameters GetLocalExposureParameters(const FViewInfo& View, FIntP
 
 	const FVector2f LocalExposureBilateralGridUVScale = GetLocalExposureBilateralGridUVScale(ViewRectSize);
 
+	float HighlightContrast = Settings.LocalExposureHighlightContrast;
+	float ShadowContrast = Settings.LocalExposureShadowContrast;
+
+	const float AverageSceneLuminance = View.GetLastAverageSceneLuminance();
+	if (AverageSceneLuminance > 0)
+	{
+		const float LuminanceMax = LuminanceMaxFromLensAttenuation();
+		// We need the Log2(0.18) to convert from average luminance to saturation luminance
+		const float LuminanceEV100 = LuminanceToEV100(LuminanceMax, AverageSceneLuminance) + FMath::Log2(1.0f / 0.18f);
+
+		if (Settings.LocalExposureHighlightContrastCurve)
+		{
+			HighlightContrast *= Settings.LocalExposureHighlightContrastCurve->GetFloatValue(LuminanceEV100);
+		}
+
+		if (Settings.LocalExposureShadowContrastCurve)
+		{
+			ShadowContrast *= Settings.LocalExposureShadowContrastCurve->GetFloatValue(LuminanceEV100);
+		}
+	}
+
 	FLocalExposureParameters Parameters;
-	Parameters.HighlightContrastScale = Settings.LocalExposureHighlightContrast;
-	Parameters.ShadowContrastScale = Settings.LocalExposureShadowContrast;
+	Parameters.HighlightContrastScale = HighlightContrast;
+	Parameters.ShadowContrastScale = ShadowContrast;
 	Parameters.DetailStrength = Settings.LocalExposureDetailStrength;
 	Parameters.BlurredLuminanceBlend = Settings.LocalExposureBlurredLuminanceBlend;
 	Parameters.MiddleGreyExposureCompensation = LocalExposureMiddleGreyExposureCompensation;
diff --git a/Engine/Source/Runtime/Renderer/Private/PostProcess/PostProcessing.cpp b/Engine/Source/Runtime/Renderer/Private/PostProcess/PostProcessing.cpp
--- a/Engine/Source/Runtime/Renderer/Private/PostProcess/PostProcessing.cpp
+++ b/Engine/Source/Runtime/Renderer/Private/PostProcess/PostProcessing.cpp
@@ -612,9 +612,11 @@ void AddPostProcessingPasses(
 			View.FinalPostProcessSettings.AutoExposureMinBrightness < View.FinalPostProcessSettings.AutoExposureMaxBrightness);
 
 		const bool bLocalExposureEnabled =
-			EngineShowFlags.VisualizeLocalExposure || 
+			EngineShowFlags.VisualizeLocalExposure ||
 			!FMath::IsNearlyEqual(View.FinalPostProcessSettings.LocalExposureHighlightContrast, 1.0f) ||
 			!FMath::IsNearlyEqual(View.FinalPostProcessSettings.LocalExposureShadowContrast, 1.0f) ||
+			View.FinalPostProcessSettings.LocalExposureHighlightContrastCurve ||
+			View.FinalPostProcessSettings.LocalExposureShadowContrastCurve ||
 			!FMath::IsNearlyEqual(View.FinalPostProcessSettings.LocalExposureDetailStrength, 1.0f);
 
 		const bool bBloomEnabled = View.FinalPostProcessSettings.BloomIntensity > 0.0f && !bVisualizeTSR;
diff --git a/Engine/Source/Runtime/Renderer/Private/SceneRendering.cpp b/Engine/Source/Runtime/Renderer/Private/SceneRendering.cpp
--- a/Engine/Source/Runtime/Renderer/Private/SceneRendering.cpp
+++ b/Engine/Source/Runtime/Renderer/Private/SceneRendering.cpp
@@ -3542,8 +3542,10 @@ void FSceneRenderer::OnRenderFinish(FRDGBuilder& GraphBuilder, FRDGTextureRef Vi
 			bNaniteEnabledButDisabledInProject = bNaniteEnabledButDisabledInProject || (WouldRenderNanite(Scene, View, /*bCheckForAtomicSupport*/ false, /*bCheckForProjectSetting*/ false) && !WouldRenderNanite(Scene, View, /*bCheckForAtomicSupport*/ false, /*bCheckForProjectSetting*/ true));
 
 			if (IsPostProcessingEnabled(View)
-				&& (!FMath::IsNearlyEqual(View.FinalPostProcessSettings.LocalExposureHighlightContrast, 1.0f) 
-					|| !FMath::IsNearlyEqual(View.FinalPostProcessSettings.LocalExposureShadowContrast, 1.0f) 
+				&& (!FMath::IsNearlyEqual(View.FinalPostProcessSettings.LocalExposureHighlightContrast, 1.0f)
+					|| !FMath::IsNearlyEqual(View.FinalPostProcessSettings.LocalExposureShadowContrast, 1.0f)
+					|| View.FinalPostProcessSettings.LocalExposureShadowContrastCurve
+					|| View.FinalPostProcessSettings.LocalExposureHighlightContrastCurve
 					|| !FMath::IsNearlyEqual(View.FinalPostProcessSettings.LocalExposureDetailStrength, 1.0f)))
 				bLocalExposureEnabledOnAnyView = true;
 		}
