Mark instances with zero scale as hidden in GPUScene flags.
This removes the instances early during the culling pass.

diff --git a/Engine/Source/Runtime/Engine/Private/InstanceUniformShaderParameters.cpp b/Engine/Source/Runtime/Engine/Private/InstanceUniformShaderParameters.cpp
--- a/Engine/Source/Runtime/Engine/Private/InstanceUniformShaderParameters.cpp
+++ b/Engine/Source/Runtime/Engine/Private/InstanceUniformShaderParameters.cpp
@@ -65,7 +65,8 @@ void FInstanceSceneShaderData::BuildInternal
 {
 	// Note: layout must match GetInstanceData in SceneData.ush and InitializeInstanceSceneData in GPUSceneWriter.ush
 
-	if (LocalToWorld.RotDeterminant() < 0.0f)
+	const float RotDeterminant = LocalToWorld.RotDeterminant();
+	if (RotDeterminant < 0.0f)
 	{
 		InstanceFlags |= INSTANCE_SCENE_DATA_FLAG_DETERMINANT_SIGN;
 	}
@@ -74,7 +75,8 @@ void FInstanceSceneShaderData::BuildInternal
 		InstanceFlags &= ~INSTANCE_SCENE_DATA_FLAG_DETERMINANT_SIGN;
 	}
 
-	if (!bIsVisible)
+	// Mark zero scaled instances as hidden.
+	if (!bIsVisible || RotDeterminant == 0)
 	{
 		InstanceFlags |= INSTANCE_SCENE_DATA_FLAG_HIDDEN;
 	}
