Fix up dev mode query for graph nodes to check more specifically for cook context, as that's what the cooker option was intended for.

(RE: issue 373857)

diff --git a/Engine/Source/Runtime/Engine/Private/EdGraph/EdGraphNode.cpp b/Engine/Source/Runtime/Engine/Private/EdGraph/EdGraphNode.cpp
--- a/Engine/Source/Runtime/Engine/Private/EdGraph/EdGraphNode.cpp
+++ b/Engine/Source/Runtime/Engine/Private/EdGraph/EdGraphNode.cpp
@@ -992,7 +992,7 @@ bool UEdGraphNode::IsInDevelopmentMode() const
 {
 #if WITH_EDITOR
 	// By default, development mode is implied when running in the editor and not cooking via commandlet, unless enabled in the project settings.
-	return !IsRunningCommandlet() || GetDefault<UCookerSettings>()->bCompileBlueprintsInDevelopmentMode;
+	return !IsRunningCookCommandlet() || GetDefault<UCookerSettings>()->bCompileBlueprintsInDevelopmentMode;
 #else
 	return false;
 #endif
