Iris - Dormancy notification on client to mimic ActorChannel behavior. This is to have initially dormant net startup actors that have been replicated properly marked as awake so thath they're added to the DemoNetDriver as replicated actors.

diff --git a/Engine/Source/Runtime/Engine/Private/Net/Iris/ReplicationSystem/ActorReplicationBridge.cpp b/Engine/Source/Runtime/Engine/Private/Net/Iris/ReplicationSystem/ActorReplicationBridge.cpp
--- a/Engine/Source/Runtime/Engine/Private/Net/Iris/ReplicationSystem/ActorReplicationBridge.cpp
+++ b/Engine/Source/Runtime/Engine/Private/Net/Iris/ReplicationSystem/ActorReplicationBridge.cpp
@@ -731,11 +731,11 @@ bool UActorReplicationBridge::OnInstantiatedFromRemote(UObject* Instance, const
 	if (BridgeHeader->bIsActor)
 	{
 		const FActorCreationHeader* Header = static_cast<const FActorCreationHeader*>(InHeader);
-
-		if (Header->bIsDynamic)
+		
+		AActor* Actor = CastChecked<AActor>(Instance);
+		if (Actor)
 		{
 			// OnActorChannelOpen
-			AActor* Actor = CastChecked<AActor>(Instance);
 			{
 				UNetConnection* Connection = NetDriver->GetConnectionById(ConnectionId);
 				FInBunch Bunch(Connection, const_cast<uint8*>(Header->CustomCreationData.GetData()), Header->CustomCreationDataBitCount);
@@ -748,6 +748,9 @@ bool UActorReplicationBridge::OnInstantiatedFromRemote(UObject* Instance, const
 					return false;
 				}
 			}
+
+			// Wake up from dormancy. This is important for client replays.
+			WakeUpObjectInstantiatedFromRemote(Actor);
 		}
 	}
 	
@@ -1118,6 +1121,34 @@ float UActorReplicationBridge::GetPollFrequencyOfRootObject(const UObject* Repli
 	return PollFrequency;
 }
 
+void UActorReplicationBridge::WakeUpObjectInstantiatedFromRemote(AActor* Actor) const
+{
+	// If the actor is already awake or can't be woken up then return immediately.
+	if (Actor->NetDormancy <= DORM_Awake)
+	{
+		return;
+	}
+
+	ENetDormancy OldDormancy = Actor->NetDormancy;
+	Actor->NetDormancy = DORM_Awake;
+
+	if (!NetDriver)
+	{
+		return;
+	}
+
+	if (FWorldContext* WorldContext = GEngine->GetWorldContextFromWorld(NetDriver->GetWorld()))
+	{
+		for (FNamedNetDriver& Driver : WorldContext->ActiveNetDrivers)
+		{
+			if (Driver.NetDriver != nullptr && Driver.NetDriver != NetDriver && Driver.NetDriver->ShouldReplicateActor(Actor))
+			{
+				Driver.NetDriver->NotifyActorClientDormancyChanged(Actor, OldDormancy);
+			}
+		}
+	}
+}
+
 #else //!UE_WITH_IRIS
 
 UActorReplicationBridge::UActorReplicationBridge() = default;
diff --git a/Engine/Source/Runtime/Engine/Public/Net/Iris/ReplicationSystem/ActorReplicationBridge.h b/Engine/Source/Runtime/Engine/Public/Net/Iris/ReplicationSystem/ActorReplicationBridge.h
--- a/Engine/Source/Runtime/Engine/Public/Net/Iris/ReplicationSystem/ActorReplicationBridge.h
+++ b/Engine/Source/Runtime/Engine/Public/Net/Iris/ReplicationSystem/ActorReplicationBridge.h
@@ -102,6 +102,8 @@ private:
 
 	void OnMaxTickRateChanged(UNetDriver* InNetDriver, int32 NewMaxTickRate, int32 OldMaxTickRate);
 
+	void WakeUpObjectInstantiatedFromRemote(AActor* Actor) const;
+
 private:
 	UNetDriver* NetDriver;
 
