World Partition: missing const ref& in URuntimePartition::IsValidPartitionTokens.

diff --git a/Engine/Source/Runtime/Engine/Private/WorldPartition/RuntimeHashSet/RuntimePartitionLHGrid.cpp b/Engine/Source/Runtime/Engine/Private/WorldPartition/RuntimeHashSet/RuntimePartitionLHGrid.cpp
--- a/Engine/Source/Runtime/Engine/Private/WorldPartition/RuntimeHashSet/RuntimePartitionLHGrid.cpp
+++ b/Engine/Source/Runtime/Engine/Private/WorldPartition/RuntimeHashSet/RuntimePartitionLHGrid.cpp
@@ -74,7 +74,7 @@ void URuntimePartitionLHGrid::SetDefaultValues()
 	CellSize = LoadingRange / 2;
 }
 
-bool URuntimePartitionLHGrid::IsValidPartitionTokens(const TArray<FName> InPartitionTokens) const
+bool URuntimePartitionLHGrid::IsValidPartitionTokens(const TArray<FName>& InPartitionTokens) const
 {
 	return InPartitionTokens.Num() == 1;
 }
diff --git a/Engine/Source/Runtime/Engine/Private/WorldPartition/RuntimeHashSet/RuntimePartitionLevelStreaming.cpp b/Engine/Source/Runtime/Engine/Private/WorldPartition/RuntimeHashSet/RuntimePartitionLevelStreaming.cpp
--- a/Engine/Source/Runtime/Engine/Private/WorldPartition/RuntimeHashSet/RuntimePartitionLevelStreaming.cpp
+++ b/Engine/Source/Runtime/Engine/Private/WorldPartition/RuntimeHashSet/RuntimePartitionLevelStreaming.cpp
@@ -5,7 +5,7 @@
 #include "WorldPartition/WorldPartitionStreamingGenerationContext.h"
 
 #if WITH_EDITOR
-bool URuntimePartitionLevelStreaming::IsValidPartitionTokens(const TArray<FName> InPartitionTokens) const
+bool URuntimePartitionLevelStreaming::IsValidPartitionTokens(const TArray<FName>& InPartitionTokens) const
 {
 	return InPartitionTokens.Num() && (InPartitionTokens.Num() <= 2);
 }
diff --git a/Engine/Source/Runtime/Engine/Public/WorldPartition/RuntimeHashSet/RuntimePartition.h b/Engine/Source/Runtime/Engine/Public/WorldPartition/RuntimeHashSet/RuntimePartition.h
--- a/Engine/Source/Runtime/Engine/Public/WorldPartition/RuntimeHashSet/RuntimePartition.h
+++ b/Engine/Source/Runtime/Engine/Public/WorldPartition/RuntimeHashSet/RuntimePartition.h
@@ -63,7 +63,7 @@ public:
 
 	virtual void SetDefaultValues();
 	virtual bool SupportsHLODs() const PURE_VIRTUAL(URuntimePartition::SupportsHLODs, return true;);
-	virtual bool IsValidPartitionTokens(const TArray<FName> InPartitionTokens) const PURE_VIRTUAL(URuntimePartition::IsValidPartitionTokens, return false;);
+	virtual bool IsValidPartitionTokens(const TArray<FName>& InPartitionTokens) const PURE_VIRTUAL(URuntimePartition::IsValidPartitionTokens, return false;);
 	virtual bool GenerateStreaming(const FGenerateStreamingParams& InParams, FGenerateStreamingResult& OutResult) PURE_VIRTUAL(URuntimePartition::GenerateStreaming, return false;);
 #endif
 
diff --git a/Engine/Source/Runtime/Engine/Public/WorldPartition/RuntimeHashSet/RuntimePartitionLHGrid.h b/Engine/Source/Runtime/Engine/Public/WorldPartition/RuntimeHashSet/RuntimePartitionLHGrid.h
--- a/Engine/Source/Runtime/Engine/Public/WorldPartition/RuntimeHashSet/RuntimePartitionLHGrid.h
+++ b/Engine/Source/Runtime/Engine/Public/WorldPartition/RuntimeHashSet/RuntimePartitionLHGrid.h
@@ -19,7 +19,7 @@ public:
 	//~ Begin URuntimePartition interface
 	virtual bool SupportsHLODs() const override { return true; }
 	virtual void SetDefaultValues() override;
-	virtual bool IsValidPartitionTokens(const TArray<FName> InPartitionTokens) const override;
+	virtual bool IsValidPartitionTokens(const TArray<FName>& InPartitionTokens) const override;
 	virtual bool GenerateStreaming(const FGenerateStreamingParams& InParams, FGenerateStreamingResult& OutResult) override;
 	//~ End URuntimePartition interface
 #endif
diff --git a/Engine/Source/Runtime/Engine/Public/WorldPartition/RuntimeHashSet/RuntimePartitionLevelStreaming.h b/Engine/Source/Runtime/Engine/Public/WorldPartition/RuntimeHashSet/RuntimePartitionLevelStreaming.h
--- a/Engine/Source/Runtime/Engine/Public/WorldPartition/RuntimeHashSet/RuntimePartitionLevelStreaming.h
+++ b/Engine/Source/Runtime/Engine/Public/WorldPartition/RuntimeHashSet/RuntimePartitionLevelStreaming.h
@@ -14,7 +14,7 @@ public:
 #if WITH_EDITOR
 	//~ Begin URuntimePartition interface
 	virtual bool SupportsHLODs() const override { return true; }
-	virtual bool IsValidPartitionTokens(const TArray<FName> InPartitionTokens) const override;
+	virtual bool IsValidPartitionTokens(const TArray<FName>& InPartitionTokens) const override;
 	virtual bool GenerateStreaming(const FGenerateStreamingParams& InParams, FGenerateStreamingResult& OutResult) override;
 	//~ End URuntimePartition interface
 #endif
diff --git a/Engine/Source/Runtime/Engine/Public/WorldPartition/RuntimeHashSet/RuntimePartitionPersistent.h b/Engine/Source/Runtime/Engine/Public/WorldPartition/RuntimeHashSet/RuntimePartitionPersistent.h
--- a/Engine/Source/Runtime/Engine/Public/WorldPartition/RuntimeHashSet/RuntimePartitionPersistent.h
+++ b/Engine/Source/Runtime/Engine/Public/WorldPartition/RuntimeHashSet/RuntimePartitionPersistent.h
@@ -14,7 +14,7 @@ public:
 #if WITH_EDITOR
 	//~ Begin URuntimePartition interface
 	virtual bool SupportsHLODs() const override { return false; }
-	virtual bool IsValidPartitionTokens(const TArray<FName> InPartitionTokens) const override { return true; }
+	virtual bool IsValidPartitionTokens(const TArray<FName>& InPartitionTokens) const override { return true; }
 	virtual bool GenerateStreaming(const FGenerateStreamingParams& InParams, FGenerateStreamingResult& OutResult) override;
 	//~ End URuntimePartition interface
 #endif
