World Partition HLODs: FHLODCreationContext now takes a functor to get the runtime grid of generated actors based on the HLOD layer.

diff --git a/Engine/Plugins/Editor/WorldPartitionHLODUtilities/Source/Private/WorldPartition/HLOD/Utilities/WorldPartitionHLODUtilities.cpp b/Engine/Plugins/Editor/WorldPartitionHLODUtilities/Source/Private/WorldPartition/HLOD/Utilities/WorldPartitionHLODUtilities.cpp
--- a/Engine/Plugins/Editor/WorldPartitionHLODUtilities/Source/Private/WorldPartition/HLOD/Utilities/WorldPartitionHLODUtilities.cpp
+++ b/Engine/Plugins/Editor/WorldPartitionHLODUtilities/Source/Private/WorldPartition/HLOD/Utilities/WorldPartitionHLODUtilities.cpp
@@ -216,7 +216,7 @@ TArray<AWorldPartitionHLOD*> FWorldPartitionHLODUtilities::CreateHLODActors(FHLO
 		}
 
 		// Runtime grid
-		const FName RuntimeGrid = InCreationParams.RuntimeGrid.IsNone() ? HLODLayer->GetRuntimeGrid(InCreationParams.HLODLevel) : InCreationParams.RuntimeGrid;
+		const FName RuntimeGrid = InCreationParams.GetRuntimeGrid(HLODLayer);
 		if (HLODActor->GetRuntimeGrid() != RuntimeGrid)
 		{
 			HLODActor->SetRuntimeGrid(RuntimeGrid);
diff --git a/Engine/Source/Runtime/Engine/Private/WorldPartition/RuntimeHashSet/WorldPartitionRuntimeHashSetHLODGeneration.cpp b/Engine/Source/Runtime/Engine/Private/WorldPartition/RuntimeHashSet/WorldPartitionRuntimeHashSetHLODGeneration.cpp
--- a/Engine/Source/Runtime/Engine/Private/WorldPartition/RuntimeHashSet/WorldPartitionRuntimeHashSetHLODGeneration.cpp
+++ b/Engine/Source/Runtime/Engine/Private/WorldPartition/RuntimeHashSet/WorldPartitionRuntimeHashSetHLODGeneration.cpp
@@ -232,7 +232,7 @@ bool UWorldPartitionRuntimeHashSet::SetupHLODActors(const IStreamingGenerationCo
 					HLODCreationParams.CellName = CellUniqueId.Name;
 					HLODCreationParams.CellGuid = CellUniqueId.Guid;			
 					HLODCreationParams.CellBounds = CellDescInstance.Bounds;
-					HLODCreationParams.RuntimeGrid = *FString::Printf(TEXT("%s:%s"), *FString::JoinBy(MainPartitionTokens, TEXT("."), [](const FName Token) { return Token.ToString(); }), *HLODLayer->GetName());
+					HLODCreationParams.GetRuntimeGrid = [&MainPartitionTokens](const UHLODLayer* HLODLayer) { return FName(*FString::Printf(TEXT("%s:%s"), *FString::JoinBy(MainPartitionTokens, TEXT("."), [](const FName Token) { return Token.ToString(); }), *HLODLayer->GetName())); };
 					HLODCreationParams.HLODLevel = HLODLevel;
 					HLODCreationParams.MinVisibleDistance = RuntimePartition->LoadingRange;
 					HLODCreationParams.ContentBundleGuid = CellDescInstance.ContentBundleID;
@@ -337,4 +337,4 @@ bool UWorldPartitionRuntimeHashSet::SetupHLODActors(const IStreamingGenerationCo
 
 	return false;
 }
-#endif
\ No newline at end of file
+#endif
diff --git a/Engine/Source/Runtime/Engine/Private/WorldPartition/RuntimeSpatialHash/RuntimeSpatialHashHLOD.cpp b/Engine/Source/Runtime/Engine/Private/WorldPartition/RuntimeSpatialHash/RuntimeSpatialHashHLOD.cpp
--- a/Engine/Source/Runtime/Engine/Private/WorldPartition/RuntimeSpatialHash/RuntimeSpatialHashHLOD.cpp
+++ b/Engine/Source/Runtime/Engine/Private/WorldPartition/RuntimeSpatialHash/RuntimeSpatialHashHLOD.cpp
@@ -203,6 +203,7 @@ static TArray<FGuid> GenerateHLODActorsForGrid(UWorldPartition* WorldPartition,
 				CreationParams.CellGuid = CellGuid;
 				CreationParams.CellName = CellName;
 				CreationParams.CellBounds = CellBounds;
+				CreationParams.GetRuntimeGrid = [HLODLevel](const UHLODLayer* HLODLayer) { return HLODLayer->GetRuntimeGrid(HLODLevel); };
 				CreationParams.HLODLevel = HLODLevel;
 				CreationParams.MinVisibleDistance = RuntimeGrid.LoadingRange;
 				CreationParams.ContentBundleGuid = GridCellDataChunk.GetContentBundleID();
diff --git a/Engine/Source/Runtime/Engine/Public/WorldPartition/HLOD/IWorldPartitionHLODUtilities.h b/Engine/Source/Runtime/Engine/Public/WorldPartition/HLOD/IWorldPartitionHLODUtilities.h
--- a/Engine/Source/Runtime/Engine/Public/WorldPartition/HLOD/IWorldPartitionHLODUtilities.h
+++ b/Engine/Source/Runtime/Engine/Public/WorldPartition/HLOD/IWorldPartitionHLODUtilities.h
@@ -30,7 +30,7 @@ struct FHLODCreationParams
 	FGuid CellGuid;
 	FString CellName;
 	FBox CellBounds;
-	FName RuntimeGrid;
+	TUniqueFunction<FName(const UHLODLayer*)> GetRuntimeGrid;
 	uint32 HLODLevel;
 	FGuid ContentBundleGuid;
 	TArray<const UDataLayerInstance*> DataLayerInstances;
