Fix Incremental UnrealEditor Mac/Linux: declaration shadows a structured binding.

diff --git a/Engine/Source/Runtime/Engine/Private/WorldPartition/RuntimeHashSet/WorldPartitionRuntimeHashSetHLODGeneration.cpp b/Engine/Source/Runtime/Engine/Private/WorldPartition/RuntimeHashSet/WorldPartitionRuntimeHashSetHLODGeneration.cpp
--- a/Engine/Source/Runtime/Engine/Private/WorldPartition/RuntimeHashSet/WorldPartitionRuntimeHashSetHLODGeneration.cpp
+++ b/Engine/Source/Runtime/Engine/Private/WorldPartition/RuntimeHashSet/WorldPartitionRuntimeHashSetHLODGeneration.cpp
@@ -232,7 +232,7 @@ bool UWorldPartitionRuntimeHashSet::SetupHLODActors(const IStreamingGenerationCo
 					HLODCreationParams.CellName = CellUniqueId.Name;
 					HLODCreationParams.CellGuid = CellUniqueId.Guid;			
 					HLODCreationParams.CellBounds = CellDescInstance.Bounds;
-					HLODCreationParams.GetRuntimeGrid = [&MainPartitionTokens](const UHLODLayer* HLODLayer) { return FName(*FString::Printf(TEXT("%s:%s"), *FString::JoinBy(MainPartitionTokens, TEXT("."), [](const FName Token) { return Token.ToString(); }), *HLODLayer->GetName())); };
+					HLODCreationParams.GetRuntimeGrid = [&MainPartitionTokens](const UHLODLayer* InHLODLayer) { return FName(*FString::Printf(TEXT("%s:%s"), *FString::JoinBy(MainPartitionTokens, TEXT("."), [](const FName Token) { return Token.ToString(); }), *InHLODLayer->GetName())); };
 					HLODCreationParams.HLODLevel = HLODLevel;
 					HLODCreationParams.MinVisibleDistance = RuntimePartition->LoadingRange;
 					HLODCreationParams.ContentBundleGuid = CellDescInstance.ContentBundleID;
