Fix regression in CL 28276825 that removed null check on Actor pointer

#rnx

diff --git a/Engine/Source/Runtime/Engine/Private/LevelInstance/LevelInstanceLevelStreaming.cpp b/Engine/Source/Runtime/Engine/Private/LevelInstance/LevelInstanceLevelStreaming.cpp
--- a/Engine/Source/Runtime/Engine/Private/LevelInstance/LevelInstanceLevelStreaming.cpp
+++ b/Engine/Source/Runtime/Engine/Private/LevelInstance/LevelInstanceLevelStreaming.cpp
@@ -115,21 +115,24 @@ void ULevelStreamingLevelInstance::OnLoadedActorsAddedToLevelPreEvent(const TArr
 		{
 			for (AActor* Actor : InActors)
 			{
-				if (Actor->IsPackageExternal())
+				if (IsValid(Actor))
 				{
-					if (bResetLoadersCalled)
+					if (Actor->IsPackageExternal())
 					{
-						ResetLoaders(Actor->GetExternalPackage());
-					}
+						if (bResetLoadersCalled)
+						{
+							ResetLoaders(Actor->GetExternalPackage());
+						}
 
-					FLevelInstanceLevelStreamingUtils::MarkObjectsInPackageAsTransientAndNonTransactional(Actor->GetExternalPackage());
-				}
+						FLevelInstanceLevelStreamingUtils::MarkObjectsInPackageAsTransientAndNonTransactional(Actor->GetExternalPackage());
+					}
 
-				// Must happen before the actors are registered with the world, which is the case for this delegate.
-				const FActorContainerID& ContainerID = LevelInstance->GetLevelInstanceID().GetContainerID();
-				FSetActorInstanceGuid SetActorInstanceGuid(Actor, ContainerID.GetActorGuid(Actor->GetActorGuid()));
+					// Must happen before the actors are registered with the world, which is the case for this delegate.
+					const FActorContainerID& ContainerID = LevelInstance->GetLevelInstanceID().GetContainerID();
+					FSetActorInstanceGuid SetActorInstanceGuid(Actor, ContainerID.GetActorGuid(Actor->GetActorGuid()));
 
-				FSetActorIsInLevelInstance SetIsInLevelInstance(Actor);
+					FSetActorIsInLevelInstance SetIsInLevelInstance(Actor);
+				}
 			}
 		}
 	}
