Bugfixes for custom data assert and incorrect handling of updates that occur during the same frame as a tree rebuild in HISM instance data paths.

#rnx

diff --git a/Engine/Source/Runtime/Engine/Private/InstancedStaticMesh/ISMInstanceDataManager.cpp b/Engine/Source/Runtime/Engine/Private/InstancedStaticMesh/ISMInstanceDataManager.cpp
--- a/Engine/Source/Runtime/Engine/Private/InstancedStaticMesh/ISMInstanceDataManager.cpp
+++ b/Engine/Source/Runtime/Engine/Private/InstancedStaticMesh/ISMInstanceDataManager.cpp
@@ -305,7 +305,6 @@
 		uint8 bInstancesChanged : 1;
 		uint8 bPrimitiveTransformChanged : 1;
 		uint8 bMaterialUsageFlagsChanged : 1;
-		uint8 bExternalDataChanged : 1;
 		uint8 bMaxDisplacementChanged : 1;
 		uint8 bStaticMeshBoundsChanged : 1;
 	};
@@ -355,6 +354,11 @@
 		check(!InstanceDataUpdateTaskInfo || InstanceDataUpdateTaskInfo->GetHeader() == InstanceDataBufferHeader);
 		InnerTaskLambda();
 		check(!InstanceDataUpdateTaskInfo || InstanceDataUpdateTaskInfo->GetHeader() == InstanceDataBufferHeader);
+		// check(InstanceDataBufferHeader.Flags == Proxy->GetData().GetFlags());
+		const FInstanceDataFlags HeaderFlags = InstanceDataBufferHeader.Flags;
+		const bool bHasAnyPayloadData = HeaderFlags.bHasPerInstanceHierarchyOffset || HeaderFlags.bHasPerInstanceLocalBounds || HeaderFlags.bHasPerInstanceDynamicData || HeaderFlags.bHasPerInstanceLMSMUVBias || HeaderFlags.bHasPerInstanceCustomData || HeaderFlags.bHasPerInstancePayloadExtension || HeaderFlags.bHasPerInstanceEditorData;
+		check(bHasAnyPayloadData || InstanceDataBufferHeader.PayloadDataStride == 0);
+		check(!bHasAnyPayloadData || InstanceDataBufferHeader.PayloadDataStride != 0);
 	};
 #else
 	TaskLambdaType OuterTaskLambda = MoveTemp(TaskLambda);
@@ -563,59 +567,16 @@
 	// Determine if the instance data proxy is exposed to the RT such that we can do early dispatch, this is important for level loading, probably, where we could move the work earlier.
 	// In fact, we could, e.g., for static ISMs kick this directly in post load perhaps.
 	const bool bIsUnattached = bNewPrimitiveProxy && Proxy->bIsNew;
-	//check(!Proxy->bIsNew || GetState() != EState::Tracked)
 	Proxy->bIsNew = false;
 
 	EShaderPlatform ShaderPlatform = ComponentData.PrimitiveSceneProxy->GetScene().GetShaderPlatform();
 	ERHIFeatureLevel::Type FeatureLevel = ComponentData.PrimitiveSceneProxy->GetScene().GetFeatureLevel();
-
-	// Mobile non-GPU-Scene renderer
-	bool bUseLegacyRenderer = !UseGPUScene(ShaderPlatform, FeatureLevel);
-
-	FChangeDesc ChangeDesc;
-
-	// Note: this must be supplied to allow calculating correct bounds for each instance, currently this is derived in the primitive proxy ctor, meaning it should be fetched from the proxy
-	//       However, we may want to remove this coupling.
-	float NewAbsMaxDisplacement = ComponentData.PrimitiveSceneProxy->GetAbsMaxDisplacement();
-
-
-	// TODO: We may decide to do so if other conditions are met (e.g., large change-set or marked for full invalidation).
-	// TODO: Need to figure this out in some other way, e.g., attachment counter or whatnot, since the creation has been moved up we no longer know if this is a fresh one.
-	ChangeDesc.bUntrackedState = GetState() != ETrackingState::Tracked;
 	check(Proxy->CheckPlatformFeatureLevel(ShaderPlatform, FeatureLevel));
 
-	// Figure out the deltas.
-	if (!ChangeDesc.bUntrackedState)
+	// Marked for externally managed update, this may be combined with tracked changes (which results in double updates)
+	if (ExternalBuildUpdateData)
 	{
-		ChangeDesc.bInstancesChanged =  HasAnyInstanceChanges();
-
-		ChangeDesc.bPrimitiveTransformChanged = bPrimitiveTransformChanged || !ComponentData.PrimitiveLocalToWorld.Equals(PrimitiveLocalToWorld);
-		ChangeDesc.bMaterialUsageFlagsChanged = Flags != ComponentData.Flags;
-		ChangeDesc.bMaxDisplacementChanged = AbsMaxDisplacement != NewAbsMaxDisplacement;
-		ChangeDesc.bStaticMeshBoundsChanged = !StaticMeshBounds.Equals(ComponentData.StaticMeshBounds);
-	}
-
-	ChangeDesc.bExternalDataChanged = GetState() == ETrackingState::External;
-
-	FMatrix PrevPrimitiveLocalToWorld = PrimitiveLocalToWorld;
-	// Update the tracked state
-	PrimitiveLocalToWorld = ComponentData.PrimitiveLocalToWorld;
-	AbsMaxDisplacement = NewAbsMaxDisplacement;
-	StaticMeshBounds = ComponentData.StaticMeshBounds;
-	Flags = ComponentData.Flags;
-
-	const bool bAnyChange = ChangeDesc.Packed != 0;
-
-	if (!bAnyChange)
-	{
-		return false;
-	}
-	
-	// Marked for externally managed update
-	if (GetState() == ETrackingState::External)
-	{
 		LOG_INST_DATA(TEXT("Rebuild with EState::External %s"), TEXT(""));
-		check(ChangeDesc.bExternalDataChanged);
 
 		FExternalGenericChangeSet ExternalChangeSet;
 		ExternalChangeSet.PrimitiveWorldSpaceOffset = FLargeWorldRenderPosition(PrimitiveLocalToWorld.GetOrigin()).GetTileOffset();
@@ -634,6 +595,9 @@
 		// Make sure we have a tracking state that matches the number of instances, but doesn't retain any history
 		ClearIdTracking(ExternalChangeSet.ExternalUpdateData.NumInstances);
 
+		// Need to store this off for tracking purposes
+		NumCustomDataFloats = ExternalChangeSet.ExternalUpdateData.NumCustomDataFloats;
+
 		// Assemble header info to enable nonblocking primitive update.
 		FInstanceDataBufferHeader InstanceDataBufferHeader;
 		InstanceDataBufferHeader.NumInstances = NumInstances;
@@ -669,13 +633,46 @@
 
 			check(ProxyRef.GetData().GetNumInstances() == InstanceIdIndexMap.GetMaxInstanceIndex());
 			check(ProxyRef.GetUpdateTaskInfo()->GetHeader().NumInstances == InstanceIdIndexMap.GetMaxInstanceIndex());
+			check(ProxyRef.GetUpdateTaskInfo()->GetHeader().Flags == ProxyRef.InstanceSceneDataBuffers.GetFlags());
 		});
-
-		TrackingState = ETrackingState::Tracked;
-
-		return true;
 	}
 
+	FChangeDesc ChangeDesc;
+
+	// Note: this must be supplied to allow calculating correct bounds for each instance, currently this is derived in the primitive proxy ctor, meaning it should be fetched from the proxy
+	//       However, we may want to remove this coupling.
+	float NewAbsMaxDisplacement = ComponentData.PrimitiveSceneProxy->GetAbsMaxDisplacement();
+
+
+	// TODO: We may decide to do so if other conditions are met (e.g., large change-set or marked for full invalidation).
+	// TODO: Need to figure this out in some other way, e.g., attachment counter or whatnot, since the creation has been moved up we no longer know if this is a fresh one.
+	ChangeDesc.bUntrackedState = GetState() != ETrackingState::Tracked;
+
+	// Figure out the deltas.
+	if (!ChangeDesc.bUntrackedState)
+	{
+		ChangeDesc.bInstancesChanged =  HasAnyInstanceChanges();
+
+		ChangeDesc.bPrimitiveTransformChanged = bPrimitiveTransformChanged || !ComponentData.PrimitiveLocalToWorld.Equals(PrimitiveLocalToWorld);
+		ChangeDesc.bMaterialUsageFlagsChanged = Flags != ComponentData.Flags;
+		ChangeDesc.bMaxDisplacementChanged = AbsMaxDisplacement != NewAbsMaxDisplacement;
+		ChangeDesc.bStaticMeshBoundsChanged = !StaticMeshBounds.Equals(ComponentData.StaticMeshBounds);
+	}
+
+	FMatrix PrevPrimitiveLocalToWorld = PrimitiveLocalToWorld;
+	// Update the tracked state
+	PrimitiveLocalToWorld = ComponentData.PrimitiveLocalToWorld;
+	AbsMaxDisplacement = NewAbsMaxDisplacement;
+	StaticMeshBounds = ComponentData.StaticMeshBounds;
+	Flags = ComponentData.Flags;
+
+	const bool bAnyChange = ChangeDesc.Packed != 0;
+
+	if (!bAnyChange)
+	{
+		return false;
+	}
+	
 	// 
 	{
 		// TODO: Maybe specialize for only bPrimitiveTransformChanged (no need to send/modify anything _other_ than the transform data)
@@ -893,10 +890,7 @@
 
 	if (GetState() != ETrackingState::Tracked)
 	{
-		if (GetState() != ETrackingState::External)
-		{
-			MarkComponentRenderInstancesDirty();
-		}
+		MarkComponentRenderInstancesDirty();
 		return;
 	}
 	MarkChangeHelper(TrackingArray, IndexToId(InstanceIndex));
@@ -908,10 +902,7 @@
 
 	if (GetState() != ETrackingState::Tracked)
 	{
-		if (GetState() != ETrackingState::External)
-		{
-			MarkComponentRenderInstancesDirty();
-		}
+		MarkComponentRenderInstancesDirty();
 		return;
 	}
 
@@ -1045,9 +1036,11 @@
 
 void FPrimitiveInstanceDataManager::MarkForRebuildFromExternal(int32 InNumInstances, TUniqueFunction<FExternalUpdateData(TArray<TRefCountPtr<HHitProxy>> &OutHitProxies)> &&BuildUpdateData)
 {
+	// Forget any tracked changes that happened before (and ID-mapping)
 	ClearIdTracking(InNumInstances);
 	ExternalBuildUpdateData = MoveTemp(BuildUpdateData);
-	TrackingState = ETrackingState::External;
+	// Must actually track the changes because we may see incremental changes in the same frame,
+	TrackingState = ETrackingState::Tracked;
 	MarkComponentRenderInstancesDirty();
 }
 
diff --git a/Engine/Source/Runtime/Engine/Private/InstancedStaticMesh/ISMInstanceDataSceneProxy.cpp b/Engine/Source/Runtime/Engine/Private/InstancedStaticMesh/ISMInstanceDataSceneProxy.cpp
--- a/Engine/Source/Runtime/Engine/Private/InstancedStaticMesh/ISMInstanceDataSceneProxy.cpp
+++ b/Engine/Source/Runtime/Engine/Private/InstancedStaticMesh/ISMInstanceDataSceneProxy.cpp
@@ -473,7 +473,6 @@ void FISMCInstanceDataSceneProxyLegacyReordered::BuildFromLegacyData(TUniquePtr<
 	ProxyData.InstanceLightShadowUVBias.SetNumZeroed(ProxyData.Flags.bHasPerInstanceLMSMUVBias ? NumInstances : 0);
 	ProxyData.InstanceLocalBounds = MakeArrayView(&InstanceLocalBounds, 1);
 	ProxyData.NumCustomDataFloats = ExternalLegacyData->GetNumCustomDataFloats();
-	ProxyData.Flags.bHasPerInstanceCustomData = ProxyData.NumCustomDataFloats != 0;
 	ProxyData.InstanceCustomData.SetNumZeroed(ProxyData.Flags.bHasPerInstanceCustomData ? NumInstances * ProxyData.NumCustomDataFloats : 0); 
 
 	ProxyData.InstanceRandomIDs.SetNumZeroed(ProxyData.Flags.bHasPerInstanceRandom ? NumInstances : 0);
diff --git a/Engine/Source/Runtime/Engine/Public/InstancedStaticMesh/ISMInstanceDataManager.h b/Engine/Source/Runtime/Engine/Public/InstancedStaticMesh/ISMInstanceDataManager.h
--- a/Engine/Source/Runtime/Engine/Public/InstancedStaticMesh/ISMInstanceDataManager.h
+++ b/Engine/Source/Runtime/Engine/Public/InstancedStaticMesh/ISMInstanceDataManager.h
@@ -51,7 +51,6 @@ public:
 	{
 		Initial, // In the initial state, there is no proxy and therefore changes do not need to be tracked, e.g., during initial setup of an ISM component.
 		Tracked,
-		External,
 	};
 
 	enum class EMode
