Properly aggregate material stats when performing multi process cook.

These stats will be reported when we report shader stats:

=== Material stats (aggregated across all cook processes) ===
Materials Cooked:          1583
Materials Translated:      1275
Material Translate Time: 1.70 s

- Materials Cooked, Materials Translated, and Material Translate time are the stats that are aggregated and reported to analytics.
- Piggyback the material counters aggregation in the existing shader stat aggregation for simplicity.
- Removed outdated, or not impactful material cook stats.  In most cases these values are extremely low and don't matter for a cook.

diff --git a/Engine/Source/Runtime/Engine/Private/Materials/HLSLMaterialTranslator.cpp b/Engine/Source/Runtime/Engine/Private/Materials/HLSLMaterialTranslator.cpp
--- a/Engine/Source/Runtime/Engine/Private/Materials/HLSLMaterialTranslator.cpp
+++ b/Engine/Source/Runtime/Engine/Private/Materials/HLSLMaterialTranslator.cpp
@@ -102,15 +102,8 @@ struct FCsvLogFile
 #include "ProfilingDebugging/ScopedTimers.h"
 namespace MaterialTranslatorCookStats
 {
-	static int32 MaterialTranslateCalls = 0;
-	static double MaterialTranslateTimeSec = 0.0f;
-
 	static FCookStatsManager::FAutoRegisterCallback RegisterCookStats([](FCookStatsManager::AddStatFuncRef AddStat)
 		{
-			AddStat(TEXT("Material"), FCookStatsManager::CreateKeyValueArray(
-				TEXT("MaterialTranslateCalls"), MaterialTranslateCalls,
-				TEXT("MaterialTranslateTimeSec"), MaterialTranslateTimeSec
-			));
 			FCsvLogFile::Get().Save();
 		});
 }
@@ -883,8 +876,6 @@ bool FHLSLMaterialTranslator::Translate()
 		<< FHLSLMaterialTranslatorTranslate.MaterialName(*TraceMaterialName);
 #endif
 
-	COOK_STAT(MaterialTranslatorCookStats::MaterialTranslateCalls++);
-	COOK_STAT(FScopedDurationTimer DurationTimer(MaterialTranslatorCookStats::MaterialTranslateTimeSec));
 	COOK_STAT(FDateTime TranslationDateTime = FDateTime::Now());
 
 	STAT(double HLSLTranslateTime = 0);
@@ -1971,8 +1962,12 @@ bool FHLSLMaterialTranslator::Translate()
 		// Store the number of unique VT samples
 		MaterialCompilationOutput.EstimatedNumVirtualTextureLookups = NumVtSamples;
 	}
+
 	ClearAllFunctionStacks();
-		
+	
+#if STATS
+	GShaderCompilerStats->IncrementMaterialTranslateTime(HLSLTranslateTime);
+#endif // STATS
 	INC_FLOAT_STAT_BY(STAT_ShaderCompiling_HLSLTranslation,(float)HLSLTranslateTime);
 
 #if ENABLE_COOK_STATS
diff --git a/Engine/Source/Runtime/Engine/Private/Materials/Material.cpp b/Engine/Source/Runtime/Engine/Private/Materials/Material.cpp
--- a/Engine/Source/Runtime/Engine/Private/Materials/Material.cpp
+++ b/Engine/Source/Runtime/Engine/Private/Materials/Material.cpp
@@ -101,20 +101,6 @@
 
 #include UE_INLINE_GENERATED_CPP_BY_NAME(Material)
 
-#if ENABLE_COOK_STATS
-namespace MaterialCookStats
-{
-	static double MaterialUpdateCachedExpressionDataSec = 0.0;
-
-	static FCookStatsManager::FAutoRegisterCallback RegisterCookStats([](FCookStatsManager::AddStatFuncRef AddStat)
-		{
-			AddStat(TEXT("Material"), FCookStatsManager::CreateKeyValueArray(
-				TEXT("MaterialUpdateCachedExpressionDataSec"), MaterialUpdateCachedExpressionDataSec
-			));
-		});
-}
-#endif
-
 #define LOCTEXT_NAMESPACE "Material"
 
 static TAutoConsoleVariable<int32> CVarMaterialParameterLegacyChecks(
@@ -2050,8 +2036,6 @@ void UMaterial::UpdateTransientExpressionData()
 #if WITH_EDITOR
 void UMaterial::UpdateCachedExpressionData()
 {
-	COOK_STAT(FScopedDurationTimer BlockingTimer(MaterialCookStats::MaterialUpdateCachedExpressionDataSec));
-
 	//@note FH: temporary preemptive PostLoad until zenloader load ordering improvements
 	ConditionalPostLoad();
 
diff --git a/Engine/Source/Runtime/Engine/Private/Materials/MaterialInstance.cpp b/Engine/Source/Runtime/Engine/Private/Materials/MaterialInstance.cpp
--- a/Engine/Source/Runtime/Engine/Private/Materials/MaterialInstance.cpp
+++ b/Engine/Source/Runtime/Engine/Private/Materials/MaterialInstance.cpp
@@ -63,21 +63,6 @@ DECLARE_CYCLE_STAT(TEXT("MaterialInstance CopyMatInstParams"), STAT_MaterialInst
 DECLARE_CYCLE_STAT(TEXT("MaterialInstance Serialize"), STAT_MaterialInstance_Serialize, STATGROUP_Shaders);
 DECLARE_CYCLE_STAT(TEXT("MaterialInstance CopyUniformParamsInternal"), STAT_MaterialInstance_CopyUniformParamsInternal, STATGROUP_Shaders);
 
-#if ENABLE_COOK_STATS
-#include "ProfilingDebugging/ScopedTimers.h"
-namespace MaterialInstanceCookStats
-{
-	static double MaterialInstancePostLoadSec = 0.0;
-
-	static FCookStatsManager::FAutoRegisterCallback RegisterMaterialInstanceCookStats([](FCookStatsManager::AddStatFuncRef AddStat)
-	{
-		AddStat(TEXT("Material"), FCookStatsManager::CreateKeyValueArray(
-			TEXT("MaterialInstancePostLoadSec"), MaterialInstancePostLoadSec
-		));
-	});
-}
-#endif
-
 // This flag controls whether MaterialInstances parents should be restricted to be either uncooked, to be
 // user defined, part of the engine or part of the base game.
 ENGINE_API bool bEnableRestrictiveMaterialInstanceParents = false;
@@ -2981,7 +2966,6 @@ void UMaterialInstance::PostLoad()
 {
 	LLM_SCOPE(ELLMTag::MaterialInstance);
 	SCOPED_LOADTIMER(MaterialInstancePostLoad);
-	COOK_STAT(FScopedDurationTimer BlockingTimer(MaterialInstanceCookStats::MaterialInstancePostLoadSec));
 
 #if WITH_EDITORONLY_DATA // fixup serialization before everything else
 	if (IsEditorOnlyDataValid() && !GetEditorOnlyData()->StaticParameters.StaticSwitchParameters_DEPRECATED.IsEmpty())
diff --git a/Engine/Source/Runtime/Engine/Private/Materials/MaterialInstanceConstant.cpp b/Engine/Source/Runtime/Engine/Private/Materials/MaterialInstanceConstant.cpp
--- a/Engine/Source/Runtime/Engine/Private/Materials/MaterialInstanceConstant.cpp
+++ b/Engine/Source/Runtime/Engine/Private/Materials/MaterialInstanceConstant.cpp
@@ -7,7 +7,6 @@
 #include "Materials/MaterialInstanceConstant.h"
 #include "Materials/Material.h"
 #include "Materials/MaterialInstanceSupport.h"
-#include "ProfilingDebugging/CookStats.h"
 #include "MaterialCachedData.h"
 
 #include UE_INLINE_GENERATED_CPP_BY_NAME(MaterialInstanceConstant)
@@ -17,21 +16,6 @@
 #include "ObjectCacheEventSink.h"
 #endif
 
-#if ENABLE_COOK_STATS
-#include "ProfilingDebugging/ScopedTimers.h"
-namespace MaterialInstanceCookStats
-{
-static double MaterialInstanceUpdateCachedExpressionDataSec = 0.0;
-
-static FCookStatsManager::FAutoRegisterCallback RegisterCookStats([](FCookStatsManager::AddStatFuncRef AddStat)
-	{
-		AddStat(TEXT("Material"), FCookStatsManager::CreateKeyValueArray(
-			TEXT("MaterialInstanceUpdateCachedExpressionDataSec"), MaterialInstanceUpdateCachedExpressionDataSec
-		));
-	});
-}
-#endif
-
 UMaterialInstanceConstant::UMaterialInstanceConstant(const FObjectInitializer& ObjectInitializer)
 	: Super(ObjectInitializer)
 {
@@ -187,8 +171,6 @@ void FMaterialInstanceCachedData::InitializeForConstant(const FMaterialLayersFun
 
 void UMaterialInstanceConstant::UpdateCachedData()
 {
-	COOK_STAT(FScopedDurationTimer BlockingTimer(MaterialInstanceCookStats::MaterialInstanceUpdateCachedExpressionDataSec));
-
 	// Don't need to rebuild cached data if it was serialized
 	if (!bLoadedCachedData)
 	{
diff --git a/Engine/Source/Runtime/Engine/Private/Materials/MaterialInterface.cpp b/Engine/Source/Runtime/Engine/Private/Materials/MaterialInterface.cpp
--- a/Engine/Source/Runtime/Engine/Private/Materials/MaterialInterface.cpp
+++ b/Engine/Source/Runtime/Engine/Private/Materials/MaterialInterface.cpp
@@ -40,6 +40,7 @@
 #include "Components/DecalComponent.h"
 #include "UObject/ArchiveCookContext.h"
 #include "UObject/Package.h"
+#include "ShaderCompiler.h"
 
 #if WITH_EDITOR
 #include "ObjectCacheEventSink.h"
@@ -48,20 +49,6 @@
 
 #define LOCTEXT_NAMESPACE "MaterialInterface"
 
-#if ENABLE_COOK_STATS
-namespace MaterialCookStats
-{
-	static int32 NumMaterialsCooked = 0;
-
-	static FCookStatsManager::FAutoRegisterCallback RegisterCookStats2([](FCookStatsManager::AddStatFuncRef AddStat)
-	{
-		AddStat(TEXT("Material"), FCookStatsManager::CreateKeyValueArray(
-			TEXT("NumMaterialsCooked"), NumMaterialsCooked
-		));
-	});
-}
-#endif
-
 /**
  * This is used to deprecate data that has been built with older versions.
  * To regenerate the data, commands like "BUILDMATERIALTEXTURESTREAMINGDATA" can be used in the editor.
@@ -1656,7 +1643,7 @@ void UMaterialInterface::PreSave(FObjectPreSaveContext ObjectSaveContext)
 #if WITH_EDITOR
 	if (ObjectSaveContext.IsCooking())
 	{
-		COOK_STAT(MaterialCookStats::NumMaterialsCooked++);
+		GShaderCompilerStats->IncrementMaterialCook();
 	}
 #endif // WITH_EDITOR
 }
diff --git a/Engine/Source/Runtime/Engine/Private/Materials/MaterialShader.cpp b/Engine/Source/Runtime/Engine/Private/Materials/MaterialShader.cpp
--- a/Engine/Source/Runtime/Engine/Private/Materials/MaterialShader.cpp
+++ b/Engine/Source/Runtime/Engine/Private/Materials/MaterialShader.cpp
@@ -57,17 +57,9 @@ static TAutoConsoleVariable<FString> CVarShaderCompilerDebugDDCKeyAsset(
 namespace MaterialShaderCookStats
 {
 	static FCookStats::FDDCResourceUsageStats UsageStats;
-	static int32 MaterialShadersCompiled = 0;
-	static double MaterialShaderMapSubmitCompileJobsTime = 0.0;
-	static int32 MaterialShaderMapSubmitCompileCalls = 0;
 	static FCookStatsManager::FAutoRegisterCallback RegisterCookStats([](FCookStatsManager::AddStatFuncRef AddStat)
 	{
 		UsageStats.LogStats(AddStat, TEXT("MaterialShader.Usage"), TEXT(""));
-		AddStat(TEXT("Material"), FCookStatsManager::CreateKeyValueArray(
-			TEXT("MaterialShadersCompiled"), MaterialShadersCompiled,
-			TEXT("MaterialShaderMapSubmitCompileJobsTime"), MaterialShaderMapSubmitCompileJobsTime,
-			TEXT("MaterialShaderMapSubmitCompileCalls"), MaterialShaderMapSubmitCompileCalls
-			));
 	});
 }
 #endif
@@ -1211,7 +1203,6 @@ static void PrepareMaterialShaderCompileJob(EShaderPlatform Platform,
 	FShaderCompilerEnvironment& ShaderEnvironment = NewJob->Input.Environment;
 
 	UE_LOG(LogShaders, Verbose, TEXT("			%s"), ShaderType->GetName());
-	COOK_STAT(MaterialShaderCookStats::MaterialShadersCompiled++);
 
 	//update material shader stats
 	UpdateMaterialShaderCompilingStats(Material);
@@ -1848,8 +1839,6 @@ int32 FMaterialShaderMap::SubmitCompileJobs(uint32 CompilingShaderMapId,
 	EShaderCompileJobPriority InPriority) const
 {
 	TRACE_CPUPROFILER_EVENT_SCOPE(FMaterialShaderMap::SubmitCompileJobs);
-	COOK_STAT(MaterialShaderCookStats::MaterialShaderMapSubmitCompileCalls++);
-	COOK_STAT(FScopedDurationTimer DurationTimer(MaterialShaderCookStats::MaterialShaderMapSubmitCompileJobsTime));
 
 	check(CompilingShaderMapId != 0u);
 	check(MaterialEnvironment);
diff --git a/Engine/Source/Runtime/Engine/Private/Materials/MeshMaterialShader.cpp b/Engine/Source/Runtime/Engine/Private/Materials/MeshMaterialShader.cpp
--- a/Engine/Source/Runtime/Engine/Private/Materials/MeshMaterialShader.cpp
+++ b/Engine/Source/Runtime/Engine/Private/Materials/MeshMaterialShader.cpp
@@ -6,20 +6,6 @@
 
 #include "MeshMaterialShader.h"
 #include "ShaderCompiler.h"
-#include "ProfilingDebugging/CookStats.h"
-
-#if ENABLE_COOK_STATS
-namespace MaterialMeshCookStats
-{
-	static int32 MeshMaterialShadersCompiled = 0;
-	static FCookStatsManager::FAutoRegisterCallback RegisterCookStats([](FCookStatsManager::AddStatFuncRef AddStat)
-	{
-		AddStat(TEXT("Material"), FCookStatsManager::CreateKeyValueArray(
-			TEXT("MeshMaterialShadersCompiled"), MeshMaterialShadersCompiled
-			));
-	});
-}
-#endif
 
 #if WITH_EDITOR
 
@@ -56,7 +42,6 @@ static void PrepareMeshMaterialShaderCompileJob(EShaderPlatform Platform,
 	UpdateMaterialShaderCompilingStats(Material);
 
 	UE_LOG(LogShaders, Verbose, TEXT("			%s"), ShaderType->GetName());
-	COOK_STAT(MaterialMeshCookStats::MeshMaterialShadersCompiled++);
 
 	// Allow the shader type to modify the compile environment.
 	ShaderType->SetupCompileEnvironment(Platform, MaterialParameters, VertexFactoryType, Key.PermutationId, PermutationFlags, ShaderEnvironment);
diff --git a/Engine/Source/Runtime/Engine/Private/ShaderCompiler/ShaderCompiler.cpp b/Engine/Source/Runtime/Engine/Private/ShaderCompiler/ShaderCompiler.cpp
--- a/Engine/Source/Runtime/Engine/Private/ShaderCompiler/ShaderCompiler.cpp
+++ b/Engine/Source/Runtime/Engine/Private/ShaderCompiler/ShaderCompiler.cpp
@@ -4615,6 +4615,8 @@ void FShaderCompilerStats::WriteStatSummary()
 		}
 	}
 
+	MaterialCounters.WriteStatSummary(AggregatedSuffix);
+
 	UE_LOG(LogShaderCompilers, Display, TEXT("================================================"));
 }
 
@@ -4695,6 +4697,8 @@ void FShaderCompilerStats::GatherAnalytics(const FString& BaseName, TArray<FAnal
 			Attributes.Emplace(MoveTemp(AttrName), Counters.CacheMemBudget);
 		}
 	}
+
+	MaterialCounters.GatherAnalytics(Attributes);
 }
 
 uint32 FShaderCompilerStats::GetTotalShadersCompiled()
@@ -4794,6 +4798,8 @@ void FShaderCompilerStats::Aggregate(FShaderCompilerStats& Other)
 			ShaderTimings.Add(TimingsKeyValue);
 		}
 	}
+
+	MaterialCounters += Other.MaterialCounters;
 }
 
 void FShaderCompilerStats::WriteToCompactBinary(FCbWriter& Writer)
@@ -4801,6 +4807,8 @@ void FShaderCompilerStats::WriteToCompactBinary(FCbWriter& Writer)
 	FScopeLock Lock(&CompileStatsLock);
 	Writer.AddBinary("Counters", &Counters, sizeof(Counters));
 
+	Writer.AddBinary("MaterialCounters", &MaterialCounters, sizeof(MaterialCounters));
+
 	Writer.BeginArray("CompileStatIndices");	
 	// Write the array of valid indices this worker has in the compile stats sparse array
 	for (TSparseArray<ShaderCompilerStats>::TConstIterator It(CompileStats); It; ++It)
@@ -4886,6 +4894,10 @@ void FShaderCompilerStats::ReadFromCompactBinary(FCbObjectView& Reader)
 	check(CountersMem.GetSize() == sizeof(FCounters));
 	Counters = *reinterpret_cast<const FCounters*>(CountersMem.GetData());
 
+	FMemoryView MaterialCountersMem = Reader["MaterialCounters"].AsBinaryView();
+	check(MaterialCountersMem.GetSize() == sizeof(FMaterialCounters));
+	MaterialCounters = *reinterpret_cast<const FMaterialCounters*>(MaterialCountersMem.GetData());
+
 	FCbArrayView CompileStatIndicesView = Reader["CompileStatIndices"].AsArrayView();
 	FCbArrayView CompileStatsView = Reader["CompileStats"].AsArrayView();
 	check(CompileStatIndicesView.Num() == CompileStatsView.Num());
@@ -5105,6 +5117,39 @@ void FShaderCompilerStats::RegisterJobBatch(int32 NumJobs, EExecutionType ExecTy
 	}
 }
 
+void FShaderCompilerStats::FMaterialCounters::WriteStatSummary(const TCHAR* AggregatedSuffix)
+{
+	UE_LOG(LogShaderCompilers, Display, TEXT("=== Material stats%s ==="), AggregatedSuffix);
+	UE_LOG(LogShaderCompilers, Display, TEXT("Materials Cooked:        %d"), NumMaterialsCooked);
+	UE_LOG(LogShaderCompilers, Display, TEXT("Materials Translated:    %d"), MaterialTranslateCalls);
+	UE_LOG(LogShaderCompilers, Display, TEXT("Material Translate Time: %.2f s"), MaterialTranslateTimeSec);
+}
+
+void FShaderCompilerStats::FMaterialCounters::GatherAnalytics(TArray<FAnalyticsEventAttribute>& Attributes)
+{
+	Attributes.Emplace(TEXT("Material_NumMaterialsCooked"), NumMaterialsCooked);
+	Attributes.Emplace(TEXT("Material_MaterialTranslateCalls"), MaterialTranslateCalls);
+	Attributes.Emplace(TEXT("Material_MaterialTranslateTimeSec"), MaterialTranslateTimeSec);
+}
+
+void FShaderCompilerStats::IncrementMaterialCook()
+{
+	FScopeLock Lock(&CompileStatsLock);
+	MaterialCounters.NumMaterialsCooked++;
+}
+
+void FShaderCompilerStats::IncrementMaterialsTranslated()
+{
+	FScopeLock Lock(&CompileStatsLock);
+	MaterialCounters.MaterialTranslateCalls++;
+}
+
+void FShaderCompilerStats::IncrementMaterialTranslateTime(double InTime)
+{
+	FScopeLock Lock(&CompileStatsLock);
+	MaterialCounters.MaterialTranslateTimeSec += InTime;
+}
+
 void FShaderCompilerStats::RegisterCookedShaders(uint32 NumCooked, float CompileTime, EShaderPlatform Platform, const FString MaterialPath, FString PermutationString)
 {
 	FScopeLock Lock(&CompileStatsLock);
diff --git a/Engine/Source/Runtime/Engine/Public/ShaderCompiler.h b/Engine/Source/Runtime/Engine/Public/ShaderCompiler.h
--- a/Engine/Source/Runtime/Engine/Public/ShaderCompiler.h
+++ b/Engine/Source/Runtime/Engine/Public/ShaderCompiler.h
@@ -362,6 +362,10 @@ public:
 		}
 	};
 
+	void IncrementMaterialCook();
+	void IncrementMaterialsTranslated();
+	void IncrementMaterialTranslateTime(double InTime);
+
 	ENGINE_API void RegisterCookedShaders(uint32 NumCooked, float CompileTime, EShaderPlatform Platform, const FString MaterialPath, FString PermutationString = FString(""));
 	ENGINE_API void RegisterCompiledShaders(uint32 NumPermutations, EShaderPlatform Platform, const FString MaterialPath, FString PermutationString = FString(""));
 	const TSparseArray<ShaderCompilerStats>& GetShaderCompilerStats() { return CompileStats; }
@@ -549,6 +553,32 @@ private:
 
 	FCounters Counters;
 
+	struct FMaterialCounters
+	{
+		/** The total number of materials cooked.  This corresponds to UMaterialInterface::Presave() */
+		int32 NumMaterialsCooked = 0;
+
+		/** The total number of materials that have been translated.  */
+		int32 MaterialTranslateCalls = 0;
+
+		/** The total time in seconds to translate all materials.  */
+		double MaterialTranslateTimeSec = 0.0f;
+
+		FMaterialCounters& operator+=(const FMaterialCounters& Other)
+		{
+			NumMaterialsCooked += Other.NumMaterialsCooked;
+			MaterialTranslateCalls += Other.MaterialTranslateCalls;
+			MaterialTranslateTimeSec += Other.MaterialTranslateTimeSec;
+
+			return *this;
+		}
+
+		void WriteStatSummary(const TCHAR* AggregatedSuffix);
+		void GatherAnalytics(TArray<FAnalyticsEventAttribute>& Attributes);
+	};
+
+	FMaterialCounters MaterialCounters;
+
 	/** Accumulates the job lifetimes without overlaps */
 	TArray<TInterval<double>> JobLifeTimeIntervals;
 
