Fix build errors in MaterialTranslation.cpp

diff --git a/Engine/Source/Runtime/Engine/Private/Materials/HLSLMaterialTranslator.cpp b/Engine/Source/Runtime/Engine/Private/Materials/HLSLMaterialTranslator.cpp
--- a/Engine/Source/Runtime/Engine/Private/Materials/HLSLMaterialTranslator.cpp
+++ b/Engine/Source/Runtime/Engine/Private/Materials/HLSLMaterialTranslator.cpp
@@ -56,7 +56,7 @@
 #include "ShaderPlatformCachedIniValue.h"
 #endif
 
-
+#if ENABLE_COOK_STATS && STATS
 /**
  * Utility used to create a MaterialTranslationLog.txt file during cooks that contains the list of all translated materials
  * and other info such as how long the translation took.
@@ -98,7 +98,6 @@ struct FCsvLogFile
 
 };
 
-#if ENABLE_COOK_STATS
 #include "ProfilingDebugging/ScopedTimers.h"
 namespace MaterialTranslatorCookStats
 {
@@ -876,7 +875,10 @@ bool FHLSLMaterialTranslator::Translate()
 		<< FHLSLMaterialTranslatorTranslate.MaterialName(*TraceMaterialName);
 #endif
 
-	COOK_STAT(FDateTime TranslationDateTime = FDateTime::Now());
+
+#if ENABLE_COOK_STATS && STATS
+	FDateTime TranslationDateTime = FDateTime::Now();
+#endif
 
 	STAT(double HLSLTranslateTime = 0);
 	{
@@ -1970,7 +1972,7 @@ bool FHLSLMaterialTranslator::Translate()
 #endif // STATS
 	INC_FLOAT_STAT_BY(STAT_ShaderCompiling_HLSLTranslation,(float)HLSLTranslateTime);
 
-#if ENABLE_COOK_STATS
+#if ENABLE_COOK_STATS && STATS
 	// Write out a CSV file MaterialTranslationLog.txt containing info about all material translations.
 	FCsvLogFile::Get().AddEntry(Material->GetMaterialInterface()->GetFullName(), TranslationDateTime, HLSLTranslateTime);
 #endif
