World Partition RuntimeHashSet HLODs: with CL 28409908, we can now send all actors to FWorldPartitionHLODUtilities::CreateHLODActors instead of manually splitting them into their respective HLOD layers.

diff --git a/Engine/Source/Runtime/Engine/Private/WorldPartition/RuntimeHashSet/WorldPartitionRuntimeHashSetHLODGeneration.cpp b/Engine/Source/Runtime/Engine/Private/WorldPartition/RuntimeHashSet/WorldPartitionRuntimeHashSetHLODGeneration.cpp
--- a/Engine/Source/Runtime/Engine/Private/WorldPartition/RuntimeHashSet/WorldPartitionRuntimeHashSetHLODGeneration.cpp
+++ b/Engine/Source/Runtime/Engine/Private/WorldPartition/RuntimeHashSet/WorldPartitionRuntimeHashSetHLODGeneration.cpp
@@ -201,78 +201,69 @@ bool UWorldPartitionRuntimeHashSet::SetupHLODActors(const IStreamingGenerationCo
 			for (URuntimePartition::FCellDescInstance& CellDescInstance : CellDescInstances)
 			{
 				// Here we split actors into their respective HLOD layers because we want to provide a specific runtime grid name for our HLOD setup to FWorldPartitionHLODUtilities::CreateHLODActors. 
-				TMap<const UHLODLayer*, TArray<IStreamingGenerationContext::FActorInstance>> HLODLayerActorInstances;
+				TArray<IStreamingGenerationContext::FActorInstance> ActorInstances;
 				for (const IStreamingGenerationContext::FActorSetInstance* ActorSetInstance : CellDescInstance.ActorSetInstances)
 				{
-					ActorSetInstance->ForEachActor([this, ActorSetInstance, &HLODLayerActorInstances, &GetRuntimePartitionDescHLODSetup](const FGuid& ActorGuid)
+					ActorSetInstance->ForEachActor([this, ActorSetInstance, &ActorInstances](const FGuid& ActorGuid)
 					{
-						IStreamingGenerationContext::FActorInstance ActorInstance(ActorGuid, ActorSetInstance);
-						const FWorldPartitionActorDescView& ActorDescView = ActorInstance.GetActorDescView();
-			
-						if (const UHLODLayer* ActorInstanceHLODLayer = Cast<const UHLODLayer>(ActorDescView.GetHLODLayer().ResolveObject()))
-						{
-							HLODLayerActorInstances.FindOrAdd(ActorInstanceHLODLayer).Add(ActorInstance);
-						}
+						ActorInstances.Emplace(ActorGuid, ActorSetInstance);
 					});
 				}
 
-				for (auto& [HLODLayer, ActorInstances] : HLODLayerActorInstances)
-				{
-					// Fake tick
-					PrivateUtils::GameTick(WorldPartition->GetWorld());
+				// Fake tick
+				PrivateUtils::GameTick(WorldPartition->GetWorld());
 
-					const FCellUniqueId CellUniqueId = GetCellUniqueId(CellDescInstance);
+				const FCellUniqueId CellUniqueId = GetCellUniqueId(CellDescInstance);
 
-					TArray<FName> MainPartitionTokens;
-					TArray<FName> HLODPartitionTokens;
-					verify(ParseGridName(ActorInstances[0].ActorSetInstance->RuntimeGrid, MainPartitionTokens, HLODPartitionTokens));
+				TArray<FName> MainPartitionTokens;
+				TArray<FName> HLODPartitionTokens;
+				verify(ParseGridName(ActorInstances[0].ActorSetInstance->RuntimeGrid, MainPartitionTokens, HLODPartitionTokens));
 					
-					FHLODCreationParams HLODCreationParams;
-					HLODCreationParams.WorldPartition = WorldPartition;
-					HLODCreationParams.CellName = CellUniqueId.Name;
-					HLODCreationParams.CellGuid = CellUniqueId.Guid;			
-					HLODCreationParams.CellBounds = CellDescInstance.Bounds;
-					HLODCreationParams.GetRuntimeGrid = [&MainPartitionTokens](const UHLODLayer* InHLODLayer) { return FName(*FString::Printf(TEXT("%s:%s"), *FString::JoinBy(MainPartitionTokens, TEXT("."), [](const FName Token) { return Token.ToString(); }), *InHLODLayer->GetName())); };
-					HLODCreationParams.HLODLevel = HLODLevel;
-					HLODCreationParams.MinVisibleDistance = RuntimePartition->LoadingRange;
-					HLODCreationParams.ContentBundleGuid = CellDescInstance.ContentBundleID;
-					HLODCreationParams.DataLayerInstances = CellDescInstance.DataLayerInstances;
-
-					IWorldPartitionHLODUtilities* WPHLODUtilities = FModuleManager::Get().LoadModuleChecked<IWorldPartitionHLODUtilitiesModule>("WorldPartitionHLODUtilities").GetUtilities();
-					TArray<AWorldPartitionHLOD*> CellHLODActors = WPHLODUtilities->CreateHLODActors(HLODCreationContext, HLODCreationParams, ActorInstances);
-
-					if (!CellHLODActors.IsEmpty())
+				FHLODCreationParams HLODCreationParams;
+				HLODCreationParams.WorldPartition = WorldPartition;
+				HLODCreationParams.CellName = CellUniqueId.Name;
+				HLODCreationParams.CellGuid = CellUniqueId.Guid;			
+				HLODCreationParams.CellBounds = CellDescInstance.Bounds;
+				HLODCreationParams.GetRuntimeGrid = [&MainPartitionTokens](const UHLODLayer* InHLODLayer) { return FName(*FString::Printf(TEXT("%s:%s"), *FString::JoinBy(MainPartitionTokens, TEXT("."), [](const FName Token) { return Token.ToString(); }), *InHLODLayer->GetName())); };
+				HLODCreationParams.HLODLevel = HLODLevel;
+				HLODCreationParams.MinVisibleDistance = RuntimePartition->LoadingRange;
+				HLODCreationParams.ContentBundleGuid = CellDescInstance.ContentBundleID;
+				HLODCreationParams.DataLayerInstances = CellDescInstance.DataLayerInstances;
+
+				IWorldPartitionHLODUtilities* WPHLODUtilities = FModuleManager::Get().LoadModuleChecked<IWorldPartitionHLODUtilitiesModule>("WorldPartitionHLODUtilities").GetUtilities();
+				TArray<AWorldPartitionHLOD*> CellHLODActors = WPHLODUtilities->CreateHLODActors(HLODCreationContext, HLODCreationParams, ActorInstances);
+
+				if (!CellHLODActors.IsEmpty())
+				{
+					for (AWorldPartitionHLOD* CellHLODActor : CellHLODActors)
 					{
-						for (AWorldPartitionHLOD* CellHLODActor : CellHLODActors)
-						{
-							FGuid ActorGuid = CellHLODActor->GetActorGuid();
+						FGuid ActorGuid = CellHLODActor->GetActorGuid();
 
-							UPackage* CellHLODActorPackage = CellHLODActor->GetPackage();
-							if (CellHLODActorPackage->HasAnyPackageFlags(PKG_NewlyCreated))
-							{
-								// Get a reference to newly create actors so they get unloaded when we release the references
-								HLODCreationContext.ActorReferences.Emplace(WorldPartition, CellHLODActor->GetActorGuid());
-							}
-
-							HLODActorGuids.Add(ActorGuid);
-							NumNextLayerHLODActors += CellHLODActor->GetHLODLayer() ? 1 : 0;
+						UPackage* CellHLODActorPackage = CellHLODActor->GetPackage();
+						if (CellHLODActorPackage->HasAnyPackageFlags(PKG_NewlyCreated))
+						{
+							// Get a reference to newly create actors so they get unloaded when we release the references
+							HLODCreationContext.ActorReferences.Emplace(WorldPartition, CellHLODActor->GetActorGuid());
 						}
 
-						if (!Params.bReportOnly)
+						HLODActorGuids.Add(ActorGuid);
+						NumNextLayerHLODActors += CellHLODActor->GetHLODLayer() ? 1 : 0;
+					}
+
+					if (!Params.bReportOnly)
+					{
+						for (AWorldPartitionHLOD* CellHLODActor : CellHLODActors)
 						{
-							for (AWorldPartitionHLOD* CellHLODActor : CellHLODActors)
+							if (CellHLODActor->GetPackage()->IsDirty())
 							{
-								if (CellHLODActor->GetPackage()->IsDirty())
-								{
-									PrivateUtils::SavePackage(CellHLODActor->GetPackage(), Params.SourceControlHelper);
-								}
+								PrivateUtils::SavePackage(CellHLODActor->GetPackage(), Params.SourceControlHelper);
 							}
 						}
 					}
-
-					// Unload actors
-					HLODCreationContext.ActorReferences.Empty();
 				}
+
+				// Unload actors
+				HLODCreationContext.ActorReferences.Empty();
 			}
 		}
 
