Level Instance: Clear filter cache on hierarchy when level instance is reloaded

Fixes filter not properly updating when editing the defaults in the source level

#rnx

diff --git a/Engine/Source/Runtime/Engine/Private/LevelInstance/LevelInstanceComponent.cpp b/Engine/Source/Runtime/Engine/Private/LevelInstance/LevelInstanceComponent.cpp
--- a/Engine/Source/Runtime/Engine/Private/LevelInstance/LevelInstanceComponent.cpp
+++ b/Engine/Source/Runtime/Engine/Private/LevelInstance/LevelInstanceComponent.cpp
@@ -222,6 +222,11 @@ void ULevelInstanceComponent::UpdateEditFilter()
 	EditFilter = Filter;
 }
 
+void ULevelInstanceComponent::ClearCachedFilter()
+{
+	CachedFilteredActorsPerContainer.Reset();
+}
+
 bool ULevelInstanceComponent::IsEditFilter() const
 {
 	if (ILevelInstanceInterface* LevelInstance = Cast<ILevelInstanceInterface>(GetOwner()))
diff --git a/Engine/Source/Runtime/Engine/Private/LevelInstance/LevelInstanceInterface.cpp b/Engine/Source/Runtime/Engine/Private/LevelInstance/LevelInstanceInterface.cpp
--- a/Engine/Source/Runtime/Engine/Private/LevelInstance/LevelInstanceInterface.cpp
+++ b/Engine/Source/Runtime/Engine/Private/LevelInstance/LevelInstanceInterface.cpp
@@ -127,6 +127,17 @@ void ILevelInstanceInterface::UpdateLevelInstanceFromWorldAsset()
 	{
 		if (ULevelInstanceSubsystem* LevelInstanceSubsystem = GetLevelInstanceSubsystem())
 		{
+#if WITH_EDITOR
+			LevelInstanceSubsystem->ForEachLevelInstanceAncestorsAndSelf(CastChecked<AActor>(this), [](const ILevelInstanceInterface* Ancestor)
+			{
+				if (ULevelInstanceComponent* LevelInstanceComponent = Ancestor->GetLevelInstanceComponent())
+				{
+					LevelInstanceComponent->ClearCachedFilter();
+				}
+				return true;
+			});
+#endif
+
 			if (IsWorldAssetValid() && IsLoadingEnabled())
 			{
 				const bool bForceUpdate = true;
diff --git a/Engine/Source/Runtime/Engine/Public/LevelInstance/LevelInstanceComponent.h b/Engine/Source/Runtime/Engine/Public/LevelInstance/LevelInstanceComponent.h
--- a/Engine/Source/Runtime/Engine/Public/LevelInstance/LevelInstanceComponent.h
+++ b/Engine/Source/Runtime/Engine/Public/LevelInstance/LevelInstanceComponent.h
@@ -45,6 +45,7 @@ public:
 	ENGINE_API void SetFilter(const FWorldPartitionActorFilter& InFilter, bool bNotify = true);
 	ENGINE_API const TMap<FActorContainerID, TSet<FGuid>>& GetFilteredActorsPerContainer() const;
 	ENGINE_API void UpdateEditFilter();
+	ENGINE_API void ClearCachedFilter();
 private:
 	ENGINE_API bool ShouldShowSpriteComponent() const;
 	ENGINE_API void OnFilterChanged();
