[NetPhysics] Clear non-physics target from PhysicsReplication.

If a replicated target ends up in PhysicsReplication using Predictive Interpolation without being marked as simulating / physics replicating on the server, make sure we clear that target after reaching it, else it will linger.

diff --git a/Engine/Source/Runtime/Engine/Classes/Engine/ReplicatedState.h b/Engine/Source/Runtime/Engine/Classes/Engine/ReplicatedState.h
--- a/Engine/Source/Runtime/Engine/Classes/Engine/ReplicatedState.h
+++ b/Engine/Source/Runtime/Engine/Classes/Engine/ReplicatedState.h
@@ -76,6 +76,7 @@ namespace ERigidBodyFlags
 		None = 0x00,
 		Sleeping = 0x01,
 		NeedsUpdate = 0x02,
+		RepPhysics = 0x04 // This RigidBody was replicated from the server with simulating physics
 	};
 }
 
@@ -185,7 +186,10 @@ struct FRepMovement
 		RBState.Quaternion = Rotation.Quaternion();
 		RBState.LinVel = LinearVelocity;
 		RBState.AngVel = AngularVelocity;
-		RBState.Flags = (decltype(FRigidBodyState::Flags))(bSimulatedPhysicSleep ? ERigidBodyFlags::Sleeping : ERigidBodyFlags::None) | ERigidBodyFlags::NeedsUpdate;
+		RBState.Flags = 
+			(decltype(FRigidBodyState::Flags))(bSimulatedPhysicSleep ? ERigidBodyFlags::Sleeping : ERigidBodyFlags::None)
+			| ERigidBodyFlags::NeedsUpdate
+			| (bRepPhysics ? ERigidBodyFlags::RepPhysics : ERigidBodyFlags::None);
 	}
 
 	bool operator==(const FRepMovement& Other) const
diff --git a/Engine/Source/Runtime/Engine/Private/Actor.cpp b/Engine/Source/Runtime/Engine/Private/Actor.cpp
--- a/Engine/Source/Runtime/Engine/Private/Actor.cpp
+++ b/Engine/Source/Runtime/Engine/Private/Actor.cpp
@@ -2659,7 +2659,26 @@ void AActor::PrestreamTextures( float Seconds, bool bEnableStreaming, int32 Cine
 
 void AActor::OnRep_Instigator() {}
 
-void AActor::OnRep_ReplicateMovement() {}
+void AActor::OnRep_ReplicateMovement()
+{
+	// If the actor stops replicating movement and is using PhysicsReplication PredictiveInterpolation, remove replicated target else it will linger waiting for a sleep state.
+	if (!IsReplicatingMovement() && GetPhysicsReplicationMode() == EPhysicsReplicationMode::PredictiveInterpolation)
+	{
+		if (UPrimitiveComponent* RootPrimComp = Cast<UPrimitiveComponent>(RootComponent))
+		{
+			if (UWorld* World = GetWorld())
+			{
+				if (FPhysScene* PhysScene = World->GetPhysicsScene())
+				{
+					if (IPhysicsReplication* PhysicsReplication = PhysScene->GetPhysicsReplication())
+					{
+						PhysicsReplication->RemoveReplicatedTarget(RootPrimComp);
+					}
+				}
+			}
+		}
+	}
+}
 
 void AActor::RouteEndPlay(const EEndPlayReason::Type EndPlayReason)
 {
diff --git a/Engine/Source/Runtime/Engine/Private/ActorReplication.cpp b/Engine/Source/Runtime/Engine/Private/ActorReplication.cpp
--- a/Engine/Source/Runtime/Engine/Private/ActorReplication.cpp
+++ b/Engine/Source/Runtime/Engine/Private/ActorReplication.cpp
@@ -210,7 +210,7 @@ void AActor::OnRep_ReplicatedMovement()
 					break;
 				}
 			}
-			DrawDebugCapsule(GetWorld(), LocalRepMovement.Location, GetSimpleCollisionHalfHeight(), GetSimpleCollisionRadius(), LocalRepMovement.Rotation.Quaternion(), DebugColor, false, 1.f);
+			DrawDebugCapsule(GetWorld(), LocalRepMovement.Location, FMath::Max(GetSimpleCollisionHalfHeight(), 25.0f), FMath::Max(GetSimpleCollisionRadius(), 25.0f), LocalRepMovement.Rotation.Quaternion(), DebugColor, false, 1.f);
 		}
 #endif
 
diff --git a/Engine/Source/Runtime/Engine/Private/PhysicsEngine/PhysicsReplication.cpp b/Engine/Source/Runtime/Engine/Private/PhysicsEngine/PhysicsReplication.cpp
--- a/Engine/Source/Runtime/Engine/Private/PhysicsEngine/PhysicsReplication.cpp
+++ b/Engine/Source/Runtime/Engine/Private/PhysicsEngine/PhysicsReplication.cpp
@@ -1285,16 +1285,17 @@ bool FPhysicsReplicationAsync::PredictiveInterpolation(Chaos::FPBDRigidParticleH
 	const bool bCanSimulate = Handle->IsDynamic() || Handle->IsSleeping();
 
 	// Helper for sleep and target clearing at replication end
-	auto EndReplicationHelper = [RigidsSolver, Handle, bCanSimulate](FReplicatedPhysicsTargetAsync& Target, bool bAllowSleep) -> bool
+	auto EndReplicationHelper = [RigidsSolver, Handle, bCanSimulate](FReplicatedPhysicsTargetAsync& Target, bool bOkToClear) -> bool
 	{
 		const bool bShouldSleep = (Target.TargetState.Flags & ERigidBodyFlags::Sleeping) != 0;
+		const bool bReplicatingPhysics = (Target.TargetState.Flags & ERigidBodyFlags::RepPhysics) != 0;
 
-		if (bAllowSleep && bShouldSleep && bCanSimulate)
+		if (bOkToClear && bShouldSleep && bCanSimulate)
 		{
 			RigidsSolver->GetEvolution()->SetParticleObjectState(Handle, Chaos::EObjectStateType::Sleeping);
 		}
 
-		const bool bClearTarget = (!bCanSimulate || (bAllowSleep && bShouldSleep)) && !PhysicsReplicationCVars::PredictiveInterpolationCVars::bDontClearTarget;
+		const bool bClearTarget = (!bCanSimulate || (bOkToClear && bShouldSleep) || (bOkToClear && !bReplicatingPhysics)) && !PhysicsReplicationCVars::PredictiveInterpolationCVars::bDontClearTarget;
 		return bClearTarget;
 	};
 
