[NetPhysics] Fix variable type conversion of flag int to int8

diff --git a/Engine/Source/Runtime/Engine/Classes/Engine/ReplicatedState.h b/Engine/Source/Runtime/Engine/Classes/Engine/ReplicatedState.h
--- a/Engine/Source/Runtime/Engine/Classes/Engine/ReplicatedState.h
+++ b/Engine/Source/Runtime/Engine/Classes/Engine/ReplicatedState.h
@@ -189,7 +189,7 @@ struct FRepMovement
 		RBState.Flags = 
 			(decltype(FRigidBodyState::Flags))(bSimulatedPhysicSleep ? ERigidBodyFlags::Sleeping : ERigidBodyFlags::None)
 			| ERigidBodyFlags::NeedsUpdate
-			| (bRepPhysics ? ERigidBodyFlags::RepPhysics : ERigidBodyFlags::None);
+			| (decltype(FRigidBodyState::Flags))(bRepPhysics ? ERigidBodyFlags::RepPhysics : ERigidBodyFlags::None);
 	}
 
 	bool operator==(const FRepMovement& Other) const
