Root Material Node defaults fix

diff --git a/Engine/Source/Editor/UnrealEd/Private/MaterialGraphNode_Root.cpp b/Engine/Source/Editor/UnrealEd/Private/MaterialGraphNode_Root.cpp
--- a/Engine/Source/Editor/UnrealEd/Private/MaterialGraphNode_Root.cpp
+++ b/Engine/Source/Editor/UnrealEd/Private/MaterialGraphNode_Root.cpp
@@ -29,6 +29,13 @@ static TAutoConsoleVariable<bool> CVarEnableRootNodeInlineControls(
 UMaterialGraphNode_Root::UMaterialGraphNode_Root(const FObjectInitializer& ObjectInitializer)
 	: Super(ObjectInitializer)
 {
+	CVarEnableRootNodeInlineControls->SetOnChangedCallback(FConsoleVariableDelegate::CreateLambda([this](IConsoleVariable* InVariable)
+		{
+			for (UEdGraphPin* Pin : this->Pins)
+			{
+				UpdateInputUseConstant(Pin, InVariable->GetBool());
+			}
+		}));
 }
 
 void UMaterialGraphNode_Root::UpdateInputUseConstant(UEdGraphPin* Pin, bool bUseConstant)
@@ -36,6 +43,9 @@ void UMaterialGraphNode_Root::UpdateInputUseConstant(UEdGraphPin* Pin, bool bUse
 	const UMaterialGraph* MaterialGraph = CastChecked<UMaterialGraph>(GetGraph());
 	const FMaterialInputInfo& MaterialInput = MaterialGraph->MaterialInputs[Pin->SourceIndex];
 
+	const bool ShouldEnableInlineControls = CVarEnableRootNodeInlineControls.GetValueOnAnyThread();
+	bUseConstant &= ShouldEnableInlineControls;
+	
 	UMaterialEditorOnlyData* EditorOnlyData = Material->GetEditorOnlyData();
 	EMaterialProperty Property = MaterialInput.GetProperty();
 	switch (Property)
diff --git a/Engine/Source/Runtime/Engine/Private/Materials/Material.cpp b/Engine/Source/Runtime/Engine/Private/Materials/Material.cpp
--- a/Engine/Source/Runtime/Engine/Private/Materials/Material.cpp
+++ b/Engine/Source/Runtime/Engine/Private/Materials/Material.cpp
@@ -7298,12 +7298,23 @@ bool UMaterial::IsRefractionPinPluggedIn(const UMaterialEditorOnlyData* EditorOn
 UMaterialEditorOnlyData::UMaterialEditorOnlyData()
 {
 	BaseColor.Constant = FColor(128, 128, 128);
-	Metallic.Constant = 0.0f;
-	Specular.Constant = 0.5f;
-	Roughness.Constant = 0.5f;
-
-	Opacity.Constant = 1.0f;
-	OpacityMask.Constant = 1.0f;
+	Metallic.Constant = FMaterialAttributeDefinitionMap::GetDefaultValue(MP_Metallic).X;
+	Specular.Constant = FMaterialAttributeDefinitionMap::GetDefaultValue(MP_Specular).X;
+	Roughness.Constant = FMaterialAttributeDefinitionMap::GetDefaultValue(MP_Roughness).X;
+	Anisotropy.Constant = FMaterialAttributeDefinitionMap::GetDefaultValue(MP_Anisotropy).X;
+	Normal.Constant = FMaterialAttributeDefinitionMap::GetDefaultValue(MP_Normal);
+	Tangent.Constant = FMaterialAttributeDefinitionMap::GetDefaultValue(MP_Tangent);
+	EmissiveColor.Constant = FLinearColor(FMaterialAttributeDefinitionMap::GetDefaultValue(MP_EmissiveColor)).ToFColorSRGB();
+	Opacity.Constant = FMaterialAttributeDefinitionMap::GetDefaultValue(MP_Opacity).X;
+	OpacityMask.Constant = FMaterialAttributeDefinitionMap::GetDefaultValue(MP_OpacityMask).X;
+	WorldPositionOffset.Constant = FMaterialAttributeDefinitionMap::GetDefaultValue(MP_WorldPositionOffset);
+	Displacement.Constant = FMaterialAttributeDefinitionMap::GetDefaultValue(MP_Displacement).X;
+	SubsurfaceColor.Constant = FLinearColor(FMaterialAttributeDefinitionMap::GetDefaultValue(MP_SubsurfaceColor)).ToFColorSRGB();
+	ClearCoat.Constant = FMaterialAttributeDefinitionMap::GetDefaultValue(MP_CustomData0).X;
+	ClearCoatRoughness.Constant = FMaterialAttributeDefinitionMap::GetDefaultValue(MP_CustomData1).X;
+	AmbientOcclusion.Constant = FMaterialAttributeDefinitionMap::GetDefaultValue(MP_AmbientOcclusion).X;
+	Refraction.Constant = FMaterialAttributeDefinitionMap::GetDefaultValue(MP_Refraction).X;
+	SurfaceThickness.Constant = FMaterialAttributeDefinitionMap::GetDefaultValue(MP_SurfaceThickness).X;
 }
 
 #undef LOCTEXT_NAMESPACE
diff --git a/Engine/Source/Runtime/Engine/Private/Materials/MaterialShared.cpp b/Engine/Source/Runtime/Engine/Private/Materials/MaterialShared.cpp
--- a/Engine/Source/Runtime/Engine/Private/Materials/MaterialShared.cpp
+++ b/Engine/Source/Runtime/Engine/Private/Materials/MaterialShared.cpp
@@ -628,11 +628,8 @@ FString FColorMaterialInput::GetDefaultValue() const
 {
 	FString DefaultValue;
 #if WITH_EDITOR
-	if (UseConstant)
-	{
-		FLinearColor Color = Constant.ReinterpretAsLinear();
-		DefaultValue = Color.ToString();
-	}
+	FLinearColor Color = Constant.ReinterpretAsLinear();
+	DefaultValue = Color.ToString();
 #endif
 	return DefaultValue;
 }
@@ -649,10 +646,7 @@ FString FScalarMaterialInput::GetDefaultValue() const
 {
 	FString DefaultValue;
 #if WITH_EDITOR
-	if (UseConstant)
-	{
-		DefaultValue = FString::SanitizeFloat(Constant);
-	}
+	DefaultValue = FString::SanitizeFloat(Constant);
 #endif
 	return DefaultValue;
 }
@@ -671,10 +665,7 @@ FString FVector2MaterialInput::GetDefaultValue() const
 {
 	FString DefaultValue;
 #if WITH_EDITOR
-	if (UseConstant)
-	{
-		DefaultValue = FString(TEXT("(X=")) + FString::SanitizeFloat(Constant.X) + FString(TEXT(",Y=")) + FString::SanitizeFloat(Constant.Y) + FString(TEXT(")"));
-	}
+	DefaultValue = FString(TEXT("(X=")) + FString::SanitizeFloat(Constant.X) + FString(TEXT(",Y=")) + FString::SanitizeFloat(Constant.Y) + FString(TEXT(")"));
 #endif
 	return DefaultValue;
 }
@@ -697,10 +688,7 @@ FString FVectorMaterialInput::GetDefaultValue() const
 {
 	FString DefaultValue;
 #if WITH_EDITOR
-	if (UseConstant)
-	{
-		DefaultValue = FString::SanitizeFloat(Constant.X) + FString(TEXT(",")) + FString::SanitizeFloat(Constant.Y) + FString(TEXT(",")) + FString::SanitizeFloat(Constant.Z);
-	}
+	DefaultValue = FString::SanitizeFloat(Constant.X) + FString(TEXT(",")) + FString::SanitizeFloat(Constant.Y) + FString(TEXT(",")) + FString::SanitizeFloat(Constant.Z);
 #endif
 	return DefaultValue;
 }
