Iris
- Do not warn if DataStreamChannel is not open

diff --git a/Engine/Source/Runtime/Engine/Private/Net/Experimental/Iris/DataStreamChannel.cpp b/Engine/Source/Runtime/Engine/Private/Net/Experimental/Iris/DataStreamChannel.cpp
--- a/Engine/Source/Runtime/Engine/Private/Net/Experimental/Iris/DataStreamChannel.cpp
+++ b/Engine/Source/Runtime/Engine/Private/Net/Experimental/Iris/DataStreamChannel.cpp
@@ -206,7 +206,7 @@ void UDataStreamChannel::Tick()
 	// We probably want separate bandwidth management for iris as we are not pre-filling sendbuffer before call to NetReady.
 	if (!IsNetReady(UE::Net::Private::bIrisSaturateBandwidth))
 	{
-		UE_LOG(LogIris, Warning, TEXT("Disallowed to write first object in batch, with Iris this is not good!"))
+		UE_CLOG(bHandshakeComplete, LogIris, Warning, TEXT("Disallowed to write first packet in batch, with Iris this is not good!"))
 		return;
 	}
 
