Remove legacy http test code.

diff --git a/Engine/Source/Runtime/Engine/Private/Tests/OSSTests.cpp b/Engine/Source/Runtime/Engine/Private/Tests/OSSTests.cpp
--- a/Engine/Source/Runtime/Engine/Private/Tests/OSSTests.cpp
+++ b/Engine/Source/Runtime/Engine/Private/Tests/OSSTests.cpp
@@ -18,16 +18,6 @@ DEFINE_LOG_CATEGORY_STATIC(LogHackAutomationTests, Log, All);
 //////////////////////////////////////////////////////////////////////////
 
 
-IMPLEMENT_SIMPLE_AUTOMATION_TEST(FHttpTest, "System.OSS.Test Http", EAutomationTestFlags::Disabled | EAutomationTestFlags::EditorContext | EAutomationTestFlags::EngineFilter)
-
-bool FHttpTest::RunTest(const FString& Parameters)
-{
-	ADD_LATENT_AUTOMATION_COMMAND(FEngineWaitLatentCommand(3.0));
-	ADD_LATENT_AUTOMATION_COMMAND(FExecStringLatentCommand(TEXT("http test 3 \"www.google.com\" ")));
-
-	return true;
-}
-
 IMPLEMENT_SIMPLE_AUTOMATION_TEST(FMCPFriendsTest, "System.OSS.MCP.Test Friends", EAutomationTestFlags::Disabled | EAutomationTestFlags::EditorContext | EAutomationTestFlags::EngineFilter)
 
 bool FMCPFriendsTest::RunTest(const FString& Parameters)
