Added log in NetDriver initialization telling what replication mode its set to.

diff --git a/Engine/Source/Runtime/Engine/Classes/Engine/NetDriver.h b/Engine/Source/Runtime/Engine/Classes/Engine/NetDriver.h
--- a/Engine/Source/Runtime/Engine/Classes/Engine/NetDriver.h
+++ b/Engine/Source/Runtime/Engine/Classes/Engine/NetDriver.h
@@ -2130,6 +2130,9 @@ private:
 	void PostDispatchSendUpdate();
 #endif
 
+	/** Description of the replication model used by this Driver (RepGraph, Iris or Generic) */
+	FString GetReplicationModelName() const;
+
 	void InitNetTraceId();
 
 	/** Called from RPC processing code to forward RPC to other NetDrivers if ShouldForwardFunction returns true. */
diff --git a/Engine/Source/Runtime/Engine/Private/NetDriver.cpp b/Engine/Source/Runtime/Engine/Private/NetDriver.cpp
--- a/Engine/Source/Runtime/Engine/Private/NetDriver.cpp
+++ b/Engine/Source/Runtime/Engine/Private/NetDriver.cpp
@@ -1493,26 +1493,14 @@ bool UNetDriver::InitBase(bool bInitAsClient, FNetworkNotify* InNotify, const FU
 	}
 #endif // UE_WITH_IRIS
 
-
+	
 	if (NetDriverDefinition == NAME_GameNetDriver)
 	{
-		static FString CrashContext_ReplicationDriver = TEXT("ReplicationDriver");
-
-		if (UReplicationDriver* RepDriver = GetReplicationDriver())
-		{
-			FGenericCrashContext::SetEngineData(CrashContext_ReplicationDriver, RepDriver->GetClass()->GetName());
-		}
-#if UE_WITH_IRIS
-		else if (ReplicationSystem)
-		{
-			FGenericCrashContext::SetEngineData(CrashContext_ReplicationDriver, TEXT("Iris"));
-		}
-#endif
-		else
-		{
-			FGenericCrashContext::SetEngineData(CrashContext_ReplicationDriver, TEXT("Generic"));
-		}
+		static FString CrashContext_ReplicationModel = TEXT("ReplicationModel");
+		FGenericCrashContext::SetEngineData(CrashContext_ReplicationModel, *GetReplicationModelName());
 	}
+
+	UE_LOG(LogNet, Log, TEXT("InitBase %s (NetDriverDefinition %s) using replication model %s"), *NetDriverName.ToString(), *NetDriverDefinition.ToString(), *GetReplicationModelName());
 	
 	InitNetTraceId();
 
@@ -1529,6 +1517,24 @@ bool UNetDriver::InitBase(bool bInitAsClient, FNetworkNotify* InNotify, const FU
 	return bSuccess;
 }
 
+FString UNetDriver::GetReplicationModelName() const
+{
+	if (UReplicationDriver* RepDriver = GetReplicationDriver())
+	{
+		return RepDriver->GetClass()->GetName();
+	}
+#if UE_WITH_IRIS
+	else if (ReplicationSystem)
+	{
+		return TEXT("Iris");
+	}
+#endif
+	else
+	{
+		return TEXT("Generic");
+	}
+}
+
 void UNetDriver::InitConnectionlessHandler()
 {
 	check(!ConnectionlessHandler.IsValid());
