Add a toggle r.AreShaderErrorsFatal which controls whether or not DefaultMaterial and global shader errors are fatal.

- This defaults to `True`, which is the current behavior.
- The idea is to turn off Fatal errors (and just Error) in situations where we want to see _all_ shader compile errors and not simply just stop when Default Material or global shaders fails.

diff --git a/Engine/Source/Runtime/Engine/Private/Materials/Material.cpp b/Engine/Source/Runtime/Engine/Private/Materials/Material.cpp
--- a/Engine/Source/Runtime/Engine/Private/Materials/Material.cpp
+++ b/Engine/Source/Runtime/Engine/Private/Materials/Material.cpp
@@ -2505,7 +2505,14 @@ namespace MaterialImpl
 			
 			if (bIsDefaultMaterial)
 			{
-				UE_ASSET_LOG(LogMaterial, Fatal, This, TEXT("%s"), *ErrorString);
+				if (AreShaderErrorsFatal())
+				{
+					UE_ASSET_LOG(LogMaterial, Fatal, This, TEXT("%s"), *ErrorString);
+				}
+				else
+				{
+					UE_ASSET_LOG(LogMaterial, Error, This, TEXT("%s"), *ErrorString);
+				}
 			}
 			else if (CVarMaterialLogErrorOnFailure.GetValueOnAnyThread())
 			{
diff --git a/Engine/Source/Runtime/Engine/Private/Materials/MaterialShared.cpp b/Engine/Source/Runtime/Engine/Private/Materials/MaterialShared.cpp
--- a/Engine/Source/Runtime/Engine/Private/Materials/MaterialShared.cpp
+++ b/Engine/Source/Runtime/Engine/Private/Materials/MaterialShared.cpp
@@ -2989,7 +2989,14 @@ bool FMaterial::CacheShaders(const FMaterialShaderMapId& ShaderMapId, EShaderPla
 					}
 
 					// Assert if the default material could not be compiled, since there will be nothing for other failed materials to fall back on.
-					UE_LOG(LogMaterial, Fatal,TEXT("Failed to compile default material %s!"), *GetFriendlyName());
+					if (AreShaderErrorsFatal())
+					{
+						UE_LOG(LogMaterial, Fatal, TEXT("Failed to compile default material %s!"), *GetFriendlyName());
+					}
+					else
+					{
+						UE_LOG(LogMaterial, Error, TEXT("Failed to compile default material %s!"), *GetFriendlyName());
+					}
 				}
 #endif // WITH_EDITOR
 			}
diff --git a/Engine/Source/Runtime/Engine/Private/ShaderCompiler/ShaderCompiler.cpp b/Engine/Source/Runtime/Engine/Private/ShaderCompiler/ShaderCompiler.cpp
--- a/Engine/Source/Runtime/Engine/Private/ShaderCompiler/ShaderCompiler.cpp
+++ b/Engine/Source/Runtime/Engine/Private/ShaderCompiler/ShaderCompiler.cpp
@@ -201,6 +201,18 @@ static TAutoConsoleVariable<bool> CVarShaderCompilerPerShaderDDCCook(
 	ECVF_Default
 );
 
+static TAutoConsoleVariable<bool> CVarAreShaderErrorsFatal(
+	TEXT("r.AreShaderErrorsFatal"),
+	true,
+	TEXT("When enabled, when a the default material or global shaders fail to compile it will issue a Fatal error.  Otherwise just an Error.\n")
+	TEXT("Default: true"),
+	ECVF_RenderThreadSafe);
+
+bool AreShaderErrorsFatal()
+{
+	return CVarAreShaderErrorsFatal.GetValueOnAnyThread();
+}
+
 static bool IsShaderJobCacheDDCRemotePolicyEnabled()
 {
 	return CVarJobCacheDDCPolicy.GetValueOnAnyThread();
@@ -6488,8 +6500,16 @@ void FShaderCompilingManager::ProcessCompiledShaderMaps(
 						}
 
 						ErrorString += FString::Printf(TEXT("Failed to compile default material %s!"), *Material->GetBaseMaterialPathName());
-						// Assert if a default material could not be compiled, since there will be nothing for other failed materials to fall back on.
-						UE_LOG(LogShaderCompilers, Fatal, TEXT("%s"), *ErrorString);
+
+						if (AreShaderErrorsFatal())
+						{
+							// Assert if a default material could not be compiled, since there will be nothing for other failed materials to fall back on.
+							UE_LOG(LogShaderCompilers, Fatal, TEXT("%s"), *ErrorString);
+						}
+						else
+						{
+							UE_LOG(LogShaderCompilers, Error, TEXT("%s"), *ErrorString);
+						}
 					}
 					
 					FString ErrorString;
@@ -10088,10 +10108,20 @@ static inline FShader* ProcessCompiledJob(FShaderCompileJob* SingleJob, const FS
 	}
 	else
 	{
-		UE_LOG(LogShaders, Fatal, TEXT("Failed to compile global shader %s %s %s.  Enable 'r.ShaderDevelopmentMode' in ConsoleVariables.ini for retries."),
-			GlobalShaderType->GetName(),
-			Pipeline ? TEXT("for pipeline") : TEXT(""),
-			Pipeline ? Pipeline->GetName() : TEXT(""));
+		if (AreShaderErrorsFatal())
+		{
+			UE_LOG(LogShaders, Fatal, TEXT("Failed to compile global shader %s %s %s.  Enable 'r.ShaderDevelopmentMode' in ConsoleVariables.ini for retries."),
+				GlobalShaderType->GetName(),
+				Pipeline ? TEXT("for pipeline") : TEXT(""),
+				Pipeline ? Pipeline->GetName() : TEXT(""));
+		}
+		else
+		{
+			UE_LOG(LogShaders, Error, TEXT("Failed to compile global shader %s %s %s.  Enable 'r.ShaderDevelopmentMode' in ConsoleVariables.ini for retries."),
+				GlobalShaderType->GetName(),
+				Pipeline ? TEXT("for pipeline") : TEXT(""),
+				Pipeline ? Pipeline->GetName() : TEXT(""));
+		}
 	}
 
 	return Shader;
diff --git a/Engine/Source/Runtime/Engine/Public/ShaderCompiler.h b/Engine/Source/Runtime/Engine/Public/ShaderCompiler.h
--- a/Engine/Source/Runtime/Engine/Public/ShaderCompiler.h
+++ b/Engine/Source/Runtime/Engine/Public/ShaderCompiler.h
@@ -45,6 +45,8 @@ struct FAnalyticsEventAttribute;
 
 #define DEBUG_INFINITESHADERCOMPILE 0
 
+bool AreShaderErrorsFatal();
+
 extern ENGINE_API bool IsShaderJobCacheDDCEnabled();
 extern ENGINE_API bool IsMaterialMapDDCEnabled();
 
