Texture.cpp commented out for testing, bump all loaded textures up to latest Oodle Texture Sdk ver

#rnx

diff --git a/Engine/Source/Runtime/Engine/Private/Texture.cpp b/Engine/Source/Runtime/Engine/Private/Texture.cpp
--- a/Engine/Source/Runtime/Engine/Private/Texture.cpp
+++ b/Engine/Source/Runtime/Engine/Private/Texture.cpp
@@ -1132,6 +1132,9 @@ void UTexture::PostLoad()
 	Super::PostLoad();
 
 #if WITH_EDITORONLY_DATA
+	///for testing, bump all loaded textures up to latest Oodle Texture Sdk ver:
+	//UpdateOodleTextureSdkVersionToLatest();
+
 	if (AssetImportData == nullptr)
 	{
 		AssetImportData = NewObject<UAssetImportData>(this, TEXT("AssetImportData"));
