Add missing call to increase number of materials translated count for cook stats.

diff --git a/Engine/Source/Runtime/Engine/Private/Materials/HLSLMaterialTranslator.cpp b/Engine/Source/Runtime/Engine/Private/Materials/HLSLMaterialTranslator.cpp
--- a/Engine/Source/Runtime/Engine/Private/Materials/HLSLMaterialTranslator.cpp
+++ b/Engine/Source/Runtime/Engine/Private/Materials/HLSLMaterialTranslator.cpp
@@ -1970,6 +1970,8 @@ bool FHLSLMaterialTranslator::Translate()
 #if STATS
 	GShaderCompilerStats->IncrementMaterialTranslateTime(HLSLTranslateTime);
 #endif // STATS
+	GShaderCompilerStats->IncrementMaterialsTranslated();
+	
 	INC_FLOAT_STAT_BY(STAT_ShaderCompiling_HLSLTranslation,(float)HLSLTranslateTime);
 
 #if ENABLE_COOK_STATS && STATS
