[ZenLoader] Make static analyzer happy

diff --git a/Engine/Source/Runtime/CoreUObject/Private/Serialization/AsyncLoading2.cpp b/Engine/Source/Runtime/CoreUObject/Private/Serialization/AsyncLoading2.cpp
--- a/Engine/Source/Runtime/CoreUObject/Private/Serialization/AsyncLoading2.cpp
+++ b/Engine/Source/Runtime/CoreUObject/Private/Serialization/AsyncLoading2.cpp
@@ -8999,7 +8999,12 @@ void FAsyncLoadingThread2::FlushLoading(TConstArrayView<int32> RequestIDs)
 	if (IsAsyncLoadingPackages())
 	{
 		// We can't possibly support flushing from async loading thread unless we have the partial request support active.
-		const bool bIsFlushSupportedOnCurrentThread = IsInGameThread() || (WITH_PARTIAL_REQUEST_DURING_RECURSION && IsInAsyncLoadingThread());
+#if WITH_PARTIAL_REQUEST_DURING_RECURSION
+		const bool bIsFlushSupportedOnCurrentThread = IsInGameThread() || IsInAsyncLoadingThread();
+#else
+		const bool bIsFlushSupportedOnCurrentThread = IsInGameThread();
+#endif
+
 		ensureMsgf(bIsFlushSupportedOnCurrentThread, TEXT("The current loader '%s' is unable to FlushAsyncLoading from the current thread."), *GetLoaderName().ToString());
 
 		if (!bIsFlushSupportedOnCurrentThread)
diff --git a/Engine/Source/Runtime/Engine/Private/Tests/Loading/AsyncLoadingTests_ConvertFromType.cpp b/Engine/Source/Runtime/Engine/Private/Tests/Loading/AsyncLoadingTests_ConvertFromType.cpp
--- a/Engine/Source/Runtime/Engine/Private/Tests/Loading/AsyncLoadingTests_ConvertFromType.cpp
+++ b/Engine/Source/Runtime/Engine/Private/Tests/Loading/AsyncLoadingTests_ConvertFromType.cpp
@@ -100,8 +100,11 @@ bool FConvertFromType_SoftToHard::RunTest(const FString& Parameters)
 		FlushAsyncLoading(RequestID);
 
 		TestTrue(TEXT("The object should have been properly loaded"), LoadedObject != nullptr);
-		TestTrue(TEXT("The hard-ref should now point to the object inside the second package"), LoadedObject->Reference.Get() != nullptr);
-		TestTrue(TEXT("The hard-ref should now point to the object inside the second package"), LoadedObject->Reference->GetPathName() == ObjectPath2);
+		if (LoadedObject)
+		{
+			TestTrue(TEXT("The hard-ref should now point to the object inside the second package"), LoadedObject->Reference.Get() != nullptr);
+			TestTrue(TEXT("The hard-ref should now point to the object inside the second package"), LoadedObject->Reference->GetPathName() == ObjectPath2);
+		}
 	}
 
 	// Cleanup phase
