Body setup should not create cooked collision data when bNeverNeedsCookedCollisionData is set.

diff --git a/Engine/Source/Runtime/Engine/Private/PhysicsEngine/BodySetup.cpp b/Engine/Source/Runtime/Engine/Private/PhysicsEngine/BodySetup.cpp
--- a/Engine/Source/Runtime/Engine/Private/PhysicsEngine/BodySetup.cpp
+++ b/Engine/Source/Runtime/Engine/Private/PhysicsEngine/BodySetup.cpp
@@ -1370,16 +1370,22 @@ FByteBulkData* UBodySetup::GetCookedData(FName Format, bool bRuntimeOnlyOptimize
 {
 	if (IsTemplate())
 	{
-		return NULL;
+		return nullptr;
+	}
+
+	// Geometry should never have collision data, cooked data will never be present
+	if (bNeverNeedsCookedCollisionData)
+	{
+		return nullptr;
 	}
 
 	IInterface_CollisionDataProvider* CDP = Cast<IInterface_CollisionDataProvider>(GetOuter());
 
 	// If there is nothing to cook or if we are reading data from a cooked package for an asset with no collision, 
 	// we want to return here
-	if ((AggGeom.ConvexElems.Num() == 0 && CDP == NULL) || !bHasCookedCollisionData)
+	if ((AggGeom.ConvexElems.Num() == 0 && CDP == nullptr) || !bHasCookedCollisionData)
 	{
-		return NULL;
+		return nullptr;
 	}
 
 #if WITH_EDITOR
@@ -1398,7 +1404,7 @@ FByteBulkData* UBodySetup::GetCookedData(FName Format, bool bRuntimeOnlyOptimize
 	{
 		SCOPE_CYCLE_COUNTER(STAT_PhysXCooking);
 
-		if (AggGeom.ConvexElems.Num() == 0 && (CDP == NULL || CDP->ContainsPhysicsTriMeshData(bMeshCollideAll) == false))
+		if (AggGeom.ConvexElems.Num() == 0 && (CDP == nullptr || CDP->ContainsPhysicsTriMeshData(bMeshCollideAll) == false))
 		{
 			return nullptr;
 		}
