Misc. fixes for the WIP FMaterialHLSLGenerator:

- Fix a bug where the black dummy texture is bound instead of the actual page table indirection texture when doing RVT samples with the new translator enabled
- Fix an issue where stale material resources can be used for rendering when toggling between the old and new translators
- Added missing view mip bias support to texture sampling code emission
- Move reused VT stack expression to a common parent scope instead of including the scope into the hash. This fixes the issue where VT sampling code uses VT addressing code from an inaccessible scope without unnecessarily duplicating VT stacks
- Properly set MATERIAL_VIRTUALTEXTURE_FEEDBACK
- Fix grass not showing up on terrains
- Landscape VF now works with the new translator
- Implemented GenerateHLSLExpression for UMaterialExpressionSwitch

diff --git a/Engine/Shaders/Private/LandscapeVertexFactory.ush b/Engine/Shaders/Private/LandscapeVertexFactory.ush
--- a/Engine/Shaders/Private/LandscapeVertexFactory.ush
+++ b/Engine/Shaders/Private/LandscapeVertexFactory.ush
@@ -559,6 +559,9 @@ FMaterialVertexParameters GetMaterialVertexParameters(FVertexFactoryInput Input,
 #endif // 1
 #endif // 0
 
+#if ENABLE_NEW_HLSL_GENERATOR
+	EvaluateVertexMaterialAttributes(Result);
+#endif
 	return Result;
 }
 
diff --git a/Engine/Source/Runtime/Engine/Classes/Materials/MaterialExpressionSwitch.h b/Engine/Source/Runtime/Engine/Classes/Materials/MaterialExpressionSwitch.h
--- a/Engine/Source/Runtime/Engine/Classes/Materials/MaterialExpressionSwitch.h
+++ b/Engine/Source/Runtime/Engine/Classes/Materials/MaterialExpressionSwitch.h
@@ -76,6 +76,8 @@ class UMaterialExpressionSwitch : public UMaterialExpression
 	virtual FExpressionInput* GetInput(int32 InputIndex) override;
 	virtual FName GetInputName(int32 InputIndex) const override;
 	virtual uint32 GetInputType(int32 InputIndex) override {return MCT_Unknown;}
+
+	virtual bool GenerateHLSLExpression(FMaterialHLSLGenerator& Generator, UE::HLSLTree::FScope& Scope, int32 OutputIndex, UE::HLSLTree::FExpression const*& OutExpression) const override;
 #endif // WITH_EDITOR
 	//~ End UMaterialExpression Interface
 };
diff --git a/Engine/Source/Runtime/Engine/Private/Materials/MaterialCachedData.cpp b/Engine/Source/Runtime/Engine/Private/Materials/MaterialCachedData.cpp
--- a/Engine/Source/Runtime/Engine/Private/Materials/MaterialCachedData.cpp
+++ b/Engine/Source/Runtime/Engine/Private/Materials/MaterialCachedData.cpp
@@ -26,6 +26,7 @@
 #include "Materials/MaterialExpressionNamedReroute.h"
 #include "Materials/MaterialExpressionFunctionInput.h"
 #include "Materials/MaterialExpressionStaticBool.h"
+#include "Materials/MaterialExpressionLandscapeGrassOutput.h"
 #include "Materials/MaterialFunctionInterface.h"
 #include "Materials/MaterialParameterCollection.h"
 #include "MaterialHLSLTree.h"
@@ -979,6 +980,17 @@ void FMaterialCachedExpressionData::UpdateForCachedHLSLTree(const FMaterialCache
 
 	::Private::PrepareHLSLTree(EmitContext, CachedTree, *this, SF_Pixel);
 	::Private::PrepareHLSLTree(EmitContext, CachedTree, *this, SF_Vertex);
+
+	for (const UMaterialExpressionCustomOutput* CustomOutput : CachedTree.GetMaterialCustomOutputs())
+	{
+		if (const UMaterialExpressionLandscapeGrassOutput* ExpressionGrassOutput = Cast<UMaterialExpressionLandscapeGrassOutput>(CustomOutput))
+		{
+			for (const FGrassInput& GrassInput : ExpressionGrassOutput->GrassTypes)
+			{
+				GrassTypes.AddUnique(GrassInput.GrassType);
+			}
+		}
+	}
 }
 
 void FMaterialCachedExpressionData::Validate()
diff --git a/Engine/Source/Runtime/Engine/Private/Materials/MaterialExpressionHLSL.cpp b/Engine/Source/Runtime/Engine/Private/Materials/MaterialExpressionHLSL.cpp
--- a/Engine/Source/Runtime/Engine/Private/Materials/MaterialExpressionHLSL.cpp
+++ b/Engine/Source/Runtime/Engine/Private/Materials/MaterialExpressionHLSL.cpp
@@ -3331,6 +3331,30 @@ bool UMaterialExpressionIf::GenerateHLSLExpression(FMaterialHLSLGenerator& Gener
 	return true;
 }
 
+bool UMaterialExpressionSwitch::GenerateHLSLExpression(FMaterialHLSLGenerator& Generator, UE::HLSLTree::FScope& Scope, int32 OutputIndex, UE::HLSLTree::FExpression const*& OutExpression) const
+{
+	using namespace UE::HLSLTree;
+
+	FTree& Tree = Generator.GetTree();
+	const FExpression* SwitchValueExpression = SwitchValue.AcquireHLSLExpressionOrConstant(Generator, Scope, ConstSwitchValue);
+	OutExpression = Default.AcquireHLSLExpressionOrConstant(Generator, Scope, ConstDefault);
+
+	for (int32 Index = Inputs.Num() - 1; Index >= 0; --Index)
+	{
+		const FExpressionInput& Input = Inputs[Index].Input;
+		const FExpression* InputExpression = Input.AcquireHLSLExpression(Generator, Scope);
+		if (!InputExpression)
+		{
+			return false;
+		}
+
+		const FExpression* ConditionExpression = Tree.NewLess(SwitchValueExpression, Tree.NewConstant(Index + 1.f));
+		OutExpression = Generator.GenerateBranch(Scope, ConditionExpression, InputExpression, OutExpression);
+	}
+
+	return true;
+}
+
 bool UMaterialExpressionFresnel::GenerateHLSLExpression(FMaterialHLSLGenerator& Generator, UE::HLSLTree::FScope& Scope, int32 OutputIndex, UE::HLSLTree::FExpression const*& OutExpression) const
 {
 	using namespace UE::HLSLTree;
diff --git a/Engine/Source/Runtime/Engine/Private/Materials/MaterialHLSLEmitter.cpp b/Engine/Source/Runtime/Engine/Private/Materials/MaterialHLSLEmitter.cpp
--- a/Engine/Source/Runtime/Engine/Private/Materials/MaterialHLSLEmitter.cpp
+++ b/Engine/Source/Runtime/Engine/Private/Materials/MaterialHLSLEmitter.cpp
@@ -490,6 +490,14 @@ static void GetMaterialEnvironment(EShaderPlatform InPlatform,
 	// Count the number of VTStacks (each stack will allocate a feedback slot)
 	OutEnvironment.SetDefine(TEXT("NUM_VIRTUALTEXTURE_SAMPLES"), EmitMaterialData.VTStacks.Num());
 
+	// Check if any feedback slots are in use. We can simplify shader and remove EARLYZ optimizations if none are.
+	bool bGenerateFeedback = false;
+	for (int32 i = 0; i < EmitMaterialData.VTStacks.Num() && !bGenerateFeedback; ++i)
+	{
+		bGenerateFeedback |= EmitMaterialData.VTStacks[i].bGenerateFeedback;
+	}
+	OutEnvironment.SetDefine(TEXT("MATERIAL_VIRTUALTEXTURE_FEEDBACK"), bGenerateFeedback);
+
 	// Setup defines to map each VT stack to either 1 or 2 page table textures, depending on how many layers it uses
 	for (int i = 0; i < EmitMaterialData.VTStacks.Num(); ++i)
 	{
diff --git a/Engine/Source/Runtime/Engine/Private/Materials/MaterialHLSLTree.cpp b/Engine/Source/Runtime/Engine/Private/Materials/MaterialHLSLTree.cpp
--- a/Engine/Source/Runtime/Engine/Private/Materials/MaterialHLSLTree.cpp
+++ b/Engine/Source/Runtime/Engine/Private/Materials/MaterialHLSLTree.cpp
@@ -925,7 +925,7 @@ uint32 AcquireVTStackIndex(
 	bool bGenerateFeedback)
 {
 	FHasher Hasher;
-	AppendHashes(Hasher, &Scope, MipValueMode, AddressU, AddressV, AspectRatio, EmitTexCoordValue, EmitTexCoordValueDdx, EmitTexCoordValueDdy, EmitMipValue, PreallocatedStackTextureIndex, bAdaptive, bGenerateFeedback);
+	AppendHashes(Hasher, MipValueMode, AddressU, AddressV, AspectRatio, EmitTexCoordValue, EmitTexCoordValueDdx, EmitTexCoordValueDdy, EmitMipValue, PreallocatedStackTextureIndex, bAdaptive, bGenerateFeedback);
 	const FXxHash64 Hash = Hasher.Finalize();
 
 	// First check to see if we have an existing VTStack that matches this key, that can still fit another layer
@@ -937,7 +937,6 @@ uint32 AcquireVTStackIndex(
 			Entry.EmitTexCoordValue == EmitTexCoordValue &&
 			Entry.EmitTexCoordValueDdx == EmitTexCoordValueDdx &&
 			Entry.EmitTexCoordValueDdy == EmitTexCoordValueDdy &&
-			Entry.Scope == &Scope &&
 			Entry.EmitMipValue == EmitMipValue &&
 			Entry.MipValueMode == MipValueMode &&
 			Entry.AddressU == AddressU &&
@@ -947,6 +946,7 @@ uint32 AcquireVTStackIndex(
 			Entry.bAdaptive == bAdaptive &&
 			Entry.bGenerateFeedback == bGenerateFeedback)
 		{
+			UE::HLSLTree::Private::MoveToScope(Entry.EmitResult, Scope);
 			return Index;
 		}
 	}
@@ -958,7 +958,7 @@ uint32 AcquireVTStackIndex(
 	Entry.EmitTexCoordValue = EmitTexCoordValue;
 	Entry.EmitTexCoordValueDdx = EmitTexCoordValueDdx;
 	Entry.EmitTexCoordValueDdy = EmitTexCoordValueDdy;
-	Entry.Scope = &Scope;
+	Entry.EmitMipValue = EmitMipValue;
 	Entry.MipValueMode = MipValueMode;
 	Entry.AddressU = AddressU;
 	Entry.AddressV = AddressV;
@@ -1140,15 +1140,47 @@ FEmitShaderExpression* EmitTextureSampleShader(
 		EmitTexCoordValueDdy = nullptr;
 	}
 
+	auto EmitManualMipViewBias = [&]()
+	{
+		if (MipValueMode == TMVM_Derivative)
+		{
+			// When doing derivative based sampling, multiply.
+			FEmitShaderExpression* EmitMultiplier = Context.EmitInlineExpression(Scope, Shader::EValueType::Float1, TEXT("View.MaterialTextureDerivativeMultiply"));
+			EmitTexCoordValueDdx = Context.EmitExpression(Scope, Shader::EValueType::Float2, TEXT("(% * %)"), EmitTexCoordValueDdx, EmitMultiplier);
+			EmitTexCoordValueDdy = Context.EmitExpression(Scope, Shader::EValueType::Float2, TEXT("(% * %)"), EmitTexCoordValueDdy, EmitMultiplier);
+		}
+		else if (MipValueMode == TMVM_MipLevel || MipValueMode == TMVM_MipBias)
+		{
+			check(EmitMipValue);
+			// Adds bias to existing input level bias.
+			EmitMipValue = Context.EmitExpression(Scope, Shader::EValueType::Float1, TEXT("(% + View.MaterialTextureMipBias)"), EmitMipValue);
+		}
+		else
+		{
+			// Sets bias.
+			EmitMipValue = Context.EmitInlineExpression(Scope, Shader::EValueType::Float1, TEXT("View.MaterialTextureMipBias"));
+		}
+
+		// If no Mip mode, then use MipBias.
+		MipValueMode = MipValueMode == TMVM_None ? TMVM_MipBias : MipValueMode;
+	};
+
 	FEmitShaderExpression* EmitTextureResult = nullptr;
 	if (TextureType == MCT_TextureVirtual)
 	{
+		// VT does not have explicit samplers (and always requires manual view mip bias)
+		if (bAutomaticViewMipBias)
+		{
+			EmitManualMipViewBias();
+		}
+
 		FEmitData& EmitMaterialData = Context.FindData<FEmitData>();
 
 		FMaterialTextureParameterInfo TextureParameterInfo;
 		TextureParameterInfo.ParameterInfo = TextureValue.ParameterInfo;
 		TextureParameterInfo.TextureIndex = Context.Material->GetReferencedTextures().Find(TextureObject);
 		TextureParameterInfo.SamplerSource = SamplerSource;
+		TextureParameterInfo.VirtualTextureLayerIndex = TextureLayerIndex;
 		check(TextureParameterInfo.TextureIndex != INDEX_NONE);
 		const int32 TextureParameterIndex = Context.MaterialCompilationOutput->UniformExpressionSet.FindOrAddTextureParameter(EMaterialTextureParameterType::Virtual, TextureParameterInfo);
 
@@ -1263,6 +1295,7 @@ FEmitShaderExpression* EmitTextureSampleShader(
 		Private::EmitTextureShader(Context, TextureValue, FormattedTexture);
 
 		TStringBuilder<256> FormattedSampler;
+		bool bRequiresManualViewMipBias = bAutomaticViewMipBias;
 		switch (SamplerSource)
 		{
 		case SSM_FromTextureAsset:
@@ -1272,22 +1305,30 @@ FEmitShaderExpression* EmitTextureSampleShader(
 			FormattedSampler.Appendf(TEXT("GetMaterialSharedSampler(%sSampler,%s)"),
 				FormattedTexture.ToString(),
 				bAutomaticViewMipBias ? TEXT("View.MaterialTextureBilinearWrapedSampler") : TEXT("Material.Wrap_WorldGroupSettings"));
+			bRequiresManualViewMipBias = false;
 			break;
 		case SSM_Clamp_WorldGroupSettings:
 			FormattedSampler.Appendf(TEXT("GetMaterialSharedSampler(%sSampler,%s)"),
 				FormattedTexture.ToString(),
 				bAutomaticViewMipBias ? TEXT("View.MaterialTextureBilinearClampedSampler") : TEXT("Material.Clamp_WorldGroupSettings"));
+			bRequiresManualViewMipBias = false;
 			break;
 		case SSM_TerrainWeightmapGroupSettings:
 			FormattedSampler.Appendf(TEXT("GetMaterialSharedSampler(%sSampler,%s)"),
 				FormattedTexture.ToString(),
 				TEXT("View.LandscapeWeightmapSampler"));
+			bRequiresManualViewMipBias = false;
 			break;
 		default:
 			checkNoEntry();
 			break;
 		}
 
+		if (bRequiresManualViewMipBias)
+		{
+			EmitManualMipViewBias();
+		}
+
 		switch (MipValueMode)
 		{
 		case TMVM_Derivative:
diff --git a/Engine/Source/Runtime/Engine/Private/Materials/MaterialInstance.cpp b/Engine/Source/Runtime/Engine/Private/Materials/MaterialInstance.cpp
--- a/Engine/Source/Runtime/Engine/Private/Materials/MaterialInstance.cpp
+++ b/Engine/Source/Runtime/Engine/Private/Materials/MaterialInstance.cpp
@@ -2273,6 +2273,29 @@ void UMaterialInstance::CacheResourceShadersForRendering(EMaterialShaderPrecompi
 			FMaterialResource* CurrentResource = FindOrCreateMaterialResource(StaticPermutationMaterialResources, BaseMaterial, this, FeatureLevel, ActiveQualityLevel);
 			check(CurrentResource);
 
+			if (IsUsingNewHLSLGenerator())
+			{
+				// Release resources from unused qualities. For some reason, FindOrCreateMaterialResource checks material quality usage
+				// but FindMaterialResource doesn't. The two functions can choose differently if unused quality resources aren't removed.
+				// When that happens, stale material resources may be used for rendering and cause troubles
+				TArray<bool, TInlineAllocator<EMaterialQualityLevel::Num>> QualityLevelsUsed;
+				GetQualityLevelUsage(QualityLevelsUsed, ShaderPlatform);
+
+				for (int32 Index = 0; Index < StaticPermutationMaterialResources.Num(); ++Index)
+				{
+					FMaterialResource* MaterialResource = StaticPermutationMaterialResources[Index];
+					if (MaterialResource != CurrentResource
+						&& MaterialResource->GetFeatureLevel() == FeatureLevel
+						&& MaterialResource->GetQualityLevel() != EMaterialQualityLevel::Num
+						&& !QualityLevelsUsed[MaterialResource->GetQualityLevel()])
+					{
+						OutResourcesToFree.Add(MaterialResource);
+						StaticPermutationMaterialResources.RemoveAtSwap(Index);
+						--Index;
+					}
+				}
+			}
+
 #if STORE_ONLY_ACTIVE_SHADERMAPS
 			if (!CurrentResource->GetGameThreadShaderMap())
 			{
diff --git a/Engine/Source/Runtime/Engine/Private/Materials/MaterialShared.cpp b/Engine/Source/Runtime/Engine/Private/Materials/MaterialShared.cpp
--- a/Engine/Source/Runtime/Engine/Private/Materials/MaterialShared.cpp
+++ b/Engine/Source/Runtime/Engine/Private/Materials/MaterialShared.cpp
@@ -3932,14 +3932,15 @@ void FMaterialVirtualTextureStack::GetTextureValues(const FMaterialRenderContext
 void FMaterialVirtualTextureStack::GetTextureValue(const FMaterialRenderContext& Context, const FUniformExpressionSet& UniformExpressionSet, const URuntimeVirtualTexture*& OutValue) const
 {
 	OutValue = nullptr;
-	if (NumLayers > 0)
+	for (uint32 LayerIndex = 0u; LayerIndex < NumLayers; ++LayerIndex)
 	{
-		const int32 ParameterIndex = LayerUniformExpressionIndices[0];
+		const int32 ParameterIndex = LayerUniformExpressionIndices[LayerIndex];
 		if (ParameterIndex != INDEX_NONE)
 		{
 			const URuntimeVirtualTexture* Texture = nullptr;
 			UniformExpressionSet.GetTextureValue(ParameterIndex, Context, Context.Material, Texture);
 			OutValue = Texture;
+			break;
 		}
 	}
 }
diff --git a/Engine/Source/Runtime/Engine/Public/HLSLTree/HLSLTreeEmit.h b/Engine/Source/Runtime/Engine/Public/HLSLTree/HLSLTreeEmit.h
--- a/Engine/Source/Runtime/Engine/Public/HLSLTree/HLSLTreeEmit.h
+++ b/Engine/Source/Runtime/Engine/Public/HLSLTree/HLSLTreeEmit.h
@@ -640,6 +640,8 @@ struct FEmitOwnerScope
 
 namespace Private
 {
+void MoveToScope(FEmitShaderNode* EmitNode, FEmitScope& Scope);
+
 void EmitPreshaderField(
 	FEmitContext& Context,
 	TMemoryImageArray<FMaterialUniformPreshaderHeader>& UniformPreshaders,
diff --git a/Engine/Source/Runtime/Engine/Public/MaterialCachedHLSLTree.h b/Engine/Source/Runtime/Engine/Public/MaterialCachedHLSLTree.h
--- a/Engine/Source/Runtime/Engine/Public/MaterialCachedHLSLTree.h
+++ b/Engine/Source/Runtime/Engine/Public/MaterialCachedHLSLTree.h
@@ -70,6 +70,8 @@ public:
 	const UE::Shader::FStructType* GetMaterialAttributesType() const { return MaterialAttributesType; }
 	const UE::Shader::FStructType* GetVTPageTableResultType() const { return VTPageTableResultType; }
 
+	const TArray<UMaterialExpressionCustomOutput*>& GetMaterialCustomOutputs() const { return MaterialCustomOutputs; }
+
 	const TMap<FMaterialConnectionKey, const UE::HLSLTree::FExpression*>& GetConnections() const { return ConnectionMap; }
 
 	ENGINE_API void SetRequestedFields(const UE::HLSLTree::FEmitContext& Context, UE::HLSLTree::FRequestedType& OutRequestedType) const;
diff --git a/Engine/Source/Runtime/Engine/Public/MaterialHLSLTree.h b/Engine/Source/Runtime/Engine/Public/MaterialHLSLTree.h
--- a/Engine/Source/Runtime/Engine/Public/MaterialHLSLTree.h
+++ b/Engine/Source/Runtime/Engine/Public/MaterialHLSLTree.h
@@ -955,7 +955,6 @@ struct FVTStackEntry
 	FEmitShaderExpression* EmitTexCoordValueDdx;
 	FEmitShaderExpression* EmitTexCoordValueDdy;
 	FEmitShaderExpression* EmitMipValue;
-	FEmitScope* Scope;
 	FEmitShaderExpression* EmitResult;
 	ETextureMipValueMode MipValueMode;
 	TextureAddress AddressU;
