World Partition RuntimeHashSet HLODs: make sure to set the bIsHLOD flag on HLOD streaming cells.

diff --git a/Engine/Source/Runtime/Engine/Private/WorldPartition/RuntimeHashSet/WorldPartitionRuntimeHashSetStreamingGeneration.cpp b/Engine/Source/Runtime/Engine/Private/WorldPartition/RuntimeHashSet/WorldPartitionRuntimeHashSetStreamingGeneration.cpp
--- a/Engine/Source/Runtime/Engine/Private/WorldPartition/RuntimeHashSet/WorldPartitionRuntimeHashSetStreamingGeneration.cpp
+++ b/Engine/Source/Runtime/Engine/Private/WorldPartition/RuntimeHashSet/WorldPartitionRuntimeHashSetStreamingGeneration.cpp
@@ -178,7 +178,7 @@ bool UWorldPartitionRuntimeHashSet::GenerateStreaming(UWorldPartitionStreamingPo
 		RuntimeCell->SetPriority(CellDescInstance.Priority);
 		RuntimeCell->SetClientOnlyVisible(CellDescInstance.bClientOnlyVisible);
 		RuntimeCell->SetBlockOnSlowLoading(CellDescInstance.bBlockOnSlowStreaming);
-		RuntimeCell->SetIsHLOD(false);
+		RuntimeCell->SetIsHLOD(CellDescInstance.SourcePartition->bIsHLODSetup);
 		RuntimeCell->SetGuid(CellUniqueId.Guid);
 
 		UWorldPartitionRuntimeCellDataSpatialHashSet* RuntimeCellDataHashSet = CastChecked<UWorldPartitionRuntimeCellDataSpatialHashSet>(RuntimeCell->RuntimeCellData);
