Fix for very long iteration times when resizing landscapes with foliage instances

The main problem was that the cell count calculation was overflowing, when passed WORLD_MAX/MIN as input bounds -- so it thought that iterating the entire space was faster than iterating the populated cells.
I also modified the query code to not query WORLD_MAX/MIN on Z, but to instead expand the world bounds by a good amount, which helps it take more optimal paths by limiting the query volume.

diff --git a/Engine/Source/Editor/LandscapeEditor/Private/LandscapeEdMode.cpp b/Engine/Source/Editor/LandscapeEditor/Private/LandscapeEdMode.cpp
--- a/Engine/Source/Editor/LandscapeEditor/Private/LandscapeEdMode.cpp
+++ b/Engine/Source/Editor/LandscapeEditor/Private/LandscapeEdMode.cpp
@@ -3810,9 +3810,14 @@ ALandscape* FEdModeLandscape::ChangeComponentSetting(int32 NumComponentsX, int32
 
 							if (NewCollisionComponent && FBoxSphereBounds::BoxesIntersect(NewCollisionComponent->Bounds, OldCollisionComponent->Bounds))
 							{
+								// only transfer instances overlapping the new box in x,y
 								FBox Box = NewCollisionComponent->Bounds.GetBox();
-								Box.Min.Z = -WORLD_MAX;
-								Box.Max.Z = WORLD_MAX;
+								FBox OldBox = OldCollisionComponent->Bounds.GetBox();
+
+								// but allow just about any Z (expand old bounds by max extent)
+								double Extent = OldBox.GetExtent().GetMax();
+								Box.Min.Z = OldBox.Min.Z - Extent;
+								Box.Max.Z = OldBox.Max.Z + Extent;
 
 								AInstancedFoliageActor::MoveInstancesToNewComponent(World, OldCollisionComponent, Box, NewCollisionComponent);
 							}
diff --git a/Engine/Source/Runtime/Engine/Public/Instances/InstancedPlacementHash.h b/Engine/Source/Runtime/Engine/Public/Instances/InstancedPlacementHash.h
--- a/Engine/Source/Runtime/Engine/Public/Instances/InstancedPlacementHash.h
+++ b/Engine/Source/Runtime/Engine/Public/Instances/InstancedPlacementHash.h
@@ -176,15 +176,20 @@ public:
 		FKey MinKey = MakeKey(InBox.Min);
 		FKey MaxKey = MakeKey(InBox.Max);
 
-		int64 CellCount = (MaxKey.Z - MinKey.Z + 1) * (MaxKey.Y - MinKey.Y + 1) * (MaxKey.X - MinKey.X + 1);
+		// compute in doubles to avoid overflow issues (queries using WORLD_MAX bounds can produce very large numbers...)
+		double CellsX = (MaxKey.X - MinKey.X + 1);
+		double CellsY = (MaxKey.Y - MinKey.Y + 1);
+		double CellsZ = (MaxKey.Z - MinKey.Z + 1);
+		double CellCount = CellsX * CellsY * CellsZ;
 
 		// The idea here is to decide when it is faster to just check every populated cell.
 		// The actual ideal threshold will depend on the exact state of the cells, but we just pick
 		// some arbitrary ratio where we switch to checking populated cells over potential cells within range.
 		// In practice the exact threshold is not that important; we just want to know if there are different orders of magnitude involved.
 		// (i.e. it is faster to check 32 populated cells than 3,000,000,000 potential cells...)
-		constexpr int32 RelativeCostOfCheckingPopulatedCells = 2;
-		if (CellCount > RelativeCostOfCheckingPopulatedCells * CellMap.Num())
+		constexpr double RelativeCostOfCheckingPopulatedCells = 2.0;
+		double Threshold = RelativeCostOfCheckingPopulatedCells * CellMap.Num();
+		if (CellCount > Threshold)
 		{
 			// check every populated cell
 			for (auto& Pair : CellMap)
diff --git a/Engine/Source/Runtime/Landscape/Private/LandscapeConfigHelper.cpp b/Engine/Source/Runtime/Landscape/Private/LandscapeConfigHelper.cpp
--- a/Engine/Source/Runtime/Landscape/Private/LandscapeConfigHelper.cpp
+++ b/Engine/Source/Runtime/Landscape/Private/LandscapeConfigHelper.cpp
@@ -383,9 +383,14 @@ void FLandscapeConfigHelper::MoveFoliageToLandscape(ULandscapeInfo* InLandscapeI
 
 				if (NewCollisionComponent && FBoxSphereBounds::BoxesIntersect(NewCollisionComponent->Bounds, OldCollisionComponent->Bounds))
 				{
+					// only transfer instances overlapping the new box in x,y
 					FBox Box = NewCollisionComponent->Bounds.GetBox();
-					Box.Min.Z = -WORLD_MAX;
-					Box.Max.Z = WORLD_MAX;
+					FBox OldBox = OldCollisionComponent->Bounds.GetBox();
+
+					// but allow just about any Z (expand old bounds by max extent)
+					double Extent = OldBox.GetExtent().GetMax();
+					Box.Min.Z = OldBox.Min.Z - Extent;
+					Box.Max.Z = OldBox.Max.Z + Extent;
 
 					AInstancedFoliageActor::MoveInstancesToNewComponent(World, OldCollisionComponent, Box, NewCollisionComponent);
 				}
