Add a flag to BodyInstance to enable/disable the initial-overlap depenetration system on a per-body basis.

diff --git a/Engine/Source/Editor/DetailCustomizations/Private/BodyInstanceCustomization.cpp b/Engine/Source/Editor/DetailCustomizations/Private/BodyInstanceCustomization.cpp
--- a/Engine/Source/Editor/DetailCustomizations/Private/BodyInstanceCustomization.cpp
+++ b/Engine/Source/Editor/DetailCustomizations/Private/BodyInstanceCustomization.cpp
@@ -1265,6 +1265,7 @@ void FBodyInstanceCustomizationHelper::UpdateFilters()
 	bDisplayAngularDamping = true;
 	bDisplayEnableGravity = true;
 	bDisplayInertiaConditioning = true;
+	bDisplayInitialOverlapDepenetration = true;
 	bDisplayWalkableSlopeOverride = true;
 	bDisplayAutoWeld = true;
 	bDisplayStartAwake = true;
@@ -1291,6 +1292,7 @@ void FBodyInstanceCustomizationHelper::UpdateFilters()
 				bDisplayAngularDamping = false;
 				bDisplayEnableGravity = false;
 				bDisplayInertiaConditioning = false;
+				bDisplayInitialOverlapDepenetration = false;
 				bDisplayWalkableSlopeOverride = false;
 				bDisplayAutoWeld = false;
 				bDisplayStartAwake = false;
@@ -1354,6 +1356,7 @@ void FBodyInstanceCustomizationHelper::CustomizeDetails( IDetailLayoutBuilder& D
 		EnablePhysicsProperty(GET_MEMBER_NAME_CHECKED(FBodyInstance, AngularDamping), bDisplayAngularDamping);
 		EnablePhysicsProperty(GET_MEMBER_NAME_CHECKED(FBodyInstance, bEnableGravity), bDisplayEnableGravity);
 		EnablePhysicsProperty(GET_MEMBER_NAME_CHECKED(FBodyInstance, bInertiaConditioning), bDisplayInertiaConditioning);
+		EnablePhysicsProperty(GET_MEMBER_NAME_CHECKED(FBodyInstance, bInitialOverlapDepenetration), bDisplayInitialOverlapDepenetration);
 		EnablePhysicsProperty(GET_MEMBER_NAME_CHECKED(FBodyInstance, WalkableSlopeOverride), bDisplayWalkableSlopeOverride);
 
 		AddBodyConstraint(PhysicsCategory, BodyInstanceHandler);
diff --git a/Engine/Source/Editor/DetailCustomizations/Private/BodyInstanceCustomization.h b/Engine/Source/Editor/DetailCustomizations/Private/BodyInstanceCustomization.h
--- a/Engine/Source/Editor/DetailCustomizations/Private/BodyInstanceCustomization.h
+++ b/Engine/Source/Editor/DetailCustomizations/Private/BodyInstanceCustomization.h
@@ -193,6 +193,7 @@ private:
 	bool bDisplayAngularDamping;
 	bool bDisplayEnableGravity;
 	bool bDisplayInertiaConditioning;
+	bool bDisplayInitialOverlapDepenetration;
 	bool bDisplayWalkableSlopeOverride;
 	bool bDisplayAutoWeld;
 	bool bDisplayStartAwake;
diff --git a/Engine/Source/Runtime/Engine/Classes/PhysicsEngine/BodyInstance.h b/Engine/Source/Runtime/Engine/Classes/PhysicsEngine/BodyInstance.h
--- a/Engine/Source/Runtime/Engine/Classes/PhysicsEngine/BodyInstance.h
+++ b/Engine/Source/Runtime/Engine/Classes/PhysicsEngine/BodyInstance.h
@@ -440,6 +440,15 @@ protected:
 	UPROPERTY(EditAnywhere, AdvancedDisplay, BlueprintReadWrite, Category = Physics)
 	uint8 bInertiaConditioning : 1;
 
+	/**
+	 * @brief Enable the initial-overlap de-penetration mechanism
+	 * 
+	 * When enabled, initial overlaps are de-penetrated at a speed that is limited by the physics project configuration.
+	 * If disabled, initial overlaps are not tracked and bodies will attempt to de-penetrate in one tick which can lead to large separating velocities.
+	 */
+	UPROPERTY(EditAnywhere, AdvancedDisplay, BlueprintReadWrite, Category = Physics)
+	uint8 bInitialOverlapDepenetration : 1;
+
 
 public:
 	/** Current scale of physics - used to know when and how physics must be rescaled to match current transform of OwnerComponent. */
@@ -812,6 +821,12 @@ public:
 	/** If set to true, this body will treat bodies that do not have the flag set as having infinite mass */
 	ENGINE_API void SetOneWayInteraction(bool InOneWayInteraction = true);
 
+	/** Is the controlled inital-overlap depenetration system enabled for this body */
+	ENGINE_API bool IsInitialOverlapDepenetrationEnabled() const;
+
+	/** Enable/disable the controlled inital-overlap depenetration system for this body */
+	ENGINE_API void SetInitialOverlapDepenetrationEnabled(bool bInEnabled);
+
 	/** Add a torque to this body */
 	ENGINE_API void AddTorqueInRadians(const FVector& Torque, bool bAllowSubstepping = true, bool bAccelChange = false, const FAsyncPhysicsTimestamp TimeStamp = FAsyncPhysicsTimestamp(), APlayerController* PlayerController = nullptr);
 	/** Clear accumulated torques on this body */
diff --git a/Engine/Source/Runtime/Engine/Private/PhysicsEngine/BodyInstance.cpp b/Engine/Source/Runtime/Engine/Private/PhysicsEngine/BodyInstance.cpp
--- a/Engine/Source/Runtime/Engine/Private/PhysicsEngine/BodyInstance.cpp
+++ b/Engine/Source/Runtime/Engine/Private/PhysicsEngine/BodyInstance.cpp
@@ -336,6 +336,7 @@ FBodyInstance::FBodyInstance()
 	, bInterpolateWhenSubStepping(true)
 	, bPendingCollisionProfileSetup(false)
 	, bInertiaConditioning(true)
+	, bInitialOverlapDepenetration(true)
 	, Scale3D(1.0f)
 	, CollisionProfileName(UCollisionProfile::CustomCollisionProfileName)
 	, PositionSolverIterationCount(8)
@@ -1210,8 +1211,9 @@ void FInitBodiesHelperBase::CreateActor_AssumesLocked(FBodyInstance* Instance, c
 		FPhysicsInterface::SetMaxLinearVelocity_AssumesLocked(Instance->ActorHandle, TNumericLimits<float>::Max());
 		FPhysicsInterface::SetSmoothEdgeCollisionsEnabled_AssumesLocked(Instance->ActorHandle, Instance->bSmoothEdgeCollisions);
 		FPhysicsInterface::SetInertiaConditioningEnabled_AssumesLocked(Instance->ActorHandle, Instance->bInertiaConditioning);
+		FPhysicsInterface::SetInitialOverlapDepenetrationEnabled_AssumesLocked(Instance->ActorHandle, Instance->bInitialOverlapDepenetration);
 
-		// Set sleep even notification
+		// Set sleep event notification
 		FPhysicsInterface::SetSendsSleepNotifies_AssumesLocked(Instance->ActorHandle, Instance->bGenerateWakeEvents);
 	}
 }
@@ -3381,6 +3383,27 @@ void FBodyInstance::SetOneWayInteraction(bool InOneWayInteraction /*= true*/)
 		});
 }
 
+bool FBodyInstance::IsInitialOverlapDepenetrationEnabled() const
+{
+	return bInitialOverlapDepenetration;
+}
+
+void FBodyInstance::SetInitialOverlapDepenetrationEnabled(bool bInEnabled)
+{
+	if (bInEnabled != bInitialOverlapDepenetration)
+	{
+		bInitialOverlapDepenetration = bInEnabled;
+
+		FPhysicsCommand::ExecuteWrite(ActorHandle, [bInEnabled](const FPhysicsActorHandle& Actor)
+			{
+				if (FChaosEngineInterface::IsValid(Actor))
+				{
+					Actor->GetGameThreadAPI().SetInitialOverlapDepenetrationEnabled(bInEnabled);
+				}
+			});
+	}
+}
+
 void FBodyInstance::AddTorqueInRadians(const FVector& Torque, bool bAllowSubstepping, bool bAccelChange, const FAsyncPhysicsTimestamp TimeStamp, APlayerController* PlayerController)
 {
 	const bool bIsInternal = TimeStamp.IsValid();
diff --git a/Engine/Source/Runtime/Experimental/Chaos/Public/PhysicsProxy/SingleParticlePhysicsProxy.h b/Engine/Source/Runtime/Experimental/Chaos/Public/PhysicsProxy/SingleParticlePhysicsProxy.h
--- a/Engine/Source/Runtime/Experimental/Chaos/Public/PhysicsProxy/SingleParticlePhysicsProxy.h
+++ b/Engine/Source/Runtime/Experimental/Chaos/Public/PhysicsProxy/SingleParticlePhysicsProxy.h
@@ -458,6 +458,30 @@ public:
 		});
 	}
 
+	bool InitialOverlapDepenetrationEnabled() const
+	{
+		return Read([](auto* Particle)
+			{
+				if (auto Rigid = Particle->CastToRigidParticle())
+				{
+					return Rigid->InitialOverlapDepenetrationEnabled();
+				}
+
+				return true;
+			});
+	}
+
+	void SetInitialOverlapDepenetrationEnabled(const bool bEnabled)
+	{
+		Write([bEnabled](auto* Particle)
+			{
+				if (auto Rigid = Particle->CastToRigidParticle())
+				{
+					return Rigid->SetInitialOverlapDepenetrationEnabled(bEnabled);
+				}
+			});
+	}
+
 	bool InertiaConditioningEnabled() const
 	{
 		return Read([](auto* Particle)
diff --git a/Engine/Source/Runtime/PhysicsCore/Private/ChaosEngineInterface.cpp b/Engine/Source/Runtime/PhysicsCore/Private/ChaosEngineInterface.cpp
--- a/Engine/Source/Runtime/PhysicsCore/Private/ChaosEngineInterface.cpp
+++ b/Engine/Source/Runtime/PhysicsCore/Private/ChaosEngineInterface.cpp
@@ -1193,6 +1193,16 @@ void FChaosEngineInterface::SetInertiaConditioningEnabled_AssumesLocked(const FP
 	InActorReference->GetGameThreadAPI().SetInertiaConditioningEnabled(bEnabled);
 }
 
+bool FChaosEngineInterface::IsInitialOverlapDepenetrationEnabled_AssumesLocked(const FPhysicsActorHandle& InActorReference)
+{
+	return InActorReference->GetGameThreadAPI().InitialOverlapDepenetrationEnabled();
+}
+
+void FChaosEngineInterface::SetInitialOverlapDepenetrationEnabled_AssumesLocked(const FPhysicsActorHandle& InActorReference, bool bEnabled)
+{
+	InActorReference->GetGameThreadAPI().SetInitialOverlapDepenetrationEnabled(bEnabled);
+}
+
 void FChaosEngineInterface::SetIsSimulationShape(const FPhysicsShapeHandle& InShape,bool bIsSimShape)
 {
 	InShape.Shape->SetSimEnabled(bIsSimShape);
diff --git a/Engine/Source/Runtime/PhysicsCore/Public/Chaos/ChaosEngineInterface.h b/Engine/Source/Runtime/PhysicsCore/Public/Chaos/ChaosEngineInterface.h
--- a/Engine/Source/Runtime/PhysicsCore/Public/Chaos/ChaosEngineInterface.h
+++ b/Engine/Source/Runtime/PhysicsCore/Public/Chaos/ChaosEngineInterface.h
@@ -472,6 +472,9 @@ public:
 	static PHYSICSCORE_API bool IsInertiaConditioningEnabled_AssumesLocked(const FPhysicsActorHandle& InActorReference);
 	static PHYSICSCORE_API void SetInertiaConditioningEnabled_AssumesLocked(const FPhysicsActorHandle& InActorReference, bool bEnabled);
 
+	static PHYSICSCORE_API bool IsInitialOverlapDepenetrationEnabled_AssumesLocked(const FPhysicsActorHandle& InActorReference);
+	static PHYSICSCORE_API void SetInitialOverlapDepenetrationEnabled_AssumesLocked(const FPhysicsActorHandle& InActorReference, bool bEnabled);
+
 	static PHYSICSCORE_API float GetStabilizationEnergyThreshold_AssumesLocked(const FPhysicsActorHandle& InHandle);
 	static PHYSICSCORE_API void SetStabilizationEnergyThreshold_AssumesLocked(const FPhysicsActorHandle& InHandle,float InThreshold);
 	static PHYSICSCORE_API uint32 GetSolverPositionIterationCount_AssumesLocked(const FPhysicsActorHandle& InHandle);
