Add client on-screen messages when server captures are in progress.

diff --git a/Engine/Source/Runtime/Engine/Private/NetworkProfiler.cpp b/Engine/Source/Runtime/Engine/Private/NetworkProfiler.cpp
--- a/Engine/Source/Runtime/Engine/Private/NetworkProfiler.cpp
+++ b/Engine/Source/Runtime/Engine/Private/NetworkProfiler.cpp
@@ -689,9 +689,6 @@ void FNetworkProfiler::TrackSessionChange( bool bShouldContinueTracking, const F
 			uint8 Type = NPTYPE_EndOfStreamMarker;
 			( *FileWriter ) << Type;
 
-			// Close file writer so we can rename the file to its final destination.
-			FileWriter->Close();
-
 			if (OnNetworkProfileFinished().IsBound())
 			{
 				OnNetworkProfileFinished().Broadcast(FileWriter->GetArchiveName());
@@ -741,6 +738,11 @@ void FNetworkProfiler::TrackSessionChange( bool bShouldContinueTracking, const F
 
 			//Mark that tracking truly is enabled now
 			bIsTrackingEnabled = bShouldTrackingBeEnabled = true;
+
+			if (OnNetworkProfileStarted().IsBound())
+			{
+				OnNetworkProfileStarted().Broadcast(FileWriter->GetArchiveName());
+			}
 		}
 		else
 		{
diff --git a/Engine/Source/Runtime/Engine/Public/Net/NetworkProfiler.h b/Engine/Source/Runtime/Engine/Public/Net/NetworkProfiler.h
--- a/Engine/Source/Runtime/Engine/Public/Net/NetworkProfiler.h
+++ b/Engine/Source/Runtime/Engine/Public/Net/NetworkProfiler.h
@@ -30,6 +30,7 @@ CSV_DECLARE_CATEGORY_EXTERN(NetworkProfiler);
 		x; \
 	}
 
+DECLARE_MULTICAST_DELEGATE_OneParam(FOnNetworkProfileStarted, const FString& /*Filename */);
 DECLARE_MULTICAST_DELEGATE_OneParam(FOnNetworkProfileFinished, const FString& /*Filename */);
 
 enum class ENetworkProfilerVersionHistory : uint32;
@@ -118,6 +119,9 @@ private:
 	/** Last known address																			*/
 	TSharedPtr<const FInternetAddr>			LastAddress;
 
+	/** Delegate that's fired when tracking starts on the current network profile */
+	FOnNetworkProfileStarted				OnNetworkProfileStartedDelegate;
+
 	/** Delegate that's fired when tracking stops on the current network profile */
 	FOnNetworkProfileFinished				OnNetworkProfileFinishedDelegate;
 
@@ -489,6 +493,9 @@ public:
 	bool FORCEINLINE IsTrackingEnabled() const { return bIsTrackingEnabled || bShouldTrackingBeEnabled; }
 	bool IsComparisonTrackingEnabled() const { return (bIsTrackingEnabled || bShouldTrackingBeEnabled) && bIsComparisonTrackingEnabled; }
 
+	/** Return the network profile started delegate */
+	FOnNetworkProfileStarted& OnNetworkProfileStarted() { return OnNetworkProfileStartedDelegate; }
+
 	/** Return the network profile finished delegate */
 	FOnNetworkProfileFinished& OnNetworkProfileFinished() { return OnNetworkProfileFinishedDelegate; }
 };
