Add null checks for child components

diff --git a/Engine/Source/Runtime/Engine/Private/Components/SceneComponent.cpp b/Engine/Source/Runtime/Engine/Private/Components/SceneComponent.cpp
--- a/Engine/Source/Runtime/Engine/Private/Components/SceneComponent.cpp
+++ b/Engine/Source/Runtime/Engine/Private/Components/SceneComponent.cpp
@@ -474,7 +474,7 @@ static bool CheckDescendantsAreAlsoCulledForTarget(USceneComponent const* SceneC
 
 	for (USceneComponent* ChildSceneComponent : AttachedChildren)
 	{
-		if (SceneComponentNeedsLoadForTarget(ChildSceneComponent, TargetPlatform))
+		if (ChildSceneComponent && SceneComponentNeedsLoadForTarget(ChildSceneComponent, TargetPlatform))
 		{
 			return false;
 		}
@@ -525,7 +525,7 @@ EDataValidationResult USceneComponent::IsDataValid(FDataValidationContext& Conte
 	{
 		for (USceneComponent* ChildSceneComponent : GetAttachChildren())
 		{
-			if (!ChildSceneComponent->IsEditorOnly())
+			if (ChildSceneComponent && !ChildSceneComponent->IsEditorOnly())
 			{
 				Context.AddError(FText::Format(LOCTEXT("SceneComponent_AttachmentEditorOnlyMismatch",
 					"Component {0} is editor-only but it has an attached child {1} that is not"),
@@ -540,7 +540,7 @@ EDataValidationResult USceneComponent::IsDataValid(FDataValidationContext& Conte
 	{
 		for (USceneComponent* ChildSceneComponent : GetAttachChildren())
 		{
-			if (ChildSceneComponent->IsEditorOnly())
+			if (!ChildSceneComponent || ChildSceneComponent->IsEditorOnly())
 			{
 				continue;
 			}
