[NetPhysics] Client-Server physics frameoffset refactor.

- Implemented a new simple flow to set the initial client-server physics frame offset.
- Implemented a solution to keep the above offset safe with a buffer and to keep that buffer in check via client-side time-dilation dictated by the server.

diff --git a/Engine/Source/Runtime/Engine/Classes/GameFramework/PlayerController.h b/Engine/Source/Runtime/Engine/Classes/GameFramework/PlayerController.h
--- a/Engine/Source/Runtime/Engine/Classes/GameFramework/PlayerController.h
+++ b/Engine/Source/Runtime/Engine/Classes/GameFramework/PlayerController.h
@@ -2242,122 +2242,21 @@ private:
 	FString CurrentInputModeDebugString;
 #endif
 
-public: 
-
-	/**
-	 * Frame number exchange. This doesn't inherently do anything but is used by the network prediction physics system.
-	 * This may be moved out at some point.
-	 * 
-	 * This is meant ot provide a mechanism for client side prediction to correlate client input and server frame numbers.
-	 * Frame is a loose concept here. It doesn't necessary mean GFrameNumber. Its just an arbitrary increasing sequence of numbers that is used
-	 * to label disrete units of client->Server input. For example the main thread may tick at a high variable rate but input is generated at a fixed
-	 * step interval.
-	 */
-
-	struct FInputCmdBuffer
-	{
-		int32 HeadFrame() const { return LastWritten; }
-		int32 TailFrame() const { return FMath::Max(0, LastWritten - Buffer.Num() + 1); }
-		TArray<uint8>& Write(int32 Frame) { LastWritten = FMath::Max(Frame, LastWritten); return Buffer[Frame % Buffer.Num()]; }
-		const TArray<uint8>& Get(int32 Frame) const { return Buffer[Frame % Buffer.Num()]; }
-
-	private:
-		int32 LastWritten = INDEX_NONE;
-		TStaticArray<TArray<uint8>, 16> Buffer;
-	};
-
-	FInputCmdBuffer& GetInputBuffer() { return InputBuffer; }
-
-	// -------------------------------------------------------------------------
-	// Client
-	// -------------------------------------------------------------------------
-
-	struct FClientFrameInfo
-	{
-		int32 LastRecvInputFrame = INDEX_NONE;	// The latest inputcmd that the server acknowledged receiving, but not yet processed (this is our frame number that we gave them)
-		int32 LastProcessedInputFrame = INDEX_NONE; // The latest InputCmd that the server actually processed (this is our frame number that we gave them)
-		int32 LastRecvServerFrame = INDEX_NONE; // the latest ServerFrame number that the processing of LastRecvInputFrame happened on (Server's local frame number)
-
-		int8 QuantizedTimeDilation = 1; // Server sent this to this client, telling them to dilate local time either catch up or slow down
-		float TargetNumBufferedCmds = 0.f;
-
-		int32 GetLocalFrameOffset() const { return LastProcessedInputFrame - LastRecvServerFrame; }
-	};	
-
-	// Client pushes input data locally. RPC is sent here but also includes redundant data
-	ENGINE_API void PushClientInput(int32 ClientInputFrame, TArray<uint8>& Data);
-
-	// Client says "Here is input frame number X" (and then calls other RPCs to deliver InputCmd payload)
-	UFUNCTION(Server, unreliable)
-	ENGINE_API void ServerRecvClientInputFrame(int32 RecvClientInputFrame, const TArray<uint8>& Data);
-
-	/** When enabled the client will send the last N inputs to the server, N dictated by CVar: np2.NumRedundantCmds 3 */
-	ENGINE_API void EnableNetworkedPhysicsInputSync(bool EnablePrediction)
-	{
-		bSyncInputsForNetworkedPhysics = EnablePrediction;
-	}
-
-	const FClientFrameInfo& GetClientFrameInfo() const { return ClientFrameInfo; }
-	FClientFrameInfo& GetClientFrameInfo() { return ClientFrameInfo; }
-
-	// -------------------------------------------------------------------------
-	// Server
-	// -------------------------------------------------------------------------
-
-	struct FServerFrameInfo
-	{
-		int32 LastProcessedInputFrame = INDEX_NONE;	// The last client frame number we processed. "processed" is arbitrary and we are informed about when commands are processed via SetServerProessedInputFrame
-		int32 LastLocalFrame = INDEX_NONE; // The local frame number that we processed the latest client input frame on. Again, processed is arbitrary and set via SetServerProessedInputFrame
-		int32 LastSentLocalFrame = INDEX_NONE;	// Tracks the latest LastLocalFrame that we sent to the client. Just to prevent redundantly sending info via RPC
-
-		float TargetTimeDilation = 1.f;
-		int8 QuantizedTimeDilation = 1; // Server sets this to tell client to slowdown or speed up
-		float TargetNumBufferedCmds = 1.f; // How many buffered cmds the server thinks this client should ideally have to absorb PL and latency variance
-		bool bFault = true;
-	};
-
-	// We call this in ::SendClientAdjustment to tell the client what the last processed input frame was for it and on what local frame number it was processed
-	UFUNCTION(Client, unreliable)
-	ENGINE_API void ClientRecvServerAckFrame(int32 LastProcessedInputFrame, int32 RecvServerFrameNumber, int8 TimeDilation);
-
-	UFUNCTION(Client, unreliable)
-	ENGINE_API void ClientRecvServerAckFrameDebug(uint8 NumBuffered, float TargetNumBufferedCmds);
-
-	FServerFrameInfo& GetServerFrameInfo() { return ServerFrameInfo; };
 private:
 
-	FInputCmdBuffer InputBuffer;
-	FClientFrameInfo ClientFrameInfo;
-	FServerFrameInfo ServerFrameInfo;
-
-	/** If true, inputs are synced from client to server in a way that's expected for NetworkPhysicsComponent */
-	bool bSyncInputsForNetworkedPhysics = false;
-
 	/** The estimated offset between the local async physics tick frame number and the server's
 	*	This is used to synchronize events that happen in the async physics tick */
-	int32 LocalToServerAsyncPhysicsTickOffset;
-
-	/** The estimated offset between the server async physics tick frame number and the local's
-	*	This is used to synchronize events that happen in the async physics tick */
-	int32 ServerToLocalAsyncPhysicsTickOffset = INDEX_NONE;
+	int32 LocalToServerAsyncPhysicsTickOffset = INDEX_NONE;
+	bool LocalToServerAsyncPhysicsTickOffsetAssigned = false;
 	
-	/** The latest server step we've received an offset correction for. This allows us to ignore out of order corrections that arrive late */
-	int32 ClientLatestCorrectedOffsetServerStep = INDEX_NONE;
-
 	/** The latest physics step we've sent to the server. Due to async we need to avoid duplicate sends */
 	int32 ClientLatestAsyncPhysicsStepSent = INDEX_NONE;
 
 	/** The latest server step we've received a time dilation for. Needed for out of order updates */
 	int32 ClientLatestTimeDilationServerStep = INDEX_NONE;
 
-	/** The server tells the client to speed up or slow down in order to keep its buffer full */
-	float ServerAsyncPhysicsTimeDilationToSend = 1.f;
-
-	/** The server records the latest timestamp it has to correct. This is used to update client (which may not happen on every physics step) */
-	FAsyncPhysicsTimestamp ServerLatestTimestampToCorrect;
-
-	/** The latest timestamp the client has sent to the server with prediction of which server frame it corresponds to. */
-	FAsyncPhysicsTimestamp ServerPendingTimestamp;
+	/** The latest physics step we've received from the client. */
+	int32 ServerLatestAsyncPhysicsStepReceived = INDEX_NONE;
 
 	/** Update the tick ofsset in between the local client and the server */
 	ENGINE_API void UpdateServerAsyncPhysicsTickOffset();
@@ -2365,7 +2264,7 @@ private:
 	UFUNCTION(Server, Unreliable)
 	ENGINE_API void ServerSendLatestAsyncPhysicsTimestamp(FAsyncPhysicsTimestamp Timestamp);
 
-	UFUNCTION(Client, Unreliable)
+	UFUNCTION(Client, Reliable)
 	ENGINE_API void ClientCorrectionAsyncPhysicsTimestamp(FAsyncPhysicsTimestamp Timestamp);
 
 	UFUNCTION(Client, Unreliable)
@@ -2387,14 +2286,8 @@ public:
 
 	/** Returns the current estimated offset between the local async physics step and the server. This is useful for dealing with low level synchronization.
 		In general it's recommended to use GetAsyncPhysicsTimestamp which accounts for the offset automatically*/
-	int32 GetLocalToServerAsyncPhysicsTickOffset() const { return LocalToServerAsyncPhysicsTickOffset; }
-
-	/** Returns the current estimated offset between the server async physics step and the local one. */
-	int32 GetServerToLocalAsyncPhysicsTickOffset() const { return (ServerToLocalAsyncPhysicsTickOffset != INDEX_NONE) ? ServerToLocalAsyncPhysicsTickOffset : LocalToServerAsyncPhysicsTickOffset; }
-
-	/** Set the offset between the server async physics step and the local one.*/
-	void SetServerToLocalAsyncPhysicsTickOffset( const int32 AsyncPhysicsTickOffset) { ServerToLocalAsyncPhysicsTickOffset = AsyncPhysicsTickOffset; }
-	
+	int32 GetLocalToServerAsyncPhysicsTickOffset() const { return LocalToServerAsyncPhysicsTickOffset; }	
+	bool GetLocalToServerAsyncPhysicsTickOffsetAssigned() const { return LocalToServerAsyncPhysicsTickOffsetAssigned; }
 };
 
 #if UE_ENABLE_INCLUDE_ORDER_DEPRECATED_IN_5_2
diff --git a/Engine/Source/Runtime/Engine/Private/PhysicsEngine/AsyncPhysicsInputComponent.cpp b/Engine/Source/Runtime/Engine/Private/PhysicsEngine/AsyncPhysicsInputComponent.cpp
--- a/Engine/Source/Runtime/Engine/Private/PhysicsEngine/AsyncPhysicsInputComponent.cpp
+++ b/Engine/Source/Runtime/Engine/Private/PhysicsEngine/AsyncPhysicsInputComponent.cpp
@@ -38,19 +38,12 @@ struct FAsyncPhysicsInputRewindCallback : public Chaos::IRewindCallback
 		{
 			if (APlayerController* PC = World->GetFirstPlayerController())
 			{
-				if (Chaos::FPhysicsSolverBase::IsNetworkPhysicsPredictionEnabled())
+				ensure(Chaos::FPhysicsSolverBase::IsNetworkPhysicsPredictionEnabled());
+				if (PC->GetLocalToServerAsyncPhysicsTickOffsetAssigned())
 				{
 					const int32 LocalToServerOffset = PC->GetLocalToServerAsyncPhysicsTickOffset();
 					ServerFrame = PhysicsStep + LocalToServerOffset;
 				}
-				else
-				{
-					APlayerController::FClientFrameInfo& ClientFrameInfo = PC->GetClientFrameInfo();
-                    if (ClientFrameInfo.LastProcessedInputFrame != INDEX_NONE)
-                    {
-                    	ServerFrame = PhysicsStep - ClientFrameInfo.GetLocalFrameOffset();
-                    }
-				}
 			}
 		}
 		DispatchPhysicsTick.Broadcast(PhysicsStep, NumSteps, ServerFrame);
@@ -79,18 +72,13 @@ struct FAsyncPhysicsInputRewindCallback : public Chaos::IRewindCallback
 				// -----------------------------------------------------------------
 				// Calculate latest frame offset to map server frame to local frame
 				// -----------------------------------------------------------------
-				APlayerController::FClientFrameInfo& ClientFrameInfo = PC->GetClientFrameInfo();
 
-				if (Chaos::FPhysicsSolverBase::IsNetworkPhysicsPredictionEnabled())
+				ensure(Chaos::FPhysicsSolverBase::IsNetworkPhysicsPredictionEnabled());
+				if (PC->GetLocalToServerAsyncPhysicsTickOffsetAssigned())
 				{
 					const int32 LocalToServerOffset = PC->GetLocalToServerAsyncPhysicsTickOffset();
 					CachedServerFrame = PhysicsStep + LocalToServerOffset;
 				}
-				else if (ClientFrameInfo.LastProcessedInputFrame != INDEX_NONE)
-				{
-					const int32 FrameOffset = ClientFrameInfo.GetLocalFrameOffset();
-					CachedServerFrame = PhysicsStep - FrameOffset; // Local = Server + Offset
-				}
 			}
 		}
 	}
diff --git a/Engine/Source/Runtime/Engine/Private/PhysicsEngine/NetworkPhysicsComponent.cpp b/Engine/Source/Runtime/Engine/Private/PhysicsEngine/NetworkPhysicsComponent.cpp
--- a/Engine/Source/Runtime/Engine/Private/PhysicsEngine/NetworkPhysicsComponent.cpp
+++ b/Engine/Source/Runtime/Engine/Private/PhysicsEngine/NetworkPhysicsComponent.cpp
@@ -14,72 +14,6 @@
 
 #include UE_INLINE_GENERATED_CPP_BY_NAME(NetworkPhysicsComponent)
 
-namespace InputCmdCVars
-{
-	static int32 ForceFault = 0;
-	static FAutoConsoleVariableRef CVarForceFault(TEXT("p.net.ForceFault"), ForceFault, TEXT("Forces server side input fault"));
-
-	static int32 MaxBufferedCmds = 16;
-	static FAutoConsoleVariableRef CVarMaxBufferedCmds(TEXT("p.net.MaxBufferedCmds"), MaxBufferedCmds, TEXT("MaxNumber of buffered server side commands"));
-
-	static int32 TimeDilationEnabled = 0;
-	static FAutoConsoleVariableRef CVarTimeDilationEnabled(TEXT("p.net.TimeDilationEnabled"), TimeDilationEnabled, TEXT("Enable clientside TimeDilation"));
-
-	static float MaxTargetNumBufferedCmds = 5.0;
-	static FAutoConsoleVariableRef CVarMaxTargetNumBufferedCmds(TEXT("p.net.MaxTargetNumBufferedCmds"), MaxTargetNumBufferedCmds, TEXT("Maximum number of buffered inputs the server will target per client."));
-
-	static float MaxTimeDilationMag = 0.01f;
-	static FAutoConsoleVariableRef CVarMaxTimeDilationMag(TEXT("p.net.MaxTimeDilationMag"), MaxTimeDilationMag, TEXT("Maximum time dilation that client will use to slow down / catch up with server"));
-
-	static float TimeDilationAlpha = 0.1f;
-	static FAutoConsoleVariableRef CVarTimeDilationAlpha(TEXT("p.net.TimeDilationAlpha"), TimeDilationAlpha, TEXT("Lerp strength for sliding client time dilation"));
-
-	static float TargetNumBufferedCmdsDeltaOnFault = 1.0f;
-	static FAutoConsoleVariableRef CVarTargetNumBufferedCmdsDeltaOnFault(TEXT("p.net.TargetNumBufferedCmdsDeltaOnFault"), TargetNumBufferedCmdsDeltaOnFault, TEXT("How much to increase TargetNumBufferedCmds when an input fault occurs"));
-
-	static float TargetNumBufferedCmds = 1.9f;
-	static FAutoConsoleVariableRef CVarTargetNumBufferedCmds(TEXT("p.net.TargetNumBufferedCmds"), TargetNumBufferedCmds, TEXT("How much to increase TargetNumBufferedCmds when an input fault occurs"));
-
-	static float TargetNumBufferedCmdsAlpha = 0.005f;
-	static FAutoConsoleVariableRef CVarTargetNumBufferedCmdsAlpha(TEXT("p.net.TargetNumBufferedCmdsAlpha"), TargetNumBufferedCmdsAlpha, TEXT("Lerp strength for TargetNumBufferedCmds"));
-
-	static int32 LerpTargetNumBufferedCmdsAggresively = 0;
-	static FAutoConsoleVariableRef CVarLerpTargetNumBufferedCmdsAggresively(TEXT("p.net.LerpTargetNumBufferedCmdsAggresively"), LerpTargetNumBufferedCmdsAggresively, TEXT("Aggresively lerp towards TargetNumBufferedCmds. Reduces server side buffering but can cause more artifacts."));
-}
-
-// --------------------------------------------------------------------------------------------------------------------------------------------------
-//	Client InputCmd Stream stuff
-// --------------------------------------------------------------------------------------------------------------------------------------------------
-
-namespace
-{
-
-int8 QuantizeTimeDilation(float F)
-{
-	if (F == 1.f)
-	{
-		return 0;
-	}
-
-	float Normalized = FMath::Clamp<float>((F - 1.f) / InputCmdCVars::MaxTimeDilationMag, -1.f, 1.f);
-	return (int8)(Normalized * 128.f);
-}
-
-float DeQuantizeTimeDilation(int8 i)
-{
-	if (i == 0)
-	{
-		return 1.f;
-	}
-
-	float Normalized = (float)i / 128.f;
-	float Uncompressed = 1.f + (Normalized * InputCmdCVars::MaxTimeDilationMag);
-	return Uncompressed;
-}
-
-}
-
-
 bool FNetworkPhysicsRewindDataProxy::NetSerializeBase(FArchive& Ar, class UPackageMap* Map, bool& bOutSuccess, TUniqueFunction<TUniquePtr<Chaos::FBaseRewindHistory>()> CreateHistoryFunction)
 {
 	Ar << Owner;
@@ -197,115 +131,6 @@ int32 FNetworkPhysicsCallback::TriggerRewindIfNeeded_Internal(int32 LatestStepCo
 	return ResimFrame;
 }
 
-void FNetworkPhysicsCallback::UpdateClientPlayer_External(int32 PhysicsStep)
-{
-	int32 LocalOffset = 0;
-	if (APlayerController* PC = World->GetFirstPlayerController())
-	{
-		// ------------------------------------------------
-		// Send RPC to server telling them what (client/local) physics step we are running
-		//	* Note that SendData is empty because of the existing API, should change this
-		// ------------------------------------------------	
-
-		TArray<uint8> SendData;
-		PC->PushClientInput(PhysicsStep, SendData);
-
-		// -----------------------------------------------------------------
-		// Calculate latest frame offset to map server frame to local frame
-		// -----------------------------------------------------------------
-		APlayerController::FClientFrameInfo& ClientFrameInfo = PC->GetClientFrameInfo();
-
-		// -----------------------------------------------------------------
-		// Apply local TIme Dilation based on server's recommendation.
-		// This speeds up or slows down our consumption of real time (by like < 1%)
-		// Ultimately this causes us to send InputCmds at a lower or higher rate in
-		// order to keep server side buffer at optimal capacity. 
-		// Optimal capacity = as small as possible without ever "missing" a frame (e.g, minimal buffer yet always a new fresh cmd to consume server side)
-		// -----------------------------------------------------------------
-		const float RealTimeDilation = DeQuantizeTimeDilation(ClientFrameInfo.QuantizedTimeDilation);
-
-		if (InputCmdCVars::TimeDilationEnabled > 0)
-		{
-			FPhysScene* PhysScene = World->GetPhysicsScene();
-			PhysScene->SetNetworkDeltaTimeScale(RealTimeDilation);
-		}
-	}
-}
-
-void FNetworkPhysicsCallback::UpdateServerPlayer_External(int32 PhysicsStep)
-{
-	// -----------------------------------------------
-	// Server: "consume" an InputCmd from each Player Controller
-	// All this means in this context is updating FServerFrameInfo:: LastProcessedInputFrame, LastLocalFrame
-	// (E.g, telling each client what "Input" of theirs we were processing and our local physics frame number.
-	// In cases where the buffer has a fault, we calculate a suggested time dilation to temporarily make client speed up 
-	// or slow down their input cmd production.
-	// -----------------------------------------------
-
-	const bool bForceFault = InputCmdCVars::ForceFault > 0;
-	InputCmdCVars::ForceFault = FMath::Max(0, InputCmdCVars::ForceFault - 1);
-
-	for (FConstPlayerControllerIterator Iterator = World->GetPlayerControllerIterator(); Iterator; ++Iterator)
-	{
-		if (APlayerController* PC = Iterator->Get())
-		{
-			APlayerController::FServerFrameInfo& FrameInfo = PC->GetServerFrameInfo();
-			APlayerController::FInputCmdBuffer& InputBuffer = PC->GetInputBuffer();
-
-			{
-				const int32 NumBufferedInputCmds = bForceFault ? 0 : (InputBuffer.HeadFrame() - FrameInfo.LastProcessedInputFrame);
-
-				// Check Overflow
-				if (NumBufferedInputCmds > InputCmdCVars::MaxBufferedCmds)
-				{
-					UE_LOG(LogPhysics, Warning, TEXT("[Remote.Input] overflow %d %d -> %d"), InputBuffer.HeadFrame(), FrameInfo.LastProcessedInputFrame, NumBufferedInputCmds);
-					FrameInfo.LastProcessedInputFrame = InputBuffer.HeadFrame() - InputCmdCVars::MaxBufferedCmds + 1;
-				}
-
-				// Check fault - we are waiting for Cmds to reach TargetNumBufferedCmds before continuing
-				if (FrameInfo.bFault)
-				{
-					if (NumBufferedInputCmds < (int32)FrameInfo.TargetNumBufferedCmds)
-					{
-						// Skip this because it is in fault. We will use the prev input for this frame.
-						UE_CLOG(FrameInfo.LastProcessedInputFrame != INDEX_NONE, LogPhysics, Warning, TEXT("[Remote.Input] in fault. Reusing Inputcmd. (Client) Input: %d. (Server) Local Frame: %d"), FrameInfo.LastProcessedInputFrame, FrameInfo.LastLocalFrame);
-						continue;
-					}
-					FrameInfo.bFault = false;
-				}
-				else if (NumBufferedInputCmds <= 0)
-				{
-					// No Cmds to process, enter fault state. Increment TargetNumBufferedCmds each time this happens.
-					// TODO: We should have something to bring this back down (which means skipping frames) we don't want temporary poor conditions to cause permanent high input buffering
-					FrameInfo.bFault = true;
-					FrameInfo.TargetNumBufferedCmds = FMath::Min(FrameInfo.TargetNumBufferedCmds + InputCmdCVars::TargetNumBufferedCmdsDeltaOnFault, InputCmdCVars::MaxTargetNumBufferedCmds);
-
-					UE_CLOG(FrameInfo.LastProcessedInputFrame != INDEX_NONE, LogPhysics, Warning, TEXT("[Remote.Input] ENTERING fault. New Target: %.2f. (Client) Input: %d. (Server) Local Frame: %d"), FrameInfo.TargetNumBufferedCmds, FrameInfo.LastProcessedInputFrame, FrameInfo.LastLocalFrame);
-					continue;
-				}
-
-				float TargetTimeDilation = 1.f;
-				if (NumBufferedInputCmds < (int32)FrameInfo.TargetNumBufferedCmds)
-				{
-					TargetTimeDilation += InputCmdCVars::MaxTimeDilationMag; // Tell client to speed up, we are starved on cmds
-				}
-
-				FrameInfo.TargetTimeDilation = FMath::Lerp(FrameInfo.TargetTimeDilation, TargetTimeDilation, InputCmdCVars::TimeDilationAlpha);
-				FrameInfo.QuantizedTimeDilation = QuantizeTimeDilation(TargetTimeDilation);
-
-				if (InputCmdCVars::LerpTargetNumBufferedCmdsAggresively != 0)
-				{
-					// When aggressive, always lerp towards target
-					FrameInfo.TargetNumBufferedCmds = FMath::Lerp(FrameInfo.TargetNumBufferedCmds, InputCmdCVars::TargetNumBufferedCmds, InputCmdCVars::TargetNumBufferedCmdsAlpha);
-				}
-
-				FrameInfo.LastProcessedInputFrame++;
-				FrameInfo.LastLocalFrame = PhysicsStep;
-			}
-		}
-	}
-}
-
 void FNetworkPhysicsCallback::InjectInputs_External(int32 PhysicsStep, int32 NumSteps)
 {
 	InjectInputsExternal.Broadcast(PhysicsStep, NumSteps);
@@ -320,15 +145,6 @@ void FNetworkPhysicsCallback::ProcessInputs_External(int32 PhysicsStep, const TA
 			SimCallbackObject.CallbackObject->ProcessInputs_External(PhysicsStep);
 		}
 	}
-
-	if (NetMode == NM_Client)
-	{
-		UpdateClientPlayer_External(PhysicsStep);
-	}
-	else
-	{
-		UpdateServerPlayer_External(PhysicsStep);
-	}
 }
 
 UNetworkPhysicsSystem::UNetworkPhysicsSystem()
@@ -413,11 +229,6 @@ void UNetworkPhysicsComponent::BeginPlay()
 			{
 				if (UPhysicsSettings::Get()->PhysicsPrediction.bEnablePhysicsPrediction)
 				{
-					if (APlayerController* PlayerController = GetPlayerController())
-					{
-						PlayerController->EnableNetworkedPhysicsInputSync(true);
-					}
-
 					if (Solver->GetRewindData() == nullptr)
 					{
 						const int32 NumFrames = FMath::Max<int32>(1, UPhysicsSettings::Get()->GetPhysicsHistoryCount());
@@ -512,13 +323,18 @@ void UNetworkPhysicsComponent::AsyncPhysicsTickComponent(float DeltaTime, float
 
 void UNetworkPhysicsComponent::SendLocalInputsDatas()
 {
-	if (HasLocalController() && InputsHistory)
+	const APlayerController* PlayerController = GetPlayerController();
+	if (!PlayerController)
 	{
-		// We just check that the local client to server offset is valid before doing something
-		const int32 LocalOffset = HasServerWorld() ? 0 : GetPlayerController()->GetLocalToServerAsyncPhysicsTickOffset();
+		return;
+	}
 
-		if (LocalOffset >= 0)
+	if (PlayerController->IsLocalController() && InputsHistory)
+	{
+		// We just check that the local client to server offset is valid before doing something
+		if (HasServerWorld() || PlayerController->GetLocalToServerAsyncPhysicsTickOffsetAssigned())
 		{
+			const int32 LocalOffset = HasServerWorld() ? 0 : PlayerController->GetLocalToServerAsyncPhysicsTickOffset();
 			const int32 NextIndex = (InputsIndex + 1) % (InputsRedundancy + 1);
 
 			// if on server (Listen server) we should send the inputs onto all the clients through repnotify
@@ -527,19 +343,16 @@ void UNetworkPhysicsComponent::SendLocalInputsDatas()
 			if (!HasServerWorld())
 			{
 #if DEBUG_NETWORK_PHYSICS
-				if (APlayerController* PlayerController = GetPlayerController())
-				{
-					FAsyncPhysicsTimestamp Timestamp = PlayerController->GetAsyncPhysicsTimestamp();
+				FAsyncPhysicsTimestamp Timestamp = PlayerController->GetAsyncPhysicsTimestamp();
 
-					TArray<int32> LocalFrames, ServerFrames, InputFrames;
-					InputsHistory->DebugDatas(*ReplicatedInputs.History, LocalFrames, ServerFrames, InputFrames);
+				TArray<int32> LocalFrames, ServerFrames, InputFrames;
+				InputsHistory->DebugDatas(*ReplicatedInputs.History, LocalFrames, ServerFrames, InputFrames);
 
-					UE_LOG(LogTemp, Log, TEXT("CLIENT | GT | SendLocalInputsDatas | Sending %d inputs from CLIENT | Component = %s"), LocalFrames.Num(), *GetFullName());
-					for (int32 FrameIndex = 0; FrameIndex < LocalFrames.Num(); ++FrameIndex)
-					{
-						UE_LOG(LogTemp, Log, TEXT("		Debugging local inputs at local frame = %d | server frame = %d | Current local frame = %d | Current server frame = %d"),
-							LocalFrames[FrameIndex], ServerFrames[FrameIndex], Timestamp.LocalFrame, Timestamp.ServerFrame);
-					}
+				UE_LOG(LogTemp, Log, TEXT("CLIENT | GT | SendLocalInputsDatas | Sending %d inputs from CLIENT | Component = %s"), LocalFrames.Num(), *GetFullName());
+				for (int32 FrameIndex = 0; FrameIndex < LocalFrames.Num(); ++FrameIndex)
+				{
+					UE_LOG(LogTemp, Log, TEXT("		Debugging local inputs at local frame = %d | server frame = %d | Current local frame = %d | Current server frame = %d"),
+						LocalFrames[FrameIndex], ServerFrames[FrameIndex], Timestamp.LocalFrame, Timestamp.ServerFrame);
 				}
 #endif
 
@@ -562,29 +375,6 @@ void UNetworkPhysicsComponent::SendLocalStatesDatas()
 	}
 }
 
-void UNetworkPhysicsComponent::CorrectServerToLocalOffset(const int32 LocalToServerOffset)
-{
-	if (HasLocalController() && !HasServerWorld() && StatesHistory)
-	{
-		TArray<int32> LocalFrames, ServerFrames, InputFrames;
-		StatesHistory->DebugDatas(*ReplicatedStates.History, LocalFrames, ServerFrames, InputFrames);
-
-		int32 ServerToLocalOffset = LocalToServerOffset;
-		for (int32 FrameIndex = 0; FrameIndex < LocalFrames.Num(); ++FrameIndex)
-		{
-#if DEBUG_NETWORK_PHYSICS || DEBUG_REWIND_DATA
-			UE_LOG(LogTemp, Log, TEXT("CLIENT | GT | CorrectServerToLocalOffset | Server frame = %d | Client Frame = %d"), ServerFrames[FrameIndex], InputFrames[FrameIndex]);
-#endif
-			ServerToLocalOffset = FMath::Min(ServerToLocalOffset, ServerFrames[FrameIndex] - InputFrames[FrameIndex]);
-		}
-
-		GetPlayerController()->SetServerToLocalAsyncPhysicsTickOffset(ServerToLocalOffset);
-#if DEBUG_NETWORK_PHYSICS || DEBUG_REWIND_DATA
-		UE_LOG(LogTemp, Log, TEXT("CLIENT | GT | CorrectServerToLocalOffset | Server to local offset = %d | Local to server offset = %d"), ServerToLocalOffset, LocalToServerOffset);
-#endif
-	}
-}
-
 void UNetworkPhysicsComponent::OnRep_SetReplicatedStates()
 {
 	// The replicated states should only be used on the client since the server already have authoritative local ones
@@ -595,8 +385,7 @@ void UNetworkPhysicsComponent::OnRep_SetReplicatedStates()
 		{
 			PlayerController = GetWorld()->GetFirstPlayerController();
 		}
-		CorrectServerToLocalOffset(PlayerController->GetLocalToServerAsyncPhysicsTickOffset());
-		const int32 LocalOffset = PlayerController->GetServerToLocalAsyncPhysicsTickOffset();
+		const int32 LocalOffset = PlayerController->GetLocalToServerAsyncPhysicsTickOffset();
 
 		// Record the received states from the server into the history for future use
 		if (UWorld* World = GetWorld())
@@ -640,7 +429,7 @@ void UNetworkPhysicsComponent::OnRep_SetReplicatedInputs()
 		{
 			PlayerController = GetWorld()->GetFirstPlayerController();
 		}
-		const int32 LocalOffset = PlayerController->GetServerToLocalAsyncPhysicsTickOffset();
+		const int32 LocalOffset = PlayerController->GetLocalToServerAsyncPhysicsTickOffset();
 
 		// Record the received inputs from the server into the history for future use
 		if (UWorld* World = GetWorld())
@@ -810,22 +599,23 @@ void UNetworkPhysicsComponent::OnPostProcessInputsInternal(const int32 PhysicsSt
 			InputsOffsets[InputsIndex] = FMath::Max(InputsOffsets[InputsIndex], PhysicsStep + 1);
 			InputsHistory->RecordDatas(PhysicsStep, PhysicsDatas);
 
-	#if DEBUG_NETWORK_PHYSICS
+#if DEBUG_NETWORK_PHYSICS
 			UE_LOG(LogTemp, Log, TEXT("		Recording local inputs at frame %d | Component = %s"), PhysicsDatas->LocalFrame, *GetFullName());
-	#endif
+#endif
 		}
-		// Compute of the local frame coming from the client that was used to generate this state
-		int32 InputFrame = INDEX_NONE;
+
+		if (HasServerWorld())
 		{
-			FNetworkPhysicsDatas* PhysicsDatas = InputsDatas.Get();
-			if(InputsHistory->ExtractDatas(PhysicsStep, false, PhysicsDatas, true))
+			// Compute of the local frame coming from the client that was used to generate this state
+			int32 InputFrame = INDEX_NONE;
 			{
-				InputFrame = PhysicsDatas->InputFrame;
+				FNetworkPhysicsDatas* PhysicsDatas = InputsDatas.Get();
+				if(InputsHistory->ExtractDatas(PhysicsStep, false, PhysicsDatas, true))
+				{
+					InputFrame = PhysicsDatas->InputFrame;
+				}
 			}
-		}
 
-		if (HasServerWorld() && !bIsSolverResim)
-		{
 			FNetworkPhysicsDatas* PhysicsDatas = StatesDatas.Get();
 			PhysicsDatas->LocalFrame = PhysicsStep;
 			PhysicsDatas->ServerFrame = PhysicsStep;
@@ -833,12 +623,12 @@ void UNetworkPhysicsComponent::OnPostProcessInputsInternal(const int32 PhysicsSt
 
 			PhysicsDatas->BuildDatas(ActorComponent);
 
-			StatesOffsets[StatesIndex] = FMath::Max(StatesOffsets[StatesIndex], PhysicsStep+1);
+			StatesOffsets[StatesIndex] = FMath::Max(StatesOffsets[StatesIndex], PhysicsStep + 1);
 			StatesHistory->RecordDatas(PhysicsStep, PhysicsDatas);
 
-	#if DEBUG_NETWORK_PHYSICS
+#if DEBUG_NETWORK_PHYSICS
 			UE_LOG(LogTemp, Log, TEXT("		Recording local states at frame %d | from input frame = %d | Component = %s"), PhysicsDatas->LocalFrame, PhysicsDatas->InputFrame, *GetFullName());
-	#endif
+#endif
 		}
 	}
 }
diff --git a/Engine/Source/Runtime/Engine/Private/PhysicsEngine/PhysicsReplication.cpp b/Engine/Source/Runtime/Engine/Private/PhysicsEngine/PhysicsReplication.cpp
--- a/Engine/Source/Runtime/Engine/Private/PhysicsEngine/PhysicsReplication.cpp
+++ b/Engine/Source/Runtime/Engine/Private/PhysicsEngine/PhysicsReplication.cpp
@@ -314,7 +314,7 @@ void FPhysicsReplication::OnTick(float DeltaSeconds, TMap<TWeakObjectPtr<UPrimit
 			{
 				if (APlayerController* PlayerController = World->GetFirstPlayerController())
 				{
-					LocalFrameOffset = PlayerController->GetServerToLocalAsyncPhysicsTickOffset();
+					LocalFrameOffset = PlayerController->GetLocalToServerAsyncPhysicsTickOffset();
 				}
 			}
 		}
diff --git a/Engine/Source/Runtime/Engine/Private/PlayerController.cpp b/Engine/Source/Runtime/Engine/Private/PlayerController.cpp
--- a/Engine/Source/Runtime/Engine/Private/PlayerController.cpp
+++ b/Engine/Source/Runtime/Engine/Private/PlayerController.cpp
@@ -109,21 +109,29 @@ namespace PlayerControllerCVars
 
 namespace NetworkPhysicsCvars
 {
-	int32 NumRedundantCmds = 3;
-	FAutoConsoleVariableRef CVarNumRedundantCmds(TEXT("np2.NumRedundantCmds"), NumRedundantCmds, TEXT("Number of redundant user cmds to send per frame"));
-
-#if (UE_BUILD_SHIPPING || UE_BUILD_TEST)
-	int32 EnableDebugRPC = 0;
-#else
-	int32 EnableDebugRPC = 1;
-#endif
-	FAutoConsoleVariableRef CVarEnableDebugRPC(TEXT("np2.EnableDebugRPC"), EnableDebugRPC, TEXT("Sends extra debug information to clients about server side input buffering"));
-
-	int32 NetworkPhysicsPredictionFrameOffset = 4;
-	FAutoConsoleVariableRef CVarNetworkPhysicsPredictionFrameOffset(TEXT("np2.NetworkPhysicsPredictionFrameOffset"), NetworkPhysicsPredictionFrameOffset, TEXT("Additional frame offset to be added to the local to server offset used by network prediction"));
+	int32 PredictionAsyncFrameBuffer = 3;
+	FAutoConsoleVariableRef CVarPredictionAsyncFrameBuffer(TEXT("np2.PredictionAsyncFrameBuffer"), PredictionAsyncFrameBuffer, TEXT("Additional frame offset to be added to the local to server offset used by network prediction"));
 
 	int32 TickOffsetUpdateInterval = 10;
 	FAutoConsoleVariableRef CVarTickOffsetUpdateInterval(TEXT("np2.TickOffsetUpdateInterval"), TickOffsetUpdateInterval, TEXT("How many physics ticks to wait between each tick offset update. Lowest viable value = 1, which means update each tick."));
+
+	bool TimeDilationEnabled = true;
+	FAutoConsoleVariableRef CVarTimeDilationEnabled(TEXT("np2.TimeDilationEnabled"), TimeDilationEnabled, TEXT("Server-side CVar, the server calculates a time dilation multiplier and sends to the client which manipulates the accumulated time on client to speed up or slow down the physics simulation in order to correct the physics frame offset between client and server."));
+	
+	float TimeDilationAmount = 0.01f;
+	FAutoConsoleVariableRef CVarTimeDilationAmount(TEXT("np2.TimeDilationAmount"), TimeDilationAmount, TEXT("Server-side CVar, Default: 0.01 | Value is in percent where 0.01 = 1% dilation. Example: 1.0/0.01 = 100, meaning that over the time it usually takes to tick 100 physics steps we will tick 99 or 101 depending on if we dilate up or down."));
+
+	bool TimeDilationEscalation = true;
+	FAutoConsoleVariableRef CVarTimeDilationEscalation(TEXT("np2.TimeDilationEscalation"), TimeDilationEscalation, TEXT("Server-side CVar, Dilate the time more depending on how many ticks we need to adjust. When set to false we use the set TimeDilationAmount and wait the amount of time it takes to perform correct the offset. When set to true we multiply the TimeDilationAmount with the buffer offset count which will correct the offset in one TimeDilationAmount cycle."));
+
+	float TimeDilationEscalationDecay = 0.05f;
+	FAutoConsoleVariableRef CVarTimeDilationEscalationDecay(TEXT("np2.TimeDilationEscalationDecay"), TimeDilationEscalationDecay, TEXT("Value is a multiplier, Default: 0.05. For each escalated TimeDilation amount, also decay by this much. Disable by setting to 0."));
+
+	float TimeDilationMax = 1.1f;
+	FAutoConsoleVariableRef CVarTimeDilationMax(TEXT("np2.TimeDilationMax"), TimeDilationMax, TEXT("Max value of the time dilation multiplier."));
+
+	float TimeDilationMin = 0.9f;
+	FAutoConsoleVariableRef CVarTimeDilationMin(TEXT("np2.TimeDilationMin"), TimeDilationMin, TEXT("Min value of the time dilation multiplier"));
 }
 
 const float RetryClientRestartThrottleTime = 0.5f;
@@ -1676,34 +1684,6 @@ void APlayerController::ClientSetCameraFade_Implementation(bool bEnableFading, F
 
 void APlayerController::SendClientAdjustment()
 {
-	if (ServerFrameInfo.LastProcessedInputFrame != INDEX_NONE && ServerFrameInfo.LastProcessedInputFrame != ServerFrameInfo.LastSentLocalFrame)
-	{
-		ServerFrameInfo.LastSentLocalFrame = ServerFrameInfo.LastProcessedInputFrame;		
-		ClientRecvServerAckFrame(ServerFrameInfo.LastProcessedInputFrame, ServerFrameInfo.LastLocalFrame, ServerFrameInfo.QuantizedTimeDilation);
-
-		if (NetworkPhysicsCvars::EnableDebugRPC)
-		{
-			ClientRecvServerAckFrameDebug(InputBuffer.HeadFrame() - ServerFrameInfo.LastProcessedInputFrame, ServerFrameInfo.TargetNumBufferedCmds);
-		}
-	}
-
-	if (UPhysicsSettings::Get()->PhysicsPrediction.bEnablePhysicsPrediction)
-	{
-		if (ServerLatestTimestampToCorrect.ServerFrame == INDEX_NONE)
-		{
-			//Nothing to correct so do not update client
-			return;
-		}
-
-		ClientCorrectionAsyncPhysicsTimestamp(ServerLatestTimestampToCorrect);
-		ServerLatestTimestampToCorrect.ServerFrame = INDEX_NONE;
-
-		if (AcknowledgedPawn != GetPawn() && !GetSpectatorPawn())
-		{
-			return;
-		}
-	}
-
 	// Server sends updates.
 	// Note: we do this for both the pawn and spectator in case an implementation has a networked spectator.
 	APawn* RemotePawn = GetPawnOrSpectator();
@@ -1717,60 +1697,6 @@ void APlayerController::SendClientAdjustment()
 	}
 }
 
-void APlayerController::PushClientInput(int32 InRecvClientInputFrame, TArray<uint8>& Data)
-{
-	if (!bSyncInputsForNetworkedPhysics && !UPhysicsSettings::Get()->PhysicsPrediction.bEnablePhysicsResimulation)
-	{
-		return;
-	}
-
-	InputBuffer.Write(InRecvClientInputFrame) = MoveTemp(Data);
-
-	// Do the RPC right here, including the redundant send. This should probably be time based and managed somewhere else like in Tick eventually
-	for (int32 Frame = FMath::Max(1, InRecvClientInputFrame-NetworkPhysicsCvars::NumRedundantCmds+1); Frame <= InRecvClientInputFrame; ++Frame)
-	{
-		ServerRecvClientInputFrame(Frame, InputBuffer.Get(Frame));
-	}
-}
-
-void APlayerController::ServerRecvClientInputFrame_Implementation(int32 InRecvClientInputFrame, const TArray<uint8>& Data)
-{
-	if (InRecvClientInputFrame < 0)
-	{
-		return;
-	}
-
-	for (int32 DroppedFrame = InputBuffer.HeadFrame()+1; DroppedFrame < InRecvClientInputFrame && DroppedFrame > 0; ++DroppedFrame)
-	{
-		UE_LOG(LogPlayerController, Warning, TEXT( "ClientInput Gap in frames (Dropped). %s [%s]. DroppedFrame: %d. RecvFrame: %d. LastProcessInputFrame: %d" ), *GetName(), PlayerState ? *PlayerState->GetPlayerName() : TEXT("???"), DroppedFrame, InRecvClientInputFrame, ServerFrameInfo.LastProcessedInputFrame);
-		//ServerFrameInfo.LastProcessedInputFrame++; // Ehh lets try this for now
-		InputBuffer.Write(DroppedFrame) = InputBuffer.Get(DroppedFrame-1);
-	}
-
-	InputBuffer.Write(InRecvClientInputFrame) = MoveTemp(const_cast<TArray<uint8>&>(Data));
-
-	if (ServerFrameInfo.LastProcessedInputFrame < InputBuffer.TailFrame())
-	{
-		// At this point, things are pretty bad and we are going to drop commands. 
-		// We still need guards client side to not send too many commands
-		UE_LOG(LogPlayerController, Warning, TEXT( "ClientInput buffer overflow. %s [%s]. RecvFrame: %d. LastProcessInputFrame: %d." ), *GetName(), PlayerState ? *PlayerState->GetPlayerName() : TEXT("???"), InRecvClientInputFrame, ServerFrameInfo.LastProcessedInputFrame);
-		ServerFrameInfo.LastProcessedInputFrame = (InputBuffer.TailFrame() + InputBuffer.HeadFrame()) / 2;
-	}
-}
-
-void APlayerController::ClientRecvServerAckFrame_Implementation(int32 LastProcessedInputFrame, int32 RecvServerFrameNumber, int8 TimeDilation)
-{
-	ClientFrameInfo.LastRecvServerFrame = RecvServerFrameNumber;
-	ClientFrameInfo.LastProcessedInputFrame = LastProcessedInputFrame;
-	ClientFrameInfo.QuantizedTimeDilation = TimeDilation;
-}
-
-void APlayerController::ClientRecvServerAckFrameDebug_Implementation(uint8 NumBuffered, float TargetNumBufferedCmds)
-{
-	ClientFrameInfo.LastRecvInputFrame = ClientFrameInfo.LastProcessedInputFrame + NumBuffered;
-	ClientFrameInfo.TargetNumBufferedCmds = TargetNumBufferedCmds;
-}
-
 /// @cond DOXYGEN_WARNINGS
 
 void APlayerController::ClientCapBandwidth_Implementation(int32 Cap)
@@ -6134,16 +6060,23 @@ FAsyncPhysicsTimestamp APlayerController::GetAsyncPhysicsTimestamp(float DeltaSe
 
 void APlayerController::UpdateServerAsyncPhysicsTickOffset()
 {
+	using namespace Chaos;
+
 	if (UWorld* World = GetWorld())
 	{
 		if (FPhysScene* PhysScene = World->GetPhysicsScene())
 		{
-			if(PhysScene->GetSolver()->GetEvolution()->IsResimming())
+			if (FPBDRigidsSolver* Solver = static_cast<FPBDRigidsSolver*>(PhysScene->GetSolver()))
 			{
-				return;
+				// Make sure that the client is actively using rewind data and that we are not currently resimulating
+				if (!Solver->GetRewindData() || Solver->GetEvolution()->IsResimming())
+				{
+					return;
+				}
 			}
 		}
 	}
+
 	FAsyncPhysicsTimestamp Timestamp = GetAsyncPhysicsTimestamp();
 	if(ClientLatestAsyncPhysicsStepSent + NetworkPhysicsCvars::TickOffsetUpdateInterval > Timestamp.LocalFrame)
 	{
@@ -6154,62 +6087,86 @@ void APlayerController::UpdateServerAsyncPhysicsTickOffset()
 	}
 
 	ClientLatestAsyncPhysicsStepSent = Timestamp.LocalFrame;
+	Timestamp.ServerFrame = LocalToServerAsyncPhysicsTickOffsetAssigned ? Timestamp.ServerFrame : INDEX_NONE; // If offset is not yet assigned, set an invalid ServerFrame
 	ServerSendLatestAsyncPhysicsTimestamp(Timestamp);
 }
 
 void APlayerController::ServerSendLatestAsyncPhysicsTimestamp_Implementation(FAsyncPhysicsTimestamp Timestamp)
 {
+	ensure(UPhysicsSettings::Get()->PhysicsPrediction.bEnablePhysicsPrediction);
+
 	//This tells the server how the client thinks the async physics tick will line up.
-	//Only cache the most up to date timestamp based on LocalFrame
-	if (Timestamp.LocalFrame > ServerPendingTimestamp.LocalFrame)
+	//If we have already received a more up to date timestamp from the client, early out
+	if (Timestamp.LocalFrame <= ServerLatestAsyncPhysicsStepReceived)
 	{
-		ServerPendingTimestamp = Timestamp;
+		return;
 	}
-}
 
-void APlayerController::ClientCorrectionAsyncPhysicsTimestamp_Implementation(FAsyncPhysicsTimestamp Timestamp)
-{
-	//This tells the client that a timestamp it sent out was wrong (for example we ran locally on step 5 and expected server to run on step 10, but it actually ran on step 11).
-	//The error can only be later. That is, it can never be that we expect to run on server step 10 but actually ran on sever step 9
-	//Once a timestamp has been corrected, any earlier timestamps can be ignored (these can be out of order because of networking)
+	ServerLatestAsyncPhysicsStepReceived = Timestamp.LocalFrame;
+		
+	// Get current server timestamp and add the frame buffer to the ServerFrame
+	FAsyncPhysicsTimestamp ActualTimestamp = GetAsyncPhysicsTimestamp();
+	ActualTimestamp.ServerFrame += NetworkPhysicsCvars::PredictionAsyncFrameBuffer;
+
+	// Mark offset as assigned when we get a valid predicted server frame.
+	const int32 PredictedServerFrame = Timestamp.ServerFrame;
+	LocalToServerAsyncPhysicsTickOffsetAssigned |= PredictedServerFrame != INDEX_NONE;
 
-	ensureMsgf(Timestamp.ServerFrame != ClientLatestCorrectedOffsetServerStep, TEXT("Server only sends at most one correction per timestamp, duplicate corrections means something is wrong"));
-	if(Timestamp.ServerFrame < ClientLatestCorrectedOffsetServerStep)
+	// Send update to client if offset is not assigned
+	// Note that we are sending the current ServerFrame along with the frame buffer added, to the client.
+	if (!LocalToServerAsyncPhysicsTickOffsetAssigned)
 	{
-		//already corrected after this so do nothing
-		return;
+		Timestamp.ServerFrame = ActualTimestamp.ServerFrame;
+		LocalToServerAsyncPhysicsTickOffset = Timestamp.ServerFrame - Timestamp.LocalFrame;
+		ClientCorrectionAsyncPhysicsTimestamp(Timestamp); /* Reliable RPC */
 	}
 
-	if (UWorld* World = GetWorld())
+	/*
+	* Use time dilation on client to adjust the frame offset and keep the buffer in check
+	* If buffer goes low, speed up the client by raising the time dilation multiplier (each deltaTime accounts for more of the accumulated time, filling the accumulated time faster to tick the next physics step)
+	* If buffer goes high, slow down the client by lowering the time dilation multiplier (each deltaTime accounts for less of the accumulated time, taking longer to fill the accumulated time
+	*/
+	if (NetworkPhysicsCvars::TimeDilationEnabled && LocalToServerAsyncPhysicsTickOffsetAssigned)
 	{
-		if (FPhysScene* PhysScene = World->GetPhysicsScene())
+		// Get the buffer offset amount that deviates from the target buffer (Note: the buffer is already added to ActualTimestamp.ServerFrame here and in the PredictedServerFrame received from the client)
+		// 0 means buffer is perfect, positive value means the buffer is too large, negative value means the buffer is too small
+		int32 CurrentFrameBufferOffset = PredictedServerFrame - ActualTimestamp.ServerFrame;
+
+		if (NetworkPhysicsCvars::TimeDilationEscalation == false)
 		{
-			if (PhysScene->GetSolver()->GetEvolution()->IsResimming())
-			{
-				return;
-			}
+			CurrentFrameBufferOffset = FMath::Clamp(CurrentFrameBufferOffset, -1, 1);
 		}
+			
+		// Calculate desired dilation and send to client
+		const float TimeDilationDecay = 1.0f - (NetworkPhysicsCvars::TimeDilationEscalationDecay * FMath::Abs(CurrentFrameBufferOffset));
+		float CalculatedTimeDilation = 1.0f + ((NetworkPhysicsCvars::TimeDilationAmount * -CurrentFrameBufferOffset) * TimeDilationDecay);
+		CalculatedTimeDilation = FMath::Clamp(CalculatedTimeDilation, NetworkPhysicsCvars::TimeDilationMin, NetworkPhysicsCvars::TimeDilationMax);
+
+		ClientAckTimeDilation(CalculatedTimeDilation, ActualTimestamp.LocalFrame);
 	}
-	FAsyncPhysicsTimestamp CurrentTimestamp = GetAsyncPhysicsTimestamp();
-	// We need to avoid changing this offset as much as possible since it will invalidate histories and will trigger resim
-	// To deal with that we compute a safe margin based on a user cvar + half the RTT
-	// This margin will only be applied the first time we will compute the offset 
-	const int32 FrameOffset = LocalToServerAsyncPhysicsTickOffset == 0 ? (NetworkPhysicsCvars::NetworkPhysicsPredictionFrameOffset + (CurrentTimestamp.LocalFrame - Timestamp.LocalFrame) / 2) : 0;
+	
+}
 
-	const int32 NewOffset = FMath::Max(LocalToServerAsyncPhysicsTickOffset,Timestamp.ServerFrame - Timestamp.LocalFrame + FrameOffset); //The new offset as reported by the server
-	LocalToServerAsyncPhysicsTickOffset = NewOffset;
+void APlayerController::ClientCorrectionAsyncPhysicsTimestamp_Implementation(FAsyncPhysicsTimestamp Timestamp)
+{
+	ensure(UPhysicsSettings::Get()->PhysicsPrediction.bEnablePhysicsPrediction);
+
+	// Assign async physics tick offset
+	if (!LocalToServerAsyncPhysicsTickOffsetAssigned)
+	{
+		LocalToServerAsyncPhysicsTickOffsetAssigned = true;
+		LocalToServerAsyncPhysicsTickOffset = Timestamp.ServerFrame - Timestamp.LocalFrame;
+	}
 }
 
 void APlayerController::ClientAckTimeDilation_Implementation(float TimeDilation, int32 ServerStep)
 {
-	if(ServerStep < ClientLatestTimeDilationServerStep)
+	if (ClientLatestTimeDilationServerStep <= ServerStep)
 	{
-		//Stale ack so do nothing
 		return;
 	}
-
-	ensureMsgf(ClientLatestTimeDilationServerStep != ServerStep, TEXT("Server should send at most one time dilation per step, duplicate means something is wrong"));
 	ClientLatestTimeDilationServerStep = ServerStep;
+
 	if(UWorld* World = GetWorld())
 	{
 		World->GetPhysicsScene()->SetNetworkDeltaTimeScale(TimeDilation);
@@ -6219,26 +6176,6 @@ void APlayerController::ClientAckTimeDilation_Implementation(float TimeDilation,
 void APlayerController::AsyncPhysicsTickActor(float DeltaTime, float SimTime)
 {
 	Super::AsyncPhysicsTickActor(DeltaTime, SimTime);
-
-	if(UPhysicsSettings::Get()->PhysicsPrediction.bEnablePhysicsPrediction)
-	{
-		//TODO: only kick this off if server and using this feature
-		if (IsLocalController())
-		{
-			return;
-		}
-
-		// If the client has not sent the correct estimate of which server frame its local frame would correspond to, send back which server frame that local frame actually corresponded to
-		const FAsyncPhysicsTimestamp ActualTimestamp = GetAsyncPhysicsTimestamp();
-		
-		if (ServerPendingTimestamp.ServerFrame != INDEX_NONE && ServerPendingTimestamp.ServerFrame != ActualTimestamp.ServerFrame)
-		{
-			ServerLatestTimestampToCorrect.ServerFrame = ActualTimestamp.ServerFrame;
-			ServerLatestTimestampToCorrect.LocalFrame = ServerPendingTimestamp.LocalFrame;
-			ServerPendingTimestamp.ServerFrame = INDEX_NONE;
-		}
-
-	}
 }
 
 
diff --git a/Engine/Source/Runtime/Engine/Public/Physics/NetworkPhysicsComponent.h b/Engine/Source/Runtime/Engine/Public/Physics/NetworkPhysicsComponent.h
--- a/Engine/Source/Runtime/Engine/Public/Physics/NetworkPhysicsComponent.h
+++ b/Engine/Source/Runtime/Engine/Public/Physics/NetworkPhysicsComponent.h
@@ -67,7 +67,7 @@ struct TNetRewindHistory : public Chaos::TDatasRewindHistory<DatasType>
 	{
 		TNetRewindHistory& NetNewDatas = static_cast<TNetRewindHistory&>(NewDatas);
 
-		if ((FrameOffset >= 0) && NetNewDatas.NumFrames > 0)
+		if (NetNewDatas.NumFrames > 0)
 		{
 			for (int32 FrameIndex = 0; FrameIndex < NetNewDatas.NumFrames; ++FrameIndex)
 			{
@@ -268,12 +268,6 @@ struct FNetworkPhysicsCallback : public Chaos::IRewindCallback
 	// but also needs to be able to access GT data (actor iterator, actor state)
 	void UpdateReplicationMap_Internal(int32 PhysicsStep);
 
-	// Update client player on GT
-	void UpdateClientPlayer_External(int32 PhysicsStep);
-
-	// Update server player on GT
-	void UpdateServerPlayer_External(int32 PhysicsStep);
-
 	// Cache the current netmode for use in PT
 	void UpdateNetMode()
 	{
@@ -399,9 +393,6 @@ public:
 	// Delegate linked to the physics rewind callback to send record local inputs/states
 	ENGINE_API void OnPostProcessInputsInternal(const int32 PhysicsStep);
 
-	// Correct the player controller Server to local offset based on the received replicated states
-	ENGINE_API void CorrectServerToLocalOffset(const int32 LocalToServerOffset);
-
 	// Used to create any physics engine information for this component 
 	ENGINE_API virtual void BeginPlay() override;
 
