World Partition: remove the possibility to not unload the map when invoking a commandlet from the editor as this codepath is untested and probably not working.

diff --git a/Engine/Source/Editor/WorldPartitionEditor/Private/WorldPartitionEditorModule.cpp b/Engine/Source/Editor/WorldPartitionEditor/Private/WorldPartitionEditorModule.cpp
--- a/Engine/Source/Editor/WorldPartitionEditor/Private/WorldPartitionEditorModule.cpp
+++ b/Engine/Source/Editor/WorldPartitionEditor/Private/WorldPartitionEditorModule.cpp
@@ -724,7 +724,7 @@ bool FWorldPartitionEditorModule::Build(const FRunBuilderParams& InParams)
 
 	// Unload map if required
 	FString MapPackage = ParamsCopy.World->GetPackage()->GetName();
-	if (ParamsCopy.bUnloadMap && !UnloadCurrentMap(MapPackage))
+	if (!UnloadCurrentMap(MapPackage))
 	{
 		return false;
 	}
@@ -752,18 +752,7 @@ bool FWorldPartitionEditorModule::Build(const FRunBuilderParams& InParams)
 	RunCommandletAsExternalProcess(CommandletArgs, OperationDescription, Result, bCancelled);
 
 	RescanAssets(MapPackage);
-
-	if (ParamsCopy.bUnloadMap)
-	{
-		LoadMap(MapPackage);
-	}
-	else
-	{
-		if (UWorldPartition* WorldPartition = ParamsCopy.World->GetWorldPartition())
-		{
-			WorldPartition->Update();
-		}
-	}
+	LoadMap(MapPackage);
 
 	if (bCancelled)
 	{
diff --git a/Engine/Source/Runtime/Engine/Private/WorldPartition/ActorDescContainer.cpp b/Engine/Source/Runtime/Engine/Private/WorldPartition/ActorDescContainer.cpp
--- a/Engine/Source/Runtime/Engine/Private/WorldPartition/ActorDescContainer.cpp
+++ b/Engine/Source/Runtime/Engine/Private/WorldPartition/ActorDescContainer.cpp
@@ -28,11 +28,6 @@ UActorDescContainer::UActorDescContainer(const FObjectInitializer& ObjectInitial
 {}
 
 #if WITH_EDITOR
-void UActorDescContainer::Initialize(UWorld* InWorld, FName InPackageName)
-{
-	Initialize({ InWorld, InPackageName });
-}
-
 void UActorDescContainer::Initialize(const FInitializeParams& InitParams)
 {
 	TRACE_CPUPROFILER_EVENT_SCOPE(UActorDescContainer::Initialize);
@@ -129,72 +124,6 @@ void UActorDescContainer::Initialize(const FInitializeParams& InitParams)
 	bContainerInitialized = true;
 }
 
-void UActorDescContainer::Update()
-{
-	check(bContainerInitialized);
-	TArray<FAssetData> Assets;
-
-	const FString ContainerExternalActorsPath = GetExternalActorPath();
-
-	// Do a synchronous scan of the level external actors path.			
-	IAssetRegistry& AssetRegistry = FModuleManager::LoadModuleChecked<FAssetRegistryModule>(TEXT("AssetRegistry")).Get();
-	AssetRegistry.ScanPathsSynchronous({ ContainerExternalActorsPath }, /*bForceRescan*/false, /*bIgnoreDenyListScanFilters*/false);
-
-	FARFilter Filter;
-	Filter.bRecursivePaths = true;
-	Filter.bIncludeOnlyOnDiskAssets = true;
-	Filter.PackagePaths.Add(*ContainerExternalActorsPath);
-
-	AssetRegistry.GetAssets(Filter, Assets);
-
-	UWorld* OwningWorld = GetWorldPartition() ? GetWorldPartition()->GetWorld() : nullptr;
-	check(OwningWorld);
-
-	TSet<FGuid> ActorGuids;
-	for (const FAssetData& Asset : Assets)
-	{
-		TUniquePtr<FWorldPartitionActorDesc> NewActorDesc = FWorldPartitionActorDescUtils::GetActorDescriptorFromAssetData(Asset);
-
-		if (NewActorDesc.IsValid() && NewActorDesc->GetNativeClass().IsValid())
-		{
-			ActorGuids.Add(NewActorDesc->GetGuid());
-
-			if (TUniquePtr<FWorldPartitionActorDesc>* ExistingActorDesc = GetActorDescriptor(NewActorDesc->GetGuid()))
-			{
-				if (!NewActorDesc->Equals((*ExistingActorDesc).Get()))
-				{
-					OnActorDescUpdating(ExistingActorDesc->Get());
-					FWorldPartitionActorDescUtils::UpdateActorDescriptorFromActorDescriptor(NewActorDesc, *ExistingActorDesc);
-					OnActorDescUpdated(ExistingActorDesc->Get());
-				}
-			}
-			else
-			{
-				FWorldPartitionActorDesc* ActorDescPtr = NewActorDesc.Release();
-				RegisterActorDescriptor(ActorDescPtr, OwningWorld);
-				OnActorDescAdded(ActorDescPtr);
-			}
-		}
-	}
-
-	TArray<FGuid> ActorDescsToRemove;
-	for (FActorDescList::TIterator<> ActorDescIt(this); ActorDescIt; ++ActorDescIt)
-	{
-		if (!ActorGuids.Contains(ActorDescIt->GetGuid()))
-		{
-			ActorDescsToRemove.Add(ActorDescIt->GetGuid());
-		}
-	}
-
-	for (const FGuid& ActorDescGuidToRemove : ActorDescsToRemove)
-	{
-		if (FWorldPartitionActorDesc* ActorDesc = GetActorDesc(ActorDescGuidToRemove))
-		{
-			RemoveActor(ActorDescGuidToRemove);
-		}
-	}
-}
-
 void UActorDescContainer::Uninitialize()
 {
 	if (bContainerInitialized)
diff --git a/Engine/Source/Runtime/Engine/Private/WorldPartition/WorldPartition.cpp b/Engine/Source/Runtime/Engine/Private/WorldPartition/WorldPartition.cpp
--- a/Engine/Source/Runtime/Engine/Private/WorldPartition/WorldPartition.cpp
+++ b/Engine/Source/Runtime/Engine/Private/WorldPartition/WorldPartition.cpp
@@ -830,20 +830,6 @@ bool UWorldPartition::IsInitialized() const
 	return InitState == EWorldPartitionInitState::Initialized;
 }
 
-void UWorldPartition::Update()
-{
-#if WITH_EDITOR
-	UWorld* OuterWorld = GetTypedOuter<UWorld>();
-	check(OuterWorld);
-	check(!OuterWorld->IsInstanced());
-
-	ForEachActorDescContainer([](UActorDescContainer* InActorDescContainer)
-	{
-		InActorDescContainer->Update();
-	});
-#endif
-}
-
 bool UWorldPartition::SupportsStreaming() const
 {
 	return World ? World->GetWorldSettings()->SupportsWorldPartitionStreaming() : false;
diff --git a/Engine/Source/Runtime/Engine/Public/WorldPartition/ActorDescContainer.h b/Engine/Source/Runtime/Engine/Public/WorldPartition/ActorDescContainer.h
--- a/Engine/Source/Runtime/Engine/Public/WorldPartition/ActorDescContainer.h
+++ b/Engine/Source/Runtime/Engine/Public/WorldPartition/ActorDescContainer.h
@@ -43,10 +43,9 @@ public:
 		TUniqueFunction<bool(const FWorldPartitionActorDesc*)> FilterActorDesc;
 	};
 
-	UE_DEPRECATED(5.1, "UActorDescContainer::Initialize is deprecated, UActorDescContainer::Initialize with UActorDescContainer::FInitializeParams should be used instead.")
-	ENGINE_API void Initialize(UWorld* InWorld, FName InPackageName);
 	ENGINE_API void Initialize(const FInitializeParams& InitParams);
-	ENGINE_API void Update();
+	UE_DEPRECATED(5.4, "UActorDescContainer::Update is deprecated.")
+	ENGINE_API void Update() {}
 	ENGINE_API void Uninitialize();
 
 	bool IsInitialized() const { return bContainerInitialized; }
diff --git a/Engine/Source/Runtime/Engine/Public/WorldPartition/IWorldPartitionEditorModule.h b/Engine/Source/Runtime/Engine/Public/WorldPartition/IWorldPartitionEditorModule.h
--- a/Engine/Source/Runtime/Engine/Public/WorldPartition/IWorldPartitionEditorModule.h
+++ b/Engine/Source/Runtime/Engine/Public/WorldPartition/IWorldPartitionEditorModule.h
@@ -26,7 +26,6 @@ public:
 		UWorld* World = nullptr;
 		FString ExtraArgs;
 		FText OperationDescription;
-		bool bUnloadMap = true;
 	};
 
 	virtual bool RunBuilder(TSubclassOf<UWorldPartitionBuilder> BuilderClass, UWorld* InWorld);
diff --git a/Engine/Source/Runtime/Engine/Public/WorldPartition/WorldPartition.h b/Engine/Source/Runtime/Engine/Public/WorldPartition/WorldPartition.h
--- a/Engine/Source/Runtime/Engine/Public/WorldPartition/WorldPartition.h
+++ b/Engine/Source/Runtime/Engine/Public/WorldPartition/WorldPartition.h
@@ -296,9 +296,6 @@ public:
 	//~ End IWorldPartitionCookPackageGenerator Interface 
 	// End Cooking
 
-	UE_DEPRECATED(5.1, "GetWorldBounds is deprecated, use GetEditorWorldBounds or GetRuntimeWorldBounds instead.")
-	FBox GetWorldBounds() const { return GetRuntimeWorldBounds(); }
-
 	ENGINE_API FBox GetEditorWorldBounds() const;
 	ENGINE_API FBox GetRuntimeWorldBounds() const;
 	
@@ -399,6 +396,7 @@ public:
 	ENGINE_API bool CanInitialize(UWorld* InWorld) const;
 	ENGINE_API void Initialize(UWorld* World, const FTransform& InTransform);
 	ENGINE_API bool IsInitialized() const;
+	UE_DEPRECATED(5.4, "UWorldPartition::Update is deprecated.")
 	ENGINE_API void Update();
 	ENGINE_API void Uninitialize();
 
