Refactored render target to texture conversion so that it can be done on any UTextureRenderTarget / UTexture type :
* Added BP functions to create all types of UTexture from their UTextureRenderTarget counterpart (except UTextureRenderTargetCubeArray, since we don't have much of anything available to produce these ATM, but there's no reason it shouldn't work)
* Fixed crash when trying to create a UTexture out of a UTextureRenderTarget with a non-supported format. The error is also appropriately reported to the user via blueprint message log, instead of a message box in the past

Misc:
* UTextureRenderTarget have to implement CanConvertToTexture to indicate if they are currently able to be converted to a texture (e.g. valid size, supported format, etc.)
* There's now a generic UpdateTexture and ConstructTexture function on UTextureRenderTarget that can be used to update an existing texture or create a new one
* D3D11 and D3D12's RHIReadSurfaceData now support reading from (cube) texture arrays
* Added R8G8 to FColor conversion (the missing .BA channels are filled with 0 and 255, respectively)
* Deprecated all redundant ReadPixels / ReadFloat16Pixels from child classes of FRenderTarget (FTextureResource, FTextureRenderTargetVolumeResource, FTextureRenderTarget2DArrayResource). Note that ReadPixels / ReadFloat16Pixels / ReadLinearColorPixels are now virtual to avoid this situation and yet allow FTextureRenderTargetVolumeResource to implement them in its own way (unlike all other implementations, it uses Read3DSurfaceFloatData internally)
* Fixed typo in FTextureRenderTargetVolumeResource::GetDisplayGamma(), which didn't take PF_R32_FORMAT into account
* Added FTextureSourceFormatInfo and GTextureSourceFormats, which provides information about the various supported formats in ETextureSourceFormat (similar to EPixelFormat and GPixelFormats)

#tests editor

diff --git a/Engine/Plugins/Experimental/LandscapePatch/Source/LandscapePatch/Private/LandscapeTextureBackedRenderTarget.cpp b/Engine/Plugins/Experimental/LandscapePatch/Source/LandscapePatch/Private/LandscapeTextureBackedRenderTarget.cpp
--- a/Engine/Plugins/Experimental/LandscapePatch/Source/LandscapePatch/Private/LandscapeTextureBackedRenderTarget.cpp
+++ b/Engine/Plugins/Experimental/LandscapePatch/Source/LandscapePatch/Private/LandscapeTextureBackedRenderTarget.cpp
@@ -4,6 +4,7 @@
 #include "Engine/Texture2D.h"
 #include "Engine/World.h"
 #include "LandscapePatchUtil.h" // CopyTextureOnRenderThread
+#include "LandscapePatchLogging.h" 
 #include "RenderGraphBuilder.h"
 #include "TextureCompiler.h"
 #include "RenderingThread.h"
@@ -261,7 +262,7 @@ void ULandscapeWeightTextureBackedRenderTarget::CopyToInternalTexture()
 	{
 		Modify();
 		InternalTexture = CreateTexture(this);
-		// The sizing and format doesn't matter because the UpdateTexture2D call below will deal with it.
+		// The sizing and format doesn't matter because the UpdateTexture call below will deal with it.
 	}
 	else
 	{
@@ -271,8 +272,16 @@ void ULandscapeWeightTextureBackedRenderTarget::CopyToInternalTexture()
 		InternalTexture->Modify(false);
 	}
 
-	RenderTarget->UpdateTexture2D(InternalTexture, GetInternalTextureFormat());
-	InternalTexture->UpdateResource();
+	FText ErrorMessage;
+	if (RenderTarget->UpdateTexture(InternalTexture, CTF_Default, /*InAlphaOverride = */nullptr, /*InTextureChangingDelegate =*/ [](UTexture*){}, &ErrorMessage))
+	{
+		check(InternalTexture->Source.GetFormat() == GetInternalTextureFormat());
+		InternalTexture->UpdateResource();
+	}
+	else
+	{
+		UE_LOG(LogLandscapePatch, Error, TEXT("Couldn't copy render target to internal texture: %s"), *ErrorMessage.ToString());
+	}
 #endif // WITH_EDITOR
 }
 
@@ -413,12 +422,17 @@ void ULandscapeHeightTextureBackedRenderTarget::CopyToInternalTexture()
 			});
 	}
 
-	// TODO: The header for this requires the texture to be square power of 2, but it actually doesn't seem to
-	// be an enforced requirement. If that changes, we'll need our own ReadPixels followed by locking a mip
-	// and writing to it.
 	// This call does a flush for us, so the render target should be updated.
-	NativeEncodingRenderTarget->UpdateTexture2D(InternalTexture, GetInternalTextureFormat());
-	InternalTexture->UpdateResource();
+	FText ErrorMessage;
+	if (RenderTarget->UpdateTexture(InternalTexture, CTF_Default, /*InAlphaOverride = */nullptr, /*InTextureChangingDelegate =*/ [](UTexture*) {}, &ErrorMessage))
+	{
+		check(InternalTexture->Source.GetFormat() == GetInternalTextureFormat());
+		InternalTexture->UpdateResource();
+	}
+	else
+	{
+		UE_LOG(LogLandscapePatch, Error, TEXT("Couldn't copy render target to internal texture: %s"), *ErrorMessage.ToString());
+	}
 #endif // WITH_EDITOR
 }
 
diff --git a/Engine/Plugins/Experimental/LandscapePatch/Source/LandscapePatch/Private/LandscapeTexturePatch.cpp b/Engine/Plugins/Experimental/LandscapePatch/Source/LandscapePatch/Private/LandscapeTexturePatch.cpp
--- a/Engine/Plugins/Experimental/LandscapePatch/Source/LandscapePatch/Private/LandscapeTexturePatch.cpp
+++ b/Engine/Plugins/Experimental/LandscapePatch/Source/LandscapePatch/Private/LandscapeTexturePatch.cpp
@@ -667,22 +667,22 @@ void ULandscapeTexturePatch::ReinitializeHeight(UTextureRenderTarget2D* InCombin
 
 		FRDGBuilder GraphBuilder(RHICmdList, RDG_EVENT_NAME("LandscapeTexturePatchReinitializeHeight"));
 
-		FReinitializeLandscapePatchPS::FParameters* HeightmapResalmpleParams = GraphBuilder.AllocParameters<FReinitializeLandscapePatchPS::FParameters>();
+		FReinitializeLandscapePatchPS::FParameters* HeightmapResampleParams = GraphBuilder.AllocParameters<FReinitializeLandscapePatchPS::FParameters>();
 
 		FRDGTextureRef HeightmapSource = GraphBuilder.RegisterExternalTexture(CreateRenderTarget(Source->GetTexture2DRHI(), TEXT("ReinitializationSource")));
 		FRDGTextureSRVRef SourceSRV = GraphBuilder.CreateSRV(FRDGTextureSRVDesc::CreateForMipLevel(HeightmapSource, 0));
-		HeightmapResalmpleParams->InSource = SourceSRV;
-		HeightmapResalmpleParams->InSourceSampler = TStaticSamplerState<SF_Bilinear, AM_Clamp, AM_Clamp>::GetRHI();
-		HeightmapResalmpleParams->InPatchToSource = PatchToSource;
+		HeightmapResampleParams->InSource = SourceSRV;
+		HeightmapResampleParams->InSourceSampler = TStaticSamplerState<SF_Bilinear, AM_Clamp, AM_Clamp>::GetRHI();
+		HeightmapResampleParams->InPatchToSource = PatchToSource;
 
 		FRDGTextureRef DestinationTexture = GraphBuilder.RegisterExternalTexture(CreateRenderTarget(Destination->GetTexture2DRHI(), TEXT("ReinitializationDestination")));
 
 		if (OffsetToApply != 0)
 		{
 			FRDGTextureRef TemporaryDestination = GraphBuilder.CreateTexture(DestinationTexture->Desc, TEXT("LandscapeTextureHeightPatchInputCopy"));
-			HeightmapResalmpleParams->RenderTargets[0] = FRenderTargetBinding(TemporaryDestination, ERenderTargetLoadAction::ENoAction, /*InMipIndex = */0);
+			HeightmapResampleParams->RenderTargets[0] = FRenderTargetBinding(TemporaryDestination, ERenderTargetLoadAction::ENoAction, /*InMipIndex = */0);
 
-			FReinitializeLandscapePatchPS::AddToRenderGraph(GraphBuilder, HeightmapResalmpleParams);
+			FReinitializeLandscapePatchPS::AddToRenderGraph(GraphBuilder, HeightmapResampleParams);
 
 			FOffsetHeightmapPS::FParameters* OffsetParams = GraphBuilder.AllocParameters<FOffsetHeightmapPS::FParameters>();
 
@@ -695,8 +695,8 @@ void ULandscapeTexturePatch::ReinitializeHeight(UTextureRenderTarget2D* InCombin
 		}
 		else
 		{
-			HeightmapResalmpleParams->RenderTargets[0] = FRenderTargetBinding(DestinationTexture, ERenderTargetLoadAction::ENoAction, /*InMipIndex = */0);
-			FReinitializeLandscapePatchPS::AddToRenderGraph(GraphBuilder, HeightmapResalmpleParams);
+			HeightmapResampleParams->RenderTargets[0] = FRenderTargetBinding(DestinationTexture, ERenderTargetLoadAction::ENoAction, /*InMipIndex = */0);
+			FReinitializeLandscapePatchPS::AddToRenderGraph(GraphBuilder, HeightmapResampleParams);
 		}
 
 		GraphBuilder.Execute();
@@ -704,9 +704,19 @@ void ULandscapeTexturePatch::ReinitializeHeight(UTextureRenderTarget2D* InCombin
 
 	// The Modify() calls currently don't really help because we don't transact inside Render_Native. Maybe someday
 	// we'll add that ability (though it sounds messy).
-	HeightInternalData->GetInternalTexture()->Modify();
-	TemporaryNativeHeightCopy->UpdateTexture2D(HeightInternalData->GetInternalTexture(), ETextureSourceFormat::TSF_BGRA8);
-	HeightInternalData->GetInternalTexture()->UpdateResource();
+	UTexture2D* InternalTexture = HeightInternalData->GetInternalTexture();
+	InternalTexture->Modify();
+	FText ErrorMessage;
+	if (TemporaryNativeHeightCopy->UpdateTexture(InternalTexture, CTF_Default, /*InAlphaOverride = */nullptr, /*InTextureChangingDelegate =*/ [](UTexture*) {}, &ErrorMessage))
+	{
+		check(InternalTexture->Source.GetFormat() == ETextureSourceFormat::TSF_BGRA8);
+		InternalTexture->UpdateResource();
+	}
+	else
+	{
+		UE_LOG(LogLandscapePatch, Error, TEXT("Couldn't copy heightmap render target to internal texture: %s"), *ErrorMessage.ToString());
+	}
+	InternalTexture->UpdateResource();
 
 	if (IsValid(HeightInternalData->GetRenderTarget()))
 	{
diff --git a/Engine/Source/Runtime/D3D12RHI/Private/D3D12RenderTarget.cpp b/Engine/Source/Runtime/D3D12RHI/Private/D3D12RenderTarget.cpp
--- a/Engine/Source/Runtime/D3D12RHI/Private/D3D12RenderTarget.cpp
+++ b/Engine/Source/Runtime/D3D12RHI/Private/D3D12RenderTarget.cpp
@@ -372,9 +372,6 @@ TRefCountPtr<FD3D12Resource> FD3D12DynamicRHI::GetStagingTexture(FRHITexture* Te
 	FD3D12Texture* Texture = GetD3D12TextureFromRHITexture(TextureRHI, GPUIndex);
 	D3D12_RESOURCE_DESC const& SourceDesc = Texture->GetResource()->GetDesc();
 
-	// Ensure we're dealing with a Texture2D, which the rest of this function already assumes
-	check(TextureRHI->GetTexture2D());
-
 	bool bRequiresTempStagingTexture = Texture->GetResource()->GetHeapType() != D3D12_HEAP_TYPE_READBACK;
 	if (bRequiresTempStagingTexture == false)
 	{
@@ -422,11 +419,12 @@ TRefCountPtr<FD3D12Resource> FD3D12DynamicRHI::GetStagingTexture(FRHITexture* Te
 	if (Texture->GetDesc().IsTextureCube())
 	{
 		uint32 D3DFace = GetD3D12CubeFace(InFlags.GetCubeFace());
-		Subresource = CalcSubresource(InFlags.GetMip(), D3DFace, TextureRHI->GetNumMips());
+		Subresource = CalcSubresource(InFlags.GetMip(), InFlags.GetArrayIndex() * 6 + D3DFace, TextureRHI->GetNumMips());
 	}
 	else
 	{
-		Subresource = CalcSubresource(InFlags.GetMip(), 0, TextureRHI->GetNumMips());
+		const bool bIsTextureArray = Texture->GetDesc().IsTextureArray();
+		Subresource = CalcSubresource(InFlags.GetMip(), bIsTextureArray ? InFlags.GetArrayIndex() : 0, TextureRHI->GetNumMips());
 	}
 
 	D3D12_BOX* RectPtr = nullptr; // API prefers NULL for entire texture.
@@ -581,6 +579,9 @@ static void ConvertDXGIToFColor(DXGI_FORMAT Format, uint32 Width, uint32 Height,
 	case DXGI_FORMAT_R8_UNORM:
 		ConvertRawR8DataToFColor(Width, Height, In, SrcPitch, Out);
 		break;
+	case DXGI_FORMAT_R8G8_UNORM:
+		ConvertRawR8G8DataToFColor(Width, Height, In, SrcPitch, Out);
+		break;
 	default:
 		checkf(0, TEXT("Unknown surface format!"));
 		break;
@@ -1040,19 +1041,17 @@ void FD3D12DynamicRHI::ReadSurfaceDataMSAARaw(FRHITexture* TextureRHI, FIntRect
 	const uint32 MipBytesAligned = XBytesAligned * SizeY;
 	VERIFYD3D12RESULT(Adapter->CreateBuffer(D3D12_HEAP_TYPE_READBACK, NodeMask, NodeMask, MipBytesAligned, StagingTexture2D.GetInitReference(), nullptr));
 
-	// Ensure we're dealing with a Texture2D, which the rest of this function already assumes
-	check(TextureRHI->GetTexture2D());
-
 	// Determine the subresource index for cubemaps.
 	uint32 Subresource = 0;
 	if (Texture->GetDesc().IsTextureCube())
 	{
 		uint32 D3DFace = GetD3D12CubeFace(InFlags.GetCubeFace());
-		Subresource = CalcSubresource(InFlags.GetMip(), D3DFace, TextureRHI->GetNumMips());
+		Subresource = CalcSubresource(InFlags.GetMip(), InFlags.GetArrayIndex() * 6 + D3DFace, TextureRHI->GetNumMips());
 	}
 	else
 	{
-		Subresource = CalcSubresource(InFlags.GetMip(), 0, TextureRHI->GetNumMips());
+		const bool bIsTextureArray = Texture->GetDesc().IsTextureArray();
+		Subresource = CalcSubresource(InFlags.GetMip(), bIsTextureArray ? InFlags.GetArrayIndex() : 0, TextureRHI->GetNumMips());
 	}
 
 	// Setup the descriptions for the copy to the readback heap.
@@ -1416,7 +1415,7 @@ void FD3D12DynamicRHI::RHIRead3DSurfaceFloatData(FRHITexture* TextureRHI, FIntRe
 	}
 	else if (bIsR32FFmt)
 	{
-		// Texture is R16F format
+		// Texture is PF_R32_FLOAT format
 		for (int32 Z = ZMinMax.X; Z < ZMinMax.Y; ++Z)
 		{
 			for (int32 Y = InRect.Min.Y; Y < InRect.Max.Y; ++Y)
diff --git a/Engine/Source/Runtime/Engine/Classes/Engine/TextureDefines.h b/Engine/Source/Runtime/Engine/Classes/Engine/TextureDefines.h
--- a/Engine/Source/Runtime/Engine/Classes/Engine/TextureDefines.h
+++ b/Engine/Source/Runtime/Engine/Classes/Engine/TextureDefines.h
@@ -13,6 +13,8 @@
 #include "TextureDefines.generated.h"
 #endif // WITH_ENGINE
 
+enum EPixelFormat : uint8;
+
 /**
  * @warning: if this is changed:
  *     update BaseEngine.ini [SystemSettings]
@@ -355,6 +357,22 @@ enum ETextureSourceFormat : int
 	TSF_RGBE8 UE_DEPRECATED(5.1,"Legacy ETextureSourceFormat not supported, use BGRE8") = TSF_RGBE8_DEPRECATED
 };
 
+/**
+* Information about a texture source format
+*/
+struct FTextureSourceFormatInfo
+{
+	FTextureSourceFormatInfo() = delete;
+	FTextureSourceFormatInfo(ETextureSourceFormat InTextureSourceFormat, EPixelFormat InPixelFormat, const TCHAR* InName, int32 InNumComponents, int32 InBytesPerPixel);
+
+	ETextureSourceFormat TextureSourceFormat;
+	EPixelFormat PixelFormat;
+	const TCHAR* Name;
+	int32 NumComponents;
+	int32 BytesPerPixel;
+};
+extern ENGINE_API FTextureSourceFormatInfo GTextureSourceFormats[TSF_MAX];		// Maps members of ETextureSourceFormat to a FTextureSourceFormatInfo describing the format.
+
 // This needs to be mirrored in EditorFactories.cpp.
 // TC_EncodedReflectionCapture is no longer used and could be deleted
 UENUM()
diff --git a/Engine/Source/Runtime/Engine/Classes/Engine/TextureRenderTarget.h b/Engine/Source/Runtime/Engine/Classes/Engine/TextureRenderTarget.h
--- a/Engine/Source/Runtime/Engine/Classes/Engine/TextureRenderTarget.h
+++ b/Engine/Source/Runtime/Engine/Classes/Engine/TextureRenderTarget.h
@@ -12,8 +12,11 @@
 #include "CoreMinimal.h"
 #include "UObject/ObjectMacros.h"
 #include "Engine/Texture.h"
+#include "RenderUtils.h"
 #include "TextureRenderTarget.generated.h"
 
+enum EConstructTextureFlags : uint32;
+
 UCLASS(abstract, MinimalAPI)
 class UTextureRenderTarget : public UTexture
 {
@@ -48,6 +51,56 @@ class UTextureRenderTarget : public UTexture
 	ENGINE_API virtual class FTextureResource* CreateResource() override;
 	ENGINE_API virtual EMaterialValueType GetMaterialType() const override;
 	//~ End UTexture Interface
+
+	/**
+	 * Validates that the UTextureRenderTarget can be converted to a UTexture (e.g. supported format, valid size, etc.)
+	 * @param OutPixelFormat - The format of the render target (only if conversion is possible)
+	 * @param OutTextureSourceFormat - The format of the texture that would be used if the UTextureRenderTarget was converted to a UTexture (only if conversion is possible)
+	 * @param OutErrorMessage - (Optional) If passed, will contain an error message in case conversion is impossible
+	 * @return True if render target can be converted to a texture
+	 */
+	ENGINE_API virtual bool CanConvertToTexture(ETextureSourceFormat& OutTextureSourceFormat, EPixelFormat& OutPixelFormat, FText* OutErrorMessage = nullptr) const PURE_VIRTUAL(UTextureRenderTarget, return false;);
+
+	/**
+	 * Validates that the UTextureRenderTarget can be converted to a UTexture (e.g. supported format, valid size, etc.)
+	 * @param OutErrorMessage - (Optional) If passed, will contain an error message in case conversion is impossible
+	 * @return True if render target can be converted to a texture
+	 */
+	ENGINE_API bool CanConvertToTexture(FText* OutErrorMessage = nullptr) const;
+
+	/**
+	 * Returns the UTexture class that corresponds to this render target (e.g. UTexture2D for UTextureRenderTarget2D)
+	 */
+	ENGINE_API virtual TSubclassOf<UTexture> GetTextureUClass() const PURE_VIRTUAL(UTextureRenderTarget, return nullptr;);
+
+#if WITH_EDITOR
+	/**
+	 * Utility for updating an existing UTexture from a TextureRenderTarget
+	 * @param InTexture - Texture which will contain the content of this render target after the call.
+	 * @param InFlags - Various control flags for operation (see EConstructTextureFlags)
+	 * @param InAlphaOverride - (Optional) the values here will become the alpha values in the resulting texture
+	 * @param InOnTextureChangingDelegate - (Optional) If the texture needs to be modified (as its properties or content will change), this delegate will be called beforehand. 
+	 * @param OutErrorMessage - (Optional) Error message in case of failure
+	 * @return true if conversion from render to texture was successful
+	 */
+	using FOnTextureChangingDelegate = TFunctionRef<void(UTexture* /*InTexture*/)>;
+	ENGINE_API bool UpdateTexture(UTexture* InTexture, EConstructTextureFlags InFlags = CTF_Default, const TArray<uint8>* InAlphaOverride = nullptr, FOnTextureChangingDelegate InOnTextureChangingDelegate = [](UTexture*){}, FText* OutErrorMessage = nullptr);
+
+	/**
+	 * Utility for creating a new UTexture from a UTextureRenderTarget. The method can fail if conversion is impossible
+	 * @param InOuter - Outer to use when constructing the new UTexture.
+	 * @param InNewTextureName - Name of new UTexture object.
+	 * @param InObjectFlags - Flags to apply to the new UTexture object
+	 * @param InFlags - Various control flags for operation (see EConstructTextureFlags)
+	 * @param InAlphaOverride - (Optional) the values here will become the alpha values in the resulting texture (8 bits alpha, regardless of the texture's format)
+	 * @param OutErrorMessage - (Optional) Error message in case of failure
+	 * @return New UTexture object.
+	 */
+	ENGINE_API UTexture* ConstructTexture(UObject* InOuter, const FString& InNewTextureName, EObjectFlags InObjectFlags = RF_NoFlags, EConstructTextureFlags InFlags = CTF_Default, const TArray<uint8>* InAlphaOverride = nullptr, FText* OutErrorMessage = nullptr);
+#endif // WITH_EDITOR
+
+protected:
+	ETextureSourceFormat ValidateTextureFormatForConversionToTextureInternal(EPixelFormat InFormat, const TArrayView<const EPixelFormat>& InCompatibleFormats, FText* OutErrorMessage) const;
 };
 
 
diff --git a/Engine/Source/Runtime/Engine/Classes/Engine/TextureRenderTarget2D.h b/Engine/Source/Runtime/Engine/Classes/Engine/TextureRenderTarget2D.h
--- a/Engine/Source/Runtime/Engine/Classes/Engine/TextureRenderTarget2D.h
+++ b/Engine/Source/Runtime/Engine/Classes/Engine/TextureRenderTarget2D.h
@@ -156,39 +156,49 @@ class UTextureRenderTarget2D : public UTextureRenderTarget
 	ENGINE_API void ResizeTarget(uint32 InSizeX, uint32 InSizeY);
 
 	/**
-	 * Utility for creating a new UTexture2D from a TextureRenderTarget2D
-	 * TextureRenderTarget2D must be square and a power of two size.
-	 * @param Outer - Outer to use when constructing the new Texture2D.
-	 * @param NewTexName - Name of new UTexture2D object.
-	 * @param ObjectFlags - Flags to apply to the new Texture2D object
-	 * @param Flags - Various control flags for operation (see EConstructTextureFlags)
-	 * @param AlphaOverride - If specified, the values here will become the alpha values in the resulting texture
+	 * Utility for creating a new UTexture2D from a UTextureRenderTarget2D
+	 * @param InOuter - Outer to use when constructing the new UTexture2D.
+	 * @param InNewTextureName - Name of new UTexture2D object.
+	 * @param InObjectFlags - Flags to apply to the new UTexture2D object
+	 * @param InFlags - Various control flags for operation (see EConstructTextureFlags)
+	 * @param InAlphaOverride - If specified, the values here will become the alpha values in the resulting texture
 	 * @return New UTexture2D object.
 	 */
-	ENGINE_API UTexture2D* ConstructTexture2D(UObject* InOuter, const FString& NewTexName, EObjectFlags InObjectFlags, uint32 Flags=CTF_Default, TArray<uint8>* AlphaOverride=NULL);
+	ENGINE_API UTexture2D* ConstructTexture2D(UObject* InOuter, const FString& InNewTextureName, EObjectFlags InObjectFlags, uint32 InFlags = CTF_Default, TArray<uint8>* InAlphaOverride = nullptr);
+	
+	UE_DEPRECATED(5.4, "Use CanConvertToTexture")
 	ENGINE_API ETextureSourceFormat GetTextureFormatForConversionToTexture2D() const;
 
+	//~ Begin UTextureRenderTarget Interface
+	virtual bool CanConvertToTexture(ETextureSourceFormat& OutTextureSourceFormat, EPixelFormat& OutPixelFormat, FText* OutErrorMessage) const override;
+	virtual TSubclassOf<UTexture> GetTextureUClass() const override;
+	//~ End UTextureRenderTarget Interface
+
 	/**
 	 * Utility for updating an existing UTexture2D from a TextureRenderTarget2D
-	 * TextureRenderTarget2D must be square and a power of two size.
 	 * @param InTexture2D					Texture which will contain the content of this render target after the call.
 	 * @param InTextureFormat				Format in which the texture should be stored.
 	 * @param Flags				Optional	Various control flags for operation (see EConstructTextureFlags)
 	 * @param AlphaOverride		Optional	If non-null, the values here will become the alpha values in the resulting texture
 	 */
-	ENGINE_API void UpdateTexture2D(UTexture2D* InTexture2D, ETextureSourceFormat InTextureFormat, uint32 Flags = CTF_Default, TArray<uint8>* AlphaOverride = NULL);
+	UE_DEPRECATED(5.4, "Use URenderTarget::UpdateTexture")
+	ENGINE_API void UpdateTexture2D(UTexture2D* InTexture2D, ETextureSourceFormat InTextureFormat, uint32 Flags = CTF_Default, TArray<uint8>* AlphaOverride = nullptr);
 
 	/**
 	 * Utility for updating an existing UTexture2D from a TextureRenderTarget2D
-	 * TextureRenderTarget2D must be square and a power of two size.
 	 * @param InTexture2D				Texture which will contain the content of this render target after the call.
 	 * @param InTextureFormat			Format in which the texture should be stored.
 	 * @param Flags						Various control flags for operation (see EConstructTextureFlags)
 	 * @param AlphaOverride				If non-null, the values here will become the alpha values in the resulting texture
 	 * @param TextureChangingDelegate	If the texture needs to be modified (as it's properties or content will change), this delegate will be called beforehand.
 	 */
+	UE_DEPRECATED(5.4, "Use URenderTarget::FOnTextureChangingDelegate")
 	DECLARE_DELEGATE_OneParam(FTextureChangingDelegate, UTexture2D*);
+
+PRAGMA_DISABLE_DEPRECATION_WARNINGS
+	UE_DEPRECATED(5.4, "Use URenderTarget::UpdateTexture")
 	ENGINE_API void UpdateTexture2D(UTexture2D* InTexture2D, ETextureSourceFormat InTextureFormat, uint32 Flags, TArray<uint8>* AlphaOverride, FTextureChangingDelegate TextureChangingDelegate);
+PRAGMA_ENABLE_DEPRECATION_WARNINGS
 
 	/**
 	 * Updates (resolves) the render target texture immediately.
diff --git a/Engine/Source/Runtime/Engine/Classes/Engine/TextureRenderTarget2DArray.h b/Engine/Source/Runtime/Engine/Classes/Engine/TextureRenderTarget2DArray.h
--- a/Engine/Source/Runtime/Engine/Classes/Engine/TextureRenderTarget2DArray.h
+++ b/Engine/Source/Runtime/Engine/Classes/Engine/TextureRenderTarget2DArray.h
@@ -66,14 +66,15 @@ class UTextureRenderTarget2DArray : public UTextureRenderTarget
 	ENGINE_API void UpdateResourceImmediate(bool bClearRenderTarget/*=true*/);
 
 	/**
-	* Utility for creating a new UTexture2DArray from a TextureRenderTarget2DArray.
-	* TextureRenderTarget2DArray must be square and a power of two size.
-	* @param	Outer			Outer to use when constructing the new Texture2DArray.
-	* @param	NewTexName		Name of new UTexture2DArray object.
-	* @param	Flags			Various control flags for operation (see EObjectFlags)
-	* @return					New UTexture2DArray object.
-	*/
-	ENGINE_API class UTexture2DArray* ConstructTexture2DArray(UObject* InOuter, const FString& NewTexName, EObjectFlags InFlags);
+	 * Utility for creating a new UTexture2DArray from a UTextureRenderTarget2DArray
+	 * @param InOuter - Outer to use when constructing the new UTexture2DArray.
+	 * @param InNewTextureName - Name of new UTexture2DArray object.
+	 * @param InObjectFlags - Flags to apply to the new UTexture2DArray object
+	 * @param InFlags - Various control flags for operation (see EConstructTextureFlags)
+	 * @param InAlphaOverride - If specified, the values here will become the alpha values in the resulting texture
+	 * @return New UTexture2DArray object.
+	 */
+	ENGINE_API class UTexture2DArray* ConstructTexture2DArray(UObject* InOuter, const FString& InNewTextureName, EObjectFlags InObjectFlags, uint32 InFlags = CTF_Default, TArray<uint8>* InAlphaOverride = nullptr);
 
 	//~ Begin UTexture Interface.
 	virtual float GetSurfaceWidth() const  override { return static_cast<float>(SizeX); }
@@ -109,7 +110,12 @@ class UTextureRenderTarget2DArray : public UTextureRenderTarget
 	virtual void PostLoad() override;
 	virtual void GetResourceSizeEx(FResourceSizeEx& CumulativeResourceSize) override;
 	virtual FString GetDesc() override;
-	//~ Begin UObject Interface
+	//~ End UObject Interface
+
+	//~ Begin UTextureRenderTarget Interface
+	virtual bool CanConvertToTexture(ETextureSourceFormat& OutTextureSourceFormat, EPixelFormat& OutPixelFormat, FText* OutErrorMessage) const override;
+	virtual TSubclassOf<UTexture> GetTextureUClass() const override;
+	//~ End UTextureRenderTarget Interface
 };
 
 
diff --git a/Engine/Source/Runtime/Engine/Classes/Engine/TextureRenderTargetCube.h b/Engine/Source/Runtime/Engine/Classes/Engine/TextureRenderTargetCube.h
--- a/Engine/Source/Runtime/Engine/Classes/Engine/TextureRenderTargetCube.h
+++ b/Engine/Source/Runtime/Engine/Classes/Engine/TextureRenderTargetCube.h
@@ -57,14 +57,15 @@ class UTextureRenderTargetCube : public UTextureRenderTarget
 	ENGINE_API void UpdateResourceImmediate(bool bClearRenderTarget=true);
 
 	/**
-	* Utility for creating a new UTextureCube from a TextureRenderTargetCube.
-	* TextureRenderTargetCube must be square and a power of two size.
-	* @param	Outer			Outer to use when constructing the new TextureCube.
-	* @param	NewTexName		Name of new UTextureCube object.
-	* @param	Flags			Various control flags for operation (see EObjectFlags)
-	* @return					New UTextureCube object.
-	*/
-	ENGINE_API class UTextureCube* ConstructTextureCube(UObject* InOuter, const FString& NewTexName, EObjectFlags InFlags);
+	 * Utility for creating a new UTextureCube from a UTextureRenderTargetCube
+	 * @param InOuter - Outer to use when constructing the new UTextureCube.
+	 * @param InNewTextureName - Name of new UTextureCube object.
+	 * @param InObjectFlags - Flags to apply to the new UTextureCube object
+	 * @param InFlags - Various control flags for operation (see EConstructTextureFlags)
+	 * @param InAlphaOverride - If specified, the values here will become the alpha values in the resulting texture
+	 * @return New UTextureCube object.
+	 */
+	ENGINE_API class UTextureCube* ConstructTextureCube(UObject* InOuter, const FString& InNewTextureName, EObjectFlags InObjectFlags, uint32 InFlags = CTF_Default, TArray<uint8>* InAlphaOverride = nullptr);
 
 	//~ Begin UTexture Interface.
 	virtual float GetSurfaceWidth() const  override { return static_cast<float>(SizeX); }
@@ -100,7 +101,12 @@ class UTextureRenderTargetCube : public UTextureRenderTarget
 	virtual void PostLoad() override;
 	virtual void GetResourceSizeEx(FResourceSizeEx& CumulativeResourceSize) override;
 	virtual FString GetDesc() override;
-	//~ Begin UObject Interface
+	//~ End UObject Interface
+
+	//~ Begin UTextureRenderTarget Interface
+	virtual bool CanConvertToTexture(ETextureSourceFormat& OutTextureSourceFormat, EPixelFormat& OutPixelFormat, FText* OutErrorMessage) const override;
+	virtual TSubclassOf<UTexture> GetTextureUClass() const override;
+	//~ End UTextureRenderTarget Interface
 };
 
 
diff --git a/Engine/Source/Runtime/Engine/Classes/Engine/TextureRenderTargetVolume.h b/Engine/Source/Runtime/Engine/Classes/Engine/TextureRenderTargetVolume.h
--- a/Engine/Source/Runtime/Engine/Classes/Engine/TextureRenderTargetVolume.h
+++ b/Engine/Source/Runtime/Engine/Classes/Engine/TextureRenderTargetVolume.h
@@ -65,14 +65,15 @@ class UTextureRenderTargetVolume : public UTextureRenderTarget
 	ENGINE_API void UpdateResourceImmediate(bool bClearRenderTarget/*=true*/);
 
 	/**
-	* Utility for creating a new UVolumeTexture from a TextureRenderTargetVolume.
-	* TextureRenderTargetVolume must be square and a power of two size.
-	* @param	Outer			Outer to use when constructing the new VolumeTexture.
-	* @param	NewTexName		Name of new UVolumeTexture object.
-	* @param	Flags			Various control flags for operation (see EObjectFlags)
-	* @return					New UVolumeTexture object.
-	*/
-	ENGINE_API class UVolumeTexture* ConstructTextureVolume(UObject* InOuter, const FString& NewTexName, EObjectFlags InFlags);
+	 * Utility for creating a new UVolumeTexture from a UTextureRenderTargetVolume
+	 * @param InOuter - Outer to use when constructing the new UVolumeTexture.
+	 * @param InNewTextureName - Name of new UVolumeTexture object.
+	 * @param InObjectFlags - Flags to apply to the new UVolumeTexture object
+	 * @param InFlags - Various control flags for operation (see EConstructTextureFlags)
+	 * @param InAlphaOverride - If specified, the values here will become the alpha values in the resulting texture
+	 * @return New UVolumeTexture object.
+	 */
+	ENGINE_API class UVolumeTexture* ConstructTextureVolume(UObject* InOuter, const FString& InNewTextureName, EObjectFlags InObjectFlags, uint32 InFlags = CTF_Default, TArray<uint8>* InAlphaOverride = nullptr);
 
 	//~ Begin UTexture Interface.
 	virtual float GetSurfaceWidth() const  override { return static_cast<float>(SizeX); }
@@ -106,7 +107,12 @@ class UTextureRenderTargetVolume : public UTextureRenderTarget
 	virtual void PostLoad() override;
 	virtual void GetResourceSizeEx(FResourceSizeEx& CumulativeResourceSize) override;
 	virtual FString GetDesc() override;
-	//~ Begin UObject Interface
+	//~ End UObject Interface
+
+	//~ Begin UTextureRenderTarget Interface
+	virtual bool CanConvertToTexture(ETextureSourceFormat& OutTextureSourceFormat, EPixelFormat& OutPixelFormat, FText* OutErrorMessage) const override;
+	virtual TSubclassOf<UTexture> GetTextureUClass() const override;
+	//~ End UTextureRenderTarget Interface
 };
 
 
diff --git a/Engine/Source/Runtime/Engine/Classes/Kismet/KismetRenderingLibrary.h b/Engine/Source/Runtime/Engine/Classes/Kismet/KismetRenderingLibrary.h
--- a/Engine/Source/Runtime/Engine/Classes/Kismet/KismetRenderingLibrary.h
+++ b/Engine/Source/Runtime/Engine/Classes/Kismet/KismetRenderingLibrary.h
@@ -86,16 +86,58 @@ class UKismetRenderingLibrary : public UBlueprintFunctionLibrary
 	* Creates a new Static Texture from a Render Target 2D.
 	* Only works in the editor
 	*/
-	UFUNCTION(BlueprintCallable, meta = (DisplayName = "Render Target Create Static Texture Editor Only", Keywords = "Create Static Texture from Render Target", UnsafeDuringActorConstruction = "true"), Category = Game)
+	UFUNCTION(BlueprintCallable, Category = "Editor Scripting | Rendering", meta = (DisplayName = "Render Target 2D Create Static Texture 2D Editor Only", Keywords = "Create Static Texture from Render Target", UnsafeDuringActorConstruction = "true"))
 	static ENGINE_API UTexture2D* RenderTargetCreateStaticTexture2DEditorOnly(UTextureRenderTarget2D* RenderTarget, FString Name = "Texture", enum TextureCompressionSettings CompressionSettings = TC_Default, enum TextureMipGenSettings MipSettings = TMGS_FromTextureGroup);
 
 	/**
-	 * Copies the contents of a render target to a UTexture2D
+	* Creates a new Static Texture 2D Array from a Render Target 2D Array.
+	* Only works in the editor
+	*/
+	UFUNCTION(BlueprintCallable, Category = "Editor Scripting | Rendering", meta = (DisplayName = "Render Target 2D Array Create Static Texture 2D Array Editor Only", Keywords = "Create Static Texture from Render Target", UnsafeDuringActorConstruction = "true"))
+	static ENGINE_API UTexture2DArray* RenderTargetCreateStaticTexture2DArrayEditorOnly(UTextureRenderTarget2DArray* RenderTarget, FString Name = "Texture", enum TextureCompressionSettings CompressionSettings = TC_Default, enum TextureMipGenSettings MipSettings = TMGS_FromTextureGroup);
+
+	/**
+	* Creates a new Static Texture Cube from a Render Target Cube.
+	* Only works in the editor
+	*/
+	UFUNCTION(BlueprintCallable, Category = "Editor Scripting | Rendering", meta = (DisplayName = "Render Target Cube Create Static Texture Cube Editor Only", Keywords = "Create Static Texture from Render Target", UnsafeDuringActorConstruction = "true"))
+	static ENGINE_API UTextureCube* RenderTargetCreateStaticTextureCubeEditorOnly(UTextureRenderTargetCube* RenderTarget, FString Name = "Texture", enum TextureCompressionSettings CompressionSettings = TC_Default, enum TextureMipGenSettings MipSettings = TMGS_FromTextureGroup);
+
+	/**
+	* Creates a new Static Volume Texture from a Render Target Volume.
+	* Only works in the editor
+	*/
+	UFUNCTION(BlueprintCallable, Category = "Editor Scripting | Rendering", meta = (DisplayName = "Render Target Volume Create Static Volume Texture Editor Only", Keywords = "Create Static Texture from Render Target", UnsafeDuringActorConstruction = "true"))
+	static ENGINE_API UVolumeTexture* RenderTargetCreateStaticVolumeTextureEditorOnly(UTextureRenderTargetVolume* RenderTarget, FString Name = "Texture", enum TextureCompressionSettings CompressionSettings = TC_Default, enum TextureMipGenSettings MipSettings = TMGS_FromTextureGroup);
+
+	/**
+	 * Copies the contents of a UTextureRenderTarget2D to a UTexture2D
 	 * Only works in the editor
 	 */
-	UFUNCTION(BlueprintCallable, Category = "Rendering", meta = (Keywords = "ConvertRenderTarget", WorldContext = "WorldContextObject", UnsafeDuringActorConstruction = "true"))
+	UFUNCTION(BlueprintCallable, Category = "Editor Scripting | Rendering", meta = (Keywords = "Convert Render Target", WorldContext = "WorldContextObject", UnsafeDuringActorConstruction = "true"))
 	static ENGINE_API void ConvertRenderTargetToTexture2DEditorOnly(UObject* WorldContextObject, UTextureRenderTarget2D* RenderTarget, UTexture2D* Texture);
 
+	/**
+	 * Copies the contents of a UTextureRenderTarget2DArray to a UTexture2DArray
+	 * Only works in the editor
+	 */
+	UFUNCTION(BlueprintCallable, Category = "Editor Scripting | Rendering", meta = (Keywords = "Convert Render Target", WorldContext = "WorldContextObject", UnsafeDuringActorConstruction = "true"))
+	static ENGINE_API void ConvertRenderTargetToTexture2DArrayEditorOnly(UObject* WorldContextObject, UTextureRenderTarget2DArray* RenderTarget, UTexture2DArray* Texture);
+
+	/**
+	 * Copies the contents of a UTextureRenderTargetCube to a UTextureCube
+	 * Only works in the editor
+	 */
+	UFUNCTION(BlueprintCallable, Category = "Editor Scripting | Rendering", meta = (Keywords = "Convert Render Target", WorldContext = "WorldContextObject", UnsafeDuringActorConstruction = "true"))
+	static ENGINE_API void ConvertRenderTargetToTextureCubeEditorOnly(UObject* WorldContextObject, UTextureRenderTargetCube* RenderTarget, UTextureCube* Texture);
+
+	/**
+	 * Copies the contents of a UTextureRenderTargetVolume to a UVolumeTexture
+	 * Only works in the editor
+	 */
+	UFUNCTION(BlueprintCallable, Category = "Editor Scripting | Rendering", meta = (Keywords = "Convert Render Target", WorldContext = "WorldContextObject", UnsafeDuringActorConstruction = "true"))
+	static ENGINE_API void ConvertRenderTargetToTextureVolumeEditorOnly(UObject* WorldContextObject, UTextureRenderTargetVolume* RenderTarget, UVolumeTexture* Texture);
+
 	/**
 	 * Exports a render target as a HDR or PNG image onto the disk (depending on the format of the render target)
 	 */
diff --git a/Engine/Source/Runtime/Engine/Private/KismetRenderingLibrary.cpp b/Engine/Source/Runtime/Engine/Private/KismetRenderingLibrary.cpp
--- a/Engine/Source/Runtime/Engine/Private/KismetRenderingLibrary.cpp
+++ b/Engine/Source/Runtime/Engine/Private/KismetRenderingLibrary.cpp
@@ -12,16 +12,22 @@
 #include "RHIContext.h"
 #include "RHIUtilities.h"
 #include "RenderingThread.h"
+#include "RenderCaptureInterface.h"
 #include "Engine/Engine.h"
 #include "Engine/Canvas.h"
 #include "TextureResource.h"
 #include "Logging/MessageLog.h"
 #include "Engine/GameViewportClient.h"
+#include "Engine/TextureRenderTarget2D.h"
 #include "Engine/TextureRenderTarget2DArray.h"
+#include "Engine/TextureRenderTargetCube.h"
 #include "Engine/TextureRenderTargetVolume.h"
 #include "ImageUtils.h"
 #include "ClearQuad.h"
 #include "Engine/Texture2D.h"
+#include "Engine/Texture2DArray.h"
+#include "Engine/TextureCube.h"
+#include "Engine/VolumeTexture.h"
 #include "UObject/Package.h"
 #include "EngineModule.h"
 #include "SceneInterface.h"
@@ -535,42 +541,65 @@ ENGINE_API TArray<FLinearColor> UKismetRenderingLibrary::ReadRenderTargetRawUVAr
 	return ReadRenderTargetRawPixelArea(WorldContextObject, TextureRenderTarget, MinX, MinY, MaxX, MaxY, bNormalize);
 }
 
-/*
-
-void UKismetRenderingLibrary::CreateTexture2DFromRenderTarget(UObject* WorldContextObject, UTextureRenderTarget2D* RenderTarget, const FString &TextureAssetName)
+namespace UE::Kismet::RenderingLibrary
+{
+void ConvertRenderTargetToTextureEditorOnly(UObject* WorldContextObject, UTextureRenderTarget* InRenderTarget, UTexture* InTexture)
 {
-	if (RenderTarget && Texture)
+#if WITH_EDITOR
+	if (InRenderTarget == nullptr)
 	{
-
-		//FImageUtils::CreateTexture2D
-
-		UTexture2D* NewTexture = RenderTarget->ConstructTexture2D(Texture->GetOuter(), Texture->GetName(), RenderTarget->GetMaskedFlags(), CTF_Default, nullptr);
-
-		check(NewTexture == Texture);
-		NewTexture->UpdateResource();
+		FMessageLog("Blueprint").Warning(FText::Format(LOCTEXT("ConvertRenderTargetToTexture_InvalidRenderTarget", "ConvertRenderTargetToTextureEditorOnly[{0}]: RenderTarget must be non-null."), FText::FromString(GetPathNameSafe(WorldContextObject))));
+		return;
 	}
-	else if (!RenderTarget)
+	
+	if (InRenderTarget->GetResource() == nullptr)
 	{
-		FMessageLog("Blueprint").Warning(LOCTEXT("ConvertRenderTargetToTexture2D_InvalidRenderTarget", "ExportRenderTarget: RenderTarget must be non-null."));
+		FMessageLog("Blueprint").Warning(FText::Format(LOCTEXT("ConvertRenderTargetToTexture_ReleasedTextureRenderTarget", "ConvertRenderTargetToTextureEditorOnly[{0}]: render target has been released."), FText::FromString(GetPathNameSafe(WorldContextObject))));
+		return;
 	}
-	else if (!Texture)
+	
+	if (InTexture == nullptr)
 	{
-		FMessageLog("Blueprint").Warning(LOCTEXT("ConvertRenderTargetToTexture2D_InvalidTexture", "ExportRenderTarget: Texture must be non-null."));
+		FMessageLog("Blueprint").Warning(FText::Format(LOCTEXT("ConvertRenderTargetToTexture_InvalidTexture", "ConvertRenderTargetToTextureEditorOnly[{0}]: Texture must be non-null."), FText::FromString(GetPathNameSafe(WorldContextObject))));
+		return;
 	}
 
-}*/
+	// We don't want to create a new texture here, we already have one.  We want to preserve any configuration
+	// by a developer and just update the minimal pixel data and format data.
+	FText ErrorMessage;
+	if (!InRenderTarget->CanConvertToTexture(&ErrorMessage))
+	{
+		FMessageLog("Blueprint").Warning(ErrorMessage);
+		return;
+	}
 
-UTexture2D* UKismetRenderingLibrary::RenderTargetCreateStaticTexture2DEditorOnly(UTextureRenderTarget2D* RenderTarget, FString InName, enum TextureCompressionSettings CompressionSettings, enum TextureMipGenSettings MipSettings)
+	if (!InRenderTarget->UpdateTexture(InTexture, /*InFlags = */CTF_Default, /*InAlphaOverride = */nullptr, /*InOnTextureChangingDelegate = */[](UTexture*){}, &ErrorMessage))
+	{
+		FMessageLog("Blueprint").Warning(ErrorMessage);
+		return;
+	}
+
+	InTexture->Modify();
+	InTexture->MarkPackageDirty();
+	InTexture->PostEditChange();
+	InTexture->UpdateResource();
+
+#else // WITH_EDITOR
+	FMessageLog("Blueprint").Error(LOCTEXT("ConvertRenderTargetToTexture_CannotCallAtRuntime", "ConvertRenderTarget: Can't convert render target to texture at run time. "));
+#endif // !WITH_EDITOR
+}
+
+UTexture* RenderTargetCreateStaticTextureEditorOnly(UTextureRenderTarget* InRenderTarget, FString InName, TextureCompressionSettings InCompressionSettings, TextureMipGenSettings InMipSettings)
 {
 #if WITH_EDITOR
-	if (!RenderTarget)
+	if (InRenderTarget == nullptr)
 	{
-		FMessageLog("Blueprint").Warning(LOCTEXT("RenderTargetCreateStaticTexture2D_InvalidRenderTarget", "RenderTargetCreateStaticTexture2DEditorOnly: RenderTarget must be non-null."));
+		FMessageLog("Blueprint").Warning(LOCTEXT("RenderTargetCreateStaticTexture_InvalidRenderTarget", "RenderTargetCreateStaticTextureEditorOnly: RenderTarget must be non-null."));
 		return nullptr;
 	}
-	else if (!RenderTarget->GetResource())
+	else if (InRenderTarget->GetResource() == nullptr)
 	{
-		FMessageLog("Blueprint").Warning(LOCTEXT("RenderTargetCreateStaticTexture2D_ReleasedRenderTarget", "RenderTargetCreateStaticTexture2DEditorOnly: RenderTarget has been released."));
+		FMessageLog("Blueprint").Warning(LOCTEXT("RenderTargetCreateStaticTexture_ReleasedRenderTarget", "RenderTargetCreateStaticTextureEditorOnly: RenderTarget has been released."));
 		return nullptr;
 	}
 	else
@@ -582,7 +611,7 @@ UTexture2D* UKismetRenderingLibrary::RenderTargetCreateStaticTexture2DEditorOnly
 		//Use asset name only if directories are specified, otherwise full path
 		if (!InName.StartsWith(TEXT("/")))
 		{
-			FString AssetName = RenderTarget->GetOutermost()->GetName();
+			FString AssetName = InRenderTarget->GetOutermost()->GetName();
 			const FString SanitizedBasePackageName = UPackageTools::SanitizePackageName(AssetName);
 			const FString PackagePath = FPackageName::GetLongPackagePath(SanitizedBasePackageName) + TEXT("/");
 			AssetTools.CreateUniqueAssetName(PackagePath, InName, PackageName, Name);
@@ -594,65 +623,76 @@ UTexture2D* UKismetRenderingLibrary::RenderTargetCreateStaticTexture2DEditorOnly
 
 		UObject* NewObj = nullptr;
 
-		// create a static 2d texture
-		NewObj = RenderTarget->ConstructTexture2D(CreatePackage( *PackageName), Name, RenderTarget->GetMaskedFlags() | RF_Public | RF_Standalone, CTF_Default | CTF_AllowMips | CTF_SkipPostEdit, nullptr);
-		UTexture2D* NewTex = Cast<UTexture2D>(NewObj);
+		// create a static texture
+		FText ErrorMessage;
+		NewObj = InRenderTarget->ConstructTexture(CreatePackage(*PackageName), Name, InRenderTarget->GetMaskedFlags() | RF_Public | RF_Standalone, 
+			static_cast<EConstructTextureFlags>(CTF_Default | CTF_AllowMips | CTF_SkipPostEdit), /*InAlphaOverride = */nullptr, &ErrorMessage);
 
-		if (NewTex != nullptr)
+		UTexture* NewTex = Cast<UTexture>(NewObj);
+		if (NewTex == nullptr)
 		{
-			// package needs saving
-			NewObj->MarkPackageDirty();
+			FMessageLog("Blueprint").Warning(ErrorMessage);
+			return nullptr;
+		}
 
-			// Update Compression and Mip settings
-			NewTex->CompressionSettings = CompressionSettings;
-			NewTex->MipGenSettings = MipSettings;
-			NewTex->PostEditChange();
+		// package needs saving
+		NewObj->MarkPackageDirty();
 
-			// Notify the asset registry
-			FAssetRegistryModule::AssetCreated(NewObj);
+		// Update Compression and Mip settings
+		NewTex->CompressionSettings = InCompressionSettings;
+		NewTex->MipGenSettings = InMipSettings;
+		NewTex->PostEditChange();
 
-			return NewTex;
-		}
-		FMessageLog("Blueprint").Warning(LOCTEXT("RenderTargetCreateStaticTexture2D_FailedToCreateTexture", "RenderTargetCreateStaticTexture2DEditorOnly: Failed to create a new texture."));
+		// Notify the asset registry
+		FAssetRegistryModule::AssetCreated(NewObj);
+
+		return NewTex;
 	}
-#else
-	FMessageLog("Blueprint").Error(LOCTEXT("Texture2D's cannot be created at runtime.", "RenderTargetCreateStaticTexture2DEditorOnly: Can't create Texture2D at run time. "));
-#endif
+#else // WITH_EDITOR
+	FMessageLog("Blueprint").Error(LOCTEXT("RenderTargetCreateStaticTexture_CannotCallAtRuntime", "RenderTargetCreateStaticTextureEditorOnly: Can't create texture at run time. "));
+#endif // !WITH_EDITOR
 	return nullptr;
 }
+} // namespace UE::Kismet::RenderingLibrary
+
+UTexture2D* UKismetRenderingLibrary::RenderTargetCreateStaticTexture2DEditorOnly(UTextureRenderTarget2D* RenderTarget, FString InName, enum TextureCompressionSettings CompressionSettings, enum TextureMipGenSettings MipSettings)
+{
+	return Cast<UTexture2D>(UE::Kismet::RenderingLibrary::RenderTargetCreateStaticTextureEditorOnly(RenderTarget, InName, CompressionSettings, MipSettings));
+}
 
+UTexture2DArray* UKismetRenderingLibrary::RenderTargetCreateStaticTexture2DArrayEditorOnly(UTextureRenderTarget2DArray* RenderTarget, FString InName, enum TextureCompressionSettings CompressionSettings, enum TextureMipGenSettings MipSettings)
+{
+	return Cast<UTexture2DArray>(UE::Kismet::RenderingLibrary::RenderTargetCreateStaticTextureEditorOnly(RenderTarget, InName, CompressionSettings, MipSettings));
+}
 
-void UKismetRenderingLibrary::ConvertRenderTargetToTexture2DEditorOnly( UObject* WorldContextObject, UTextureRenderTarget2D* RenderTarget, UTexture2D* Texture )
+UTextureCube* UKismetRenderingLibrary::RenderTargetCreateStaticTextureCubeEditorOnly(UTextureRenderTargetCube* RenderTarget, FString InName, enum TextureCompressionSettings CompressionSettings, enum TextureMipGenSettings MipSettings)
 {
-#if WITH_EDITOR
-	if (!RenderTarget)
-	{
-		FMessageLog("Blueprint").Warning(FText::Format(LOCTEXT("ConvertRenderTargetToTexture2D_InvalidRenderTarget", "ConvertRenderTargetToTexture2DEditorOnly[{0}]: RenderTarget must be non-null."), FText::FromString(GetPathNameSafe(WorldContextObject))));
-	}
-	else if (!RenderTarget->GetResource())
-	{
-		FMessageLog("Blueprint").Warning(FText::Format(LOCTEXT("ConvertRenderTargetToTexture2D_ReleasedTextureRenderTarget", "ConvertRenderTargetToTexture2DEditorOnly[{0}]: render target has been released."), FText::FromString(GetPathNameSafe(WorldContextObject))));
-	}
-	else if (!Texture)
-	{
-		FMessageLog("Blueprint").Warning(FText::Format(LOCTEXT("ConvertRenderTargetToTexture2D_InvalidTexture", "ConvertRenderTargetToTexture2DEditorOnly[{0}]: Texture must be non-null."), FText::FromString(GetPathNameSafe(WorldContextObject))));
-	}
-	else
-	{
-		// We don't want to create a new texture here, we already have one.  We want to preserve any configuration
-		// by a developer and just update the minimal pixel data and format data.
-		const ETextureSourceFormat TextureFormat = RenderTarget->GetTextureFormatForConversionToTexture2D();
-		RenderTarget->UpdateTexture2D(Texture, TextureFormat);
+	return Cast<UTextureCube>(UE::Kismet::RenderingLibrary::RenderTargetCreateStaticTextureEditorOnly(RenderTarget, InName, CompressionSettings, MipSettings));
+}
 
-		Texture->Modify();
-		Texture->MarkPackageDirty();
-		Texture->PostEditChange();
-		Texture->UpdateResource();
-	}
-#else
-	FMessageLog("Blueprint").Error(LOCTEXT("Convert to render target can't be used at run time.", "ConvertRenderTarget: Can't convert render target to texture2d at run time. "));
-#endif
+UVolumeTexture* UKismetRenderingLibrary::RenderTargetCreateStaticVolumeTextureEditorOnly(UTextureRenderTargetVolume* RenderTarget, FString InName, enum TextureCompressionSettings CompressionSettings, enum TextureMipGenSettings MipSettings)
+{
+	return Cast<UVolumeTexture>(UE::Kismet::RenderingLibrary::RenderTargetCreateStaticTextureEditorOnly(RenderTarget, InName, CompressionSettings, MipSettings));
+}
+
+void UKismetRenderingLibrary::ConvertRenderTargetToTexture2DEditorOnly(UObject* WorldContextObject, UTextureRenderTarget2D* RenderTarget, UTexture2D* Texture)
+{
+	UE::Kismet::RenderingLibrary::ConvertRenderTargetToTextureEditorOnly(WorldContextObject, RenderTarget, Texture);
+}
 
+void UKismetRenderingLibrary::ConvertRenderTargetToTexture2DArrayEditorOnly(UObject* WorldContextObject, UTextureRenderTarget2DArray* RenderTarget, UTexture2DArray* Texture)
+{
+	UE::Kismet::RenderingLibrary::ConvertRenderTargetToTextureEditorOnly(WorldContextObject, RenderTarget, Texture);
+}
+
+void UKismetRenderingLibrary::ConvertRenderTargetToTextureCubeEditorOnly(UObject* WorldContextObject, UTextureRenderTargetCube* RenderTarget, UTextureCube* Texture)
+{
+	UE::Kismet::RenderingLibrary::ConvertRenderTargetToTextureEditorOnly(WorldContextObject, RenderTarget, Texture);
+}
+
+void UKismetRenderingLibrary::ConvertRenderTargetToTextureVolumeEditorOnly(UObject* WorldContextObject, UTextureRenderTargetVolume* RenderTarget, UVolumeTexture* Texture)
+{
+	UE::Kismet::RenderingLibrary::ConvertRenderTargetToTextureEditorOnly(WorldContextObject, RenderTarget, Texture);
 }
 
 void UKismetRenderingLibrary::ExportTexture2D(UObject* WorldContextObject, UTexture2D* Texture, const FString& FilePath, const FString& FileName)
diff --git a/Engine/Source/Runtime/Engine/Private/Rendering/TextureSourceFormat.cpp b/Engine/Source/Runtime/Engine/Private/Rendering/TextureSourceFormat.cpp
new file mode 100644
--- /dev/null
+++ b/Engine/Source/Runtime/Engine/Private/Rendering/TextureSourceFormat.cpp
@@ -0,0 +1,33 @@
+// Copyright Epic Games, Inc. All Rights Reserved.
+
+#include "Engine/TextureDefines.h"
+#include "PixelFormat.h"
+
+//
+//	Texture source format information.
+//
+
+FTextureSourceFormatInfo::FTextureSourceFormatInfo(ETextureSourceFormat InTextureSourceFormat, EPixelFormat InPixelFormat, const TCHAR* InName, int32 InNumComponents, int32 InBytesPerPixel)
+	: TextureSourceFormat(InTextureSourceFormat)
+	, PixelFormat(InPixelFormat)
+	, Name(InName)
+	, NumComponents(InNumComponents)
+	, BytesPerPixel(InBytesPerPixel)
+{}
+
+FTextureSourceFormatInfo GTextureSourceFormats[TSF_MAX] =
+{
+	//						TextureSourceFormat		PixelFormat							Name						NumComponents	BytesPerPixel
+	FTextureSourceFormatInfo(TSF_Invalid,			EPixelFormat::PF_Unknown,			TEXT("Invalid"),			0,				0),
+	FTextureSourceFormatInfo(TSF_G8,				EPixelFormat::PF_G8,				TEXT("G8"),					1,				1),
+	FTextureSourceFormatInfo(TSF_BGRA8,				EPixelFormat::PF_B8G8R8A8,			TEXT("BGRA8"),				4,				4),
+	FTextureSourceFormatInfo(TSF_BGRE8,				EPixelFormat::PF_B8G8R8A8,			TEXT("BGRE8"),				4,				4),
+	FTextureSourceFormatInfo(TSF_RGBA16,			EPixelFormat::PF_R16G16B16A16_UINT,	TEXT("RGBA16"),				4,				8),
+	FTextureSourceFormatInfo(TSF_RGBA16F,			EPixelFormat::PF_A16B16G16R16,		TEXT("RGBA16F"),			4,				8),
+	FTextureSourceFormatInfo(TSF_RGBA8_DEPRECATED,	EPixelFormat::PF_B8G8R8A8,			TEXT("RGBA8_DEPRECATED"),	4,				4),
+	FTextureSourceFormatInfo(TSF_RGBE8_DEPRECATED,	EPixelFormat::PF_B8G8R8A8,			TEXT("RGBE8_DEPRECATED"),	4,				4),
+	FTextureSourceFormatInfo(TSF_G16,				EPixelFormat::PF_G16,				TEXT("G16"),				1,				2),
+	FTextureSourceFormatInfo(TSF_RGBA32F,			EPixelFormat::PF_A32B32G32R32F,		TEXT("RGBA32F"),			4,				16),
+	FTextureSourceFormatInfo(TSF_R16F,				EPixelFormat::PF_R16F,				TEXT("R16F"),				1,				2),
+	FTextureSourceFormatInfo(TSF_R32F,				EPixelFormat::PF_R32_FLOAT,			TEXT("R32F"),				1,				4),
+};
diff --git a/Engine/Source/Runtime/Engine/Private/TextureRenderTarget.cpp b/Engine/Source/Runtime/Engine/Private/TextureRenderTarget.cpp
--- a/Engine/Source/Runtime/Engine/Private/TextureRenderTarget.cpp
+++ b/Engine/Source/Runtime/Engine/Private/TextureRenderTarget.cpp
@@ -9,6 +9,10 @@
 #include "TextureResource.h"
 #include "SceneRenderTargetParameters.h"
 #include "EngineModule.h"
+#include "EngineLogs.h"
+#include "RenderUtils.h"
+#include "TextureCompiler.h"
+#include "Logging/MessageLog.h"
 
 #include UE_INLINE_GENERATED_CPP_BY_NAME(TextureRenderTarget)
 
@@ -62,6 +66,335 @@ EMaterialValueType UTextureRenderTarget::GetMaterialType() const
 	return MCT_Texture;
 }
 
+bool UTextureRenderTarget::CanConvertToTexture(FText* OutErrorMessage) const
+{
+	ETextureSourceFormat TextureSourceFormat;
+	EPixelFormat PixelFormat;
+	return CanConvertToTexture(TextureSourceFormat, PixelFormat, OutErrorMessage);
+}
+
+ETextureSourceFormat UTextureRenderTarget::ValidateTextureFormatForConversionToTextureInternal(EPixelFormat InFormat, const TArrayView<const EPixelFormat>& InCompatibleFormats, FText* OutErrorMessage) const
+{
+	if (!InCompatibleFormats.Contains(InFormat))
+	{
+		if (OutErrorMessage != nullptr)
+		{
+			TArray<const TCHAR*> CompatibleFormatStrings;
+			Algo::Transform(InCompatibleFormats, CompatibleFormatStrings, [](EPixelFormat InCompatiblePixelFormat) { return GetPixelFormatString(InCompatiblePixelFormat); });
+
+			*OutErrorMessage = FText::Format(NSLOCTEXT("TextureRenderTarget", "UnsupportedFormatForConversionToTexture", "Unsupported format ({0}) for converting {1} to {2}. Supported formats are: {3}"),
+				FText::FromString(FString(GetPixelFormatString(InFormat))),
+				FText::FromString(GetClass()->GetName()),
+				FText::FromString(GetTextureUClass()->GetName()),
+				FText::FromString(FString::Join(CompatibleFormatStrings, TEXT(","))));
+		}
+		return TSF_Invalid;
+	}
+
+	// Return what ETextureSourceFormat corresponds to this EPixelFormat (must match the conversion capabilities of UTextureRenderTarget::UpdateTexture) : 
+	ETextureSourceFormat TextureFormat = TSF_Invalid;
+	switch (InFormat)
+	{
+	case PF_B8G8R8A8:
+	case PF_R8G8: // Will produce 2 useless channels as it will be expanded to a TSF_BGRA8 texture source but there's no 2 channel ETextureSourceFormat
+		TextureFormat = TSF_BGRA8;
+		break;
+	case PF_R32_FLOAT: // Narrowed down to 16bits float (if available by the FRenderTarget::ReadFloat16Pixels implementation and corresponding RHI implementation)
+	case PF_FloatRGBA:
+		TextureFormat = TSF_RGBA16F;
+		break;
+	case PF_G8:
+		TextureFormat = TSF_G8;
+		break;
+	default:
+	{
+		checkf(false, TEXT("Only some select texture formats are currently supported by UTextureRenderTarget::UpdateTexture. Please update this switch if you add conversion code."));
+	}
+	}
+
+	return TextureFormat;
+}
+
+#if WITH_EDITOR
+
+bool UTextureRenderTarget::UpdateTexture(UTexture* InTexture, EConstructTextureFlags InFlags, const TArray<uint8>* InAlphaOverride, FOnTextureChangingDelegate InOnTextureChangingDelegate, FText* OutErrorMessage)
+{
+
+	if (InTexture == nullptr)
+	{
+		if (OutErrorMessage != nullptr)
+		{
+			*OutErrorMessage = NSLOCTEXT("TextureRenderTarget", "NullTexture", "Cannot update texture : null texture.");
+		}
+		return false;
+	}
+
+	UClass* TextureUClass = GetTextureUClass();
+	check(TextureUClass != nullptr);
+	if (!InTexture->GetClass()->IsChildOf(TextureUClass))
+	{
+		if (OutErrorMessage != nullptr)
+		{
+			*OutErrorMessage = FText::Format(NSLOCTEXT("TextureRenderTarget", "InvalidTextureClass", "Cannot update texture {0} : its class is {1} while a {2} is expected."), 
+				FText::FromString(InTexture->GetName()),
+				FText::FromString(InTexture->GetClass()->GetName()),
+				FText::FromString(GetTextureUClass()->GetClass()->GetName()));
+		}
+		return false;
+	}
+
+	ETextureSourceFormat TextureSourceFormat = TSF_Invalid;
+	EPixelFormat PixelFormat = PF_Unknown; 
+	if (!CanConvertToTexture(TextureSourceFormat, PixelFormat, OutErrorMessage))
+	{
+		return false;
+	}
+
+	// Note : use UTexture's GetTextureClass() here, as URenderTarget's GetTextureClass() unconveniently returns ETextureClass::RenderTarget :
+	ETextureClass TextureClass = InTexture->GetTextureClass();
+	const bool bIsVolume = (TextureClass == ETextureClass::Volume);
+	const bool bIsArray = (TextureClass == ETextureClass::Array) || (TextureClass == ETextureClass::CubeArray);
+	const bool bIsCube = (TextureClass == ETextureClass::Cube) || (TextureClass == ETextureClass::CubeArray);
+	const int32 SizeX = GetSurfaceWidth();
+	const int32 SizeY = GetSurfaceHeight();
+	const int32 SizeZ = bIsVolume ? GetSurfaceDepth() : 1;
+	const int32 NumSurfaces = bIsVolume ? SizeZ : (bIsArray || bIsCube) ? GetSurfaceArraySize() : 1;
+
+	FTextureRenderTargetResource* RenderTarget = GameThread_GetRenderTargetResource();
+	if (RenderTarget == nullptr)
+	{
+		if (OutErrorMessage != nullptr)
+		{
+			*OutErrorMessage = FText::Format(NSLOCTEXT("TextureRenderTarget", "NullRenderTargetResource", "Cannot update texture : render target {0} has been released."), FText::FromString(GetPathNameSafe(this)));
+		}
+		return false;
+	}
+
+	TextureCompressionSettings CompressionSettingsForTexture = (PixelFormat == EPixelFormat::PF_FloatRGBA) ? TC_HDR : TC_Default;
+
+	const int32 NumMips = 1; // There's only support for mip 0 ATM
+
+	bool bTextureChanging = false;
+	bTextureChanging |= InTexture->Source.GetSizeX() != SizeX;
+	bTextureChanging |= InTexture->Source.GetSizeY() != SizeY;
+	bTextureChanging |= InTexture->Source.GetVolumeSizeZ() != SizeZ;
+	// Note : for FTextureSource, NumSlices means NumSurfaces (e.g. 6 * GetNumSlices() for a cube map array) : 
+	bTextureChanging |= InTexture->Source.GetNumSlices() != NumSurfaces;
+	bTextureChanging |= InTexture->Source.GetFormat() != TextureSourceFormat;
+	bTextureChanging |= InTexture->CompressionSettings != CompressionSettingsForTexture;
+
+	uint32 OldDataHash = 0;
+
+	// Hashing the content is only really useful if none of the other texture settings have changed.
+	const bool bHashContent = !bTextureChanging;
+	if (bHashContent)
+	{
+		const uint8* OldData = InTexture->Source.LockMipReadOnly(0);
+		const int32 OldDataSize = InTexture->Source.CalcMipSize(0);
+		OldDataHash = FCrc::MemCrc32(OldData, OldDataSize);
+		InTexture->Source.UnlockMip(0);
+	}
+
+	// Temp storage that will be used for the texture init.
+	const int32 NumPixelsPerSurface = SizeX * SizeY;
+	const int32 NumBytesPerPixel = GTextureSourceFormats[TextureSourceFormat].BytesPerPixel;
+	const int32 DestinationSurfaceMipSize = NumPixelsPerSurface * NumBytesPerPixel;
+	const int32 DestinationSize = NumPixelsPerSurface * NumSurfaces * NumBytesPerPixel;
+	check(NumBytesPerPixel > 0);
+	checkf((InAlphaOverride == nullptr) || (NumPixelsPerSurface == InAlphaOverride->Num()), TEXT("If InAlphaOverride is specified, it must have the same number of pixels as a single slice (currently, render target has %d pixels per slice, AlphaOverride: %d). It is also expected to be 1-byte encoded (1 element per pixel)."),
+		NumPixelsPerSurface, InAlphaOverride->Num());
+
+	TArray<uint8> NewData;
+	NewData.SetNumUninitialized(DestinationSize);
+
+	// Read the data surface by surface (i.e. for a 2D array or a volume: slice by slice, for a cubemap: face by face and for a cubemap array: face by face, slice by slice
+	for (int32 SurfaceIndex = 0; SurfaceIndex < NumSurfaces; ++SurfaceIndex)
+	{
+		TArrayView<uint8> SurfaceByteData = MakeArrayView(NewData.GetData() + SurfaceIndex * DestinationSurfaceMipSize, DestinationSurfaceMipSize);
+		FReadSurfaceDataFlags ReadSurfaceDataFlags(RCM_UNorm, bIsCube ? static_cast<ECubeFace>(SurfaceIndex % 6) : CubeFace_MAX);
+		ReadSurfaceDataFlags.SetArrayIndex(bIsCube ? SurfaceIndex / 6 : SurfaceIndex);
+		switch (TextureSourceFormat)
+		{
+		case TSF_BGRA8:
+		{
+			TArray<FColor> NewDataColor;
+			RenderTarget->ReadPixels(NewDataColor, ReadSurfaceDataFlags);
+			check(NewDataColor.Num() == NumPixelsPerSurface);
+
+			const int32 NumBytes = NewDataColor.Num() * NewDataColor.GetTypeSize();
+			check(NumBytes == SurfaceByteData.Num());
+			FMemory::Memcpy(SurfaceByteData.GetData(), NewDataColor.GetData(), NumBytes);
+
+			// override the alpha if desired
+			TArrayView<FColor> SurfaceColorData = MakeArrayView(reinterpret_cast<FColor*>(SurfaceByteData.GetData()), NumPixelsPerSurface);
+			if (InAlphaOverride != nullptr)
+			{
+				for (int32 PixelIndex = 0; PixelIndex < NumPixelsPerSurface; ++PixelIndex)
+				{
+					SurfaceColorData[PixelIndex].A = (*InAlphaOverride)[PixelIndex];
+				}
+			}
+			else if (InFlags & CTF_RemapAlphaAsMasked)
+			{
+				// if the target was rendered with a masked texture, then the depth will probably have been written instead of 0/255 for the
+				// alpha, and the depth when unwritten will be 255, so remap 255 to 0 (masked out area) and anything else as 255 (written to area)
+				for (int32 PixelIndex = 0; PixelIndex < NumPixelsPerSurface; ++PixelIndex)
+				{
+					SurfaceColorData[PixelIndex].A = (SurfaceColorData[PixelIndex].A == 255) ? 0 : 255;
+				}
+			}
+			else if (InFlags & CTF_ForceOpaque)
+			{
+				for (int32 PixelIndex = 0; PixelIndex < NumPixelsPerSurface; ++PixelIndex)
+				{
+					SurfaceColorData[PixelIndex].A = 255;
+				}
+			}
+		}
+		break;
+		case TSF_RGBA16F:
+		{
+			TArray<FFloat16Color> NewDataFloat16Color;
+			check(NumBytesPerPixel == sizeof(FFloat16Color));
+			RenderTarget->ReadFloat16Pixels(NewDataFloat16Color, ReadSurfaceDataFlags);
+			check(NewDataFloat16Color.Num() == NumPixelsPerSurface);
+			const int32 NumBytes = NewDataFloat16Color.Num() * NewDataFloat16Color.GetTypeSize();
+			check(NumBytes == SurfaceByteData.Num());
+			FMemory::Memcpy(SurfaceByteData.GetData(), NewDataFloat16Color.GetData(), NumBytes);
+
+			// override the alpha if desired
+			TArrayView<FFloat16Color> SurfaceFloat16ColorData = MakeArrayView(reinterpret_cast<FFloat16Color*>(SurfaceByteData.GetData()), NumPixelsPerSurface);
+			if (InAlphaOverride != nullptr)
+			{
+				for (int32 PixelIndex = 0; PixelIndex < NumPixelsPerSurface; ++PixelIndex)
+				{
+					SurfaceFloat16ColorData[PixelIndex].A = (*InAlphaOverride)[PixelIndex];
+				}
+			}
+			else if (InFlags & CTF_RemapAlphaAsMasked)
+			{
+				// if the target was rendered with a masked texture, then the depth will probably have been written instead of 0/255 for the
+				// alpha, and the depth when unwritten will be 255, so remap 255 to 0 (masked out area) and anything else as 255 (written to area)
+				for (int32 PixelIndex = 0; PixelIndex < NumPixelsPerSurface; ++PixelIndex)
+				{
+					SurfaceFloat16ColorData[PixelIndex].A = (SurfaceFloat16ColorData[PixelIndex].A == 255) ? 0 : 255;
+				}
+			}
+			else if (InFlags & CTF_ForceOpaque)
+			{
+				for (int32 PixelIndex = 0; PixelIndex < NumPixelsPerSurface; ++PixelIndex)
+				{
+					SurfaceFloat16ColorData[PixelIndex].A = 255;
+				}
+			}
+		}
+		break;
+		case TSF_G8:
+		{
+			
+			TArray<FColor> NewDataColor;
+			RenderTarget->ReadPixels(NewDataColor, ReadSurfaceDataFlags);
+			check(NewDataColor.Num() == NumPixelsPerSurface);
+			for (int32 PixelIndex = 0; PixelIndex < NumPixelsPerSurface; PixelIndex++)
+			{
+				SurfaceByteData[PixelIndex] = NewDataColor[PixelIndex].R;
+			}
+		}
+		break;
+		default:
+			checkf(false, TEXT("Unsupported format, it needs to be added to CanConvertToTexture and UpdateTexture"));
+		}
+	}
+
+	if (bHashContent)
+	{
+		uint32 NewDataHash = FCrc::MemCrc32(NewData.GetData(), NewData.Num());
+		bTextureChanging = OldDataHash != NewDataHash;
+	}
+
+	if (bTextureChanging)
+	{
+		InOnTextureChangingDelegate(InTexture);
+
+		// Ensure the texture is not being compiled
+		FTextureCompilingManager::Get().FinishCompilation({ InTexture });
+
+		// init to the same size as the render target
+		InTexture->Source.Init(SizeX, SizeY, NumSurfaces, NumMips, TextureSourceFormat, NewData.GetData());
+		InTexture->CompressionSettings = CompressionSettingsForTexture;
+	}
+
+	return true;
+}
+
+UTexture* UTextureRenderTarget::ConstructTexture(UObject* InOuter, const FString& InNewTextureName, EObjectFlags InObjectFlags, EConstructTextureFlags InFlags, const TArray<uint8>* InAlphaOverride, FText* OutErrorMessage)
+{
+	FTextureRenderTargetResource* RenderTarget = GameThread_GetRenderTargetResource();
+	if (RenderTarget == nullptr)
+	{
+		if (OutErrorMessage != nullptr)
+		{
+			*OutErrorMessage = FText::Format(NSLOCTEXT("TextureRenderTarget", "NullRenderTargetResource", "Cannot construct texture : render target {0} has been released."), FText::FromString(GetPathNameSafe(this)));
+		}
+		return nullptr;
+	}
+
+	// Make sure the conversion is possible between this render target and this texture. After this point, the conversion shouldn't be able to ever fail : 
+	if (!CanConvertToTexture(OutErrorMessage))
+	{
+		return nullptr;
+	}
+
+	// Create the texture object
+	UClass* TextureUClass = GetTextureUClass();
+	check(TextureUClass != nullptr);
+	UTexture* Result = NewObject<UTexture>(InOuter, TextureUClass, FName(*InNewTextureName), InObjectFlags);
+
+	verify(UpdateTexture(Result, InFlags, InAlphaOverride, /*InOnTextureChangingDelegate = */[](UTexture*) {}, OutErrorMessage));
+
+	// if render target gamma used was 1.0 then disable SRGB for the static texture
+	// note: UTextureRenderTarget2D also has an explicit SRGB flag in the UTexture parent class
+	//	  these are NOT correctly kept in sync
+	//		I see SRGB = 1 but Gamma = 1.0
+	//	 see also IsSRGB() which is yet another query that has different ideas
+	//	furthermore, float formats do not support anything but Linear gamma
+	float Gamma = RenderTarget->GetDisplayGamma();
+	if (FMath::Abs(Gamma - 1.0f) < UE_KINDA_SMALL_NUMBER)
+	{
+		InFlags = static_cast<EConstructTextureFlags>(InFlags & ~CTF_SRGB);
+	}
+
+	Result->SRGB = (InFlags & CTF_SRGB) != 0;
+
+	Result->MipGenSettings = TMGS_FromTextureGroup;
+
+	if ((InFlags & CTF_AllowMips) == 0)
+	{
+		Result->MipGenSettings = TMGS_NoMipmaps;
+	}
+
+	if (InFlags & CTF_Compress)
+	{
+		// Set compression options.
+		Result->DeferCompression = (InFlags & CTF_DeferCompression) ? true : false;
+	}
+	else
+	{
+		// Disable compression
+		Result->CompressionNone = true;
+		Result->DeferCompression = false;
+	}
+
+	if ((InFlags & CTF_SkipPostEdit) != 0)
+	{
+		Result->PostEditChange();
+	}
+
+	return Result;
+}
+
+#endif // WITH_EDITOR
+
 /*-----------------------------------------------------------------------------
 	FTextureRenderTargetResource
 -----------------------------------------------------------------------------*/
diff --git a/Engine/Source/Runtime/Engine/Private/TextureRenderTarget2D.cpp b/Engine/Source/Runtime/Engine/Private/TextureRenderTarget2D.cpp
--- a/Engine/Source/Runtime/Engine/Private/TextureRenderTarget2D.cpp
+++ b/Engine/Source/Runtime/Engine/Private/TextureRenderTarget2D.cpp
@@ -12,6 +12,7 @@
 #include "Engine/Texture2D.h"
 #include "DeviceProfiles/DeviceProfile.h"
 #include "DeviceProfiles/DeviceProfileManager.h"
+#include "EngineLogs.h"
 #include "GenerateMips.h"
 #include "RenderGraphUtils.h"
 #include "UObject/Package.h"
@@ -315,245 +316,90 @@ FString UTextureRenderTarget2D::GetDesc()
 	return FString::Printf( TEXT("Render to Texture %dx%d[%s]"), SizeX, SizeY, GPixelFormats[GetFormat()].Name );
 }
 
-UTexture2D* UTextureRenderTarget2D::ConstructTexture2D(UObject* Outer, const FString& NewTexName, EObjectFlags InObjectFlags, uint32 Flags, TArray<uint8>* AlphaOverride)
+UTexture2D* UTextureRenderTarget2D::ConstructTexture2D(UObject* InOuter, const FString& InNewTextureName, EObjectFlags InObjectFlags, uint32 InFlags, TArray<uint8>* InAlphaOverride)
 {
 	UTexture2D* Result = nullptr;
 
 #if WITH_EDITOR
-	// Check render target size is valid.
-	const bool bIsValidSize = (SizeX != 0 && SizeY != 0);
-
-	// The render to texture resource will be needed to read its surface contents
-	FRenderTarget* RenderTarget = GameThread_GetRenderTargetResource();
-
-	const ETextureSourceFormat TextureFormat = GetTextureFormatForConversionToTexture2D();
-
-	// exit if source is not compatible.
-	if (bIsValidSize == false || RenderTarget == nullptr || TextureFormat == TSF_Invalid)
-	{
-		return Result;
-	}
-
-	// Create the 2d texture
-	Result = NewObject<UTexture2D>(Outer, FName(*NewTexName), InObjectFlags);
-	
-	UpdateTexture2D(Result, TextureFormat, Flags, AlphaOverride);
-
-	// if render target gamma used was 1.0 then disable SRGB for the static texture
-	// note: UTextureRenderTarget2D also has an explicit SRGB flag in the UTexture parent class
-	//	  these are NOT correctly kept in sync
-	//		I see SRGB = 1 but Gamma = 1.0
-	//	 see also IsSRGB() which is yet another query that has different ideas
-	//	furthermore, float formats do not support anything but Linear gamma
-	float Gamma = RenderTarget->GetDisplayGamma();
-	if (FMath::Abs(Gamma - 1.0f) < UE_KINDA_SMALL_NUMBER)
-	{
-		Flags &= ~CTF_SRGB;
+	FText ErrorMessage;
+	Result = Cast<UTexture2D>(ConstructTexture(InOuter, InNewTextureName, InObjectFlags, static_cast<EConstructTextureFlags>(InFlags), InAlphaOverride, &ErrorMessage));
+	if (Result == nullptr)
+	{ 
+		UE_LOG(LogTexture, Error, TEXT("Couldn't construct texture : %s"), *ErrorMessage.ToString());
 	}
-
-	Result->SRGB = (Flags & CTF_SRGB) != 0;
-
-	Result->MipGenSettings = TMGS_FromTextureGroup;
-
-	if ((Flags & CTF_AllowMips) == 0)
-	{
-		Result->MipGenSettings = TMGS_NoMipmaps;
-	}
-
-	if (Flags & CTF_Compress)
-	{
-		// Set compression options.
-		Result->DeferCompression = (Flags & CTF_DeferCompression) ? true : false;
-	}
-	else
-	{
-		// Disable compression
-		Result->CompressionNone = true;
-		Result->DeferCompression = false;
-	}
-
-	if ((Flags & CTF_SkipPostEdit) != 0)
-	{
-		Result->PostEditChange();
-	}
-#endif
+#endif // WITH_EDITOR
 
 	return Result;
 }
 
+
+PRAGMA_DISABLE_DEPRECATION_WARNINGS
 ETextureSourceFormat UTextureRenderTarget2D::GetTextureFormatForConversionToTexture2D() const
 {
-	const EPixelFormat PixelFormat = GetFormat();
-	ETextureSourceFormat TextureFormat = TSF_Invalid;
-	switch (PixelFormat)
+	ETextureSourceFormat TextureSourceFormat = TSF_Invalid;
+	EPixelFormat PixelFormat = PF_Unknown;
+	if (!CanConvertToTexture(TextureSourceFormat, PixelFormat, /*OutErrorMessage = */nullptr))
 	{
-	case PF_B8G8R8A8:
-		TextureFormat = TSF_BGRA8;
-		break;
-	case PF_FloatRGBA:
-		TextureFormat = TSF_RGBA16F;
-		break;
-	case PF_G8:
-		TextureFormat = TSF_G8;
-		break;
-	default:
-	{
-		FText InvalidFormatMessage = NSLOCTEXT("TextureRenderTarget2D", "UnsupportedFormatRenderTarget2DWarning", "Unsupported format when creating Texture2D from TextureRenderTarget2D. Supported formats are B8G8R8A8, FloatRGBA and G8.");
-		FMessageDialog::Open(EAppMsgType::Ok, InvalidFormatMessage);
-	}
+		return TSF_Invalid;
 	}
 
-	return TextureFormat;
+	return TextureSourceFormat;
 }
+PRAGMA_ENABLE_DEPRECATION_WARNINGS
 
-void UTextureRenderTarget2D::UpdateTexture2D(UTexture2D* InTexture2D, ETextureSourceFormat InTextureFormat, uint32 Flags, TArray<uint8>* AlphaOverride)
+TSubclassOf<UTexture> UTextureRenderTarget2D::GetTextureUClass() const
 {
-	UpdateTexture2D(InTexture2D, InTextureFormat, Flags, AlphaOverride, FTextureChangingDelegate());
+	return UTexture2D::StaticClass();
 }
 
-void UTextureRenderTarget2D::UpdateTexture2D(UTexture2D* InTexture2D, ETextureSourceFormat InTextureFormat, uint32 Flags, TArray<uint8>* AlphaOverride, FTextureChangingDelegate TextureChangingDelegate)
+bool UTextureRenderTarget2D::CanConvertToTexture(ETextureSourceFormat& OutTextureSourceFormat, EPixelFormat& OutPixelFormat, FText* OutErrorMessage) const
 {
-#if WITH_EDITOR
-	const EPixelFormat PixelFormat = GetFormat();
-	TextureCompressionSettings CompressionSettingsForTexture = PixelFormat == EPixelFormat::PF_FloatRGBA ? TC_HDR : TC_Default;
-
-	const int32 NbSlices = 1;
-	const int32 NbMips = 1;
-
-	bool bTextureChanging = false;
-	
-	bTextureChanging |= InTexture2D->Source.GetSizeX() != SizeX;
-	bTextureChanging |= InTexture2D->Source.GetSizeY() != SizeY;
-	bTextureChanging |= InTexture2D->Source.GetNumSlices() != NbSlices;
-	bTextureChanging |= InTexture2D->Source.GetNumMips() != NbMips;
-	bTextureChanging |= InTexture2D->Source.GetFormat() != InTextureFormat;
-	bTextureChanging |= InTexture2D->CompressionSettings != CompressionSettingsForTexture;
-	
-	uint32 OldDataHash = 0;
-
-	// Hashing the content is only really useful if none of the other texture settings have changed.
-	const bool bHashContent = !bTextureChanging;
-	if (bHashContent)
+	const EPixelFormat LocalFormat = GetFormat();
+	// These are the formats currently available for conversion to texture for UTextureRenderTarget2D : 
+	const ETextureSourceFormat TextureSourceFormat = ValidateTextureFormatForConversionToTextureInternal(LocalFormat, { PF_G8, PF_R8G8, PF_B8G8R8A8, PF_FloatRGBA }, OutErrorMessage);
+	if (TextureSourceFormat == TSF_Invalid)
 	{
-		const uint8* OldData = InTexture2D->Source.LockMipReadOnly(0);
-		const int32 OldDataSize = InTexture2D->Source.CalcMipSize(0);
-		OldDataHash = FCrc::MemCrc32(OldData, OldDataSize);
-		InTexture2D->Source.UnlockMip(0);
-	}
-
-	// Temp storage that will be used for the texture init.
-	TArray<uint8> NewDataGrayscale;
-	TArray<FColor> NewDataColor;
-	TArray<FFloat16Color> NewDataFloat16Color;
-
-	const uint8* NewData = nullptr;
-	int32 NewDataSize = 0;
-
-	// read the 2d surface
-	FRenderTarget* RenderTarget = GameThread_GetRenderTargetResource();
-	if (InTextureFormat == TSF_BGRA8)
-	{
-		RenderTarget->ReadPixels(NewDataColor);
-		check(NewDataColor.Num() == SizeX * SizeY);
-
-		// override the alpha if desired
-		if (AlphaOverride)
-		{
-			check(NewDataColor.Num() == AlphaOverride->Num());
-			for (int32 Pixel = 0; Pixel < NewDataColor.Num(); Pixel++)
-			{
-				NewDataColor[Pixel].A = (*AlphaOverride)[Pixel];
-			}
-		}
-		else if (Flags & CTF_RemapAlphaAsMasked)
-		{
-			// if the target was rendered with a masked texture, then the depth will probably have been written instead of 0/255 for the
-			// alpha, and the depth when unwritten will be 255, so remap 255 to 0 (masked out area) and anything else as 255 (written to area)
-			for (int32 Pixel = 0; Pixel < NewDataColor.Num(); Pixel++)
-			{
-				NewDataColor[Pixel].A = (NewDataColor[Pixel].A == 255) ? 0 : 255;
-			}
-		}
-		else if (Flags & CTF_ForceOpaque)
-		{
-			for (int32 Pixel = 0; Pixel < NewDataColor.Num(); Pixel++)
-			{
-				NewDataColor[Pixel].A = 255;
-			}
-		}
-
-		NewData = (uint8*)NewDataColor.GetData();
-		NewDataSize = NewDataColor.Num() * NewDataColor.GetTypeSize();
+		return TSF_Invalid;
 	}
-	else if (InTextureFormat == TSF_RGBA16F)
-	{
-		RenderTarget->ReadFloat16Pixels(NewDataFloat16Color);
-		check(NewDataFloat16Color.Num() == SizeX * SizeY);
 
-		// override the alpha if desired
-		if (AlphaOverride)
-		{
-			check(NewDataFloat16Color.Num() == AlphaOverride->Num());
-			for (int32 Pixel = 0; Pixel < NewDataFloat16Color.Num(); Pixel++)
-			{
-				NewDataFloat16Color[Pixel].A = ((float)(*AlphaOverride)[Pixel]) / 255.0f;
-			}
-		}
-		else if (Flags & CTF_RemapAlphaAsMasked)
-		{
-			// if the target was rendered with a masked texture, then the depth will probably have been written instead of 0/255 for the
-			// alpha, and the depth when unwritten will be 255, so remap 255 to 0 (masked out area) and anything else as 1 (written to area)
-			for (int32 Pixel = 0; Pixel < NewDataFloat16Color.Num(); Pixel++)
-			{
-				NewDataFloat16Color[Pixel].A = (NewDataFloat16Color[Pixel].A == 255) ? 0.0f : 1.0f;
-			}
-		}
-		else if (Flags & CTF_ForceOpaque)
-		{
-			for (int32 Pixel = 0; Pixel < NewDataFloat16Color.Num(); Pixel++)
-			{
-				NewDataFloat16Color[Pixel].A = 1.0f;
-			}
-		}
-
-		NewData = (uint8*)NewDataFloat16Color.GetData();
-		NewDataSize = NewDataFloat16Color.Num() * NewDataFloat16Color.GetTypeSize();
-	}
-	else if (InTextureFormat == TSF_G8)
+	if ((SizeX <= 0) || (SizeY <= 0))
 	{
-		RenderTarget->ReadPixels(NewDataColor);
-		check(NewDataColor.Num() == SizeX * SizeY);
-		NewDataGrayscale.SetNumUninitialized(NewDataColor.Num());
-
-		for (int32 Pixel = 0; Pixel < NewDataColor.Num(); Pixel++)
+		if (OutErrorMessage != nullptr)
 		{
-			NewDataGrayscale[Pixel] = NewDataColor[Pixel].R;
+			*OutErrorMessage = FText::Format(NSLOCTEXT("TextureRenderTarget2D", "InvalidSizeForConversionToTexture", "Invalid size ({0},{1}) for converting {2} to {3}"),
+				FText::AsNumber(SizeX),
+				FText::AsNumber(SizeY),
+				FText::FromString(GetClass()->GetName()),
+				FText::FromString(GetTextureUClass()->GetName()));
 		}
-
-		NewData = (uint8*)NewDataGrayscale.GetData();
-		NewDataSize = NewDataGrayscale.Num() * NewDataGrayscale.GetTypeSize();
+		return false;
 	}
 
+	OutPixelFormat = LocalFormat;
+	OutTextureSourceFormat = TextureSourceFormat;
+	return true;
+}
 
-	if (bHashContent)
-	{
-		uint32 NewDataHash = FCrc::MemCrc32(NewData, NewDataSize);
-		bTextureChanging = OldDataHash != NewDataHash;
-	}
+PRAGMA_DISABLE_DEPRECATION_WARNINGS
+void UTextureRenderTarget2D::UpdateTexture2D(UTexture2D* InTexture2D, ETextureSourceFormat InTextureFormat, uint32 Flags, TArray<uint8>* AlphaOverride)
+{
+	UpdateTexture2D(InTexture2D, InTextureFormat, Flags, AlphaOverride, FTextureChangingDelegate());
+}
 
-	if (bTextureChanging)
+void UTextureRenderTarget2D::UpdateTexture2D(UTexture2D* InTexture2D, ETextureSourceFormat InTextureFormat, uint32 Flags, TArray<uint8>* AlphaOverride, FTextureChangingDelegate TextureChangingDelegate)
+{
+#if WITH_EDITOR
+	// Simply forward the internal delegate to the external one (the only difference being the specialized type of the texture) :
+	auto OnTextureChangingDelegate = [TextureChangingDelegate](UTexture* InTexture) { TextureChangingDelegate.ExecuteIfBound(CastChecked<UTexture2D>(InTexture)); };
+	FText ErrorMessage;
+	bool bSuccess = UpdateTexture(InTexture2D, static_cast<EConstructTextureFlags>(Flags), AlphaOverride, OnTextureChangingDelegate, &ErrorMessage);
+	if (!bSuccess)
 	{
-		TextureChangingDelegate.ExecuteIfBound(InTexture2D);
-
-		// Ensure the texture is not being compiled
-		FTextureCompilingManager::Get().FinishCompilation({InTexture2D});
-
-		// init to the same size as the 2d texture
-		InTexture2D->Source.Init(SizeX, SizeY, NbSlices, NbMips, InTextureFormat, NewData);
-		InTexture2D->CompressionSettings = CompressionSettingsForTexture;
+		UE_LOG(LogTexture, Error, TEXT("Cannot update texture : %s"), *ErrorMessage.ToString());
 	}
-#endif
+#endif // WITH_EDITOR
 }
+PRAGMA_ENABLE_DEPRECATION_WARNINGS
 
 /*-----------------------------------------------------------------------------
 	FTextureRenderTarget2DResource
diff --git a/Engine/Source/Runtime/Engine/Private/TextureRenderTarget2DArray.cpp b/Engine/Source/Runtime/Engine/Private/TextureRenderTarget2DArray.cpp
--- a/Engine/Source/Runtime/Engine/Private/TextureRenderTarget2DArray.cpp
+++ b/Engine/Source/Runtime/Engine/Private/TextureRenderTarget2DArray.cpp
@@ -134,82 +134,54 @@ FString UTextureRenderTarget2DArray::GetDesc()
 	return FString::Printf( TEXT("Render to Texture 2DArray %dx%d[%s]"), SizeX, SizeX, GPixelFormats[GetFormat()].Name);
 }
 
-UTexture2DArray* UTextureRenderTarget2DArray::ConstructTexture2DArray(UObject* ObjOuter, const FString& NewTexName, EObjectFlags InFlags)
+TSubclassOf<UTexture> UTextureRenderTarget2DArray::GetTextureUClass() const
 {
-#if WITH_EDITOR
-	if (SizeX == 0 || SizeY == 0 || Slices == 0)
-	{
-		return nullptr;
-	}
-
-	const EPixelFormat PixelFormat = GetFormat();
-	ETextureSourceFormat TextureFormat = TSF_Invalid;
-	switch (PixelFormat)
-	{
-		case PF_FloatRGBA:
-			TextureFormat = TSF_RGBA16F;
-			break;
-	}
+	return UTexture2DArray::StaticClass();
+}
 
-	if (TextureFormat == TSF_Invalid)
+bool UTextureRenderTarget2DArray::CanConvertToTexture(ETextureSourceFormat& OutTextureSourceFormat, EPixelFormat& OutPixelFormat, FText* OutErrorMessage) const
+{
+	const EPixelFormat LocalFormat = GetFormat();
+	// These are the formats currently available for conversion to texture for UTextureRenderTarget2DArray : 
+	const ETextureSourceFormat TextureSourceFormat = ValidateTextureFormatForConversionToTextureInternal(GetFormat(), { PF_G8, PF_R8G8, PF_B8G8R8A8, PF_FloatRGBA }, OutErrorMessage);
+	if (TextureSourceFormat == TSF_Invalid)
 	{
-		return nullptr;
+		return false;
 	}
 
-	FTextureRenderTarget2DArrayResource* TextureResource = (FTextureRenderTarget2DArrayResource*)GameThread_GetRenderTargetResource();
-	if (TextureResource == nullptr)
+	if ((SizeX <= 0) || (SizeY <= 0) || (Slices <= 0))
 	{
-		return nullptr;
+		if (OutErrorMessage != nullptr)
+		{
+			*OutErrorMessage = FText::Format(NSLOCTEXT("TextureRenderTargetVolume", "InvalidSizeForConversionToTexture", "Invalid size (({0},{1}), {2} slices) for converting {3} to {4}"),
+				FText::AsNumber(SizeX),
+				FText::AsNumber(SizeY),
+				FText::AsNumber(Slices),
+				FText::FromString(GetClass()->GetName()),
+				FText::FromString(GetTextureUClass()->GetName()));
+		}
+		return false;
 	}
 
-	// Create texture
-	UTexture2DArray* Texture2DArray = NewObject<UTexture2DArray>(ObjOuter, FName(*NewTexName), InFlags);
+	OutPixelFormat = LocalFormat;
+	OutTextureSourceFormat = TextureSourceFormat;
+	return true;
+}
 
-	bool bSRGB = true;
-	// if render target gamma used was 1.0 then disable SRGB for the static texture
-	if (FMath::Abs(TextureResource->GetDisplayGamma() - 1.0f) < UE_KINDA_SMALL_NUMBER)
-	{
-		bSRGB = false;
-	}
-	Texture2DArray->Source.Init(SizeX, SizeX, Slices, 1, TextureFormat);
+UTexture2DArray* UTextureRenderTarget2DArray::ConstructTexture2DArray(UObject* InOuter, const FString& InNewTextureName, EObjectFlags InObjectFlags, uint32 InFlags, TArray<uint8>* InAlphaOverride)
+{
+	UTexture2DArray* Result = nullptr;
 
-	const int32 SrcMipSize = CalculateImageBytes(SizeX, SizeY, 1, PixelFormat);
-	const int32 DstMipSize = CalculateImageBytes(SizeX, SizeY, 1, PF_FloatRGBA);
-	uint8* SliceData = Texture2DArray->Source.LockMip(0);
-	switch (TextureFormat)
+#if WITH_EDITOR
+	FText ErrorMessage;
+	Result = Cast<UTexture2DArray>(ConstructTexture(InOuter, InNewTextureName, InObjectFlags, static_cast<EConstructTextureFlags>(InFlags), InAlphaOverride, &ErrorMessage));
+	if (Result == nullptr)
 	{
-		case TSF_RGBA16F:
-		{
-			for (int i = 0; i < Slices; ++i)
-			{
-				TArray<FFloat16Color> OutputBuffer;
-				FReadSurfaceDataFlags ReadSurfaceDataFlags(RCM_UNorm);
-				if (TextureResource->ReadPixels(OutputBuffer, i))
-				{
-					FMemory::Memcpy((FFloat16Color*)(SliceData + i * DstMipSize), OutputBuffer.GetData(), DstMipSize);
-				}
-			}
-			break;
-		}
-
-		default:
-			// Missing conversion from PF -> TSF
-			check(false);
-			break;
+		UE_LOG(LogTexture, Error, TEXT("Couldn't construct texture : %s"), *ErrorMessage.ToString());
 	}
-
-	Texture2DArray->Source.UnlockMip(0);
-	Texture2DArray->SRGB = bSRGB;
-	// If HDR source image then choose HDR compression settings..
-	Texture2DArray->CompressionSettings = TextureFormat == TSF_RGBA16F ? TextureCompressionSettings::TC_HDR : TextureCompressionSettings::TC_Default; //-V547 - future proofing
-	// Default to no mip generation for cube render target captures.
-	Texture2DArray->MipGenSettings = TextureMipGenSettings::TMGS_NoMipmaps;
-	Texture2DArray->PostEditChange();
-
-	return Texture2DArray;
-#else
-	return nullptr;
 #endif // #if WITH_EDITOR
+
+	return Result;
 }
 
 /*-----------------------------------------------------------------------------
@@ -357,50 +329,18 @@ float FTextureRenderTarget2DArrayResource::GetDisplayGamma() const
 	return FTextureRenderTargetResource::GetDisplayGamma();
 }
 
+PRAGMA_DISABLE_DEPRECATION_WARNINGS
 bool FTextureRenderTarget2DArrayResource::ReadPixels(TArray<FColor>& OutImageData, int32 InSlice, FIntRect InRect)
 {
-	if (InRect == FIntRect(0, 0, 0, 0))
-	{
-		InRect = FIntRect(0, 0, GetSizeXY().X, GetSizeXY().Y);
-	}
-
-	OutImageData.Reset();
-
-	ENQUEUE_RENDER_COMMAND(ReadSurfaceCommand)
-	(
-		[RenderTarget_RT=this, OutData_RT=&OutImageData, Rect_RT=InRect, Slice_RT=InSlice, bSRGB_RT= bSRGB](FRHICommandListImmediate& RHICmdList)
-		{
-			TArray<FFloat16Color> TempData;
-			RHICmdList.ReadSurfaceFloatData(RenderTarget_RT->TextureRHI, Rect_RT, TempData, (ECubeFace)0, Slice_RT, 0);
-			for (const FFloat16Color& SrcColor : TempData)
-			{
-				OutData_RT->Emplace(FLinearColor(SrcColor).ToFColor(bSRGB_RT));
-			}
-		}
-	);
-	FlushRenderingCommands();
-
-	return true;
+	FReadSurfaceDataFlags Flags;
+	Flags.SetArrayIndex(InSlice);
+	return FRenderTarget::ReadPixels(OutImageData, Flags, InRect);
 }
 
 bool FTextureRenderTarget2DArrayResource::ReadPixels(TArray<FFloat16Color>& OutImageData, int32 InSlice, FIntRect InRect)
 {
-	if (InRect == FIntRect(0, 0, 0, 0))
-	{
-		InRect = FIntRect(0, 0, GetSizeXY().X, GetSizeXY().Y);
-	}
-
-	OutImageData.Reset();
-
-	ENQUEUE_RENDER_COMMAND(ReadSurfaceCommand)
-	(
-		[RenderTarget_RT=this, OutData_RT=&OutImageData, Rect_RT=InRect, Slice_RT=InSlice](FRHICommandListImmediate& RHICmdList)
-		{
-			RHICmdList.ReadSurfaceFloatData(RenderTarget_RT->TextureRHI, Rect_RT, *OutData_RT, (ECubeFace)0, Slice_RT, 0);
-		}
-	);
-	FlushRenderingCommands();
-
-	return true;
+	FReadSurfaceDataFlags Flags;
+	Flags.SetArrayIndex(InSlice);
+	return FRenderTarget::ReadFloat16Pixels(OutImageData, Flags, InRect);
 }
-
+PRAGMA_ENABLE_DEPRECATION_WARNINGS
diff --git a/Engine/Source/Runtime/Engine/Private/TextureRenderTargetCube.cpp b/Engine/Source/Runtime/Engine/Private/TextureRenderTargetCube.cpp
--- a/Engine/Source/Runtime/Engine/Private/TextureRenderTargetCube.cpp
+++ b/Engine/Source/Runtime/Engine/Private/TextureRenderTargetCube.cpp
@@ -156,88 +156,51 @@ FString UTextureRenderTargetCube::GetDesc()
 	return FString::Printf( TEXT("Render to Texture Cube %dx%d[%s]"), SizeX, SizeX, GPixelFormats[GetFormat()].Name);
 }
 
-UTextureCube* UTextureRenderTargetCube::ConstructTextureCube(
-  UObject* ObjOuter,
-  const FString& NewTexName,
-  EObjectFlags InFlags
-)
+TSubclassOf<UTexture> UTextureRenderTargetCube::GetTextureUClass() const
 {
-	UTextureCube* Result = NULL;
-#if WITH_EDITOR
-	// Check render target size is valid and power of two.
-	if (SizeX != 0 && !(SizeX & (SizeX - 1)))
+	return UTextureCube::StaticClass();
+}
+
+bool UTextureRenderTargetCube::CanConvertToTexture(ETextureSourceFormat& OutTextureSourceFormat, EPixelFormat& OutPixelFormat, FText* OutErrorMessage) const
+{
+	const EPixelFormat LocalFormat = GetFormat();
+	// These are the formats currently available for conversion to texture for UTextureRenderTargetCube : 
+	const ETextureSourceFormat TextureSourceFormat = ValidateTextureFormatForConversionToTextureInternal(GetFormat(), { PF_B8G8R8A8, PF_FloatRGBA }, OutErrorMessage);
+	if (TextureSourceFormat == TSF_Invalid)
 	{
-		const EPixelFormat PixelFormat = GetFormat();
-		ETextureSourceFormat TextureFormat = TSF_Invalid;
-		switch (PixelFormat)
-		{
-			case PF_FloatRGBA:
-				TextureFormat = TSF_RGBA16F;
-				break;
-			case PF_B8G8R8A8:
-				TextureFormat = TSF_BGRA8;
-				break;
-			default:
-				return nullptr;
-		}
+		return false;
+	}
 
-		// The r2t resource will be needed to read its surface contents
-		FTextureRenderTargetCubeResource* CubeResource = (FTextureRenderTargetCubeResource*)GameThread_GetRenderTargetResource();
-		if (CubeResource && TextureFormat != TSF_Invalid)
+	if ((SizeX <= 0) || (SizeX & (SizeX - 1)))
+	{
+		if (OutErrorMessage != nullptr)
 		{
-			// create the cube texture
-			Result = NewObject<UTextureCube>(ObjOuter, FName(*NewTexName), InFlags);
-
-			bool bSRGB = true;
-			// if render target gamma used was 1.0 then disable SRGB for the static texture
-			if (FMath::Abs(CubeResource->GetDisplayGamma() - 1.0f) < UE_KINDA_SMALL_NUMBER)
-			{
-				bSRGB = false;
-			}
-
-			Result->Source.Init(SizeX, SizeX, 6, 1, TextureFormat);
+			*OutErrorMessage = FText::Format(NSLOCTEXT("TextureRenderTargetVolume", "InvalidSizeForConversionToTexture", "Invalid size ({0},{0}) for converting {1} to {2}. Needs to be a power of 2."),
+				FText::AsNumber(SizeX),
+				FText::FromString(GetClass()->GetName()),
+				FText::FromString(GetTextureUClass()->GetName()));
+		}
+		return false;
+	}
 
-			int32 MipSize = CalculateImageBytes(SizeX, SizeX, 0, PixelFormat);
-			uint8* SliceData = Result->Source.LockMip(0);
+	OutPixelFormat = LocalFormat;
+	OutTextureSourceFormat = TextureSourceFormat;
+	return true;
+}
 
-			switch (TextureFormat)
-			{
-				case TSF_BGRA8:
-				{
-					TArray<FColor> OutputBuffer;
-					for (int32 SliceIndex = 0; SliceIndex < 6; SliceIndex++)
-					{
-						if (CubeResource->ReadPixels(OutputBuffer, FReadSurfaceDataFlags(RCM_UNorm, (ECubeFace)SliceIndex)))
-						{
-							FMemory::Memcpy((FColor*)(SliceData + SliceIndex * MipSize), OutputBuffer.GetData(), MipSize);
-						}
-					}
-				}
-				break;
-				case TSF_RGBA16F:
-				{
-					TArray<FFloat16Color> OutputBuffer;
-					for (int32 SliceIndex = 0; SliceIndex < 6; SliceIndex++)
-					{
-						if (CubeResource->ReadPixels(OutputBuffer, FReadSurfaceDataFlags(RCM_UNorm, (ECubeFace)SliceIndex)))
-						{
-							FMemory::Memcpy((FFloat16Color*)(SliceData + SliceIndex * MipSize), OutputBuffer.GetData(), MipSize);
-						}
-					}
-				}
-				break;
-			}
+UTextureCube* UTextureRenderTargetCube::ConstructTextureCube(UObject* InOuter, const FString& InNewTextureName, EObjectFlags InObjectFlags, uint32 InFlags, TArray<uint8>* InAlphaOverride)
+{
+	UTextureCube* Result = NULL;
 
-			Result->Source.UnlockMip(0);
-			Result->SRGB = bSRGB;
-			// If HDR source image then choose HDR compression settings..
-			Result->CompressionSettings = TextureFormat == TSF_RGBA16F ? TextureCompressionSettings::TC_HDR : TextureCompressionSettings::TC_Default;
-			// Default to no mip generation for cube render target captures.
-			Result->MipGenSettings = TextureMipGenSettings::TMGS_NoMipmaps;
-			Result->PostEditChange();
-		}
+#if WITH_EDITOR
+	FText ErrorMessage;
+	Result = Cast<UTextureCube>(ConstructTexture(InOuter, InNewTextureName, InObjectFlags, static_cast<EConstructTextureFlags>(InFlags), InAlphaOverride, &ErrorMessage));
+	if (Result == nullptr)
+	{
+		UE_LOG(LogTexture, Error, TEXT("Couldn't construct texture : %s"), *ErrorMessage.ToString());
 	}
 #endif // #if WITH_EDITOR
+
 	return Result;
 }
 
@@ -414,7 +377,7 @@ float FTextureRenderTargetCubeResource::GetDisplayGamma() const
 	return FTextureRenderTargetResource::GetDisplayGamma();
 }
 
-
+PRAGMA_DISABLE_DEPRECATION_WARNINGS
 /**
 * Copy the texels of a single face of a cube texture into an array.
 * @param OutImageData - RGBA8 values will be stored in this array.
@@ -424,42 +387,7 @@ float FTextureRenderTargetCubeResource::GetDisplayGamma() const
 */
 bool FTextureRenderTargetCubeResource::ReadPixels(TArray< FColor >& OutImageData, FReadSurfaceDataFlags InFlags, FIntRect InRect)
 {
-	if (InRect == FIntRect(0, 0, 0, 0))
-	{
-		InRect = FIntRect(0, 0, GetSizeXY().X, GetSizeXY().Y);
-	}
-
-	// Read the render target surface data back.
-	struct FReadSurfaceContext
-	{
-		FTextureRenderTargetCubeResource* SrcRenderTarget;
-		TArray<FColor>* OutData;
-		FIntRect Rect;
-		FReadSurfaceDataFlags Flags;
-	};
-
-	OutImageData.Reset();
-	FReadSurfaceContext Context =
-	{
-		this,
-		&OutImageData,
-		InRect,
-		InFlags
-	};
-
-	ENQUEUE_RENDER_COMMAND(ReadSurfaceCommand)(
-		[Context](FRHICommandListImmediate& RHICmdList)
-		{
-			RHICmdList.ReadSurfaceData(
-				Context.SrcRenderTarget->TextureRHI,
-				Context.Rect,
-				*Context.OutData,
-				Context.Flags
-			);
-		});
-	FlushRenderingCommands();
-
-	return true;
+	return FRenderTarget::ReadPixels(OutImageData, InFlags, InRect);
 }
 
 /**
@@ -471,42 +399,6 @@ bool FTextureRenderTargetCubeResource::ReadPixels(TArray< FColor >& OutImageData
 */
 bool FTextureRenderTargetCubeResource::ReadPixels(TArray<FFloat16Color>& OutImageData, FReadSurfaceDataFlags InFlags, FIntRect InRect)
 {
-	if (InRect == FIntRect(0, 0, 0, 0))
-	{
-		InRect = FIntRect(0, 0, GetSizeXY().X, GetSizeXY().Y);
-	}
-	// Read the render target surface data back.
-	struct FReadSurfaceFloatContext
-	{
-		FTextureRenderTargetCubeResource* SrcRenderTarget;
-		TArray<FFloat16Color>* OutData;
-		FIntRect Rect;
-		ECubeFace CubeFace;
-	};
-
-	FReadSurfaceFloatContext Context =
-	{
-		this,
-		&OutImageData,
-		InRect,
-		InFlags.GetCubeFace()
-	};
-
-	ENQUEUE_RENDER_COMMAND(ReadSurfaceFloatCommand)(
-		[Context](FRHICommandListImmediate& RHICmdList)
-		{
-			RHICmdList.ReadSurfaceFloatData(
-				Context.SrcRenderTarget->TextureRHI,
-				Context.Rect,
-				*Context.OutData,
-				Context.CubeFace,
-				0,
-				0
-			);
-	});
-
-	FlushRenderingCommands();
-
-	return true;
+	return FRenderTarget::ReadFloat16Pixels(OutImageData, InFlags, InRect);
 }
-
+PRAGMA_ENABLE_DEPRECATION_WARNINGS
diff --git a/Engine/Source/Runtime/Engine/Private/TextureRenderTargetVolume.cpp b/Engine/Source/Runtime/Engine/Private/TextureRenderTargetVolume.cpp
--- a/Engine/Source/Runtime/Engine/Private/TextureRenderTargetVolume.cpp
+++ b/Engine/Source/Runtime/Engine/Private/TextureRenderTargetVolume.cpp
@@ -138,84 +138,56 @@ FString UTextureRenderTargetVolume::GetDesc()
 	return FString::Printf( TEXT("Render to Texture Volume %dx%d[%s]"), SizeX, SizeX, GPixelFormats[GetFormat()].Name);
 }
 
-UVolumeTexture* UTextureRenderTargetVolume::ConstructTextureVolume(UObject* ObjOuter, const FString& NewTexName, EObjectFlags InFlags)
+TSubclassOf<UTexture> UTextureRenderTargetVolume::GetTextureUClass() const
 {
-#if WITH_EDITOR
-	if (SizeX == 0 || SizeY == 0 || SizeZ == 0)
-	{
-		return nullptr;
-	}
-
-
-	const EPixelFormat PixelFormat = GetFormat();
-	ETextureSourceFormat TextureFormat = TSF_Invalid;
-	switch (PixelFormat)
-	{
-		case PF_R32_FLOAT:
-		case PF_FloatRGBA:
-			TextureFormat = TSF_RGBA16F;
-			break;
-	}
+	return UVolumeTexture::StaticClass();
+}
 
-	if (TextureFormat == TSF_Invalid)
+bool UTextureRenderTargetVolume::CanConvertToTexture(ETextureSourceFormat& OutTextureSourceFormat, EPixelFormat& OutPixelFormat, FText* OutErrorMessage) const
+{
+	const EPixelFormat LocalFormat = GetFormat();
+	// These are the formats currently available for conversion to texture for UTextureRenderTargetVolume : 
+	// PF_R32_FLOAT and PF_FloatRGBA are supported by FTextureRenderTargetVolumeResource::ReadFloat16Pixels
+	// TODO [jonathan.bard] : PF_R16F is also supported by by FTextureRenderTargetVolumeResource::ReadFloat16Pixels so it could be added to the list
+	ETextureSourceFormat TextureSourceFormat = ValidateTextureFormatForConversionToTextureInternal(GetFormat(), { PF_R32_FLOAT, PF_FloatRGBA }, OutErrorMessage);
+	if (TextureSourceFormat == TSF_Invalid)
 	{
-		return nullptr;
+		return false;
 	}
 
-	FTextureRenderTargetVolumeResource* TextureResource = (FTextureRenderTargetVolumeResource*)GameThread_GetRenderTargetResource();
-	if (TextureResource == nullptr)
+	if ((SizeX <= 0) || (SizeY <= 0) || (SizeZ <= 0))
 	{
-		return nullptr;
+		if (OutErrorMessage != nullptr)
+		{
+			*OutErrorMessage = FText::Format(NSLOCTEXT("TextureRenderTargetVolume", "InvalidSizeForConversionToTexture", "Invalid size ({0},{1},{2}) for converting {3} to {4}"),
+				FText::AsNumber(SizeX),
+				FText::AsNumber(SizeY),
+				FText::AsNumber(SizeZ),
+				FText::FromString(GetClass()->GetName()),
+				FText::FromString(GetTextureUClass()->GetName()));
+		}
+		return false;
 	}
 
-	// Create texture
-	UVolumeTexture* VolumeTexture = NewObject<UVolumeTexture>(ObjOuter, FName(*NewTexName), InFlags);
+	OutPixelFormat = LocalFormat;
+	OutTextureSourceFormat = TextureSourceFormat;
+	return true;
+}
 
-	bool bSRGB = true;
-	// if render target gamma used was 1.0 then disable SRGB for the static texture
-	if (FMath::Abs(TextureResource->GetDisplayGamma() - 1.0f) < UE_KINDA_SMALL_NUMBER)
-	{
-		bSRGB = false;
-	}
-	VolumeTexture->Source.Init(SizeX, SizeX, SizeZ, 1, TextureFormat);
+UVolumeTexture* UTextureRenderTargetVolume::ConstructTextureVolume(UObject* InOuter, const FString& InNewTextureName, EObjectFlags InObjectFlags, uint32 InFlags, TArray<uint8>* InAlphaOverride)
+{
+	UVolumeTexture* Result = nullptr;
 
-	const int32 SrcMipSize = CalculateImageBytes(SizeX, SizeX, 1, PixelFormat);
-	const int32 DstMipSize = CalculateImageBytes(SizeX, SizeX, 1, PF_FloatRGBA);
-	uint8* SliceData = VolumeTexture->Source.LockMip(0);
-	switch (TextureFormat)
+#if WITH_EDITOR
+	FText ErrorMessage;
+	Result = Cast<UVolumeTexture>(ConstructTexture(InOuter, InNewTextureName, InObjectFlags, static_cast<EConstructTextureFlags>(InFlags), InAlphaOverride, &ErrorMessage));
+	if (Result == nullptr)
 	{
-		case TSF_RGBA16F:
-		{
-			for (int i = 0; i < SizeZ; ++i)
-			{
-				TArray<FFloat16Color> OutputBuffer;
-				FReadSurfaceDataFlags ReadSurfaceDataFlags(RCM_UNorm);
-				if (TextureResource->ReadPixels(OutputBuffer, i))
-				{
-					FMemory::Memcpy((FFloat16Color*)(SliceData + i * DstMipSize), OutputBuffer.GetData(), DstMipSize);
-				}
-			}
-			break;
-		}
-
-		default:
-			// Missing conversion from PF -> TSF
-			check(false);
-			break;
+		UE_LOG(LogTexture, Error, TEXT("Couldn't construct texture : %s"), *ErrorMessage.ToString());
 	}
-
-	VolumeTexture->Source.UnlockMip(0);
-	VolumeTexture->SRGB = bSRGB;
-	// If HDR source image then choose HDR compression settings..
-	VolumeTexture->CompressionSettings = TextureFormat == TSF_RGBA16F ? TextureCompressionSettings::TC_HDR : TextureCompressionSettings::TC_Default; //-V547 - future proofing
-	// Default to no mip generation for cube render target captures.
-	VolumeTexture->MipGenSettings = TextureMipGenSettings::TMGS_NoMipmaps;
-	VolumeTexture->PostEditChange();
-
-	return VolumeTexture;
-#else
-	return nullptr;
 #endif // #if WITH_EDITOR
+
+	return Result;
 }
 
 /*-----------------------------------------------------------------------------
@@ -353,32 +325,34 @@ float FTextureRenderTargetVolumeResource::GetDisplayGamma() const
 		return Owner->TargetGamma;
 	}
 	EPixelFormat Format = Owner->GetFormat();
-	if(Format == PF_FloatRGB || Format == PF_FloatRGBA || Owner->bForceLinearGamma)
+	if(Format == PF_R32_FLOAT || Format == PF_FloatRGBA || Owner->bForceLinearGamma)
 	{
 		return 1.0f;
 	}
 	return FTextureRenderTargetResource::GetDisplayGamma();
 }
 
-bool FTextureRenderTargetVolumeResource::ReadPixels(TArray<FColor>& OutImageData, int32 InDepthSlice, FIntRect InRect)
+bool FTextureRenderTargetVolumeResource::ReadPixels(TArray<FColor>& OutImageData, FReadSurfaceDataFlags InFlags, FIntRect InSrcRect)
+{
+	checkf(false, TEXT("Not implemented : volume textures only support float16 ATM. Use ReadFloat16Pixels and convert to the desired format"));
+	return false;
+}
+
+bool FTextureRenderTargetVolumeResource::ReadFloat16Pixels(TArray<FFloat16Color>& OutImageData, FReadSurfaceDataFlags InFlags, FIntRect InSrcRect)
 {
-	if (InRect == FIntRect(0, 0, 0, 0))
+	if (InSrcRect == FIntRect(0, 0, 0, 0))
 	{
-		InRect = FIntRect(0, 0, GetSizeXY().X, GetSizeXY().Y);
+		InSrcRect = FIntRect(0, 0, GetSizeXY().X, GetSizeXY().Y);
 	}
 
 	OutImageData.Reset();
 
+	// Read the render target surface data back.	
 	ENQUEUE_RENDER_COMMAND(ReadSurfaceCommand)
 	(
-		[RenderTarget_RT=this, OutData_RT=&OutImageData, Rect_RT=InRect, DepthSlice_RT=InDepthSlice, bSRGB_RT= bSRGB](FRHICommandListImmediate& RHICmdList)
+		[RenderTarget_RT = this, SrcRect_RT = InSrcRect, OutData_RT = &OutImageData, Flags_RT = InFlags](FRHICommandListImmediate& RHICmdList)
 		{
-			TArray<FFloat16Color> TempData;
-			RHICmdList.Read3DSurfaceFloatData(RenderTarget_RT->TextureRHI, Rect_RT, FIntPoint(DepthSlice_RT, DepthSlice_RT + 1), TempData);
-			for (const FFloat16Color& SrcColor : TempData)
-			{
-				OutData_RT->Emplace(FLinearColor(SrcColor).ToFColor(bSRGB_RT));
-			}
+			RHICmdList.Read3DSurfaceFloatData(RenderTarget_RT->GetRenderTargetTexture(), SrcRect_RT, FIntPoint(Flags_RT.GetArrayIndex(), Flags_RT.GetArrayIndex() + 1), *OutData_RT);
 		}
 	);
 	FlushRenderingCommands();
@@ -386,24 +360,33 @@ bool FTextureRenderTargetVolumeResource::ReadPixels(TArray<FColor>& OutImageData
 	return true;
 }
 
-bool FTextureRenderTargetVolumeResource::ReadPixels(TArray<FFloat16Color>& OutImageData, int32 InDepthSlice, FIntRect InRect)
+bool FTextureRenderTargetVolumeResource::ReadLinearColorPixels(TArray<FLinearColor>& OutImageData, FReadSurfaceDataFlags InFlags, FIntRect InSrcRect)
 {
-	if (InRect == FIntRect(0, 0, 0, 0))
-	{
-		InRect = FIntRect(0, 0, GetSizeXY().X, GetSizeXY().Y);
-	}
-
-	OutImageData.Reset();
+	checkf(false, TEXT("Not implemented : volume textures only support float16 ATM. Use ReadFloat16Pixels and convert to the desired format"));
+	return false;
+}
 
-	ENQUEUE_RENDER_COMMAND(ReadSurfaceCommand)
-	(
-		[RenderTarget_RT=this, OutData_RT=&OutImageData, Rect_RT=InRect, DepthSlice_RT=InDepthSlice](FRHICommandListImmediate& RHICmdList)
+PRAGMA_DISABLE_DEPRECATION_WARNINGS
+bool FTextureRenderTargetVolumeResource::ReadPixels(TArray<FColor>& OutImageData, int32 InDepthSlice, FIntRect InRect)
+{
+	FReadSurfaceDataFlags Flags;
+	Flags.SetArrayIndex(InDepthSlice);
+	TArray<FFloat16Color> TempData;
+	if (FRenderTarget::ReadFloat16Pixels(TempData, Flags, InRect))
+	{
+		for (const FFloat16Color& SrcColor : TempData)
 		{
-			RHICmdList.Read3DSurfaceFloatData(RenderTarget_RT->TextureRHI, Rect_RT, FIntPoint(DepthSlice_RT, DepthSlice_RT+1), *OutData_RT);
+			OutImageData.Emplace(FLinearColor(SrcColor).ToFColor(bSRGB));
 		}
-	);
-	FlushRenderingCommands();
+	}
 
 	return true;
 }
 
+bool FTextureRenderTargetVolumeResource::ReadPixels(TArray<FFloat16Color>& OutImageData, int32 InDepthSlice, FIntRect InRect)
+{
+	FReadSurfaceDataFlags Flags;
+	Flags.SetArrayIndex(InDepthSlice);
+	return FRenderTarget::ReadFloat16Pixels(OutImageData, Flags, InRect);
+}
+PRAGMA_ENABLE_DEPRECATION_WARNINGS
diff --git a/Engine/Source/Runtime/Engine/Private/UnrealClient.cpp b/Engine/Source/Runtime/Engine/Private/UnrealClient.cpp
--- a/Engine/Source/Runtime/Engine/Private/UnrealClient.cpp
+++ b/Engine/Source/Runtime/Engine/Private/UnrealClient.cpp
@@ -49,66 +49,31 @@ int32 FViewport::PresentAndStopMovieDelay = 0;
 
 static const FName NAME_DummyViewport = FName(TEXT("DummyViewport"));
 
-/**
-* Reads the viewport's displayed pixels into a preallocated color buffer.
-* @param OutImageData - RGBA8 values will be stored in this buffer
-* @param TopLeftX - Top left X pixel to capture
-* @param TopLeftY - Top left Y pixel to capture
-* @param Width - Width of image in pixels to capture
-* @param Height - Height of image in pixels to capture
-* @return True if the read succeeded.
-*/
-bool FRenderTarget::ReadPixels(TArray< FColor >& OutImageData, FReadSurfaceDataFlags InFlags, FIntRect InRect)
+bool FRenderTarget::ReadPixels(TArray< FColor >& OutImageData, FReadSurfaceDataFlags InFlags, FIntRect InSrcRect)
 {
-	if(InRect == FIntRect(0, 0, 0, 0))
+	if(InSrcRect == FIntRect(0, 0, 0, 0))
 	{
-		InRect = FIntRect(0, 0, GetSizeXY().X, GetSizeXY().Y);
+		InSrcRect = FIntRect(0, 0, GetSizeXY().X, GetSizeXY().Y);
 	}
 
-	// Read the render target surface data back.	
-	struct FReadSurfaceContext
-	{
-		FRenderTarget* SrcRenderTarget;
-		TArray<FColor>* OutData;
-		FIntRect Rect;
-		FReadSurfaceDataFlags Flags;
-	};
-
 	OutImageData.Reset();
-	FReadSurfaceContext Context =
-	{
-		this,
-		&OutImageData,
-		InRect,
-		InFlags
-	};
 
+	// Read the render target surface data back.	
 	ENQUEUE_RENDER_COMMAND(ReadSurfaceCommand)(
-		[Context](FRHICommandListImmediate& RHICmdList)
+		[RenderTarget_RT = this, SrcRect_RT = InSrcRect, OutData_RT = &OutImageData, Flags_RT = InFlags](FRHICommandListImmediate& RHICmdList)
 		{
-			RHICmdList.ReadSurfaceData(
-				Context.SrcRenderTarget->GetRenderTargetTexture(),
-				Context.Rect,
-				*Context.OutData,
-				Context.Flags
-				);
+			RHICmdList.ReadSurfaceData(RenderTarget_RT->GetRenderTargetTexture(), SrcRect_RT, *OutData_RT, Flags_RT);
 		});
 	FlushRenderingCommands();
 
 	return OutImageData.Num() > 0;
 }
 
-
-/**
-* Reads the viewport's displayed pixels into a preallocated color buffer.
-* @param OutputBuffer - RGBA8 values will be stored in this buffer
-* @return True if the read succeeded.
-*/
-bool FRenderTarget::ReadPixelsPtr(FColor* OutImageBytes, FReadSurfaceDataFlags InFlags, FIntRect InRect)
+bool FRenderTarget::ReadPixelsPtr(FColor* OutImageBytes, FReadSurfaceDataFlags InFlags, FIntRect InSrcRect)
 {
 	TArray<FColor> SurfaceData;
 
-	bool bResult = ReadPixels( SurfaceData, InFlags, InRect );
+	bool bResult = ReadPixels( SurfaceData, InFlags, InSrcRect);
 	if( bResult )
 	{
 		FMemory::Memcpy( OutImageBytes, &SurfaceData[ 0 ], SurfaceData.Num() * sizeof(FColor) );
@@ -117,133 +82,58 @@ bool FRenderTarget::ReadPixelsPtr(FColor* OutImageBytes, FReadSurfaceDataFlags I
 	return bResult;
 }
 
-/**
- * Reads the viewport's displayed pixels into a preallocated color buffer.
- * @param OutImageBytes - RGBA16F values will be stored in this buffer.  Buffer must be preallocated with the correct size!
- * @param CubeFace - optional cube face for when reading from a cube render target
- * @return True if the read succeeded.
- */
-bool FRenderTarget::ReadFloat16Pixels(FFloat16Color* OutImageData,ECubeFace CubeFace)
+PRAGMA_DISABLE_DEPRECATION_WARNINGS
+bool FRenderTarget::ReadFloat16Pixels(TArray<FFloat16Color>& OutImageData, ECubeFace CubeFace)
 {
-	// Read the render target surface data back.	
-	struct FReadSurfaceFloatContext
-	{
-		FRenderTarget* SrcRenderTarget;
-		TArray<FFloat16Color>* OutData;
-		FIntRect Rect;
-		ECubeFace CubeFace;
-	};
-	
-	TArray<FFloat16Color> SurfaceData;
-	FReadSurfaceFloatContext Context =
-	{
-		this,
-		&SurfaceData,
-		FIntRect(0, 0, GetSizeXY().X, GetSizeXY().Y),
-		CubeFace	
-	};
-
-	ENQUEUE_RENDER_COMMAND(ReadSurfaceFloatCommand)(
-		[Context](FRHICommandListImmediate& RHICmdList)
-		{
-			RHICmdList.ReadSurfaceFloatData(
-				Context.SrcRenderTarget->GetRenderTargetTexture(),
-				Context.Rect,
-				*Context.OutData,
-				Context.CubeFace,
-				0,
-				0
-				);
-		});
-	FlushRenderingCommands();
-
-	// Copy the surface data into the output array.
-	FFloat16Color* OutImageColors = reinterpret_cast< FFloat16Color* >(OutImageData);
+	return ReadFloat16Pixels(OutImageData, FReadSurfaceDataFlags(RCM_UNorm, CubeFace));
+}
+PRAGMA_ENABLE_DEPRECATION_WARNINGS
 
-	// Cache width and height as its very expensive to call these virtuals in inner loop (never inlined)
-	const int32 ImageWidth = GetSizeXY().X;
-	const int32 ImageHeight = GetSizeXY().Y;
-	for (int32 Y = 0; Y < ImageHeight; Y++)
+bool FRenderTarget::ReadFloat16Pixels(TArray<FFloat16Color>& OutImageData, FReadSurfaceDataFlags InFlags, FIntRect InSrcRect)
+{
+	if (InSrcRect == FIntRect(0, 0, 0, 0))
 	{
-		FFloat16Color* SourceData = (FFloat16Color*)SurfaceData.GetData() + Y * ImageWidth;
-		for (int32 X = 0; X < ImageWidth; X++)
-		{
-			OutImageColors[ Y * ImageWidth + X ] = SourceData[X];
-		}
+		InSrcRect = FIntRect(0, 0, GetSizeXY().X, GetSizeXY().Y);
 	}
 
-	return true;
-}
+	OutImageData.Reset();
 
-/**
- * Reads the viewport's displayed pixels into the given color buffer.
- * @param OutputBuffer - RGBA16F values will be stored in this buffer
- * @param CubeFace - optional cube face for when reading from a cube render target
- * @return True if the read succeeded.
- */
-bool FRenderTarget::ReadFloat16Pixels(TArray<FFloat16Color>& OutputBuffer,ECubeFace CubeFace)
-{
-	// Copy the surface data into the output array.
-	OutputBuffer.SetNumUninitialized(GetSizeXY().X * GetSizeXY().Y, true);
-	return ReadFloat16Pixels((FFloat16Color*)&(OutputBuffer[0]), CubeFace);
+	// Read the render target surface data back.	
+	ENQUEUE_RENDER_COMMAND(ReadSurfaceFloatCommand)(
+		[RenderTarget_RT = this, SrcRect_RT = InSrcRect, OutData_RT = &OutImageData, Flags_RT = InFlags](FRHICommandListImmediate& RHICmdList)
+	{
+		RHICmdList.ReadSurfaceFloatData(RenderTarget_RT->GetRenderTargetTexture(), SrcRect_RT, *OutData_RT, Flags_RT);
+	});
+	FlushRenderingCommands();
+
+	return OutImageData.Num() > 0;;
 }
 
-/**
-* Reads the viewport's displayed pixels into a preallocated color buffer.
-* @param OutImageData - LinearColor array to fill!
-* @param CubeFace - optional cube face for when reading from a cube render target
-* @return True if the read succeeded.
-*/
-bool FRenderTarget::ReadLinearColorPixels(TArray<FLinearColor> &OutImageData, FReadSurfaceDataFlags InFlags, FIntRect InRect)
+bool FRenderTarget::ReadLinearColorPixels(TArray<FLinearColor> &OutImageData, FReadSurfaceDataFlags InFlags, FIntRect InSrcRect)
 {
-	if (InRect == FIntRect(0, 0, 0, 0))
+	if (InSrcRect == FIntRect(0, 0, 0, 0))
 	{
-		InRect = FIntRect(0, 0, GetSizeXY().X, GetSizeXY().Y);
+		InSrcRect = FIntRect(0, 0, GetSizeXY().X, GetSizeXY().Y);
 	}
 
-	// Read the render target surface data back.	
-	struct FReadSurfaceContext
-	{
-		FRenderTarget* SrcRenderTarget;
-		TArray<FLinearColor>* OutData;
-		FIntRect Rect;
-		FReadSurfaceDataFlags Flags;
-	};
-
 	OutImageData.Reset();
-	FReadSurfaceContext Context =
-	{
-		this,
-		&OutImageData,
-		InRect,
-		InFlags
-	};
 
+	// Read the render target surface data back.	
 	ENQUEUE_RENDER_COMMAND(ReadSurfaceCommand)(
-		[Context](FRHICommandListImmediate& RHICmdList)
+		[RenderTarget_RT = this, SrcRect_RT = InSrcRect, OutData_RT = &OutImageData, Flags_RT = InFlags](FRHICommandListImmediate& RHICmdList)
 		{
-			RHICmdList.ReadSurfaceData(
-			Context.SrcRenderTarget->GetRenderTargetTexture(),
-				Context.Rect,
-				*Context.OutData,
-				Context.Flags
-				);
+			RHICmdList.ReadSurfaceData(RenderTarget_RT->GetRenderTargetTexture(), SrcRect_RT, *OutData_RT, Flags_RT);
 		});
 	FlushRenderingCommands();
 
 	return OutImageData.Num() > 0;
 }
 
-/**
-* Reads the viewport's displayed pixels into a preallocated color buffer.
-* @param OutputBuffer - RGBA8 values will be stored in this buffer
-* @return True if the read succeeded.
-*/
-bool FRenderTarget::ReadLinearColorPixelsPtr(FLinearColor* OutImageBytes, FReadSurfaceDataFlags InFlags, FIntRect InRect)
+bool FRenderTarget::ReadLinearColorPixelsPtr(FLinearColor* OutImageBytes, FReadSurfaceDataFlags InFlags, FIntRect InSrcRect)
 {
 	TArray<FLinearColor> SurfaceData;
 
-	bool bResult = ReadLinearColorPixels(SurfaceData, InFlags, InRect);
+	bool bResult = ReadLinearColorPixels(SurfaceData, InFlags, InSrcRect);
 	if (bResult)
 	{
 		check(SurfaceData.Num() != 0);
diff --git a/Engine/Source/Runtime/Engine/Private/WorldPartition/WorldPartitionMiniMapHelper.cpp b/Engine/Source/Runtime/Engine/Private/WorldPartition/WorldPartitionMiniMapHelper.cpp
--- a/Engine/Source/Runtime/Engine/Private/WorldPartition/WorldPartitionMiniMapHelper.cpp
+++ b/Engine/Source/Runtime/Engine/Private/WorldPartition/WorldPartitionMiniMapHelper.cpp
@@ -11,6 +11,7 @@
 #include "Components/SceneCaptureComponent2D.h"
 #include "Engine/Scene.h"
 #include "Engine/TextureRenderTarget2D.h"
+#include "Engine/Texture2D.h"
 #include "Engine/SceneCapture2D.h"
 #include "Engine/Level.h"
 #include "Engine/World.h"
@@ -109,7 +110,15 @@ void FWorldPartitionMiniMapHelper::CaptureBoundsMiniMapToTexture(UWorld* InWorld
 	}
 	else
 	{
-		RenderTargetTexture->UpdateTexture2D(InOutMiniMapTexture, TSF_BGRA8, CTF_Default);
+		FText ErrorMessage;
+		if (RenderTargetTexture->UpdateTexture(InOutMiniMapTexture, CTF_Default, /*InAlphaOverride = */nullptr, /*InTextureChangingDelegate =*/ [](UTexture*) {}, &ErrorMessage))
+		{
+			check(InOutMiniMapTexture->Source.GetFormat() == TSF_BGRA8);
+		}
+		else
+		{
+			UE_LOG(LogWorldPartitionMiniMapHelper, Error, TEXT("Couldn't copy render target to internal texture: %s"), *ErrorMessage.ToString());
+		}
 	}
 }
 
diff --git a/Engine/Source/Runtime/Engine/Public/TextureRenderTarget2DArrayResource.h b/Engine/Source/Runtime/Engine/Public/TextureRenderTarget2DArrayResource.h
--- a/Engine/Source/Runtime/Engine/Public/TextureRenderTarget2DArrayResource.h
+++ b/Engine/Source/Runtime/Engine/Public/TextureRenderTarget2DArrayResource.h
@@ -75,6 +75,9 @@ public:
 	*/
 	float GetDisplayGamma() const override;
 
+	// UE_DEPRECATED(5.4, "This using is there temporarily until the 2 deprecated ReadPixels 'overrides' are removed : they were hiding FRenderTarget's virtual functions")
+	using FRenderTarget::ReadPixels;
+
 	/**
 	* Copy the texels of a single depth slice of the 2d array into an array.
 	* @param OutImageData - float16 values will be stored in this array.
@@ -82,6 +85,7 @@ public:
 	* @param InRect - Rectangle of texels to copy.
 	* @return true if the read succeeded.
 	*/
+	UE_DEPRECATED(5.4, "Use FRenderTarget's ReadPixels, which is functionally equivalent")
 	ENGINE_API bool ReadPixels(TArray<FColor>& OutImageData, int32 InDepthSlice, FIntRect InRect = FIntRect(0, 0, 0, 0));
 
 	/**
@@ -91,6 +95,7 @@ public:
 	* @param InRect - Rectangle of texels to copy.
 	* @return true if the read succeeded.
 	*/
+	UE_DEPRECATED(5.4, "Use FRenderTarget's ReadFloat16Pixels, which is functionally equivalent")
 	ENGINE_API bool ReadPixels(TArray<FFloat16Color>& OutImageData, int32 InDepthSlice, FIntRect InRect = FIntRect(0, 0, 0, 0));
 
 protected:
diff --git a/Engine/Source/Runtime/Engine/Public/TextureRenderTargetVolumeResource.h b/Engine/Source/Runtime/Engine/Public/TextureRenderTargetVolumeResource.h
--- a/Engine/Source/Runtime/Engine/Public/TextureRenderTargetVolumeResource.h
+++ b/Engine/Source/Runtime/Engine/Public/TextureRenderTargetVolumeResource.h
@@ -75,6 +75,12 @@ public:
 	*/
 	float GetDisplayGamma() const override;
 
+	virtual bool ReadPixels(TArray<FColor>& OutImageData, FReadSurfaceDataFlags InFlags = FReadSurfaceDataFlags(RCM_UNorm, CubeFace_MAX), FIntRect InSrcRect = FIntRect(0, 0, 0, 0)) override;
+
+	virtual bool ReadFloat16Pixels(TArray<FFloat16Color>& OutImageData, FReadSurfaceDataFlags InFlags = FReadSurfaceDataFlags(RCM_UNorm, CubeFace_MAX), FIntRect InSrcRect = FIntRect(0, 0, 0, 0)) override;
+
+	virtual bool ReadLinearColorPixels(TArray<FLinearColor>& OutImageData, FReadSurfaceDataFlags InFlags = FReadSurfaceDataFlags(RCM_UNorm, CubeFace_MAX), FIntRect InSrcRect = FIntRect(0, 0, 0, 0)) override;
+
 	/**
 	* Copy the texels of a single depth slice of the volume into an array.
 	* @param OutImageData - float16 values will be stored in this array.
@@ -82,6 +88,7 @@ public:
 	* @param InRect - Rectangle of texels to copy.
 	* @return true if the read succeeded.
 	*/
+	UE_DEPRECATED(5.4, "Use FRenderTarget's ReadPixels, which is functionally equivalent")
 	ENGINE_API bool ReadPixels(TArray<FColor>& OutImageData, int32 InDepthSlice, FIntRect InRect = FIntRect(0, 0, 0, 0));
 
 	/**
@@ -91,6 +98,7 @@ public:
 	* @param InRect - Rectangle of texels to copy.
 	* @return true if the read succeeded.
 	*/
+	UE_DEPRECATED(5.4, "Use FRenderTarget's ReadFloat16Pixels, which is functionally equivalent")
 	ENGINE_API bool ReadPixels(TArray<FFloat16Color>& OutImageData, int32 InDepthSlice, FIntRect InRect = FIntRect(0, 0, 0, 0));
 
 protected:
diff --git a/Engine/Source/Runtime/Engine/Public/TextureResource.h b/Engine/Source/Runtime/Engine/Public/TextureResource.h
--- a/Engine/Source/Runtime/Engine/Public/TextureResource.h
+++ b/Engine/Source/Runtime/Engine/Public/TextureResource.h
@@ -601,6 +601,7 @@ public:
 	* @param InRect - Rectangle of texels to copy.
 	* @return true if the read succeeded.
 	*/
+	UE_DEPRECATED(5.4, "Use FRenderTarget's ReadPixels, which is functionally equivalent")
 	ENGINE_API bool ReadPixels(TArray< FColor >& OutImageData, FReadSurfaceDataFlags InFlags, FIntRect InRect = FIntRect(0, 0, 0, 0));
 
 	/**
@@ -610,6 +611,7 @@ public:
 	* @param InRect - Rectangle of texels to copy.
 	* @return true if the read succeeded.
 	*/
+	UE_DEPRECATED(5.4, "Use FRenderTarget's ReadFloat16Pixels, which is functionally equivalent")
 	ENGINE_API bool ReadPixels(TArray<FFloat16Color>& OutImageData, FReadSurfaceDataFlags InFlags, FIntRect InRect = FIntRect(0, 0, 0, 0));
 
 protected:
diff --git a/Engine/Source/Runtime/Engine/Public/UnrealClient.h b/Engine/Source/Runtime/Engine/Public/UnrealClient.h
--- a/Engine/Source/Runtime/Engine/Public/UnrealClient.h
+++ b/Engine/Source/Runtime/Engine/Public/UnrealClient.h
@@ -82,55 +82,75 @@ public:
 	virtual bool HasToggleFreezeCommand() { return false; };
 
 	/**
-	* Reads the viewport's displayed pixels into a preallocated color buffer.
+	* Reads the render target's displayed pixels into a preallocated color buffer.
 	* @param OutImageData - RGBA8 values will be stored in this buffer
-	* @param InRect - source rect of the image to capture
+	* @param InFlags - Additional information about how to to read the surface data (cube face, slice index, etc.)
+	* @param InSrcRect - InSrcRect not specified means the whole rect
 	* @return True if the read succeeded.
 	*/
-	ENGINE_API bool ReadPixels(TArray< FColor >& OutImageData,FReadSurfaceDataFlags InFlags = FReadSurfaceDataFlags(RCM_UNorm, CubeFace_MAX), FIntRect InRect = FIntRect(0, 0, 0, 0) );
+	ENGINE_API virtual bool ReadPixels(TArray<FColor>& OutImageData, FReadSurfaceDataFlags InFlags = FReadSurfaceDataFlags(RCM_UNorm, CubeFace_MAX), FIntRect InSrcRect = FIntRect(0, 0, 0, 0));
 
 	/**
-	* Reads the viewport's displayed pixels into a preallocated color buffer.
+	* Reads the render target's displayed pixels into a preallocated color buffer.
 	* @param OutImageBytes - RGBA8 values will be stored in this buffer.  Buffer must be preallocated with the correct size!
-	* @param InSrcRect InSrcRect not specified means the whole rect
+	* * @param InFlags - Additional information about how to to read the surface data (cube face, slice index, etc.)
+	* @param InSrcRect - InSrcRect not specified means the whole rect
 	* @return True if the read succeeded.
 	*/
 	ENGINE_API bool ReadPixelsPtr(FColor* OutImageBytes, FReadSurfaceDataFlags InFlags = FReadSurfaceDataFlags(RCM_UNorm, CubeFace_MAX), FIntRect InSrcRect = FIntRect(0, 0, 0, 0));
 
 	/**
-	 * Reads the viewport's displayed pixels into the given color buffer.
-	 * @param OutputBuffer - RGBA16F values will be stored in this buffer
+	 * Reads the render target's displayed pixels into the given color buffer.
+	 * @param OutImageData - RGBA16F values will be stored in this buffer
 	 * @param CubeFace - optional cube face for when reading from a cube render target
 	 * @return True if the read succeeded.
 	 */
-	ENGINE_API bool ReadFloat16Pixels(TArray<FFloat16Color>& OutputBuffer,ECubeFace CubeFace=CubeFace_PosX);
+	UE_DEPRECATED(5.4, "Use the other ReadFloat16Pixels variant (ECubeFace can be set in FReadSurfaceDataFlags)")
+	ENGINE_API bool ReadFloat16Pixels(TArray<FFloat16Color>& OutImageData, ECubeFace CubeFace);
+	/**
+	 * Reads the render target's displayed pixels into the given color buffer.
+	 * @param OutImageData - RGBA16F values will be stored in this buffer
+	 * @param InFlags - Additional information about how to to read the surface data (cube face, slice index, etc.)
+	 * @param InSrcRect - InSrcRect not specified means the whole rect
+	 * @return True if the read succeeded.
+	 */
+	ENGINE_API virtual bool ReadFloat16Pixels(TArray<FFloat16Color>& OutImageData, FReadSurfaceDataFlags InFlags = FReadSurfaceDataFlags(RCM_UNorm, CubeFace_MAX), FIntRect InSrcRect = FIntRect(0, 0, 0, 0));
 
 	/**
-	 * Reads the viewport's displayed pixels into the given color buffer.
-	 * @param OutputBuffer - Linear color array to store the value
-	 * @param CubeFace - optional cube face for when reading from a cube render target
+	* Reads the render target's displayed pixels into a preallocated color buffer.
+	* @param OutImageBytes - RGBA16F values will be stored in this buffer.  Buffer must be preallocated with the correct size!
+	* @param InFlags - Additional information about how to to read the surface data (cube face, slice index, etc.)
+	* @param InSrcRect - InSrcRect not specified means the whole rect
+	* @return True if the read succeeded.
+	*/
+	ENGINE_API bool ReadFloat16PixelsPtr(FFloat16Color* OutImageBytes, FReadSurfaceDataFlags InFlags = FReadSurfaceDataFlags(RCM_UNorm, CubeFace_MAX), FIntRect InSrcRect = FIntRect(0, 0, 0, 0));
+
+	/**
+	 * Reads the render target's displayed pixels into the given color buffer.
+	 * @param OutImageData - Linear color array to store the value
+	 * @param InFlags - Additional information about how to to read the surface data (cube face, slice index, etc.)
+	 * @param InSrcRect - InSrcRect not specified means the whole rect
 	 * @return True if the read succeeded.
 	 */
-	ENGINE_API bool ReadLinearColorPixels(TArray<FLinearColor>& OutputBuffer, FReadSurfaceDataFlags InFlags = FReadSurfaceDataFlags(RCM_MinMax, CubeFace_MAX), FIntRect InRect = FIntRect(0, 0, 0, 0));
+	ENGINE_API virtual bool ReadLinearColorPixels(TArray<FLinearColor>& OutImageData, FReadSurfaceDataFlags InFlags = FReadSurfaceDataFlags(RCM_MinMax, CubeFace_MAX), FIntRect InSrcRect = FIntRect(0, 0, 0, 0));
 
-	ENGINE_API bool ReadLinearColorPixelsPtr(FLinearColor* OutImageBytes, FReadSurfaceDataFlags InFlags = FReadSurfaceDataFlags(RCM_MinMax, CubeFace_MAX), FIntRect InRect = FIntRect(0, 0, 0, 0));
+	/**
+	 * Reads the render target's displayed pixels into the given color buffer.
+	 * @param OutImageBytes - Linear color array will be stored in this buffer.  Buffer must be preallocated with the correct size!
+	 * @param InFlags - Additional information about how to to read the surface data (cube face, slice index, etc.)
+	 * @param InSrcRect - InSrcRect not specified means the whole rect
+	 * @return True if the read succeeded.
+	 */
+	ENGINE_API bool ReadLinearColorPixelsPtr(FLinearColor* OutImageBytes, FReadSurfaceDataFlags InFlags = FReadSurfaceDataFlags(RCM_MinMax, CubeFace_MAX), FIntRect InSrcRect = FIntRect(0, 0, 0, 0));
 
 	/**
-	 * Returns the GPU nodes on which to render this rendertarget.
+	 * Returns the GPU nodes on which to render this render target.
 	 **/
 	virtual FRHIGPUMask GetGPUMask(FRHICommandListImmediate& RHICmdList) const { return FRHIGPUMask::GPU0(); }
 
 protected:
 
 	FTextureRHIRef RenderTargetTextureRHI;
-
-	/**
-	 * Reads the viewport's displayed pixels into a preallocated color buffer.
-	 * @param OutImageBytes - RGBA16F values will be stored in this buffer.  Buffer must be preallocated with the correct size!
-	 * @param CubeFace - optional cube face for when reading from a cube render target
-	 * @return True if the read succeeded.
-	 */
-	bool ReadFloat16Pixels(FFloat16Color* OutImageBytes,ECubeFace CubeFace=CubeFace_PosX);
 };
 
 
diff --git a/Engine/Source/Runtime/RHI/Public/RHISurfaceDataConversion.h b/Engine/Source/Runtime/RHI/Public/RHISurfaceDataConversion.h
--- a/Engine/Source/Runtime/RHI/Public/RHISurfaceDataConversion.h
+++ b/Engine/Source/Runtime/RHI/Public/RHISurfaceDataConversion.h
@@ -394,6 +394,21 @@ inline void ConvertRawR8DataToFColor(uint32 Width, uint32 Height, uint8 *In, uin
 	}
 }
 
+inline void ConvertRawR8G8DataToFColor(uint32 Width, uint32 Height, uint8* In, uint32 SrcPitch, FColor* Out)
+{
+	for (uint32 Y = 0; Y < Height; Y++)
+	{
+		uint8* SrcPtr = (uint8*)(In + Y * SrcPitch);
+		FColor* DestPtr = Out + Y * Width;
+		for (uint32 X = 0; X < Width; X++)
+		{
+			*DestPtr = FColor(*SrcPtr, *(SrcPtr + 1), 0);
+			SrcPtr += 2;
+			++DestPtr;
+		}
+	}
+}
+
 inline void ConvertRawD32S8DataToFColor(uint32 Width, uint32 Height, uint8 *In, uint32 SrcPitch, FColor* Out, FReadSurfaceDataFlags InFlags)
 {
 	bool bLinearToGamma = InFlags.GetLinearToGamma();
diff --git a/Engine/Source/Runtime/RHI/Public/RHITypes.h b/Engine/Source/Runtime/RHI/Public/RHITypes.h
--- a/Engine/Source/Runtime/RHI/Public/RHITypes.h
+++ b/Engine/Source/Runtime/RHI/Public/RHITypes.h
@@ -24,7 +24,6 @@ public:
 
 	ECubeFace GetCubeFace() const
 	{
-		checkSlow(CubeFace <= CubeFace_NegZ);
 		return CubeFace;
 	}
 
diff --git a/Engine/Source/Runtime/RenderCore/Public/RenderUtils.h b/Engine/Source/Runtime/RenderCore/Public/RenderUtils.h
--- a/Engine/Source/Runtime/RenderCore/Public/RenderUtils.h
+++ b/Engine/Source/Runtime/RenderCore/Public/RenderUtils.h
@@ -128,8 +128,8 @@ struct FPackedPosition
 };
 
 
-/** Flags that control ConstructTexture2D */
-enum EConstructTextureFlags
+/** Flags that control ConstructTexture(2D/2DArray/Volume/etc.) */
+enum EConstructTextureFlags : uint32
 {
 	/** Compress RGBA8 to DXT */
 	CTF_Compress =				0x01,
diff --git a/Engine/Source/Runtime/VulkanRHI/Private/VulkanRenderTarget.cpp b/Engine/Source/Runtime/VulkanRHI/Private/VulkanRenderTarget.cpp
--- a/Engine/Source/Runtime/VulkanRHI/Private/VulkanRenderTarget.cpp
+++ b/Engine/Source/Runtime/VulkanRHI/Private/VulkanRenderTarget.cpp
@@ -107,6 +107,10 @@ static void ConvertRawDataToFColor(VkFormat VulkanFormat, uint32 DestWidth, uint
 		ConvertRawR8DataToFColor(DestWidth, DestHeight, In, SrcPitch, Dest);
 		break;
 
+	case VK_FORMAT_R8G8_UNORM:
+		ConvertRawR8G8DataToFColor(DestWidth, DestHeight, In, SrcPitch, Dest);
+		break;
+
 	case VK_FORMAT_R16_UNORM:
 		ConvertRawR16DataToFColor(DestWidth, DestHeight, In, SrcPitch, Dest);
 		break;
diff --git a/Engine/Source/Runtime/Windows/D3D11RHI/Private/D3D11RenderTarget.cpp b/Engine/Source/Runtime/Windows/D3D11RHI/Private/D3D11RenderTarget.cpp
--- a/Engine/Source/Runtime/Windows/D3D11RHI/Private/D3D11RenderTarget.cpp
+++ b/Engine/Source/Runtime/Windows/D3D11RHI/Private/D3D11RenderTarget.cpp
@@ -336,11 +336,12 @@ TRefCountPtr<ID3D11Texture2D> FD3D11DynamicRHI::GetStagingTexture(FRHITexture* T
 	if( SourceDesc.MiscFlags == D3D11_RESOURCE_MISC_TEXTURECUBE )
 	{
 		uint32 D3DFace = GetD3D11CubeFace(InFlags.GetCubeFace());
-		Subresource = D3D11CalcSubresource(InFlags.GetMip(),D3DFace,TextureRHI->GetNumMips());
+		Subresource = D3D11CalcSubresource(InFlags.GetMip(), InFlags.GetArrayIndex() * 6 + D3DFace, TextureRHI->GetNumMips());
 	}
 	else
 	{
-		Subresource = D3D11CalcSubresource(InFlags.GetMip(), 0, TextureRHI->GetNumMips());
+		const bool bIsTextureArray = Texture->GetDesc().IsTextureArray();
+		Subresource = D3D11CalcSubresource(InFlags.GetMip(), bIsTextureArray ? InFlags.GetArrayIndex() : 0, TextureRHI->GetNumMips());
 	}
 
 	D3D11_BOX* RectPtr = NULL; // API prefers NULL for entire texture.
@@ -482,6 +483,9 @@ static void ConvertDXGIToFColor(DXGI_FORMAT Format, uint32 Width, uint32 Height,
 		case DXGI_FORMAT_R8_UNORM:
 			ConvertRawR8DataToFColor(Width, Height, In, SrcPitch, Out);
 			break;
+		case DXGI_FORMAT_R8G8_UNORM:
+			ConvertRawR8G8DataToFColor(Width, Height, In, SrcPitch, Out);
+			break;
 		default:
 			checkf(0, TEXT("Unknown surface format!"));
 			break;
@@ -598,11 +602,16 @@ void FD3D11DynamicRHI::ReadSurfaceDataMSAARaw(FRHITexture* TextureRHI,FIntRect I
 	VERIFYD3D11RESULT_EX(Direct3DDevice->CreateTexture2D(&StagingDesc,NULL,StagingTexture2D.GetInitReference()), Direct3DDevice);
 
 	// Determine the subresource index for cubemaps.
-	uint32 Subresource = InFlags.GetMip();
-	if( TextureDesc.MiscFlags == D3D11_RESOURCE_MISC_TEXTURECUBE )
+	uint32 Subresource = 0;
+	if (TextureDesc.MiscFlags == D3D11_RESOURCE_MISC_TEXTURECUBE)
 	{
 		uint32 D3DFace = GetD3D11CubeFace(InFlags.GetCubeFace());
-		Subresource = D3D11CalcSubresource(0,D3DFace,1);
+		Subresource = D3D11CalcSubresource(InFlags.GetMip(), InFlags.GetArrayIndex() * 6 + D3DFace, TextureRHI->GetNumMips());
+	}
+	else
+	{
+		const bool bIsTextureArray = Texture->GetDesc().IsTextureArray();
+		Subresource = D3D11CalcSubresource(InFlags.GetMip(), bIsTextureArray ? InFlags.GetArrayIndex() : 0, TextureRHI->GetNumMips());
 	}
 	
 	// Allocate the output buffer.
