Clean up Substrate material complexity to remove duplicated and non-consistent logic.

diff --git a/Engine/Shaders/Public/ShaderVersion.ush b/Engine/Shaders/Public/ShaderVersion.ush
--- a/Engine/Shaders/Public/ShaderVersion.ush
+++ b/Engine/Shaders/Public/ShaderVersion.ush
@@ -2,4 +2,4 @@
 // Each time the console command is executed it generates a new GUID. As this file is included
 // in Platform.ush (which should be included in any shader) it allows to invalidate the shader DDC.
 // 
-#pragma message("UESHADERMETADATA_VERSION 2B90878D-98D5-4E42-933B-178267A501C0")
+#pragma message("UESHADERMETADATA_VERSION 4037B637-C512-4BA2-B763-29C04DF6E818")
diff --git a/Engine/Shaders/Shared/SubstrateDefinitions.h b/Engine/Shaders/Shared/SubstrateDefinitions.h
--- a/Engine/Shaders/Shared/SubstrateDefinitions.h
+++ b/Engine/Shaders/Shared/SubstrateDefinitions.h
@@ -71,6 +71,11 @@
 #define SUBSTRATE_TILE_TYPE_DECAL_COMPLEX				8
 #define SUBSTRATE_TILE_TYPE_COUNT						9
 
+#define SUBSTRATE_MATERIAL_TYPE_SIMPLE					0
+#define SUBSTRATE_MATERIAL_TYPE_SINGLE					1
+#define SUBSTRATE_MATERIAL_TYPE_COMPLEX					2
+#define SUBSTRATE_MATERIAL_TYPE_COMPLEX_SPECIAL			3
+
 // sizeof(FRHIDrawIndirectParameters) = 4 uints = 16 bytes
 #define GetSubstrateTileTypeDrawIndirectArgOffset_Byte(x)  (x * 16)
 #define GetSubstrateTileTypeDrawIndirectArgOffset_DWord(x) (x * 4)
diff --git a/Engine/Source/Editor/MaterialEditor/Private/SMaterialEditorSubstrateWidget.cpp b/Engine/Source/Editor/MaterialEditor/Private/SMaterialEditorSubstrateWidget.cpp
--- a/Engine/Source/Editor/MaterialEditor/Private/SMaterialEditorSubstrateWidget.cpp
+++ b/Engine/Source/Editor/MaterialEditor/Private/SMaterialEditorSubstrateWidget.cpp
@@ -377,6 +377,18 @@ SMaterialEditorSubstrateWidget::~SMaterialEditorSubstrateWidget()
 {
 }
 
+static FString SubstrateMaterialTypeToString(uint8 InMaterialType)
+{
+	switch (InMaterialType)
+	{
+		case SUBSTRATE_MATERIAL_TYPE_SIMPLE:  			return FString::Printf(TEXT("SIMPLE  (Diffuse, albedo, roughness)\r\n"));
+		case SUBSTRATE_MATERIAL_TYPE_SINGLE:  			return FString::Printf(TEXT("SINGLE  (BSDF all features except anisotropy)\r\n"));
+		case SUBSTRATE_MATERIAL_TYPE_COMPLEX:  			return FString::Printf(TEXT("COMPLEX (Anisotropy, multi-slabs)\r\n"));
+		case SUBSTRATE_MATERIAL_TYPE_COMPLEX_SPECIAL:  	return FString::Printf(TEXT("COMPLEX SPECIAL (Glints, specular profile)\r\n"));
+		default: 										return FString::Printf(TEXT("UNKOWN => ERROR!\r\n"));
+	}
+}
+
 void SMaterialEditorSubstrateWidget::Tick(const FGeometry& AllottedGeometry, const double InCurrentTime, const float InDeltaTime)
 {
 	if (!bUpdateRequested || !Substrate::IsSubstrateEnabled())
@@ -420,25 +432,7 @@ void SMaterialEditorSubstrateWidget::Tick(const FGeometry& AllottedGeometry, con
 				}
 				MaterialDescription += FString::Printf(TEXT("BSDF Count	                  = %i\r\n"), CompilationOutput.SubstrateBSDFCount);
 				MaterialDescription += FString::Printf(TEXT("Local bases Count            = %i\r\n"), CompilationOutput.SharedLocalBasesCount);
-
-				switch (CompilationOutput.SubstrateMaterialType)
-				{
-				case 0:
-					MaterialDescription += FString::Printf(TEXT("Material complexity          = SIMPLE (diffuse, albedo, roughness)\r\n"));
-					break;
-				case 1:
-					MaterialDescription += FString::Printf(TEXT("Material complexity          = SINGLE (BSDF all features except anisotropy)\r\n"));
-					break;
-				case 2:
-					MaterialDescription += FString::Printf(TEXT("Material complexity          = COMPLEX\r\n"));
-					break;
-				case 3:
-					MaterialDescription += FString::Printf(TEXT("Material complexity          = COMPLEX SPECIAL\r\n"));
-					break;
-				default:
-					MaterialDescription += FString::Printf(TEXT("Material complexity          = UNKOWN => ERROR!\r\n"));
-				}
-
+				MaterialDescription += FString::Printf(TEXT("Material complexity          = %s\r\n"), *SubstrateMaterialTypeToString(CompilationOutput.SubstrateMaterialType));
 				MaterialDescription += FString::Printf(TEXT("Root Node Is Thin            = %i\r\n"), CompilationOutput.bIsThin);
 
 				//if (CompilationOutput.SubstrateBSDFCount == 1)
diff --git a/Engine/Source/Runtime/Engine/Classes/Materials/MaterialExpressionSubstrate.h b/Engine/Source/Runtime/Engine/Classes/Materials/MaterialExpressionSubstrate.h
--- a/Engine/Source/Runtime/Engine/Classes/Materials/MaterialExpressionSubstrate.h
+++ b/Engine/Source/Runtime/Engine/Classes/Materials/MaterialExpressionSubstrate.h
@@ -8,6 +8,7 @@
 #include "MaterialExpressionSubstrate.generated.h"
 
 class FMaterialCompiler;
+struct FSubstrateMaterialComplexity;
 
 /**
  * Compile a special blend function for Substrate when blending material attribute
@@ -342,14 +343,7 @@ class UMaterialExpressionSubstrateSlabBSDF : public UMaterialExpressionSubstrate
 	bool HasGlint() const;
 	bool HasSpecularProfile() const;
 
-	struct FComplexity
-	{
-		bool bSubstrateMaterialIsComplexSpecial;
-		bool bSubstrateMaterialIsComplex;
-		bool bSubstrateMaterialIsSingle;
-		// If all the above are false, complexity is Simple
-	};
-	FComplexity GetComplexity() const;
+	FSubstrateMaterialComplexity GetComplexity() const;
 #endif
 	//~ End UMaterialExpression Interface
 };
diff --git a/Engine/Source/Runtime/Engine/Private/Materials/HLSLMaterialTranslator.cpp b/Engine/Source/Runtime/Engine/Private/Materials/HLSLMaterialTranslator.cpp
--- a/Engine/Source/Runtime/Engine/Private/Materials/HLSLMaterialTranslator.cpp
+++ b/Engine/Source/Runtime/Engine/Private/Materials/HLSLMaterialTranslator.cpp
@@ -342,8 +342,7 @@ void FHLSLMaterialTranslator::FSubstrateCompilationContext::Initialise()
 	SubstrateMaterialExpressionToOperatorIndex.Reserve(SUBSTRATE_MAX_COMPILER_REGISTERED_OPERATOR_COUNT);
 	SubstrateMaterialBSDFCount = 0;
 	SubstrateMaterialRequestedSizeByte = 0;
-	bSubstrateMaterialIsSimple = false;
-	bSubstrateMaterialIsSingle = false;
+	SubstrateMaterialComplexity.Reset();
 	bSubstrateMaterialIsUnlitNode = false;
 	bSubstrateWritesEmissive = false;
 	bSubstrateWritesAmbientOcclusion = false;
@@ -2367,10 +2366,10 @@ void FHLSLMaterialTranslator::GetMaterialEnvironment(EShaderPlatform InPlatform,
 			// For now, the fully simplified mode is used for Lumen or anything else supported inlined evaluation. The export is only valid for the default case.
 			// SUBSTRATE_TODO: generate an export for the different context (need to generate two export functions: the default one and the FullSimplification one)
 			FSubstrateCompilationContext& SubstrateCtx = SubstrateCompilationContext[ESubstrateCompilationContext::SCC_Default];
-			OutEnvironment.SetDefine(TEXT("SUBSTRATE_SINGLEPATH"), SubstrateCtx.bSubstrateMaterialIsSingle ? TEXT("1") : TEXT("0"));
-			OutEnvironment.SetDefine(TEXT("SUBSTRATE_FASTPATH"), SubstrateCtx.bSubstrateMaterialIsSingle ? TEXT("0") : (SubstrateCtx.bSubstrateMaterialIsSimple ? TEXT("1") : TEXT("0")));
+			OutEnvironment.SetDefine(TEXT("SUBSTRATE_SINGLEPATH"), SubstrateCtx.SubstrateMaterialComplexity.IsSingle() ? TEXT("1") : TEXT("0"));
+			OutEnvironment.SetDefine(TEXT("SUBSTRATE_FASTPATH"), SubstrateCtx.SubstrateMaterialComplexity.IsSimple() ? TEXT("1") : TEXT("0"));
+			OutEnvironment.SetDefine(TEXT("SUBSTRATE_COMPLEXSPECIALPATH"), SubstrateCtx.SubstrateMaterialComplexity.IsComplexSpecial() ? TEXT("1") : TEXT("0"));
 			OutEnvironment.SetDefine(TEXT("SUBSTRATE_CLAMPED_BSDF_COUNT"), SubstrateCtx.SubstrateMaterialBSDFCount);
-			OutEnvironment.SetDefine(TEXT("SUBSTRATE_COMPLEXSPECIALPATH"), MaterialCompilationOutput.SubstrateMaterialCompilationOutput.bUsesComplexSpecialRenderPath ? TEXT("1") : TEXT("0"));
 		}
 
 
@@ -2411,7 +2410,6 @@ void FHLSLMaterialTranslator::GetMaterialEnvironment(EShaderPlatform InPlatform,
 					return TEXT("ERROR");
 				};
 				FString SubstrateCompilationContextName = GetSubstrateCompilationContextName(SubstrateCompilationContextIndex);
-
 				SubstrateMaterialContextDescription += FString::Printf(TEXT("----- SUBSTRATE - %s -----\r\n"), *SubstrateCompilationContextName);
 				SubstrateMaterialContextDescription += FString::Printf(TEXT("SubstrateCompilationInfo -\r\n"));
 				SubstrateMaterialContextDescription += FString::Printf(TEXT(" - Byte Per Pixel Budget                           %u\r\n"), SubstrateBytePerPixel);
@@ -2420,7 +2418,7 @@ void FHLSLMaterialTranslator::GetMaterialEnvironment(EShaderPlatform InPlatform,
 				SubstrateMaterialContextDescription += FString::Printf(TEXT(" - Requested Byte Size after simplification        %u (%d UINT32)\r\n"), SubstrateCtx.SubstrateMaterialRequestedSizeByte, SubstrateCtx.SubstrateMaterialRequestedSizeByte / 4);
 				SubstrateMaterialContextDescription += FString::Printf(TEXT(" - Requested Closure Count before simplification   %u\r\n"), SubstrateSimplificationStatus.OriginalRequestedClosureCount);
 				SubstrateMaterialContextDescription += FString::Printf(TEXT(" - Requested Closure Count after simplification    %u\r\n"), SubstrateCtx.SubstrateMaterialClosureCount);
-				SubstrateMaterialContextDescription += FString::Printf(TEXT(" - Material complexity                             %s\r\n"), SubstrateCtx.bSubstrateMaterialIsSingle ? TEXT("SINGLE") : (SubstrateCtx.bSubstrateMaterialIsSimple ? TEXT("SIMPLE") : TEXT("COMPLEX")));
+				SubstrateMaterialContextDescription += FString::Printf(TEXT(" - Material complexity                             %s\r\n"), *FSubstrateMaterialComplexity::ToString(SubstrateCtx.SubstrateMaterialComplexity.SubstrateMaterialType(), true /* Upper case */));
 				SubstrateMaterialContextDescription += FString::Printf(TEXT(" - BSDF Count                                      %i\r\n"), SubstrateCtx.SubstrateMaterialBSDFCount); // REMOVE?
 				if (RequestedSharedLocalBasesCount > SUBSTRATE_MAX_SHAREDLOCALBASES_REGISTERS)
 				{
@@ -7375,8 +7373,8 @@ int32 FHLSLMaterialTranslator::SceneTextureLookup(int32 ViewportUV, uint32 InSce
 	// for optimization purpose (simple/single/complex). To avoid compiling out single or complex unpacking paths (due to defines set by analyzing 
 	// the current shader, vs. scene texture pixels), we force Simple/Single versions to be disabled
 	FSubstrateCompilationContext& SubstrateCtx = SubstrateCompilationContext[CurrentSubstrateCompilationContext];
-	SubstrateCtx.bSubstrateMaterialIsSimple = false;
-	SubstrateCtx.bSubstrateMaterialIsSingle = false;
+	SubstrateCtx.SubstrateMaterialComplexity.bIsSimple = false;
+	SubstrateCtx.SubstrateMaterialComplexity.bIsSingle = false;
 	
 	if (SceneTextureId == PPI_PostProcessInput0 && Material->GetMaterialDomain() == MD_PostProcess && Material->GetBlendableLocation() != BL_AfterTonemapping)
 	{
@@ -12026,11 +12024,11 @@ bool FHLSLMaterialTranslator::FSubstrateCompilationContext::SubstrateGenerateDer
 			SubstrateMaterialRequestedSizeByte = 0;
 
 			// 1. Evaluate simple/single BSDF
-			bSubstrateMaterialIsSimple = SubstrateMaterialBSDFCount == 1;
-			bSubstrateMaterialIsSingle = SubstrateMaterialBSDFCount == 1;
+			SubstrateMaterialComplexity.bIsSimple = SubstrateMaterialBSDFCount == 1;
+			SubstrateMaterialComplexity.bIsSingle = SubstrateMaterialBSDFCount == 1;
+			SubstrateMaterialComplexity.bIsComplexSpecial = false;
 			bool bIsFastWaterPath = false;
 			bool bCustomEncoding = false;
-			bool bUsesComplexSpecialRenderPath = false;
 			for (const auto& It : SubstrateMaterialExpressionRegisteredOperators)
 			{
 				if (It.IsDiscarded())
@@ -12049,29 +12047,29 @@ bool FHLSLMaterialTranslator::FSubstrateCompilationContext::SubstrateGenerateDer
 					{
 					case SUBSTRATE_BSDF_TYPE_SLAB:
 					{
-						bSubstrateMaterialIsSimple = bSubstrateMaterialIsSimple && !bMayHaveColoredWeight && !It.bBSDFHasAnisotropy && !It.bBSDFHasEdgeColor && !It.bBSDFHasFuzz && !It.bBSDFHasSecondRoughnessOrSimpleClearCoat && !It.bBSDFHasMFPPluggedIn && !It.bBSDFHasSSS && !It.bBSDFHasGlint && !It.bBSDFHasSpecularProfile;
-						bSubstrateMaterialIsSingle = bSubstrateMaterialIsSingle && !bMayHaveColoredWeight && !It.bBSDFHasAnisotropy && !It.bBSDFHasGlint && !It.bBSDFHasSpecularProfile;
-						bUsesComplexSpecialRenderPath |= It.bBSDFHasGlint || It.bBSDFHasSpecularProfile;
+						SubstrateMaterialComplexity.bIsSimple = SubstrateMaterialComplexity.bIsSimple && !bMayHaveColoredWeight && !It.bBSDFHasAnisotropy && !It.bBSDFHasEdgeColor && !It.bBSDFHasFuzz && !It.bBSDFHasSecondRoughnessOrSimpleClearCoat && !It.bBSDFHasMFPPluggedIn && !It.bBSDFHasSSS && !It.bBSDFHasGlint && !It.bBSDFHasSpecularProfile;
+						SubstrateMaterialComplexity.bIsSingle = SubstrateMaterialComplexity.bIsSingle && !bMayHaveColoredWeight && !It.bBSDFHasAnisotropy && !It.bBSDFHasGlint && !It.bBSDFHasSpecularProfile;
+						SubstrateMaterialComplexity.bIsComplexSpecial |= It.bBSDFHasGlint || It.bBSDFHasSpecularProfile;
 						break;
 					}
 					case SUBSTRATE_BSDF_TYPE_HAIR:
 					{
-						bSubstrateMaterialIsSimple = false;
-						bSubstrateMaterialIsSingle = false;
+						SubstrateMaterialComplexity.bIsSimple = false;
+						SubstrateMaterialComplexity.bIsSingle = false;
 						bCustomEncoding			= true;
 						break;
 					}
 					case SUBSTRATE_BSDF_TYPE_EYE:
 					{
-						bSubstrateMaterialIsSimple = false;
-						bSubstrateMaterialIsSingle = false;
+						SubstrateMaterialComplexity.bIsSimple = false;
+						SubstrateMaterialComplexity.bIsSingle = false;
 						bCustomEncoding			= true;
 						break;
 					}
 					case SUBSTRATE_BSDF_TYPE_SINGLELAYERWATER:
 					{
-						bSubstrateMaterialIsSimple = false;
-						bSubstrateMaterialIsSingle = false;
+						SubstrateMaterialComplexity.bIsSimple = false;
+						SubstrateMaterialComplexity.bIsSingle = false;
 						bIsFastWaterPath		= true;
 						break;
 					}
@@ -12081,17 +12079,17 @@ bool FHLSLMaterialTranslator::FSubstrateCompilationContext::SubstrateGenerateDer
 				case SUBSTRATE_OPERATOR_WEIGHT:
 				{
 					// If a BSDF modified by a weight operator, its weight will be < 1.0f, and it won't be a single material anymore
-					bSubstrateMaterialIsSimple = false;
-					bSubstrateMaterialIsSingle = false;
+					SubstrateMaterialComplexity.bIsSimple = false;
+					SubstrateMaterialComplexity.bIsSingle = false;
 					break;
 				}
 				}
 			}
-			bSubstrateMaterialIsSingle = bSubstrateMaterialIsSingle && !bSubstrateMaterialIsSimple;
+			SubstrateMaterialComplexity.bIsSingle = SubstrateMaterialComplexity.bIsSingle && !SubstrateMaterialComplexity.bIsSimple;
 
 			// 2. Header
 
-			if (!bSubstrateMaterialIsSimple && !bSubstrateMaterialIsSingle && !bCustomEncoding && !bIsFastWaterPath) // header written later, 
+			if (!SubstrateMaterialComplexity.bIsSimple && !SubstrateMaterialComplexity.bIsSingle && !bCustomEncoding && !bIsFastWaterPath) // header written later, 
 			{
 				// Packed Header
 				SubstrateMaterialRequestedSizeByte += UintByteSize;
@@ -12131,7 +12129,7 @@ bool FHLSLMaterialTranslator::FSubstrateCompilationContext::SubstrateGenerateDer
 					// From the compiler side, we can only assume the top layer has gray scale luminance weight.
 					const bool bMayHaveColoredWeight = !It.bIsTop;
 
-					if (bSubstrateMaterialIsSimple)
+					if (SubstrateMaterialComplexity.bIsSimple)
 					{
 						// Header
 						SubstrateMaterialRequestedSizeByte += UintByteSize;
@@ -12139,7 +12137,7 @@ bool FHLSLMaterialTranslator::FSubstrateCompilationContext::SubstrateGenerateDer
 						SubstrateMaterialRequestedSizeByte += UintByteSize;
 						break; // Stop here
 					}
-					else if (bSubstrateMaterialIsSingle)
+					else if (SubstrateMaterialComplexity.bIsSingle)
 					{
 						// Header
 						SubstrateMaterialRequestedSizeByte += UintByteSize;
@@ -12278,10 +12276,9 @@ bool FHLSLMaterialTranslator::FSubstrateCompilationContext::SubstrateGenerateDer
 			if (CompilationContextIndex == ESubstrateCompilationContext::SCC_Default)
 			{
 				// Only write those data for the default material
-				Compiler->MaterialCompilationOutput.SubstrateMaterialCompilationOutput.SubstrateMaterialType = bSubstrateMaterialIsSimple ? 0 : (bSubstrateMaterialIsSingle ? 1 : (bUsesComplexSpecialRenderPath ? 3 : 2));
+				Compiler->MaterialCompilationOutput.SubstrateMaterialCompilationOutput.SubstrateMaterialType = SubstrateMaterialComplexity.SubstrateMaterialType();
 				Compiler->MaterialCompilationOutput.SubstrateMaterialCompilationOutput.SubstrateBSDFCount = SubstrateMaterialBSDFCount;
 				Compiler->MaterialCompilationOutput.SubstrateMaterialCompilationOutput.SubstrateUintPerPixel = uint8(FMath::Clamp(RequestedSizeInUint, 0u, 0xFF));
-				Compiler->MaterialCompilationOutput.SubstrateMaterialCompilationOutput.bUsesComplexSpecialRenderPath = bUsesComplexSpecialRenderPath;
 
 #if WITH_EDITOR
 				Compiler->MaterialCompilationOutput.SubstrateMaterialCompilationOutput.SharedLocalBasesCount = 0; // FinalUsedSharedLocalBasesCount is not valid yet
diff --git a/Engine/Source/Runtime/Engine/Private/Materials/HLSLMaterialTranslator.h b/Engine/Source/Runtime/Engine/Private/Materials/HLSLMaterialTranslator.h
--- a/Engine/Source/Runtime/Engine/Private/Materials/HLSLMaterialTranslator.h
+++ b/Engine/Source/Runtime/Engine/Private/Materials/HLSLMaterialTranslator.h
@@ -495,8 +495,7 @@ protected:
 		uint32 SubstrateMaterialBSDFCount; // Also acts as requested 
 		uint32 SubstrateMaterialRequestedSizeByte;
 		uint32 SubstrateMaterialClosureCount;
-		bool bSubstrateMaterialIsSimple;
-		bool bSubstrateMaterialIsSingle;
+		FSubstrateMaterialComplexity SubstrateMaterialComplexity;
 		bool bSubstrateMaterialIsUnlitNode;
 
 		bool bSubstrateWritesEmissive;
diff --git a/Engine/Source/Runtime/Engine/Private/Materials/MaterialExpressions.cpp b/Engine/Source/Runtime/Engine/Private/Materials/MaterialExpressions.cpp
--- a/Engine/Source/Runtime/Engine/Private/Materials/MaterialExpressions.cpp
+++ b/Engine/Source/Runtime/Engine/Private/Materials/MaterialExpressions.cpp
@@ -23979,53 +23979,35 @@ int32 UMaterialExpressionSubstrateSlabBSDF::Compile(class FMaterialCompiler* Com
 	return OutputCodeChunk;
 }
 
-UMaterialExpressionSubstrateSlabBSDF::FComplexity UMaterialExpressionSubstrateSlabBSDF::GetComplexity() const
+FSubstrateMaterialComplexity UMaterialExpressionSubstrateSlabBSDF::GetComplexity() const
 {
-	UMaterialExpressionSubstrateSlabBSDF::FComplexity Complexity;
-	Complexity.bSubstrateMaterialIsComplexSpecial = HasGlint() || HasSpecularProfile();
-	Complexity.bSubstrateMaterialIsComplex = HasAnisotropy();
-	Complexity.bSubstrateMaterialIsSingle = HasEdgeColor() || HasFuzz() || HasSecondRoughness() || HasMFPPluggedIn() || HasSSS();
+	FSubstrateMaterialComplexity Out;
+	Out.bIsComplexSpecial = HasGlint() || HasSpecularProfile();
+	Out.bIsSingle 		  = !HasAnisotropy() && (HasEdgeColor() || HasFuzz() || HasSecondRoughness() || HasMFPPluggedIn() || HasSSS());
+	Out.bIsSimple 		  = !HasAnisotropy();
 
-	// MAsk out to only have a single possibility
-	Complexity.bSubstrateMaterialIsComplex &= !Complexity.bSubstrateMaterialIsComplexSpecial;
-	Complexity.bSubstrateMaterialIsSingle &= !Complexity.bSubstrateMaterialIsComplexSpecial && !Complexity.bSubstrateMaterialIsComplex;
+	// Mask out to only have a single possibility
+	Out.bIsSingle &= !Out.bIsComplexSpecial;
+	Out.bIsSimple &= !Out.bIsComplexSpecial && !Out.bIsSingle;
 
-	return Complexity;
-}
-
-static FString GetSlabComplexityString(UMaterialExpressionSubstrateSlabBSDF::FComplexity Complexity)
-{
-	FString ComplexityString = TEXT("Simple");
-	if (Complexity.bSubstrateMaterialIsComplexSpecial)
-	{
-		ComplexityString = TEXT("ComplexSpecial");
-	}
-	else if (Complexity.bSubstrateMaterialIsComplex)
-	{
-		ComplexityString = TEXT("Complex");
-	}
-	else if (Complexity.bSubstrateMaterialIsSingle)
-	{
-		ComplexityString = TEXT("Single");
-	}
-	return ComplexityString;
+	return Out;
 }
 
 void UMaterialExpressionSubstrateSlabBSDF::GetCaption(TArray<FString>& OutCaptions) const
 {
 	// The node complexity is manually maintained to match FSubstrateCompilationContext::SubstrateGenerateDerivedMaterialOperatorData and shaders.
-	OutCaptions.Add(TEXT("Substrate Slab BSDF - ") + GetSlabComplexityString(GetComplexity()));
+	OutCaptions.Add(TEXT("Substrate Slab BSDF - ") + FSubstrateMaterialComplexity::ToString(GetComplexity().SubstrateMaterialType()));
 }
 
 void UMaterialExpressionSubstrateSlabBSDF::GetExpressionToolTip(TArray<FString>& OutToolTip)
 {
 	OutToolTip.Add(TEXT("Substrate Slab BSDF"));
-	OutToolTip.Add(TEXT("Complexity = ") + GetSlabComplexityString(GetComplexity()));
+	OutToolTip.Add(TEXT("Complexity = ") + FSubstrateMaterialComplexity::ToString(GetComplexity().SubstrateMaterialType()));
 	OutToolTip.Add(TEXT("The complexity represents the cost of the shading path (Lighting, Lumen, SSS) the material will follow:"));
 	OutToolTip.Add(TEXT(" - Simple means the Slab only relies on Diffuse, F0 and Roughness. It will follow a fast shading path."));
 	OutToolTip.Add(TEXT(" - Single means the Slab uses more features such as F90, Fuzz, Second Roughness, MFP or SSS. It will follow a more expenssive shading path."));
 	OutToolTip.Add(TEXT(" - Complex means a Slab uses anisotropic lighting, with any of the previous features."));
-	OutToolTip.Add(TEXT(" - ComplexSpecial means the Slab is using more advanced features such as glints or specular LUT. This is the most expenssive shading path."));
+	OutToolTip.Add(TEXT(" - Complex Special means the Slab is using more advanced features such as glints or specular LUT. This is the most expenssive shading path."));
 }
 
 uint32 UMaterialExpressionSubstrateSlabBSDF::GetOutputType(int32 OutputIndex)
diff --git a/Engine/Source/Runtime/Engine/Private/Materials/MaterialInterface.cpp b/Engine/Source/Runtime/Engine/Private/Materials/MaterialInterface.cpp
--- a/Engine/Source/Runtime/Engine/Private/Materials/MaterialInterface.cpp
+++ b/Engine/Source/Runtime/Engine/Private/Materials/MaterialInterface.cpp
@@ -483,7 +483,6 @@ FMaterialRelevance UMaterialInterface::GetRelevance_Internal(const UMaterial* Ma
 		const uint8 SubstrateBSDFCount = FMath::Max(MaterialResource->MaterialGetSubstrateBSDFCount_GameThread(), uint8(1u));
 		const uint8 SubstrateBSDFCountMask = 1u << uint8(FMath::Min(SubstrateBSDFCount - 1, 8));
 		const uint8 SubstrateUintPerPixel = FMath::Max(MaterialResource->MaterialGetSubstrateUintPerPixel_GameThread(), uint8(1u));
-		const uint8 bUsesComplexSpecialRenderPath = MaterialResource->MaterialGetSubstrateUsesComplexSpecialRenderPath_GameThread();
 
 		MaterialRelevance.bOpaque = !bIsTranslucent;
 		MaterialRelevance.bMasked = IsMasked();
@@ -511,7 +510,7 @@ FMaterialRelevance UMaterialInterface::GetRelevance_Internal(const UMaterial* Ma
 		MaterialRelevance.bUsesAnisotropy = bUsesAnisotropy;
 		MaterialRelevance.SubstrateBSDFCountMask = SubstrateBSDFCountMask;
 		MaterialRelevance.SubstrateUintPerPixel = SubstrateUintPerPixel;
-		MaterialRelevance.bUsesComplexSpecialRenderPath = bUsesComplexSpecialRenderPath;
+		MaterialRelevance.bUsesComplexSpecialRenderPath = MaterialResource->MaterialGetSubstrateUsesComplexSpecialRenderPath_GameThread();
 
 		return MaterialRelevance;
 	}
diff --git a/Engine/Source/Runtime/Engine/Private/Materials/SubstrateMaterial.h b/Engine/Source/Runtime/Engine/Private/Materials/SubstrateMaterial.h
--- a/Engine/Source/Runtime/Engine/Private/Materials/SubstrateMaterial.h
+++ b/Engine/Source/Runtime/Engine/Private/Materials/SubstrateMaterial.h
@@ -5,7 +5,7 @@
 #include "CoreMinimal.h"
 #include "Containers/Map.h"
 #include "MaterialCompiler.h"
-
+#include "SubstrateDefinitions.h"
 
 class FMaterialCompiler;
 
@@ -14,3 +14,39 @@ FString GetSubstrateBSDFName(uint8 BSDFType);
 FSubstrateRegisteredSharedLocalBasis SubstrateCompilationInfoCreateNullSharedLocalBasis();
 FSubstrateRegisteredSharedLocalBasis SubstrateCompilationInfoCreateSharedLocalBasis(FMaterialCompiler* Compiler, int32 NormalCodeChunk, int32 TangentCodeChunk = INDEX_NONE);
 
+struct FSubstrateMaterialComplexity
+{
+	bool bIsComplexSpecial;
+	bool bIsSingle;
+	bool bIsSimple;
+	// If all the above are false, complexity is Complex
+
+	void Reset()
+	{
+		bIsComplexSpecial = false;
+		bIsSingle = false;
+		bIsSimple = false;
+	}
+
+	uint8 SubstrateMaterialType() const
+	{
+		return bIsSimple ? SUBSTRATE_MATERIAL_TYPE_SIMPLE : (bIsSingle ? SUBSTRATE_MATERIAL_TYPE_SINGLE : (bIsComplexSpecial ? SUBSTRATE_MATERIAL_TYPE_COMPLEX_SPECIAL : SUBSTRATE_MATERIAL_TYPE_COMPLEX));
+	}
+
+	bool IsSimple() const			{ return SubstrateMaterialType() == SUBSTRATE_MATERIAL_TYPE_SIMPLE; }
+	bool IsSingle() const			{ return SubstrateMaterialType() == SUBSTRATE_MATERIAL_TYPE_SINGLE; }
+	bool IsComplex() const			{ return SubstrateMaterialType() == SUBSTRATE_MATERIAL_TYPE_COMPLEX; }
+	bool IsComplexSpecial() const 	{ return SubstrateMaterialType() == SUBSTRATE_MATERIAL_TYPE_COMPLEX_SPECIAL; }
+
+	static FString ToString(uint8 InMaterialType, bool bUpperCase = false)
+	{
+		switch (InMaterialType)
+		{
+			case SUBSTRATE_MATERIAL_TYPE_SIMPLE: 			return bUpperCase ? TEXT("SIMPLE")  		: TEXT("Simple");
+			case SUBSTRATE_MATERIAL_TYPE_SINGLE: 			return bUpperCase ? TEXT("SINGLE")  		: TEXT("Single");
+			case SUBSTRATE_MATERIAL_TYPE_COMPLEX: 			return bUpperCase ? TEXT("COMPLEX") 		: TEXT("Complex");
+			case SUBSTRATE_MATERIAL_TYPE_COMPLEX_SPECIAL: 	return bUpperCase ? TEXT("COMPLEX SPECIAL") : TEXT("Complex Special");
+			default: 										return bUpperCase ? TEXT("UNKOWN") 			: TEXT("Unknown");
+		}
+	}
+};
\ No newline at end of file
diff --git a/Engine/Source/Runtime/Engine/Public/MaterialShared.h b/Engine/Source/Runtime/Engine/Public/MaterialShared.h
--- a/Engine/Source/Runtime/Engine/Public/MaterialShared.h
+++ b/Engine/Source/Runtime/Engine/Public/MaterialShared.h
@@ -1404,7 +1404,7 @@ public:
 	uint8 GetSubstrateMaterialType() const { return GetSubstrateMaterialCompilationOutput().SubstrateMaterialType; }
 	uint8 GetSubstrateBSDFCount() const { return GetSubstrateMaterialCompilationOutput().SubstrateBSDFCount; }
 	uint8 GetSubstrateUintPerPixel() const { return GetSubstrateMaterialCompilationOutput().SubstrateUintPerPixel; }
-	bool GetSubstrateUsesComplexSpecialRenderPath() const { return GetSubstrateMaterialCompilationOutput().bUsesComplexSpecialRenderPath; }
+	bool GetSubstrateUsesComplexSpecialRenderPath() const { return GetSubstrateMaterialCompilationOutput().SubstrateMaterialType == SUBSTRATE_MATERIAL_TYPE_COMPLEX_SPECIAL; }
 	
 #if WITH_EDITOR
 	uint32 GetNumUsedUVScalars() const { return GetContent()->MaterialCompilationOutput.NumUsedUVScalars; }
diff --git a/Engine/Source/Runtime/Engine/Public/Rendering/SubstrateMaterialShared.cpp b/Engine/Source/Runtime/Engine/Public/Rendering/SubstrateMaterialShared.cpp
--- a/Engine/Source/Runtime/Engine/Public/Rendering/SubstrateMaterialShared.cpp
+++ b/Engine/Source/Runtime/Engine/Public/Rendering/SubstrateMaterialShared.cpp
@@ -9,7 +9,6 @@ FSubstrateMaterialCompilationOutput::FSubstrateMaterialCompilationOutput()
 	: SubstrateMaterialType(0)
 	, SubstrateBSDFCount(0)
 	, SubstrateUintPerPixel(0)
-	, bUsesComplexSpecialRenderPath(0)
 #if WITH_EDITOR
 	, SubstrateMaterialDescription()
 	, SharedLocalBasesCount(0)
diff --git a/Engine/Source/Runtime/Engine/Public/Rendering/SubstrateMaterialShared.h b/Engine/Source/Runtime/Engine/Public/Rendering/SubstrateMaterialShared.h
--- a/Engine/Source/Runtime/Engine/Public/Rendering/SubstrateMaterialShared.h
+++ b/Engine/Source/Runtime/Engine/Public/Rendering/SubstrateMaterialShared.h
@@ -103,7 +103,7 @@ public:
 	//// The following data is required at runtime
 	////
 
-	/** Substrate material type, at compile time (0:simple, 1:single, 2: complex) */
+	/** Substrate material type, at compile time (Possible values from SUBSTRATE_MATERIAL_TYPE_XXX: simple/single/complex/complex special) */
 	LAYOUT_FIELD(uint8, SubstrateMaterialType);
 
 	/** Substrate BSDF count, at compile time (0-7) */
@@ -112,9 +112,6 @@ public:
 	/** Substrate uint per pixel, at compile time (0-255) */
 	LAYOUT_FIELD(uint8, SubstrateUintPerPixel);
 
-	/** If Glints or SpecularLUT are used, we need to go the special more expenssive render path */
-	LAYOUT_BITFIELD(uint8, bUsesComplexSpecialRenderPath, 1);
-
 	////
 	//// The following data is only needed when compiling with the editor.
 	////
diff --git a/Engine/Source/Runtime/Renderer/Private/BasePassRendering.cpp b/Engine/Source/Runtime/Renderer/Private/BasePassRendering.cpp
--- a/Engine/Source/Runtime/Renderer/Private/BasePassRendering.cpp
+++ b/Engine/Source/Runtime/Renderer/Private/BasePassRendering.cpp
@@ -1758,7 +1758,7 @@ bool FBasePassMeshProcessor::Process(
 			// and Single materials, as other materials (Complex, ...) get decal applied during the base pass.
 			const uint8 SubstrateMaterialType	= MaterialResource.MaterialGetSubstrateMaterialType_RenderThread();
 			const uint32 DecalResponse		= MaterialResource.GetMaterialDecalResponse();
-			const bool bSupportDBufferPass	= SubstrateMaterialType == 0 || SubstrateMaterialType == 1; // Simple or single
+			const bool bSupportDBufferPass	= SubstrateMaterialType == SUBSTRATE_MATERIAL_TYPE_SIMPLE || SubstrateMaterialType == SUBSTRATE_MATERIAL_TYPE_SINGLE;
 			const bool bRoughnessResponse	= bSupportDBufferPass && (DecalResponse == MDR_ColorNormalRoughness || DecalResponse == MDR_ColorRoughness	|| DecalResponse == MDR_NormalRoughness || DecalResponse == MDR_Roughness);
 			const bool bColorResponse		= bSupportDBufferPass && (DecalResponse == MDR_ColorNormalRoughness || DecalResponse == MDR_ColorNormal		|| DecalResponse == MDR_ColorRoughness	|| DecalResponse == MDR_Color);
 			const bool bNormalResponse		= bSupportDBufferPass && (DecalResponse == MDR_ColorNormalRoughness || DecalResponse == MDR_ColorNormal		|| DecalResponse == MDR_NormalRoughness || DecalResponse == MDR_Normal);
