Add scope counter to UClusterUnionReplicatedProxyComponent::PostRepNotifies so cost can be seen easier ( seen taking ~3.5ms in a frame with samples ).

diff --git a/Engine/Source/Runtime/Engine/Private/PhysicsEngine/ClusterUnionReplicatedProxyComponent.cpp b/Engine/Source/Runtime/Engine/Private/PhysicsEngine/ClusterUnionReplicatedProxyComponent.cpp
--- a/Engine/Source/Runtime/Engine/Private/PhysicsEngine/ClusterUnionReplicatedProxyComponent.cpp
+++ b/Engine/Source/Runtime/Engine/Private/PhysicsEngine/ClusterUnionReplicatedProxyComponent.cpp
@@ -112,6 +112,8 @@ void UClusterUnionReplicatedProxyComponent::OnRep_ParticleChildToParents()
 
 void UClusterUnionReplicatedProxyComponent::PostRepNotifies()
 {
+	QUICK_SCOPE_CYCLE_COUNTER(UClusterUnionReplicatedProxyComponent_PostRepNotifies);
+
 	UActorComponent::PostRepNotifies();
 
 	if (IsPendingDeletion())
