Filter out garbage worlds that were not yet cleaned up.
Also replace check with ensure.

diff --git a/Engine/Source/Editor/UnrealEd/Private/FileHelpers.cpp b/Engine/Source/Editor/UnrealEd/Private/FileHelpers.cpp
--- a/Engine/Source/Editor/UnrealEd/Private/FileHelpers.cpp
+++ b/Engine/Source/Editor/UnrealEd/Private/FileHelpers.cpp
@@ -5172,6 +5172,12 @@ void FEditorFileUtils::GetDirtyWorldPackages(TArray<UPackage*>& OutDirtyPackages
 
 	for (TObjectIterator<UWorld> WorldIt; WorldIt; ++WorldIt)
 	{
+		// Filter out pending-delete worlds that may have leaked, e.g. from PIE sessions which were not cleaned up which cleared the PKG_PlayInEditor flag
+		if (!IsValid(*WorldIt))
+		{
+			continue;
+		}
+
 		UPackage* WorldPackage = WorldIt->GetOutermost();
 		if (!WorldPackage->HasAnyPackageFlags(PKG_PlayInEditor)
 			&& !WorldPackage->HasAnyFlags(RF_Transient)
diff --git a/Engine/Source/Runtime/Engine/Private/Level.cpp b/Engine/Source/Runtime/Engine/Private/Level.cpp
--- a/Engine/Source/Runtime/Engine/Private/Level.cpp
+++ b/Engine/Source/Runtime/Engine/Private/Level.cpp
@@ -4008,6 +4008,13 @@ TArray<FString> ULevel::GetOnDiskExternalActorPackages(bool bTryUsingPackageLoad
 
 TArray<UPackage*> ULevel::GetLoadedExternalObjectPackages() const
 {
+	// We also need to provide empty packages (for example, actors that were converted to non-external)
+	UWorld* World = GetTypedOuter<UWorld>();
+	if (!ensure(IsValid(World)))
+	{
+		return {};
+	}
+
 	TSet<UPackage*> ExternalObjectPackages;
 
 	// Get external packages (including deleted actors)
@@ -4026,10 +4033,6 @@ TArray<UPackage*> ULevel::GetLoadedExternalObjectPackages() const
 		return InPath;
 	};
 
-	// We also need to provide empty packages (for example, actors that were converted to non-external)
-	UWorld* World = GetTypedOuter<UWorld>();
-	check(IsValid(World));
-
 	TArray<FString> ExternalObjectsPaths;
 	ExternalObjectsPaths.Add(SanitizeExternalPath(ULevel::GetExternalActorsPath(World->GetPackage(), (World->OriginalWorldName == NAME_None) ? World->GetName() : World->OriginalWorldName.ToString())));
 	ExternalObjectsPaths.Add(SanitizeExternalPath(FExternalPackageHelper::GetExternalObjectsPath(World->GetPackage(), (World->OriginalWorldName == NAME_None) ? World->GetName() : World->OriginalWorldName.ToString())));
