WIP New HLSL translator Material Expressions
MaterialExpressionSpeedTree
MaterialExpressionDecalMipmapLevel
MaterialExpressionAbsorptionMediumMaterialOutput
MaterialExpressionBentNormalCustomOutput
MaterialExpressionDBufferTexture
MaterialExpressionSobol
MaterialExpressionDoubleVectorParameter
MaterialExpressionScalarParameter
MaterialExpressionVectorParameter
MaterialExpressionTemporalSobol
MaterialExpressionPrecomputedAOMask

#tests materials, editor, PIE, win64

diff --git a/Engine/Source/Runtime/Engine/Classes/Materials/MaterialExpressionAbsorptionMediumMaterialOutput.h b/Engine/Source/Runtime/Engine/Classes/Materials/MaterialExpressionAbsorptionMediumMaterialOutput.h
--- a/Engine/Source/Runtime/Engine/Classes/Materials/MaterialExpressionAbsorptionMediumMaterialOutput.h
+++ b/Engine/Source/Runtime/Engine/Classes/Materials/MaterialExpressionAbsorptionMediumMaterialOutput.h
@@ -22,6 +22,7 @@ public:
 	//~ Begin UMaterialExpression Interface
 	virtual int32 Compile(class FMaterialCompiler* Compiler, int32 OutputIndex) override;
 	virtual void GetCaption(TArray<FString>& OutCaptions) const override;
+	virtual bool GenerateHLSLExpression(FMaterialHLSLGenerator& Generator, UE::HLSLTree::FScope& Scope, int32 OutputIndex, UE::HLSLTree::FExpression const*& OutExpression) const override;
 	//~ End UMaterialExpression Interface
 #endif
 
@@ -29,5 +30,8 @@ public:
 	virtual int32 GetNumOutputs() const override;
 	virtual FString GetFunctionName() const override;
 	virtual FString GetDisplayName() const override;
+#if WITH_EDITOR
+	virtual UE::Shader::EValueType GetCustomOutputType(int32 OutputIndex) const override;
+#endif
 	//~ End UMaterialExpressionCustomOutput Interface
 };
diff --git a/Engine/Source/Runtime/Engine/Classes/Materials/MaterialExpressionBentNormalCustomOutput.h b/Engine/Source/Runtime/Engine/Classes/Materials/MaterialExpressionBentNormalCustomOutput.h
--- a/Engine/Source/Runtime/Engine/Classes/Materials/MaterialExpressionBentNormalCustomOutput.h
+++ b/Engine/Source/Runtime/Engine/Classes/Materials/MaterialExpressionBentNormalCustomOutput.h
@@ -23,10 +23,14 @@ class UMaterialExpressionBentNormalCustomOutput : public UMaterialExpressionCust
 	virtual void GetCaption(TArray<FString>& OutCaptions) const override;
 	virtual uint32 GetInputType(int32 InputIndex) override { return MCT_Float3; }
 	virtual FExpressionInput* GetInput(int32 InputIndex) override;
+	virtual bool GenerateHLSLExpression(FMaterialHLSLGenerator& Generator, UE::HLSLTree::FScope& Scope, int32 OutputIndex, UE::HLSLTree::FExpression const*& OutExpression) const override;
 #endif
 	virtual int32 GetNumOutputs() const override { return 1; }
 	virtual FString GetFunctionName() const override { return TEXT("GetBentNormal"); }
 	virtual FString GetDisplayName() const override { return TEXT("BentNormal"); }
+#if WITH_EDITOR
+	virtual UE::Shader::EValueType GetCustomOutputType(int32 OutputIndex) const override;
+#endif
 };
 
 
diff --git a/Engine/Source/Runtime/Engine/Classes/Materials/MaterialExpressionDBufferTexture.h b/Engine/Source/Runtime/Engine/Classes/Materials/MaterialExpressionDBufferTexture.h
--- a/Engine/Source/Runtime/Engine/Classes/Materials/MaterialExpressionDBufferTexture.h
+++ b/Engine/Source/Runtime/Engine/Classes/Materials/MaterialExpressionDBufferTexture.h
@@ -35,6 +35,7 @@ class UMaterialExpressionDBufferTexture : public UMaterialExpression
 #if WITH_EDITOR
 	virtual int32 Compile(class FMaterialCompiler* Compiler, int32 OutputIndex) override;
 	virtual void GetCaption(TArray<FString>& OutCaptions) const override;
+	virtual bool GenerateHLSLExpression(FMaterialHLSLGenerator& Generator, UE::HLSLTree::FScope& Scope, int32 OutputIndex, UE::HLSLTree::FExpression const*& OutExpression) const override;
 #endif
 	//~ End UMaterialExpression Interface
 };
diff --git a/Engine/Source/Runtime/Engine/Classes/Materials/MaterialExpressionDecalMipmapLevel.h b/Engine/Source/Runtime/Engine/Classes/Materials/MaterialExpressionDecalMipmapLevel.h
--- a/Engine/Source/Runtime/Engine/Classes/Materials/MaterialExpressionDecalMipmapLevel.h
+++ b/Engine/Source/Runtime/Engine/Classes/Materials/MaterialExpressionDecalMipmapLevel.h
@@ -29,6 +29,7 @@ class UMaterialExpressionDecalMipmapLevel : public UMaterialExpression
 #if WITH_EDITOR
 	virtual int32 Compile(class FMaterialCompiler* Compiler, int32 OutputIndex) override;
 	virtual void GetCaption(TArray<FString>& OutCaptions) const override;
+	virtual bool GenerateHLSLExpression(FMaterialHLSLGenerator& Generator, UE::HLSLTree::FScope& Scope, int32 OutputIndex, UE::HLSLTree::FExpression const*& OutExpression) const override;
 #endif
 	//~ End UMaterialExpression Interface
 };
diff --git a/Engine/Source/Runtime/Engine/Classes/Materials/MaterialExpressionPrecomputedAOMask.h b/Engine/Source/Runtime/Engine/Classes/Materials/MaterialExpressionPrecomputedAOMask.h
--- a/Engine/Source/Runtime/Engine/Classes/Materials/MaterialExpressionPrecomputedAOMask.h
+++ b/Engine/Source/Runtime/Engine/Classes/Materials/MaterialExpressionPrecomputedAOMask.h
@@ -22,6 +22,7 @@ class UMaterialExpressionPrecomputedAOMask : public UMaterialExpression
 #if WITH_EDITOR
 	virtual int32 Compile(class FMaterialCompiler* Compiler, int32 OutputIndex) override;
 	virtual void GetCaption(TArray<FString>& OutCaptions) const override;
+	virtual bool GenerateHLSLExpression(FMaterialHLSLGenerator& Generator, UE::HLSLTree::FScope& Scope, int32 OutputIndex, UE::HLSLTree::FExpression const*& OutExpression) const override;
 #endif
 	//~ End UMaterialExpression Interface
 };
diff --git a/Engine/Source/Runtime/Engine/Classes/Materials/MaterialExpressionSobol.h b/Engine/Source/Runtime/Engine/Classes/Materials/MaterialExpressionSobol.h
--- a/Engine/Source/Runtime/Engine/Classes/Materials/MaterialExpressionSobol.h
+++ b/Engine/Source/Runtime/Engine/Classes/Materials/MaterialExpressionSobol.h
@@ -41,6 +41,7 @@ class UMaterialExpressionSobol : public UMaterialExpression
 #if WITH_EDITOR
 	virtual int32 Compile(class FMaterialCompiler* Compiler, int32 OutputIndex) override;
 	virtual void GetCaption(TArray<FString>& OutCaptions) const override;
+	virtual bool GenerateHLSLExpression(FMaterialHLSLGenerator& Generator, UE::HLSLTree::FScope& Scope, int32 OutputIndex, UE::HLSLTree::FExpression const*& OutExpression) const override;
 #endif
 	//~ End UMaterialExpression Interface
 };
diff --git a/Engine/Source/Runtime/Engine/Classes/Materials/MaterialExpressionSpeedTree.h b/Engine/Source/Runtime/Engine/Classes/Materials/MaterialExpressionSpeedTree.h
--- a/Engine/Source/Runtime/Engine/Classes/Materials/MaterialExpressionSpeedTree.h
+++ b/Engine/Source/Runtime/Engine/Classes/Materials/MaterialExpressionSpeedTree.h
@@ -80,6 +80,7 @@ class UMaterialExpressionSpeedTree : public UMaterialExpression
 #if WITH_EDITOR
 	virtual int32 Compile(class FMaterialCompiler* Compiler, int32 OutputIndex) override;
 	virtual void GetCaption(TArray<FString>& OutCaptions) const override;
+	virtual bool GenerateHLSLExpression(FMaterialHLSLGenerator& Generator, UE::HLSLTree::FScope& Scope, int32 OutputIndex, UE::HLSLTree::FExpression const*& OutExpression) const override;
 #endif
 	//~ End UMaterialExpression Interface
 };
diff --git a/Engine/Source/Runtime/Engine/Classes/Materials/MaterialExpressionTemporalSobol.h b/Engine/Source/Runtime/Engine/Classes/Materials/MaterialExpressionTemporalSobol.h
--- a/Engine/Source/Runtime/Engine/Classes/Materials/MaterialExpressionTemporalSobol.h
+++ b/Engine/Source/Runtime/Engine/Classes/Materials/MaterialExpressionTemporalSobol.h
@@ -36,6 +36,7 @@ class UMaterialExpressionTemporalSobol : public UMaterialExpression
 #if WITH_EDITOR
 	virtual int32 Compile(class FMaterialCompiler* Compiler, int32 OutputIndex) override;
 	virtual void GetCaption(TArray<FString>& OutCaptions) const override;
+	virtual bool GenerateHLSLExpression(FMaterialHLSLGenerator& Generator, UE::HLSLTree::FScope& Scope, int32 OutputIndex, UE::HLSLTree::FExpression const*& OutExpression) const override;
 #endif
 	//~ End UMaterialExpression Interface
 };
diff --git a/Engine/Source/Runtime/Engine/Private/Materials/MaterialExpressionHLSL.cpp b/Engine/Source/Runtime/Engine/Private/Materials/MaterialExpressionHLSL.cpp
--- a/Engine/Source/Runtime/Engine/Private/Materials/MaterialExpressionHLSL.cpp
+++ b/Engine/Source/Runtime/Engine/Private/Materials/MaterialExpressionHLSL.cpp
@@ -13,8 +13,10 @@
 #include "Materials/MaterialAttributeDefinitionMap.h"
 #include "Materials/MaterialExpression.h"
 #include "Materials/MaterialExpressionAbs.h"
+#include "Materials/MaterialExpressionAbsorptionMediumMaterialOutput.h"
 #include "Materials/MaterialExpressionActorPositionWS.h"
 #include "Materials/MaterialExpressionAdd.h"
+#include "Materials/MaterialExpressionAntialiasedTextureMask.h"
 #include "Materials/MaterialExpressionAppendVector.h"
 #include "Materials/MaterialExpressionArccosine.h"
 #include "Materials/MaterialExpressionArccosineFast.h"
@@ -26,6 +28,7 @@
 #include "Materials/MaterialExpressionArctangentFast.h"
 #include "Materials/MaterialExpressionAtmosphericLightColor.h"
 #include "Materials/MaterialExpressionAtmosphericLightVector.h"
+#include "Materials/MaterialExpressionBentNormalCustomOutput.h"
 #include "Materials/MaterialExpressionBinaryOp.h"
 #include "Materials/MaterialExpressionBlackBody.h"
 #include "Materials/MaterialExpressionBlendMaterialAttributes.h"
@@ -50,11 +53,13 @@
 #include "Materials/MaterialExpressionCurveAtlasRowParameter.h"
 #include "Materials/MaterialExpressionCustom.h"
 #include "Materials/MaterialExpressionDataDrivenShaderPlatformInfoSwitch.h"
+#include "Materials/MaterialExpressionDBufferTexture.h"
 #include "Materials/MaterialExpressionDDX.h"
 #include "Materials/MaterialExpressionDDY.h"
 #include "Materials/MaterialExpressionDecalColor.h"
 #include "Materials/MaterialExpressionDecalDerivative.h"
 #include "Materials/MaterialExpressionDecalLifetimeOpacity.h"
+#include "Materials/MaterialExpressionDecalMipmapLevel.h"
 #include "Materials/MaterialExpressionDeltaTime.h"
 #include "Materials/MaterialExpressionDepthFade.h"
 #include "Materials/MaterialExpressionDepthOfFieldFunction.h"
@@ -68,6 +73,7 @@
 #include "Materials/MaterialExpressionDistanceToNearestSurface.h"
 #include "Materials/MaterialExpressionDivide.h"
 #include "Materials/MaterialExpressionDotProduct.h"
+#include "Materials/MaterialExpressionDoubleVectorParameter.h"
 #include "Materials/MaterialExpressionDynamicParameter.h"
 #include "Materials/MaterialExpressionExecBegin.h"
 #include "Materials/MaterialExpressionExecEnd.h"
@@ -140,6 +146,7 @@
 #include "Materials/MaterialExpressionPixelDepth.h"
 #include "Materials/MaterialExpressionPixelNormalWS.h"
 #include "Materials/MaterialExpressionPower.h"
+#include "Materials/MaterialExpressionPrecomputedAOMask.h"
 #include "Materials/MaterialExpressionPreSkinnedLocalBounds.h"
 #include "Materials/MaterialExpressionPreSkinnedNormal.h"
 #include "Materials/MaterialExpressionPreSkinnedPosition.h"
@@ -158,6 +165,7 @@
 #include "Materials/MaterialExpressionRuntimeVirtualTextureSample.h"
 #include "Materials/MaterialExpressionRuntimeVirtualTextureSampleParameter.h"
 #include "Materials/MaterialExpressionSaturate.h"
+#include "Materials/MaterialExpressionScalarParameter.h"
 #include "Materials/MaterialExpressionSceneColor.h"
 #include "Materials/MaterialExpressionSceneDepth.h"
 #include "Materials/MaterialExpressionSceneDepthWithoutWater.h"
@@ -178,6 +186,8 @@
 #include "Materials/MaterialExpressionSkyAtmosphereViewLuminance.h"
 #include "Materials/MaterialExpressionSkyLightEnvMapSample.h"
 #include "Materials/MaterialExpressionSmoothStep.h"
+#include "Materials/MaterialExpressionSobol.h"
+#include "Materials/MaterialExpressionSpeedTree.h"
 #include "Materials/MaterialExpressionSphereMask.h"
 #include "Materials/MaterialExpressionSquareRoot.h"
 #include "Materials/MaterialExpressionStaticBool.h"
@@ -189,6 +199,7 @@
 #include "Materials/MaterialExpressionSwitch.h"
 #include "Materials/MaterialExpressionTangent.h"
 #include "Materials/MaterialExpressionTangentOutput.h"
+#include "Materials/MaterialExpressionTemporalSobol.h"
 #include "Materials/MaterialExpressionTextureCoordinate.h"
 #include "Materials/MaterialExpressionTextureObject.h"
 #include "Materials/MaterialExpressionTextureObjectParameter.h"
@@ -203,6 +214,7 @@
 #include "Materials/MaterialExpressionTruncateLWC.h"
 #include "Materials/MaterialExpressionTwoSidedSign.h"
 #include "Materials/MaterialExpressionVectorNoise.h"
+#include "Materials/MaterialExpressionVectorParameter.h"
 #include "Materials/MaterialExpressionVertexColor.h"
 #include "Materials/MaterialExpressionVertexInterpolator.h"
 #include "Materials/MaterialExpressionVertexNormalWS.h"
@@ -719,7 +731,28 @@ bool UMaterialExpressionParameter::GenerateHLSLExpression(FMaterialHLSLGenerator
 		return Generator.Error(TEXT("Failed to get parameter value"));
 	}
 
-	OutExpression = Generator.GenerateMaterialParameter(ParameterName, ParameterMeta);
+	if (ParameterMeta.PrimitiveDataIndex != INDEX_NONE)
+	{
+		using namespace UE::HLSLTree;
+		FTree& Tree = Generator.GetTree();
+		if (ParameterMeta.Value.Type == EMaterialParameterType::Scalar)
+		{
+			OutExpression = Tree.NewExpression<Material::FExpressionCustomPrimitiveDataFunction>(ParameterMeta.PrimitiveDataIndex);
+		}
+		else
+		{
+			const FExpression* X = Tree.NewExpression<Material::FExpressionCustomPrimitiveDataFunction>(ParameterMeta.PrimitiveDataIndex + 0);
+			const FExpression* Y = Tree.NewExpression<Material::FExpressionCustomPrimitiveDataFunction>(ParameterMeta.PrimitiveDataIndex + 1);
+			const FExpression* Z = Tree.NewExpression<Material::FExpressionCustomPrimitiveDataFunction>(ParameterMeta.PrimitiveDataIndex + 2);
+			const FExpression* W = Tree.NewExpression<Material::FExpressionCustomPrimitiveDataFunction>(ParameterMeta.PrimitiveDataIndex + 3);
+			OutExpression = Tree.NewAppend(X, Y, Z, W);
+		}
+	}
+	else
+	{
+		OutExpression = Generator.GenerateMaterialParameter(ParameterName, ParameterMeta);
+	}
+
 	return true;
 }
 
@@ -1337,6 +1370,36 @@ bool UMaterialExpressionSkyLightEnvMapSample::GenerateHLSLExpression(FMaterialHL
 	return true;
 }
 
+bool UMaterialExpressionSpeedTree::GenerateHLSLExpression(FMaterialHLSLGenerator& Generator, UE::HLSLTree::FScope& Scope, int32 OutputIndex, UE::HLSLTree::FExpression const*& OutExpression) const
+{
+	using namespace UE::HLSLTree;
+
+	const FExpression* GeometryExpression = GeometryInput.AcquireHLSLExpressionOrConstant(Generator, Scope, int32(GeometryType));
+	const FExpression* WindExpression = WindInput.AcquireHLSLExpressionOrConstant(Generator, Scope, int32(WindType));
+	const FExpression* LODExpression = LODInput.AcquireHLSLExpressionOrConstant(Generator, Scope, int32(LODType));
+
+	bool bExtraBend = (ExtraBendWS.GetTracedInput().Expression != nullptr);
+	const FExpression* ExtraBendExpression = ExtraBendWS.AcquireHLSLExpressionOrConstant(Generator, Scope, FVector3f(0.f, 0.f, 0.f));
+
+	OutExpression = Generator.GetTree().NewExpression<Material::FExpressionSpeedTree>(GeometryExpression, WindExpression, LODExpression, ExtraBendExpression, bExtraBend, bAccurateWindVelocities, BillboardThreshold, false);
+	return true;
+}
+
+bool UMaterialExpressionDecalMipmapLevel::GenerateHLSLExpression(FMaterialHLSLGenerator& Generator, UE::HLSLTree::FScope& Scope, int32 OutputIndex, UE::HLSLTree::FExpression const*& OutExpression) const
+{
+	using namespace UE::HLSLTree;
+
+	const UMaterial* BaseMaterial = Generator.GetTargetMaterial();
+	if (BaseMaterial && (BaseMaterial->MaterialDomain != MD_DeferredDecal))
+	{
+		return Generator.Error(TEXT("Node only works for the deferred decal material domain."));
+	}
+
+	const FExpression* TextureSizeExpression = TextureSize.AcquireHLSLExpressionOrConstant(Generator, Scope, FVector2f(ConstWidth, ConstHeight));
+	OutExpression = Generator.GetTree().NewExpression<Material::FExpressionDecalMipmapLevel>(TextureSizeExpression);
+	return true;
+}
+
 bool UMaterialExpressionTextureSample::GenerateHLSLExpressionBase(FMaterialHLSLGenerator& Generator, UE::HLSLTree::FScope& Scope, const UE::HLSLTree::FExpression* TextureExpression, UE::HLSLTree::FExpression const*& OutExpression) const
 {
 	using namespace UE::HLSLTree;
@@ -2065,6 +2128,30 @@ bool UMaterialExpressionDepthFade::GenerateHLSLExpression(FMaterialHLSLGenerator
 	return true;
 }
 
+bool UMaterialExpressionDBufferTexture::GenerateHLSLExpression(FMaterialHLSLGenerator& Generator, UE::HLSLTree::FScope& Scope, int32 OutputIndex, UE::HLSLTree::FExpression const*& OutExpression) const
+{
+
+	using namespace UE::HLSLTree;
+
+	const UMaterial* BaseMaterial = Generator.GetTargetMaterial();
+	if (BaseMaterial)
+	{
+		if (Material->MaterialDomain != MD_Surface || IsTranslucentBlendMode(Material->GetBlendMode()))
+		{
+			return Generator.Error(TEXT("DBuffer scene textures are only available on opaque or masked surfaces."));
+		}
+	}
+
+	const FExpression* UVExpression = nullptr;
+	if (Coordinates.GetTracedInput().Expression)
+	{
+		UVExpression = Coordinates.AcquireHLSLExpression(Generator, Scope);
+	}
+
+	OutExpression = Generator.GetTree().NewExpression<Material::FExpressionDBufferTexture>(UVExpression, DBufferTextureId);
+	return true;
+}
+
 bool UMaterialExpressionOneMinus::GenerateHLSLExpression(FMaterialHLSLGenerator& Generator, UE::HLSLTree::FScope& Scope, int32 OutputIndex, UE::HLSLTree::FExpression const*& OutExpression) const
 {
 	const UE::HLSLTree::FExpression* InputExpression = Input.AcquireHLSLExpression(Generator, Scope);
@@ -2301,6 +2388,37 @@ bool UMaterialExpressionDepthOfFieldFunction::GenerateHLSLExpression(FMaterialHL
 	return OutExpression != nullptr;
 }
 
+bool UMaterialExpressionSobol::GenerateHLSLExpression(FMaterialHLSLGenerator& Generator, UE::HLSLTree::FScope& Scope, int32 OutputIndex, UE::HLSLTree::FExpression const*& OutExpression) const
+{
+	using namespace UE::HLSLTree;
+
+	const FExpression* CellExpression = Cell.AcquireHLSLExpressionOrConstant(Generator, Scope, FVector2f(0.f, 0.f));
+	const FExpression* IndexExpression = Index.AcquireHLSLExpressionOrConstant(Generator, Scope, int32(ConstIndex));
+	const FExpression* SeedExpression = Seed.AcquireHLSLExpressionOrConstant(Generator, Scope, FVector2f(ConstSeed.X, ConstSeed.Y));
+
+	OutExpression = Generator.GetTree().NewExpression<Material::FExpressionSobolFunction>(CellExpression, IndexExpression, SeedExpression, false);
+	return true;
+}
+
+bool UMaterialExpressionTemporalSobol::GenerateHLSLExpression(FMaterialHLSLGenerator& Generator, UE::HLSLTree::FScope& Scope, int32 OutputIndex, UE::HLSLTree::FExpression const*& OutExpression) const
+{
+	using namespace UE::HLSLTree;
+
+	const FExpression* IndexExpression = Index.AcquireHLSLExpressionOrConstant(Generator, Scope, int32(ConstIndex));
+	const FExpression* SeedExpression = Seed.AcquireHLSLExpressionOrConstant(Generator, Scope, FVector2f(ConstSeed.X, ConstSeed.Y));
+
+	OutExpression = Generator.GetTree().NewExpression<Material::FExpressionSobolFunction>(nullptr, IndexExpression, SeedExpression, true);
+	return true;
+}
+
+bool UMaterialExpressionPrecomputedAOMask::GenerateHLSLExpression(FMaterialHLSLGenerator& Generator, UE::HLSLTree::FScope& Scope, int32 OutputIndex, UE::HLSLTree::FExpression const*& OutExpression) const
+{
+	using namespace UE::HLSLTree;
+	const FExpression* AOMaskExpression = Generator.GetTree().NewExpression<Material::FExpressionExternalInput>(Material::EExternalInput::AOMask);
+	OutExpression = Generator.GetTree().NewExpression<Material::FExpressionAOMaskFunction>(AOMaskExpression);
+	return true;
+}
+
 static bool GenerateHLSLExpressionTrig(FMaterialHLSLGenerator& Generator,
 	UE::HLSLTree::EOperation Op,
 	const UE::HLSLTree::FExpression* InputExpression,
@@ -3702,6 +3820,45 @@ UE::Shader::EValueType UMaterialExpressionThinTranslucentMaterialOutput::GetCust
 	return EValueType::Void;
 }
 
+bool UMaterialExpressionAbsorptionMediumMaterialOutput::GenerateHLSLExpression(FMaterialHLSLGenerator& Generator, UE::HLSLTree::FScope& Scope, int32 OutputIndex, UE::HLSLTree::FExpression const*& OutExpression) const
+{
+	if (OutputIndex == 0)
+	{
+		OutExpression = TransmittanceColor.AcquireHLSLExpressionOrConstant(Generator, Scope, FVector3f(1.0f, 1.0f, 1.0f));
+		return true;
+	}
+
+	return Generator.Error(TEXT("Invalid output"));
+}
+
+UE::Shader::EValueType UMaterialExpressionAbsorptionMediumMaterialOutput::GetCustomOutputType(int32 OutputIndex) const
+{
+	using namespace UE::Shader;
+	if (OutputIndex == 0)
+	{
+		return EValueType::Float3;
+	}
+
+	return EValueType::Void;
+}
+
+bool UMaterialExpressionBentNormalCustomOutput::GenerateHLSLExpression(FMaterialHLSLGenerator& Generator, UE::HLSLTree::FScope& Scope, int32 OutputIndex, UE::HLSLTree::FExpression const*& OutExpression) const
+{
+	if (Input.GetTracedInput().Expression)
+	{
+		OutExpression = Input.AcquireHLSLExpression(Generator, Scope);
+		return true;
+	}
+
+	return Generator.Error(TEXT("Input missing"));
+}
+
+UE::Shader::EValueType UMaterialExpressionBentNormalCustomOutput::GetCustomOutputType(int32 OutputIndex) const
+{
+	using namespace UE::Shader;
+	return EValueType::Float3;
+}
+
 bool UMaterialExpressionVolumetricCloudEmptySpaceSkippingInput::GenerateHLSLExpression(FMaterialHLSLGenerator& Generator, UE::HLSLTree::FScope& Scope, int32 OutputIndex, UE::HLSLTree::FExpression const*& OutExpression) const
 {
 	return Generator.Error(TEXT("UMaterialExpressionVolumetricCloudEmptySpaceSkippingInput::GenerateHLSLExpression unimplemented"));
diff --git a/Engine/Source/Runtime/Engine/Private/Materials/MaterialHLSLEmitter.cpp b/Engine/Source/Runtime/Engine/Private/Materials/MaterialHLSLEmitter.cpp
--- a/Engine/Source/Runtime/Engine/Private/Materials/MaterialHLSLEmitter.cpp
+++ b/Engine/Source/Runtime/Engine/Private/Materials/MaterialHLSLEmitter.cpp
@@ -389,7 +389,12 @@ static void GetMaterialEnvironment(EShaderPlatform InPlatform,
 		OutEnvironment.SetDefine(TEXT("SPHERICAL_PARTICLE_OPACITY"), TEXT("1"));
 	}
 
-	if (false)//bUsesParticleSubUVs)
+	const bool bNeedsParticleSubUVs =
+		EmitMaterialData.IsExternalInputUsed(Material::EExternalInput::ParticleSubUVCoords0) ||
+		EmitMaterialData.IsExternalInputUsed(Material::EExternalInput::ParticleSubUVCoords1) ||
+		EmitMaterialData.IsExternalInputUsed(Material::EExternalInput::ParticleSubUVLerp);
+
+	if (bNeedsParticleSubUVs)
 	{
 		OutEnvironment.SetDefine(TEXT("USE_PARTICLE_SUBUVS"), TEXT("1"));
 	}
@@ -399,14 +404,14 @@ static void GetMaterialEnvironment(EShaderPlatform InPlatform,
 		OutEnvironment.SetDefine(TEXT("LIGHTMAP_UV_ACCESS"), TEXT("1"));
 	}
 
-	if (false)//bUsesAOMaterialMask)
+	if (EmitMaterialData.IsExternalInputUsed(Material::EExternalInput::AOMask))
 	{
 		OutEnvironment.SetDefine(TEXT("USES_AO_MATERIAL_MASK"), TEXT("1"));
 	}
 
-	if (false)//bUsesSpeedTree)
+	if (EmitContext.bUsesSpeedTree)
 	{
-		OutEnvironment.SetDefine(TEXT("USES_SPEEDTREE"), TEXT("1"));
+		OutEnvironment.SetDefine(TEXT("USES_SPEEDTREE"), 1);
 	}
 
 	const bool bNeedsWorldPositionExcludingShaderOffsets =
@@ -427,16 +432,6 @@ static void GetMaterialEnvironment(EShaderPlatform InPlatform,
 		OutEnvironment.SetDefine(TEXT("NEEDS_PARTICLE_SIZE"), TEXT("1"));
 	}
 
-	const bool bNeedsParticleSubUVs =
-		EmitMaterialData.IsExternalInputUsed(Material::EExternalInput::ParticleSubUVCoords0) ||
-		EmitMaterialData.IsExternalInputUsed(Material::EExternalInput::ParticleSubUVCoords1) ||
-		EmitMaterialData.IsExternalInputUsed(Material::EExternalInput::ParticleSubUVLerp);
-
-	if (bNeedsParticleSubUVs)
-	{
-		OutEnvironment.SetDefine(TEXT("USE_PARTICLE_SUBUVS"), TEXT("1"));
-	}
-
 	if (MaterialCompilationOutput.bNeedsSceneTextures)
 	{
 		OutEnvironment.SetDefine(TEXT("NEEDS_SCENE_TEXTURES"), TEXT("1"));
diff --git a/Engine/Source/Runtime/Engine/Private/Materials/MaterialHLSLTree.cpp b/Engine/Source/Runtime/Engine/Private/Materials/MaterialHLSLTree.cpp
--- a/Engine/Source/Runtime/Engine/Private/Materials/MaterialHLSLTree.cpp
+++ b/Engine/Source/Runtime/Engine/Private/Materials/MaterialHLSLTree.cpp
@@ -167,6 +167,8 @@ FExternalInputDescription GetExternalInputDescription(EExternalInput Input)
 
 	case EExternalInput::IsOrthographic: return FExternalInputDescription(TEXT("IsOrthographic"), Shader::EValueType::Float1);
 
+	case EExternalInput::AOMask: return FExternalInputDescription(TEXT("PrecomputedAOMask"), Shader::EValueType::Float1);
+
 	default: checkNoEntry(); return FExternalInputDescription(TEXT("Invalid"), Shader::EValueType::Void);
 	}
 }
@@ -444,6 +446,8 @@ void FExpressionExternalInput::EmitValueShader(FEmitContext& Context, FEmitScope
 
 		case EExternalInput::IsOrthographic: Code = TEXT("((View.ViewToClip[3][3] < 1.0f) ? 0.0f : 1.0f)"); break;
 
+		case EExternalInput::AOMask: Code = TEXT("Parameters.AOMaterialMask"); break;
+
 		default:
 			checkNoEntry();
 			break;
@@ -805,6 +809,131 @@ void FExpressionSkyLightEnvMapSample::EmitValueShader(FEmitContext& Context, FEm
 	OutResult.Code = Context.EmitExpression(Scope, Shader::EValueType::Float3, TEXT("MaterialExpressionSkyLightEnvMapSample(%, %)"), EmitDirectionExpression, EmitRoughnessExpression);
 }
 
+const FExpression* FExpressionSpeedTree::ComputePreviousFrame(FTree& Tree, const FRequestedType& RequestedType) const
+{
+	// if bAccurateWind and computing previous frame use new expression
+	if (bAccurateWind)
+	{
+		return Tree.NewExpression<Material::FExpressionSpeedTree>(
+			GeometryExpression, 
+			WindExpression, 
+			LODExpression, 
+			ExtraBendExpression, 
+			bExtraBend, 
+			bAccurateWind,
+			BillboardThreshold, 
+			true);
+	}
+
+	return nullptr;
+}
+
+bool FExpressionSpeedTree::PrepareValue(FEmitContext& Context, FEmitScope& Scope, const FRequestedType& RequestedType, FPrepareValueResult& OutResult) const
+{
+	if (Context.Material && Context.Material->IsUsedWithSkeletalMesh())
+	{
+		return Context.Error(TEXT("SpeedTree node not currently supported for Skeletal Meshes, please disable usage flag."));
+	}
+
+	if (Context.ShaderFrequency != SF_Vertex)
+	{
+		return Context.Error(TEXT("Invalid node used in pixel/hull/domain shader input."));
+	}
+
+	if (Context.PrepareExpression(GeometryExpression, Scope, RequestedType).IsVoid() || 
+		Context.PrepareExpression(WindExpression, Scope, RequestedType).IsVoid() ||
+		Context.PrepareExpression(LODExpression, Scope, RequestedType).IsVoid() ||
+		Context.PrepareExpression(ExtraBendExpression, Scope, RequestedType).IsVoid())
+	{
+		return false;
+	}
+
+	if (Context.bMarkLiveValues)
+	{
+		FEmitData& EmitMaterialData = Context.FindData<FEmitData>();
+		for (int32 i = (int32)EExternalInput::TexCoord2; i <= (int32)EExternalInput::TexCoord7; ++i)
+		{
+			EmitMaterialData.ExternalInputMask[SF_Vertex][i] = true;
+		}
+	}
+
+	return OutResult.SetType(Context, RequestedType, EExpressionEvaluation::Shader, Shader::EValueType::Float3);
+}
+
+void FExpressionSpeedTree::EmitValueShader(FEmitContext& Context, FEmitScope& Scope, const FRequestedType& RequestedType, FEmitValueShaderResult& OutResult) const
+{
+	Context.bUsesSpeedTree = true;
+
+	FEmitShaderExpression* EmitGeometryExpression = GeometryExpression->GetValueShader(Context, Scope, Shader::EValueType::Int1);
+	FEmitShaderExpression* EmitWindExpression = WindExpression->GetValueShader(Context, Scope, Shader::EValueType::Int1);
+	FEmitShaderExpression* EmitLODExpression = LODExpression->GetValueShader(Context, Scope, Shader::EValueType::Int1);
+	FEmitShaderExpression* EmitExtraBendExpression = ExtraBendExpression->GetValueShader(Context, Scope, Shader::EValueType::Float3);
+	OutResult.Code = Context.EmitInlineExpression(
+		Scope, Shader::EValueType::Float3, 
+		TEXT("GetSpeedTreeVertexOffset(%, %, %, %, %, %, %)"), 
+		EmitGeometryExpression, 
+		EmitWindExpression,
+		EmitLODExpression,
+		BillboardThreshold,
+		bPreviousFrame ? TEXT("true") : TEXT("false"),
+		bExtraBend ? TEXT("true") : TEXT("false"),
+		EmitExtraBendExpression);
+}
+
+bool FExpressionDecalMipmapLevel::PrepareValue(FEmitContext& Context, FEmitScope& Scope, const FRequestedType& RequestedType, FPrepareValueResult& OutResult) const
+{
+	if (Context.PrepareExpression(TextureSizeExpression, Scope, RequestedType).IsVoid())
+	{
+		return false;
+	}
+
+	return OutResult.SetType(Context, RequestedType, EExpressionEvaluation::Shader, Shader::EValueType::Float1);
+}
+
+void FExpressionDecalMipmapLevel::EmitValueShader(FEmitContext& Context, FEmitScope& Scope, const FRequestedType& RequestedType, FEmitValueShaderResult& OutResult) const
+{
+	FEmitShaderExpression* EmitSizeExpression = TextureSizeExpression->GetValueShader(Context, Scope, Shader::EValueType::Float2);
+	OutResult.Code = Context.EmitInlineExpression(
+		Scope, Shader::EValueType::Float1,
+		TEXT("ComputeDecalMipmapLevel(Parameters, %)"),
+		EmitSizeExpression);
+}
+
+bool FExpressionDBufferTexture::PrepareValue(FEmitContext& Context, FEmitScope& Scope, const FRequestedType& RequestedType, FPrepareValueResult& OutResult) const
+{
+	if (UVExpression && Context.PrepareExpression(UVExpression, Scope, RequestedType).IsVoid())
+	{
+		return false;
+	}
+
+	if (Context.bMarkLiveValues && Context.MaterialCompilationOutput)
+	{
+		Context.MaterialCompilationOutput->SetIsDBufferTextureUsed(DBufferTextureID);
+		Context.MaterialCompilationOutput->SetIsDBufferTextureLookupUsed(true);
+	}
+
+	return OutResult.SetType(Context, RequestedType, EExpressionEvaluation::Shader, Shader::EValueType::Float4);
+}
+
+void FExpressionDBufferTexture::EmitValueShader(FEmitContext& Context, FEmitScope& Scope, const FRequestedType& RequestedType, FEmitValueShaderResult& OutResult) const
+{
+	FEmitShaderExpression* EmitTexCoord = nullptr;
+	if (UVExpression)
+	{
+		EmitTexCoord = UVExpression->GetValueShader(Context, Scope, Shader::EValueType::Float2);
+		EmitTexCoord = Context.EmitExpression(Scope, Shader::EValueType::Float2, TEXT("ClampSceneTextureUV(ViewportUVToSceneTextureUV(%), 0)"), EmitTexCoord);
+	}
+	else
+	{
+		EmitTexCoord = Context.EmitExpression(Scope, Shader::EValueType::Float2, TEXT("GetDefaultSceneTextureUV(Parameters, 0)"));
+	}
+
+	OutResult.Code = Context.EmitInlineExpression(
+		Scope, Shader::EValueType::Float4,
+		TEXT("MaterialExpressionDBufferTextureLookup(Parameters, %, %)"),
+		EmitTexCoord, DBufferTextureID);
+}
+
 namespace Private
 {
 
@@ -2658,6 +2787,82 @@ void FExpressionDepthOfFieldFunction::EmitValueShader(FEmitContext& Context, FEm
 		FunctionValue);
 }
 
+bool FExpressionSobolFunction::PrepareValue(FEmitContext& Context, FEmitScope& Scope, const FRequestedType& RequestedType, FPrepareValueResult& OutResult) const
+{
+	if ((CellExpression && Context.PrepareExpression(CellExpression, Scope, Shader::EValueType::Float2).IsVoid()) ||
+		Context.PrepareExpression(IndexExpression, Scope, Shader::EValueType::Int1).IsVoid() ||
+		Context.PrepareExpression(SeedExpression, Scope, Shader::EValueType::Float2).IsVoid())
+	{
+		return false;
+	}
+
+	return OutResult.SetType(Context, RequestedType, EExpressionEvaluation::Shader, Shader::EValueType::Float2);
+}
+
+void FExpressionSobolFunction::EmitValueShader(FEmitContext& Context, FEmitScope& Scope, const FRequestedType& RequestedType, FEmitValueShaderResult& OutResult) const
+{
+	FEmitShaderExpression* EmitCell = CellExpression ? CellExpression->GetValueShader(Context, Scope, Shader::EValueType::Float2) : nullptr;
+	FEmitShaderExpression* EmitIndex = IndexExpression->GetValueShader(Context, Scope, Shader::EValueType::Int1);
+	FEmitShaderExpression* EmitSeed = SeedExpression->GetValueShader(Context, Scope, Shader::EValueType::Float2);
+
+	if (bTemporal)
+	{
+		OutResult.Code = Context.EmitExpression(Scope, Shader::EValueType::Float2,
+			TEXT("float2(SobolIndex(SobolPixel(uint2(Parameters.SvPosition.xy)), uint(View.StateFrameIndexMod8 + 8 * %)) ^ uint2(% * 0x10000) & 0xffff) / 0x10000"),
+			EmitIndex, EmitSeed);
+	}
+	else
+	{
+		check(EmitCell);
+		OutResult.Code = Context.EmitExpression(Scope, Shader::EValueType::Float2,
+			TEXT("floor(%) + float2(SobolIndex(SobolPixel(uint2(%)), uint(%)) ^ uint2(% * 0x10000) & 0xffff) / 0x10000"),
+			EmitCell, EmitCell, EmitIndex, EmitSeed);
+	}
+}
+
+bool FExpressionCustomPrimitiveDataFunction::PrepareValue(FEmitContext& Context, FEmitScope& Scope, const FRequestedType& RequestedType, FPrepareValueResult& OutResult) const
+{
+	if (Index < FCustomPrimitiveData::NumCustomPrimitiveDataFloats)
+	{
+		return OutResult.SetType(Context, RequestedType, EExpressionEvaluation::Shader, Shader::EValueType::Float1);
+	}
+	
+	// out of range values set to 0
+	return OutResult.SetType(Context, RequestedType, EExpressionEvaluation::ConstantZero, Shader::EValueType::Float1);
+}
+
+void FExpressionCustomPrimitiveDataFunction::EmitValueShader(FEmitContext& Context, FEmitScope& Scope, const FRequestedType& RequestedType, FEmitValueShaderResult& OutResult) const
+{
+	if (Index < FCustomPrimitiveData::NumCustomPrimitiveDataFloats)
+	{
+		const int32 CustomDataIndex = Index / 4;
+		const int32 ElementIndex = Index % 4; // x, y, z or w
+
+		OutResult.Code = Context.EmitInlineExpression(Scope, Shader::EValueType::Float1,
+			TEXT("GetPrimitiveData(Parameters).CustomPrimitiveData[%][%]"),
+			CustomDataIndex, ElementIndex);
+	}
+	else
+	{
+		OutResult.Code = Context.EmitConstantZero(Scope, Shader::EValueType::Float1);
+	}
+}
+
+bool FExpressionAOMaskFunction::PrepareValue(FEmitContext& Context, FEmitScope& Scope, const FRequestedType& RequestedType, FPrepareValueResult& OutResult) const
+{
+	if (Context.TargetParameters.FeatureLevel < ERHIFeatureLevel::SM5)
+	{
+		return Context.Errorf(TEXT("Node not supported in feature level %d. %d required."), Context.TargetParameters.FeatureLevel, ERHIFeatureLevel::SM5);
+	}
+
+	if (Context.ShaderFrequency != SF_Pixel && Context.ShaderFrequency != SF_Compute)
+	{
+		return Context.Error(TEXT("Invalid node used in vertex/hull/domain shader."));
+	}
+
+	return FExpressionForward::PrepareValue(Context, Scope, RequestedType, OutResult);
+}
+
 bool FExpressionNaniteReplaceFunction::PrepareValue(FEmitContext& Context, FEmitScope& Scope, const FRequestedType& RequestedType, FPrepareValueResult& OutResult) const
 {
 	const FPreparedType& DefaultType = Context.PrepareExpression(DefaultExpression, Scope, RequestedType);
diff --git a/Engine/Source/Runtime/Engine/Public/HLSLTree/HLSLTreeEmit.h b/Engine/Source/Runtime/Engine/Public/HLSLTree/HLSLTreeEmit.h
--- a/Engine/Source/Runtime/Engine/Public/HLSLTree/HLSLTreeEmit.h
+++ b/Engine/Source/Runtime/Engine/Public/HLSLTree/HLSLTreeEmit.h
@@ -582,6 +582,7 @@ public:
 	bool bUseAnalyticDerivatives = false;
 
 	bool bUsesSkyAtmosphere = false;
+	bool bUsesSpeedTree = false;
 
 	uint32 DynamicParticleParameterMask = 0u;
 
diff --git a/Engine/Source/Runtime/Engine/Public/MaterialHLSLTree.h b/Engine/Source/Runtime/Engine/Public/MaterialHLSLTree.h
--- a/Engine/Source/Runtime/Engine/Public/MaterialHLSLTree.h
+++ b/Engine/Source/Runtime/Engine/Public/MaterialHLSLTree.h
@@ -205,6 +205,8 @@ enum class EExternalInput : uint8
 
 	IsOrthographic,
 
+	AOMask,
+
 	Num,
 };
 static constexpr int32 MaxNumTexCoords = 8;
@@ -343,6 +345,103 @@ public:
 	virtual void EmitValueShader(FEmitContext& Context, FEmitScope& Scope, const FRequestedType& RequestedType, FEmitValueShaderResult& OutResult) const override;
 };
 
+class FExpressionSpeedTree : public FExpression
+{
+public:
+	const FExpression* GeometryExpression;
+	const FExpression* WindExpression;
+	const FExpression* LODExpression;
+	const FExpression* ExtraBendExpression;
+	const float BillboardThreshold;
+	const bool bExtraBend;
+	const bool bAccurateWind;
+	const bool bPreviousFrame;
+
+	FExpressionSpeedTree(const FExpression* InGeometryExpression, const FExpression* InWindExpression, const FExpression* InLODExpression, const FExpression* InExtraBendExpression, bool bInExtraBend, bool bInAccurateWind, float InBillboardThreshold, bool bInPreviousFrame)
+		: GeometryExpression(InGeometryExpression)
+		, WindExpression(InWindExpression)
+		, LODExpression(InLODExpression)
+		, ExtraBendExpression(InExtraBendExpression)
+		, BillboardThreshold(InBillboardThreshold)
+		, bExtraBend(bInExtraBend)
+		, bAccurateWind(bInAccurateWind)
+		, bPreviousFrame(bInPreviousFrame)
+	{}
+
+	virtual bool PrepareValue(FEmitContext& Context, FEmitScope& Scope, const FRequestedType& RequestedType, FPrepareValueResult& OutResult) const override;
+	virtual void EmitValueShader(FEmitContext& Context, FEmitScope& Scope, const FRequestedType& RequestedType, FEmitValueShaderResult& OutResult) const override;
+	virtual const FExpression* ComputePreviousFrame(FTree& Tree, const FRequestedType& RequestedType) const override;
+};
+
+class FExpressionDecalMipmapLevel : public FExpression
+{
+public:
+	const FExpression* TextureSizeExpression;
+
+	FExpressionDecalMipmapLevel(const FExpression* InTextureSizeExpression)
+		: TextureSizeExpression(InTextureSizeExpression)
+	{}
+
+	virtual bool PrepareValue(FEmitContext& Context, FEmitScope& Scope, const FRequestedType& RequestedType, FPrepareValueResult& OutResult) const override;
+	virtual void EmitValueShader(FEmitContext& Context, FEmitScope& Scope, const FRequestedType& RequestedType, FEmitValueShaderResult& OutResult) const override;
+};
+
+class FExpressionDBufferTexture : public FExpression
+{
+public:
+	const FExpression* UVExpression;
+	const uint8 DBufferTextureID;
+
+	FExpressionDBufferTexture(const FExpression* InUVExpression, uint8 InDBufferTextureID)
+		: UVExpression(InUVExpression)
+		, DBufferTextureID(InDBufferTextureID)
+	{}
+
+	virtual bool PrepareValue(FEmitContext& Context, FEmitScope& Scope, const FRequestedType& RequestedType, FPrepareValueResult& OutResult) const override;
+	virtual void EmitValueShader(FEmitContext& Context, FEmitScope& Scope, const FRequestedType& RequestedType, FEmitValueShaderResult& OutResult) const override;
+};
+
+class FExpressionSobolFunction : public FExpression
+{
+public:
+	const FExpression* CellExpression;		// optional, nullptr or ignored when Temporal is true
+	const FExpression* IndexExpression;
+	const FExpression* SeedExpression;
+	bool bTemporal;
+
+	FExpressionSobolFunction(const FExpression* InCellExpression, const FExpression* InIndexExpression, const FExpression* InSeedExpression, bool bInTemporal)
+		: CellExpression(InCellExpression)
+		, IndexExpression(InIndexExpression)
+		, SeedExpression(InSeedExpression)
+		, bTemporal(bInTemporal)
+	{}
+
+	virtual bool PrepareValue(FEmitContext& Context, FEmitScope& Scope, const FRequestedType& RequestedType, FPrepareValueResult& OutResult) const override;
+	virtual void EmitValueShader(FEmitContext& Context, FEmitScope& Scope, const FRequestedType& RequestedType, FEmitValueShaderResult& OutResult) const override;
+};
+
+// Expression for a single float at Index from CustomPrimitiveData, out of range returns 0
+class FExpressionCustomPrimitiveDataFunction : public FExpression
+{
+public:
+	const uint8 Index;
+
+	FExpressionCustomPrimitiveDataFunction(uint8 InIndex)
+		: Index(InIndex)
+	{}
+
+	virtual bool PrepareValue(FEmitContext& Context, FEmitScope& Scope, const FRequestedType& RequestedType, FPrepareValueResult& OutResult) const override;
+	virtual void EmitValueShader(FEmitContext& Context, FEmitScope& Scope, const FRequestedType& RequestedType, FEmitValueShaderResult& OutResult) const override;
+};
+
+class FExpressionAOMaskFunction : public FExpressionForward
+{
+public:
+	FExpressionAOMaskFunction(const FExpression* InExpression) : FExpressionForward(InExpression) {}
+
+	virtual bool PrepareValue(FEmitContext& Context, FEmitScope& Scope, const FRequestedType& RequestedType, FPrepareValueResult& OutResult) const override;
+};
+
 class FExpressionDepthOfFieldFunction : public FExpression
 {
 public:
