Fixes XR incorrectly using the desktop setting for screen percentage mode when a separate VR/stereo setting exists.

diff --git a/Engine/Source/Runtime/Engine/Private/GameViewportClient.cpp b/Engine/Source/Runtime/Engine/Private/GameViewportClient.cpp
--- a/Engine/Source/Runtime/Engine/Private/GameViewportClient.cpp
+++ b/Engine/Source/Runtime/Engine/Private/GameViewportClient.cpp
@@ -1707,6 +1707,9 @@ void UGameViewportClient::Draw(FViewport* InViewport, FCanvas* SceneCanvas)
 
 	FinalizeViews(&ViewFamily, PlayerViewMap);
 
+	// Update engine flags stored by viewport.
+	EngineShowFlags = ViewFamily.EngineShowFlags;
+
 	// Update level streaming.
 	MyWorld->UpdateLevelStreaming();
 
