DemoRecording:
* Fix CanSkipUpdate when demorecording is enabled alongside an Iris netdriver
* PushModel is not enabled to generic NetDrivers when Iris is enabled in another net driver.

diff --git a/Engine/Source/Runtime/Engine/Private/DataReplication.cpp b/Engine/Source/Runtime/Engine/Private/DataReplication.cpp
--- a/Engine/Source/Runtime/Engine/Private/DataReplication.cpp
+++ b/Engine/Source/Runtime/Engine/Private/DataReplication.cpp
@@ -1754,6 +1754,9 @@ bool FObjectReplicator::CanSkipUpdate(FReplicationFlags RepFlags)
 
 	bool bCanSkip = true;
 
+	// Is the pushmodel handle properly assigned.
+	bCanSkip = bCanSkip && RepChangelistState.HasValidPushModelHandle();
+
 	// Have the RepFlags changed ?
 	bCanSkip = bCanSkip && SendingRepState.RepFlags.Value == RepFlags.Value; 
 
diff --git a/Engine/Source/Runtime/Engine/Private/RepLayout.cpp b/Engine/Source/Runtime/Engine/Private/RepLayout.cpp
--- a/Engine/Source/Runtime/Engine/Private/RepLayout.cpp
+++ b/Engine/Source/Runtime/Engine/Private/RepLayout.cpp
@@ -1182,6 +1182,11 @@ bool FRepChangelistState::HasAnyDirtyProperties() const
 {
 	return UEPushModelPrivate::DoesHaveDirtyPropertiesOrRecentlyCollectedGarbage(PushModelObjectHandle);
 }
+
+bool FRepChangelistState::HasValidPushModelHandle() const
+{
+	return PushModelObjectHandle.IsValid();
+}
 #endif
 
 void FRepChangelistState::CountBytes(FArchive& Ar) const
diff --git a/Engine/Source/Runtime/Engine/Public/Net/RepLayout.h b/Engine/Source/Runtime/Engine/Public/Net/RepLayout.h
--- a/Engine/Source/Runtime/Engine/Public/Net/RepLayout.h
+++ b/Engine/Source/Runtime/Engine/Public/Net/RepLayout.h
@@ -422,6 +422,8 @@ public:
 
 	bool HasAnyDirtyProperties() const;
 
+	bool HasValidPushModelHandle() const;
+
 private:
 	const UEPushModelPrivate::FPushModelPerNetDriverHandle PushModelObjectHandle;
 #endif
