World Partition: deprecated FWorldPartitionLevelHelper::LoadActors and added new version with FWorldPartitionLevelHelper::FLoadActorsParams as input parameter.

diff --git a/Engine/Source/Runtime/Engine/Private/WorldPartition/WorldPartitionLevelHelper.cpp b/Engine/Source/Runtime/Engine/Private/WorldPartition/WorldPartitionLevelHelper.cpp
--- a/Engine/Source/Runtime/Engine/Private/WorldPartition/WorldPartitionLevelHelper.cpp
+++ b/Engine/Source/Runtime/Engine/Private/WorldPartition/WorldPartitionLevelHelper.cpp
@@ -405,9 +405,23 @@ ULevel* FWorldPartitionLevelHelper::CreateEmptyLevelForRuntimeCell(const UWorldP
 	return NewLevel;
 }
 
-bool FWorldPartitionLevelHelper::LoadActors(UWorld* InOuterWorld, ULevel* InDestLevel, TArrayView<FWorldPartitionRuntimeCellObjectMapping> InActorPackages, FWorldPartitionLevelHelper::FPackageReferencer& InPackageReferencer, TFunction<void(bool)> InCompletionCallback, bool bInLoadAsync, FLinkerInstancingContext InstancingContext)
+bool FWorldPartitionLevelHelper::LoadActors(UWorld* InOuterWorld, ULevel* InDestLevel, TArrayView<FWorldPartitionRuntimeCellObjectMapping> InActorPackages, FWorldPartitionLevelHelper::FPackageReferencer& InPackageReferencer, TFunction<void(bool)> InCompletionCallback, bool bInLoadAsync, FLinkerInstancingContext InInstancingContext)
 {
-	UPackage* DestPackage = InDestLevel ? InDestLevel->GetPackage() : nullptr;
+	FLoadActorsParams Params = FLoadActorsParams()
+		.SetOuterWorld(InOuterWorld)
+		.SetDestLevel(InDestLevel)
+		.SetActorPackages(InActorPackages)
+		.SetPackageReferencer(&InPackageReferencer)
+		.SetCompletionCallback(InCompletionCallback)
+		.SetLoadAsync(bInLoadAsync)
+		.SetInstancingContext(InInstancingContext);
+
+	return LoadActors(Params);
+}
+
+bool FWorldPartitionLevelHelper::LoadActors(const FLoadActorsParams& InParams)
+{
+	UPackage* DestPackage = InParams.DestLevel ? InParams.DestLevel->GetPackage() : nullptr;
 	FString ShortLevelPackageName = DestPackage? FPackageName::GetShortName(DestPackage->GetFName()) : FString();
 
 	struct FLoadProgress
@@ -419,20 +433,20 @@ bool FWorldPartitionLevelHelper::LoadActors(UWorld* InOuterWorld, ULevel* InDest
 
 	// Actors to load
 	TArray<FWorldPartitionRuntimeCellObjectMapping*> ActorPackages;
-	ActorPackages.Reserve(InActorPackages.Num());
+	ActorPackages.Reserve(InParams.ActorPackages.Num());
 
 	TMap<FActorContainerID, FLinkerInstancingContext> LinkerInstancingContexts;
 	// Add Main container context
-	LinkerInstancingContexts.Add(FActorContainerID::GetMainContainerID(), MoveTemp(InstancingContext));
+	LinkerInstancingContexts.Add(FActorContainerID::GetMainContainerID(), MoveTemp(InParams.InstancingContext));
 			
-	for (FWorldPartitionRuntimeCellObjectMapping& PackageObjectMapping : InActorPackages)
+	for (FWorldPartitionRuntimeCellObjectMapping& PackageObjectMapping : InParams.ActorPackages)
 	{
 		FLinkerInstancingContext* Context = LinkerInstancingContexts.Find(PackageObjectMapping.ContainerID);
 		if (!Context)
 		{
 			check(!PackageObjectMapping.ContainerID.IsMainContainer());
 		
-			const FString DestLevelPackageName = InDestLevel ? InDestLevel->GetPackage()->GetName() : FString();
+			const FString DestLevelPackageName = InParams.DestLevel ? InParams.DestLevel->GetPackage()->GetName() : FString();
 			const FName ContainerPackageInstanceName(GetContainerPackage(PackageObjectMapping.ContainerID, PackageObjectMapping.ContainerPackage.ToString(), DestLevelPackageName));
 
 			FLinkerInstancingContext& NewContext = LinkerInstancingContexts.Add(PackageObjectMapping.ContainerID);
@@ -466,7 +480,7 @@ bool FWorldPartitionLevelHelper::LoadActors(UWorld* InOuterWorld, ULevel* InDest
 
 	for (FWorldPartitionRuntimeCellObjectMapping* PackageObjectMapping : ActorPackages)
 	{
-		FLoadPackageAsyncDelegate CompletionCallback = FLoadPackageAsyncDelegate::CreateLambda([LoadProgress, PackageObjectMapping, &InPackageReferencer, InOuterWorld, InDestLevel, InCompletionCallback](const FName& LoadedPackageName, UPackage* LoadedPackage, EAsyncLoadingResult::Type Result)
+		FLoadPackageAsyncDelegate CompletionCallback = FLoadPackageAsyncDelegate::CreateLambda([LoadProgress, PackageObjectMapping, PackageReferencer = InParams.PackageReferencer, OuterWorld = InParams.OuterWorld, DestLevel = InParams.DestLevel, CompletionCallback = InParams.CompletionCallback](const FName& LoadedPackageName, UPackage* LoadedPackage, EAsyncLoadingResult::Type Result)
 		{
 			const FName ActorName = *FPaths::GetExtension(PackageObjectMapping->Path.ToString());
 			check(LoadProgress->NumPendingLoadRequests);
@@ -475,7 +489,7 @@ bool FWorldPartitionLevelHelper::LoadActors(UWorld* InOuterWorld, ULevel* InDest
 			// In PIE, we make sure to clear RF_Standalone flag on objects in external packages (UMetaData) 
 			// This guarantees that external packages of actors that are destroyed during the PIE session will
 			// properly get GC'ed and will allow future edits/modifications of OFPA actors.
-			if (LoadedPackage && InDestLevel && InDestLevel->GetPackage()->HasAnyPackageFlags(PKG_PlayInEditor))
+			if (LoadedPackage && DestLevel && DestLevel->GetPackage()->HasAnyPackageFlags(PKG_PlayInEditor))
 			{
 				ForEachObjectWithPackage(LoadedPackage, [](UObject* Object)
 				{
@@ -488,7 +502,7 @@ bool FWorldPartitionLevelHelper::LoadActors(UWorld* InOuterWorld, ULevel* InDest
 
 			if (Actor)
 			{
-				const UWorld* ContainerWorld = PackageObjectMapping->ContainerID.IsMainContainer() ? InOuterWorld : Actor->GetTypedOuter<UWorld>();
+				const UWorld* ContainerWorld = PackageObjectMapping->ContainerID.IsMainContainer() ? OuterWorld : Actor->GetTypedOuter<UWorld>();
 				
 				TOptional<FName> SrcActorFolderPath;
 
@@ -500,7 +514,7 @@ bool FWorldPartitionLevelHelper::LoadActors(UWorld* InOuterWorld, ULevel* InDest
 						ContainerWorld->PersistentLevel->bFixupActorFoldersAtLoad = false;
 						ContainerWorld->PersistentLevel->FixupActorFolders();
 					}
-				
+
 					// Since actor's level doesn't necessarily uses actor folders, access Folder Guid directly
 					const bool bDirectAccess = true;
 					const FGuid ActorFolderGuid = Actor->GetFolderGuid(bDirectAccess);
@@ -514,24 +528,24 @@ bool FWorldPartitionLevelHelper::LoadActors(UWorld* InOuterWorld, ULevel* InDest
 				if (!PackageObjectMapping->ContainerID.IsMainContainer())
 				{					
 					// Add Cache handle on world so it gets unloaded properly
-					InPackageReferencer.AddReference(ContainerWorld->GetPackage());
+					PackageReferencer->AddReference(ContainerWorld->GetPackage());
 										
 					// We only care about the source paths here
 					FString SourceWorldPath, DummyUnusedPath;
 					// Verify that it is indeed an instanced world
 					verify(ContainerWorld->GetSoftObjectPathMapping(SourceWorldPath, DummyUnusedPath));
 					FString SourceOuterWorldPath;
-					InOuterWorld->GetSoftObjectPathMapping(SourceOuterWorldPath, DummyUnusedPath);
+					OuterWorld->GetSoftObjectPathMapping(SourceOuterWorldPath, DummyUnusedPath);
 
 					// Rename through UObject to avoid changing Actor's external packaging and folder properties
-					Actor->UObject::Rename(*FString::Printf(TEXT("%s_%s"), *Actor->GetName(), *PackageObjectMapping->ContainerID.ToShortString()), InDestLevel, REN_NonTransactional | REN_ForceNoResetLoaders | REN_DoNotDirty | REN_DontCreateRedirectors);
+					Actor->UObject::Rename(*FString::Printf(TEXT("%s_%s"), *Actor->GetName(), *PackageObjectMapping->ContainerID.ToShortString()), DestLevel, REN_NonTransactional | REN_ForceNoResetLoaders | REN_DoNotDirty | REN_DontCreateRedirectors);
 
 					// Handle child actors
-					Actor->ForEachComponent<UChildActorComponent>(true, [InDestLevel, PackageObjectMapping](UChildActorComponent* ChildActorComponent)
+					Actor->ForEachComponent<UChildActorComponent>(true, [DestLevel = DestLevel, PackageObjectMapping](UChildActorComponent* ChildActorComponent)
 					{
 						if (AActor* ChildActor = ChildActorComponent->GetChildActor())
 						{
-							ChildActor->UObject::Rename(*FString::Printf(TEXT("%s_%s"), *ChildActor->GetName(), *PackageObjectMapping->ContainerID.ToShortString()), InDestLevel, REN_NonTransactional | REN_ForceNoResetLoaders | REN_DoNotDirty | REN_DontCreateRedirectors);
+							ChildActor->UObject::Rename(*FString::Printf(TEXT("%s_%s"), *ChildActor->GetName(), *PackageObjectMapping->ContainerID.ToShortString()), DestLevel, REN_NonTransactional | REN_ForceNoResetLoaders | REN_DoNotDirty | REN_DontCreateRedirectors);
 						}
 					});
 					
@@ -556,7 +570,7 @@ bool FWorldPartitionLevelHelper::LoadActors(UWorld* InOuterWorld, ULevel* InDest
 					{
 						if (!Value.IsNull() && Value.GetAssetPathString().Equals(SourceWorldPath, ESearchCase::IgnoreCase))
 						{
-							InOuterWorld->GetWorldPartition()->ConvertContainerPathToEditorPath(PackageObjectMapping->ContainerID, FSoftObjectPath(Value), Value);
+							OuterWorld->GetWorldPartition()->ConvertContainerPathToEditorPath(PackageObjectMapping->ContainerID, FSoftObjectPath(Value), Value);
 						}
 					});
 					FixupArchive.Fixup(Actor);
@@ -567,26 +581,26 @@ bool FWorldPartitionLevelHelper::LoadActors(UWorld* InOuterWorld, ULevel* InDest
 					}
 				}
 
-				if (InDestLevel)
+				if (DestLevel)
 				{
 					// Propagate resolved actor folder path
-					check(!InDestLevel->IsUsingActorFolders());
+					check(!DestLevel->IsUsingActorFolders());
 					if (SrcActorFolderPath.IsSet())
 					{
 						Actor->SetFolderPath(*SrcActorFolderPath);
 					}
 
 					check(Actor->IsPackageExternal());
-					InDestLevel->Actors.Add(Actor);
-					checkf(Actor->GetLevel() == InDestLevel, TEXT("Levels mismatch, got : %s, expected: %s\nActor: %s\nActorFullName: %s\nActorPackage: %s"), *InDestLevel->GetFullName(), *Actor->GetLevel()->GetFullName(), *Actor->GetActorNameOrLabel(), *Actor->GetFullName(), *Actor->GetPackage()->GetFullName());
+					DestLevel->Actors.Add(Actor);
+					checkf(Actor->GetLevel() == DestLevel, TEXT("Levels mismatch, got : %s, expected: %s\nActor: %s\nActorFullName: %s\nActorPackage: %s"), *DestLevel->GetFullName(), *Actor->GetLevel()->GetFullName(), *Actor->GetActorNameOrLabel(), *Actor->GetFullName(), *Actor->GetPackage()->GetFullName());
 
 					// Handle child actors
-					Actor->ForEachComponent<UChildActorComponent>(true, [InDestLevel](UChildActorComponent* ChildActorComponent)
+					Actor->ForEachComponent<UChildActorComponent>(true, [DestLevel = DestLevel](UChildActorComponent* ChildActorComponent)
 					{
 						if (AActor* ChildActor = ChildActorComponent->GetChildActor())
 						{
-							InDestLevel->Actors.Add(ChildActor);
-							check(ChildActor->GetLevel() == InDestLevel);
+							DestLevel->Actors.Add(ChildActor);
+							check(ChildActor->GetLevel() == DestLevel);
 						}
 					});
 				}
@@ -615,7 +629,7 @@ bool FWorldPartitionLevelHelper::LoadActors(UWorld* InOuterWorld, ULevel* InDest
 
 			if (!LoadProgress->NumPendingLoadRequests)
 			{
-				InCompletionCallback(!LoadProgress->NumFailedLoadedRequests);
+				CompletionCallback(!LoadProgress->NumFailedLoadedRequests);
 			}
 		});
 
@@ -623,7 +637,7 @@ bool FWorldPartitionLevelHelper::LoadActors(UWorld* InOuterWorld, ULevel* InDest
 		const FLinkerInstancingContext& ContainerInstancingContext = LinkerInstancingContexts.FindChecked(PackageObjectMapping->ContainerID);
 		FName PackageName = ContainerInstancingContext.RemapPackage(PackageToLoad);
 
-		if (bInLoadAsync)
+		if (InParams.bLoadAsync)
 		{
 			FPackagePath PackagePath = FPackagePath::FromPackageNameChecked(PackageToLoad);
 
diff --git a/Engine/Source/Runtime/Engine/Private/WorldPartition/WorldPartitionLevelStreamingDynamic.cpp b/Engine/Source/Runtime/Engine/Private/WorldPartition/WorldPartitionLevelStreamingDynamic.cpp
--- a/Engine/Source/Runtime/Engine/Private/WorldPartition/WorldPartitionLevelStreamingDynamic.cpp
+++ b/Engine/Source/Runtime/Engine/Private/WorldPartition/WorldPartitionLevelStreamingDynamic.cpp
@@ -440,7 +440,16 @@ bool UWorldPartitionLevelStreamingDynamic::IssueLoadRequests()
 	// Load saved actors
 	if (ChildPackagesToLoad.Num())
 	{
-		FWorldPartitionLevelHelper::LoadActors(GetStreamingWorld(), RuntimeLevel, ChildPackagesToLoad, PackageReferencer, FinalizeLoading, World->IsGameWorld(), MoveTemp(InstancingContext));
+		FWorldPartitionLevelHelper::FLoadActorsParams Params = FWorldPartitionLevelHelper::FLoadActorsParams()
+			.SetOuterWorld(GetStreamingWorld())
+			.SetDestLevel(RuntimeLevel)
+			.SetActorPackages(ChildPackagesToLoad)
+			.SetPackageReferencer(&PackageReferencer)
+			.SetCompletionCallback(FinalizeLoading)
+			.SetLoadAsync(World->IsGameWorld())
+			.SetInstancingContext(MoveTemp(InstancingContext));
+
+		FWorldPartitionLevelHelper::LoadActors(Params);
 	}
 	else
 	{
diff --git a/Engine/Source/Runtime/Engine/Private/WorldPartition/WorldPartitionRuntimeLevelStreamingCell.cpp b/Engine/Source/Runtime/Engine/Private/WorldPartition/WorldPartitionRuntimeLevelStreamingCell.cpp
--- a/Engine/Source/Runtime/Engine/Private/WorldPartition/WorldPartitionRuntimeLevelStreamingCell.cpp
+++ b/Engine/Source/Runtime/Engine/Private/WorldPartition/WorldPartitionRuntimeLevelStreamingCell.cpp
@@ -284,14 +284,20 @@ bool UWorldPartitionRuntimeLevelStreamingCell::PopulateGeneratorPackageForCook(T
 
 	if (GetActorCount() > 0)
 	{
-		FWorldPartitionLevelHelper::FPackageReferencer PackageReferencer;
-		const bool bLoadAsync = false;
 		UWorld* OuterWorld = GetOuterWorld();
 		UWorldPartition* WorldPartition = OuterWorld->GetWorldPartition();
 
-		// Don't do SoftObjectPath remapping for PersistentLevel actors because references can end up in different cells
-		const bool bSoftObjectRemappingEnabled = false;
-		verify(FWorldPartitionLevelHelper::LoadActors(OuterWorld, nullptr, Packages, PackageReferencer, [](bool) {}, bLoadAsync, FLinkerInstancingContext(bSoftObjectRemappingEnabled)));
+		FWorldPartitionLevelHelper::FPackageReferencer PackageReferencer;
+		FWorldPartitionLevelHelper::FLoadActorsParams Params = FWorldPartitionLevelHelper::FLoadActorsParams()
+			.SetOuterWorld(OuterWorld)
+			.SetDestLevel(nullptr)
+			.SetActorPackages(Packages)
+			.SetPackageReferencer(&PackageReferencer)
+			.SetCompletionCallback([](bool) {})
+			.SetLoadAsync(false)
+			.SetInstancingContext(FLinkerInstancingContext(false)); // Don't do SoftObjectPath remapping for PersistentLevel actors because references can end up in different cells
+
+		verify(FWorldPartitionLevelHelper::LoadActors(Params));
 
 		FWorldPartitionLevelHelper::MoveExternalActorsToLevel(Packages, OuterWorld->PersistentLevel, OutModifiedPackages);
 
@@ -346,11 +352,16 @@ bool UWorldPartitionRuntimeLevelStreamingCell::PopulateGeneratedPackageForCook(U
 
 		// Load cell Actors
 		FWorldPartitionLevelHelper::FPackageReferencer PackageReferencer;
-		const bool bLoadAsync = false;
-
-		// Don't do SoftObjectPath remapping for PersistentLevel actors because references can end up in different cells
-		const bool bSoftObjectRemappingEnabled = false;
-		verify(FWorldPartitionLevelHelper::LoadActors(OuterWorld, nullptr, Packages, PackageReferencer, [](bool) {}, bLoadAsync, FLinkerInstancingContext(bSoftObjectRemappingEnabled)));
+		FWorldPartitionLevelHelper::FLoadActorsParams Params = FWorldPartitionLevelHelper::FLoadActorsParams()
+			.SetOuterWorld(OuterWorld)
+			.SetDestLevel(nullptr)
+			.SetActorPackages(Packages)
+			.SetPackageReferencer(&PackageReferencer)
+			.SetCompletionCallback([](bool) {})
+			.SetLoadAsync(false)
+			.SetInstancingContext(FLinkerInstancingContext(false)); // Don't do SoftObjectPath remapping for PersistentLevel actors because references can end up in different cells
+
+		verify(FWorldPartitionLevelHelper::LoadActors(Params));
 
 		// Create a level and move these actors in it
 		ULevel* NewLevel = FWorldPartitionLevelHelper::CreateEmptyLevelForRuntimeCell(this, OuterWorld, LevelStreaming->GetWorldAsset().ToString(), InPackage);
diff --git a/Engine/Source/Runtime/Engine/Public/WorldPartition/WorldPartitionLevelHelper.h b/Engine/Source/Runtime/Engine/Public/WorldPartition/WorldPartitionLevelHelper.h
--- a/Engine/Source/Runtime/Engine/Public/WorldPartition/WorldPartitionLevelHelper.h
+++ b/Engine/Source/Runtime/Engine/Public/WorldPartition/WorldPartitionLevelHelper.h
@@ -37,7 +37,37 @@ public:
 	static void MoveExternalActorsToLevel(const TArray<FWorldPartitionRuntimeCellObjectMapping>& InChildPackages, ULevel* InLevel, TArray<UPackage*>& OutModifiedPackages);
 	static void RemapLevelSoftObjectPaths(ULevel* InLevel, UWorldPartition* InWorldPartition);
 	
-	static bool LoadActors(UWorld* InOuterWorld, ULevel* InDestLevel, TArrayView<FWorldPartitionRuntimeCellObjectMapping> InActorPackages, FPackageReferencer& InPackageReferencer, TFunction<void(bool)> InCompletionCallback, bool bInLoadAsync, FLinkerInstancingContext InOutInstancingContext);
+	UE_DEPRECATED(5.4, "LoadActors is deprecated, LoadActors with FLoadActorsParams should be used instead.")
+	static bool LoadActors(UWorld* InOuterWorld, ULevel* InDestLevel, TArrayView<FWorldPartitionRuntimeCellObjectMapping> InActorPackages, FPackageReferencer& InPackageReferencer, TFunction<void(bool)> InCompletionCallback, bool bInLoadAsync, FLinkerInstancingContext InInstancingContext);
+
+	/* Struct of optional parameters passed to foreach actordesc functions. */
+	struct FLoadActorsParams
+	{
+		FLoadActorsParams()
+			: OuterWorld(nullptr)
+			, DestLevel(nullptr)
+			, PackageReferencer(nullptr)
+			, bLoadAsync(false)
+		{}
+
+		UWorld* OuterWorld;
+		ULevel* DestLevel;
+		TArrayView<FWorldPartitionRuntimeCellObjectMapping> ActorPackages;
+		FPackageReferencer* PackageReferencer;
+		TFunction<void(bool)> CompletionCallback;
+		bool bLoadAsync;
+		mutable FLinkerInstancingContext InstancingContext;
+
+		FLoadActorsParams& SetOuterWorld(UWorld* InOuterWorld) { OuterWorld = InOuterWorld; return *this; }
+		FLoadActorsParams& SetDestLevel(ULevel* InDestLevel) { DestLevel = InDestLevel; return *this; }
+		FLoadActorsParams& SetActorPackages(TArrayView<FWorldPartitionRuntimeCellObjectMapping> InActorPackages) { ActorPackages = InActorPackages; return *this; }
+		FLoadActorsParams& SetPackageReferencer(FPackageReferencer* InPackageReferencer) { PackageReferencer = InPackageReferencer; return *this; }
+		FLoadActorsParams& SetCompletionCallback(TFunction<void(bool)> InCompletionCallback) { CompletionCallback = InCompletionCallback; return *this; }
+		FLoadActorsParams& SetLoadAsync(bool bInLoadAsync) { bLoadAsync = bInLoadAsync; return *this; }
+		FLoadActorsParams& SetInstancingContext(FLinkerInstancingContext InInstancingContext) { InstancingContext = InInstancingContext; return *this; }
+	};
+
+	static bool LoadActors(const FLoadActorsParams& InParams);
 	
 	static FSoftObjectPath RemapActorPath(const FActorContainerID& InContainerID, const FString& SourceWorldName, const FSoftObjectPath& InActorPath);
 
