Fix Incremental Build Minimal / NonUnity UnrealEditor Win64 builds

diff --git a/Engine/Source/Runtime/Engine/Public/WorldPartition/WorldPartitionLevelHelper.h b/Engine/Source/Runtime/Engine/Public/WorldPartition/WorldPartitionLevelHelper.h
--- a/Engine/Source/Runtime/Engine/Public/WorldPartition/WorldPartitionLevelHelper.h
+++ b/Engine/Source/Runtime/Engine/Public/WorldPartition/WorldPartitionLevelHelper.h
@@ -10,6 +10,7 @@
 #pragma once
 
 #include "Engine/World.h"
+#include "UObject/LinkerInstancingContext.h"
 #include "WorldPartition/WorldPartitionRuntimeCell.h"
 
 class FWorldPartitionPackageHelper;
