Fixing tooltip typo
#rd trivial

diff --git a/Engine/Source/Runtime/Engine/Classes/Sound/ReverbEffect.h b/Engine/Source/Runtime/Engine/Classes/Sound/ReverbEffect.h
--- a/Engine/Source/Runtime/Engine/Classes/Sound/ReverbEffect.h
+++ b/Engine/Source/Runtime/Engine/Classes/Sound/ReverbEffect.h
@@ -54,7 +54,7 @@ class UReverbEffect : public UObject
 	UPROPERTY(Category=LateReflections, meta=(ClampMin = "0.0", ClampMax = "1.0", EditCondition = "!bBypassLateReflections"), EditAnywhere)
 	float		AirAbsorptionGainHF;
 
-	/** Decay High Frequency Ratio - 0.1 < 0.83 < 2.0 - how much the quicker or slower the high frequencies decay relative to the lower frequencies. */
+	/** Decay High Frequency Ratio - 0.1 < 0.83 < 2.0 - how much quicker or slower the high frequencies decay relative to the lower frequencies. */
 	UPROPERTY(Category=LateReflections, meta=(ClampMin = "0.1", ClampMax = "2.0", EditCondition = "!bBypassLateReflections"), EditAnywhere)
 	float		DecayHFRatio;
 
