Use `GetCompressedData(Format)` to get the actual memory of a ForceInline sound wave.

When displaying ForceInline percentage in the DedbugView, use the Percentange calculation based on the ForceInlineBytes / TotalBytes. Then, figure out how to display that percentage in the "Bar" graph by doing an inverse calculation to figure out how many "chunks" it would occupy.

diff --git a/Engine/Source/Runtime/Engine/Private/AudioStreamingCache.cpp b/Engine/Source/Runtime/Engine/Private/AudioStreamingCache.cpp
--- a/Engine/Source/Runtime/Engine/Private/AudioStreamingCache.cpp
+++ b/Engine/Source/Runtime/Engine/Private/AudioStreamingCache.cpp
@@ -993,8 +993,18 @@ void FAudioChunkCache::ClearCache()
 
 void FAudioChunkCache::AddForceInlineSoundWave(const FSoundWaveProxyPtr& SoundWave)
 {
-	check(SoundWave->GetLoadingBehavior() == ESoundWaveLoadingBehavior::ForceInline);
-	ForceInlineMemoryCounterBytes += SoundWave->GetResourceSize();
+	check(SoundWave.IsValid());
+	ensureMsgf(SoundWave->GetLoadingBehavior() == ESoundWaveLoadingBehavior::ForceInline, 
+		TEXT("AudioStreamingCache::AddForceInlineSoundWave: Attempted to add SoundWave not set to ForceInline: %s"), *SoundWave->GetFName().ToString());
+
+	if (SoundWave->GetLoadingBehavior() != ESoundWaveLoadingBehavior::ForceInline)
+	{
+		return;
+	}
+
+	FName Format = SoundWave->GetRuntimeFormat();
+	FByteBulkData* Data = SoundWave->GetCompressedData(Format);
+	ForceInlineMemoryCounterBytes += Data ? Data->GetBulkDataSize() : 0;
 
 	const uint64 MemoryUsageBytes = MemoryCounterBytes + ForceInlineMemoryCounterBytes;
 	if (TrimCacheWhenOverBudgetCVar != 0 && MemoryUsageBytes > MemoryLimitBytes)
@@ -1015,8 +1025,18 @@ void FAudioChunkCache::AddForceInlineSoundWave(const FSoundWaveProxyPtr& SoundWa
 
 void FAudioChunkCache::RemoveForecInlineSoundWave(const FSoundWaveProxyPtr& SoundWave)
 {
-	check(SoundWave->GetLoadingBehavior() == ESoundWaveLoadingBehavior::ForceInline);
-	ForceInlineMemoryCounterBytes -= SoundWave->GetResourceSize();
+	check(SoundWave.IsValid());
+	ensureMsgf(SoundWave->GetLoadingBehavior() == ESoundWaveLoadingBehavior::ForceInline,
+		TEXT("AudioStreamingCache::RemoveForceInlineSoundWave: Attempted to remove SoundWave not set to ForceInline: %s"), *SoundWave->GetFName().ToString());
+
+	if (SoundWave->GetLoadingBehavior() != ESoundWaveLoadingBehavior::ForceInline)
+	{
+		return;
+	}
+
+	FName Format = SoundWave->GetRuntimeFormat();
+	FByteBulkData* Data = SoundWave->GetCompressedData(Format);
+	ForceInlineMemoryCounterBytes -= Data ? Data->GetBulkDataSize() : 0;
 }
 
 uint64 FAudioChunkCache::TrimMemory(uint64 BytesToFree, bool bInAllowRetainedChunkTrimming)
@@ -2313,9 +2333,6 @@ TPair<int, int> FAudioChunkCache::DebugDisplay(UWorld* World, FViewport* Viewpor
 	int32 NumTrimmed = 0;
 	int32 NumLoadInProgress = 0;
 	int32 NumOther = 0;
-	// approximate how many chunks the force inline memory is using for debug purposes
-	int32 NumForceInline = ForceInlineMemoryCounterBytes.Load() / MaxChunkSize;
-
 
 	for (int i = 0; i < ChunksInUse; ++i)
 	{
@@ -2399,13 +2416,18 @@ TPair<int, int> FAudioChunkCache::DebugDisplay(UWorld* World, FViewport* Viewpor
 		}
 	}
 
-	NumBytesCounter += ForceInlineMemoryCounterBytes.Load();
+	uint64 ForceInlineBytes = ForceInlineMemoryCounterBytes.Load();
+	NumBytesCounter += ForceInlineBytes;
 	// Convert to megabytes and print the total size:
 	const double NumMegabytesInUse = (double)NumBytesCounter / (1024 * 1024);
 	const double MaxCacheSizeMB = ((double)MemoryLimitBytes) / (1024 * 1024);
 
+	// calculate ForceInline bytes and percentage
+	const double NumMegabytesForceInline = (double)ForceInlineBytes / (1024 * 1024);
+	const float PercentageForceInline = NumBytesCounter > 0 ? (double)ForceInlineBytes / NumBytesCounter : 0;
+
 	FString CacheMemoryUsage = *FString::Printf(TEXT("Using: %.4f Megabytes (%lu bytes). Max Potential Usage: %.4f Megabytes."), 
-		NumMegabytesInUse, MemoryCounterBytes.Load() + ForceInlineMemoryCounterBytes.Load(), MaxCacheSizeMB);
+		NumMegabytesInUse, MemoryCounterBytes.Load() + ForceInlineBytes, MaxCacheSizeMB);
 
 	// We're going to align this horizontally with the number of elements right above it.
 	Canvas->DrawShadowedString(X, Y, *CacheMemoryUsage, UEngine::GetMediumFont(), FLinearColor::White);
@@ -2424,8 +2446,33 @@ TPair<int, int> FAudioChunkCache::DebugDisplay(UWorld* World, FViewport* Viewpor
 		+ NumLOD
 		+ NumTrimmed
 		+ NumLoadInProgress
-		+ NumOther
-		+ NumForceInline;
+		+ NumOther;
+
+
+	if (FMath::IsNearlyEqual(PercentageForceInline, 1.0))
+	{
+		// if the ForceInline Percentage is basically 1, then just set the "number of chunks" to a really big number.
+		// so everything else just gets zeroed out
+		NumChunks = UE_BIG_NUMBER;
+	}
+	else if (PercentageForceInline > 0.0)
+	{
+		// calculate the NumForceInline based on the percentage of memory used.
+		int32 NumForceInline = NumChunks * (PercentageForceInline / (1 - PercentageForceInline));
+
+		// derivation:
+		// NumChunks + NumForceInline = TotalNumChunks
+		// NumForceInline = TotalNumChunks * PercentageForceInline
+		// 
+		// NumChunks = TotalNumChunks * (1 - PercentageForceInline)
+		// TotalNumChunks = NumChunks / (1 - PercentageForceInline)
+		// 
+		// - using substitution with the above
+		//  NumForceInline = NumChunks * PercentageForceInline / (1 - PercentageForceInline)
+
+		// Add the newly calculated ForceInline to the mix 
+		NumChunks += NumForceInline;
+	}
 
 	if (NumChunks == 0)
 	{
@@ -2450,7 +2497,6 @@ TPair<int, int> FAudioChunkCache::DebugDisplay(UWorld* World, FViewport* Viewpor
 	const float PercentageTrimmed = (NumTrimmed / NumChunks);
 	const float PercentageLoadInProgress = (NumLoadInProgress / NumChunks);
 	const float PercentageOther = (NumOther / NumChunks);
-	const float PercentageForceInline = (NumForceInline / NumChunks);
 
 	const int32 BarWidthRetainedAndPlaying = PercentageRetainedAndPlaying * BarWidth;
 	const int32 BarWidthRetained = PercentageRetained * BarWidth;
@@ -2492,7 +2538,7 @@ TPair<int, int> FAudioChunkCache::DebugDisplay(UWorld* World, FViewport* Viewpor
 	Canvas->DrawShadowedString(X, Y, *TempString, UEngine::GetSmallFont(), ColorLoadInProgress);
 	Y += 15;
 
-	TempString = *FString::Printf(TEXT("Force Inline: %.2f %%"), 100.f * PercentageForceInline);
+	TempString = *FString::Printf(TEXT("Force Inline: %.2f %% (%.2f MB)"), 100.f * PercentageForceInline, NumMegabytesForceInline);
 	Canvas->DrawShadowedString(X, Y, *TempString, UEngine::GetSmallFont(), ColorForceInline);
 	Y += 25;
 
