Refactor in preparation of supporting Data Layer Instances saved in their own external package.
- Adapted existing code to propertly support external package on object not directly outered to the level.
- Generalized handling of external objects when renaming a world, resetting package loaders in ULevelStreamingLevelInstance, handling dirty external packages
- Implemented FExternalPackageHelper::CopyObjectsExternalPackageFilePathToClipboard
- Added optional PrimaryAssetDisplayName asset registry tag to allow for systems to use a display name (generalized version of ActorLabel tag)

diff --git a/Engine/Plugins/Editor/DataValidation/Source/DataValidation/Private/DataValidationChangelist.cpp b/Engine/Plugins/Editor/DataValidation/Source/DataValidation/Private/DataValidationChangelist.cpp
--- a/Engine/Plugins/Editor/DataValidation/Source/DataValidation/Private/DataValidationChangelist.cpp
+++ b/Engine/Plugins/Editor/DataValidation/Source/DataValidation/Private/DataValidationChangelist.cpp
@@ -56,6 +56,10 @@ FString UDataValidationChangelist::GetPrettyPackageName(const FName& InPackageNa
 		{
 			Assets[0].GetTagValue(NAME_ActorLabel, AssetName);
 		}
+		else if (Assets[0].FindTag(FPrimaryAssetId::PrimaryAssetDisplayNameTag))
+		{
+			Assets[0].GetTagValue(FPrimaryAssetId::PrimaryAssetDisplayNameTag, AssetName);
+		}
 		else
 		{
 			AssetName = Assets[0].AssetName.ToString();
diff --git a/Engine/Source/Editor/LevelEditor/Private/LevelEditorActions.cpp b/Engine/Source/Editor/LevelEditor/Private/LevelEditorActions.cpp
--- a/Engine/Source/Editor/LevelEditor/Private/LevelEditorActions.cpp
+++ b/Engine/Source/Editor/LevelEditor/Private/LevelEditorActions.cpp
@@ -46,6 +46,7 @@
 #include "WorldPartition/WorldPartitionRuntimeHash.h"
 #include "WorldPartition/IWorldPartitionEditorModule.h"
 #include "WorldBrowserModule.h"
+#include "ExternalPackageHelper.h"
 
 #include "Elements/Framework/TypedElementCommonActions.h"
 #include "Elements/Framework/TypedElementSelectionSet.h"
@@ -2929,21 +2930,12 @@ void FLevelEditorActionCallbacks::SnapObjectToView_Clicked()
 
 void FLevelEditorActionCallbacks::CopyActorFilePathtoClipboard_Clicked()
 {
-	FString Result;
-
+	TArray<const UObject*> Objects;
 	for (FSelectionIterator It(GEditor->GetSelectedActorIterator()); It; ++It)
 	{
-		const AActor* Actor = Cast<AActor>(*It);
-		const UPackage* Package = Actor->GetPackage();
-		const FString LocalFullPath(Package->GetLoadedPath().GetLocalFullPath());
-		Result += FPaths::ConvertRelativePathToFull(LocalFullPath);
-		Result += TEXT("\n");
-	}
-
-	if (Result.Len())
-	{
-		FPlatformApplicationMisc::ClipboardCopy(*Result);
+		Objects.Add(*It);
 	}
+	FExternalPackageHelper::CopyObjectsExternalPackageFilePathToClipboard(Objects);
 }
 
 void FLevelEditorActionCallbacks::SaveActor_Clicked()
diff --git a/Engine/Source/Editor/SourceControlWindows/Private/SSourceControlCommon.cpp b/Engine/Source/Editor/SourceControlWindows/Private/SSourceControlCommon.cpp
--- a/Engine/Source/Editor/SourceControlWindows/Private/SSourceControlCommon.cpp
+++ b/Engine/Source/Editor/SourceControlWindows/Private/SSourceControlCommon.cpp
@@ -68,11 +68,16 @@ static FString RetrieveAssetName(const FAssetData& InAssetData)
 
 	if (InAssetData.FindTag(NAME_ActorLabel))
 	{
-		FString ResultAssetName = TEXT("");
-
+		FString ResultAssetName;
 		InAssetData.GetTagValue(NAME_ActorLabel, ResultAssetName);
 		return ResultAssetName;
 	}
+	else if (InAssetData.FindTag(FPrimaryAssetId::PrimaryAssetDisplayNameTag))
+	{
+		FString ResultAssetName;
+		InAssetData.GetTagValue(FPrimaryAssetId::PrimaryAssetDisplayNameTag, ResultAssetName);
+		return ResultAssetName;
+	}
 	else if (InAssetData.AssetClassPath == UActorFolder::StaticClass()->GetClassPathName())
 	{
 		FString ActorFolderPath = UActorFolder::GetAssetRegistryInfoFromPackage(InAssetData.PackageName).GetDisplayName();
diff --git a/Engine/Source/Editor/UnrealEd/Private/FileHelpers.cpp b/Engine/Source/Editor/UnrealEd/Private/FileHelpers.cpp
--- a/Engine/Source/Editor/UnrealEd/Private/FileHelpers.cpp
+++ b/Engine/Source/Editor/UnrealEd/Private/FileHelpers.cpp
@@ -81,7 +81,6 @@
 #include "WorldPartition/LoaderAdapter/LoaderAdapterShape.h"
 #include "LevelInstance/LevelInstanceSubsystem.h"
 #include "PackageSourceControlHelper.h"
-#include "ActorFolder.h"
 #include "InterchangeManager.h"
 #include "SourceControlHelpers.h"
 #include "InterchangeProjectSettings.h"
@@ -3095,19 +3094,10 @@ EAutosaveContentPackagesResult::Type FEditorFileUtils::AutosaveMapEx(const FStri
 						&& FPackageName::IsValidLongPackageName(ExternalPackage->GetName(), /*bIncludeReadOnlyRoots=*/false))
 					{
 						// Don't try to save external packages that will get deleted
-						ForEachObjectWithPackage(ExternalPackage, [ExternalPackage, &ExternalPackagesToSave](UObject* Object)
+						if (IsValid(ExternalPackage->FindAssetInPackage()))
 						{
-							// @todo_ow: Find better way
-							if (Object->IsA<AActor>() || Object->IsA<UActorFolder>())
-							{
-								if (IsValid(Object))
-								{
-									ExternalPackagesToSave.Add(ExternalPackage);
-									return false;
-								}
-							}
-							return true;
-						}, false);
+							ExternalPackagesToSave.Add(ExternalPackage);
+						}
 					}
 				}
 
@@ -5255,20 +5245,7 @@ void FEditorFileUtils::GetDirtyWorldPackages(TArray<UPackage*>& OutDirtyPackages
 							// Skip unsaved packages containing only pending kill actors
 							if (ExternalPackage->HasAnyPackageFlags(PKG_NewlyCreated))
 							{
-								bActorPackageNeedsToSave = false;
-								ForEachObjectWithPackage(ExternalPackage, [&bActorPackageNeedsToSave](UObject* Object)
-								{
-									// @todo_ow: Find better way
-									if (Object->IsA<AActor>() || Object->IsA<UActorFolder>())
-									{
-										if (IsValid(Object))
-										{
-											bActorPackageNeedsToSave = true;
-											return false;
-										}
-									}
-									return true;
-								}, false);
+								bActorPackageNeedsToSave = IsValid(ExternalPackage->FindAssetInPackage());
 							}
 
 							// Filter out Actors that might be unsaved (/Temp folder)
diff --git a/Engine/Source/Editor/UnrealEd/Private/UnrealEdMisc.cpp b/Engine/Source/Editor/UnrealEd/Private/UnrealEdMisc.cpp
--- a/Engine/Source/Editor/UnrealEd/Private/UnrealEdMisc.cpp
+++ b/Engine/Source/Editor/UnrealEd/Private/UnrealEdMisc.cpp
@@ -1450,11 +1450,12 @@ void FUnrealEdMisc::OnMessageTokenActivated(const TSharedRef<IMessageToken>& Tok
 FText FUnrealEdMisc::OnGetAssetDataDisplayName(const FAssetData& InObject, const bool bFullPath)
 {
 	static FName NAME_ActorLabel("ActorLabel");
-	FString ActorLabel;
+	FString DisplayName;
 	TStringBuilder<FName::StringBufferSize> Buffer;
-	if (InObject.GetTagValue(NAME_ActorLabel, ActorLabel))
+	if (InObject.GetTagValue(NAME_ActorLabel, DisplayName) || 
+		InObject.GetTagValue(FPrimaryAssetId::PrimaryAssetDisplayNameTag, DisplayName))
 	{
-		Buffer << ActorLabel << TEXT(" (");
+		Buffer << DisplayName << TEXT(" (");
 		InObject.AppendObjectPath(Buffer);
 		Buffer << TEXT(")");
 		return FText::FromStringView(Buffer.ToView());
diff --git a/Engine/Source/Runtime/CoreUObject/Private/UObject/Obj.cpp b/Engine/Source/Runtime/CoreUObject/Private/UObject/Obj.cpp
--- a/Engine/Source/Runtime/CoreUObject/Private/UObject/Obj.cpp
+++ b/Engine/Source/Runtime/CoreUObject/Private/UObject/Obj.cpp
@@ -2149,6 +2149,7 @@ static void GetAssetRegistryTagsFromSearchableProperties(const UObject* Object,
 
 const FName FPrimaryAssetId::PrimaryAssetTypeTag(TEXT("PrimaryAssetType"));
 const FName FPrimaryAssetId::PrimaryAssetNameTag(TEXT("PrimaryAssetName"));
+const FName FPrimaryAssetId::PrimaryAssetDisplayNameTag(TEXT("PrimaryAssetDisplayName"));
 
 void UObject::GetAssetRegistryTags(TArray<FAssetRegistryTag>& OutTags) const
 {
diff --git a/Engine/Source/Runtime/CoreUObject/Public/UObject/PrimaryAssetId.h b/Engine/Source/Runtime/CoreUObject/Public/UObject/PrimaryAssetId.h
--- a/Engine/Source/Runtime/CoreUObject/Public/UObject/PrimaryAssetId.h
+++ b/Engine/Source/Runtime/CoreUObject/Public/UObject/PrimaryAssetId.h
@@ -113,6 +113,8 @@ struct FPrimaryAssetId
 	/** Static names to represent the AssetRegistry tags for the above data */
 	static COREUOBJECT_API const FName PrimaryAssetTypeTag;
 	static COREUOBJECT_API const FName PrimaryAssetNameTag;
+	/** AssetRegistry tag used to store primary asset display name (optional) */
+	static COREUOBJECT_API const FName PrimaryAssetDisplayNameTag;
 
 	FPrimaryAssetId() = default;
 	FPrimaryAssetId(FPrimaryAssetType InAssetType, FName InAssetName)
diff --git a/Engine/Source/Runtime/Engine/Private/ActorFolder.cpp b/Engine/Source/Runtime/Engine/Private/ActorFolder.cpp
--- a/Engine/Source/Runtime/Engine/Private/ActorFolder.cpp
+++ b/Engine/Source/Runtime/Engine/Private/ActorFolder.cpp
@@ -31,7 +31,7 @@ UActorFolder* UActorFolder::Create(ULevel* InLevel, const FString& InFolderLabel
 	const bool bShouldDirtyLevel = !bUseExternalObject;
 	const EObjectFlags Flags = (bIsTransientFolder ? RF_Transient : RF_NoFlags) | RF_Transactional;
 
-	UPackage* ExternalPackage = bUseExternalObject ? FExternalPackageHelper::CreateExternalPackage(InLevel, *GloballyUniqueObjectPath, GetExternalPackageFlags()) : nullptr;
+	UPackage* ExternalPackage = bUseExternalObject ? FExternalPackageHelper::CreateExternalPackage(InLevel, *GloballyUniqueObjectPath) : nullptr;
 
 	UActorFolder* ActorFolder = NewObject<UActorFolder>(InLevel, UActorFolder::StaticClass(), FName(FolderShortName), Flags, nullptr, /*bCopyTransientsFromClassDefaults*/false, /*InstanceGraph*/nullptr, ExternalPackage);
 	check(ActorFolder);
@@ -309,7 +309,7 @@ FFolder UActorFolder::GetFolder() const
 
 void UActorFolder::SetPackageExternal(bool bInExternal, bool bShouldDirty)
 {
-	FExternalPackageHelper::SetPackagingMode(this, GetOuterULevel(), bInExternal, bShouldDirty, GetExternalPackageFlags());
+	FExternalPackageHelper::SetPackagingMode(this, GetOuterULevel(), bInExternal, bShouldDirty);
 }
 
 #endif
diff --git a/Engine/Source/Runtime/Engine/Private/ExternalPackageHelper.cpp b/Engine/Source/Runtime/Engine/Private/ExternalPackageHelper.cpp
--- a/Engine/Source/Runtime/Engine/Private/ExternalPackageHelper.cpp
+++ b/Engine/Source/Runtime/Engine/Private/ExternalPackageHelper.cpp
@@ -1,15 +1,24 @@
 // Copyright Epic Games, Inc. All Rights Reserved.
 
 #include "ExternalPackageHelper.h"
-#include "Misc/ArchiveMD5.h"
 
 #if WITH_EDITOR
 
+#include "Misc/ArchiveMD5.h"
+#include "Misc/Paths.h"
+#include "HAL/PlatformApplicationMisc.h"
+
 FExternalPackageHelper::FOnObjectPackagingModeChanged FExternalPackageHelper::OnObjectPackagingModeChanged;
 
+EPackageFlags FExternalPackageHelper::GetDefaultExternalPackageFlags()
+{
+	return (PKG_EditorOnly | PKG_ContainsMapData | PKG_NewlyCreated);
+}
+
 UPackage* FExternalPackageHelper::CreateExternalPackage(UObject* InObjectOuter, const FString& InObjectPath, EPackageFlags InFlags)
 {
-	UPackage* Package = CreatePackage(*FExternalPackageHelper::GetExternalPackageName(InObjectOuter->GetPackage()->GetName(), InObjectPath));
+	UPackage* OutermostPackage = InObjectOuter->IsA<UPackage>() ? CastChecked<UPackage>(InObjectOuter) : InObjectOuter->GetOutermostObject()->GetPackage();
+	UPackage* Package = CreatePackage(*FExternalPackageHelper::GetExternalPackageName(OutermostPackage->GetName(), InObjectPath));
 	Package->SetPackageFlags(InFlags);
 	return Package;
 }
@@ -141,4 +150,32 @@ void FExternalPackageHelper::GetExternalSaveableObjects(UObject* InOuter, TArray
 	}
 }
 
+TArray<FString> FExternalPackageHelper::GetObjectsExternalPackageFilePath(const TArray<const UObject*>& InObjects)
+{
+	TArray<FString> PackageFilePaths;
+	for (const UObject* Object : InObjects)
+	{
+		if (Object && Object->IsPackageExternal())
+		{
+			const FString LocalFullPath(Object->GetExternalPackage()->GetLoadedPath().GetLocalFullPath());
+			if (!LocalFullPath.IsEmpty())
+			{
+				PackageFilePaths.Add(FPaths::ConvertRelativePathToFull(LocalFullPath));
+			}
+		}
+	}
+	return PackageFilePaths;
+}
+
+void FExternalPackageHelper::CopyObjectsExternalPackageFilePathToClipboard(const TArray<const UObject*>& InObjects)
+{
+	TArray<FString> PackageFilePaths = GetObjectsExternalPackageFilePath(InObjects);
+	if (!PackageFilePaths.IsEmpty())
+	{
+		FString Result = FString::Join(PackageFilePaths, TEXT("\n"));
+		check(Result.Len());
+		FPlatformApplicationMisc::ClipboardCopy(*Result);
+	}
+}
+
 #endif
diff --git a/Engine/Source/Runtime/Engine/Private/Level.cpp b/Engine/Source/Runtime/Engine/Private/Level.cpp
--- a/Engine/Source/Runtime/Engine/Private/Level.cpp
+++ b/Engine/Source/Runtime/Engine/Private/Level.cpp
@@ -481,7 +481,7 @@ void ULevel::CleanupLevel(bool bCleanupResources, bool bUnloadFromEditor)
 			{
 				ProcessPackage(ObjectPackage, true);
 			}
-		}, false);
+		}, /*bIncludeNestedObjects*/ true);
 	}
 #endif
 }
@@ -1369,38 +1369,21 @@ void ULevel::PreDuplicate(FObjectDuplicationParameters& DupParams)
 		FString ReplaceTo = FPaths::GetBaseFilename(*DstPackage->GetName());
 		ReplaceTo = FString::Printf(TEXT("%s.%s:"), *ReplaceTo, *ReplaceTo);
 
-		for (AActor* Actor : Actors)
+		ForEachObjectWithOuter(this, [this, &SrcPackage, &DstPackage, &ReplaceFrom, &ReplaceTo, &DupParams](UObject* Object)
 		{
-			if (UPackage* Package = Actor ? Actor->GetExternalPackage() : nullptr)
+			if (UPackage* Package = Object ? Object->GetExternalPackage() : nullptr)
 			{
-				FString Path = Actor->GetPathName();
+				FString Path = Object->GetPathName();
 				if (DstPackage != SrcPackage)
 				{
 					Path = Path.Replace(*ReplaceFrom, *ReplaceTo);
 				}
-				UPackage* DupPackage = CreateActorPackage(DstPackage, GetActorPackagingScheme(), Path);
+				UPackage* DupPackage = Object->IsA<AActor>() ? CreateActorPackage(DstPackage, GetActorPackagingScheme(), Path) : FExternalPackageHelper::CreateExternalPackage(DstPackage, Path);
 				DupPackage->MarkAsFullyLoaded();
 				DupPackage->MarkPackageDirty();
 				DupParams.DuplicationSeed.Add(Package, DupPackage);
 			}
-		}
-
-		ForEachActorFolder([&SrcPackage, &DstPackage, &ReplaceFrom, &ReplaceTo, &DupParams](UActorFolder* ActorFolder)
-		{
-			if (UPackage* Package = ActorFolder ? ActorFolder->GetExternalPackage() : nullptr)
-			{
-				FString Path = ActorFolder->GetPathName();
-				if (DstPackage != SrcPackage)
-				{
-					Path = Path.Replace(*ReplaceFrom, *ReplaceTo);
-				}
-				UPackage* DupPackage = FExternalPackageHelper::CreateExternalPackage(DstPackage, Path, UActorFolder::GetExternalPackageFlags());
-				DupPackage->MarkAsFullyLoaded();
-				DupPackage->MarkPackageDirty();
-				DupParams.DuplicationSeed.Add(Package, DupPackage);
-			}
-			return true;
-		});
+		}, /*bIncludeNestedObjects*/ true);
 	}
 #endif
 }
diff --git a/Engine/Source/Runtime/Engine/Private/LevelInstance/LevelInstanceLevelStreaming.cpp b/Engine/Source/Runtime/Engine/Private/LevelInstance/LevelInstanceLevelStreaming.cpp
--- a/Engine/Source/Runtime/Engine/Private/LevelInstance/LevelInstanceLevelStreaming.cpp
+++ b/Engine/Source/Runtime/Engine/Private/LevelInstance/LevelInstanceLevelStreaming.cpp
@@ -122,6 +122,13 @@ void ULevelStreamingLevelInstance::OnLoadedActorsAddedToLevelPreEvent(const TArr
 						if (bResetLoadersCalled)
 						{
 							ResetLoaders(Actor->GetExternalPackage());
+							ForEachObjectWithOuter(Actor, [](UObject* InObject)
+							{
+								if (InObject && InObject->IsPackageExternal())
+								{
+									ResetLoaders(InObject->GetExternalPackage());
+								}
+							}, /*bIncludeNestedObjects*/ true);
 						}
 
 						FLevelInstanceLevelStreamingUtils::MarkObjectsInPackageAsTransientAndNonTransactional(Actor->GetExternalPackage());
@@ -194,22 +201,13 @@ void ULevelStreamingLevelInstance::ResetLevelInstanceLoaders()
 			LinkerLoad->DetachLoader();
 		}
 
-		for (AActor* Actor : LoadedLevel->Actors)
-		{
-			if (Actor && Actor->IsPackageExternal())
-			{
-				ResetLoaders(Actor->GetExternalPackage());
-			}
-		}
-
-		LoadedLevel->ForEachActorFolder([](UActorFolder* ActorFolder)
+		ForEachObjectWithOuter(LoadedLevel, [](UObject* InObject)
 		{
-			if (ActorFolder->IsPackageExternal())
+			if (InObject && InObject->IsPackageExternal())
 			{
-				ResetLoaders(ActorFolder->GetExternalPackage());
+				ResetLoaders(InObject->GetExternalPackage());
 			}
-			return true;
-		});
+		}, /*bIncludeNestedObjects*/ true);
 
 		bResetLoadersCalled = true;
 	}
@@ -339,7 +337,7 @@ void ULevelStreamingLevelInstance::UnloadInstance(ULevelStreamingLevelInstance*
 				return false;
 			}
 			return true;
-		}, true);
+		}, /*bIncludeNestedObjects*/ true);
 
 		LevelStreaming->GetWorld()->GetSubsystem<ULevelInstanceSubsystem>()->RemoveLevelsFromWorld({ LevelStreaming->GetLoadedLevel() }, bResetTrans);
 	}
@@ -362,14 +360,14 @@ void ULevelStreamingLevelInstance::OnLevelLoadedChanged(ULevel* InLevel)
 			check(!InLevel->bAreComponentsCurrentlyRegistered);
 			FLevelInstanceLevelStreamingUtils::MarkObjectsInPackageAsTransientAndNonTransactional(NewLoadedLevel->GetPackage());
 
-			NewLoadedLevel->ForEachActorFolder([](UActorFolder* ActorFolder)
+			ForEachObjectWithOuter(NewLoadedLevel, [](UObject* InObject)
 			{
-				if (ActorFolder->IsPackageExternal())
+				// Skip actors as they are already handled in OnLoadedActorsAddedToLevelPreEvent
+				if (InObject && InObject->IsPackageExternal() && !InObject->IsA<AActor>())
 				{
-					FLevelInstanceLevelStreamingUtils::MarkObjectsInPackageAsTransientAndNonTransactional(ActorFolder->GetPackage());
+					FLevelInstanceLevelStreamingUtils::MarkObjectsInPackageAsTransientAndNonTransactional(InObject->GetPackage());
 				}
-				return true;
-			});
+			}, /*bIncludeNestedObjects*/ true);
 
 			OnLoadedActorsAddedToLevelPreEvent(InLevel->Actors);
 
diff --git a/Engine/Source/Runtime/Engine/Private/World.cpp b/Engine/Source/Runtime/Engine/Private/World.cpp
--- a/Engine/Source/Runtime/Engine/Private/World.cpp
+++ b/Engine/Source/Runtime/Engine/Private/World.cpp
@@ -49,6 +49,7 @@
 #include "Engine/CullDistanceVolume.h"
 #include "Engine/Console.h"
 #include "Engine/WorldComposition.h"
+#include "ExternalPackageHelper.h"
 #include "WorldPartition/WorldPartition.h"
 #include "WorldPartition/DataLayer/DataLayerManager.h"
 #include "WorldPartition/DataLayer/WorldDataLayers.h"
@@ -930,14 +931,18 @@ bool UWorld::Rename(const TCHAR* InName, UObject* NewOuter, ERenameFlags Flags)
 			}
 		}
 
+		// Process external objects other than actors
 		if (PersistentLevel->IsUsingExternalObjects())
 		{
-			PersistentLevel->ForEachActorFolder([](UActorFolder* ActorFolder)
+			ForEachObjectWithOuter(PersistentLevel, [this](UObject* Object)
 			{
-				ActorFolder->SetPackageExternal(false);
-				ActorFolder->SetPackageExternal(true);
+				if (Object->IsPackageExternal() && !Object->IsA<AActor>())
+				{
+					FExternalPackageHelper::SetPackagingMode(Object, PersistentLevel, false);
+					FExternalPackageHelper::SetPackagingMode(Object, PersistentLevel, true);
+				}
 				return true;
-			});
+			}, /*bIncludeNestedObjects*/ true);
 		}
 	}
 
diff --git a/Engine/Source/Runtime/Engine/Private/WorldPartition/WorldPartitionPackageHelper.cpp b/Engine/Source/Runtime/Engine/Private/WorldPartition/WorldPartitionPackageHelper.cpp
--- a/Engine/Source/Runtime/Engine/Private/WorldPartition/WorldPartitionPackageHelper.cpp
+++ b/Engine/Source/Runtime/Engine/Private/WorldPartition/WorldPartitionPackageHelper.cpp
@@ -43,7 +43,7 @@ void FWorldPartitionPackageHelper::UnloadPackage(UPackage* InPackage)
 				{
 					TrashPackage(ExternalPackage);
 				}
-			}, false);
+			}, /*bIncludeNestedObjects*/ true);
 		}
 	}
 }
diff --git a/Engine/Source/Runtime/Engine/Public/ActorFolder.h b/Engine/Source/Runtime/Engine/Public/ActorFolder.h
--- a/Engine/Source/Runtime/Engine/Public/ActorFolder.h
+++ b/Engine/Source/Runtime/Engine/Public/ActorFolder.h
@@ -22,6 +22,7 @@ public:
 	// Helper that creates a new UActorFolder
 	static ENGINE_API UActorFolder* Create(ULevel* InOuter, const FString& InFolderLabel, UActorFolder* InParent);
 
+	UE_DEPRECATED(5.4, "This function is deprecated.")
 	static EPackageFlags GetExternalPackageFlags() { return (PKG_EditorOnly | PKG_ContainsMapData | PKG_NewlyCreated); }
 	
 	// Returns actor folder info stored in its package
diff --git a/Engine/Source/Runtime/Engine/Public/ExternalPackageHelper.h b/Engine/Source/Runtime/Engine/Public/ExternalPackageHelper.h
--- a/Engine/Source/Runtime/Engine/Public/ExternalPackageHelper.h
+++ b/Engine/Source/Runtime/Engine/Public/ExternalPackageHelper.h
@@ -28,7 +28,10 @@ public:
 	 * @param InFlags the package flags to apply
 	 * @return the created package
 	 */
-	static ENGINE_API UPackage* CreateExternalPackage(UObject* InObjectOuter, const FString& InObjectPath, EPackageFlags InFlags);
+	static ENGINE_API UPackage* CreateExternalPackage(UObject* InObjectOuter, const FString& InObjectPath, EPackageFlags InFlags = FExternalPackageHelper::GetDefaultExternalPackageFlags());
+
+	/** Returns default external package flags used to create external packages. */
+	static ENGINE_API EPackageFlags GetDefaultExternalPackageFlags();
 
 	/**
 	 * Set the object packaging mode.
@@ -38,7 +41,7 @@ public:
 	 * @param bInShouldDirty should dirty or not the object's outer package
 	 * @param InExternalPackageFlags the flags to apply to the external package if bInIsPackageExternal is true
 	 */
-	static ENGINE_API void SetPackagingMode(UObject* InObject, UObject* InObjectOuter, bool bInIsPackageExternal, bool bInShouldDirty, EPackageFlags InExternalPackageFlags);
+	static ENGINE_API void SetPackagingMode(UObject* InObject, UObject* InObjectOuter, bool bInIsPackageExternal, bool bInShouldDirty = true, EPackageFlags InExternalPackageFlags = FExternalPackageHelper::GetDefaultExternalPackageFlags());
 
 	/**
 	 * Get the path containing the external objects for this path
@@ -77,6 +80,17 @@ public:
 	 * @param OutObjects	The objects that should be saved
 	 */
 	static ENGINE_API void GetExternalSaveableObjects(UObject* InOuter, TArray<UObject*>& OutObjects);
+
+	/**
+	 * Returns an array of external package file paths for the provided objects
+	 * @param InObjects		The objects to process
+	 */
+	static ENGINE_API TArray<FString> GetObjectsExternalPackageFilePath(const TArray<const UObject*>& InObjects);
+
+	/*
+	 * Copies the file path of objects external package to the clipboard
+	 */
+	static ENGINE_API void CopyObjectsExternalPackageFilePathToClipboard(const TArray<const UObject*>& InObjects);
 	
 private:
 	/** Get the external object package instance name. */
@@ -86,7 +100,8 @@ private:
 template<typename T>
 void FExternalPackageHelper::LoadObjectsFromExternalPackages(UObject* InOuter, TFunctionRef<void(T*)> Operation)
 {
-	const FString ExternalObjectsPath = FExternalPackageHelper::GetExternalObjectsPath(InOuter->GetPackage(), FString(), /*bTryUsingPackageLoadedPath*/ true);
+	UPackage* OutermostPackage = InOuter->IsA<UPackage>() ? CastChecked<UPackage>(InOuter) : InOuter->GetOutermostObject()->GetPackage();
+	const FString ExternalObjectsPath = FExternalPackageHelper::GetExternalObjectsPath(OutermostPackage, FString(), /*bTryUsingPackageLoadedPath*/ true);
 	TArray<FString> ObjectPackageNames;
 
 	// Do a synchronous scan of the world external objects path.			
@@ -122,7 +137,20 @@ void FExternalPackageHelper::LoadObjectsFromExternalPackages(UObject* InOuter, T
 			OuterInstancingContext = &OuterLinker->GetInstancingContext();
 		}
 
-		InstancingContext.AddPackageMapping(PackageResourceName, OuterPackage->GetFName());
+		// Add packages of all outers to instancing context
+		TSet<const UPackage*> OuterPackages;
+		UObject* ItObj = InOuter;
+		while (ItObj)
+		{
+			const UPackage* Package = ItObj->GetPackage();
+			bool bIsAlreadyInSet = false;
+			OuterPackages.Add(Package, &bIsAlreadyInSet);
+			if (!bIsAlreadyInSet)
+			{
+				InstancingContext.AddPackageMapping(Package->GetLoadedPath().GetPackageFName(), Package->GetFName());
+			}
+			ItObj = ItObj->GetOuter();
+		}
 
 		for (const FString& ObjectPackageName : ObjectPackageNames)
 		{
