Fixed string localization issues

diff --git a/Engine/Source/Runtime/Engine/Private/TextureRenderTarget.cpp b/Engine/Source/Runtime/Engine/Private/TextureRenderTarget.cpp
--- a/Engine/Source/Runtime/Engine/Private/TextureRenderTarget.cpp
+++ b/Engine/Source/Runtime/Engine/Private/TextureRenderTarget.cpp
@@ -165,7 +165,7 @@ bool UTextureRenderTarget::UpdateTexture(UTexture* InTexture, EConstructTextureF
 	{
 		if (OutErrorMessage != nullptr)
 		{
-			*OutErrorMessage = FText::Format(NSLOCTEXT("TextureRenderTarget", "NullRenderTargetResource", "Cannot update texture : render target {0} has been released."), FText::FromString(GetPathNameSafe(this)));
+			*OutErrorMessage = FText::Format(NSLOCTEXT("TextureRenderTarget", "NullRenderTargetResourceForUpdate", "Cannot update texture : render target {0} has been released."), FText::FromString(GetPathNameSafe(this)));
 		}
 		return false;
 	}
@@ -334,7 +334,7 @@ UTexture* UTextureRenderTarget::ConstructTexture(UObject* InOuter, const FString
 	{
 		if (OutErrorMessage != nullptr)
 		{
-			*OutErrorMessage = FText::Format(NSLOCTEXT("TextureRenderTarget", "NullRenderTargetResource", "Cannot construct texture : render target {0} has been released."), FText::FromString(GetPathNameSafe(this)));
+			*OutErrorMessage = FText::Format(NSLOCTEXT("TextureRenderTarget", "NullRenderTargetResourceForConstruct", "Cannot construct texture : render target {0} has been released."), FText::FromString(GetPathNameSafe(this)));
 		}
 		return nullptr;
 	}
diff --git a/Engine/Source/Runtime/Engine/Private/TextureRenderTarget2DArray.cpp b/Engine/Source/Runtime/Engine/Private/TextureRenderTarget2DArray.cpp
--- a/Engine/Source/Runtime/Engine/Private/TextureRenderTarget2DArray.cpp
+++ b/Engine/Source/Runtime/Engine/Private/TextureRenderTarget2DArray.cpp
@@ -153,7 +153,7 @@ bool UTextureRenderTarget2DArray::CanConvertToTexture(ETextureSourceFormat& OutT
 	{
 		if (OutErrorMessage != nullptr)
 		{
-			*OutErrorMessage = FText::Format(NSLOCTEXT("TextureRenderTargetVolume", "InvalidSizeForConversionToTexture", "Invalid size (({0},{1}), {2} slices) for converting {3} to {4}"),
+			*OutErrorMessage = FText::Format(NSLOCTEXT("TextureRenderTarget2DArray", "InvalidSizeForConversionToTexture", "Invalid size (({0},{1}), {2} slices) for converting {3} to {4}"),
 				FText::AsNumber(SizeX),
 				FText::AsNumber(SizeY),
 				FText::AsNumber(Slices),
diff --git a/Engine/Source/Runtime/Engine/Private/TextureRenderTargetCube.cpp b/Engine/Source/Runtime/Engine/Private/TextureRenderTargetCube.cpp
--- a/Engine/Source/Runtime/Engine/Private/TextureRenderTargetCube.cpp
+++ b/Engine/Source/Runtime/Engine/Private/TextureRenderTargetCube.cpp
@@ -175,7 +175,7 @@ bool UTextureRenderTargetCube::CanConvertToTexture(ETextureSourceFormat& OutText
 	{
 		if (OutErrorMessage != nullptr)
 		{
-			*OutErrorMessage = FText::Format(NSLOCTEXT("TextureRenderTargetVolume", "InvalidSizeForConversionToTexture", "Invalid size ({0},{0}) for converting {1} to {2}. Needs to be a power of 2."),
+			*OutErrorMessage = FText::Format(NSLOCTEXT("TextureRenderTargetCube", "InvalidSizeForConversionToTexture", "Invalid size ({0},{0}) for converting {1} to {2}. Needs to be a power of 2."),
 				FText::AsNumber(SizeX),
 				FText::FromString(GetClass()->GetName()),
 				FText::FromString(GetTextureUClass()->GetName()));
