Expose SetUpdateKinematicFromSimulation to Blueprint

diff --git a/Engine/Source/Runtime/Engine/Classes/Components/PrimitiveComponent.h b/Engine/Source/Runtime/Engine/Classes/Components/PrimitiveComponent.h
--- a/Engine/Source/Runtime/Engine/Classes/Components/PrimitiveComponent.h
+++ b/Engine/Source/Runtime/Engine/Classes/Components/PrimitiveComponent.h
@@ -2589,6 +2589,14 @@ public:
 	UFUNCTION(BlueprintPure, Category="Physics")
 	ENGINE_API virtual bool IsGravityEnabled() const;
 
+	/** Enables/disables whether this component should be updated by simulation when it is kinematic. This is needed if (for example) its velocity needs to be accessed. */
+	UFUNCTION(BlueprintCallable, Category="Physics")
+	ENGINE_API virtual void SetUpdateKinematicFromSimulation(bool bUpdateKinematicFromSimulation);
+
+	/** Returns whether this component should be updated by simulation when it is kinematic. */
+	UFUNCTION(BlueprintPure, Category="Physics")
+	ENGINE_API virtual bool GetUpdateKinematicFromSimulation() const;
+
 	/** Sets the linear damping of this component. */
 	UFUNCTION(BlueprintCallable, Category="Physics")
 	ENGINE_API virtual void SetLinearDamping(float InDamping);
diff --git a/Engine/Source/Runtime/Engine/Private/PrimitiveComponentPhysics.cpp b/Engine/Source/Runtime/Engine/Private/PrimitiveComponentPhysics.cpp
--- a/Engine/Source/Runtime/Engine/Private/PrimitiveComponentPhysics.cpp
+++ b/Engine/Source/Runtime/Engine/Private/PrimitiveComponentPhysics.cpp
@@ -439,6 +439,26 @@ bool UPrimitiveComponent::IsGravityEnabled() const
 	return false;
 }
 
+void UPrimitiveComponent::SetUpdateKinematicFromSimulation(bool bUpdateKinematicFromSimulation)
+{
+	FBodyInstance* BI = GetBodyInstance();
+	if (BI)
+	{
+		BI->SetUpdateKinematicFromSimulation(bUpdateKinematicFromSimulation);
+	}
+}
+
+bool UPrimitiveComponent::GetUpdateKinematicFromSimulation() const
+{
+	FBodyInstance* BI = GetBodyInstance();
+	if (BI)
+	{
+		return BI->bUpdateKinematicFromSimulation;
+	}
+
+	return false;
+}
+
 void UPrimitiveComponent::SetLinearDamping(float InDamping)
 {
 	FBodyInstance* BI = GetBodyInstance();
