Allow components to opt out of the render proxy delay strategy if PSO precaching hasn't finished.

This is implemented through a virtual function that individual component classes can override, and not as a flag or member variable since there's no clear use case for that yet.

An example of the use case is for UWaterBodyInfoMeshComponent, which needs to opt out because it's rendered as part of a one-off capture where a delayed proxy causes issues.

diff --git a/Engine/Plugins/Experimental/Water/Source/Runtime/Private/WaterBodyInfoMeshComponent.h b/Engine/Plugins/Experimental/Water/Source/Runtime/Private/WaterBodyInfoMeshComponent.h
--- a/Engine/Plugins/Experimental/Water/Source/Runtime/Private/WaterBodyInfoMeshComponent.h
+++ b/Engine/Plugins/Experimental/Water/Source/Runtime/Private/WaterBodyInfoMeshComponent.h
@@ -28,6 +28,12 @@ public:
 	// For every water body, there is a dilated and non-dilated mesh. Both are needed for generating the water info texture.
 	UPROPERTY()
 	bool bIsDilatedMesh = false;
+
+protected:
+	// Even though we may be still precaching PSOs, we don't want to delay render proxy creation
+	// or fallback to the default material. We're always returning false here so that we decline
+	// those optimizations.
+	virtual bool UsePSOPrecacheRenderProxyDelay() const override { return false; }
 };
 
 struct FWaterBodyInfoMeshSceneProxy : public FStaticMeshSceneProxy
diff --git a/Engine/Source/Runtime/Engine/Classes/Components/PrimitiveComponent.h b/Engine/Source/Runtime/Engine/Classes/Components/PrimitiveComponent.h
--- a/Engine/Source/Runtime/Engine/Classes/Components/PrimitiveComponent.h
+++ b/Engine/Source/Runtime/Engine/Classes/Components/PrimitiveComponent.h
@@ -931,6 +931,17 @@ public:
 	 */
 	ENGINE_API bool CheckPSOPrecachingAndBoostPriority();
 
+protected:
+
+	/**
+	 * Returns true if this component opts in to participate in the render proxy delay mechanism that kicks in
+	 * if PSO precaching hasn't finished. Otherwise, PSO precaching will still be active but the render proxy
+	 * will be created as normal.
+	 */
+	ENGINE_API virtual bool UsePSOPrecacheRenderProxyDelay() const;
+
+public:
+
 	/**
 	 * Set of actors to ignore during component sweeps in MoveComponent().
 	 * All components owned by these actors will be ignored when this component moves or updates overlaps.
diff --git a/Engine/Source/Runtime/Engine/Private/Components/PrimitiveComponent.cpp b/Engine/Source/Runtime/Engine/Private/Components/PrimitiveComponent.cpp
--- a/Engine/Source/Runtime/Engine/Private/Components/PrimitiveComponent.cpp
+++ b/Engine/Source/Runtime/Engine/Private/Components/PrimitiveComponent.cpp
@@ -4617,7 +4617,7 @@ void UPrimitiveComponent::RequestRecreateRenderStateWhenPSOPrecacheFinished(cons
 #if UE_WITH_PSO_PRECACHING
 	// If the proxy creation strategy relies on knowing when the precached PSO has been compiled,
 	// schedule a task to mark the render state dirty when all PSOs are compiled so the proxy gets recreated.
-	if (GetPSOPrecacheProxyCreationStrategy() != EPSOPrecacheProxyCreationStrategy::AlwaysCreate && !PSOPrecacheCompileEvents.IsEmpty())
+	if (UsePSOPrecacheRenderProxyDelay() && GetPSOPrecacheProxyCreationStrategy() != EPSOPrecacheProxyCreationStrategy::AlwaysCreate && !PSOPrecacheCompileEvents.IsEmpty())
 	{
 		PSOPrecacheCompileEvent = TGraphTask<FMarkRenderStateDirtyTask>::CreateTask(&PSOPrecacheCompileEvents).ConstructAndDispatchWhenReady(this);
 	}
@@ -4626,6 +4626,15 @@ void UPrimitiveComponent::RequestRecreateRenderStateWhenPSOPrecacheFinished(cons
 #endif // UE_WITH_PSO_PRECACHING
 }
 
+bool UPrimitiveComponent::UsePSOPrecacheRenderProxyDelay() const
+{
+#if UE_WITH_PSO_PRECACHING
+	return true;
+#else
+	return false;
+#endif
+}
+
 bool UPrimitiveComponent::IsPSOPrecaching() const
 {
 #if UE_WITH_PSO_PRECACHING
