Fixed UBlendSpaces not syncing when using inertialization. Added support for properly syncing UBlendSpaces when using "Transition Leader" and "Transition Follower" sync group roles.

[RN] minor

diff --git a/Engine/Source/Runtime/AnimGraphRuntime/Private/AnimNodes/AnimNode_BlendSpacePlayer.cpp b/Engine/Source/Runtime/AnimGraphRuntime/Private/AnimNodes/AnimNode_BlendSpacePlayer.cpp
--- a/Engine/Source/Runtime/AnimGraphRuntime/Private/AnimNodes/AnimNode_BlendSpacePlayer.cpp
+++ b/Engine/Source/Runtime/AnimGraphRuntime/Private/AnimNodes/AnimNode_BlendSpacePlayer.cpp
@@ -5,6 +5,7 @@
 #include "Animation/AnimSequence.h"
 #include "Animation/AnimInstanceProxy.h"
 #include "AnimGraphRuntimeTrace.h"
+#include "Animation/AnimInertializationSyncScope.h"
 #include "Animation/AnimSync.h"
 #include "Animation/AnimSyncScope.h"
 #include "Animation/AnimStats.h"
@@ -97,11 +98,11 @@ void FAnimNode_BlendSpacePlayerBase::UpdateInternal(const FAnimationUpdateContex
 			IsEvaluator(), Context.GetFinalBlendWeight(), /*inout*/ InternalTimeAccumulator, MarkerTickRecord);
 		TickRecord.RootMotionWeightModifier = Context.GetRootMotionWeightModifier();
 		TickRecord.DeltaTimeRecord = &DeltaTimeRecord;
-
-		UE::Anim::FAnimSyncParams SyncParams(GetGroupName(), GetGroupRole(), GetGroupMethod());
+		TickRecord.bRequestedInertialization = Context.GetMessage<UE::Anim::FAnimInertializationSyncScope>() != nullptr;
+		
 		TickRecord.GatherContextData(Context);
 
-		SyncScope.AddTickRecord(TickRecord, SyncParams, UE::Anim::FAnimSyncDebugInfo(Context));
+		SyncScope.AddTickRecord(TickRecord, GetSyncParams(TickRecord.bRequestedInertialization), UE::Anim::FAnimSyncDebugInfo(Context));
 
 		TRACE_ANIM_TICK_RECORD(Context, TickRecord);
 
diff --git a/Engine/Source/Runtime/Engine/Classes/Animation/AnimNode_AssetPlayerBase.h b/Engine/Source/Runtime/Engine/Classes/Animation/AnimNode_AssetPlayerBase.h
--- a/Engine/Source/Runtime/Engine/Classes/Animation/AnimNode_AssetPlayerBase.h
+++ b/Engine/Source/Runtime/Engine/Classes/Animation/AnimNode_AssetPlayerBase.h
@@ -6,6 +6,7 @@
 #include "UObject/ObjectMacros.h"
 #include "Animation/AnimationAsset.h"
 #include "Animation/AnimNode_RelevantAssetPlayerBase.h"
+#include "Animation/AnimSync.h"
 #include "AnimNode_AssetPlayerBase.generated.h"
 
 /* Base class for any asset playing anim node */
@@ -85,4 +86,7 @@ protected:
 
 	/** Track whether we have been full weight previously. Reset when we reach 0 weight*/
 	bool bHasBeenFullWeight = false;
+
+	/** Get sync parameters for the tick record produced by this asset player */
+	ENGINE_API UE::Anim::FAnimSyncParams GetSyncParams(bool bRequestedInertialization) const;
 };
diff --git a/Engine/Source/Runtime/Engine/Private/Animation/AnimNode_AssetPlayerBase.cpp b/Engine/Source/Runtime/Engine/Private/Animation/AnimNode_AssetPlayerBase.cpp
--- a/Engine/Source/Runtime/Engine/Private/Animation/AnimNode_AssetPlayerBase.cpp
+++ b/Engine/Source/Runtime/Engine/Private/Animation/AnimNode_AssetPlayerBase.cpp
@@ -31,47 +31,15 @@ void FAnimNode_AssetPlayerBase::CreateTickRecordForNode(const FAnimationUpdateCo
 
 	UE::Anim::FAnimSyncGroupScope& SyncScope = Context.GetMessageChecked<UE::Anim::FAnimSyncGroupScope>();
 
-	const EAnimGroupRole::Type SyncGroupRole = GetGroupRole();
-	const FName SyncGroupName = GetGroupName();
-	FName GroupNameToUse = SyncGroupName;
-	EAnimSyncMethod MethodToUse = GetGroupMethod();
-	const bool bRequestedInertialization = Context.GetMessage<UE::Anim::FAnimInertializationSyncScope>() != nullptr;
-	
-	// Skip sync based on roles.
-	{
-		// Only allow transition leader/follower part of a sync group once after inertialization request. (Inertilization)
-		if (bRequestedInertialization)
-		{
-			if (SyncGroupRole == EAnimGroupRole::TransitionLeader || SyncGroupRole == EAnimGroupRole::TransitionFollower)
-			{
-				GroupNameToUse = NAME_None;
-			}
-		}
-		// Only allow transition leader/follower part of a sync group once it has full weight (Standard blend).
-		else if ((SyncGroupRole == EAnimGroupRole::TransitionLeader || SyncGroupRole == EAnimGroupRole::TransitionFollower) && !bHasBeenFullWeight)
-		{
-			GroupNameToUse = NAME_None;
-		}
-
-		// Do not use sync groups.
-		if (GroupNameToUse == NAME_None && MethodToUse == EAnimSyncMethod::SyncGroup)
-		{
-			MethodToUse = EAnimSyncMethod::DoNotSync;
-		}
-	}
-
-	// Setup parameters for synchronizer
-	const UE::Anim::FAnimSyncParams SyncParams(GroupNameToUse, SyncGroupRole, MethodToUse);
-
 	// Active asset player's tick record
 	FAnimTickRecord TickRecord(Sequence, bLooping, PlayRate, bIsEvaluator, FinalBlendWeight, /*inout*/ InternalTimeAccumulator, MarkerTickRecord);
 	TickRecord.GatherContextData(Context);
 	TickRecord.RootMotionWeightModifier = Context.GetRootMotionWeightModifier();
 	TickRecord.DeltaTimeRecord = &DeltaTimeRecord;
-	TickRecord.bRequestedInertialization = bRequestedInertialization;
+	TickRecord.bRequestedInertialization = Context.GetMessage<UE::Anim::FAnimInertializationSyncScope>() != nullptr;;
 
 	// Add asset player to synchronizer
-	SyncScope.AddTickRecord(TickRecord, SyncParams, UE::Anim::FAnimSyncDebugInfo(Context));
+	SyncScope.AddTickRecord(TickRecord, GetSyncParams(TickRecord.bRequestedInertialization), UE::Anim::FAnimSyncDebugInfo(Context));
 
 	TRACE_ANIM_TICK_RECORD(Context, TickRecord);
 }
@@ -105,3 +73,36 @@ const FDeltaTimeRecord* FAnimNode_AssetPlayerBase::GetDeltaTimeRecord() const
 {
 	return &DeltaTimeRecord;
 }
+
+UE::Anim::FAnimSyncParams FAnimNode_AssetPlayerBase::GetSyncParams(bool bRequestedInertialization) const
+{
+	const EAnimGroupRole::Type SyncGroupRole = GetGroupRole();
+	const FName SyncGroupName = GetGroupName();
+	FName GroupNameToUse = SyncGroupName;
+	EAnimSyncMethod MethodToUse = GetGroupMethod();
+	
+	// Skip sync based on roles.
+	{
+		// Only allow transition leader/follower part of a sync group once after inertialization request. (Inertilization)
+		if (bRequestedInertialization)
+		{
+			if (SyncGroupRole == EAnimGroupRole::TransitionLeader || SyncGroupRole == EAnimGroupRole::TransitionFollower)
+			{
+				GroupNameToUse = NAME_None;
+			}
+		}
+		// Only allow transition leader/follower part of a sync group once it has full weight (Standard blend).
+		else if ((SyncGroupRole == EAnimGroupRole::TransitionLeader || SyncGroupRole == EAnimGroupRole::TransitionFollower) && !bHasBeenFullWeight)
+		{
+			GroupNameToUse = NAME_None;
+		}
+
+		// Do not use sync groups.
+		if (GroupNameToUse == NAME_None && MethodToUse == EAnimSyncMethod::SyncGroup)
+		{
+			MethodToUse = EAnimSyncMethod::DoNotSync;
+		}
+	}
+
+	return UE::Anim::FAnimSyncParams(GroupNameToUse, SyncGroupRole, MethodToUse);
+}
