Stops checking for MaterialDomain == MD_Surface on UMaterialInterface::HasPixelAnimation()

Many surface only parameters don't check anyway.

diff --git a/Engine/Source/Runtime/Engine/Private/Materials/Material.cpp b/Engine/Source/Runtime/Engine/Private/Materials/Material.cpp
--- a/Engine/Source/Runtime/Engine/Private/Materials/Material.cpp
+++ b/Engine/Source/Runtime/Engine/Private/Materials/Material.cpp
@@ -6759,7 +6759,7 @@ bool UMaterial::ShouldAlwaysEvaluateWorldPositionOffset() const
 
 bool UMaterial::HasPixelAnimation() const
 {
-	return bHasPixelAnimation && MaterialDomain == MD_Surface;
+	return bHasPixelAnimation;
 }
 
 void UMaterial::SetShadingModel(EMaterialShadingModel NewModel)
diff --git a/Engine/Source/Runtime/Engine/Private/Materials/MaterialInstance.cpp b/Engine/Source/Runtime/Engine/Private/Materials/MaterialInstance.cpp
--- a/Engine/Source/Runtime/Engine/Private/Materials/MaterialInstance.cpp
+++ b/Engine/Source/Runtime/Engine/Private/Materials/MaterialInstance.cpp
@@ -4609,8 +4609,7 @@ bool UMaterialInstance::ShouldAlwaysEvaluateWorldPositionOffset() const
 
 bool UMaterialInstance::HasPixelAnimation() const
 {
-	const UMaterial* Material = GetMaterial();
-	return Material ? bHasPixelAnimation && GetMaterial()->MaterialDomain == MD_Surface : false;
+	return bHasPixelAnimation;
 }
 
 bool UMaterialInstance::IsMasked() const
