WorldPartition: Added option in WorldDataLayers actor to save each DataLayerInstances in their own file (disabled by default)

diff --git a/Engine/Plugins/Editor/EngineAssetDefinitions/Source/Private/AssetDefinition_DataLayerInstance.h b/Engine/Plugins/Editor/EngineAssetDefinitions/Source/Private/AssetDefinition_DataLayerInstance.h
new file mode 100644
--- /dev/null
+++ b/Engine/Plugins/Editor/EngineAssetDefinitions/Source/Private/AssetDefinition_DataLayerInstance.h
@@ -0,0 +1,32 @@
+// Copyright Epic Games, Inc. All Rights Reserved.
+
+#pragma once
+
+#include "WorldPartition/DataLayer/DataLayerInstance.h"
+#include "AssetDefinitionDefault.h"
+
+#include "AssetDefinition_DataLayerInstance.generated.h"
+
+struct FAssetCategoryPath;
+
+UCLASS()
+class UAssetDefinition_DataLayerInstance : public UAssetDefinitionDefault
+{
+	GENERATED_BODY()
+
+public:
+	// UAssetDefinition Begin
+	virtual FText GetAssetDisplayName() const override { return NSLOCTEXT("AssetTypeActions", "AssetTypeActions_DataLayerInstance", "Data Layer Instance"); }
+	virtual FLinearColor GetAssetColor() const override { return FLinearColor(FColor(52, 213, 235)); }
+	virtual TSoftClassPtr<UObject> GetAssetClass() const override { return UDataLayerInstance::StaticClass(); }
+	virtual TConstArrayView<FAssetCategoryPath> GetAssetCategories() const override
+	{
+		return TConstArrayView<FAssetCategoryPath>();
+	}
+	virtual FText GetObjectDisplayNameText(UObject* Object) const override
+	{
+		UDataLayerInstance* DataLayerInstance = CastChecked<UDataLayerInstance>(Object);
+		return FText::FromString(DataLayerInstance->GetDataLayerShortName());
+	}
+	// UAssetDefinition End
+};
diff --git a/Engine/Source/Editor/DataLayerEditor/Private/DataLayer/DataLayerMode.cpp b/Engine/Source/Editor/DataLayerEditor/Private/DataLayer/DataLayerMode.cpp
--- a/Engine/Source/Editor/DataLayerEditor/Private/DataLayer/DataLayerMode.cpp
+++ b/Engine/Source/Editor/DataLayerEditor/Private/DataLayer/DataLayerMode.cpp
@@ -1265,6 +1265,7 @@ void FDataLayerMode::RegisterContextMenu()
 			check(Mode);
 			TArray<UDataLayerInstance*> SelectedDataLayers = Mode->GetSelectedDataLayers(SceneOutliner);
 			const bool bSelectedDataLayersContainsLocked = Algo::AnyOf(SelectedDataLayers, [](const UDataLayerInstance* DataLayerInstance) { return DataLayerInstance->IsLocked(); });
+			const bool bSelectedDataLayersContainsExternalPackage = Algo::AnyOf(SelectedDataLayers, [](const UDataLayerInstance* DataLayerInstance) { return DataLayerInstance->IsPackageExternal(); });
 
 			TArray<AActor*> SelectedActors;
 			GEditor->GetSelectedActors()->GetSelectedObjects<AActor>(SelectedActors);
@@ -1516,6 +1517,18 @@ void FDataLayerMode::RegisterContextMenu()
 							}}),
 						FCanExecuteAction::CreateLambda([SelectedDataLayers] { return !SelectedDataLayers.IsEmpty(); })
 					));
+
+				Section.AddMenuEntry("Copy Selected Data Layer Instances(s) File Path", LOCTEXT("CopySelectedDataLayerInstancessFilePath", "Copy Selected Data Layer Instance(s) File Path"), FText(), FSlateIcon(),
+					FUIAction(
+						FExecuteAction::CreateLambda([SelectedDataLayers]() {
+							check(!SelectedDataLayers.IsEmpty());
+							{
+								TArray<const UObject*> Objects;
+								Algo::Transform(SelectedDataLayers, Objects, [](UDataLayerInstance* DataLayerInstance) { return DataLayerInstance; });
+								FExternalPackageHelper::CopyObjectsExternalPackageFilePathToClipboard(Objects);
+							}}),
+						FCanExecuteAction::CreateLambda([bSelectedDataLayersContainsExternalPackage] { return bSelectedDataLayersContainsExternalPackage; })
+					));
 			}
 
 			{
diff --git a/Engine/Source/Editor/WorldPartitionEditor/Private/WorldPartition/Customizations/WorldDataLayersActorDetails.cpp b/Engine/Source/Editor/WorldPartitionEditor/Private/WorldPartition/Customizations/WorldDataLayersActorDetails.cpp
new file mode 100644
--- /dev/null
+++ b/Engine/Source/Editor/WorldPartitionEditor/Private/WorldPartition/Customizations/WorldDataLayersActorDetails.cpp
@@ -0,0 +1,65 @@
+// Copyright Epic Games, Inc. All Rights Reserved.
+
+#include "WorldDataLayersActorDetails.h"
+#include "DetailCategoryBuilder.h"
+#include "DetailLayoutBuilder.h"
+#include "DetailWidgetRow.h"
+#include "Widgets/Input/SCheckBox.h"
+#include "Widgets/Text/STextBlock.h"
+#include "WorldPartition/DataLayer/WorldDataLayers.h"
+
+#define LOCTEXT_NAMESPACE "WorldDataLayersActorDetails"
+
+TSharedRef<IDetailCustomization> FWorldDataLayersActorDetails::MakeInstance()
+{
+	return MakeShareable( new FWorldDataLayersActorDetails);
+}
+
+void FWorldDataLayersActorDetails::CustomizeDetails( IDetailLayoutBuilder& DetailBuilder )
+{
+	TArray<TWeakObjectPtr<UObject>> CustomizedObjects;
+	DetailBuilder.GetObjectsBeingCustomized(CustomizedObjects);
+	if (AWorldDataLayers* CustomizedWorldDataLayers = (CustomizedObjects.Num() == 1) ? Cast<AWorldDataLayers>(CustomizedObjects[0]) : nullptr)
+	{
+		if (ULevel* CustomizedLevel = CustomizedWorldDataLayers->GetLevel())
+		{
+			IDetailCategoryBuilder& WorldCategory = DetailBuilder.EditCategory("World");
+			WorldCategory.AddCustomRow(LOCTEXT("UseExternalPackageDataLayerInstancesRow", "UseExternalPackageDataLayerInstances"), true)
+			.NameContent()
+			[
+				SNew(STextBlock)
+				.Text(LOCTEXT("UseExternalPackageDataLayerInstances", "Use External Package Data Layer Instances"))
+				.ToolTipText(LOCTEXT("UseExternalPackageDataLayerInstances_ToolTip", "Use External Package Data Layer Instances, data layer instances will be stored in their own external package."))
+				.Font(IDetailLayoutBuilder::GetDetailFont())
+				.IsEnabled(this, &FWorldDataLayersActorDetails::UseExternalPackageDataLayerInstancesEnabled, CustomizedWorldDataLayers)
+			]
+			.ValueContent()
+			[
+				SNew(SCheckBox)
+				.OnCheckStateChanged(this, &FWorldDataLayersActorDetails::OnUseExternalPackageDataLayerInstancesChanged, CustomizedWorldDataLayers)
+				.IsChecked(this, &FWorldDataLayersActorDetails::IsUseExternalPackageDataLayerInstancesChecked, CustomizedWorldDataLayers)
+				.IsEnabled(this, &FWorldDataLayersActorDetails::UseExternalPackageDataLayerInstancesEnabled, CustomizedWorldDataLayers)
+			];
+		}
+	}
+}
+
+bool FWorldDataLayersActorDetails::UseExternalPackageDataLayerInstancesEnabled(AWorldDataLayers* WorldDataLayers) const
+{
+	return WorldDataLayers && WorldDataLayers->SupportsExternalPackageDataLayerInstances();
+}
+
+ECheckBoxState FWorldDataLayersActorDetails::IsUseExternalPackageDataLayerInstancesChecked(AWorldDataLayers* WorldDataLayers) const
+{
+	return WorldDataLayers && WorldDataLayers->IsUsingExternalPackageDataLayerInstances() ? ECheckBoxState::Checked : ECheckBoxState::Unchecked;
+}
+
+void FWorldDataLayersActorDetails::OnUseExternalPackageDataLayerInstancesChanged(ECheckBoxState BoxState, AWorldDataLayers* WorldDataLayers)
+{
+	if (WorldDataLayers)
+	{
+		WorldDataLayers->SetUseExternalPackageDataLayerInstances(BoxState == ECheckBoxState::Checked, /*bInteractiveMode*/ true);
+	}
+}
+
+#undef LOCTEXT_NAMESPACE
diff --git a/Engine/Source/Editor/WorldPartitionEditor/Private/WorldPartition/Customizations/WorldDataLayersActorDetails.h b/Engine/Source/Editor/WorldPartitionEditor/Private/WorldPartition/Customizations/WorldDataLayersActorDetails.h
new file mode 100644
--- /dev/null
+++ b/Engine/Source/Editor/WorldPartitionEditor/Private/WorldPartition/Customizations/WorldDataLayersActorDetails.h
@@ -0,0 +1,25 @@
+// Copyright Epic Games, Inc. All Rights Reserved.
+
+#pragma once
+
+#include "IDetailCustomization.h"
+#include "Templates/SharedPointer.h"
+#include "Styling/SlateTypes.h"
+
+class IDetailLayoutBuilder;
+class AWorldDataLayers;
+
+class FWorldDataLayersActorDetails : public IDetailCustomization
+{
+public:
+	/** Makes a new instance of this detail layout class for a specific detail view requesting it */
+	static TSharedRef<IDetailCustomization> MakeInstance();
+
+	/** IDetailCustomization interface */
+	virtual void CustomizeDetails( IDetailLayoutBuilder& DetailBuilder ) override;
+
+private:
+	bool UseExternalPackageDataLayerInstancesEnabled(AWorldDataLayers* WorldDataLayers) const;
+	ECheckBoxState IsUseExternalPackageDataLayerInstancesChecked(AWorldDataLayers* WorldDataLayers) const;
+	void OnUseExternalPackageDataLayerInstancesChanged(ECheckBoxState BoxState, AWorldDataLayers* WorldDataLayers);
+};
diff --git a/Engine/Source/Editor/WorldPartitionEditor/Private/WorldPartitionEditorModule.cpp b/Engine/Source/Editor/WorldPartitionEditor/Private/WorldPartitionEditorModule.cpp
--- a/Engine/Source/Editor/WorldPartitionEditor/Private/WorldPartitionEditorModule.cpp
+++ b/Engine/Source/Editor/WorldPartitionEditor/Private/WorldPartitionEditorModule.cpp
@@ -14,6 +14,7 @@
 #include "WorldPartition/Customizations/WorldPartitionDetailsCustomization.h"
 #include "WorldPartition/Customizations/WorldPartitionHLODDetailsCustomization.h"
 #include "WorldPartition/Customizations/WorldPartitionRuntimeSpatialHashDetailsCustomization.h"
+#include "WorldPartition/Customizations/WorldDataLayersActorDetails.h"
 #include "WorldPartition/SWorldPartitionConvertDialog.h"
 #include "WorldPartition/WorldPartitionConvertOptions.h"
 #include "WorldPartition/WorldPartitionEditorSettings.h"
@@ -182,6 +183,7 @@ void FWorldPartitionEditorModule::StartupModule()
 	PropertyEditor.RegisterCustomClassLayout("WorldPartition", FOnGetDetailCustomizationInstance::CreateStatic(&FWorldPartitionDetails::MakeInstance));
 	PropertyEditor.RegisterCustomClassLayout("WorldPartitionRuntimeSpatialHash", FOnGetDetailCustomizationInstance::CreateStatic(&FWorldPartitionRuntimeSpatialHashDetails::MakeInstance));
 	PropertyEditor.RegisterCustomClassLayout("WorldPartitionHLOD", FOnGetDetailCustomizationInstance::CreateStatic(&FWorldPartitionHLODDetailsCustomization::MakeInstance));
+	PropertyEditor.RegisterCustomClassLayout("WorldDataLayers", FOnGetDetailCustomizationInstance::CreateStatic(&FWorldDataLayersActorDetails::MakeInstance));
 
 	FWorldPartitionClassDescRegistry().Get().Initialize();
 
diff --git a/Engine/Source/Runtime/Core/Public/UObject/FortniteMainBranchObjectVersion.h b/Engine/Source/Runtime/Core/Public/UObject/FortniteMainBranchObjectVersion.h
--- a/Engine/Source/Runtime/Core/Public/UObject/FortniteMainBranchObjectVersion.h
+++ b/Engine/Source/Runtime/Core/Public/UObject/FortniteMainBranchObjectVersion.h
@@ -381,6 +381,9 @@ struct FFortniteMainBranchObjectVersion
 		// Changed default Local Exposure Contrast from 1.0 to 0.8
 		LocalExposureDefaultChangeFrom1,
 
+		// Added support of external packaging of Data Layer Instances
+		AddDataLayerInstanceExternalPackage,
+
 		// -----<new versions can be added above this line>-------------------------------------------------
 		VersionPlusOne,
 		LatestVersion = VersionPlusOne - 1
diff --git a/Engine/Source/Runtime/Engine/Classes/Engine/Level.h b/Engine/Source/Runtime/Engine/Classes/Engine/Level.h
--- a/Engine/Source/Runtime/Engine/Classes/Engine/Level.h
+++ b/Engine/Source/Runtime/Engine/Classes/Engine/Level.h
@@ -1369,6 +1369,7 @@ private:
 	friend struct FLevelActorFoldersHelper;
 	friend class FWorldPartitionLevelHelper;
 	friend class UActorFolder;
+	friend class AWorldDataLayers;
 
 	/** Replace the existing LSA (if set) by spawning a new one based on this level's script blueprint */
 	void RegenerateLevelScriptActor();
diff --git a/Engine/Source/Runtime/Engine/Private/WorldPartition/DataLayer/DataLayerInstance.cpp b/Engine/Source/Runtime/Engine/Private/WorldPartition/DataLayer/DataLayerInstance.cpp
--- a/Engine/Source/Runtime/Engine/Private/WorldPartition/DataLayer/DataLayerInstance.cpp
+++ b/Engine/Source/Runtime/Engine/Private/WorldPartition/DataLayer/DataLayerInstance.cpp
@@ -7,8 +7,12 @@
 #include "WorldPartition/WorldPartitionLog.h"
 #include "WorldPartition/WorldPartitionActorLoaderInterface.h"
 #include "WorldPartition/DataLayer/DataLayerUtils.h"
+#include "WorldPartition/DataLayer/DataLayerAsset.h"
 #include "WorldPartition/DataLayer/WorldDataLayers.h"
+#include "WorldPartition/DataLayer/WorldDataLayersActorDesc.h"
 #include "WorldPartition/ErrorHandling/WorldPartitionStreamingGenerationErrorHandler.h"
+#include "AssetRegistry/AssetData.h"
+#include "AssetRegistry/AssetRegistryHelpers.h"
 
 #include UE_INLINE_GENERATED_CPP_BY_NAME(DataLayerInstance)
 
@@ -52,6 +56,100 @@ AWorldDataLayers* UDataLayerInstance::GetOuterWorldDataLayers() const
 }
 
 #if WITH_EDITOR
+bool UDataLayerInstance::IsAsset() const
+{
+	// When using external packaging, Data Layer Instances are considered assets to allow using the asset logic for save dialogs, etc.
+	// Also, they return true even if pending kill, in order to show up as deleted in these dialogs.
+	return IsPackageExternal() && !GetPackage()->HasAnyFlags(RF_Transient) && !HasAnyFlags(RF_Transient | RF_ClassDefaultObject) && !GetPackage()->HasAnyPackageFlags(PKG_PlayInEditor);
+}
+
+namespace DataLayerInstance
+{
+	static const FName NAME_DataLayerInstanceName(TEXT("DataLayerInstanceName"));
+	static const FName NAME_DataLayerInstanceParentName(TEXT("DataLayerInstanceParentName"));
+	static const FName NAME_DataLayerInstanceAssetPath(TEXT("DataLayerInstanceAssetPath"));
+	static const FName NAME_DataLayerInstanceIsPrivate(TEXT("DataLayerInstanceIsPrivate"));
+	static const FName NAME_DataLayerInstanceIsIncludedInActorFilterDefault(TEXT("DataLayerInstanceIsIncludedInActorFilterDefault"));
+	static const FName NAME_DataLayerInstancePrivateDataLayerSupportsActorFilter(TEXT("DataLayerInstancePrivateDataLayerSupportsActorFilter"));
+	static const FName NAME_DataLayerInstancePrivateShortName(TEXT("DataLayerInstancePrivateShortName"));
+};
+
+void UDataLayerInstance::GetAssetRegistryTags(TArray<FAssetRegistryTag>& OutTags) const
+{
+	Super::GetAssetRegistryTags(OutTags);
+
+	if (IsPackageExternal())
+	{
+		// Set generic FPrimaryAssetId::PrimaryAssetDisplayNameTag
+		OutTags.Add(FAssetRegistryTag(FPrimaryAssetId::PrimaryAssetDisplayNameTag, *GetDataLayerShortName(), FAssetRegistryTag::TT_Hidden));
+
+		// Set DataLayerInstance specific tags
+		OutTags.Add(FAssetRegistryTag(DataLayerInstance::NAME_DataLayerInstanceName, *GetDataLayerFName().ToString(), FAssetRegistryTag::TT_Hidden));
+		if (GetParent())
+		{
+			OutTags.Add(FAssetRegistryTag(DataLayerInstance::NAME_DataLayerInstanceParentName, *GetParent()->GetDataLayerFName().ToString(), FAssetRegistryTag::TT_Hidden));
+		}
+		if (const UDataLayerAsset* DataLayerAsset = GetAsset())
+		{
+			OutTags.Add(FAssetRegistryTag(DataLayerInstance::NAME_DataLayerInstanceAssetPath, *DataLayerAsset->GetPathName(), FAssetRegistryTag::TT_Hidden));
+			OutTags.Add(FAssetRegistryTag(DataLayerInstance::NAME_DataLayerInstanceIsPrivate, DataLayerAsset->IsPrivate() ? TEXT("1") : TEXT("0"), FAssetRegistryTag::TT_Hidden));
+		}
+		OutTags.Add(FAssetRegistryTag(DataLayerInstance::NAME_DataLayerInstanceIsIncludedInActorFilterDefault, IsIncludedInActorFilterDefault() ? TEXT("1") : TEXT("0"), FAssetRegistryTag::TT_Hidden));
+		OutTags.Add(FAssetRegistryTag(DataLayerInstance::NAME_DataLayerInstancePrivateDataLayerSupportsActorFilter, SupportsActorFilters() ? TEXT("1") : TEXT("0"), FAssetRegistryTag::TT_Hidden));
+		OutTags.Add(FAssetRegistryTag(DataLayerInstance::NAME_DataLayerInstancePrivateShortName, *GetDataLayerShortName(), FAssetRegistryTag::TT_Hidden));
+	}
+}
+
+bool UDataLayerInstance::GetAssetRegistryInfoFromPackage(FName InDataLayerInstancePackageName, FDataLayerInstanceDesc& OutDataLayerInstanceDesc)
+{
+	IAssetRegistry& AssetRegistry = FModuleManager::LoadModuleChecked<FAssetRegistryModule>(TEXT("AssetRegistry")).Get();
+	OutDataLayerInstanceDesc.bIsUsingAsset = true;
+	TArray<FAssetData> Assets;
+	AssetRegistry.GetAssetsByPackageName(InDataLayerInstancePackageName, Assets, true);
+	check(Assets.Num() <= 1);
+	if (Assets.Num() == 1)
+	{
+		return GetAssetRegistryInfoFromPackage(Assets[0], OutDataLayerInstanceDesc);
+	}
+	return false;
+}
+
+bool UDataLayerInstance::GetAssetRegistryInfoFromPackage(const FAssetData& InAsset, FDataLayerInstanceDesc& OutDataLayerInstanceDesc)
+{
+	FString Value;
+	if (InAsset.GetTagValue(DataLayerInstance::NAME_DataLayerInstanceName, Value))
+	{
+		OutDataLayerInstanceDesc.Name = *Value;
+	}
+	if (InAsset.GetTagValue(DataLayerInstance::NAME_DataLayerInstanceParentName, Value))
+	{
+		OutDataLayerInstanceDesc.ParentName = *Value;
+	}
+	if (InAsset.GetTagValue(DataLayerInstance::NAME_DataLayerInstanceAssetPath, Value))
+	{
+		FName AssetPath = *Value;
+		UAssetRegistryHelpers::FixupRedirectedAssetPath(AssetPath);
+		OutDataLayerInstanceDesc.AssetPath = AssetPath;
+	}
+	if (InAsset.GetTagValue(DataLayerInstance::NAME_DataLayerInstanceIsPrivate, Value))
+	{
+		OutDataLayerInstanceDesc.bIsPrivate = (Value == TEXT("1"));
+	}
+	if (InAsset.GetTagValue(DataLayerInstance::NAME_DataLayerInstanceIsIncludedInActorFilterDefault, Value))
+	{
+		OutDataLayerInstanceDesc.bIsIncludedInActorFilterDefault = (Value == TEXT("1"));
+	}
+	if (InAsset.GetTagValue(DataLayerInstance::NAME_DataLayerInstancePrivateDataLayerSupportsActorFilter, Value))
+	{
+		OutDataLayerInstanceDesc.bPrivateDataLayerSupportsActorFilter = (Value == TEXT("1"));
+	}
+	if (InAsset.GetTagValue(DataLayerInstance::NAME_DataLayerInstancePrivateShortName, Value))
+	{
+		OutDataLayerInstanceDesc.PrivateShortName = Value;
+	}
+	return !OutDataLayerInstanceDesc.Name.IsNone();
+}
+
 void UDataLayerInstance::PreEditUndo()
 {
 	Super::PreEditUndo();
@@ -305,8 +403,6 @@ void UDataLayerInstance::SetChildParent(UDataLayerInstance* InParent)
 	}
 
 	check(!InParent || InParent->CanHaveChildDataLayers());
-
-	Modify();
 	while (Children.Num())
 	{
 		Children[0]->SetParent(InParent);
@@ -315,7 +411,7 @@ void UDataLayerInstance::SetChildParent(UDataLayerInstance* InParent)
 
 void UDataLayerInstance::RemoveChild(UDataLayerInstance* InDataLayer)
 {
-	Modify();
+	Modify(false);
 	check(Children.Contains(InDataLayer));
 	Children.RemoveSingle(InDataLayer);
 }
@@ -407,7 +503,7 @@ void UDataLayerInstance::AddChild(UDataLayerInstance* InDataLayer)
 {
 	check(InDataLayer->GetOuterWorldDataLayers() == GetOuterWorldDataLayers())
 	check(CanHaveChildDataLayers());
-	Modify();
+	Modify(false);
 	checkSlow(!Children.Contains(InDataLayer));
 	Children.Add(InDataLayer);
 }
diff --git a/Engine/Source/Runtime/Engine/Private/WorldPartition/DataLayer/DataLayerManager.cpp b/Engine/Source/Runtime/Engine/Private/WorldPartition/DataLayer/DataLayerManager.cpp
--- a/Engine/Source/Runtime/Engine/Private/WorldPartition/DataLayer/DataLayerManager.cpp
+++ b/Engine/Source/Runtime/Engine/Private/WorldPartition/DataLayer/DataLayerManager.cpp
@@ -192,6 +192,9 @@ void UDataLayerManager::Initialize()
 		}
 	}
 
+	// Initialize WorldDataLayers
+	WorldDataLayers->OnDataLayerManagerInitialized();
+
 	// Some levels do not have the WorldDataLayer actor serialized as part of their Actors array. Make sure we add it here if it isn't.
 	// Make sure WorldDataLayers is part of the Actors list so that it gets cooked properly as part of the Persistent Level
 	// This auto-corrects itself when resaving the level.
@@ -228,11 +231,21 @@ void UDataLayerManager::Initialize()
 		check(IsValid(*It));
 		It->FixupDataLayers();
 	}
+#else
+	if (AWorldDataLayers* WorldDataLayers = GetWorldDataLayers())
+	{
+		WorldDataLayers->OnDataLayerManagerInitialized();
+	}
 #endif
 }
 
 void UDataLayerManager::DeInitialize()
 {
+	if (AWorldDataLayers* WorldDataLayers = GetWorldDataLayers())
+	{
+		WorldDataLayers->OnDataLayerManagerDeinitialized();
+	}
+
 #if WITH_EDITOR
 	UActorDescContainer::OnActorDescContainerInitialized.RemoveAll(this);
 
diff --git a/Engine/Source/Runtime/Engine/Private/WorldPartition/DataLayer/WorldDataLayers.cpp b/Engine/Source/Runtime/Engine/Private/WorldPartition/DataLayer/WorldDataLayers.cpp
--- a/Engine/Source/Runtime/Engine/Private/WorldPartition/DataLayer/WorldDataLayers.cpp
+++ b/Engine/Source/Runtime/Engine/Private/WorldPartition/DataLayer/WorldDataLayers.cpp
@@ -17,6 +17,7 @@
 #include "ProfilingDebugging/CsvProfiler.h"
 #include "Net/UnrealNetwork.h"
 #include "Net/Core/PushModel/PushModel.h"
+#include "UObject/FortniteMainBranchObjectVersion.h"
 
 #include UE_INLINE_GENERATED_CPP_BY_NAME(WorldDataLayers)
 
@@ -25,6 +26,8 @@
 #include "WorldPartition/DataLayer/WorldDataLayersActorDesc.h"
 #include "LevelInstance/LevelInstanceSubsystem.h"
 #include "Algo/Find.h"
+#include "Misc/MessageDialog.h"
+#include "Editor.h"
 #endif
 
 #define LOCTEXT_NAMESPACE "WorldDataLayers"
@@ -47,6 +50,7 @@ static FString JoinDataLayerShortNamesFromInstanceNames(AWorldDataLayers* InWorl
 AWorldDataLayers::AWorldDataLayers(const FObjectInitializer& ObjectInitializer)
 	: Super(ObjectInitializer.DoNotCreateDefaultSubobject(TEXT("Sprite")))
 #if WITH_EDITORONLY_DATA
+	, bUseExternalPackageDataLayerInstances(false)
 	, bAllowRuntimeDataLayerEditing(true)
 #endif
 	, DataLayersStateEpoch(0)
@@ -107,14 +111,8 @@ void AWorldDataLayers::InitializeDataLayerRuntimeStates()
 
 	if (GetWorld()->IsGameWorld())
 	{
-#if WITH_EDITOR
-		if (!IsRuntimeRelevant())
-		{
-			return;
-		}
-#endif
 		TArray<UDataLayerInstance*> RuntimeDataLayerInstances;
-		RuntimeDataLayerInstances.Reserve(DataLayerInstances.Num());
+		RuntimeDataLayerInstances.Reserve(GetDataLayerInstances().Num());
 
 		ForEachDataLayerInstance([this, &RuntimeDataLayerInstances](UDataLayerInstance* DataLayerInstance)
 		{
@@ -475,10 +473,11 @@ void AWorldDataLayers::ResolveEffectiveRuntimeState(const UDataLayerInstance* In
 			OnDataLayerRuntimeStateChanged(InDataLayerInstance, NewEffectiveRuntimeState);
 		}
 
-		for (const UDataLayerInstance* Child : InDataLayerInstance->GetChildren())
+		InDataLayerInstance->ForEachChild([this](const UDataLayerInstance* Child)
 		{
 			ResolveEffectiveRuntimeState(Child);
-		}
+			return true;
+		});
 	}
 }
 
@@ -509,10 +508,11 @@ void AWorldDataLayers::DumpDataLayerRecursively(const UDataLayerInstance* DataLa
 		*DataLayer->GetDataLayerShortName(),
 		*GetDataLayerRuntimeStateString(DataLayer));
 
-	for (const UDataLayerInstance* Child : DataLayer->GetChildren())
+	DataLayer->ForEachChild([this, &Prefix, &OutputDevice](const UDataLayerInstance* Child)
 	{
 		DumpDataLayerRecursively(Child, Prefix + TEXT(" | "), OutputDevice);
-	}
+		return true;
+	});
 }
 
 void AWorldDataLayers::DumpDataLayers(FOutputDevice& OutputDevice) const
@@ -640,21 +640,51 @@ TArray<const UDataLayerInstance*> AWorldDataLayers::GetDataLayerInstances(const
 	return OutDataLayers;
 }
 
-int32 AWorldDataLayers::RemoveDataLayers(const TArray<UDataLayerInstance*>& InDataLayerInstances)
+bool AWorldDataLayers::IsEmpty() const
+{
+	return GetDataLayerInstances().IsEmpty();
+}
+
+void AWorldDataLayers::AddDataLayerInstance(UDataLayerInstance* InDataLayerInstance)
+{
+	check(InDataLayerInstance);
+	TSet<TObjectPtr<UDataLayerInstance>>& UsedDataLayerInstances = GetDataLayerInstances();
+	// Only dirty actor when not using external package
+	const bool bDirty = !bUseExternalPackageDataLayerInstances;
+	Modify(bDirty);
+	check(GetLevel());
+	check(GetLevel()->IsUsingExternalObjects());
+	check(!UsedDataLayerInstances.Contains(InDataLayerInstance));
+	UsedDataLayerInstances.Add(InDataLayerInstance);
+	if (InDataLayerInstance->IsPackageExternal())
+	{
+		InDataLayerInstance->MarkPackageDirty();
+	}
+}
+
+int32 AWorldDataLayers::RemoveDataLayers(const TArray<UDataLayerInstance*>& InDataLayerInstances, bool bInResolveActorDescContainers)
 {
 	int32 RemovedCount = 0;
 
+	TSet<TObjectPtr<UDataLayerInstance>>& UsedDataLayerInstances = GetDataLayerInstances();
 	for (UDataLayerInstance* DataLayerInstance : InDataLayerInstances)
 	{
-		if (ContainsDataLayer(DataLayerInstance))
+		if (UsedDataLayerInstances.Contains(DataLayerInstance))
 		{
-			Modify();
+			// Only dirty actor when not using external package
+			const bool bDirty = !bUseExternalPackageDataLayerInstances;
+			Modify(bDirty);
+			DataLayerInstance->Modify();
 			DataLayerInstance->SetChildParent(DataLayerInstance->GetParent());
-			DataLayerInstances.Remove(DataLayerInstance);
+			UsedDataLayerInstances.Remove(DataLayerInstance);
 			if (DataLayerInstance->IsA<UDeprecatedDataLayerInstance>())
 			{
 				DeprecatedDataLayerNameToDataLayerInstance.Remove(DataLayerInstance->GetDataLayerFName());
 			}
+			if (DataLayerInstance->IsPackageExternal())
+			{
+				DataLayerInstance->MarkAsGarbage();
+			}
 			RemovedCount++;
 		}
 	}
@@ -662,31 +692,18 @@ int32 AWorldDataLayers::RemoveDataLayers(const TArray<UDataLayerInstance*>& InDa
 	if (RemovedCount > 0)
 	{
 		UpdateContainsDeprecatedDataLayers();
-
-		ResolveActorDescContainers();
+		if (bInResolveActorDescContainers)
+		{
+			ResolveActorDescContainers();
+		}
 	}
 
 	return RemovedCount;
 }
 
-bool AWorldDataLayers::RemoveDataLayer(const UDataLayerInstance* InDataLayerInstance)
+bool AWorldDataLayers::RemoveDataLayer(const UDataLayerInstance* InDataLayerInstance, bool bInResolveActorDescContainers)
 {
-	if (ContainsDataLayer(InDataLayerInstance))
-	{
-		Modify();
-		DataLayerInstances.Remove(InDataLayerInstance);
-
-		if (InDataLayerInstance->IsA<UDeprecatedDataLayerInstance>())
-		{
-			DeprecatedDataLayerNameToDataLayerInstance.Remove(InDataLayerInstance->GetDataLayerFName());
-			UpdateContainsDeprecatedDataLayers();
-		}
-
-		ResolveActorDescContainers();
-		
-		return true;
-	}
-	return false;
+	return RemoveDataLayers({ const_cast<UDataLayerInstance*>(InDataLayerInstance) }, bInResolveActorDescContainers) > 0;
 }
 
 void AWorldDataLayers::SetAllowRuntimeDataLayerEditing(bool bInAllowRuntimeDataLayerEditing)
@@ -781,18 +798,24 @@ TArray<UDataLayerInstance*> AWorldDataLayers::GetActorEditorContextDataLayers()
 
 bool AWorldDataLayers::ContainsDataLayer(const UDataLayerInstance* InDataLayerInstance) const 
 {
-	return DataLayerInstances.Contains(InDataLayerInstance);
+	return GetDataLayerInstances().Contains(InDataLayerInstance);
 }
 
 const UDataLayerInstance* AWorldDataLayers::GetDataLayerInstance(const FName& InDataLayerInstanceName) const
 {
 #if WITH_EDITOR	
-	for (UDataLayerInstance* DataLayerInstance : DataLayerInstances)
 	{
-		if (DataLayerInstance->GetDataLayerFName() == InDataLayerInstanceName)
+		const UDataLayerInstance* FoundDataLayerInstance = nullptr;
+		ForEachDataLayerInstance([InDataLayerInstanceName, &FoundDataLayerInstance](UDataLayerInstance* DataLayerInstance)
 		{
-			return DataLayerInstance;
-		}
+			if (DataLayerInstance->GetDataLayerFName() == InDataLayerInstanceName)
+			{
+				FoundDataLayerInstance = DataLayerInstance;
+				return false;
+			}
+			return true;
+		});
+		return FoundDataLayerInstance;
 	}
 #else
 	if (const UDataLayerInstance* const* FoundDataLayerInstance = InstanceNameToInstance.Find(InDataLayerInstanceName))
@@ -814,13 +837,17 @@ const UDataLayerInstance* AWorldDataLayers::GetDataLayerInstance(const FName& In
 const UDataLayerInstance* AWorldDataLayers::GetDataLayerInstanceFromAssetName(const FName& InDataLayerAssetFullName) const
 {
 #if WITH_EDITOR	
-	for (UDataLayerInstance* DataLayerInstance : DataLayerInstances)
+	const UDataLayerInstance* FoundDataLayerInstance = nullptr;
+	ForEachDataLayerInstance([InDataLayerAssetFullName, &FoundDataLayerInstance](UDataLayerInstance* DataLayerInstance)
 	{
 		if (DataLayerInstance->GetDataLayerFullName().Equals(InDataLayerAssetFullName.ToString(), ESearchCase::IgnoreCase))
 		{
-			return DataLayerInstance;
+			FoundDataLayerInstance = DataLayerInstance;
+			return false;
 		}
-	}
+		return true;
+	});
+	return FoundDataLayerInstance;
 #else
 	if (const UDataLayerInstance* const* FoundDataLayerInstance = AssetNameToInstance.Find(InDataLayerAssetFullName.ToString()))
 	{
@@ -838,13 +865,17 @@ const UDataLayerInstance* AWorldDataLayers::GetDataLayerInstance(const UDataLaye
 	}
 
 #if WITH_EDITOR	
-	for (const UDataLayerInstance* DataLayerInstance : DataLayerInstances)
+	const UDataLayerInstance* FoundDataLayerInstance = nullptr;
+	ForEachDataLayerInstance([InDataLayerAsset, &FoundDataLayerInstance](UDataLayerInstance* DataLayerInstance)
 	{
 		if (DataLayerInstance->GetAsset() == InDataLayerAsset)
 		{
-			return DataLayerInstance;
+			FoundDataLayerInstance = DataLayerInstance;
+			return false;
 		}
-	}
+		return true;
+	});
+	return FoundDataLayerInstance;
 #else
 	if (const UDataLayerInstance* const* FoundDataLayerInstance = AssetNameToInstance.Find(InDataLayerAsset->GetFullName()))
 	{
@@ -857,24 +888,18 @@ const UDataLayerInstance* AWorldDataLayers::GetDataLayerInstance(const UDataLaye
 
 void AWorldDataLayers::ForEachDataLayerInstance(TFunctionRef<bool(UDataLayerInstance*)> Func)
 {
-	for (UDataLayerInstance* DataLayerInstance : DataLayerInstances)
+	for (UDataLayerInstance* DataLayerInstance : GetDataLayerInstances())
 	{
 		if (DataLayerInstance && !Func(DataLayerInstance))
 		{
-			break;
+			return;
 		}
 	}
 }
 
 void AWorldDataLayers::ForEachDataLayerInstance(TFunctionRef<bool(UDataLayerInstance*)> Func) const
 {
-	for (UDataLayerInstance* DataLayerInstance : DataLayerInstances)
-	{
-		if (DataLayerInstance && !Func(DataLayerInstance))
-		{
-			break;
-		}
-	}
+	const_cast<AWorldDataLayers*>(this)->ForEachDataLayerInstance(Func);
 }
 
 TArray<const UDataLayerInstance*> AWorldDataLayers::GetDataLayerInstances(const TArray<FName>& InDataLayerInstanceNames) const
@@ -893,33 +918,46 @@ TArray<const UDataLayerInstance*> AWorldDataLayers::GetDataLayerInstances(const
 	return OutDataLayers;
 }
 
-void AWorldDataLayers::PostLoad()
+TSet<TObjectPtr<UDataLayerInstance>>& AWorldDataLayers::GetDataLayerInstances()
 {
-	Super::PostLoad();
-
-	GetLevel()->ConditionalPostLoad();
-
-	// Patch WorldDataLayer in UWorld.
-	// Only the "main" world data Layer is named AWorldDataLayers::StaticClass()->GetFName() for a given world.
-	if ((GetTypedOuter<UWorld>()->GetWorldDataLayers() == nullptr) && (GetFName() == GetWorldPartionWorldDataLayersName()))
+#if WITH_EDITOR
+	if (IsUsingExternalPackageDataLayerInstances())
 	{
-		GetTypedOuter<UWorld>()->SetWorldDataLayers(this);
+		check(DataLayerInstances.IsEmpty());
+		check(LoadedExternalPackageDataLayerInstances.IsEmpty());
+		return ExternalPackageDataLayerInstances;
+	}
+	else
+	{
+		check(ExternalPackageDataLayerInstances.IsEmpty());
 	}
+#endif
+	return DataLayerInstances;
+}
+
+const TSet<TObjectPtr<UDataLayerInstance>>& AWorldDataLayers::GetDataLayerInstances() const
+{
+	return const_cast<AWorldDataLayers*>(this)->GetDataLayerInstances();
+}
 
+void AWorldDataLayers::OnDataLayerManagerInitialized()
+{
 #if WITH_EDITOR
-	bListedInSceneOutliner = true;
+	ConditionalPostLoad();
+	// At this point, LoadedExternalPackageDataLayerInstances are fully loaded, transfer them to the DataLayerInstances list.
+	ExternalPackageDataLayerInstances.Append(LoadedExternalPackageDataLayerInstances);
+	LoadedExternalPackageDataLayerInstances.Reset();
+	if (IsRunningCookCommandlet())
+	{
+		// Embed external DataLayerInstances when cooking
+		SetUseExternalPackageDataLayerInstances(false);
+	}
 
 	ConvertDataLayerToInstancces();
 
 	// Remove all Editor Data Layers when cooking or when in a game world
 	if (IsRunningCookCommandlet() || GetWorld()->IsGameWorld())
 	{
-		ForEachDataLayerInstance([](UDataLayerInstance* DataLayerInstance)
-		{
-			DataLayerInstance->ConditionalPostLoad();
-			return true;
-		});
-
 		TArray<UDataLayerInstance*> EditorDataLayers;
 		ForEachDataLayerInstance([&EditorDataLayers](UDataLayerInstance* DataLayerInstance)
 		{
@@ -932,18 +970,12 @@ void AWorldDataLayers::PostLoad()
 		RemoveDataLayers(EditorDataLayers);
 	}
 
-	// Empty Data Layers for non-runtime relevant AWorldDataLayers
-	if (GetWorld()->IsGameWorld() && !IsEmpty() && !IsRuntimeRelevant())
-	{
-		UE_LOG(LogWorldPartition, Warning, TEXT("WorldDataLayers %s is not runtime relevant, Data Layers of this level will be ignored."), *GetPathName());
-		RemoveDataLayers(DataLayerInstances.Array());
-	}
-
 	// Setup defaults before overriding with user settings
-	for (UDataLayerInstance* DataLayer : DataLayerInstances)
+	ForEachDataLayerInstance([](UDataLayerInstance* DataLayerInstance)
 	{
-		DataLayer->SetIsLoadedInEditor(DataLayer->IsInitiallyLoadedInEditor(), /*bFromUserChange*/false);
-	}
+		DataLayerInstance->SetIsLoadedInEditor(DataLayerInstance->IsInitiallyLoadedInEditor(), /*bFromUserChange*/false);
+		return true;
+	});
 
 	// Initialize DataLayer's IsLoadedInEditor based on DataLayerEditorPerProjectUserSettings
 	TArray<FName> SettingsDataLayersNotLoadedInEditor = GetMutableDefault<UWorldPartitionEditorPerProjectUserSettings>()->GetWorldDataLayersNotLoadedInEditor(GetWorld());
@@ -963,9 +995,60 @@ void AWorldDataLayers::PostLoad()
 			DataLayerInstance->SetIsLoadedInEditor(true, /*bFromUserChange*/false);
 		}
 	}
-#else
+#endif
+
+	InitializeDataLayerRuntimeStates();
+}
+
+void AWorldDataLayers::OnDataLayerManagerDeinitialized()
+{
+	ResetDataLayerRuntimeStates();
+}
+
+void AWorldDataLayers::Serialize(FArchive& Ar)
+{
+	Super::Serialize(Ar);
+
+#if WITH_EDITORONLY_DATA
+	// This handle both duplication for PIE (which includes unsaved modifications) and regular duplicate
+	if ((Ar.GetPortFlags() & PPF_Duplicate) && Ar.IsPersistent())
+	{
+		Ar << ExternalPackageDataLayerInstances;
+	}
+#endif
+}
+
+void AWorldDataLayers::PostLoad()
+{
+	Super::PostLoad();
+
+	GetLevel()->ConditionalPostLoad();
+
+	// Patch WorldDataLayer in UWorld.
+	// Only the "main" world data Layer is named AWorldDataLayers::StaticClass()->GetFName() for a given world.
+	if ((GetTypedOuter<UWorld>()->GetWorldDataLayers() == nullptr) && (GetFName() == GetWorldPartionWorldDataLayersName()))
+	{
+		GetTypedOuter<UWorld>()->SetWorldDataLayers(this);
+	}
+
+#if WITH_EDITOR
+	ULevel* Level = GetLevel();
+	if (!Level->bWasDuplicated && IsUsingExternalPackageDataLayerInstances())
+	{
+		// Load all folders for this level
+		FExternalPackageHelper::LoadObjectsFromExternalPackages<UDataLayerInstance>(this, [this](UDataLayerInstance* LoadedDataLayerInstance)
+		{
+			check(IsValid(LoadedDataLayerInstance));
+			LoadedExternalPackageDataLayerInstances.Add(LoadedDataLayerInstance);
+		});
+	}
+
+	bListedInSceneOutliner = true;
+#endif
+
+#if !WITH_EDITOR
 	// Build acceleration tables
-	for (const UDataLayerInstance* DataLayerInstance : DataLayerInstances)
+	ForEachDataLayerInstance([this](const UDataLayerInstance* DataLayerInstance)
 	{
 		InstanceNameToInstance.Add(DataLayerInstance->GetDataLayerFName(), DataLayerInstance);
 
@@ -975,14 +1058,15 @@ void AWorldDataLayers::PostLoad()
 			if (!DataLayerInstanceWithAsset->GetAsset())
 			{
 				UE_LOG(LogWorldPartition, Warning, TEXT("DataLayerWithAsset %s has null asset."), *DataLayerInstanceWithAsset->GetPathName());
-				continue;
 			}
-			AssetNameToInstance.Add(DataLayerInstanceWithAsset->GetAsset()->GetFullName(), DataLayerInstance);
+			else
+			{
+				AssetNameToInstance.Add(DataLayerInstanceWithAsset->GetAsset()->GetFullName(), DataLayerInstance);
+			}
 		}
-	}
+		return true;
+	});
 #endif
-
-	InitializeDataLayerRuntimeStates();
 }
 
 #if WITH_EDITOR
@@ -1031,6 +1115,73 @@ bool AWorldDataLayers::IsRuntimeRelevant() const
 	return false;
 }
 
+bool AWorldDataLayers::SupportsExternalPackageDataLayerInstances() const
+{
+	ULevel* Level = GetLevel();
+	return Level && Level->bIsPartitioned && Level->IsUsingExternalObjects() && !HasDeprecatedDataLayers();
+}
+
+bool AWorldDataLayers::SetUseExternalPackageDataLayerInstances(bool bInNewValue, bool bInInteractiveMode)
+{
+	if (HasAnyFlags(RF_ArchetypeObject | RF_ClassDefaultObject))
+	{
+		return false;
+	}
+
+	if (bUseExternalPackageDataLayerInstances == bInNewValue)
+	{
+		return false;
+	}
+
+	if (HasDeprecatedDataLayers())
+	{
+		UE_LOG(LogWorldPartition, Warning, TEXT("Changing external packaging of data layer instances is not supported with deprecated data layers. Convert your data to use Data Layer Assets."));
+		return false;
+	}
+
+	check(SupportsExternalPackageDataLayerInstances());
+	const bool bInteractiveMode = bInInteractiveMode && !IsRunningCommandlet();
+
+	// Validate we have a saved map
+	UPackage* Package = GetExternalPackage();
+	if (Package == GetTransientPackage()
+		|| Package->HasAnyFlags(RF_Transient)
+		|| !FPackageName::IsValidLongPackageName(Package->GetName()))
+	{
+		if (bInteractiveMode)
+		{
+			FMessageDialog::Open(EAppMsgType::Ok, LOCTEXT("UseExternalPackageDataLayerInstancesSaveActor", "You need to save the WorldDataLayers actor before changing the `Use External Package Data Layer Instances` option."));
+		}
+		UE_LOG(LogWorldPartition, Warning, TEXT("You need to save the WorldDataLayers actor before changing the `Use External Package Data Layer Instances` option."));
+		return false;
+	}
+
+	if (bInteractiveMode && !IsEmpty())
+	{
+		FText MessageTitle(LOCTEXT("ConvertDataLayerInstancesExternalPackagingDialog", "Convert Data Layer Instances External Packaging"));
+		FText Message(LOCTEXT("ConvertDataLayerInstancesExternalPackagingMsg", "Do you want to convert external packaging for all data layer instances?"));
+		EAppReturnType::Type ConvertAnswer = FMessageDialog::Open(EAppMsgType::YesNo, Message, MessageTitle);
+		if (ConvertAnswer != EAppReturnType::Yes)
+		{
+			return false;
+		}
+	}
+
+	Modify();
+	ForEachDataLayerInstance([this, bInNewValue](UDataLayerInstance* DataLayerInstance)
+	{
+		FExternalPackageHelper::SetPackagingMode(DataLayerInstance, this, bInNewValue);
+		return true;
+	});
+	bUseExternalPackageDataLayerInstances = bInNewValue;
+	Swap(DataLayerInstances, ExternalPackageDataLayerInstances);
+
+	// Operation cannot be undone
+	GEditor->ResetTransaction(LOCTEXT("EActorFolderObjectsResetTrans", "WorldDataLayers Use External Package Data Layer Instances"));
+
+	return true;
+}
+
 void AWorldDataLayers::ConvertDataLayerToInstancces()
 {
 	static_assert(DATALAYER_TO_INSTANCE_RUNTIME_CONVERSION_ENABLED, "AWorldDataLayers::ConvertDataLayerToInstancces function is deprecated and needs to be deleted.");
@@ -1096,24 +1247,25 @@ void AWorldDataLayers::ResolveActorDescContainers()
 void AWorldDataLayers::PreEditUndo()
 {
 	Super::PreEditUndo();
-	CachedDataLayerInstances = DataLayerInstances;
+	CachedDataLayerInstances = GetDataLayerInstances();
 }
 
 void AWorldDataLayers::PostEditUndo()
 {
 	Super::PostEditUndo();
 
-	bool bNeedResolve = CachedDataLayerInstances.Num() != DataLayerInstances.Num();
+	bool bNeedResolve = CachedDataLayerInstances.Num() != GetDataLayerInstances().Num();
 	if (!bNeedResolve)
 	{
-		for (UDataLayerInstance* DataLayerInstance : DataLayerInstances)
+		ForEachDataLayerInstance([this, &bNeedResolve](UDataLayerInstance* DataLayerInstance)
 		{
 			if (!CachedDataLayerInstances.Contains(DataLayerInstance))
 			{
 				bNeedResolve = true;
-				break;
+				return false;
 			}
-		}
+			return true;
+		});
 	}
 	if (bNeedResolve)
 	{
@@ -1154,16 +1306,18 @@ bool AWorldDataLayers::ContainsDataLayer(const UDEPRECATED_DataLayer* InDataLaye
 const UDataLayerInstance* AWorldDataLayers::GetDataLayerFromLabel(const FName& InDataLayerLabel) const
 {
 	const FString DataLayerLabelSanitized = *FDataLayerUtils::GetSanitizedDataLayerShortName(InDataLayerLabel.ToString());
-
-	for (const UDataLayerInstance* DataLayer : DataLayerInstances)
+	const UDataLayerInstance* FoundDataLayerInstance = nullptr;
+	ForEachDataLayerInstance([&DataLayerLabelSanitized, &FoundDataLayerInstance](UDataLayerInstance* DataLayerInstance)
 	{
-		if (DataLayer->GetDataLayerShortName() == DataLayerLabelSanitized)
+		if (DataLayerInstance->GetDataLayerShortName() == DataLayerLabelSanitized)
 		{
-			return DataLayer;
+			FoundDataLayerInstance = DataLayerInstance;
+			return false;
 		}
-	}
+		return true;
+	});
 
-	return nullptr;
+	return FoundDataLayerInstance;
 }
 
 PRAGMA_ENABLE_DEPRECATION_WARNINGS
diff --git a/Engine/Source/Runtime/Engine/Private/WorldPartition/DataLayer/WorldDataLayersActorDesc.cpp b/Engine/Source/Runtime/Engine/Private/WorldPartition/DataLayer/WorldDataLayersActorDesc.cpp
--- a/Engine/Source/Runtime/Engine/Private/WorldPartition/DataLayer/WorldDataLayersActorDesc.cpp
+++ b/Engine/Source/Runtime/Engine/Private/WorldPartition/DataLayer/WorldDataLayersActorDesc.cpp
@@ -4,6 +4,8 @@
 
 #if WITH_EDITOR
 
+#include "UObject/TopLevelAssetPath.h"
+#include "AssetRegistry/IAssetRegistry.h"
 #include "AssetRegistry/AssetRegistryHelpers.h"
 #include "WorldPartition/WorldPartitionHelpers.h"
 #include "WorldPartition/DataLayer/WorldDataLayers.h"
@@ -11,6 +13,7 @@
 #include "WorldPartition/WorldPartitionActorContainerID.h"
 #include "UObject/FortniteSeasonBranchObjectVersion.h"
 #include "UObject/FortniteMainBranchObjectVersion.h"
+#include "ExternalPackageHelper.h"
 
 FArchive& operator<<(FArchive& Ar, FDataLayerInstanceDesc& Desc)
 {
@@ -200,6 +203,7 @@ void FDataLayerInstanceDesc::Init(UDataLayerInstance* InDataLayerInstance)
 
 FWorldDataLayersActorDesc::FWorldDataLayersActorDesc()
 : bIsValid(false)
+, bUseExternalPackageDataLayerInstances(false)
 {}
 
 void FWorldDataLayersActorDesc::Init(const AActor* InActor)
@@ -207,12 +211,16 @@ void FWorldDataLayersActorDesc::Init(const AActor* InActor)
 	FWorldPartitionActorDesc::Init(InActor);
 
 	const AWorldDataLayers* WorldDataLayers = CastChecked<AWorldDataLayers>(InActor);
-	WorldDataLayers->ForEachDataLayerInstance([this](UDataLayerInstance* DataLayerInstance)
+	bUseExternalPackageDataLayerInstances = WorldDataLayers->IsUsingExternalPackageDataLayerInstances();
+	if (!bUseExternalPackageDataLayerInstances)
 	{
-		FDataLayerInstanceDesc& DataLayerInstanceDesc = DataLayerInstances.Emplace_GetRef();
-		DataLayerInstanceDesc.Init(DataLayerInstance);
-		return true;
-	});
+		WorldDataLayers->ForEachDataLayerInstance([this](UDataLayerInstance* DataLayerInstance)
+		{
+			FDataLayerInstanceDesc& DataLayerInstanceDesc = DataLayerInstances.Emplace_GetRef();
+			DataLayerInstanceDesc.Init(DataLayerInstance);
+			return true;
+		});
+	}
 
 	bIsValid = true;
 }
@@ -222,7 +230,8 @@ bool FWorldDataLayersActorDesc::Equals(const FWorldPartitionActorDesc* Other) co
 	if (FWorldPartitionActorDesc::Equals(Other))
 	{
 		const FWorldDataLayersActorDesc* OtherDesc = (FWorldDataLayersActorDesc*)Other;
-		return CompareUnsortedArrays(DataLayerInstances, OtherDesc->DataLayerInstances);
+		return (bUseExternalPackageDataLayerInstances == OtherDesc->bUseExternalPackageDataLayerInstances) && 
+			CompareUnsortedArrays(DataLayerInstances, OtherDesc->DataLayerInstances);
 	}
 	return false;
 }
@@ -241,31 +250,95 @@ void FWorldDataLayersActorDesc::Serialize(FArchive& Ar)
 			Ar << DataLayerInstances;
 			bIsValid = true;
 		}
+		if (Ar.CustomVer(FFortniteMainBranchObjectVersion::GUID) >= FFortniteMainBranchObjectVersion::AddDataLayerInstanceExternalPackage)
+		{
+			Ar << bUseExternalPackageDataLayerInstances;
+		}
+	}
+}
+
+const TArray<FDataLayerInstanceDesc>& FWorldDataLayersActorDesc::GetExternalPackageDataLayerInstances() const
+{
+	check(bUseExternalPackageDataLayerInstances);
+	if (!ExternalPackageDataLayerInstances.IsSet())
+	{
+		TArray<FDataLayerInstanceDesc> FoundDataLayerInstances;
+		if (ULevel::GetIsLevelPartitionedFromPackage(ActorPath.GetLongPackageFName()))
+		{
+			FTopLevelAssetPath MapAssetName = ActorPath.GetAssetPath();
+			FString ExternalObjectsPath = FExternalPackageHelper::GetExternalObjectsPath(MapAssetName.GetPackageName().ToString());
+			
+			// Do a synchronous scan of the world external objects path.			
+			IAssetRegistry & AssetRegistry = FModuleManager::LoadModuleChecked<FAssetRegistryModule>(TEXT("AssetRegistry")).Get();
+			AssetRegistry.ScanSynchronous({ ExternalObjectsPath }, TArray<FString>());
+
+			FARFilter Filter;
+			Filter.bRecursivePaths = true;
+			Filter.bIncludeOnlyOnDiskAssets = true;
+			Filter.ClassPaths.Add(UDataLayerInstance::StaticClass()->GetClassPathName());
+			Filter.bRecursiveClasses = true;
+			Filter.PackagePaths.Add(*ExternalObjectsPath);
+			TArray<FAssetData> Assets;
+			AssetRegistry.GetAssets(Filter, Assets);
+
+			for (const FAssetData& Asset : Assets)
+			{
+				FDataLayerInstanceDesc DataLayerInstanceDesc;
+				if (UDataLayerInstance::GetAssetRegistryInfoFromPackage(Asset, DataLayerInstanceDesc))
+				{
+					FoundDataLayerInstances.Add(DataLayerInstanceDesc);
+				}
+			}
+		}
+		ExternalPackageDataLayerInstances = MoveTemp(FoundDataLayerInstances);
+	}
+	return ExternalPackageDataLayerInstances.GetValue();
+}
+
+const TArray<FDataLayerInstanceDesc>& FWorldDataLayersActorDesc::GetDataLayerInstances() const
+{
+	return bUseExternalPackageDataLayerInstances ? GetExternalPackageDataLayerInstances() : DataLayerInstances;
+}
+
+void FWorldDataLayersActorDesc::ForEachDataLayerInstanceDesc(TFunctionRef<bool(const FDataLayerInstanceDesc&)> Func) const
+{
+	for (const FDataLayerInstanceDesc& DataLayerInstance : GetDataLayerInstances())
+	{
+		if (!Func(DataLayerInstance))
+		{
+			return;
+		}
 	}
 }
 
 const FDataLayerInstanceDesc* FWorldDataLayersActorDesc::GetDataLayerInstanceFromInstanceName(FName InDataLayerInstanceName) const
 {
-	for (const FDataLayerInstanceDesc& DataLayerInstance : DataLayerInstances)
+	const FDataLayerInstanceDesc* FoundDataLayerInstanceDesc = nullptr;
+	ForEachDataLayerInstanceDesc([InDataLayerInstanceName, &FoundDataLayerInstanceDesc](const FDataLayerInstanceDesc& DataLayerInstance)
 	{
 		if (DataLayerInstance.GetName().IsEqual(InDataLayerInstanceName, ENameCase::CaseSensitive))
 		{
-			return &DataLayerInstance;
+			FoundDataLayerInstanceDesc = &DataLayerInstance;
+			return false;
 		}
-	}
-	return nullptr;
+		return true;
+	});
+	return FoundDataLayerInstanceDesc;
 }
 
 const FDataLayerInstanceDesc* FWorldDataLayersActorDesc::GetDataLayerInstanceFromAssetPath(FName InDataLayerAssetPath) const
 {
-	for (const FDataLayerInstanceDesc& DataLayerInstance : DataLayerInstances)
+	const FDataLayerInstanceDesc* FoundDataLayerInstanceDesc = nullptr;
+	ForEachDataLayerInstanceDesc([InDataLayerAssetPath, &FoundDataLayerInstanceDesc](const FDataLayerInstanceDesc& DataLayerInstance)
 	{
 		if (FName(DataLayerInstance.GetAssetPath()) == InDataLayerAssetPath)
 		{
-			return &DataLayerInstance;
+			FoundDataLayerInstanceDesc = &DataLayerInstance;
+			return false;
 		}
-	}
-	return nullptr;
+		return true;
+	});
+	return FoundDataLayerInstanceDesc;
 }
 
 #endif
diff --git a/Engine/Source/Runtime/Engine/Public/WorldPartition/DataLayer/DataLayerInstance.h b/Engine/Source/Runtime/Engine/Public/WorldPartition/DataLayer/DataLayerInstance.h
--- a/Engine/Source/Runtime/Engine/Public/WorldPartition/DataLayer/DataLayerInstance.h
+++ b/Engine/Source/Runtime/Engine/Public/WorldPartition/DataLayer/DataLayerInstance.h
@@ -17,6 +17,7 @@ class FText;
 class IStreamingGenerationErrorHandler;
 class UWorld;
 class UDataLayerAsset;
+class FDataLayerInstanceDesc;
 
 UENUM(BlueprintType)
 enum class EDataLayerRuntimeState : uint8
@@ -67,8 +68,12 @@ public:
 	ENGINE_API AWorldDataLayers* GetOuterWorldDataLayers() const;
 
 #if WITH_EDITOR
+	//~ Begin UObject
+	ENGINE_API virtual bool IsAsset() const override;
+	ENGINE_API virtual void GetAssetRegistryTags(TArray<FAssetRegistryTag>& OutTags) const;
 	ENGINE_API virtual void PreEditUndo() override;
 	ENGINE_API virtual void PostEditUndo() override;
+	//~ End UObject
 
 	ENGINE_API void SetVisible(bool bIsVisible);
 	ENGINE_API void SetIsInitiallyVisible(bool bIsInitiallyVisible);
@@ -114,6 +119,9 @@ public:
 	// Whether the DataLayer was created by a user and can be deleted by a user.
 	virtual bool IsUserManaged() const { return true; }
 
+	static bool GetAssetRegistryInfoFromPackage(FName InDataLayerInstancePackageName, FDataLayerInstanceDesc& OutDataLayerInstanceDesc);
+	static bool GetAssetRegistryInfoFromPackage(const FAssetData& InAsset, FDataLayerInstanceDesc& OutDataLayerInstanceDesc);
+
 	ENGINE_API virtual bool Validate(IStreamingGenerationErrorHandler* ErrorHandler) const;
 
 	UE_DEPRECATED(5.3, "Use CanEditShortName instead")
diff --git a/Engine/Source/Runtime/Engine/Public/WorldPartition/DataLayer/WorldDataLayers.h b/Engine/Source/Runtime/Engine/Public/WorldPartition/DataLayer/WorldDataLayers.h
--- a/Engine/Source/Runtime/Engine/Public/WorldPartition/DataLayer/WorldDataLayers.h
+++ b/Engine/Source/Runtime/Engine/Public/WorldPartition/DataLayer/WorldDataLayers.h
@@ -10,6 +10,8 @@
 #include "WorldPartition/DataLayer/DataLayer.h"
 #include "Net/Core/PushModel/PushModel.h"
 #include "Engine/World.h"
+#include "Engine/Level.h"
+#include "ExternalPackageHelper.h"
 #include "WorldDataLayers.generated.h"
 
 class UDEPRECATED_DataLayer;
@@ -48,6 +50,7 @@ public:
 	ENGINE_API virtual void PostLoad() override;
 	ENGINE_API virtual void RewindForReplay() override;
 	ENGINE_API virtual void BeginPlay() override;
+	virtual void Serialize(FArchive& Ar) override;
 
 #if WITH_EDITOR
 	ENGINE_API virtual void PreEditUndo() override;
@@ -62,14 +65,14 @@ public:
 	static ENGINE_API AWorldDataLayers* Create(UWorld* World, FName InWorldDataLayerName = NAME_None);
 	static ENGINE_API AWorldDataLayers* Create(const FActorSpawnParameters& SpawnParameters);
 
-	bool IsEmpty() const { return DataLayerInstances.IsEmpty(); }
+	bool IsEmpty() const;
 	bool HasDeprecatedDataLayers() const { return bHasDeprecatedDataLayers; }
 
 	template<class DataLayerInstanceType, typename ...CreationsArgs>
 	DataLayerInstanceType* CreateDataLayer(CreationsArgs... InCreationArgs);
 
-	ENGINE_API bool RemoveDataLayer(const UDataLayerInstance* InDataLayer);
-	ENGINE_API int32 RemoveDataLayers(const TArray<UDataLayerInstance*>& InDataLayerInstances);
+	ENGINE_API bool RemoveDataLayer(const UDataLayerInstance* InDataLayer, bool bInResolveActorDescContainers = true);
+	ENGINE_API int32 RemoveDataLayers(const TArray<UDataLayerInstance*>& InDataLayerInstances, bool bInResolveActorDescContainers = true);
 	ENGINE_API void SetAllowRuntimeDataLayerEditing(bool bInAllowRuntimeDataLayerEditing);
 	bool GetAllowRuntimeDataLayerEditing() const { return bAllowRuntimeDataLayerEditing; }
 
@@ -93,6 +96,10 @@ public:
 	ENGINE_API bool IsReadOnly() const;
 	ENGINE_API bool IsSubWorldDataLayers() const;
 	ENGINE_API bool IsRuntimeRelevant() const;
+
+	ENGINE_API bool SupportsExternalPackageDataLayerInstances() const;
+	ENGINE_API bool IsUsingExternalPackageDataLayerInstances() const { return bUseExternalPackageDataLayerInstances; }
+	ENGINE_API bool SetUseExternalPackageDataLayerInstances(bool bInNewValue, bool bInInteractiveMode = false);
 #endif
 
 	static FName GetWorldPartionWorldDataLayersName() { return FName(TEXT("WorldDataLayers")); } // reserved for ULevel::WorldDataLayers
@@ -184,15 +191,21 @@ protected:
 	ENGINE_API void OnRep_EffectiveLoadedDataLayerNames();
 
 private:
+	ENGINE_API void OnDataLayerManagerInitialized();
+	ENGINE_API void OnDataLayerManagerDeinitialized();
 	ENGINE_API void ResolveEffectiveRuntimeState(const UDataLayerInstance* InDataLayer, bool bInNotifyChange = true);
 	ENGINE_API void DumpDataLayerRecursively(const UDataLayerInstance* DataLayer, FString Prefix, FOutputDevice& OutputDevice) const;
+	ENGINE_API TSet<TObjectPtr<UDataLayerInstance>>& GetDataLayerInstances();
+	ENGINE_API const TSet<TObjectPtr<UDataLayerInstance>>& GetDataLayerInstances() const;
 
 #if WITH_EDITOR
 	virtual bool ActorTypeIsMainWorldOnly() const override { return true; }
 
+	ENGINE_API void AddDataLayerInstance(UDataLayerInstance* InDataLayerInstance);
 	ENGINE_API void ConvertDataLayerToInstancces();
 	ENGINE_API void UpdateContainsDeprecatedDataLayers();
 	ENGINE_API void ResolveActorDescContainers();
+
 	friend class UDataLayerInstanceWithAsset;
 
 	// Used to compare state pre/post undo
@@ -205,6 +218,9 @@ private:
 #endif
 
 #if WITH_EDITORONLY_DATA
+	UPROPERTY()
+	bool bUseExternalPackageDataLayerInstances;
+
 	// True when Runtime Data Layer editing is allowed.
 	UPROPERTY()
 	bool bAllowRuntimeDataLayerEditing;
@@ -218,6 +234,14 @@ private:
 	UPROPERTY()
 	TSet<TObjectPtr<UDataLayerInstance>> DataLayerInstances;
 
+	/** Data layer instances stored in their external package (only used when UseExternalPackageDataLayerInstances is True) */
+	UPROPERTY(Transient)
+	TSet<TObjectPtr<UDataLayerInstance>> ExternalPackageDataLayerInstances;
+
+	/** Temporary array containing data layer instances manually loaded from their external packages */
+	UPROPERTY(Transient)
+	TArray<TObjectPtr<UDataLayerInstance>> LoadedExternalPackageDataLayerInstances;
+
 	static_assert(DATALAYER_TO_INSTANCE_RUNTIME_CONVERSION_ENABLED, "DeprecatedDataLayerNameToDataLayerInstance Property is deprecated and needs to be deleted.");
 	UPROPERTY()
 	TMap<FName, TWeakObjectPtr<UDataLayerInstance>> DeprecatedDataLayerNameToDataLayerInstance;
@@ -263,6 +287,8 @@ private:
 	static_assert(DATALAYER_TO_INSTANCE_RUNTIME_CONVERSION_ENABLED, "bHasDeprecatedDataLayers is deprecated and needs to be deleted.");
 	bool bHasDeprecatedDataLayers;
 
+	friend class UDataLayerManager;
+
 public:
 	DECLARE_DELEGATE_RetVal_ThreeParams(bool, FDataLayersFilterDelegate, FName /*DataLayerName*/, EDataLayerRuntimeState /*CurrentState*/, EDataLayerRuntimeState /*TargetState*/);
 
@@ -283,16 +309,40 @@ struct HasDataLayerInstanceClass<T, decltype(&T::GetDataLayerInstanceClass, 0)>
 template<class DataLayerInstanceType, typename ...CreationsArgs>
 DataLayerInstanceType* AWorldDataLayers::CreateDataLayer(CreationsArgs... InCreationArgs)
 {
-	Modify();
-
 	UClass* ClassToUse = DataLayerInstanceType::StaticClass();
 	if constexpr (HasDataLayerInstanceClass<DataLayerInstanceType>::value)
 		ClassToUse = DataLayerInstanceType::GetDataLayerInstanceClass();
+	UPackage* ExternalPackage = nullptr;
+	UObject* OuterObject = this;
+	FName NewObjectName;
+
+	if (IsUsingExternalPackageDataLayerInstances())
+	{
+		check(SupportsExternalPackageDataLayerInstances());
+		const ULevel* Level = GetLevel();
+		const bool bIsTransient = (Level->IsInstancedLevel() && !Level->IsPersistentLevel()) || Level->GetPackage()->HasAnyPackageFlags(PKG_PlayInEditor);
+		check(!bIsTransient);
+
+		// We generate a globally unique name to avoid any potential clash of 2 users creating the same object
+		const FGuid NewObjectGuid = FGuid::NewGuid();
+		FString ShortName = DataLayerInstanceType::StaticClass()->GetName() + TEXT("_UID_") + NewObjectGuid.ToString(EGuidFormats::UniqueObjectGuid);
+		TStringBuilderWithBuffer<TCHAR, NAME_SIZE> GloballyUniqueObjectPath;
+		GloballyUniqueObjectPath += OuterObject->GetPathName();
+		GloballyUniqueObjectPath += TEXT(".");
+		GloballyUniqueObjectPath += ShortName;
+		ExternalPackage = FExternalPackageHelper::CreateExternalPackage(OuterObject, *GloballyUniqueObjectPath);
+		check(ExternalPackage);
+		NewObjectName = *ShortName;
+	}
+	else
+	{
+		NewObjectName = DataLayerInstanceType::MakeName(Forward<CreationsArgs>(InCreationArgs)...);
+	}
 
-	DataLayerInstanceType* NewDataLayer = NewObject<DataLayerInstanceType>(this, ClassToUse, DataLayerInstanceType::MakeName(Forward<CreationsArgs>(InCreationArgs)...), RF_Transactional | RF_NoFlags);
+	DataLayerInstanceType* NewDataLayer = NewObject<DataLayerInstanceType>(OuterObject, ClassToUse, NewObjectName, RF_Transactional, nullptr, /*bCopyTransientsFromClassDefaults*/false, /*InstanceGraph*/nullptr, ExternalPackage);
 	check(NewDataLayer != nullptr);
 	NewDataLayer->OnCreated(Forward<CreationsArgs>(InCreationArgs)...);
-	DataLayerInstances.Add(NewDataLayer);
+	AddDataLayerInstance(NewDataLayer);
 
 	UpdateContainsDeprecatedDataLayers();
 	ResolveActorDescContainers();
diff --git a/Engine/Source/Runtime/Engine/Public/WorldPartition/DataLayer/WorldDataLayersActorDesc.h b/Engine/Source/Runtime/Engine/Public/WorldPartition/DataLayer/WorldDataLayersActorDesc.h
--- a/Engine/Source/Runtime/Engine/Public/WorldPartition/DataLayer/WorldDataLayersActorDesc.h
+++ b/Engine/Source/Runtime/Engine/Public/WorldPartition/DataLayer/WorldDataLayersActorDesc.h
@@ -6,6 +6,7 @@
 
 #include "WorldPartition/WorldPartitionActorDesc.h"
 #include "WorldPartition/DataLayer/DataLayerType.h"
+#include "Misc/Optional.h"
 
 class FDataLayerInstanceDesc
 {
@@ -53,6 +54,7 @@ private:
 	bool bDeprecatedIsRuntime;
 	//~ End UDeprecatedDataLayerInstance
 
+	friend class UDataLayerInstance;
 };
 
 /**
@@ -63,7 +65,7 @@ class FWorldDataLayersActorDesc : public FWorldPartitionActorDesc
 public:
 	ENGINE_API FWorldDataLayersActorDesc();
 	bool IsValid() const { return bIsValid; }
-	const TArray<FDataLayerInstanceDesc>& GetDataLayerInstances() const { return DataLayerInstances; }
+	ENGINE_API const TArray<FDataLayerInstanceDesc>& GetDataLayerInstances() const;
 	ENGINE_API const FDataLayerInstanceDesc* GetDataLayerInstanceFromInstanceName(FName InDataLayerInstanceName) const;
 	ENGINE_API const FDataLayerInstanceDesc* GetDataLayerInstanceFromAssetPath(FName InDataLayerAssetPath) const;
 
@@ -77,7 +79,13 @@ protected:
 	//~ End FWorldPartitionActorDesc Interface.
 
 private:
+
+	void ForEachDataLayerInstanceDesc(TFunctionRef<bool(const FDataLayerInstanceDesc&)> Func) const;
+	const TArray<FDataLayerInstanceDesc>& GetExternalPackageDataLayerInstances() const;
+
 	TArray<FDataLayerInstanceDesc> DataLayerInstances;
+	mutable TOptional<TArray<FDataLayerInstanceDesc>> ExternalPackageDataLayerInstances;
 	bool bIsValid;
+	bool bUseExternalPackageDataLayerInstances;
 };
 #endif
