Fix crash in NetworkPhysicsComponent on editor shutdown

diff --git a/Engine/Source/Runtime/Engine/Private/PhysicsEngine/NetworkPhysicsComponent.cpp b/Engine/Source/Runtime/Engine/Private/PhysicsEngine/NetworkPhysicsComponent.cpp
--- a/Engine/Source/Runtime/Engine/Private/PhysicsEngine/NetworkPhysicsComponent.cpp
+++ b/Engine/Source/Runtime/Engine/Private/PhysicsEngine/NetworkPhysicsComponent.cpp
@@ -253,22 +253,24 @@ void UNetworkPhysicsComponent::BeginPlay()
 void UNetworkPhysicsComponent::InitializeComponent()
 {
 	Super::InitializeComponent();
-	UWorld* World = GetWorld();
-
-	if (UNetworkPhysicsSystem* NetworkManager = World->GetSubsystem<UNetworkPhysicsSystem>())
+	if (UWorld* World = GetWorld())
 	{
-		NetworkManager->RegisterNetworkComponent(this);
+		if (UNetworkPhysicsSystem* NetworkManager = World->GetSubsystem<UNetworkPhysicsSystem>())
+		{
+			NetworkManager->RegisterNetworkComponent(this);
+		}
 	}
 }
 
 void UNetworkPhysicsComponent::UninitializeComponent()
 {
 	Super::UninitializeComponent();
-	UWorld* World = GetWorld();
-
-	if (UNetworkPhysicsSystem* NetworkManager = World->GetSubsystem<UNetworkPhysicsSystem>())
+	if (UWorld* World = GetWorld())
 	{
-		NetworkManager->UnregisterNetworkComponent(this);
+		if (UNetworkPhysicsSystem* NetworkManager = World->GetSubsystem<UNetworkPhysicsSystem>())
+		{
+			NetworkManager->UnregisterNetworkComponent(this);
+		}
 	}
 }
 
