Fix when an ISMC has bComponentMarkedDirty sets to true, but fails to create its proxy (eg, num instances == 0). Since FPrimitiveInstanceDataManager::FlushChanges is only called when the proxy exists, we could get into a bad state where bComponentMarkedDirty=true, CreateSceneProxy fails and any subsequent call to MarkComponentRenderInstancesDirty would be ignored, until MarkRenderStateDirty is called

#rnx

diff --git a/Engine/Source/Runtime/Engine/Private/HierarchicalInstancedStaticMesh.cpp b/Engine/Source/Runtime/Engine/Private/HierarchicalInstancedStaticMesh.cpp
--- a/Engine/Source/Runtime/Engine/Private/HierarchicalInstancedStaticMesh.cpp
+++ b/Engine/Source/Runtime/Engine/Private/HierarchicalInstancedStaticMesh.cpp
@@ -2958,6 +2958,8 @@ FPrimitiveSceneProxy* UHierarchicalInstancedStaticMeshComponent::CreateSceneProx
 	QUICK_SCOPE_CYCLE_COUNTER(STAT_HierarchicalInstancedStaticMeshComponent_CreateSceneProxy);
 	SCOPE_CYCLE_COUNTER(STAT_FoliageCreateProxy);
 
+	PrimitiveInstanceDataManager.ResetComponentDirtyTracking();
+
 	if (ProxySize)
 	{
 		DEC_DWORD_STAT_BY(STAT_FoliageInstanceBuffers, ProxySize);
diff --git a/Engine/Source/Runtime/Engine/Private/InstancedStaticMesh.cpp b/Engine/Source/Runtime/Engine/Private/InstancedStaticMesh.cpp
--- a/Engine/Source/Runtime/Engine/Private/InstancedStaticMesh.cpp
+++ b/Engine/Source/Runtime/Engine/Private/InstancedStaticMesh.cpp
@@ -2128,6 +2128,8 @@ void UInstancedStaticMeshComponent::SendRenderInstanceData_Concurrent()
 {
 	Super::SendRenderInstanceData_Concurrent();
 
+	PrimitiveInstanceDataManager.ResetComponentDirtyTracking();
+
 	// If the primitive isn't hidden update its instances.
 	const bool bDetailModeAllowsRendering = DetailMode <= GetCachedScalabilityCVars().DetailMode;
 	// The proxy may not be created, this can happen when a SM is async loading for example.
@@ -2185,6 +2187,8 @@ FPrimitiveSceneProxy* UInstancedStaticMeshComponent::CreateSceneProxy()
 {
 	ProxySize = 0;
 
+	PrimitiveInstanceDataManager.ResetComponentDirtyTracking();
+
 	// Verify that the mesh is valid before using it.
 	const bool bMeshIsValid =
 		// make sure we have instances
diff --git a/Engine/Source/Runtime/Engine/Public/InstancedStaticMesh/ISMInstanceDataManager.h b/Engine/Source/Runtime/Engine/Public/InstancedStaticMesh/ISMInstanceDataManager.h
--- a/Engine/Source/Runtime/Engine/Public/InstancedStaticMesh/ISMInstanceDataManager.h
+++ b/Engine/Source/Runtime/Engine/Public/InstancedStaticMesh/ISMInstanceDataManager.h
@@ -154,6 +154,9 @@ public:
 	 */
 	ENGINE_API SIZE_T GetAllocatedSize() const;
 
+
+	inline void ResetComponentDirtyTracking() { bComponentMarkedDirty = false; }
+
 private:
 	template <typename TaskLambdaType>
 	static void BeginUpdateTask(FInstanceDataUpdateTaskInfo &InstanceDataUpdateTaskInfo, TaskLambdaType &&TaskLambda, const FInstanceDataBufferHeader &InInstanceDataBufferHeader);
