Fixed CIS Static Analysis

diff --git a/Engine/Source/Runtime/Engine/Public/ExternalPackageHelper.h b/Engine/Source/Runtime/Engine/Public/ExternalPackageHelper.h
--- a/Engine/Source/Runtime/Engine/Public/ExternalPackageHelper.h
+++ b/Engine/Source/Runtime/Engine/Public/ExternalPackageHelper.h
@@ -100,6 +100,7 @@ private:
 template<typename T>
 void FExternalPackageHelper::LoadObjectsFromExternalPackages(UObject* InOuter, TFunctionRef<void(T*)> Operation)
 {
+	check(InOuter);
 	UPackage* OutermostPackage = InOuter->IsA<UPackage>() ? CastChecked<UPackage>(InOuter) : InOuter->GetOutermostObject()->GetPackage();
 	const FString ExternalObjectsPath = FExternalPackageHelper::GetExternalObjectsPath(OutermostPackage, FString(), /*bTryUsingPackageLoadedPath*/ true);
 	TArray<FString> ObjectPackageNames;
