Dynamic Submix fixes:
- Fix unregistering submix on disconnect
- Fix debug command to correctly view graph

diff --git a/Engine/Source/Editor/AudioEditor/Private/SoundSubmixGraphSchema.cpp b/Engine/Source/Editor/AudioEditor/Private/SoundSubmixGraphSchema.cpp
--- a/Engine/Source/Editor/AudioEditor/Private/SoundSubmixGraphSchema.cpp
+++ b/Engine/Source/Editor/AudioEditor/Private/SoundSubmixGraphSchema.cpp
@@ -132,14 +132,14 @@ bool USoundSubmixGraphSchema::ConnectionCausesLoop(const UEdGraphPin* InputPin,
 				return true;
 			}
 
-			if (OutputNode->SoundSubmix->RecurseCheckChild(MasterSubmix))
+			if (SubmixUtils::FindInGraph(MasterSubmix, OutputNode->SoundSubmix, false))
 			{
 				return true;
 			}
 		}
 	}
 
-	return OutputNode->SoundSubmix->RecurseCheckChild(InputNode->SoundSubmix);
+	return SubmixUtils::FindInGraph(OutputNode->SoundSubmix, InputNode->SoundSubmix, false);
 }
 
 void USoundSubmixGraphSchema::GetAssetsGraphHoverMessage(const TArray<FAssetData>& Assets, const UEdGraph* HoverGraph, FString& OutTooltipText, bool& OutOkIcon) const
diff --git a/Engine/Source/Runtime/AudioMixer/Private/AudioMixerDevice.cpp b/Engine/Source/Runtime/AudioMixer/Private/AudioMixerDevice.cpp
--- a/Engine/Source/Runtime/AudioMixer/Private/AudioMixerDevice.cpp
+++ b/Engine/Source/Runtime/AudioMixer/Private/AudioMixerDevice.cpp
@@ -198,48 +198,44 @@ namespace Audio
 	);
 
 
-	static void DrawSubmixHeirarchy(USoundSubmixBase* InSubmix, FMixerSubmix* InInstance, const FMixerDevice* InDevice, int32 InIdent, FOutputDevice& Ar);
-	static void DrawSubmixHeirarchy(USoundSubmixBase* InSubmix, FMixerSubmix* InInstance, const FMixerDevice* InDevice, int32 InIdent, FOutputDevice& Ar)
+	
+	static void DrawSubmixHeirarchy(USoundSubmixBase* InSubmix, const TSharedPtr<Audio::FMixerSubmix, ESPMode::ThreadSafe> InInstance, const FMixerDevice* InDevice, int32 InIdent, FOutputDevice& Ar, const TCHAR* GroupingText)
 	{
 		if (!InSubmix)
 		{
 			return;
 		}
+		
+		FString Indet = FCString::Spc(InIdent*3);
+		FString FxChain;
+		if (USoundSubmix* Submix = Cast<USoundSubmix>(InSubmix))
+		{		
+			for (const TObjectPtr<USoundEffectSubmixPreset>& i: Submix->SubmixEffectChain)
+			{
+				FxChain += FString::Printf(TEXT("[%s]"), *i->GetName());
+			}
+		}
 
-		auto SubmixInfo = [](USoundSubmixBase* InUObj, FMixerSubmix* InInst, USoundSubmixBase* InParent, FMixerSubmix* InParentInst) -> FString
-		{
-			FMixerSubmixPtr ParentInstance = InInst->GetParent().Pin();
-			const bool bParentLinkSeemsSane = ParentInstance.IsValid() ? ParentInstance.Get() == InParentInst : false;
-			const int32 NumFxRunning = InInst->GetNumEffects();
-			const int32 NumChildren = InInst->GetChildren().Num();
-			return FString::Printf(TEXT("ParentLinkGood=%s,NumFxRunning=%d,NumChildren=%d"), ToCStr(LexToString(bParentLinkSeemsSane)),NumFxRunning, NumChildren);
-		};
-
-		for (TObjectPtr<USoundSubmixBase> i : InSubmix->ChildSubmixes)
+		Ar.Logf(TEXT("%sName=%s,Instance=0x%p,Id=%u,Fx=%s,[%s]"), *Indet, *InSubmix->GetName(), InInstance.Get(), InInstance ? InInstance->GetId() : 0, *FxChain, GroupingText);
+		for (const auto& i : InSubmix->ChildSubmixes)
 		{
-			FMixerSubmixPtr ChildInstance = InDevice->GetSubmixInstance(i).Pin();
-			FString Info = SubmixInfo(i.Get(), ChildInstance.Get(), InSubmix, InInstance);
-			Ar.Logf(TEXT("%s [%s](Static)(Instance=0x%p)%s"), FCString::Spc(InIdent * 2), *InSubmix->GetName(), ChildInstance.Get(), *Info);
-			DrawSubmixHeirarchy(i.Get(), ChildInstance.Get(), InDevice, InIdent + 1, Ar);
+			DrawSubmixHeirarchy(i, InDevice->GetSubmixInstance(i).Pin(), InDevice, InIdent+1, Ar, TEXT("Static"));
 		}
-		const auto& DynamicChildren = InSubmix->DynamicChildSubmixes.FindOrAdd(InDevice->DeviceID).ChildSubmixes;
-		for (TObjectPtr<USoundSubmixBase> i : DynamicChildren)
+		const auto& DynamicSubmixes = InSubmix->DynamicChildSubmixes.FindOrAdd(InDevice->DeviceID);
+		for (const auto& i : DynamicSubmixes.ChildSubmixes)
 		{
-			FMixerSubmixPtr ChildInstance = InDevice->GetSubmixInstance(i).Pin();
-			FString Info = SubmixInfo(i.Get(), ChildInstance.Get(), InSubmix, InInstance);
-			Ar.Logf(TEXT("%s [%s](Dynamic)(Instance=0x%p)%s"), FCString::Spc(InIdent * 2), *InSubmix->GetName(), ChildInstance.Get(), *Info);
-			DrawSubmixHeirarchy(i.Get(), ChildInstance.Get(), InDevice,  InIdent + 1, Ar);
+			DrawSubmixHeirarchy(i, InDevice->GetSubmixInstance(i).Pin(), InDevice, InIdent+1, Ar, TEXT("Dynamic"));
 		}
 	}
 
 	void FMixerDevice::DrawSubmixes(FOutputDevice& Output) const
 	{
-		Output.Logf(TEXT("Submix Graph for AudioMixerDevice=%u/(0x%p) "), DeviceID, this);
+		Output.Logf(TEXT("AudioDevice=%d, Device Instance=0x%p"), DeviceID, this);		
+
 		for (int32 i = 0; i < RequiredSubmixes.Num(); i++)
 		{
-			FMixerSubmixPtr Instance = GetRequiredSubmixInstance(RequiredSubmixes[i]);
-			Output.Logf(TEXT("Slot=[%s] (Instance=0x%p)"), ToCStr(LexToString((ERequiredSubmixes)i)), Instance.Get());
-			DrawSubmixHeirarchy(RequiredSubmixes[i], Instance.Get(), this, 1, Output);
+			Output.Logf(TEXT("SlotName=[%s]"), ToCStr(LexToString((ERequiredSubmixes)i)));
+			DrawSubmixHeirarchy(RequiredSubmixes[i], RequiredSubmixInstances[i], this, 1, Output, TEXT("In Slot"));
 		}
 	}
 
@@ -2012,7 +2008,7 @@ namespace Audio
 	ISoundfieldFactory* FMixerDevice::GetFactoryForSubmixInstance(FMixerSubmixWeakPtr& SoundSubmixPtr)
 	{
 		FMixerSubmixPtr SubmixPtr = SoundSubmixPtr.Pin();
-		if (ensure(SubmixPtr.IsValid()))
+		if (SubmixPtr.IsValid())
 		{
 			return SubmixPtr->GetSoundfieldFactory();
 		}
diff --git a/Engine/Source/Runtime/Engine/Classes/Sound/SoundSubmix.h b/Engine/Source/Runtime/Engine/Classes/Sound/SoundSubmix.h
--- a/Engine/Source/Runtime/Engine/Classes/Sound/SoundSubmix.h
+++ b/Engine/Source/Runtime/Engine/Classes/Sound/SoundSubmix.h
@@ -195,6 +195,27 @@ public:
 	// Dynamic Child submixes
 	UPROPERTY(Transient)
 	TMap<uint32, FDynamicChildSubmix> DynamicChildSubmixes;
+		
+	/** Dynamically Connects to a parent submix.
+	* @param	WorldContextObject	UObject that's used to GetWorld
+	* @param	InParent	Parent Submix to connect to
+	**/
+	UFUNCTION(BlueprintCallable, Category = "Submix", meta = (WorldContext = "WorldContextObject", DisplayName = "Connect"))	
+	ENGINE_API virtual bool DynamicConnect(
+		const UObject* WorldContextObject, 
+		UPARAM(DisplayName="Parent") USoundSubmixBase* InParent) 
+	{ 
+		return false; 
+	}
+
+	/** Dynamically Disconnect from a parent.
+	* @param	WorldContextObject	UObject that's used to GetWorld
+	**/
+	UFUNCTION(BlueprintCallable, Category = "Submix", meta = (WorldContext = "WorldContextObject", DisplayName = "Disconnect"))
+	ENGINE_API virtual bool DynamicDisconnect(const UObject* WorldContextObject) 
+	{ 
+		return false; 
+	}
 	
 protected:
 	//~ Begin UObject Interface.
@@ -203,14 +224,10 @@ protected:
 	ENGINE_API virtual void PostLoad() override;
 
 public:
+
 	// Sound Submix Editor functionality
 #if WITH_EDITOR
 
-	/**
-	* @return true if the child sound class exists in the tree
-	*/
-	ENGINE_API bool RecurseCheckChild(const USoundSubmixBase* ChildSoundSubmix) const;
-
 	/**
 	* Add Referenced objects
 	*
@@ -255,15 +272,7 @@ public:
 	
 	UPROPERTY(Transient)
 	TMap<uint32, TObjectPtr<USoundSubmixBase>> DynamicParentSubmix;
-
-	UFUNCTION(BlueprintCallable, Category = "Audio", meta = (WorldContext = "WorldContextObject", DisplayName = "Connect"))
-	ENGINE_API bool DynamicConnect(const UObject* WorldContextObject, USoundSubmixBase* InParent);
-	ENGINE_API bool DynamicConnect(const FAudioDeviceHandle& Handle, USoundSubmixBase* InParent);
-
-	UFUNCTION(BlueprintCallable, Category = "Audio", meta = (WorldContext = "WorldContextObject", DisplayName = "Disconnect"))
-	ENGINE_API bool DynamicDisconnect(const UObject* WorldContextObject);
-	ENGINE_API bool DynamicDisconnect(const FAudioDeviceHandle& Handle);
-
+	
 	/**
 	* Set the parent submix of this SoundSubmix, removing it as a child from its previous owner
 	*
@@ -271,6 +280,11 @@ public:
 	*/
 	ENGINE_API void SetParentSubmix(USoundSubmixBase* InParentSubmix);
 
+	ENGINE_API virtual bool DynamicConnect(const UObject* WorldContextObject, USoundSubmixBase* Parent) override;
+	ENGINE_API bool DynamicConnect(FAudioDeviceHandle Handle, USoundSubmixBase* Parent);
+	ENGINE_API virtual bool DynamicDisconnect(const UObject* WorldContextObject) override;
+	ENGINE_API bool DynamicDisconnect(FAudioDeviceHandle Handle);
+
 protected:
 
 #if WITH_EDITOR
@@ -589,6 +603,9 @@ namespace SubmixUtils
 {
 	ENGINE_API bool AreSubmixFormatsCompatible(const USoundSubmixBase* ChildSubmix, const USoundSubmixBase* ParentSubmix);
 
+	ENGINE_API bool FindInGraph(const USoundSubmixBase* InEntryPoint, const USoundSubmixBase* InToMatch, bool bShouldAcsend, FAudioDeviceHandle InDevice = {});
+	ENGINE_API const USoundSubmixBase* FindRoot(const USoundSubmixBase* InStartingPoint, FAudioDeviceHandle InDevice);
+
 #if WITH_EDITOR
 	ENGINE_API void RefreshEditorForSubmix(const USoundSubmixBase* InSubmix);
 #endif
diff --git a/Engine/Source/Runtime/Engine/Private/SoundSubmix.cpp b/Engine/Source/Runtime/Engine/Private/SoundSubmix.cpp
--- a/Engine/Source/Runtime/Engine/Private/SoundSubmix.cpp
+++ b/Engine/Source/Runtime/Engine/Private/SoundSubmix.cpp
@@ -5,6 +5,7 @@
 #include "Engine/Engine.h"
 #include "Engine/World.h"
 #include "Sound/SampleBufferIO.h"
+#include "Stats/Stats2.h"
 
 #include UE_INLINE_GENERATED_CPP_BY_NAME(SoundSubmix)
 
@@ -621,7 +622,6 @@ void USoundSubmixBase::PostLoad()
 	}
 }
 
-
 TObjectPtr<USoundSubmixBase> USoundSubmixWithParentBase::GetParent(Audio::FDeviceId InDeviceId) const
 {
 	// Dynamic parent?
@@ -650,7 +650,7 @@ bool USoundSubmixWithParentBase::DynamicConnect(const UObject* WorldContextObjec
 	return DynamicConnect(World->GetAudioDevice(), InParent);
 }
 
-bool USoundSubmixWithParentBase::DynamicConnect(const FAudioDeviceHandle& Handle, USoundSubmixBase* InParent)
+bool USoundSubmixWithParentBase::DynamicConnect(FAudioDeviceHandle Handle, USoundSubmixBase* InParent)
 {
 	if (!Handle.IsValid())
 	{
@@ -658,6 +658,19 @@ bool USoundSubmixWithParentBase::DynamicConnect(const FAudioDeviceHandle& Handle
 		return false;
 	}
 
+	if (InParent && !SubmixUtils::AreSubmixFormatsCompatible(this, InParent))
+	{
+		UE_LOG(LogAudio, Warning, TEXT("Submix (DynamicConnect): Submix connot be connected are they are incompatible [%s] and [%s]"), *GetName(), *GetNameSafe(InParent));
+		return false;
+	}
+
+	// Already part of the graph?
+	if (InParent && SubmixUtils::FindInGraph(InParent, this, true, Handle))
+	{
+		UE_LOG(LogAudio, Warning, TEXT("Submix (DynamicConnect): Submix [%s] is already part of the graph"), *GetName());
+		return false;
+	}
+
 	const Audio::DeviceID Id = Handle.GetDeviceID();
 
 	TObjectPtr<USoundSubmixBase>& CurrentParent = DynamicParentSubmix.FindOrAdd(Id);
@@ -676,18 +689,21 @@ bool USoundSubmixWithParentBase::DynamicConnect(const FAudioDeviceHandle& Handle
 		{
 			CurrentParent->DynamicChildSubmixes.FindOrAdd(Id).ChildSubmixes.AddUnique(this);
 			UE_LOG(LogAudio, Verbose, TEXT("Submix (DynamicConnect): Reparenting [%s] and adding to it's new parent [%s]"), *GetName(), *CurrentParent->GetName());
+
+			// Disable our parents auto disable feature.
+			Handle->SetSubmixAutoDisable(Cast<USoundSubmix>(CurrentParent.Get()), false);
 		}
 
-		if (FAudioDevice* RawDevice = Handle.GetAudioDevice())
-		{
-			RawDevice->RegisterSoundSubmix(this, true);
-			UE_LOG(LogAudio, Verbose, TEXT("Submix (DynamicConnect): Registering [%s] with AudioDevice [%u]"), *GetName(), Id);
+		Handle->RegisterSoundSubmix(this, false);
+		Handle->SetSubmixAutoDisable(Cast<USoundSubmix>(this), CurrentParent != nullptr);
 
-			return CurrentParent != nullptr;
-		}
+		UE_LOG(LogAudio, Verbose, TEXT("Submix (DynamicConnect): Registering [%s] with AudioDevice [%u]"), *GetName(), Id);
+
+		return CurrentParent != nullptr;
 	}
 
-	UE_LOG(LogAudio, Warning, TEXT("Submix (DynamicConnect): Submix [%s] was already connected to [%s]"), *GetName(), *GetNameSafe(CurrentParent));
+	UE_CLOG(InParent, LogAudio, Warning, TEXT("Submix (DynamicConnect): Submix [%s] was already connected to [%s]"), *GetName(), *GetNameSafe(CurrentParent));
+	UE_CLOG(InParent == nullptr, LogAudio, Warning, TEXT("Submix (DynamicConnect): Connecting a [%s] to a null parent, but our parent is already null"), *GetName());
 	return false;
 
 }
@@ -710,7 +726,7 @@ bool USoundSubmixWithParentBase::DynamicDisconnect(const UObject* WorldContextOb
 	return DynamicDisconnect(World->GetAudioDevice());
 }
 
-bool USoundSubmixWithParentBase::DynamicDisconnect(const FAudioDeviceHandle& Handle)
+bool USoundSubmixWithParentBase::DynamicDisconnect(FAudioDeviceHandle Handle)
 {
 	if (!Handle.IsValid())
 	{
@@ -729,13 +745,17 @@ bool USoundSubmixWithParentBase::DynamicDisconnect(const FAudioDeviceHandle& Han
 		CurrentParent->DynamicChildSubmixes.FindOrAdd(Id).ChildSubmixes.Remove(this);
 		CurrentParent = nullptr;
 
-		if (FAudioDevice* RawDevice = Handle.GetAudioDevice())
+		Handle->UnregisterSoundSubmix(this);
+				
+		// If we still have a valid parent static submix? Make sure that's still live and registered.
+		if (ParentSubmix)
 		{
-			RawDevice->UnregisterSoundSubmix(this);
-			UE_LOG(LogAudio, Verbose, TEXT("Submix (DynamicDisconnect): Unregistering [%s] with AudioDevice [%u]"), *GetName(), Id);
-
-			return true;
+			Handle->RegisterSoundSubmix(this, false);
 		}
+		
+		UE_LOG(LogAudio, Verbose, TEXT("Submix (DynamicDisconnect): Unregistering [%s] with AudioDevice [%u]"), *GetName(), Id);
+
+		return true;
 	}
 
 	UE_LOG(LogAudio, Warning, TEXT("Submix (DynamicDisconnect): Submix was not connected to any dynamic parent [%s]"), *GetName());
@@ -777,7 +797,7 @@ void USoundSubmixBase::PostEditChangeProperty(FPropertyChangedEvent& PropertyCha
 			{
 				if (ChildSubmixes[ChildIndex] != nullptr && !BackupChildSubmixes.Contains(ChildSubmixes[ChildIndex]))
 				{
-					if (ChildSubmixes[ChildIndex]->RecurseCheckChild(this))
+					if (SubmixUtils::FindInGraph(this, ChildSubmixes[ChildIndex], false))
 					{
 						// Contains cycle so revert to old layout - launch notification to inform user
 						FNotificationInfo Info(NSLOCTEXT("Engine", "UnableToChangeSoundSubmixChildDueToInfiniteLoopNotification", "Could not change SoundSubmix child as it would create a loop"));
@@ -825,27 +845,6 @@ void USoundSubmixBase::PostEditChangeProperty(FPropertyChangedEvent& PropertyCha
 
 TArray<TObjectPtr<USoundSubmixBase>> USoundSubmixBase::BackupChildSubmixes;
 
-bool USoundSubmixBase::RecurseCheckChild(const USoundSubmixBase* ChildSoundSubmix) const
-{
-	for (int32 Index = 0; Index < ChildSubmixes.Num(); Index++)
-	{
-		if (ChildSubmixes[Index])
-		{
-			if (ChildSubmixes[Index] == ChildSoundSubmix)
-			{
-				return true;
-			}
-
-			if (ChildSubmixes[Index]->RecurseCheckChild(ChildSoundSubmix))
-			{
-				return true;
-			}
-		}
-	}
-
-	return false;
-}
-
 void USoundSubmixWithParentBase::SetParentSubmix(USoundSubmixBase* InParentSubmix)
 {
 	if (ParentSubmix != InParentSubmix)
@@ -1348,6 +1347,68 @@ ENGINE_API bool SubmixUtils::AreSubmixFormatsCompatible(const USoundSubmixBase*
 	return true;
 }
 
+bool SubmixUtils::FindInGraph(
+	const USoundSubmixBase* InEntryPoint, 
+	const USoundSubmixBase* InToMatch, 
+	const bool bStartFromRoot,
+	FAudioDeviceHandle InDevice /*= {}*/)
+{
+	TSet<const USoundSubmixBase*> Visited;
+	TArray<const USoundSubmixBase*> Stack;
+
+	// Optionally ascend to the root
+	const USoundSubmixBase* StartingPoint = InEntryPoint;
+	if (bStartFromRoot)
+	{
+		StartingPoint = FindRoot(InEntryPoint, InDevice);
+	}
+	
+	Stack.Push(StartingPoint);
+	while (Stack.Num() > 0)
+	{
+		if (const USoundSubmixBase* Vertex = Stack.Pop())
+		{
+			if (Vertex == InToMatch)
+			{
+				return true;
+			}
+			else if (!Visited.Contains(Vertex))
+			{
+				// Unlike parents, submixes can have both dynamic and static children so search both.
+				Stack.Append(Vertex->ChildSubmixes);
+				if (const FDynamicChildSubmix* Dynamics = Vertex->DynamicChildSubmixes.Find(InDevice.GetDeviceID()))
+				{
+					Stack.Append(Dynamics->ChildSubmixes);
+				}
+			}
+		}
+	}
+	return false;
+
+}
+
+const USoundSubmixBase* SubmixUtils::FindRoot(const USoundSubmixBase* InStartingPoint, FAudioDeviceHandle InDevice)
+{	
+	const USoundSubmixBase* HighestPoint = InStartingPoint;
+	while (HighestPoint)
+	{
+		const USoundSubmixWithParentBase* WithParent = Cast<const USoundSubmixWithParentBase>(HighestPoint);
+		if (!WithParent)
+		{
+			break;
+		}
+
+		const USoundSubmixBase* Parent = WithParent->GetParent(InDevice.GetDeviceID());
+		if (!Parent)
+		{
+			break;
+		}
+
+		HighestPoint = Parent;
+	}
+	return HighestPoint;
+}
+
 #if WITH_EDITOR
 
 ENGINE_API void SubmixUtils::RefreshEditorForSubmix(const USoundSubmixBase* InSubmix)
