Chaos - earlier detection of nulls children in UClusterUnionComponent with some useful logging

diff --git a/Engine/Source/Runtime/Engine/Private/PhysicsEngine/ClusterUnionComponent.cpp b/Engine/Source/Runtime/Engine/Private/PhysicsEngine/ClusterUnionComponent.cpp
--- a/Engine/Source/Runtime/Engine/Private/PhysicsEngine/ClusterUnionComponent.cpp
+++ b/Engine/Source/Runtime/Engine/Private/PhysicsEngine/ClusterUnionComponent.cpp
@@ -1150,11 +1150,19 @@ void UClusterUnionComponent::ForceSetChildToParent(UPrimitiveComponent* InCompon
 	for (int32 Index = 0; Index < BoneIds.Num(); ++Index)
 	{
 		Chaos::FPhysicsObjectHandle Handle = InComponent->GetPhysicsObjectById(BoneIds[Index]);
-		Objects.Add(Handle);
+		const bool bIsValidHandle = ensureMsgf(Handle != nullptr, TEXT("UClusterUnionComponent::ForceSetChildToParent invalid BoneID %d on %s"), BoneIds[Index], *InComponent->GetFullName());
+
+		if (bIsValidHandle)
+		{
+			Objects.Add(Handle);
+		}
 	}
 
 	// Need to lock the ChildToParent to prevent the PT from trying to compute it.
-	PhysicsProxy->BulkSetChildToParent_External(Objects, ChildToParent, true);
+	if (!Objects.IsEmpty())
+	{
+		PhysicsProxy->BulkSetChildToParent_External(Objects, ChildToParent, true);
+	}
 }
 
 void UClusterUnionComponent::SetSimulatePhysics(bool bSimulate)
diff --git a/Engine/Source/Runtime/Experimental/Chaos/Private/PhysicsProxy/ClusterUnionPhysicsProxy.cpp b/Engine/Source/Runtime/Experimental/Chaos/Private/PhysicsProxy/ClusterUnionPhysicsProxy.cpp
--- a/Engine/Source/Runtime/Experimental/Chaos/Private/PhysicsProxy/ClusterUnionPhysicsProxy.cpp
+++ b/Engine/Source/Runtime/Experimental/Chaos/Private/PhysicsProxy/ClusterUnionPhysicsProxy.cpp
@@ -700,7 +700,10 @@ namespace Chaos
 
 	void FClusterUnionPhysicsProxy::SetChildToParent_External(FPhysicsObjectHandle Child, const FTransform& RelativeTransform, bool bLock)
 	{
-		BulkSetChildToParent_External({ Child }, { RelativeTransform }, bLock);
+		if (ensure(Child != nullptr))
+		{
+			BulkSetChildToParent_External({ Child }, { RelativeTransform }, bLock);
+		}
 	}
 
 	void FClusterUnionPhysicsProxy::BulkSetChildToParent_External(const TArray<FPhysicsObjectHandle>& Objects, const TArray<FTransform>& Transforms, bool bLock)
