LFV - fixed visual artefact for low falloff when the camera looks at the horizon.

diff --git a/Engine/Source/Runtime/Engine/Classes/Components/LocalFogVolumeComponent.h b/Engine/Source/Runtime/Engine/Classes/Components/LocalFogVolumeComponent.h
--- a/Engine/Source/Runtime/Engine/Classes/Components/LocalFogVolumeComponent.h
+++ b/Engine/Source/Runtime/Engine/Classes/Components/LocalFogVolumeComponent.h
@@ -24,8 +24,8 @@ class ULocalFogVolumeComponent : public USceneComponent
 	UPROPERTY(EditAnywhere, BlueprintReadOnly, interp, Category = "Height Fog Distribution", meta = (DisplayName = "Height Fog Density", UIMin = "0", UIMax = "2.0", SliderExponent = 2.0, ClampMin = 0.0))
 	float HeightFogExtinction = 0.0f;
 
-	/** Controls how the density decreases as height increases. Smaller values make the visible transition larger. */
-	UPROPERTY(EditAnywhere, BlueprintReadOnly, interp, Category = "Height Fog Distribution", meta = (UIMin = "0.001", UIMax = "5000", SliderExponent = 2.0, ClampMin = 0.001))
+	/** Controls how the density decreases as height increases. Smaller values make the visible transition larger. 1.0 is the lowest value before visual artifact are visible at the horizon. */
+	UPROPERTY(EditAnywhere, BlueprintReadOnly, interp, Category = "Height Fog Distribution", meta = (UIMin = "1.0", UIMax = "5000", SliderExponent = 2.0, ClampMin = 1.0))
 	float HeightFogFalloff = 1000.0f;
 
 	/** Height offset, relative to the actor Z position. */
diff --git a/Engine/Source/Runtime/Renderer/Private/LocalFogVolumeRendering.cpp b/Engine/Source/Runtime/Renderer/Private/LocalFogVolumeRendering.cpp
--- a/Engine/Source/Runtime/Renderer/Private/LocalFogVolumeRendering.cpp
+++ b/Engine/Source/Runtime/Renderer/Private/LocalFogVolumeRendering.cpp
@@ -346,8 +346,13 @@ void GetLocalFogVolumeSortingData(const FScene* Scene, FRDGBuilder& GraphBuilder
 		FMatrix44f InvTransform = Transform.Inverse();
 		ConvertFromMatrix44fTo4x3Array(InvTransform, LocalFogVolumeGPUInstanceDataIt->InvTransform);
 
+		// Falloff needs to be made safe in order to avoid artifact when the camera is looking toward the horizon at the level of the offset.
+		const float SafeFalloffThreshold = 1.0f;
+		const float FalloffScaleUI = 0.01f;
+		const float SafeFallOff = FMath::Max(LHF->HeightFogFalloff, SafeFalloffThreshold) * FalloffScaleUI;
+
 		FVector2DHalf Data0X = FVector2DHalf(LHF->RadialFogExtinction,			LHF->HeightFogExtinction);
-		FVector2DHalf Data0Y = FVector2DHalf(LHF->HeightFogFalloff * 0.01f,		LHF->HeightFogOffset);
+		FVector2DHalf Data0Y = FVector2DHalf(SafeFallOff,						LHF->HeightFogOffset);
 		FVector2DHalf Data0Z = FVector2DHalf(LHF->FogEmissive.R,				LHF->FogEmissive.G);
 		FVector2DHalf Data0W = FVector2DHalf(LHF->FogEmissive.B,				0.0f);
 
