LFV - Changed the component default values in order to see something by default after 28538588.

diff --git a/Engine/Source/Runtime/Engine/Classes/Components/LocalFogVolumeComponent.h b/Engine/Source/Runtime/Engine/Classes/Components/LocalFogVolumeComponent.h
--- a/Engine/Source/Runtime/Engine/Classes/Components/LocalFogVolumeComponent.h
+++ b/Engine/Source/Runtime/Engine/Classes/Components/LocalFogVolumeComponent.h
@@ -16,13 +16,13 @@ class ULocalFogVolumeComponent : public USceneComponent
 
 	~ULocalFogVolumeComponent();
 
-	/** The density of the radial fog representing its extinction coefficient at the center of the sphere. */
+	/** The density of the radial fog representing its extinction coefficient at the center of the sphere. The final look of the volume is determined by combining the "Coverage=1-Transmittance" of both radial and height fog in order to achieve both soft edges and height fog.*/
 	UPROPERTY(EditAnywhere, BlueprintReadOnly, interp, Category = "Radial Fog Distribution", meta = (DisplayName = "Radial Fog Density", UIMin = "0", UIMax = "2.0", SliderExponent = 2.0, ClampMin = 0.0))
 	float RadialFogExtinction = 1.0f;
 
-	/** The density of the radial fog representing its extinction coefficient at height 0 in the unit sphere. */
+	/** The density of the radial fog representing its extinction coefficient at height 0 in the unit sphere. The final look of the volume is determined by combining the "Coverage=1-Transmittance" of both radial and height fog in order to achieve both soft edges and height fog.*/
 	UPROPERTY(EditAnywhere, BlueprintReadOnly, interp, Category = "Height Fog Distribution", meta = (DisplayName = "Height Fog Density", UIMin = "0", UIMax = "2.0", SliderExponent = 2.0, ClampMin = 0.0))
-	float HeightFogExtinction = 0.0f;
+	float HeightFogExtinction = 1.0f;
 
 	/** Controls how the density decreases as height increases. Smaller values make the visible transition larger. 1.0 is the lowest value before visual artifact are visible at the horizon. */
 	UPROPERTY(EditAnywhere, BlueprintReadOnly, interp, Category = "Height Fog Distribution", meta = (UIMin = "1.0", UIMax = "5000", SliderExponent = 2.0, ClampMin = 1.0))
