World Partition: runtime cells object mapping now contains actor's base and native classes for debugging purposes.

diff --git a/Engine/Source/Runtime/Engine/Private/WorldPartition/HLOD/HLODSourceActorsFromCell.cpp b/Engine/Source/Runtime/Engine/Private/WorldPartition/HLOD/HLODSourceActorsFromCell.cpp
--- a/Engine/Source/Runtime/Engine/Private/WorldPartition/HLOD/HLODSourceActorsFromCell.cpp
+++ b/Engine/Source/Runtime/Engine/Private/WorldPartition/HLOD/HLODSourceActorsFromCell.cpp
@@ -41,14 +41,17 @@ ULevelStreaming* UWorldPartitionHLODSourceActorsFromCell::LoadSourceActors(bool&
 	UWorldPartition* WorldPartition = World->GetWorldPartition();
 	check(WorldPartition);
 
+	const FTopLevelAssetPath WorldPartitionHLODNativeClass(AWorldPartitionHLOD::StaticClass()->GetPathName());
 	const FName LevelStreamingName = FName(*FString::Printf(TEXT("HLODLevelStreaming_%s"), *HLODActor->GetName()));
 	TArray<FWorldPartitionRuntimeCellObjectMapping> Mappings;
 	Mappings.Reserve(Actors.Num());
-	Algo::Transform(Actors, Mappings, [World, HLODActor](const FHLODSubActor& SubActor)
-	{
+	Algo::Transform(Actors, Mappings, [World, HLODActor, &WorldPartitionHLODNativeClass](const FHLODSubActor& SubActor)
+	{		
 		return FWorldPartitionRuntimeCellObjectMapping(
 			SubActor.ActorPackage,
 			SubActor.ActorPath,
+			FTopLevelAssetPath(),
+			WorldPartitionHLODNativeClass,
 			SubActor.ContainerID,
 			SubActor.ContainerTransform,
 			SubActor.ContainerPackage,
diff --git a/Engine/Source/Runtime/Engine/Private/WorldPartition/WorldPartitionRuntimeLevelStreamingCell.cpp b/Engine/Source/Runtime/Engine/Private/WorldPartition/WorldPartitionRuntimeLevelStreamingCell.cpp
--- a/Engine/Source/Runtime/Engine/Private/WorldPartition/WorldPartitionRuntimeLevelStreamingCell.cpp
+++ b/Engine/Source/Runtime/Engine/Private/WorldPartition/WorldPartitionRuntimeLevelStreamingCell.cpp
@@ -240,6 +240,8 @@ void UWorldPartitionRuntimeLevelStreamingCell::AddActorToCell(const FWorldPartit
 		Packages.Emplace(
 			ReferenceActorDesc.GetActorPackage(), 
 			*ReferenceActorDesc.GetActorSoftPath().ToString(), 
+			ActorDescView.GetBaseClass(),
+			ActorDescView.GetNativeClass(),
 			InContainerID, 
 			InContainerTransform, 
 			InContainer->GetContainerPackage(), 
@@ -252,6 +254,8 @@ void UWorldPartitionRuntimeLevelStreamingCell::AddActorToCell(const FWorldPartit
 	Packages.Emplace(
 		ActorDescView.GetActorPackage(), 
 		*ActorDescView.GetActorSoftPath().ToString(), 
+		ActorDescView.GetBaseClass(),
+		ActorDescView.GetNativeClass(),
 		InContainerID, 
 		InContainerTransform, 
 		InContainer->GetContainerPackage(), 
diff --git a/Engine/Source/Runtime/Engine/Public/WorldPartition/WorldPartitionRuntimeCell.h b/Engine/Source/Runtime/Engine/Public/WorldPartition/WorldPartitionRuntimeCell.h
--- a/Engine/Source/Runtime/Engine/Public/WorldPartition/WorldPartitionRuntimeCell.h
+++ b/Engine/Source/Runtime/Engine/Public/WorldPartition/WorldPartitionRuntimeCell.h
@@ -37,19 +37,17 @@ struct FWorldPartitionRuntimeCellObjectMapping
 
 	FWorldPartitionRuntimeCellObjectMapping()
 #if WITH_EDITORONLY_DATA
-		: Package(NAME_None)
-		, Path(NAME_None)
-		, ContainerTransform(FTransform::Identity)
-		, ContainerPackage(NAME_None)
-		, LoadedPath(NAME_None)
+		: ContainerTransform(FTransform::Identity)
 		, bIsEditorOnly(false)
 #endif
 	{}
 
-	FWorldPartitionRuntimeCellObjectMapping(FName InPackage, FName InPath, const FActorContainerID& InContainerID, const FTransform& InContainerTransform, FName InContainerPackage, FName InWorldPackage, const FGuid& InActorInstanceGuid, bool bInIsEditorOnly)
+	FWorldPartitionRuntimeCellObjectMapping(FName InPackage, FName InPath, const FTopLevelAssetPath& InBaseClass, const FTopLevelAssetPath& InNativeClass, const FActorContainerID& InContainerID, const FTransform& InContainerTransform, FName InContainerPackage, FName InWorldPackage, const FGuid& InActorInstanceGuid, bool bInIsEditorOnly)
 #if WITH_EDITORONLY_DATA
 		: Package(InPackage)
 		, Path(InPath)
+		, BaseClass(InBaseClass)
+		, NativeClass(InNativeClass)
 		, ContainerID(InContainerID)
 		, ContainerTransform(InContainerTransform)
 		, ContainerPackage(InContainerPackage)
@@ -73,6 +71,18 @@ struct FWorldPartitionRuntimeCellObjectMapping
 	UPROPERTY()
 	FName Path;
 
+	/** 
+	 * The actor's base class
+	 */
+	UPROPERTY()
+	FTopLevelAssetPath BaseClass;
+
+	/** 
+	 * The actor's native base class
+	 */
+	UPROPERTY()
+	FTopLevelAssetPath NativeClass;
+
 	/**
 	 * ID of the owning container instance
 	 */
