[NetPhysics] Physics replication improvements, align the angular replication timeline with the linear replication timeline.

diff --git a/Engine/Source/Runtime/Engine/Private/PhysicsEngine/PhysicsReplication.cpp b/Engine/Source/Runtime/Engine/Private/PhysicsEngine/PhysicsReplication.cpp
--- a/Engine/Source/Runtime/Engine/Private/PhysicsEngine/PhysicsReplication.cpp
+++ b/Engine/Source/Runtime/Engine/Private/PhysicsEngine/PhysicsReplication.cpp
@@ -104,19 +104,25 @@ namespace PhysicsReplicationCVars
 
 	namespace PredictiveInterpolationCVars
 	{
-		static float PosCorrectionTimeBase = 0.15f;
+		static float PosCorrectionTimeBase = 0.0f;
 		static FAutoConsoleVariableRef CVarPosCorrectionTimeBase(TEXT("np2.PredictiveInterpolation.PosCorrectionTimeBase"), PosCorrectionTimeBase, TEXT("Base time to correct positional offset over. RTT * PosCorrectionTimeMultiplier are added on top of this."));
 
+		static float PosCorrectionTimeMin = 0.13f;
+		static FAutoConsoleVariableRef CVarPosCorrectionTimeMin(TEXT("np2.PredictiveInterpolation.PosCorrectionTimeMin"), PosCorrectionTimeMin, TEXT("Min time time to correct positional offset over. DeltaSeconds is added on top of this."));
+
 		static float PosCorrectionTimeMultiplier = 1.0f;
 		static FAutoConsoleVariableRef CVarPosCorrectionTimeMultiplier(TEXT("np2.PredictiveInterpolation.PosCorrectionTimeMultiplier"), PosCorrectionTimeMultiplier, TEXT("Multiplier to adjust how much of RTT (network Round Trip Time) to add to positional offset correction."));
 
-		static float RotCorrectionTimeBase = 0.15f;
+		static float RotCorrectionTimeBase = 0.0f;
 		static FAutoConsoleVariableRef CVarRotCorrectionTimeBase(TEXT("np2.PredictiveInterpolation.RotCorrectionTimeBase"), RotCorrectionTimeBase, TEXT("Base time to correct positional offset over. RTT * PosCorrectionTimeMultiplier are added on top of this."));
 
+		static float RotCorrectionTimeMin = 0.13f;
+		static FAutoConsoleVariableRef CVarRotCorrectionTimeMin(TEXT("np2.PredictiveInterpolation.RotCorrectionTimeMin"), RotCorrectionTimeMin, TEXT("Min time time to correct rotational offset over. DeltaSeconds is added on top of this."));
+
 		static float RotCorrectionTimeMultiplier = 1.0f;
 		static FAutoConsoleVariableRef CVarRotCorrectionTimeMultiplier(TEXT("np2.PredictiveInterpolation.RotCorrectionTimeMultiplier"), RotCorrectionTimeMultiplier, TEXT("Multiplier to adjust how much of RTT (network Round Trip Time) to add to positional offset correction."));
 
-		static float InterpolationTimeMultiplier = 1.1f;
+		static float InterpolationTimeMultiplier = 1.25f;
 		static FAutoConsoleVariableRef CVarInterpolationTimeMultiplier(TEXT("np2.PredictiveInterpolation.InterpolationTimeMultiplier"), InterpolationTimeMultiplier, TEXT("Multiplier to adjust the replication interpolation time which is based on the sendrate of replication data from the server."));
 		
 		static float AverageReceiveIntervalSmoothing = 3.0f;
@@ -1321,14 +1327,8 @@ bool FPhysicsReplicationAsync::PredictiveInterpolation(Chaos::FPBDRigidParticleH
 		}
 	}
 
-	// Calculate position correction time based on current Round Trip Time
-	const float RTT = LatencyOneWay * 2.f;
-	const float PosCorrectionTime = FMath::Max(PhysicsReplicationCVars::PredictiveInterpolationCVars::PosCorrectionTimeBase + RTT * PhysicsReplicationCVars::PredictiveInterpolationCVars::PosCorrectionTimeMultiplier, DeltaSeconds);
-	const float RotCorrectionTime = FMath::Max(PhysicsReplicationCVars::PredictiveInterpolationCVars::RotCorrectionTimeBase + RTT * PhysicsReplicationCVars::PredictiveInterpolationCVars::RotCorrectionTimeMultiplier, DeltaSeconds);
-
-	// Calculate interpolation time based on current average receive rate of targets from the server (receive rate = send-rate from server with network conditions taken into account)
+	// Update the AverageReceiveInterval of targets from the server (receive rate = send-rate from server with network conditions taken into account)
 	Target.AverageReceiveInterval = FMath::Lerp(Target.AverageReceiveInterval, Target.ReceiveInterval, FMath::Clamp((1.0f / (Target.ReceiveInterval * PhysicsReplicationCVars::PredictiveInterpolationCVars::AverageReceiveIntervalSmoothing)), 0.0f, 1.0f));
-	const float InterpolationTime = Target.AverageReceiveInterval * DeltaSeconds * PhysicsReplicationCVars::PredictiveInterpolationCVars::InterpolationTimeMultiplier;
 
 	// CurrentState
 	FRigidBodyState CurrentState;
@@ -1403,6 +1403,16 @@ bool FPhysicsReplicationAsync::PredictiveInterpolation(Chaos::FPBDRigidParticleH
 	}
 	else // Velocity-based Replication
 	{
+		// Calculate position correction time based on current Round Trip Time
+		const float RTT = LatencyOneWay * 2.f;
+		const float PosCorrectionTime = FMath::Max(PhysicsReplicationCVars::PredictiveInterpolationCVars::PosCorrectionTimeBase + RTT * PhysicsReplicationCVars::PredictiveInterpolationCVars::PosCorrectionTimeMultiplier, DeltaSeconds + PhysicsReplicationCVars::PredictiveInterpolationCVars::PosCorrectionTimeMin);
+		const float RotCorrectionTime = FMath::Max(PhysicsReplicationCVars::PredictiveInterpolationCVars::RotCorrectionTimeBase + RTT * PhysicsReplicationCVars::PredictiveInterpolationCVars::RotCorrectionTimeMultiplier, DeltaSeconds + PhysicsReplicationCVars::PredictiveInterpolationCVars::RotCorrectionTimeMin);
+
+		/* Temp disabled until angular replication has implemented InterpolationTime
+		// Calculate interpolation time based on current average receive rate
+		const float InterpolationTime = Target.AverageReceiveInterval * DeltaSeconds * PhysicsReplicationCVars::PredictiveInterpolationCVars::InterpolationTimeMultiplier;
+		*/
+
 		if (!bXCanEarlyOut)
 		{	// --- Velocity Replication ---
 			
@@ -1418,8 +1428,8 @@ bool FPhysicsReplicationAsync::PredictiveInterpolation(Chaos::FPBDRigidParticleH
 			// Add PosDiffVelocity to LinVelDiff to get BlendedTargetVelocity
 			const FVector BlendedTargetVelocity = LinVelDiff + PosDiffVelocity;
 
-			// Multiply BlendedTargetVelocity with(deltaTime / interpolationTime), clamp to 1 and add to CurrentState.LinVel to get BlendedTargetVelocityInterpolated
-			const float BlendStepAmount = FMath::Clamp(DeltaSeconds / InterpolationTime, 0.f, 1.f);
+			// Add BlendedTargetVelocity onto current velocity
+			const float BlendStepAmount = 1.0f; /* FMath::Clamp(DeltaSeconds / InterpolationTime, 0.0f, 1.0f); Temp disabled until angular replication has implemented InterpolationTime */
 			const FVector RepLinVel = CurrentState.LinVel + (BlendedTargetVelocity * BlendStepAmount);
 			
 			Handle->SetV(RepLinVel);
@@ -1429,7 +1439,8 @@ bool FPhysicsReplicationAsync::PredictiveInterpolation(Chaos::FPBDRigidParticleH
 		}
 
 		{	// --- Angular Velocity Replication ---
-			
+			/* Todo, Implement InterpolationTime */
+
 			// Extrapolate current rotation along current angular velocity to see where we would end up
 			float CurAngVelSize;
 			FVector CurAngVelAxis;
@@ -1439,7 +1450,8 @@ bool FPhysicsReplicationAsync::PredictiveInterpolation(Chaos::FPBDRigidParticleH
 
 			// Slerp from the extrapolated current rotation towards the target rotation
 			// This takes current angular velocity into account
-			const FQuat TargetRotBlended = FQuat::Slerp(CurRotExtrap, TargetRot, 1.0f  /*BlendStepAmount*/); // BlendStepAmount is temporarily disabled for rotation
+			const float RotCorrectionAmount = FMath::Clamp(DeltaSeconds / RotCorrectionTime, 0.0f, 1.0f);
+			const FQuat TargetRotBlended = FQuat::Slerp(CurRotExtrap, TargetRot, RotCorrectionAmount);
 
 			// Get the rotational offset between the blended rotation target and the current rotation
 			const FQuat TargetRotDelta = TargetRotBlended * CurrentState.Quaternion.Inverse();
@@ -1448,7 +1460,7 @@ bool FPhysicsReplicationAsync::PredictiveInterpolation(Chaos::FPBDRigidParticleH
 			float WAngle;
 			FVector WAxis;
 			TargetRotDelta.ToAxisAndAngle(WAxis, WAngle);
-			const FVector TargetRotDeltaBlend = FVector(WAxis * (WAngle / RotCorrectionTime));
+			const FVector TargetRotDeltaBlend = FVector(WAxis * (WAngle / (DeltaSeconds * PhysicsReplicationCVars::PredictiveInterpolationCVars::InterpolationTimeMultiplier)));
 			const FVector RepAngVel = FMath::DegreesToRadians(TargetAngVel) + TargetRotDeltaBlend;
 
 			Handle->SetW(RepAngVel);
