Exposing rect light sampling on translucent to the render setting panel.

diff --git a/Engine/Source/Runtime/Engine/Classes/Engine/RendererSettings.h b/Engine/Source/Runtime/Engine/Classes/Engine/RendererSettings.h
--- a/Engine/Source/Runtime/Engine/Classes/Engine/RendererSettings.h
+++ b/Engine/Source/Runtime/Engine/Classes/Engine/RendererSettings.h
@@ -874,6 +874,15 @@ class URendererSettings : public UDeveloperSettings
 		ConfigRestartRequired = true))
 		uint32 bSupportCloudShadowOnForwardLitTranslucent : 1;
 
+	/**
+	"Enable rect light evaluation on translucent materials when using the Forward Shading mode."
+	*/
+	UPROPERTY(config, EditAnywhere, Category = Optimizations, meta = (
+		ConsoleVariable = "r.RectLightAtlas.Translucent", DisplayName = "Support Rect Light On Translucent Materials (When Using ForwardShading)",
+		ToolTip = "Enable rect light evaluation on translucent materials when using the Forward Shading mode.",
+		ConfigRestartRequired = true))
+		uint32 bSupportRectLightOnTranslucent : 1;
+
 	UPROPERTY(config, EditAnywhere, Category = Debugging, meta = (
 		ConsoleVariable = "r.GPUCrashDebugging", DisplayName = "Enable vendor specific GPU crash analysis tools",
 		ToolTip = "Enables vendor specific GPU crash analysis tools.",
