Add OverrideMinLOD and deprecate SetMinLOD in SkinnedMeshComponent.
- SetMinLOD requires bOverrideMinLod to be true for it to function, so adding a function that sets MinLOD and bOverrideMinLod to true.

#tests local PIE

diff --git a/Engine/Source/Runtime/Engine/Classes/Components/SkinnedMeshComponent.h b/Engine/Source/Runtime/Engine/Classes/Components/SkinnedMeshComponent.h
--- a/Engine/Source/Runtime/Engine/Classes/Components/SkinnedMeshComponent.h
+++ b/Engine/Source/Runtime/Engine/Classes/Components/SkinnedMeshComponent.h
@@ -821,9 +821,18 @@ public:
 	 *
 	 * @param	InNewMinLOD	Set new MinLodModel that make sure the LOD does not go below of this value. Range from [0, Max Number of LOD - 1]. This will affect in the next tick update. 
 	 */
-	UFUNCTION(BlueprintCallable, Category="Components|SkinnedMesh")
+	UE_DEPRECATED(5.5, "Use USkinnedMeshComponent::OverrideMinLOD() instead.")
+	UFUNCTION(BlueprintCallable, Category="Components|SkinnedMesh", meta = (DeprecatedFunction, DeprecationMessage = "Use USkinnedMeshComponent::OverrideMinLOD() instead."))
 	ENGINE_API void SetMinLOD(int32 InNewMinLOD);
 
+	/**
+	 * Override the Min LOD of the mesh component
+	 *
+	 * @param	InNewMinLOD	Override new MinLodModel that make sure the LOD does not go below of this value. Range from [0, Max Number of LOD - 1]. This will affect in the next tick update.
+	 */
+	UFUNCTION(BlueprintCallable, Category = "Components|SkinnedMesh")
+	ENGINE_API void OverrideMinLOD(int32 InNewMinLOD);
+
 	/**
 	 * Set ForcedLodModel of the mesh component
 	 *
diff --git a/Engine/Source/Runtime/Engine/Private/Components/SkinnedMeshComponent.cpp b/Engine/Source/Runtime/Engine/Private/Components/SkinnedMeshComponent.cpp
--- a/Engine/Source/Runtime/Engine/Private/Components/SkinnedMeshComponent.cpp
+++ b/Engine/Source/Runtime/Engine/Private/Components/SkinnedMeshComponent.cpp
@@ -3678,6 +3678,13 @@ void USkinnedMeshComponent::SetMinLOD(int32 InNewMinLOD)
 	MinLodModel = FMath::Clamp(InNewMinLOD, 0, MaxLODIndex);
 }
 
+void USkinnedMeshComponent::OverrideMinLOD(int32 InNewMinLOD)
+{
+	int32 MaxLODIndex = GetNumLODs() - 1;
+	MinLodModel = FMath::Clamp(InNewMinLOD, 0, MaxLODIndex);
+	bOverrideMinLod = true;
+}
+
 int32 USkinnedMeshComponent::ComputeMinLOD() const
 {
 	int32 MinLodIndex = bOverrideMinLod ? MinLodModel : GetSkinnedAsset()->GetMinLodIdx();
