UWorld::ComponentSweepMulti needs to return whether this was a *blocking* hit not whether there was any hit

diff --git a/Engine/Source/Runtime/Engine/Private/Collision/WorldCollision.cpp b/Engine/Source/Runtime/Engine/Private/Collision/WorldCollision.cpp
--- a/Engine/Source/Runtime/Engine/Private/Collision/WorldCollision.cpp
+++ b/Engine/Source/Runtime/Engine/Private/Collision/WorldCollision.cpp
@@ -409,8 +409,7 @@ bool UWorld::ComponentSweepMulti(TArray<struct FHitResult>& OutHits, class UPrim
 {
 	if (PrimComp)
 	{
-		ComponentSweepMultiByChannel(OutHits, PrimComp, Start, End, Quat, PrimComp->GetCollisionObjectType(), Params);
-		return (OutHits.Num() > 0);
+		return ComponentSweepMultiByChannel(OutHits, PrimComp, Start, End, Quat, PrimComp->GetCollisionObjectType(), Params);
 	}
 	else
 	{
