add UClusterUnionComponent::FindChildPhysicsObjectGivenShapeIndex which allows users to query a cluster union to get the underlying physics object given the shape index

- also introduce FSafePhysicsObjectHandle which is a wrapper around the Component and Bone ID which provides a safe way to get a physics object handle.

diff --git a/Engine/Source/Runtime/Engine/Classes/PhysicsEngine/ClusterUnionComponent.h b/Engine/Source/Runtime/Engine/Classes/PhysicsEngine/ClusterUnionComponent.h
--- a/Engine/Source/Runtime/Engine/Classes/PhysicsEngine/ClusterUnionComponent.h
+++ b/Engine/Source/Runtime/Engine/Classes/PhysicsEngine/ClusterUnionComponent.h
@@ -10,6 +10,7 @@
 #include "Engine/EngineTypes.h"
 #include "Logging/LogMacros.h"
 #include "PhysicsEngine/ExternalSpatialAccelerationPayload.h"
+#include "PhysicsEngine/SafePhysicsObjectHandle.h"
 #include "PhysicsInterfaceTypesCore.h"
 #include "PhysicsProxy/ClusterUnionPhysicsProxy.h"
 #include "UObject/ObjectKey.h"
@@ -192,6 +193,11 @@ public:
 	ENGINE_API TArray<int32> GetAddedBoneIdsForComponent(UPrimitiveComponent* Component) const;
 	ENGINE_API void ChangeIfComponentBonesAreMainParticle(UPrimitiveComponent* Component, const TArray<int32>& BoneIds, bool bIsMain);
 
+	// This function will return information from the latest physics sync. Thus adding a new bones
+	// to a cluster union will not allow those indices to be valid for use in this function
+	// until the next time we sync from the PT.
+	ENGINE_API Chaos::FPhysicsObjectHandle FindChildPhysicsObjectByShapeIndex(int32 Index) const;
+
 	friend class UClusterUnionReplicatedProxyComponent;
 	friend class UModularVehicleBaseComponent;
 protected:
@@ -238,6 +244,9 @@ private:
 	// where the component hasn't been added to the cluster union on the GT causing a mismatch in behavior.
 	TMap<TObjectKey<UPrimitiveComponent>, FClusterUnionPendingAddData> PendingComponentSync;
 
+	// At every physics sync, we keep track of which shape index matches which child primitive component and bone.
+	TArray<FSafePhysicsObjectHandle> PerShapeComponentBone;
+
 	// Data that can be changed at runtime to keep state about the cluster union consistent between the server and client.
 	UPROPERTY(ReplicatedUsing=OnRep_RigidState)
 	FClusterUnionReplicatedData ReplicatedRigidState;
diff --git a/Engine/Source/Runtime/Engine/Classes/PhysicsEngine/PhysicsObjectExternalInterface.h b/Engine/Source/Runtime/Engine/Classes/PhysicsEngine/PhysicsObjectExternalInterface.h
--- a/Engine/Source/Runtime/Engine/Classes/PhysicsEngine/PhysicsObjectExternalInterface.h
+++ b/Engine/Source/Runtime/Engine/Classes/PhysicsEngine/PhysicsObjectExternalInterface.h
@@ -55,4 +55,6 @@ public:
 
 	static ENGINE_API Chaos::FReadPhysicsObjectInterface_External GetRead_AssumesLocked();
 	static ENGINE_API Chaos::FWritePhysicsObjectInterface_External GetWrite_AssumesLocked();
+
+	static ENGINE_API UPrimitiveComponent* GetComponentFromPhysicsObject(UWorld* World, Chaos::FPhysicsObjectHandle PhysicsObject);
 };
diff --git a/Engine/Source/Runtime/Engine/Classes/PhysicsEngine/SafePhysicsObjectHandle.h b/Engine/Source/Runtime/Engine/Classes/PhysicsEngine/SafePhysicsObjectHandle.h
new file mode 100644
--- /dev/null
+++ b/Engine/Source/Runtime/Engine/Classes/PhysicsEngine/SafePhysicsObjectHandle.h
@@ -0,0 +1,23 @@
+// Copyright Epic Games, Inc. All Rights Reserved.
+
+#pragma once
+
+#include "Chaos/PhysicsObject.h"
+#include "Interfaces/IPhysicsComponent.h"
+#include "UObject/WeakInterfacePtr.h"
+
+struct FSafePhysicsObjectHandle
+{
+	TWeakInterfacePtr<IPhysicsComponent> Component;
+	Chaos::FPhysicsObjectId BoneId;
+
+	Chaos::FPhysicsObjectHandle Get() const
+	{
+		if (IPhysicsComponent* ComponentPtr = Component.Get())
+		{
+			return ComponentPtr->GetPhysicsObjectById(BoneId);
+		}
+
+		return nullptr;
+	}
+};
diff --git a/Engine/Source/Runtime/Engine/Private/PhysicsEngine/ClusterUnionComponent.cpp b/Engine/Source/Runtime/Engine/Private/PhysicsEngine/ClusterUnionComponent.cpp
--- a/Engine/Source/Runtime/Engine/Private/PhysicsEngine/ClusterUnionComponent.cpp
+++ b/Engine/Source/Runtime/Engine/Private/PhysicsEngine/ClusterUnionComponent.cpp
@@ -837,6 +837,7 @@ void UClusterUnionComponent::SyncClusterUnionFromProxy()
 
 	TArray<Chaos::FPBDRigidParticle*> ChildParticles;
 	ChildParticles.Reserve(FullData.ChildParticles.Num());
+	PerShapeComponentBone.Reset(FullData.ChildParticles.Num());
 
 	// Note that at the UClusterUnionComponent level we really only want to be dealing with components.
 	// Hence why we need to modify each of the particles that we synced from the game thread into a
@@ -854,10 +855,12 @@ void UClusterUnionComponent::SyncClusterUnionFromProxy()
 	
 				Chaos::FPhysicsObjectHandle Handle = Component->GetPhysicsObjectById(ChildData.BoneId);
 				ChildParticles.Add(Interface->GetRigidParticle(Handle));
+				PerShapeComponentBone.Add({Component, ChildData.BoneId});
 			}
 			else
 			{
 				ChildParticles.Add(nullptr);
+				PerShapeComponentBone.Add({});
 			}
 		}
 	}
@@ -1596,3 +1599,12 @@ void UClusterUnionComponent::ChangeIfComponentBonesAreMainParticle(UPrimitiveCom
 
 	PhysicsProxy->ChangeMainParticleStatus_External(PhysicsObjects, bIsMain);
 }
+
+Chaos::FPhysicsObjectHandle UClusterUnionComponent::FindChildPhysicsObjectByShapeIndex(int32 Index) const
+{
+	if (!PerShapeComponentBone.IsValidIndex(Index))
+	{
+		return nullptr;
+	}
+	return PerShapeComponentBone[Index].Get();
+}
\ No newline at end of file
diff --git a/Engine/Source/Runtime/Engine/Private/PhysicsEngine/PhysicsObjectExternalInterface.cpp b/Engine/Source/Runtime/Engine/Private/PhysicsEngine/PhysicsObjectExternalInterface.cpp
--- a/Engine/Source/Runtime/Engine/Private/PhysicsEngine/PhysicsObjectExternalInterface.cpp
+++ b/Engine/Source/Runtime/Engine/Private/PhysicsEngine/PhysicsObjectExternalInterface.cpp
@@ -1,6 +1,9 @@
 // Copyright Epic Games, Inc. All Rights Reserved.
 #include "PhysicsEngine/PhysicsObjectExternalInterface.h"
 
+#include "Engine/World.h"
+#include "Physics/Experimental/PhysScene_Chaos.h"
+
 FLockedReadPhysicsObjectExternalInterface FPhysicsObjectExternalInterface::LockRead(FChaosScene* Scene)
 {
 	return FLockedReadPhysicsObjectExternalInterface{ Scene, CreateReadInterface<Chaos::EThreadContext::External>() };
@@ -34,4 +37,21 @@ Chaos::FReadPhysicsObjectInterface_External FPhysicsObjectExternalInterface::Get
 Chaos::FWritePhysicsObjectInterface_External FPhysicsObjectExternalInterface::GetWrite_AssumesLocked()
 {
 	return CreateWriteInterface<Chaos::EThreadContext::External>();
+}
+
+UPrimitiveComponent* FPhysicsObjectExternalInterface::GetComponentFromPhysicsObject(UWorld* World, Chaos::FPhysicsObjectHandle PhysicsObject)
+{
+	if (!PhysicsObject || !World)
+	{
+		return nullptr;
+	}
+
+	IPhysicsProxyBase* Proxy = Chaos::FPhysicsObjectInterface::GetProxy({ &PhysicsObject, 1 });
+	FPhysScene* Scene = World->GetPhysicsScene();
+	if (!Scene || !Proxy)
+	{
+		return nullptr;
+	}
+
+	return Scene->GetOwningComponent<UPrimitiveComponent>(Proxy);
 }
\ No newline at end of file
