FPhysScene_Chaos::HandleGlobalCollisionEvent should bubble up the collision ShapeIndex in cases where we collide against a cluster union

diff --git a/Engine/Source/Runtime/Engine/Private/PhysicsEngine/Experimental/PhysScene_Chaos.cpp b/Engine/Source/Runtime/Engine/Private/PhysicsEngine/Experimental/PhysScene_Chaos.cpp
--- a/Engine/Source/Runtime/Engine/Private/PhysicsEngine/Experimental/PhysScene_Chaos.cpp
+++ b/Engine/Source/Runtime/Engine/Private/PhysicsEngine/Experimental/PhysScene_Chaos.cpp
@@ -1117,7 +1117,7 @@ void FPhysScene_Chaos::HandleGlobalCollisionEvent(Chaos::FCollisionDataArray con
 			CollisionEvent.Body1.PhysMaterial = InternalMat1 ? FPhysicsUserData::Get<UPhysicalMaterial>(InternalMat1->UserData) : nullptr;
 			CollisionEvent.Body2.PhysMaterial = InternalMat2 ? FPhysicsUserData::Get<UPhysicalMaterial>(InternalMat2->UserData) : nullptr;
 			
-			if (BodyPrimitive1)
+			if (BodyPrimitive1 && CollisionItem.Proxy1)
 			{
 				const FBodyInstance* BodyInst1 = GetBodyInstanceFromProxyAndShape(CollisionItem.Proxy1, CollisionItem.ShapeIndex1);
 				if (BodyInst1 != nullptr)
@@ -1132,9 +1132,13 @@ void FPhysScene_Chaos::HandleGlobalCollisionEvent(Chaos::FCollisionDataArray con
 						CollisionEvent.Body1.BoneName = NAME_None;
 					}
 				}
+				else if (CollisionItem.Proxy1->GetType() == EPhysicsProxyType::ClusterUnionProxy)
+				{
+					CollisionEvent.Body1.BodyIndex = CollisionItem.ShapeIndex1;
+				}
 			}
 
-			if (BodyPrimitive2)
+			if (BodyPrimitive2 && CollisionItem.Proxy2)
 			{
 				const FBodyInstance* BodyInst2 = GetBodyInstanceFromProxyAndShape(CollisionItem.Proxy2, CollisionItem.ShapeIndex2);
 				if (BodyInst2 != nullptr)
@@ -1149,6 +1153,10 @@ void FPhysScene_Chaos::HandleGlobalCollisionEvent(Chaos::FCollisionDataArray con
 						CollisionEvent.Body2.BoneName = NAME_None;
 					}
 				}
+				else if (CollisionItem.Proxy2->GetType() == EPhysicsProxyType::ClusterUnionProxy)
+				{
+					CollisionEvent.Body2.BodyIndex = CollisionItem.ShapeIndex2;
+				}
 			}
 		}
 	}
