Prevent MetaSound InitResources from running during cook from MetaSounds used in SoundCues

diff --git a/Engine/Plugins/Runtime/Metasound/Source/MetasoundEngine/Private/MetasoundSource.cpp b/Engine/Plugins/Runtime/Metasound/Source/MetasoundEngine/Private/MetasoundSource.cpp
--- a/Engine/Plugins/Runtime/Metasound/Source/MetasoundEngine/Private/MetasoundSource.cpp
+++ b/Engine/Plugins/Runtime/Metasound/Source/MetasoundEngine/Private/MetasoundSource.cpp
@@ -510,6 +510,7 @@ void UMetaSoundSource::InitResources()
 
 	METASOUND_LLM_SCOPE;
 	METASOUND_TRACE_CPUPROFILER_EVENT_SCOPE(UMetaSoundSource::InitResources); 
+	ensureMsgf(!IsRunningCookCommandlet(), TEXT("UMetaSoundSource::InitResources should not be called during cook."));
 
 	if (IsInGameThread())
 	{
diff --git a/Engine/Source/Runtime/Engine/Private/SoundNodeWavePlayer.cpp b/Engine/Source/Runtime/Engine/Private/SoundNodeWavePlayer.cpp
--- a/Engine/Source/Runtime/Engine/Private/SoundNodeWavePlayer.cpp
+++ b/Engine/Source/Runtime/Engine/Private/SoundNodeWavePlayer.cpp
@@ -61,7 +61,14 @@ void USoundNodeWavePlayer::LoadAsset(bool bAddToRoot)
 		if (SoundWave)
 		{
 			SoundWave->AddToCluster(this, true);
-			SoundWave->InitResources();
+			// Don't init resources when running cook, as this can trigger 
+			// registration of a MetaSound and its dependent graphs.
+			// Those will instead be registered when the MetaSound itself is cooked (FMetasoundAssetBase::CookMetaSound)
+			// in a way that does not deal with runtime data like this function does
+			if (!IsRunningCookCommandlet())
+			{
+				SoundWave->InitResources();
+			}
 		}
 	}
 	else
@@ -80,7 +87,14 @@ void USoundNodeWavePlayer::LoadAsset(bool bAddToRoot)
 				SoundWave->AddToRoot();
 			}
 			SoundWave->AddToCluster(this);
-			SoundWave->InitResources();
+			// Don't init resources when running cook, as this can trigger 
+			// registration of a MetaSound and its dependent graphs.
+			// Those will instead be registered when the MetaSound itself is cooked (FMetasoundAssetBase::CookMetaSound)
+			// in a way that does not deal with runtime data like this function does
+			if (!IsRunningCookCommandlet())
+			{
+				SoundWave->InitResources();
+			}
 		}
 	}
 }
