Remove some redundant instance update dispatches for the landscape grass use-case & fix a few cases of incorrect delta updates being sent.

diff --git a/Engine/Source/Runtime/Engine/Classes/Components/InstancedStaticMeshComponent.h b/Engine/Source/Runtime/Engine/Classes/Components/InstancedStaticMeshComponent.h
--- a/Engine/Source/Runtime/Engine/Classes/Components/InstancedStaticMeshComponent.h
+++ b/Engine/Source/Runtime/Engine/Classes/Components/InstancedStaticMeshComponent.h
@@ -481,6 +481,9 @@ public:
 	 */
 	virtual FVector GetTranslatedInstanceSpaceOrigin() const { return FVector::Zero(); }
 
+	/** Handle changes that must happen before the proxy is recreated. */
+	ENGINE_API void PreApplyComponentInstanceData(struct FInstancedStaticMeshComponentInstanceData* ComponentInstanceData);
+
 	/** Applies the cached component instance data to a newly blueprint constructed component. */
 	ENGINE_API virtual void ApplyComponentInstanceData(struct FInstancedStaticMeshComponentInstanceData* ComponentInstanceData);
 
@@ -560,6 +563,7 @@ private:
 	ENGINE_API bool BatchUpdateInstancesTransformsInternal(int32 StartInstanceIndex, TArrayView<const FTransform> NewInstancesTransforms, bool bWorldSpace, bool bMarkRenderStateDirty, bool bTeleport);
 
 protected:
+	bool bIsInstanceDataApplyCompleted = true;
 
 	FPrimitiveInstanceDataManager PrimitiveInstanceDataManager;
 
@@ -698,6 +702,8 @@ public:
 
 	virtual void ApplyToComponent(UActorComponent* Component, const ECacheApplyPhase CacheApplyPhase) override
 	{
+		// The Super::ApplyToComponent will cause the scene proxy to be recreated, so we must do what we can to make sure the state is ok before that.
+		CastChecked<UInstancedStaticMeshComponent>(Component)->PreApplyComponentInstanceData(this);
 		Super::ApplyToComponent(Component, CacheApplyPhase);
 		CastChecked<UInstancedStaticMeshComponent>(Component)->ApplyComponentInstanceData(this);
 	}
diff --git a/Engine/Source/Runtime/Engine/Private/HierarchicalInstancedStaticMesh.cpp b/Engine/Source/Runtime/Engine/Private/HierarchicalInstancedStaticMesh.cpp
--- a/Engine/Source/Runtime/Engine/Private/HierarchicalInstancedStaticMesh.cpp
+++ b/Engine/Source/Runtime/Engine/Private/HierarchicalInstancedStaticMesh.cpp
@@ -2573,6 +2573,7 @@ void UHierarchicalInstancedStaticMeshComponent::BuildComponentInstanceData(FInst
 	OutData.StaticMeshBounds = GetStaticMesh()->GetBounds();
 	OutData.NumProxyInstances = InstanceCountToRender;
 	OutData.NumSourceInstances = PerInstanceSMData.Num();
+	OutData.NumCustomDataFloats = NumCustomDataFloats;
 
 	OutData.BuildChangeSet = [&](FISMInstanceUpdateChangeSet &ChangeSet)
 	{
@@ -2702,25 +2703,7 @@ void UHierarchicalInstancedStaticMeshComponent::ApplyBuildTree(FClusterBuilder&
 	// Make sure it gets rebuilt from scratch to reflect the new instance ordering
 	// we could actually do it incrementally since no data needs to be uploaded (the instances are the same as before, just need to mark changed Indexes for all & implement the general swap functionality)
 	// BUT: the hitproxy data was rebuilt right here & the per instance random is handled differently inside the tree builder so, nope.
-	PrimitiveInstanceDataManager.MarkForRebuildFromExternal(NumBuiltRenderInstances, [
-		HitProxies, 
-		LegacyInstanceData = MoveTemp(BuiltInstanceData),
-		LegacyInstanceReorderTable = InstanceReorderTable]
-		(TArray<TRefCountPtr<HHitProxy>> &OutHitProxies) mutable
-	{
-		OutHitProxies = HitProxies;
-		FPrimitiveInstanceDataManager::FExternalUpdateData ExternalUpdateData;
-		ExternalUpdateData.NumCustomDataFloats = LegacyInstanceData ? LegacyInstanceData->GetNumCustomDataFloats() : 0;
-		ExternalUpdateData.NumInstances = LegacyInstanceData ? LegacyInstanceData->GetNumInstances() : 0;
-		ExternalUpdateData.UpdateProxy = [
-			LegacyInstanceDataInner = MoveTemp(LegacyInstanceData), 
-			LegacyInstanceReorderTableInner = MoveTemp(LegacyInstanceReorderTable)]
-			(FISMCInstanceDataSceneProxy &InstanceDataSceneProxy, const FRenderBounds &LocalBounds) mutable
-		{
-			InstanceDataSceneProxy.BuildFromLegacyData(MoveTemp(LegacyInstanceDataInner), LocalBounds, MoveTemp(LegacyInstanceReorderTableInner));
-		};
-		return ExternalUpdateData;
-	});
+	PrimitiveInstanceDataManager.MarkForRebuildFromLegacy(MoveTemp(BuiltInstanceData), InstanceReorderTable, HitProxies);
 
 	FlushAccumulatedNavigationUpdates();
 	PostBuildStats();
@@ -2972,14 +2955,16 @@ FPrimitiveSceneProxy* UHierarchicalInstancedStaticMeshComponent::CreateSceneProx
 		return nullptr;
 	}
 
-	// Verify that the mesh is valid before using it.
-	const bool bMeshIsValid =
+	// Verify that the mesh is valid & that we have instances before creating a proxy.
+	const bool bAreMeshAndInstancesValid =
+		// Make sure we have instances, or an update with instances on the way (for the external data (landscape grass) mode).
+		(GetNumInstances() > 0 ||  PrimitiveInstanceDataManager.GetMaxInstanceIndex() > 0) &&
 		// Make sure we have an actual static mesh.
 		GetStaticMesh() &&
 		!GetStaticMesh()->IsCompiling() &&
 		GetStaticMesh()->HasValidRenderData(false);
 
-	if (!bMeshIsValid)
+	if (!bAreMeshAndInstancesValid)
 	{
 		return nullptr;
 	}
diff --git a/Engine/Source/Runtime/Engine/Private/InstancedStaticMesh.cpp b/Engine/Source/Runtime/Engine/Private/InstancedStaticMesh.cpp
--- a/Engine/Source/Runtime/Engine/Private/InstancedStaticMesh.cpp
+++ b/Engine/Source/Runtime/Engine/Private/InstancedStaticMesh.cpp
@@ -2058,6 +2058,14 @@ void UInstancedStaticMeshComponent::GetComponentChildElements(TArray<FTypedEleme
 #endif	// WITH_EDITOR
 }
 
+void UInstancedStaticMeshComponent::PreApplyComponentInstanceData(struct FInstancedStaticMeshComponentInstanceData* InstancedMeshData)
+{
+#if WITH_EDITOR
+	// Prevent proxy recreate while traversing the ::ApplyToComponent stack
+	bIsInstanceDataApplyCompleted = false;
+#endif
+}
+
 void UInstancedStaticMeshComponent::ApplyComponentInstanceData(FInstancedStaticMeshComponentInstanceData* InstancedMeshData)
 {
 #if WITH_EDITOR
@@ -2105,7 +2113,9 @@ void UInstancedStaticMeshComponent::ApplyComponentInstanceData(FInstancedStaticM
 
 	bHasPerInstanceHitProxies = InstancedMeshData->bHasPerInstanceHitProxies;
 
-	// Force recreation of the render data & reset all tracking and mapping in of ID's
+	bIsInstanceDataApplyCompleted = true;
+
+	// TODO: restore ID mapping either from the serialized stuff, or the InstancedMeshData
 	PrimitiveInstanceDataManager.Invalidate(PerInstanceSMData.Num());
 #endif
 }
@@ -2187,10 +2197,11 @@ FPrimitiveSceneProxy* UInstancedStaticMeshComponent::CreateSceneProxy()
 {
 	ProxySize = 0;
 
-	PrimitiveInstanceDataManager.ResetComponentDirtyTracking();
-
-	// Verify that the mesh is valid before using it.
-	const bool bMeshIsValid =
+	// Verify that both mesh and instance data is valid before using it.
+	const bool bIsMeshAndInstanceDataValid =
+#if WITH_EDITOR
+		bIsInstanceDataApplyCompleted && 
+#endif
 		// make sure we have instances
 		PerInstanceSMData.Num() > 0 &&
 		// make sure we have an actual static mesh
@@ -2198,7 +2209,7 @@ FPrimitiveSceneProxy* UInstancedStaticMeshComponent::CreateSceneProxy()
 		GetStaticMesh()->IsCompiling() == false &&
 		GetStaticMesh()->HasValidRenderData();
 
-	if (!bMeshIsValid)
+	if (!bIsMeshAndInstanceDataValid)
 	{
 		return nullptr;
 	}
@@ -2392,6 +2403,7 @@ void UInstancedStaticMeshComponent::BuildComponentInstanceData(FInstanceUpdateCo
 	OutData.StaticMeshBounds = GetStaticMesh()->GetBounds();
 	OutData.NumProxyInstances = PerInstanceSMData.Num();
 	OutData.NumSourceInstances = PerInstanceSMData.Num();
+	OutData.NumCustomDataFloats = NumCustomDataFloats;
 
 	// Function that only gets called if we actually need to flush any changes.
 	OutData.BuildChangeSet = [&](FISMInstanceUpdateChangeSet &ChangeSet)
@@ -4571,14 +4583,15 @@ void UInstancedStaticMeshComponent::PostLoad()
 		}
 	}
 
-	// Has different implementation in HISMC
-	OnPostLoadPerInstanceData();
-
 	if (!HasAnyFlags(RF_ClassDefaultObject|RF_ArchetypeObject))
 	{
 		PrimitiveInstanceDataManager.PostLoad(PerInstanceSMData.Num(), MoveTemp(InstanceDataBufferSerializationTmp));
 	}
 
+	// Has different implementation in HISMC
+	OnPostLoadPerInstanceData();
+
+
 	// release InstanceDataBuffers
 	InstanceDataBufferSerializationTmp.Reset();
 }
diff --git a/Engine/Source/Runtime/Engine/Private/InstancedStaticMesh/ISMInstanceDataManager.cpp b/Engine/Source/Runtime/Engine/Private/InstancedStaticMesh/ISMInstanceDataManager.cpp
--- a/Engine/Source/Runtime/Engine/Private/InstancedStaticMesh/ISMInstanceDataManager.cpp
+++ b/Engine/Source/Runtime/Engine/Private/InstancedStaticMesh/ISMInstanceDataManager.cpp
@@ -396,17 +396,35 @@ void FPrimitiveInstanceDataManager::DispatchUpdateTask(bool bIsUnattached, const
 	}
 }
 
-struct FExternalGenericChangeSet
+/**
+ * Data captured that may be flushed to rebuild the proxy from legacy data AKA FStaticMeshInstanceData.
+ * Is marshalled across via the FLegacyRebuildChangeSet.
+ */
+struct FLegacyBuildData
+{
+	TUniquePtr<FStaticMeshInstanceData> LegacyStaticMeshInstanceData;
+	TArray<int32> LegacyInstanceReorderTable;
+	FInstanceIdIndexMap InstanceIdIndexMap;
+#if WITH_EDITOR
+	TArray<TRefCountPtr<HHitProxy>> HitProxies;
+#endif
+};
+
+struct FLegacyRebuildChangeSet
 {
 	FRenderBounds InstanceLocalBounds;
 	FVector PrimitiveWorldSpaceOffset;
 	FRenderTransform PrimitiveToRelativeWorld;
 	FInstanceDataFlags Flags;
 	float AbsMaxDisplacement = 0.0f;
+	TUniquePtr<FStaticMeshInstanceData> LegacyStaticMeshInstanceData;
+	TArray<int32> LegacyInstanceReorderTable;
+	FInstanceIdIndexMap InstanceIdIndexMap;
+
 #if WITH_EDITOR
 	TPimplPtr<FOpaqueHitProxyContainer> HitProxyContainer;
 #endif
-	FPrimitiveInstanceDataManager::FExternalUpdateData ExternalUpdateData;
+	int32 NumCustomDataFloats = 0; 
 };
 
 void FPrimitiveInstanceDataManager::InitChangeSet(const FChangeDesc &ChangeDesc, const FInstanceUpdateComponentDesc &ComponentData, FISMInstanceUpdateChangeSet &ChangeSet)
@@ -414,13 +432,14 @@ void FPrimitiveInstanceDataManager::InitChangeSet(const FChangeDesc &ChangeDesc,
 	if (ChangeSet.IsFullUpdate())
 	{
 		// Initialize to full update collection
+		// TODO: Initialize to empty conditionally using the Flags
 		ChangeSet.CustomDataDelta 
 #if WITH_EDITOR
 			= ChangeSet.InstanceEditorDataDelta 
 #endif
 			= ChangeSet.InstanceLightShadowUVBiasDelta 
 			= ChangeSet.TransformsDelta 
-			= FArrayIndexDelta(GetMaxInstanceIndex());
+			= FArrayIndexDelta(ComponentData.NumSourceInstances);
 	}
 	else
 	{
@@ -437,12 +456,12 @@ void FPrimitiveInstanceDataManager::InitChangeSet(const FChangeDesc &ChangeDesc,
 			{
 				int32 Index = IdToIndex(FPrimitiveInstanceId{BitIt.GetIndex()});
 				check(Index != INDEX_NONE);
-				if (!ChangeDesc.bPrimitiveTransformChanged)
+				if (!bSendAllTransforms)
 				{
 					ChangeSet.TransformsDelta.AddIndex(Index);
 				}
 
-				if (!bNumCustomDataChanged)
+				if (!bNumCustomDataChanged && Flags.bHasPerInstanceCustomData)
 				{
 					ChangeSet.CustomDataDelta.AddIndex(Index);
 				}
@@ -456,7 +475,7 @@ void FPrimitiveInstanceDataManager::InitChangeSet(const FChangeDesc &ChangeDesc,
 		}
 #if WITH_EDITOR
 		// Bulk update editor data - the hit-proxies are recreated all the time.
-		ChangeSet.InstanceEditorDataDelta = FArrayIndexDelta(GetMaxInstanceIndex());
+		ChangeSet.InstanceEditorDataDelta = FArrayIndexDelta(ComponentData.NumSourceInstances);
 #endif
 		if (bSendAllTransforms)
 		{
@@ -477,27 +496,34 @@ void FPrimitiveInstanceDataManager::InitChangeSet(const FChangeDesc &ChangeDesc,
 		}
 
 		// Update all the custom data if there was a change in number of floats (this resets the whole thing).
-		if (bNumCustomDataChanged)
-		{
-			ChangeSet.CustomDataDelta = FArrayIndexDelta(GetMaxInstanceIndex());
-		}
-		else
+		if (Flags.bHasPerInstanceCustomData)
 		{
-			for (TConstSetBitIterator<> BitIt(CustomDataChangedInstances); BitIt; ++BitIt)
+			if (bNumCustomDataChanged)
 			{
-				FPrimitiveInstanceId Id{BitIt.GetIndex()};
-				if (!AddedInstances[Id] && !RemovedInstances[Id])
+				ChangeSet.CustomDataDelta = FArrayIndexDelta(ComponentData.NumSourceInstances);
+			}
+			else
+			{
+				for (TConstSetBitIterator<> BitIt(CustomDataChangedInstances); BitIt; ++BitIt)
 				{
-					int32 Index = IdToIndex(Id);
-					check(Index != INDEX_NONE);
-					ChangeSet.CustomDataDelta.AddIndex(Index);
+					FPrimitiveInstanceId Id{BitIt.GetIndex()};
+					if (!AddedInstances[Id] && !RemovedInstances[Id])
+					{
+						int32 Index = IdToIndex(Id);
+						check(Index != INDEX_NONE);
+						ChangeSet.CustomDataDelta.AddIndex(Index);
+					}
 				}
 			}
 		}
+		else
+		{
+			ChangeSet.CustomDataDelta = FArrayIndexDelta();
+		}
 
 		if (bBakedLightingDataChanged)
 		{
-			ChangeSet.InstanceLightShadowUVBiasDelta = FArrayIndexDelta(GetMaxInstanceIndex());
+			ChangeSet.InstanceLightShadowUVBiasDelta = FArrayIndexDelta(ComponentData.NumSourceInstances);
 		}
 		else
 		{
@@ -536,9 +562,9 @@ void FPrimitiveInstanceDataManager::InitChangeSet(const FChangeDesc &ChangeDesc,
 	ChangeSet.PrimitiveToRelativeWorld = FLargeWorldRenderScalar::MakeToRelativeWorldMatrix(ChangeSet.PrimitiveWorldSpaceOffset, PrimitiveLocalToWorld);
 	ChangeSet.Flags = Flags;
 	ChangeSet.AbsMaxDisplacement = AbsMaxDisplacement;
+	ChangeSet.NumCustomDataFloats = NumCustomDataFloats;
 	if (!ChangeSet.Flags.bHasPerInstanceCustomData)
 	{
-		ChangeSet.NumCustomDataFloats = 0;
 		ChangeSet.CustomDataDelta = FArrayIndexDelta();
 	}
 	ChangeSet.InstanceIdIndexMap = FInstanceIdIndexMap(*this);
@@ -573,47 +599,50 @@ bool FPrimitiveInstanceDataManager::FlushChanges(FInstanceUpdateComponentDesc &&
 	ERHIFeatureLevel::Type FeatureLevel = ComponentData.PrimitiveSceneProxy->GetScene().GetFeatureLevel();
 	check(Proxy->CheckPlatformFeatureLevel(ShaderPlatform, FeatureLevel));
 
+
+	bool bWasUpdateQueued = false;
 	// Marked for externally managed update, this may be combined with tracked changes (which results in double updates)
-	if (ExternalBuildUpdateData)
+	if (LegacyBuildData)
 	{
 		LOG_INST_DATA(TEXT("Rebuild with EState::External %s"), TEXT(""));
 
-		FExternalGenericChangeSet ExternalChangeSet;
+		check(LegacyBuildData->InstanceIdIndexMap.GetMaxInstanceIndex() == LegacyBuildData->LegacyStaticMeshInstanceData->GetNumInstances());
+
+		// An "external build" is responsible for setting everything up in the right space.
+		bPrimitiveTransformChanged = false;
+		PrimitiveLocalToWorld = ComponentData.PrimitiveLocalToWorld;
+		Flags = ComponentData.Flags;
+		StaticMeshBounds = ComponentData.StaticMeshBounds;
+		NumCustomDataFloats = ComponentData.NumCustomDataFloats;
+
+		FLegacyRebuildChangeSet ExternalChangeSet;
 		ExternalChangeSet.PrimitiveWorldSpaceOffset = FLargeWorldRenderPosition(PrimitiveLocalToWorld.GetOrigin()).GetTileOffset();
 		ExternalChangeSet.PrimitiveToRelativeWorld = FLargeWorldRenderScalar::MakeToRelativeWorldMatrix(ExternalChangeSet.PrimitiveWorldSpaceOffset, PrimitiveLocalToWorld);
 		ExternalChangeSet.InstanceLocalBounds = ComponentData.StaticMeshBounds;
-		ExternalChangeSet.Flags = ComponentData.Flags;
+		ExternalChangeSet.Flags = Flags;
 
-		TArray<TRefCountPtr<HHitProxy>> HitProxies;
-		ExternalChangeSet.ExternalUpdateData = ExternalBuildUpdateData(HitProxies);
+		ExternalChangeSet.LegacyStaticMeshInstanceData = MoveTemp(LegacyBuildData->LegacyStaticMeshInstanceData);
+		ExternalChangeSet.LegacyInstanceReorderTable = MoveTemp(LegacyBuildData->LegacyInstanceReorderTable);
+		ExternalChangeSet.InstanceIdIndexMap = MoveTemp(LegacyBuildData->InstanceIdIndexMap);
 #if WITH_EDITOR
-		ExternalChangeSet.HitProxyContainer = MakeOpaqueHitProxyContainer(HitProxies);
+		ExternalChangeSet.HitProxyContainer = MakeOpaqueHitProxyContainer(LegacyBuildData->HitProxies);
 #endif
 
 		ExternalChangeSet.AbsMaxDisplacement = AbsMaxDisplacement;
-		
-		// Make sure we have a tracking state that matches the number of instances, but doesn't retain any history
-		ClearIdTracking(ExternalChangeSet.ExternalUpdateData.NumInstances);
-
-		// Need to store this off for tracking purposes
-		NumCustomDataFloats = ExternalChangeSet.ExternalUpdateData.NumCustomDataFloats;
 
 		// Assemble header info to enable nonblocking primitive update.
 		FInstanceDataBufferHeader InstanceDataBufferHeader;
-		InstanceDataBufferHeader.NumInstances = NumInstances;
-		InstanceDataBufferHeader.PayloadDataStride = FInstanceSceneDataBuffers::CalcPayloadDataStride(ExternalChangeSet.Flags, ExternalChangeSet.ExternalUpdateData.NumCustomDataFloats, 0);
+		InstanceDataBufferHeader.NumInstances = ExternalChangeSet.InstanceIdIndexMap.GetMaxInstanceIndex();
+		InstanceDataBufferHeader.PayloadDataStride = FInstanceSceneDataBuffers::CalcPayloadDataStride(ExternalChangeSet.Flags, NumCustomDataFloats, 0);
 		InstanceDataBufferHeader.Flags = ExternalChangeSet.Flags;
 
-		DispatchUpdateTask(bIsUnattached, InstanceDataBufferHeader,
-			[ExternalChangeSet = MoveTemp(ExternalChangeSet), 
-			Proxy = Proxy,
-			InstanceIdIndexMap = FInstanceIdIndexMap(*this)] () mutable 
+		DispatchUpdateTask(bIsUnattached, InstanceDataBufferHeader, [ExternalChangeSet = MoveTemp(ExternalChangeSet), Proxy = Proxy] () mutable 
 		{
 			FISMCInstanceDataSceneProxy &ProxyRef = *Proxy;
 			
 			// Forcibly destroy any tracking state, if the external entity manages this it can track the data on its own.
 			ProxyRef.ChangeMask.Reset();
-			ProxyRef.InstanceIdIndexMap = MoveTemp(InstanceIdIndexMap);
+			ProxyRef.InstanceIdIndexMap = MoveTemp(ExternalChangeSet.InstanceIdIndexMap);
 #if WITH_EDITOR
 			ProxyRef.HitProxyContainer = MoveTemp(ExternalChangeSet.HitProxyContainer);
 #endif
@@ -629,12 +658,15 @@ bool FPrimitiveInstanceDataManager::FlushChanges(FInstanceUpdateComponentDesc &&
 			ExternalChangeSet.InstanceLocalBounds.Min -= PadExtent;
 			ExternalChangeSet.InstanceLocalBounds.Max += PadExtent;
 
-			ExternalChangeSet.ExternalUpdateData.UpdateProxy(ProxyRef, ExternalChangeSet.InstanceLocalBounds);
+			ProxyRef.BuildFromLegacyData(MoveTemp(ExternalChangeSet.LegacyStaticMeshInstanceData), ExternalChangeSet.InstanceLocalBounds, MoveTemp(ExternalChangeSet.LegacyInstanceReorderTable));
 
-			check(ProxyRef.GetData().GetNumInstances() == InstanceIdIndexMap.GetMaxInstanceIndex());
-			check(ProxyRef.GetUpdateTaskInfo()->GetHeader().NumInstances == InstanceIdIndexMap.GetMaxInstanceIndex());
+			check(ProxyRef.GetData().GetNumInstances() == ProxyRef.InstanceIdIndexMap.GetMaxInstanceIndex());
+			check(ProxyRef.GetUpdateTaskInfo()->GetHeader().NumInstances == ProxyRef.InstanceIdIndexMap.GetMaxInstanceIndex());
 			check(ProxyRef.GetUpdateTaskInfo()->GetHeader().Flags == ProxyRef.InstanceSceneDataBuffers.GetFlags());
 		});
+
+		LegacyBuildData.Reset();
+		bWasUpdateQueued = true;
 	}
 
 	FChangeDesc ChangeDesc;
@@ -658,6 +690,12 @@ bool FPrimitiveInstanceDataManager::FlushChanges(FInstanceUpdateComponentDesc &&
 		ChangeDesc.bMaxDisplacementChanged = AbsMaxDisplacement != NewAbsMaxDisplacement;
 		ChangeDesc.bStaticMeshBoundsChanged = !StaticMeshBounds.Equals(ComponentData.StaticMeshBounds);
 	}
+	
+	// Yet another special case to handle externally managed data from landscape grass
+	if (Mode == EMode::ExternalLegacyData && (bPrimitiveTransformChanged || !ComponentData.PrimitiveLocalToWorld.Equals(PrimitiveLocalToWorld)))
+	{
+		ChangeDesc.bPrimitiveTransformChanged = true;
+	}
 
 	FMatrix PrevPrimitiveLocalToWorld = PrimitiveLocalToWorld;
 	// Update the tracked state
@@ -665,28 +703,30 @@ bool FPrimitiveInstanceDataManager::FlushChanges(FInstanceUpdateComponentDesc &&
 	AbsMaxDisplacement = NewAbsMaxDisplacement;
 	StaticMeshBounds = ComponentData.StaticMeshBounds;
 	Flags = ComponentData.Flags;
+	NumCustomDataFloats = ComponentData.NumCustomDataFloats;
 
 	const bool bAnyChange = ChangeDesc.Packed != 0;
 
 	if (!bAnyChange)
 	{
-		return false;
+		return bWasUpdateQueued;
 	}
 	
+	// No need to send an update for an ISM with zero source instances && no instance updates 
+	// except if the primitive transform changed, in which case we need to pipe that through if there are instances in the proxy.
+	if (ComponentData.NumSourceInstances == 0 && !ChangeDesc.bInstancesChanged 
+		&& !(ChangeDesc.bPrimitiveTransformChanged && ComponentData.NumProxyInstances != 0))
+	{
+		ClearChangeTracking();
+		return bWasUpdateQueued;
+	}
+
 	// 
 	{
 		// TODO: Maybe specialize for only bPrimitiveTransformChanged (no need to send/modify anything _other_ than the transform data)
 		// TODO: The states other than bCreateNewProxy _can_ be handled through delta updates (e.g., add instance + offset all the rest).
 		bool bNeedFullUpdate = ChangeDesc.bUntrackedState || ChangeDesc.bMaterialUsageFlagsChanged;
 
-		// TODO: fix this up somehow: can probably remove now since we now propagate the count with invalidations & just leave a check
-		// In this case, here's where we found out the true count of instances, since before it's not been tracked.
-		if (ChangeDesc.bUntrackedState)
-		{
-			check(ComponentData.NumSourceInstances >= 0);
-			NumInstances = ComponentData.NumSourceInstances;
-		}
-
 		FISMInstanceUpdateChangeSet ChangeSet(bNeedFullUpdate);
 
 		// Initialize the change set before collecting instance change data.
@@ -696,8 +736,7 @@ bool FPrimitiveInstanceDataManager::FlushChanges(FInstanceUpdateComponentDesc &&
 		ComponentData.BuildChangeSet(ChangeSet);
 		
 		// make sure the custom data change is correctly tracked
-		check(bNeedFullUpdate || bNumCustomDataChanged || NumCustomDataFloats == ChangeSet.NumCustomDataFloats);
-		NumCustomDataFloats = ChangeSet.NumCustomDataFloats;
+		check(!ChangeSet.Flags.bHasPerInstanceCustomData || NumCustomDataFloats == ChangeSet.NumCustomDataFloats);
 		check(ChangeSet.Flags.bHasPerInstanceCustomData || ChangeSet.CustomDataDelta.IsEmpty() && ChangeSet.PerInstanceCustomData.IsEmpty());
 
 		// If we have per-instance previous local to world, they are expected to be in the local space of the _previous_ local to world. If they are in fact not (e.g., if someone sets them explicitly from world space) then, well, this won't be correct
@@ -732,9 +771,7 @@ bool FPrimitiveInstanceDataManager::FlushChanges(FInstanceUpdateComponentDesc &&
 
 
 		// Special handling for the case of external data & transform change, may want to extend this to others (but need to send data in that case)
-		FChangeDesc OnlyPrimitiveTransformChangeDesc;
-		OnlyPrimitiveTransformChangeDesc.bPrimitiveTransformChanged = true;
-		if (Mode == EMode::ExternalLegacyData && ChangeDesc.Packed == OnlyPrimitiveTransformChangeDesc.Packed)
+		if (Mode == EMode::ExternalLegacyData && ChangeDesc.bPrimitiveTransformChanged)
 		{
 			LOG_INST_DATA(TEXT("Primitive Transform Update %s"), TEXT(""));
 			DispatchUpdateTask(bIsUnattached, InstanceDataBufferHeader, [ChangeSet = MoveTemp(ChangeSet), Proxy = Proxy] () mutable 
@@ -815,7 +852,7 @@ void FPrimitiveInstanceDataManager::ClearChangeTracking()
 	// When tracking data is cleared, we loose connection to previously tracked state until the next update is sent.
 	TrackingState = ETrackingState::Initial;
 
-	ExternalBuildUpdateData.Reset();
+	LegacyBuildData.Reset();
 
 	AddedInstances.Empty();
 	RemovedInstances.Empty();
@@ -1034,13 +1071,35 @@ void FPrimitiveInstanceDataManager::ValidateMapping() const
 #endif
 }
 
-void FPrimitiveInstanceDataManager::MarkForRebuildFromExternal(int32 InNumInstances, TUniqueFunction<FExternalUpdateData(TArray<TRefCountPtr<HHitProxy>> &OutHitProxies)> &&BuildUpdateData)
+void FPrimitiveInstanceDataManager::MarkForRebuildFromLegacy(TUniquePtr<FStaticMeshInstanceData> &&InLegacyInstanceData, const TArray<int32> &InstanceReorderTable, const TArray<TRefCountPtr<HHitProxy>> &HitProxies)
 {
-	// Forget any tracked changes that happened before (and ID-mapping)
-	ClearIdTracking(InNumInstances);
-	ExternalBuildUpdateData = MoveTemp(BuildUpdateData);
-	// Must actually track the changes because we may see incremental changes in the same frame,
-	TrackingState = ETrackingState::Tracked;
+	check(InLegacyInstanceData);
+	// TODO: restore the ID tracking when not in external mode
+	//if (Mode == EMode::ExternalLegacyData)
+	{
+		ClearIdTracking(InLegacyInstanceData->GetNumInstances());
+	}
+	//else
+	//{
+	//	ClearChangeTracking();
+	//}
+	check(GetMaxInstanceIndex() == InLegacyInstanceData->GetNumInstances());
+
+	LegacyBuildData = MakePimpl<FLegacyBuildData>();
+	
+	// Need to capture the InstanceIdIndexMap map here, as there may be further changes before the next flush is called.
+	LegacyBuildData->LegacyStaticMeshInstanceData = MoveTemp(InLegacyInstanceData);
+	LegacyBuildData->LegacyInstanceReorderTable = InstanceReorderTable;
+	LegacyBuildData->InstanceIdIndexMap = FInstanceIdIndexMap(*this);
+#if WITH_EDITOR
+	LegacyBuildData->HitProxies = HitProxies; 
+#endif
+
+	if (Mode != EMode::ExternalLegacyData)
+	{
+		// Must actually track the changes because we may see incremental changes in the same frame,
+		TrackingState = ETrackingState::Tracked;
+	}
 	MarkComponentRenderInstancesDirty();
 }
 
diff --git a/Engine/Source/Runtime/Engine/Private/InstancedStaticMesh/ISMInstanceUpdateChangeSet.cpp b/Engine/Source/Runtime/Engine/Private/InstancedStaticMesh/ISMInstanceUpdateChangeSet.cpp
--- a/Engine/Source/Runtime/Engine/Private/InstancedStaticMesh/ISMInstanceUpdateChangeSet.cpp
+++ b/Engine/Source/Runtime/Engine/Private/InstancedStaticMesh/ISMInstanceUpdateChangeSet.cpp
@@ -94,25 +94,16 @@ void FISMInstanceUpdateChangeSet::SetCustomData(const TArrayView<const float> &I
 {
 	if (Flags.bHasPerInstanceCustomData)
 	{
-		// Switch over to not having any if count is zero
-		if (InNumCustomDataFloats == 0)
-		{
-			CustomDataDelta = FArrayIndexDelta();
-			Flags.bHasPerInstanceCustomData = false;
-			return;
-		}
+		check(InNumCustomDataFloats == NumCustomDataFloats);
+		check(NumCustomDataFloats > 0);
+		check(InPerInstanceCustomData.Num() == InstanceIdIndexMap.GetMaxInstanceIndex() * NumCustomDataFloats);
 
-		check(InPerInstanceCustomData.Num() == InstanceIdIndexMap.GetMaxInstanceIndex() * InNumCustomDataFloats);
-
-		// If the count changed, we must update everything
-		if (NumCustomDataFloats != InNumCustomDataFloats)
-		{
-			// No delta state (copy everything)
-			CustomDataDelta = FArrayIndexDelta(InstanceIdIndexMap.GetMaxInstanceIndex());
-		}
-		NumCustomDataFloats = InNumCustomDataFloats;
 		CustomDataDelta.Gather(PerInstanceCustomData, InPerInstanceCustomData, InNumCustomDataFloats);
 	}
+	else
+	{
+		check(CustomDataDelta.IsEmpty());
+	}
 }
 
 void FISMInstanceUpdateChangeSet::SetInstanceLocalBounds(const FRenderBounds &Bounds)
diff --git a/Engine/Source/Runtime/Engine/Private/PrimitiveSceneProxy.cpp b/Engine/Source/Runtime/Engine/Private/PrimitiveSceneProxy.cpp
--- a/Engine/Source/Runtime/Engine/Private/PrimitiveSceneProxy.cpp
+++ b/Engine/Source/Runtime/Engine/Private/PrimitiveSceneProxy.cpp
@@ -872,7 +872,10 @@ void FPrimitiveSceneProxy::BuildUniformShaderParameters(FPrimitiveUniformShaderP
 	if (HasInstanceDataBuffers())
 	{
 		const FInstanceSceneDataBuffers* InstanceSceneDataBuffers = GetInstanceSceneDataBuffers();
-		Builder.InstanceLocalBounds(InstanceSceneDataBuffers->GetInstanceLocalBounds(0));
+		if (GetInstanceDataHeader().NumInstances > 0)
+		{
+			Builder.InstanceLocalBounds(InstanceSceneDataBuffers->GetInstanceLocalBounds(0));
+		}
 	}
 
 	if (ShouldRenderCustomDepth())
diff --git a/Engine/Source/Runtime/Engine/Public/InstancedStaticMesh/ISMInstanceDataManager.h b/Engine/Source/Runtime/Engine/Public/InstancedStaticMesh/ISMInstanceDataManager.h
--- a/Engine/Source/Runtime/Engine/Public/InstancedStaticMesh/ISMInstanceDataManager.h
+++ b/Engine/Source/Runtime/Engine/Public/InstancedStaticMesh/ISMInstanceDataManager.h
@@ -25,6 +25,8 @@ struct FInstanceUpdateComponentDesc
 	int32 NumProxyInstances = -1;
 	// Number of instances in the source (e.g., the component)
 	int32 NumSourceInstances = -1;
+	// 
+	int32 NumCustomDataFloats = 0;
 
 	// Callback to fill in the required change set use with delta-update capable onwing components.
 	// TODO: move elsewhere?
@@ -134,21 +136,10 @@ public:
 	
 	void ValidateMapping() const;
 
-	struct FExternalUpdateData
-	{
-		int32 NumCustomDataFloats = 0; 
-		int32 NumInstances = -1;
-		TUniqueFunction<void(FISMCInstanceDataSceneProxy &Proxy, const FRenderBounds &InstanceLocalBounds)> UpdateProxy;
-	};
 	/**
 	 * Call to mark the manager as needing a full rebuild & having an external driver for this.
-	 * The BuildUpdateData function may not be called, as this only happens if the proxy is actually created.
-	 * BuildUpdateData shares lifetime with the component so can safely capture the this pointer.
-	 * The FExternalUpdateData produced by BuildUpdateData is responsible for capturing data needed to update the Proxy, 
-	 * This will be moved through to the updater task and executed there, should thus not reference the component.
-	 * Calling this also invalidates any ID/Index mapping as the instance data management is completely up to the external source.
 	 */
-	ENGINE_API void MarkForRebuildFromExternal(int32 InNumInstances, TUniqueFunction<FExternalUpdateData(TArray<TRefCountPtr<HHitProxy>> &OutHitProxies)> &&BuildUpdateData);
+	ENGINE_API void MarkForRebuildFromLegacy(TUniquePtr<FStaticMeshInstanceData> &&LegacyInstanceData, const TArray<int32> &InstanceReorderTable, const TArray<TRefCountPtr<HHitProxy>> &HitProxies);
 
 	/**
 	 */
@@ -207,7 +198,7 @@ private:
 	bool bComponentMarkedDirty = false;
 	bool bEnableTracking = false;
 
-	TUniqueFunction<FExternalUpdateData(TArray<TRefCountPtr<HHitProxy>> &OutHitProxies)> ExternalBuildUpdateData;
+	TPimplPtr<struct FLegacyBuildData> LegacyBuildData;
 
 	// Used for serialized legacy data to drive the non-GPU scene rendering path.
 	TUniquePtr<FStaticMeshInstanceData> LegacyStaticMeshInstanceData;
diff --git a/Engine/Source/Runtime/Engine/Public/InstancedStaticMesh/ISMInstanceUpdateChangeSet.h b/Engine/Source/Runtime/Engine/Public/InstancedStaticMesh/ISMInstanceUpdateChangeSet.h
--- a/Engine/Source/Runtime/Engine/Public/InstancedStaticMesh/ISMInstanceUpdateChangeSet.h
+++ b/Engine/Source/Runtime/Engine/Public/InstancedStaticMesh/ISMInstanceUpdateChangeSet.h
@@ -209,11 +209,11 @@ public:
 		return bNeedFullUpdate;
 	}
 
-	bool bNeedFullUpdate;
+	bool bNeedFullUpdate = false;
 	FChangeMask ChangeMask;
 
 	FInstanceIdIndexMap InstanceIdIndexMap;
-	int32 NumCustomDataFloats;
+	int32 NumCustomDataFloats = 0;
 
 	// TODO: Make it possible to share the deltas such that we don't have to produce an unique index array for each attribute
 	//       e.g., if the only thing that happened was somethis was added and we therefore need all attributes for one instance.
diff --git a/Engine/Source/Runtime/Foliage/Private/InstancedGrass.cpp b/Engine/Source/Runtime/Foliage/Private/InstancedGrass.cpp
--- a/Engine/Source/Runtime/Foliage/Private/InstancedGrass.cpp
+++ b/Engine/Source/Runtime/Foliage/Private/InstancedGrass.cpp
@@ -90,19 +90,8 @@ void UGrassInstancedStaticMeshComponent::AcceptPrebuiltTree(TArray<FClusterNode>
 	TUniquePtr<FStaticMeshInstanceData> BuiltInstanceData = MakeUnique<FStaticMeshInstanceData>();
 	// TODO: Implement move semantics for FStaticMeshInstanceData!
 	Swap(*BuiltInstanceData, *InSharedInstanceBufferData);
+	PrimitiveInstanceDataManager.MarkForRebuildFromLegacy(MoveTemp(BuiltInstanceData), InstanceReorderTable, TArray<TRefCountPtr<HHitProxy>>());
 
-	PrimitiveInstanceDataManager.MarkForRebuildFromExternal(NumBuiltRenderInstances, [LegacyInstanceData = MoveTemp(BuiltInstanceData)] (TArray<TRefCountPtr<HHitProxy>> &OutHitProxies) mutable
-	{
-		OutHitProxies.Reset(); 
-		FPrimitiveInstanceDataManager::FExternalUpdateData ExternalUpdateData;
-		ExternalUpdateData.NumInstances = LegacyInstanceData ? LegacyInstanceData->GetNumInstances() : 0;
-		ExternalUpdateData.UpdateProxy = [LegacyInstanceDataInner = MoveTemp(LegacyInstanceData)](FISMCInstanceDataSceneProxy &InstanceDataSceneProxy, const FRenderBounds &InstanceLocalBounds) mutable
-		{
-			// No reorder table needed since we never perform delta updates.
-			InstanceDataSceneProxy.BuildFromLegacyData(MoveTemp(LegacyInstanceDataInner), InstanceLocalBounds, TArray<int32>());
-		};
-		return ExternalUpdateData;
-	});
 	MarkRenderStateDirty();
 }
 
@@ -115,13 +104,17 @@ void UGrassInstancedStaticMeshComponent::BuildComponentInstanceData(FInstanceUpd
 	OutData.StaticMeshBounds = GetStaticMesh()->GetBounds();
 	OutData.NumProxyInstances = InstanceCountToRender;
 	OutData.NumSourceInstances = 0;
+	OutData.NumCustomDataFloats = NumCustomDataFloats;
+
 #if WITH_EDITOR
 	// TODO: Do we want these for this path?
 	OutData.Flags.bHasPerInstanceEditorData = false;
 #endif
 	OutData.BuildChangeSet = [&](FISMInstanceUpdateChangeSet &ChangeSet)
 	{
-		// Cancel any update that is not transforms, because that is the only data this primitive stores for this very special use-case
+		// Cancel update as there is no source data
+		check(ChangeSet.TransformsDelta.IsEmpty());
+		ChangeSet.TransformsDelta = FArrayIndexDelta();
 		ChangeSet.InstanceLightShadowUVBiasDelta = FArrayIndexDelta();
 		ChangeSet.CustomDataDelta = FArrayIndexDelta();
 #if WITH_EDITOR
@@ -129,8 +122,8 @@ void UGrassInstancedStaticMeshComponent::BuildComponentInstanceData(FInstanceUpd
 #endif
 		BuildInstanceDataDeltaChangeSetCommon(ChangeSet);
 		check(GetTranslatedInstanceSpaceOrigin().IsNearlyZero());
-		ChangeSet.SetInstanceTransforms(MakeStridedView(PerInstanceSMData, &FInstancedStaticMeshInstanceData::Transform));
-		ChangeSet.SetInstancePrevTransforms(MakeArrayView(PerInstancePrevTransform));
+		check(PerInstanceSMData.IsEmpty());
+		check(PerInstancePrevTransform.IsEmpty());
 		// The reorder table is always empty in this path because the HISM is populated in AcceptPrebuiltTree using instances in the already sorted order.
 		check(InstanceReorderTable.IsEmpty());
 		check(ChangeSet.LegacyInstanceReorderTable.IsEmpty());
diff --git a/Engine/Source/Runtime/Renderer/Private/GPUScene.cpp b/Engine/Source/Runtime/Renderer/Private/GPUScene.cpp
--- a/Engine/Source/Runtime/Renderer/Private/GPUScene.cpp
+++ b/Engine/Source/Runtime/Renderer/Private/GPUScene.cpp
@@ -361,7 +361,7 @@ void ValidateInstanceUploadInfo(const FInstanceUploadInfo& UploadInfo, FRDGBuffe
 
 	if (UploadInfo.InstanceSceneDataBuffers != nullptr)
 	{
-		if (!ensure(UploadInfo.InstanceSceneDataBuffers->GetPrimitiveToRelativeWorld().Equals(UploadInfo.PrimitiveToWorld)))
+		if (UploadInfo.InstanceSceneDataBuffers->GetNumInstances() > 0 && !ensure(UploadInfo.InstanceSceneDataBuffers->GetPrimitiveToRelativeWorld().Equals(UploadInfo.PrimitiveToWorld)))
 		{
 			UE_LOG(LogTemp, Warning, TEXT("Mismatched Primitive transform! primitive ID %d"), UploadInfo.PrimitiveID);
 		}
