Fix GLEO error related to the use of OFPA in non World Partition map
* Objects which have the external package as their outer were left behind and led to GLEO errors
* Temporary reouter them to the level during cook - handled by ULevel::DetachAttachAllActorsPackages()

diff --git a/Engine/Source/Runtime/Engine/Classes/Engine/Level.h b/Engine/Source/Runtime/Engine/Classes/Engine/Level.h
--- a/Engine/Source/Runtime/Engine/Classes/Engine/Level.h
+++ b/Engine/Source/Runtime/Engine/Classes/Engine/Level.h
@@ -791,6 +791,10 @@ private:
 
 	UPROPERTY(Transient)
 	TWeakObjectPtr<UObject> EditorPathOwner;
+
+	/** Temporary map of objects to their associated external packages. Used when detaching/attaching external actors packages during cook. */
+	UPROPERTY(Transient)
+	TMap<TObjectPtr<UObject>, TObjectPtr<UPackage>> ObjectsToExternalPackages;
 #endif // #if WITH_EDITORONLY_DATA
 
 	enum class ERouteActorInitializationState : uint8
diff --git a/Engine/Source/Runtime/Engine/Private/Level.cpp b/Engine/Source/Runtime/Engine/Private/Level.cpp
--- a/Engine/Source/Runtime/Engine/Private/Level.cpp
+++ b/Engine/Source/Runtime/Engine/Private/Level.cpp
@@ -4090,14 +4090,49 @@ void ULevel::DetachAttachAllActorsPackages(bool bReattach)
 				Actor->ReattachExternalPackage();
 			}
 		}
+
+		// Reouter objects previously found in the actors packages to their original packages
+		for (const TPair<TObjectPtr<UObject>, TObjectPtr<UPackage>> ObjectToPackage : ObjectsToExternalPackages)
+		{
+			UObject* Object = ObjectToPackage.Key;
+			UPackage* Package = ObjectToPackage.Value;
+
+			Object->Rename(nullptr, Package, REN_ForceNoResetLoaders);
+		}
+
+		ObjectsToExternalPackages.Empty();
 	}
 	else
 	{
+		UPackage* LevelPackage = GetPackage();
+
+		check(ObjectsToExternalPackages.IsEmpty());
+
 		for (AActor* Actor : Actors)
 		{
 			if (Actor)
 			{
-				Actor->DetachExternalPackage();
+				if (UPackage* ActorExternalPackage = Actor->GetExternalPackage())
+				{
+					Actor->DetachExternalPackage();
+
+					// Process objects found in the source actor package
+					TArray<UObject*> Objects;
+					const bool bIncludeNestedSubobjects = false;
+
+					GetObjectsWithPackage(ActorExternalPackage, Objects, bIncludeNestedSubobjects, RF_NoFlags, EInternalObjectFlags::Garbage);
+					for (UObject* Object : Objects)
+					{
+						if (Object != Actor && Object->GetFName() != NAME_PackageMetaData)
+						{
+							// Move objects in the destination level package
+							Object->Rename(nullptr, LevelPackage, REN_ForceNoResetLoaders);
+
+							// Keep track of which package this object really belongs to
+							ObjectsToExternalPackages.Emplace(Object, ActorExternalPackage);
+						}
+					}
+				}
 			}
 		}
 	}
