Fixed CIS issue
*  Loop variable 'ObjectToPackage' creates a copy from type
#rnx
#buildfix

diff --git a/Engine/Source/Runtime/Engine/Private/Level.cpp b/Engine/Source/Runtime/Engine/Private/Level.cpp
--- a/Engine/Source/Runtime/Engine/Private/Level.cpp
+++ b/Engine/Source/Runtime/Engine/Private/Level.cpp
@@ -4092,7 +4092,7 @@ void ULevel::DetachAttachAllActorsPackages(bool bReattach)
 		}
 
 		// Reouter objects previously found in the actors packages to their original packages
-		for (const TPair<TObjectPtr<UObject>, TObjectPtr<UPackage>> ObjectToPackage : ObjectsToExternalPackages)
+		for (const TPair<TObjectPtr<UObject>, TObjectPtr<UPackage>>& ObjectToPackage : ObjectsToExternalPackages)
 		{
 			UObject* Object = ObjectToPackage.Key;
 			UPackage* Package = ObjectToPackage.Value;
