[Physics] Remove fast path runtime physics formats.

diff --git a/Engine/Source/Runtime/Engine/Classes/PhysicsEngine/BodySetup.h b/Engine/Source/Runtime/Engine/Classes/PhysicsEngine/BodySetup.h
--- a/Engine/Source/Runtime/Engine/Classes/PhysicsEngine/BodySetup.h
+++ b/Engine/Source/Runtime/Engine/Classes/PhysicsEngine/BodySetup.h
@@ -222,14 +222,6 @@ class UBodySetup : public UBodySetupCore
 
 	/** GUID used to uniquely identify this setup so it can be found in the DDC */
 	FGuid BodySetupGuid;
-
-private:
-#if WITH_EDITOR
-	/** Cooked physics data with runtime only optimizations. This allows us to remove editor only data (like face index remap) assuming the project doesn't use it at runtime. At runtime we load this into CookedFormatData */
-	FFormatContainer CookedFormatDataRuntimeOnlyOptimization;
-#endif
-
-public:
 	
 	/** list of chaos trimesh objects */
 	TArray<Chaos::FTriangleMeshImplicitObjectPtr> TriMeshGeometries;
@@ -324,10 +316,9 @@ private:
 	* Given a format name returns its cooked data.
 	*
 	* @param Format Physics format name.
-	* @param bRuntimeOnlyOptimizedVersion whether we want the data that has runtime only optimizations. At runtime this flag is ignored and we use the runtime only optimized data regardless.
 	* @return Cooked data or NULL of the data was not found.
 	*/
-	FByteBulkData* GetCookedData(FName Format, bool bRuntimeOnlyOptimizedVersion = false);
+	FByteBulkData* GetCookedData(FName Format);
 
 public:
 
diff --git a/Engine/Source/Runtime/Engine/Private/PhysicsEngine/BodySetup.cpp b/Engine/Source/Runtime/Engine/Private/PhysicsEngine/BodySetup.cpp
--- a/Engine/Source/Runtime/Engine/Private/PhysicsEngine/BodySetup.cpp
+++ b/Engine/Source/Runtime/Engine/Private/PhysicsEngine/BodySetup.cpp
@@ -797,9 +797,6 @@ void UBodySetup::InvalidatePhysicsData()
 	{
 		CookedFormatData.FlushData();
 	}
-#if WITH_EDITOR
-	CookedFormatDataRuntimeOnlyOptimization.FlushData();
-#endif
 }
 
 void UBodySetup::BeginDestroy()
@@ -849,16 +846,14 @@ void UBodySetup::Serialize(FArchive& Ar)
 			// Make sure to reset bHasCookedCollision data to true before calling GetCookedData for cooking
 			bHasCookedCollisionData = true;
 			FName Format = Ar.CookingTarget()->GetPhysicsFormat(this);
-			bool bUseRuntimeOnlyCookedData = !bSharedCookedData;	//For shared cook data we do not optimize for runtime only flags. This is only used by per poly skeletal mesh component at the moment. Might want to add support in future
-			bHasCookedCollisionData = GetCookedData(Format, bUseRuntimeOnlyCookedData) != NULL; // Get the data from the DDC or build it
+			bHasCookedCollisionData = GetCookedData(Format) != NULL; // Get the data from the DDC or build it
 
 			TArray<FName> ActualFormatsToSave;
 			ActualFormatsToSave.Add(Format);
 
 			FArchive_Serialize_BitfieldBool(Ar, bHasCookedCollisionData);
-
-			FFormatContainer* UseCookedFormatData = bUseRuntimeOnlyCookedData ? &CookedFormatDataRuntimeOnlyOptimization : &CookedFormatData;
-			UseCookedFormatData->Serialize(Ar, this, &ActualFormatsToSave, !bSharedCookedData);
+			
+			CookedFormatData.Serialize(Ar, this, &ActualFormatsToSave, !bSharedCookedData);
 
 #if VERIFY_COOKED_PHYS_DATA
 			// Verify that the cooked data matches the uncooked data
@@ -1211,13 +1206,13 @@ bool UBodySetup::IsCachedCookedPlatformDataLoaded(const ITargetPlatform* TargetP
 			return false;
 		}
 	}
-	GetCookedData(TargetPlatform->GetPhysicsFormat(this), true);
+	GetCookedData(TargetPlatform->GetPhysicsFormat(this));
 	return true;
 }
 
 void UBodySetup::ClearCachedCookedPlatformData( const ITargetPlatform* TargetPlatform )
 {
-	CookedFormatDataRuntimeOnlyOptimization.FlushData();
+
 }
 #endif
 
@@ -1366,7 +1361,7 @@ void GetDDCBuiltData(FByteBulkData* OutResult, DDCBuilderType& InBuilder, UBodyS
 
 #endif //#if WITH_EDITOR
 
-FByteBulkData* UBodySetup::GetCookedData(FName Format, bool bRuntimeOnlyOptimizedVersion)
+FByteBulkData* UBodySetup::GetCookedData(FName Format)
 {
 	if (IsTemplate())
 	{
@@ -1390,7 +1385,7 @@ FByteBulkData* UBodySetup::GetCookedData(FName Format, bool bRuntimeOnlyOptimize
 
 #if WITH_EDITOR
 	//We don't support runtime cook optimization for per poly skeletal mesh. This is an edge case we may want to support (only helps memory savings)
-	FFormatContainer* UseCookedData = CookedFormatDataOverride ? CookedFormatDataOverride : (bRuntimeOnlyOptimizedVersion ? &CookedFormatDataRuntimeOnlyOptimization : &CookedFormatData);
+	FFormatContainer* UseCookedData = CookedFormatDataOverride ? CookedFormatDataOverride : &CookedFormatData;
 #else
 	FFormatContainer* UseCookedData = CookedFormatDataOverride ? CookedFormatDataOverride : &CookedFormatData;
 #endif
