Adding sleep overrides to clusterunioncomponent for primitive component work

diff --git a/Engine/Source/Runtime/Engine/Classes/PhysicsEngine/ClusterUnionComponent.h b/Engine/Source/Runtime/Engine/Classes/PhysicsEngine/ClusterUnionComponent.h
--- a/Engine/Source/Runtime/Engine/Classes/PhysicsEngine/ClusterUnionComponent.h
+++ b/Engine/Source/Runtime/Engine/Classes/PhysicsEngine/ClusterUnionComponent.h
@@ -160,6 +160,9 @@ public:
 	UFUNCTION()
 	bool IsComponentAdded(UPrimitiveComponent* Component) { return PerComponentData.Contains(Component) || PendingComponentSync.Contains(Component); }
 
+	ENGINE_API virtual void WakeAllRigidBodies() override;
+	ENGINE_API virtual bool IsAnyRigidBodyAwake() override;
+
 	// Multi-trace/sweep functions that only make sense in the context of a cluster union.
 	ENGINE_API bool LineTraceComponent(TArray<FHitResult>& OutHit, const FVector Start, const FVector End, ECollisionChannel TraceChannel, const struct FCollisionQueryParams& Params, const struct FCollisionResponseParams& ResponseParams, const struct FCollisionObjectQueryParams& ObjectParams);
 	ENGINE_API bool SweepComponent(TArray<FHitResult>& OutHit, const FVector Start, const FVector End, const FQuat& ShapeWorldRotation, const FPhysicsGeometry& Geometry, ECollisionChannel TraceChannel, const struct FCollisionQueryParams& Params, const struct FCollisionResponseParams& ResponseParams, const struct FCollisionObjectQueryParams& ObjectParams);
diff --git a/Engine/Source/Runtime/Engine/Private/PhysicsEngine/ClusterUnionComponent.cpp b/Engine/Source/Runtime/Engine/Private/PhysicsEngine/ClusterUnionComponent.cpp
--- a/Engine/Source/Runtime/Engine/Private/PhysicsEngine/ClusterUnionComponent.cpp
+++ b/Engine/Source/Runtime/Engine/Private/PhysicsEngine/ClusterUnionComponent.cpp
@@ -1177,6 +1177,28 @@ void UClusterUnionComponent::SetSimulatePhysics(bool bSimulate)
 	PhysicsProxy->SetObjectState_External(bSimulate ? Chaos::EObjectStateType::Dynamic : Chaos::EObjectStateType::Kinematic);
 }
 
+void UClusterUnionComponent::WakeAllRigidBodies()
+{
+	if (!PhysicsProxy)
+	{
+		return;
+	}
+
+	Chaos::EObjectStateType State = PhysicsProxy->GetObjectState_External();
+	if (State == Chaos::EObjectStateType::Sleeping)
+	{
+		PhysicsProxy->SetObjectState_External(Chaos::EObjectStateType::Dynamic);
+	}
+}
+
+bool UClusterUnionComponent::IsAnyRigidBodyAwake()
+{
+	TArray<Chaos::FPhysicsObject*> PhysicsObjects = GetAllPhysicsObjects();
+	FLockedReadPhysicsObjectExternalInterface Interface = FPhysicsObjectExternalInterface::LockRead(PhysicsObjects);
+
+	return !Interface->AreAllSleeping(PhysicsObjects);
+}
+
 DECLARE_CYCLE_STAT(TEXT("UClusterUnionComponent::LineTraceComponentMulti"), STAT_ClusterUnionComponentLineTraceComponentMulti, STATGROUP_Chaos);
 bool UClusterUnionComponent::LineTraceComponent(TArray<FHitResult>& OutHit, const FVector Start, const FVector End, ECollisionChannel TraceChannel, const struct FCollisionQueryParams& Params, const struct FCollisionResponseParams& ResponseParams, const struct FCollisionObjectQueryParams& ObjectParams)
 {
