Prevent allocation in anim node message stack

diff --git a/Engine/Source/Runtime/Engine/Public/Animation/AnimNodeMessages.h b/Engine/Source/Runtime/Engine/Public/Animation/AnimNodeMessages.h
--- a/Engine/Source/Runtime/Engine/Public/Animation/AnimNodeMessages.h
+++ b/Engine/Source/Runtime/Engine/Public/Animation/AnimNodeMessages.h
@@ -346,8 +346,8 @@ private:
 		FNodeInfo NodeInfo;
 	};
 
-	using FMessageStackEntry = TArray<FMessageEntry, TInlineAllocator<4>>;
-	using FMessageMap = TMap<FName, FMessageStackEntry, TInlineSetAllocator<4>>;
+	using FMessageStackEntry = TArray<FMessageEntry, TInlineAllocator<8>>;
+	using FMessageMap = TMap<FName, FMessageStackEntry, TInlineSetAllocator<8>>;
 
 	// Message stack
 	FMessageMap MessageStacks;
