Fixed regression in AWorldDataLayers::GetDataLayerInstance not fallbacking on DeprecatedDataLayerNameToDataLayerInstance.

diff --git a/Engine/Source/Runtime/Engine/Private/WorldPartition/DataLayer/WorldDataLayers.cpp b/Engine/Source/Runtime/Engine/Private/WorldPartition/DataLayer/WorldDataLayers.cpp
--- a/Engine/Source/Runtime/Engine/Private/WorldPartition/DataLayer/WorldDataLayers.cpp
+++ b/Engine/Source/Runtime/Engine/Private/WorldPartition/DataLayer/WorldDataLayers.cpp
@@ -815,7 +815,10 @@ const UDataLayerInstance* AWorldDataLayers::GetDataLayerInstance(const FName& In
 			}
 			return true;
 		});
-		return FoundDataLayerInstance;
+		if (FoundDataLayerInstance)
+		{
+			return FoundDataLayerInstance;
+		}
 	}
 #else
 	if (const UDataLayerInstance* const* FoundDataLayerInstance = InstanceNameToInstance.Find(InDataLayerInstanceName))
