Removed some code related to Nanite Coarse Mesh Streaming
- don't modify DDC key based on cvar since it doesn't impact LOD generation since CL 18820056.
- removed bGenerateCoarseMeshStreamingLODs parameter from a few functions since it isn't used anymore.

#prelight 652516f53a59af0b367fbd0b

diff --git a/Engine/Source/Developer/MeshBuilder/Private/MeshBuilderModule.cpp b/Engine/Source/Developer/MeshBuilder/Private/MeshBuilderModule.cpp
--- a/Engine/Source/Developer/MeshBuilder/Private/MeshBuilderModule.cpp
+++ b/Engine/Source/Developer/MeshBuilder/Private/MeshBuilderModule.cpp
@@ -26,7 +26,7 @@ public:
 		// Unregister any modular features here
 	}
 
-	virtual bool BuildMesh(FStaticMeshRenderData& OutRenderData, UObject* Mesh, const FStaticMeshLODGroup& LODGroup, bool bGenerateCoarseMeshStreamingLODs, bool bAllowNanite) override;
+	virtual bool BuildMesh(FStaticMeshRenderData& OutRenderData, UObject* Mesh, const FStaticMeshLODGroup& LODGroup, bool bAllowNanite) override;
 
 	virtual bool BuildMeshVertexPositions(
 		UObject* StaticMesh,
@@ -42,13 +42,13 @@ private:
 
 IMPLEMENT_MODULE(FMeshBuilderModule, MeshBuilder );
 
-bool FMeshBuilderModule::BuildMesh(FStaticMeshRenderData& OutRenderData, class UObject* Mesh, const FStaticMeshLODGroup& LODGroup, bool bGenerateCoarseMeshStreamingLODs, bool bAllowNanite)
+bool FMeshBuilderModule::BuildMesh(FStaticMeshRenderData& OutRenderData, class UObject* Mesh, const FStaticMeshLODGroup& LODGroup, bool bAllowNanite)
 {
 	UStaticMesh* StaticMesh = Cast<UStaticMesh>(Mesh);
 	if (StaticMesh != nullptr)
 	{
 		//Call the static mesh builder
-		return FStaticMeshBuilder().Build(OutRenderData, StaticMesh, LODGroup, bGenerateCoarseMeshStreamingLODs, bAllowNanite);
+		return FStaticMeshBuilder().Build(OutRenderData, StaticMesh, LODGroup, bAllowNanite);
 	}
 	return false;
 }
diff --git a/Engine/Source/Developer/MeshBuilder/Private/StaticMeshBuilder.cpp b/Engine/Source/Developer/MeshBuilder/Private/StaticMeshBuilder.cpp
--- a/Engine/Source/Developer/MeshBuilder/Private/StaticMeshBuilder.cpp
+++ b/Engine/Source/Developer/MeshBuilder/Private/StaticMeshBuilder.cpp
@@ -327,7 +327,7 @@ static bool BuildNanite(
 }
 
 
-bool FStaticMeshBuilder::Build(FStaticMeshRenderData& StaticMeshRenderData, UStaticMesh* StaticMesh, const FStaticMeshLODGroup& LODGroup, bool bGenerateCoarseMeshStreamingLODs, bool bTargetSupportsNanite)
+bool FStaticMeshBuilder::Build(FStaticMeshRenderData& StaticMeshRenderData, UStaticMesh* StaticMesh, const FStaticMeshLODGroup& LODGroup, bool bTargetSupportsNanite)
 {
 	const bool bNaniteBuildEnabled = StaticMesh->IsNaniteEnabled();
 	const bool bHaveHiResSourceModel = StaticMesh->IsHiResMeshDescriptionValid();
diff --git a/Engine/Source/Developer/MeshBuilder/Public/IMeshBuilderModule.h b/Engine/Source/Developer/MeshBuilder/Public/IMeshBuilderModule.h
--- a/Engine/Source/Developer/MeshBuilder/Public/IMeshBuilderModule.h
+++ b/Engine/Source/Developer/MeshBuilder/Public/IMeshBuilderModule.h
@@ -29,7 +29,7 @@ public:
 
 	virtual void AppendToDDCKey(FString& DDCKey, bool bSkeletal) { }
 
-	virtual bool BuildMesh(class FStaticMeshRenderData& OutRenderData, class UObject* Mesh, const class FStaticMeshLODGroup& LODGroup, bool bGenerateCoarseMeshStreamingLODs, bool bAllowNanite) = 0;
+	virtual bool BuildMesh(class FStaticMeshRenderData& OutRenderData, class UObject* Mesh, const class FStaticMeshLODGroup& LODGroup, bool bAllowNanite) = 0;
 
 	virtual bool BuildMeshVertexPositions(class UObject* StaticMesh, TArray<uint32>& Indices, TArray<FVector3f>& Vertices, FStaticMeshSectionArray& Sections) = 0;
 
diff --git a/Engine/Source/Developer/MeshBuilder/Public/MeshBuilder.h b/Engine/Source/Developer/MeshBuilder/Public/MeshBuilder.h
--- a/Engine/Source/Developer/MeshBuilder/Public/MeshBuilder.h
+++ b/Engine/Source/Developer/MeshBuilder/Public/MeshBuilder.h
@@ -32,7 +32,6 @@ public:
 		FStaticMeshRenderData& OutRenderData,
 		UStaticMesh* StaticMesh,
 		const FStaticMeshLODGroup& LODGroup,
-		bool bGenerateCoarseMeshStreamingLODs,
 		bool bAllowNanite) = 0;
 
 	virtual bool BuildMeshVertexPositions(
diff --git a/Engine/Source/Developer/MeshBuilder/Public/SkeletalMeshBuilder.h b/Engine/Source/Developer/MeshBuilder/Public/SkeletalMeshBuilder.h
--- a/Engine/Source/Developer/MeshBuilder/Public/SkeletalMeshBuilder.h
+++ b/Engine/Source/Developer/MeshBuilder/Public/SkeletalMeshBuilder.h
@@ -18,7 +18,7 @@ public:
 	FSkeletalMeshBuilder();
 
 	//No support for static mesh build in this class
-	virtual bool Build(FStaticMeshRenderData& OutRenderData, UStaticMesh* StaticMesh, const FStaticMeshLODGroup& LODGroup, bool bGenerateCoarseMeshStreamingLODs, bool bAllowNanite) override
+	virtual bool Build(FStaticMeshRenderData& OutRenderData, UStaticMesh* StaticMesh, const FStaticMeshLODGroup& LODGroup, bool bAllowNanite) override
 	{
 		bool No_Support_For_StaticMesh_Build_In_FSkeletalMeshBuilder_Class = false;
 		check(No_Support_For_StaticMesh_Build_In_FSkeletalMeshBuilder_Class);
diff --git a/Engine/Source/Developer/MeshBuilder/Public/StaticMeshBuilder.h b/Engine/Source/Developer/MeshBuilder/Public/StaticMeshBuilder.h
--- a/Engine/Source/Developer/MeshBuilder/Public/StaticMeshBuilder.h
+++ b/Engine/Source/Developer/MeshBuilder/Public/StaticMeshBuilder.h
@@ -26,7 +26,6 @@ public:
 		FStaticMeshRenderData& OutRenderData,
 		UStaticMesh* StaticMesh,
 		const FStaticMeshLODGroup& LODGroup,
-		bool bGenerateCoarseMeshStreamingLODs,
 		bool bAllowNanite) override;
 
 	//No support for skeletal mesh build in this class
diff --git a/Engine/Source/Runtime/Engine/Private/StaticMesh.cpp b/Engine/Source/Runtime/Engine/Private/StaticMesh.cpp
--- a/Engine/Source/Runtime/Engine/Private/StaticMesh.cpp
+++ b/Engine/Source/Runtime/Engine/Private/StaticMesh.cpp
@@ -2786,14 +2786,6 @@ static FString BuildStaticMeshDerivedDataKeySuffix(const ITargetPlatform* Target
 		KeySuffix += TEXT("zzzzzzzz");
 	}
 
-	// Nanite Coarse mesh streaming change require regenerating the coarse mesh LODs during caching
-	static auto* VarNaniteCoarseMeshStreaming = IConsoleManager::Get().FindConsoleVariable(TEXT("r.Nanite.CoarseMeshStreaming"));
-	const bool bNaniteCoareMeshStreamingEnabled = !VarNaniteCoarseMeshStreaming || VarNaniteCoarseMeshStreaming->GetInt() != 0;
-	if (bNaniteCoareMeshStreamingEnabled)
-	{
-		KeySuffix += TEXT("_NCMS");
-	}
-
 	KeySuffix.AppendChar(Mesh->bSupportUniformlyDistributedSampling ? TEXT('1') : TEXT('0'));
 
 	if (TargetPlatform->SupportsFeature(ETargetPlatformFeatures::HardwareLZDecompression))
@@ -3103,15 +3095,10 @@ void FStaticMeshRenderData::Cache(const ITargetPlatform* TargetPlatform, UStatic
 			{
 				IMeshBuilderModule& MeshBuilderModule = IMeshBuilderModule::GetForPlatform(TargetPlatform);
 
-				// Only build the LODs if coarse mesh streaming is enabled
-				static auto* VarNaniteCoarseMeshStreaming = IConsoleManager::Get().FindConsoleVariable(TEXT("r.Nanite.CoarseMeshStreaming"));
-				const bool bNaniteCoarseMeshStreamingEnabled = !VarNaniteCoarseMeshStreaming || VarNaniteCoarseMeshStreaming->GetInt() != 0;
-				const bool bGenerateStreamingLODs = TargetPlatform->SupportsFeature(ETargetPlatformFeatures::MeshLODStreaming) && bNaniteCoarseMeshStreamingEnabled && LODGroup.IsLODStreamingSupported();
-
 				// Check if the target platform supports Nanite at all
 				const bool bAllowNanite = DoesTargetPlatformSupportNanite(TargetPlatform);
 
-				if (!MeshBuilderModule.BuildMesh(*this, Owner, LODGroup, bGenerateStreamingLODs, bAllowNanite))
+				if (!MeshBuilderModule.BuildMesh(*this, Owner, LODGroup, bAllowNanite))
 				{
 					UE_LOG(LogStaticMesh, Error, TEXT("Failed to build static mesh. See previous line(s) for details."));
 					return;
