Band-aid fix to prevent checks & crashes when the instance manager has not been informed of the actual number of instances represented in the ISM

diff --git a/Engine/Source/Runtime/Engine/Private/InstancedStaticMesh/ISMInstanceDataManager.cpp b/Engine/Source/Runtime/Engine/Private/InstancedStaticMesh/ISMInstanceDataManager.cpp
--- a/Engine/Source/Runtime/Engine/Private/InstancedStaticMesh/ISMInstanceDataManager.cpp
+++ b/Engine/Source/Runtime/Engine/Private/InstancedStaticMesh/ISMInstanceDataManager.cpp
@@ -598,6 +598,15 @@ bool FPrimitiveInstanceDataManager::FlushChanges(FInstanceUpdateComponentDesc &&
 	EShaderPlatform ShaderPlatform = ComponentData.PrimitiveSceneProxy->GetScene().GetShaderPlatform();
 	ERHIFeatureLevel::Type FeatureLevel = ComponentData.PrimitiveSceneProxy->GetScene().GetFeatureLevel();
 	check(Proxy->CheckPlatformFeatureLevel(ShaderPlatform, FeatureLevel));
+	
+	// BandAid: This is the first flush & we have not been informed correctly about the number of instances in the ISM so we need to patch that up here and now.
+	if (bFirstFlush && GetMaxInstanceIndex() == 0 && ComponentData.NumSourceInstances != 0)
+	{
+		IdToIndexMap.Reset();
+		IndexToIdMap.Reset();
+		NumInstances = ComponentData.NumSourceInstances;
+	}
+	bFirstFlush = false;
 
 
 	bool bWasUpdateQueued = false;
diff --git a/Engine/Source/Runtime/Engine/Public/InstancedStaticMesh/ISMInstanceDataManager.h b/Engine/Source/Runtime/Engine/Public/InstancedStaticMesh/ISMInstanceDataManager.h
--- a/Engine/Source/Runtime/Engine/Public/InstancedStaticMesh/ISMInstanceDataManager.h
+++ b/Engine/Source/Runtime/Engine/Public/InstancedStaticMesh/ISMInstanceDataManager.h
@@ -210,4 +210,5 @@ private:
 	int32 NumCustomDataFloats = 0;
 	float AbsMaxDisplacement = 0.0f;
 	FRenderBounds StaticMeshBounds;
+	bool bFirstFlush = true;
 };
