WorldPartition: Added Blueprint function UDataLayerManager::GetDataLayerInstances and exposed to BP UDataLayerInstance GetDataLayerShortName, GetDataLayerFullName and GetAsset.

diff --git a/Engine/Source/Runtime/Engine/Private/WorldPartition/DataLayer/DataLayerManager.cpp b/Engine/Source/Runtime/Engine/Private/WorldPartition/DataLayer/DataLayerManager.cpp
--- a/Engine/Source/Runtime/Engine/Private/WorldPartition/DataLayer/DataLayerManager.cpp
+++ b/Engine/Source/Runtime/Engine/Private/WorldPartition/DataLayer/DataLayerManager.cpp
@@ -268,6 +268,17 @@ AWorldDataLayers* UDataLayerManager::GetWorldDataLayers() const
 	return GetTypedOuter<UWorld>()->GetWorldDataLayers();
 }
 
+TArray<UDataLayerInstance*> UDataLayerManager::GetDataLayerInstances() const
+{
+	TArray<UDataLayerInstance*> Result;
+	ForEachDataLayerInstance([&Result](UDataLayerInstance* DataLayerInstance)
+	{
+		Result.Add(DataLayerInstance);
+		return true;
+	});
+	return Result;
+}
+
 const UDataLayerInstance* UDataLayerManager::GetDataLayerInstanceFromAsset(const UDataLayerAsset* InDataLayerAsset) const
 {
 	return GetDataLayerInstance(InDataLayerAsset);
diff --git a/Engine/Source/Runtime/Engine/Public/WorldPartition/DataLayer/DataLayerInstance.h b/Engine/Source/Runtime/Engine/Public/WorldPartition/DataLayer/DataLayerInstance.h
--- a/Engine/Source/Runtime/Engine/Public/WorldPartition/DataLayer/DataLayerInstance.h
+++ b/Engine/Source/Runtime/Engine/Public/WorldPartition/DataLayer/DataLayerInstance.h
@@ -161,7 +161,10 @@ public:
 	UFUNCTION(Category = "Data Layer|Runtime", BlueprintCallable)
 	EDataLayerRuntimeState GetInitialRuntimeState() const { return IsRuntime() && !IsClientOnly() && !IsServerOnly() ? InitialRuntimeState : EDataLayerRuntimeState::Unloaded; }
 
+	UFUNCTION(Category = "Data Layer", BlueprintCallable)
 	virtual FString GetDataLayerShortName() const { return TEXT("Invalid Data Layer"); }
+
+	UFUNCTION(Category = "Data Layer", BlueprintCallable)
 	virtual FString GetDataLayerFullName() const { return TEXT("Invalid Data Layer"); }
 
 	virtual bool CanHaveChildDataLayers() const { return true; }
@@ -173,6 +176,7 @@ public:
 	ENGINE_API EDataLayerRuntimeState GetRuntimeState() const;
 	ENGINE_API EDataLayerRuntimeState GetEffectiveRuntimeState() const;
 
+	UFUNCTION(Category = "Data Layer", BlueprintCallable)
 	virtual const UDataLayerAsset* GetAsset() const { return nullptr; }
 
 	const TArray<TObjectPtr<UDataLayerInstance>>& GetChildren() const { return Children; }
diff --git a/Engine/Source/Runtime/Engine/Public/WorldPartition/DataLayer/DataLayerManager.h b/Engine/Source/Runtime/Engine/Public/WorldPartition/DataLayer/DataLayerManager.h
--- a/Engine/Source/Runtime/Engine/Public/WorldPartition/DataLayer/DataLayerManager.h
+++ b/Engine/Source/Runtime/Engine/Public/WorldPartition/DataLayer/DataLayerManager.h
@@ -53,6 +53,9 @@ public:
 	}
 
 	//~ Begin Blueprint interface
+	UFUNCTION(BlueprintCallable, Category = DataLayers)
+	ENGINE_API TArray<UDataLayerInstance*> GetDataLayerInstances() const;
+
 	UFUNCTION(BlueprintCallable, Category = DataLayers)
 	ENGINE_API const UDataLayerInstance* GetDataLayerInstanceFromAsset(const UDataLayerAsset* InDataLayerAsset) const;
 
