Fix crash in UpdateOwningConnection due to null children entries
Fix buffer overflow in UpdateOwningConnection if too many children are present.

#rnx

diff --git a/Engine/Source/Runtime/Engine/Classes/GameFramework/Actor.h b/Engine/Source/Runtime/Engine/Classes/GameFramework/Actor.h
--- a/Engine/Source/Runtime/Engine/Classes/GameFramework/Actor.h
+++ b/Engine/Source/Runtime/Engine/Classes/GameFramework/Actor.h
@@ -3221,7 +3221,7 @@ protected:
 	 * Pushes the owning NetConnection for the actor and all of its children to the replication system.
 	 * This information decides whether properties with owner conditionals are replicated or not.
 	 */
-	ENGINE_API void UpdateOwningNetConnection() const;
+	ENGINE_API void UpdateOwningNetConnection();
 
 	/**
 	 * Helper to BeginReplication passing on additional parameters to the ReplicationSystem, typically called from code overriding normal BeginReplication()
diff --git a/Engine/Source/Runtime/Engine/Private/ActorReplication.cpp b/Engine/Source/Runtime/Engine/Private/ActorReplication.cpp
--- a/Engine/Source/Runtime/Engine/Private/ActorReplication.cpp
+++ b/Engine/Source/Runtime/Engine/Private/ActorReplication.cpp
@@ -1038,7 +1038,7 @@ void AActor::EndReplication(EEndPlayReason::Type EndPlayReason)
 	UE::Net::FReplicationSystemUtil::EndReplication(this, EndPlayReason);
 }
 
-void AActor::UpdateOwningNetConnection() const
+void AActor::UpdateOwningNetConnection()
 {
 	using namespace UE::Net;
 
@@ -1068,34 +1068,33 @@ void AActor::UpdateOwningNetConnection() const
 		return;
 	}
 
-	constexpr SIZE_T MaxActorCount = 512;
-	const AActor* Actors[MaxActorCount];
-	SIZE_T ActorCount = 1;
-	Actors[0] = this;
-	for (; ActorCount > 0; )
+	TArray<AActor*, TInlineAllocator<32>> EveryChildren;
+	// Include ourself
+	EveryChildren.Push(this);
+
+	do 
 	{
-		const AActor* Actor = Actors[--ActorCount];
-		check(ActorCount + Actor->Children.Num() <= UE_ARRAY_COUNT(Actors));
-		for (const AActor* Child : MakeArrayView(Actor->Children))
+		constexpr bool bNoShrinking = false;
+		if (AActor* Actor = EveryChildren.Pop(bNoShrinking))
 		{
-			Actors[ActorCount++] = Child;
-		}
+			EveryChildren.Append(Actor->Children);
 
-		if (Actor->GetIsReplicated())
-		{
-			const UE::Net::FNetRefHandle RefHandle = ObjectReplicationBridge->GetReplicatedRefHandle(Actor);
-			if (RefHandle.IsValid())
+			if (Actor->GetIsReplicated())
 			{
-				ReplicationSystem->SetOwningNetConnection(RefHandle, NewOwningNetConnectionId);
-				// Update autonomous proxy condition
-				if (Actor->GetRemoteRole() == ROLE_AutonomousProxy)
+				const UE::Net::FNetRefHandle RefHandle = ObjectReplicationBridge->GetReplicatedRefHandle(Actor);
+				if (RefHandle.IsValid())
 				{
-					const bool bEnableAutonomousCondition = true;
-					ReplicationSystem->SetReplicationConditionConnectionFilter(RefHandle, EReplicationCondition::RoleAutonomous, NewOwningNetConnectionId, bEnableAutonomousCondition);
+					ReplicationSystem->SetOwningNetConnection(RefHandle, NewOwningNetConnectionId);
+					// Update autonomous proxy condition
+					if (Actor->GetRemoteRole() == ROLE_AutonomousProxy)
+					{
+						const bool bEnableAutonomousCondition = true;
+						ReplicationSystem->SetReplicationConditionConnectionFilter(RefHandle, EReplicationCondition::RoleAutonomous, NewOwningNetConnectionId, bEnableAutonomousCondition);
+					}
 				}
 			}
 		}
-	}
+	} while (EveryChildren.Num());
 }
 #endif // UE_WITH_IRIS
 
