Only update navigation for actor instance when initially populating the ISMC to prevent unnecessary dirtiness.
Redoing changes made at CL 28040350 but wiped out in refactor at CL 28371211.

diff --git a/Engine/Plugins/Runtime/MassGameplay/Source/MassRepresentation/Private/MassVisualizationComponent.cpp b/Engine/Plugins/Runtime/MassGameplay/Source/MassRepresentation/Private/MassVisualizationComponent.cpp
--- a/Engine/Plugins/Runtime/MassGameplay/Source/MassRepresentation/Private/MassVisualizationComponent.cpp
+++ b/Engine/Plugins/Runtime/MassGameplay/Source/MassRepresentation/Private/MassVisualizationComponent.cpp
@@ -352,7 +352,7 @@ void UMassVisualizationComponent::HandleChangesWithExternalIDTracking(UInstanced
 
 	// removing instances first, since this operation is more resilient to duplicates. Plus we make an arbitrary decision 
 	// that it's better to have redundant things visible than not seeing required things
-	ProcessRemoves(ISMComponent, SharedData);
+	ProcessRemoves(ISMComponent, SharedData, /*bUpdateNavigation=*/false);
 
 	// NOTE: This code path is designed to only perform Adds, never updates so updates are filtered out along with duplicates.
 	TArray<int32>& MassInstanceIds = SharedData.UpdateInstanceIds;
@@ -394,13 +394,14 @@ void UMassVisualizationComponent::HandleChangesWithExternalIDTracking(UInstanced
 			// if these are the first entities we're adding we need to set NumCustomDataFloats so that the PerInstanceSMCustomData
 			// gets populated properly by the AddInstancesInternal call below
 			const int32 StartingCount = ISMComponent.GetNumInstances();
+			const bool bInitiallyEmpty = (StartingCount == 0); 
 			if (StartingCount == 0 && CustomFloatData.Num() && ISMComponent.Mobility != EComponentMobility::Static)
 			{
 				ISMComponent.SetNumCustomDataFloats(InNumCustomDataFloats);
 			}
 
 			check(MassInstanceIds.Num() == InstanceTransforms.Num());
-			TArray<FPrimitiveInstanceId> NewIds = ISMComponent.AddInstancesById(InstanceTransforms, true);
+			TArray<FPrimitiveInstanceId> NewIds = ISMComponent.AddInstancesById(InstanceTransforms, /*bWorldSpace=*/true, /*bUpdateNavigation =*/bInitiallyEmpty);
 			check(MassInstanceIds.Num() == NewIds.Num());
 			for (int32 i = 0; i < MassInstanceIds.Num(); ++i)
 			{
@@ -417,7 +418,7 @@ void UMassVisualizationComponent::HandleChangesWithExternalIDTracking(UInstanced
 }
 
 
-void UMassVisualizationComponent::ProcessRemoves(UInstancedStaticMeshComponent& ISMComponent, FMassISMCSharedData& SharedData)
+void UMassVisualizationComponent::ProcessRemoves(UInstancedStaticMeshComponent& ISMComponent, FMassISMCSharedData& SharedData, const bool bUpdateNavigation /*= true*/)
 {
 	if (!SharedData.GetRemoveInstanceIds().IsEmpty())
 	{
@@ -443,7 +444,7 @@ void UMassVisualizationComponent::ProcessRemoves(UInstancedStaticMeshComponent&
 			}
 		}
 
-		ISMComponent.RemoveInstancesById(ISMInstanceIds);
+		ISMComponent.RemoveInstancesById(ISMInstanceIds, bUpdateNavigation);
 	}
 }
 
diff --git a/Engine/Plugins/Runtime/MassGameplay/Source/MassRepresentation/Public/MassVisualizationComponent.h b/Engine/Plugins/Runtime/MassGameplay/Source/MassRepresentation/Public/MassVisualizationComponent.h
--- a/Engine/Plugins/Runtime/MassGameplay/Source/MassRepresentation/Public/MassVisualizationComponent.h
+++ b/Engine/Plugins/Runtime/MassGameplay/Source/MassRepresentation/Public/MassVisualizationComponent.h
@@ -80,7 +80,7 @@ protected:
 	/**
 	 * Process all removed IDs in FMassISMCSharedData and apply to the ISM component.
 	 */
-	void ProcessRemoves(UInstancedStaticMeshComponent& ISMComponent, FMassISMCSharedData& SharedData);
+	void ProcessRemoves(UInstancedStaticMeshComponent& ISMComponent, FMassISMCSharedData& SharedData, bool bUpdateNavigation = true);
 	
 	/** 
 	 * Applies changes accumulated in SharedData while manually updating the Instance ID mapping. This approach is done in preparation 
diff --git a/Engine/Source/Runtime/Engine/Classes/Components/InstancedStaticMeshComponent.h b/Engine/Source/Runtime/Engine/Classes/Components/InstancedStaticMeshComponent.h
--- a/Engine/Source/Runtime/Engine/Classes/Components/InstancedStaticMeshComponent.h
+++ b/Engine/Source/Runtime/Engine/Classes/Components/InstancedStaticMeshComponent.h
@@ -184,7 +184,7 @@ class UInstancedStaticMeshComponent : public UStaticMeshComponent, public ISMIns
 
 	/**
 	 */
-	ENGINE_API TArray<FPrimitiveInstanceId> AddInstancesById(const TArrayView<const FTransform>& InstanceTransforms, bool bWorldSpace = false);
+	ENGINE_API TArray<FPrimitiveInstanceId> AddInstancesById(const TArrayView<const FTransform>& InstanceTransforms, bool bWorldSpace = false, bool bUpdateNavigation = true);
 	ENGINE_API FPrimitiveInstanceId AddInstanceById(const FTransform& InstanceTransforms, bool bWorldSpace = false);
 	/**
 	 */
@@ -192,7 +192,7 @@ class UInstancedStaticMeshComponent : public UStaticMeshComponent, public ISMIns
 	inline void SetCustomDataById(FPrimitiveInstanceId InstanceId, TArrayView<const float> CustomDataFloats) { SetCustomDataById(MakeArrayView(&InstanceId, 1), CustomDataFloats); }
 	/**
 	 */
-	ENGINE_API virtual void RemoveInstancesById(const TArrayView<const FPrimitiveInstanceId> &InstanceIds);
+	ENGINE_API virtual void RemoveInstancesById(const TArrayView<const FPrimitiveInstanceId> &InstanceIds, bool bUpdateNavigation = true);
 	inline void RemoveInstanceById(FPrimitiveInstanceId InstanceId) { RemoveInstancesById(MakeArrayView(&InstanceId, 1)); }
 	/**
 	 */
@@ -592,7 +592,7 @@ protected:
 	ENGINE_API TArray<int32> AddInstancesInternal(TConstArrayView<FTransform> InstanceTransforms, bool bShouldReturnIndices, bool bWorldSpace, bool bUpdateNavigation = true);
 
 	/** Internal version of RemoveInstance */	
-	ENGINE_API bool RemoveInstanceInternal(int32 InstanceIndex, bool InstanceAlreadyRemoved, bool bForceRemoveAtSwap = false);
+	ENGINE_API bool RemoveInstanceInternal(int32 InstanceIndex, bool InstanceAlreadyRemoved, bool bForceRemoveAtSwap = false, bool bUpdateNavigation = true);
 
 	/**
 	 * Returns the bounds of a single instance in local space. It uses the NavCollision if available,
diff --git a/Engine/Source/Runtime/Engine/Private/InstancedStaticMesh.cpp b/Engine/Source/Runtime/Engine/Private/InstancedStaticMesh.cpp
--- a/Engine/Source/Runtime/Engine/Private/InstancedStaticMesh.cpp
+++ b/Engine/Source/Runtime/Engine/Private/InstancedStaticMesh.cpp
@@ -3163,11 +3163,11 @@ TArray<int32> UInstancedStaticMeshComponent::AddInstances(const TArray<FTransfor
 	return AddInstancesInternal(InstanceTransforms, bShouldReturnIndices, bWorldSpace, bUpdateNavigation);
 }
 
-TArray<FPrimitiveInstanceId> UInstancedStaticMeshComponent::AddInstancesById(const TArrayView<const FTransform>& InstanceTransforms, bool bWorldSpace)
+TArray<FPrimitiveInstanceId> UInstancedStaticMeshComponent::AddInstancesById(const TArrayView<const FTransform>& InstanceTransforms, bool bWorldSpace, const bool bUpdateNavigation)
 {
 	check(PrimitiveInstanceDataManager.GetMode() != FPrimitiveInstanceDataManager::EMode::ExternalLegacyData);
 	TArray<FPrimitiveInstanceId> Ids;
-	TArray<int32> Indices = AddInstancesInternal(InstanceTransforms, true, bWorldSpace);
+	TArray<int32> Indices = AddInstancesInternal(InstanceTransforms, true, bWorldSpace, bUpdateNavigation);
 	Algo::Transform(Indices, Ids, [&](int32 Index) { return FPrimitiveInstanceId{PrimitiveInstanceDataManager.IndexToId(Index) }; } );
 	return Ids;
 }
@@ -3190,12 +3190,12 @@ void UInstancedStaticMeshComponent::SetCustomDataById(const TArrayView<const FPr
 	}
 }
 
-void UInstancedStaticMeshComponent::RemoveInstancesById(const TArrayView<const FPrimitiveInstanceId>& InstanceIds)
+void UInstancedStaticMeshComponent::RemoveInstancesById(const TArrayView<const FPrimitiveInstanceId>& InstanceIds, const bool bUpdateNavigation)
 {
 	check(PrimitiveInstanceDataManager.GetMode() != FPrimitiveInstanceDataManager::EMode::ExternalLegacyData);
 	for (FPrimitiveInstanceId Id : InstanceIds)
 	{
-		RemoveInstanceInternal(PrimitiveInstanceDataManager.IdToIndex(Id), false, true);
+		RemoveInstanceInternal(PrimitiveInstanceDataManager.IdToIndex(Id), false, true, bUpdateNavigation);
 	}
 }
 
@@ -3308,7 +3308,7 @@ bool UInstancedStaticMeshComponent::SupportsRemoveSwap() const
 	return bSupportRemoveAtSwap || CVarISMForceRemoveAtSwap.GetValueOnGameThread() != 0;
 }
 
-bool UInstancedStaticMeshComponent::RemoveInstanceInternal(int32 InstanceIndex, bool InstanceAlreadyRemoved, bool bForceRemoveAtSwap)
+bool UInstancedStaticMeshComponent::RemoveInstanceInternal(int32 InstanceIndex, bool InstanceAlreadyRemoved, bool bForceRemoveAtSwap, const bool bUpdateNavigation)
 {
 #if WITH_EDITOR
 	DeletionState = InstanceAlreadyRemoved ? EInstanceDeletionReason::EntryAlreadyRemoved : EInstanceDeletionReason::EntryRemoval;
@@ -3330,9 +3330,11 @@ bool UInstancedStaticMeshComponent::RemoveInstanceInternal(int32 InstanceIndex,
 	if (!InstanceAlreadyRemoved && PerInstanceSMData.IsValidIndex(InstanceIndex))
 	{
 		const bool bWasNavRelevant = bNavigationRelevant;
-		
-		// Request navigation update
-		PartialNavigationUpdate(InstanceIndex);
+
+		if (bUpdateNavigation)
+		{
+			PartialNavigationUpdate(InstanceIndex);
+		}
 
 		if (bUseRemoveAtSwap)
 		{
