Back out CL 28506784 due to test failures, meshes rendering all gray.

diff --git a/Engine/Source/Runtime/Engine/Private/GameViewportClient.cpp b/Engine/Source/Runtime/Engine/Private/GameViewportClient.cpp
--- a/Engine/Source/Runtime/Engine/Private/GameViewportClient.cpp
+++ b/Engine/Source/Runtime/Engine/Private/GameViewportClient.cpp
@@ -1707,9 +1707,6 @@ void UGameViewportClient::Draw(FViewport* InViewport, FCanvas* SceneCanvas)
 
 	FinalizeViews(&ViewFamily, PlayerViewMap);
 
-	// Update engine flags stored by viewport.
-	EngineShowFlags = ViewFamily.EngineShowFlags;
-
 	// Update level streaming.
 	MyWorld->UpdateLevelStreaming();
 
