Restore line accidentally removed in a previous submit.

diff --git a/Engine/Source/Runtime/Engine/Private/InstancedStaticMesh.cpp b/Engine/Source/Runtime/Engine/Private/InstancedStaticMesh.cpp
--- a/Engine/Source/Runtime/Engine/Private/InstancedStaticMesh.cpp
+++ b/Engine/Source/Runtime/Engine/Private/InstancedStaticMesh.cpp
@@ -2197,6 +2197,8 @@ FPrimitiveSceneProxy* UInstancedStaticMeshComponent::CreateSceneProxy()
 {
 	ProxySize = 0;
 
+	PrimitiveInstanceDataManager.ResetComponentDirtyTracking();
+
 	// Verify that both mesh and instance data is valid before using it.
 	const bool bIsMeshAndInstanceDataValid =
 #if WITH_EDITOR
