Fix Fill Collision Underneath for Navmesh not being applied to ISM isntances.

diff --git a/Engine/Source/Runtime/Engine/Private/InstancedStaticMesh.cpp b/Engine/Source/Runtime/Engine/Private/InstancedStaticMesh.cpp
--- a/Engine/Source/Runtime/Engine/Private/InstancedStaticMesh.cpp
+++ b/Engine/Source/Runtime/Engine/Private/InstancedStaticMesh.cpp
@@ -4855,6 +4855,13 @@ bool UInstancedStaticMeshComponent::ComponentOverlapMultiImpl(TArray<struct FOve
 
 void UInstancedStaticMeshComponent::GetNavigationData(FNavigationRelevantData& Data) const
 {
+	if (bFillCollisionUnderneathForNavmesh)
+	{
+		FCompositeNavModifier CompositeNavModifier;
+		CompositeNavModifier.SetFillCollisionUnderneathForNavmesh(bFillCollisionUnderneathForNavmesh);
+		Data.Modifiers.Add(CompositeNavModifier);
+	}
+
 	// Navigation data will get refreshed once async compilation finishes
 	if (GetStaticMesh() && !GetStaticMesh()->IsCompiling() && GetStaticMesh()->GetNavCollision())
 	{
