Updated comments about WorldPartition initialization and WorldPartitionSubsystem::DoesSupportWorldType.

diff --git a/Engine/Source/Runtime/Engine/Private/Level.cpp b/Engine/Source/Runtime/Engine/Private/Level.cpp
--- a/Engine/Source/Runtime/Engine/Private/Level.cpp
+++ b/Engine/Source/Runtime/Engine/Private/Level.cpp
@@ -3009,7 +3009,7 @@ void ULevel::OnLevelLoaded()
 	}
 
 	// 1. Cook commandlet does it's own UWorldPartition::Initialize call in FWorldPartitionCookPackageSplitter::GetGenerateList
-	// 2. Do not Initialize if World doesn't have a UWorldPartitionSubsystem (Known case is when WorldType == EWorldType::Inactive)
+	// 2. Do not Initialize if World doesn't have a UWorldPartitionSubsystem
 	if (!IsRunningCookCommandlet() && OwningWorld->HasSubsystem<UWorldPartitionSubsystem>())
 	{
 		if (UWorldPartition* WorldPartition = GetWorldPartition())
diff --git a/Engine/Source/Runtime/Engine/Private/WorldPartition/WorldPartitionSubsystem.cpp b/Engine/Source/Runtime/Engine/Private/WorldPartition/WorldPartitionSubsystem.cpp
--- a/Engine/Source/Runtime/Engine/Private/WorldPartition/WorldPartitionSubsystem.cpp
+++ b/Engine/Source/Runtime/Engine/Private/WorldPartition/WorldPartitionSubsystem.cpp
@@ -519,8 +519,7 @@ void UWorldPartitionSubsystem::Deinitialize()
 	Super::Deinitialize();
 }
 
-// We allow creating UWorldPartitionSubsystem for inactive worlds as WorldPartition initialization is necessary 
-// because DataLayerManager is required to be initialized when duplicating a partitioned world.
+// WorldPartitionSubsystem is required for WorldPartition to be properly initialized (even when world is Inactive like during Cook or world duplication).
 bool UWorldPartitionSubsystem::DoesSupportWorldType(const EWorldType::Type WorldType) const
 {
 	return Super::DoesSupportWorldType(WorldType) || WorldType == EWorldType::Inactive || WorldType == EWorldType::EditorPreview;
