Iris
* Added ability to prefer a Replication System via the GameMode class.
* When a game mode selects a specific replication system, the server will create its GameNetDriver using that mode if compatible.
* Note that the clients are NOT able to ask the game mode for the proper repsystem, so use this only when you can correctly predict the proper mode before connecting on the server.  Otherwise the clients may disconnect due to incompatible netdrivers.
     * PIE clients started inside a single process will use the correct repsystem of the dedicated server.
* Setting -UseIrisReplication=0/1 on the cmdline will override the GameMode's requested system.

diff --git a/Engine/Source/Runtime/Core/Private/Misc/NetworkVersion.cpp b/Engine/Source/Runtime/Core/Private/Misc/NetworkVersion.cpp
--- a/Engine/Source/Runtime/Core/Private/Misc/NetworkVersion.cpp
+++ b/Engine/Source/Runtime/Core/Private/Misc/NetworkVersion.cpp
@@ -337,6 +337,6 @@ void FNetworkVersion::DescribeNetworkRuntimeFeaturesBitset(EEngineNetworkRuntime
 	}
 	else
 	{
-		OutVerboseDescription.Append(TEXT("LegacyReplication"));
+		OutVerboseDescription.Append(TEXT("GenericReplication"));
 	}
 }
\ No newline at end of file
diff --git a/Engine/Source/Runtime/Core/Public/Misc/NetworkVersion.h b/Engine/Source/Runtime/Core/Public/Misc/NetworkVersion.h
--- a/Engine/Source/Runtime/Core/Public/Misc/NetworkVersion.h
+++ b/Engine/Source/Runtime/Core/Public/Misc/NetworkVersion.h
@@ -44,7 +44,7 @@ public:
 enum class EEngineNetworkRuntimeFeatures : uint16
 {
 	None = 0,
-	IrisEnabled = 1 << None, // Are we running the Iris or Legacy network system
+	IrisEnabled = 1 << None, // Are we running the Iris or Generic replication system
 };
 ENUM_CLASS_FLAGS(EEngineNetworkRuntimeFeatures);
 
diff --git a/Engine/Source/Runtime/Engine/Classes/GameFramework/GameModeBase.h b/Engine/Source/Runtime/Engine/Classes/GameFramework/GameModeBase.h
--- a/Engine/Source/Runtime/Engine/Classes/GameFramework/GameModeBase.h
+++ b/Engine/Source/Runtime/Engine/Classes/GameFramework/GameModeBase.h
@@ -11,6 +11,7 @@
 #include "GameFramework/Info.h"
 #include "Engine/ServerStatReplicator.h"
 #include "Online/CoreOnline.h"
+#include "Net/Core/Connection/NetEnums.h"
 #include "GameFramework/PlayerController.h"
 #include "GameModeBase.generated.h"
 
@@ -117,6 +118,9 @@ public:
 	UPROPERTY(EditAnywhere, NoClear, BlueprintReadOnly, Category = Classes)
 	TSubclassOf<AServerStatReplicator> ServerStatReplicatorClass;
 
+	/** Return if the game mode is requesting a specific replication system to be used for the GameNetDriver */
+	EReplicationSystem GetGameNetDriverReplicationSystem() const { return GameNetDriverReplicationSystem; }
+
 	//~=============================================================================
 	// Accessors for current state
 
@@ -588,6 +592,14 @@ protected:
 	UPROPERTY(EditAnywhere, BlueprintReadWrite, Category=GameMode)
 	uint32 bPauseable : 1;
 
+	/** 
+	 * Can be used to request a specific replication system for a GameNetDriver that will replicate this game mode.
+	 * Leave to Default to use the game engine's preferred system. 
+	 * Useful when migrating from one repsystem to another and a game mode does not fully support both repsystem yet.
+	 */
+	UPROPERTY(EditAnywhere, BlueprintReadWrite, Category = GameMode)
+	EReplicationSystem GameNetDriverReplicationSystem = EReplicationSystem::Default;
+
 	/** The list of delegates to check before unpausing a game */
 	TArray<FCanUnpause> Pausers;
 
diff --git a/Engine/Source/Runtime/Engine/Private/UnrealEngine.cpp b/Engine/Source/Runtime/Engine/Private/UnrealEngine.cpp
--- a/Engine/Source/Runtime/Engine/Private/UnrealEngine.cpp
+++ b/Engine/Source/Runtime/Engine/Private/UnrealEngine.cpp
@@ -13575,15 +13575,9 @@ namespace UE::Private
 	 *		3. NetDriverDefinition match
 	 *
 	 */
-	bool IsNetDriverUsingIris(UEngine* Engine, FName InNetDriverDefinition, FName InNetDriverName)
+	bool IsNetDriverUsingIris(UEngine* Engine, const FWorldContext& Context, FName InNetDriverDefinition, FName InNetDriverName)
 	{
 #if UE_WITH_IRIS
-		// NetDrivers can only use Iris if the global setting allows it.
-		if (UE::Net::ShouldUseIrisReplication() == false)
-		{
-			return false;
-		}
-
 		FIrisNetDriverConfig* IrisConfig = nullptr;
 
 		// Search for the exact name match
@@ -13614,8 +13608,76 @@ namespace UE::Private
 				});
 		}
 
+		const bool bConfigCanUseIris = IrisConfig && IrisConfig->bCanUseIris;
+		const bool bIsEngineDefaultIris = UE::Net::ShouldUseIrisReplication();
+
+		bool bUseIrisRepSystem = bConfigCanUseIris && bIsEngineDefaultIris;
+
+		if (InNetDriverDefinition == NAME_GameNetDriver)
+		{
+			if (UWorld* World = Context.World())
+			{
+				// If we are the server look if the game mode wanted to override the replication system
+				if (AGameModeBase* GameMode = World->GetAuthGameMode())
+				{
+					EReplicationSystem RequestedRepSystem = GameMode->GetGameNetDriverReplicationSystem();
+
+					// If the gamemode requested to use the generic repsystem
+					if (RequestedRepSystem == EReplicationSystem::Generic)
+					{
+						UE_CLOG(bUseIrisRepSystem, LogNet, Log, TEXT("GameMode %s is forcing NetDriver %s (NetDefinition %s) to use the Generic replication system."), *GetNameSafe(GameMode), *InNetDriverName.ToString(), *InNetDriverDefinition.ToString());
+						bUseIrisRepSystem = false;
+					}
+					// If the gamemode requested to use the Iris repsystem
+					else if (RequestedRepSystem == EReplicationSystem::Iris)
+					{
+						UE_CLOG(!bUseIrisRepSystem && bConfigCanUseIris, LogNet, Log, TEXT("GameMode %s is forcing NetDriver %s (NetDefinition %s) to use the Iris replication system."), *GetNameSafe(GameMode), *InNetDriverName.ToString(), *InNetDriverDefinition.ToString());
+
+						// Enable Iris ONLY if the config supports it.
+						bUseIrisRepSystem = bConfigCanUseIris;
+					}
+				}
+				else
+				{
+#if WITH_EDITOR
+					// In PIE let's cheat and make sure the clients follow what the server's net driver is using
+					if (Context.WorldType == EWorldType::PIE)
+					{
+						if (FWorldContext* ServerPIEContext = GEngine->GetWorldContextFromPIEInstance(0))
+						{
+							if (ServerPIEContext->RunAsDedicated)
+							{
+								for (const FNamedNetDriver PieNetDriver : ServerPIEContext->ActiveNetDrivers)
+								{
+									if (PieNetDriver.NetDriverDef->DefName == NAME_GameNetDriver)
+									{
+										bUseIrisRepSystem = PieNetDriver.NetDriver->IsUsingIrisReplication();
+										break;
+									}
+								}
+							}
+						}
+					}
+#endif
+				}
+			}
+		}
+
+		// Ignore most of the above if the cmdline is requesting a specific system
+		const EReplicationSystem CmdlineRequest = UE::Net::GetUseIrisReplicationCmdlineValue();
+		if (CmdlineRequest == EReplicationSystem::Iris)
+		{
+			UE_CLOG(!bUseIrisRepSystem && bConfigCanUseIris, LogNet, Log, TEXT("Cmdline -UseIrisReplication=1 is forcing NetDriver %s (NetDefinition %s) to use the Iris replication system."), *InNetDriverName.ToString(), *InNetDriverDefinition.ToString());
+			bUseIrisRepSystem = bConfigCanUseIris;
+		}
+		else if (CmdlineRequest == EReplicationSystem::Generic)
+		{
+			UE_CLOG(bUseIrisRepSystem, LogNet, Log, TEXT("Cmdline -UseIrisReplication=0 is forcing NetDriver %s (NetDefinition %s) to use the Iris replication system."), *InNetDriverName.ToString(), *InNetDriverDefinition.ToString());
+			bUseIrisRepSystem = false;
+		}
+
 		// Only use Iris if the module is loaded (this happens automatically if the Iris plugin is enabled)
-		if (IrisConfig && IrisConfig->bCanUseIris)
+		if (bUseIrisRepSystem)
 		{
 			if (!ensureMsgf(FModuleManager::Get().IsModuleLoaded("IrisCore"), TEXT("%s is not using Iris because the IrisCore module isn't loaded. Check whether the Iris plugin is enabled."), *InNetDriverName.ToString()))
 			{
@@ -13623,7 +13685,7 @@ namespace UE::Private
 			}
 		}
 
-		return IrisConfig ? IrisConfig->bCanUseIris : false;
+		return bUseIrisRepSystem;
 #else
 		return false;
 #endif //UE_WITH_IRIS
@@ -13763,7 +13825,7 @@ namespace UE::Private
 				check(ReturnVal != nullptr);
 
 				const FName DriverName = InNetDriverName.IsNone() ? ReturnVal->GetFName() : InNetDriverName;
-				const bool bInitializeWithIris = IsNetDriverUsingIris(Engine, NetDriverDefinition, DriverName);
+				const bool bInitializeWithIris = IsNetDriverUsingIris(Engine, Context, NetDriverDefinition, DriverName);
 
 				ReturnVal->SetNetDriverName(DriverName);
 				ReturnVal->SetNetDriverDefinition(NetDriverDefinition);
diff --git a/Engine/Source/Runtime/Engine/Private/World.cpp b/Engine/Source/Runtime/Engine/Private/World.cpp
--- a/Engine/Source/Runtime/Engine/Private/World.cpp
+++ b/Engine/Source/Runtime/Engine/Private/World.cpp
@@ -7726,6 +7726,16 @@ UWorld* FSeamlessTravelHandler::Tick()
 			if (NetDriver)
 			{
 				NetDriver->PreSeamlessTravelGarbageCollect();
+
+				// Warn if we loaded a game mode that wanted a different replication system from the previous mode.
+				if (AGameModeBase* GameMode = LoadedWorld->GetAuthGameMode())
+				{
+					EReplicationSystem LoadedGameModeRepSystem = GameMode->GetGameNetDriverReplicationSystem();
+					const bool bIsNetDriverCompatible = LoadedGameModeRepSystem == EReplicationSystem::Default || 
+														UE::Net::GetUseIrisReplicationCmdlineValue() != EReplicationSystem::Default ||
+														(LoadedGameModeRepSystem == EReplicationSystem::Iris && NetDriver->IsUsingIrisReplication());
+					ensureMsgf(bIsNetDriverCompatible, TEXT("Seamless travel loaded game mode %s that wants a different replication system than the current NetDriver uses."), *GetNameSafe(GameMode));
+				}
 			}
 
 			GWorld = nullptr;
diff --git a/Engine/Source/Runtime/Experimental/Iris/Core/Private/Iris/IrisConfig.cpp b/Engine/Source/Runtime/Experimental/Iris/Core/Private/Iris/IrisConfig.cpp
--- a/Engine/Source/Runtime/Experimental/Iris/Core/Private/Iris/IrisConfig.cpp
+++ b/Engine/Source/Runtime/Experimental/Iris/Core/Private/Iris/IrisConfig.cpp
@@ -2,6 +2,8 @@
 
 #include "Iris/IrisConfig.h"
 #include "HAL/IConsoleManager.h"
+#include "Misc/CommandLine.h"
+#include "Misc/Parse.h"
 
 namespace UE::Net
 {
@@ -19,4 +21,17 @@ void SetUseIrisReplication(bool EnableIrisReplication)
 	CVarUseIrisReplication = EnableIrisReplication ? 1 : 0;
 }
 
+EReplicationSystem GetUseIrisReplicationCmdlineValue()
+{
+	int32 UseIrisReplication=0;
+	if (FParse::Value(FCommandLine::Get(), TEXT("UseIrisReplication="), UseIrisReplication))
+	{
+		// Try to force the requested system if the cmdline is present
+		return UseIrisReplication > 0 ? EReplicationSystem::Iris : EReplicationSystem::Generic;
+	}
+
+	// Use the default engine value
+	return EReplicationSystem::Default;
+}
+
 }
diff --git a/Engine/Source/Runtime/Experimental/Iris/Core/Public/Iris/IrisConfig.h b/Engine/Source/Runtime/Experimental/Iris/Core/Public/Iris/IrisConfig.h
--- a/Engine/Source/Runtime/Experimental/Iris/Core/Public/Iris/IrisConfig.h
+++ b/Engine/Source/Runtime/Experimental/Iris/Core/Public/Iris/IrisConfig.h
@@ -4,12 +4,20 @@
 
 #include "HAL/Platform.h"
 
+#include "Net/Core/Connection/NetEnums.h"
+
 namespace UE::Net
 {
 
+/** Returns if the preferred replication system should be Iris. */
 IRISCORE_API bool ShouldUseIrisReplication();
+
+/** Set if the preferred replication system should be Iris or not. */
 IRISCORE_API void SetUseIrisReplication(bool EnableIrisReplication);
 
+/** Returns what replication sytem was set to be used by the cmdline. Returns Default when the command line was not set. */
+IRISCORE_API EReplicationSystem GetUseIrisReplicationCmdlineValue();
+
 }
 
 /* NetBitStreamReader/Writer validation support */
diff --git a/Engine/Source/Runtime/Net/Core/Public/Net/Core/Connection/NetEnums.h b/Engine/Source/Runtime/Net/Core/Public/Net/Core/Connection/NetEnums.h
--- a/Engine/Source/Runtime/Net/Core/Public/Net/Core/Connection/NetEnums.h
+++ b/Engine/Source/Runtime/Net/Core/Public/Net/Core/Connection/NetEnums.h
@@ -70,6 +70,13 @@ namespace ENetworkFailure
 	}
 }
 
+UENUM(BlueprintType)
+enum class EReplicationSystem : uint8
+{
+	Default,
+	Generic,
+	Iris,
+};
 
 // (DEPRECATED) Security event types used for UE_SECURITY_LOG
 namespace ESecurityEvent
