Use a consistent allocator in FCSPose

diff --git a/Engine/Source/Runtime/Engine/Public/BonePose.h b/Engine/Source/Runtime/Engine/Public/BonePose.h
--- a/Engine/Source/Runtime/Engine/Public/BonePose.h
+++ b/Engine/Source/Runtime/Engine/Public/BonePose.h
@@ -408,7 +408,8 @@ struct FCSPose
 {
 	// Set up our index type based on the type of pose we are manipulating
 	typedef typename PoseType::BoneIndexType BoneIndexType;
-
+	typedef typename PoseType::Allocator AllocatorType;
+	
 	// Init Pose
 	void InitPose(const FBoneContainer* InBoneContainer)
 	{
@@ -458,7 +459,7 @@ struct FCSPose
 	}
 
 	const PoseType& GetPose() const { return Pose; }
-	const TCustomBoneIndexArray<uint8, BoneIndexType>& GetComponentSpaceFlags() const { return ComponentSpaceFlags; }
+	const TCustomBoneIndexArray<uint8, BoneIndexType, AllocatorType>& GetComponentSpaceFlags() const { return ComponentSpaceFlags; }
 
 	// Get transform for supplied bone in local space
 	FTransform GetLocalSpaceTransform(BoneIndexType BoneIndex);
@@ -529,13 +530,13 @@ protected:
 	PoseType Pose;
 
 	// Flags to track each bones current state (0 means local pose, 1 means component space pose)
-	TCustomBoneIndexArray<uint8, BoneIndexType> ComponentSpaceFlags;
+	TCustomBoneIndexArray<uint8, BoneIndexType, AllocatorType> ComponentSpaceFlags;
 
 	// Cached bone mask array to avoid reallocations
-	TCustomBoneIndexArray<uint8, BoneIndexType> BoneMask;
+	TCustomBoneIndexArray<uint8, BoneIndexType, AllocatorType> BoneMask;
 
 	// Cached conversion array for this pose, to save on allocations each frame
-	TArray<FCompactPoseBoneIndex> BonesToConvert;
+	TArray<FCompactPoseBoneIndex, AllocatorType> BonesToConvert;
 };
 
 template<class PoseType>
diff --git a/Engine/Source/Runtime/Engine/Public/CustomBoneIndexArray.h b/Engine/Source/Runtime/Engine/Public/CustomBoneIndexArray.h
--- a/Engine/Source/Runtime/Engine/Public/CustomBoneIndexArray.h
+++ b/Engine/Source/Runtime/Engine/Public/CustomBoneIndexArray.h
@@ -103,8 +103,8 @@ using TTypedIndexArrayView = TTypedIndexRange<IndexType, TArrayView<TArrayViewAr
 
 } // namespace UE::Anim
 
-template<typename ElementType, typename BoneIndexType>
-using TCustomBoneIndexArray = UE::Anim::TTypedIndexRange<BoneIndexType, TArray<ElementType>>;
+template<typename ElementType, typename BoneIndexType, typename InAllocatorType = typename TArray<ElementType>::AllocatorType>
+using TCustomBoneIndexArray = UE::Anim::TTypedIndexRange<BoneIndexType, TArray<ElementType, InAllocatorType>>;
 
 template<typename ElementType, typename BoneIndexType>
 using TCustomBoneIndexArrayView = UE::Anim::TTypedIndexRange<BoneIndexType, TArrayView<ElementType>>;
