Fix const-copy error on Linux

#rnx

diff --git a/Engine/Source/Runtime/Engine/Private/UnrealEngine.cpp b/Engine/Source/Runtime/Engine/Private/UnrealEngine.cpp
--- a/Engine/Source/Runtime/Engine/Private/UnrealEngine.cpp
+++ b/Engine/Source/Runtime/Engine/Private/UnrealEngine.cpp
@@ -13647,7 +13647,7 @@ namespace UE::Private
 						{
 							if (ServerPIEContext->RunAsDedicated)
 							{
-								for (const FNamedNetDriver PieNetDriver : ServerPIEContext->ActiveNetDrivers)
+								for (const FNamedNetDriver& PieNetDriver : ServerPIEContext->ActiveNetDrivers)
 								{
 									if (PieNetDriver.NetDriverDef->DefName == NAME_GameNetDriver)
 									{
