Making attenuation assets live updatable in editor and settable in BP.

diff --git a/Engine/Source/Runtime/Engine/Classes/Components/AudioComponent.h b/Engine/Source/Runtime/Engine/Classes/Components/AudioComponent.h
--- a/Engine/Source/Runtime/Engine/Classes/Components/AudioComponent.h
+++ b/Engine/Source/Runtime/Engine/Classes/Components/AudioComponent.h
@@ -230,9 +230,12 @@ class UAudioComponent : public USceneComponent, public ISoundParameterController
 	uint8 bAllowSpatialization:1;
 
 	/** Allows defining attenuation settings directly on this audio component without using an attenuation settings asset. */
-	UPROPERTY(EditAnywhere, BlueprintReadWrite, Category = Attenuation)
+	UPROPERTY(EditAnywhere, BlueprintReadWrite, BlueprintSetter = SetOverrideAttenuation, Category = Attenuation)
 	uint8 bOverrideAttenuation:1;
 
+	UFUNCTION(BlueprintSetter)
+	void SetOverrideAttenuation(bool bInOverrideAttenuation);
+
 	/** Whether or not to override the sound's subtitle priority. */
 	UPROPERTY(EditAnywhere, BlueprintReadWrite, Category = Subtitles, meta = (InlineEditConditionToggle))
 	uint8 bOverrideSubtitlePriority:1;
@@ -381,13 +384,19 @@ public:
 	int32 ActiveCount;
 
 	/** If bOverrideSettings is false, the asset to use to determine attenuation properties for sounds generated by this component */
-	UPROPERTY(EditAnywhere, BlueprintReadWrite, Category = Attenuation, meta = (EditCondition = "!bOverrideAttenuation", DisplayAfter = "bOverrideAttenuation", EditConditionHides))
-	TObjectPtr<class USoundAttenuation> AttenuationSettings;
+	UPROPERTY(EditAnywhere, BlueprintReadWrite, Category = Attenuation, BlueprintSetter = SetAttenuationSettings, meta = (EditCondition = "!bOverrideAttenuation", DisplayAfter = "bOverrideAttenuation", EditConditionHides))
+	TObjectPtr<USoundAttenuation> AttenuationSettings;
+
+	UFUNCTION(BlueprintSetter)
+	void SetAttenuationSettings(USoundAttenuation* InAttenuationSettings);
 
 	/** If bOverrideSettings is true, the attenuation properties to use for sounds generated by this component */
-	UPROPERTY(EditAnywhere, BlueprintReadWrite, Category = Attenuation, meta = (EditCondition = "bOverrideAttenuation", DisplayAfter = "bOverrideAttenuation", EditConditionHides))
+	UPROPERTY(EditAnywhere, BlueprintReadWrite, Category = Attenuation, BlueprintSetter = SetAttenuationOverrides, meta = (EditCondition = "bOverrideAttenuation", DisplayAfter = "bOverrideAttenuation", EditConditionHides))
 	struct FSoundAttenuationSettings AttenuationOverrides;
 
+	UFUNCTION(BlueprintSetter)
+	void SetAttenuationOverrides(const FSoundAttenuationSettings& InAttenuationOverrides);
+
 	/** What sound concurrency to use for sounds generated by this audio component */
 	UPROPERTY()
 	TObjectPtr<USoundConcurrency> ConcurrencySettings_DEPRECATED;
@@ -751,6 +760,7 @@ public:
 
 	static ENGINE_API void PlaybackCompleted(uint64 AudioComponentID, bool bFailedToStart);
 
+
 	bool GetDisableParameterUpdatesWhilePlaying() const override { return static_cast<bool>(bDisableParameterUpdatesWhilePlaying); }
 
 private:
@@ -809,6 +819,7 @@ public:
 
 	/** Returns a pointer to the attenuation settings to be used (if any) for this audio component dependent on the SoundAttenuation asset or overrides set. */
 	ENGINE_API const FSoundAttenuationSettings* GetAttenuationSettingsToApply() const;
+	ENGINE_API const TObjectPtr<USoundAttenuation> GetAttenuationSettingsAsset() const;
 
 	/** Retrieves Attenuation Settings data on the targeted Audio Component. Returns FALSE if no settings were found. 
 	 *  Because the Attenuation Settings data structure is copied, FALSE returns will return default values. 
diff --git a/Engine/Source/Runtime/Engine/Private/ActiveSound.cpp b/Engine/Source/Runtime/Engine/Private/ActiveSound.cpp
--- a/Engine/Source/Runtime/Engine/Private/ActiveSound.cpp
+++ b/Engine/Source/Runtime/Engine/Private/ActiveSound.cpp
@@ -52,6 +52,7 @@ FActiveSound::FActiveSound()
 	, WorldID(0)
 	, Sound(nullptr)
 	, SourceEffectChain(nullptr)
+	, SoundAttenuation(nullptr)
 	, AudioComponentID(0)
 	, OwnerID(0)
 	, PlayOrder(INDEX_NONE)
@@ -98,6 +99,7 @@ FActiveSound::FActiveSound()
 	, bIsFirstAttenuationUpdate(true)
 	, bStartedWithinNonBinauralRadius(false)
 	, bModulationRoutingUpdated(false)
+	, bIsAttenuationSettingsOverridden(false)
 	, UserIndex(0)
 	, FadeOut(EFadeOut::None)
 	, bIsOccluded(false)
@@ -225,6 +227,8 @@ void FActiveSound::AddReferencedObjects(FReferenceCollector& Collector)
 		Sound->SourceEffectChain->AddReferencedEffects(Collector);
 	}
 
+	Collector.AddReferencedObject(SoundAttenuation);
+
 	for (auto& Concurrency : ConcurrencySet)
 	{
 		if (Concurrency)
@@ -292,6 +296,16 @@ void FActiveSound::SetSoundClass(USoundClass* SoundClass)
 		|| (Sound && Sound->ShouldApplyInteriorVolumes());
 }
 
+void FActiveSound::SetAttenuationSettingsAsset(TObjectPtr<USoundAttenuation> InSoundAttenuation)
+{
+	SoundAttenuation = InSoundAttenuation;
+}
+
+void FActiveSound::SetAttenuationSettingsOverride(bool bInattenuationSettingsOverride)
+{
+	bIsAttenuationSettingsOverridden = bInattenuationSettingsOverride;
+}
+
 bool FActiveSound::IsPlayWhenSilent() const
 {
 	if (!AudioDevice || !AudioDevice->PlayWhenSilentEnabled())
@@ -1425,7 +1439,14 @@ void FActiveSound::CollectAttenuationShapesForVisualization(TMultiMap<EAttenuati
 
 	if (bHasAttenuationSettings)
 	{
-		AttenuationSettings.CollectAttenuationShapesForVisualization(ShapeDetailsMap);
+		if (bIsAttenuationSettingsOverridden)
+		{
+			AttenuationSettings.CollectAttenuationShapesForVisualization(ShapeDetailsMap);
+		}
+		else if (SoundAttenuation)
+		{
+			SoundAttenuation->Attenuation.CollectAttenuationShapesForVisualization(ShapeDetailsMap);
+		}
 	}
 
 	// For sound cues we'll dig in and see if we can find any attenuation sound nodes that will affect the settings
@@ -1640,8 +1661,25 @@ void FActiveSound::UpdateAttenuation(float DeltaTime, FSoundParseParameters& Par
 
 void FActiveSound::UpdateAttenuation(float DeltaTime, FSoundParseParameters& ParseParams, int32 ListenerIndex, const FSoundAttenuationSettings* SettingsAttenuationNode)
 {
+	const FSoundAttenuationSettings* Settings = nullptr;
+
 	// Get the attenuation settings to use for this application to the active sound
-	const FSoundAttenuationSettings* Settings = SettingsAttenuationNode ? SettingsAttenuationNode : &AttenuationSettings;
+	// Use the passed-in attenuation settings
+	if (SettingsAttenuationNode)
+	{
+		Settings = SettingsAttenuationNode;
+	}
+	// We use the copied off "overridden" settings
+	else if (bIsAttenuationSettingsOverridden)
+	{
+		Settings = &AttenuationSettings;
+	}
+	// We fallback to using the asset's settings directly
+	else if (SoundAttenuation)
+	{
+		Settings = &SoundAttenuation->Attenuation;
+	}
+
 	if (!Settings)
 	{
 		UE_LOG(LogAudio, Warning, TEXT("No attenuation settings found for active sound."));
diff --git a/Engine/Source/Runtime/Engine/Private/Components/AudioComponent.cpp b/Engine/Source/Runtime/Engine/Private/Components/AudioComponent.cpp
--- a/Engine/Source/Runtime/Engine/Private/Components/AudioComponent.cpp
+++ b/Engine/Source/Runtime/Engine/Private/Components/AudioComponent.cpp
@@ -727,6 +727,9 @@ void UAudioComponent::PlayInternal(const PlayInternalRequestData& InPlayRequestD
 	NewActiveSound.SetSound(SoundToPlay);
 	NewActiveSound.SetSourceEffectChain(SourceEffectChain);
 	NewActiveSound.SetSoundClass(SoundClassOverride);
+	NewActiveSound.SetAttenuationSettingsAsset(GetAttenuationSettingsAsset());
+	NewActiveSound.SetAttenuationSettingsOverride(bOverrideAttenuation);
+
 	NewActiveSound.ConcurrencySet = ConcurrencySet;
 
 	const float Volume = (VolumeModulationMax + ((VolumeModulationMin - VolumeModulationMax) * RandomStream.FRand())) * VolumeMultiplier;
@@ -1265,6 +1268,19 @@ void UAudioComponent::PostEditChangeProperty(FPropertyChangedEvent& PropertyChan
 }
 #endif // WITH_EDITOR
 
+const TObjectPtr<USoundAttenuation> UAudioComponent::GetAttenuationSettingsAsset() const
+{
+	if (AttenuationSettings)
+	{
+		return AttenuationSettings;
+	}
+	else if (Sound)
+	{
+		return Sound->AttenuationSettings;
+	}
+	return nullptr;
+}
+
 const FSoundAttenuationSettings* UAudioComponent::GetAttenuationSettingsToApply() const
 {
 	if (bOverrideAttenuation)
@@ -1282,6 +1298,54 @@ const FSoundAttenuationSettings* UAudioComponent::GetAttenuationSettingsToApply(
 	return nullptr;
 }
 
+void UAudioComponent::SetAttenuationSettings(USoundAttenuation* InSoundAttenuation)
+{
+	AttenuationSettings = InSoundAttenuation;
+
+	if (FAudioDevice* AudioDevice = GetAudioDevice())
+	{
+		if (IsActive())
+		{
+			AudioDevice->SendCommandToActiveSounds(AudioComponentID, [InSoundAttenuation](FActiveSound& ActiveSound)
+			{
+				ActiveSound.SetAttenuationSettingsAsset(InSoundAttenuation);
+			});
+		}
+	}
+}
+
+void UAudioComponent::SetAttenuationOverrides(const FSoundAttenuationSettings& InAttenuationOverrides)
+{
+	AttenuationOverrides = InAttenuationOverrides;
+
+	if (FAudioDevice* AudioDevice = GetAudioDevice())
+	{
+		if (IsActive())
+		{
+			AudioDevice->SendCommandToActiveSounds(AudioComponentID, [InAttenuationOverrides](FActiveSound& ActiveSound)
+			{
+				ActiveSound.AttenuationSettings = InAttenuationOverrides;
+			});
+		}
+	}
+}
+
+void UAudioComponent::SetOverrideAttenuation(bool bInOverrideAttenuation)
+{
+	bOverrideAttenuation = bInOverrideAttenuation;
+
+	if (FAudioDevice* AudioDevice = GetAudioDevice())
+	{
+		if (IsActive())
+		{
+			AudioDevice->SendCommandToActiveSounds(AudioComponentID, [bInOverrideAttenuation](FActiveSound& ActiveSound)
+			{
+				ActiveSound.bIsAttenuationSettingsOverridden = bInOverrideAttenuation;
+			});
+		}
+	}
+}
+
 bool UAudioComponent::BP_GetAttenuationSettingsToApply(FSoundAttenuationSettings& OutAttenuationSettings)
 {
 	if (const FSoundAttenuationSettings* Settings = GetAttenuationSettingsToApply())
diff --git a/Engine/Source/Runtime/Engine/Private/SoundCue.cpp b/Engine/Source/Runtime/Engine/Private/SoundCue.cpp
--- a/Engine/Source/Runtime/Engine/Private/SoundCue.cpp
+++ b/Engine/Source/Runtime/Engine/Private/SoundCue.cpp
@@ -383,8 +383,22 @@ void USoundCue::PostEditChangeProperty(struct FPropertyChangedEvent& PropertyCha
 		{
 			if (It->Sound == this && It->IsActive())
 			{
-				It->Stop();
-				It->Play();
+				// Allow attenuation overrides not update without stopping
+				if (PropertyChangedEvent.MemberProperty->GetFName() == GET_MEMBER_NAME_STRING_CHECKED(USoundCue, AttenuationOverrides))
+				{
+					It->SetAttenuationOverrides(AttenuationOverrides);
+				}
+				else if (PropertyChangedEvent.MemberProperty->GetFName() == GET_MEMBER_NAME_STRING_CHECKED(USoundCue, bOverrideAttenuation))
+				{
+					It->SetAttenuationOverrides(AttenuationOverrides);
+					It->SetOverrideAttenuation(bOverrideAttenuation);
+				}
+				else
+				{
+					It->Stop();
+					It->Play();
+				}
+
 			}
 		}
 
diff --git a/Engine/Source/Runtime/Engine/Public/ActiveSound.h b/Engine/Source/Runtime/Engine/Public/ActiveSound.h
--- a/Engine/Source/Runtime/Engine/Public/ActiveSound.h
+++ b/Engine/Source/Runtime/Engine/Public/ActiveSound.h
@@ -289,6 +289,7 @@ private:
 
 	TObjectPtr<USoundBase> Sound;
 	TObjectPtr<USoundEffectSourcePresetChain> SourceEffectChain;
+	TObjectPtr<USoundAttenuation> SoundAttenuation;
 
 	uint64 AudioComponentID;
 	FName AudioComponentUserID;
@@ -343,6 +344,10 @@ public:
 
 	ENGINE_API void SetSoundClass(USoundClass* SoundClass);
 
+	ENGINE_API void SetAttenuationSettingsAsset(TObjectPtr<USoundAttenuation> InSoundAttenuation);
+
+	ENGINE_API void SetAttenuationSettingsOverride(bool bInIsAttenuationSettingsOverridden);
+
 	void SetAudioDevice(FAudioDevice* InAudioDevice)
 	{
 		AudioDevice = InAudioDevice;
@@ -528,6 +533,9 @@ public:
 
 	uint8 bModulationRoutingUpdated : 1;
 
+	/** If this is true the active sound uses the overridden struct of the sound not the attenuation settings asset. */
+	uint8 bIsAttenuationSettingsOverridden : 1;
+
 	uint8 UserIndex;
 
 	/** Type of fade out currently being applied */
