Fix override vertex colors occasionally not working on skeletal meshes.

The init/release render resource commands for the override colors buffer were happening on the default pipe instead of the skeletal mesh pipe, so the mesh draw commands were using stale data.

#rnx

diff --git a/Engine/Source/Runtime/Engine/Private/Components/SkinnedMeshComponent.cpp b/Engine/Source/Runtime/Engine/Private/Components/SkinnedMeshComponent.cpp
--- a/Engine/Source/Runtime/Engine/Private/Components/SkinnedMeshComponent.cpp
+++ b/Engine/Source/Runtime/Engine/Private/Components/SkinnedMeshComponent.cpp
@@ -39,6 +39,7 @@
 #include "HAL/LowLevelMemTracker.h"
 #include "HAL/LowLevelMemStats.h"
 #include "UObject/Package.h"
+#include "Rendering/RenderCommandPipes.h"
 
 DEFINE_LOG_CATEGORY_STATIC(LogSkinnedMeshComp, Log, All);
 
@@ -4101,7 +4102,7 @@ void FSkelMeshComponentLODInfo::ReleaseOverrideVertexColorsAndBlock()
 	if (OverrideVertexColors)
 	{
 		// enqueue a rendering command to release
-		BeginReleaseResource(OverrideVertexColors);
+		BeginReleaseResource(OverrideVertexColors, &UE::RenderCommandPipe::SkeletalMesh);
 		// Ensure the RT no longer accessed the data, might slow down
 		FlushRenderingCommands();
 		// The RT thread has no access to it any more so it's safe to delete it.
@@ -4114,7 +4115,7 @@ bool FSkelMeshComponentLODInfo::BeginReleaseOverrideVertexColors()
 	if (OverrideVertexColors)
 	{
 		// enqueue a rendering command to release
-		BeginReleaseResource(OverrideVertexColors);
+		BeginReleaseResource(OverrideVertexColors, &UE::RenderCommandPipe::SkeletalMesh);
 		return true;
 	}
 
@@ -4255,7 +4256,7 @@ void USkinnedMeshComponent::SetVertexColorOverride(int32 LODIndex, const TArray<
 		Info.OverrideVertexColors = new FColorVertexBuffer;
 		Info.OverrideVertexColors->InitFromColorArray(*UseColors);
 
-		BeginInitResource(Info.OverrideVertexColors);
+		BeginInitResource(Info.OverrideVertexColors, &UE::RenderCommandPipe::SkeletalMesh);
 
 		MarkRenderStateDirty();
 	}
