Skip default material PSO pre-caching on platforms that do not support MVF.
Those platforms require mesh information to pre-cache PSO correctly

diff --git a/Engine/Source/Runtime/Engine/Private/Materials/Material.cpp b/Engine/Source/Runtime/Engine/Private/Materials/Material.cpp
--- a/Engine/Source/Runtime/Engine/Private/Materials/Material.cpp
+++ b/Engine/Source/Runtime/Engine/Private/Materials/Material.cpp
@@ -573,7 +573,8 @@ void UMaterialInterface::InitDefaultMaterials()
 
 		// Now precache PSOs for all the default materials after the default materials are marked initialize
 		// PSO precaching can request default materials so they have to marked as initialized to avoid endless recursion
-		if (bInitialized && PipelineStateCache::IsPSOPrecachingEnabled())
+		// Skip platforms that do not support MVF, non-MVF path needs mesh information for PSO 
+		if (bInitialized && PipelineStateCache::IsPSOPrecachingEnabled() && RHISupportsManualVertexFetch(GMaxRHIShaderPlatform))
 		{
 			TArray<FMaterialPSOPrecacheRequestID> MaterialPrecacheRequestIDs;
 
