Partial revert of CL28378343, CL28383888

Migration logic from LocalExposureHighlightContrastScale/LocalExposureShadowContrastScale to LocalExposureHighlightContrast/LocalExposureShadowContrast was causing issues with blueprints inheritance.

diff --git a/Engine/Source/Runtime/Core/Public/UObject/FortniteMainBranchObjectVersion.h b/Engine/Source/Runtime/Core/Public/UObject/FortniteMainBranchObjectVersion.h
--- a/Engine/Source/Runtime/Core/Public/UObject/FortniteMainBranchObjectVersion.h
+++ b/Engine/Source/Runtime/Core/Public/UObject/FortniteMainBranchObjectVersion.h
@@ -378,8 +378,8 @@ struct FFortniteMainBranchObjectVersion
 		// Serialize custom editor bounds for HLOD actor descriptors.
 		WorldPartitionHLODActorDescSerializeEditorBounds,
 
-		// Changed default Local Exposure Contrast from 1.0 to 0.8
-		LocalExposureDefaultChangeFrom1,
+		// Changed default Local Exposure Contrast from 1.0 to 0.8 (reverted)
+		LocalExposureDefaultChangeFrom1_Reverted,
 
 		// Added support of external packaging of Data Layer Instances
 		AddDataLayerInstanceExternalPackage,
diff --git a/Engine/Source/Runtime/Engine/Classes/Engine/Scene.h b/Engine/Source/Runtime/Engine/Classes/Engine/Scene.h
--- a/Engine/Source/Runtime/Engine/Classes/Engine/Scene.h
+++ b/Engine/Source/Runtime/Engine/Classes/Engine/Scene.h
@@ -901,17 +901,11 @@ struct FPostProcessSettings
 	UPROPERTY()
 	uint8 bOverride_LocalExposureContrastScale_DEPRECATED:1;
 
-	UPROPERTY()
-	uint8 bOverride_LocalExposureHighlightContrastScale_DEPRECATED :1;
-
-	UPROPERTY()
-	uint8 bOverride_LocalExposureShadowContrastScale_DEPRECATED :1;
-
 	UPROPERTY(EditAnywhere, BlueprintReadWrite, Category=Overrides, meta=(PinHiddenByDefault, InlineEditConditionToggle))
-	uint8 bOverride_LocalExposureHighlightContrast:1;
+	uint8 bOverride_LocalExposureHighlightContrastScale:1;
 
 	UPROPERTY(EditAnywhere, BlueprintReadWrite, Category=Overrides, meta=(PinHiddenByDefault, InlineEditConditionToggle))
-	uint8 bOverride_LocalExposureShadowContrast:1;
+	uint8 bOverride_LocalExposureShadowContrastScale:1;
 
 	UPROPERTY(EditAnywhere, BlueprintReadWrite, Category=Overrides, meta=(PinHiddenByDefault, InlineEditConditionToggle))
 	uint8 bOverride_LocalExposureHighlightContrastCurve:1;
@@ -1838,20 +1832,14 @@ struct FPostProcessSettings
 	UPROPERTY()
 	float LocalExposureContrastScale_DEPRECATED;
 
-	UPROPERTY()
-	float LocalExposureHighlightContrastScale_DEPRECATED;
-
-	UPROPERTY()
-	float LocalExposureShadowContrastScale_DEPRECATED;
-
 	/** 
 	 * Local Exposure decomposes luminance of the frame into a base layer and a detail layer.
 	 * Contrast of the base layer is reduced based on this value.
 	 * Value less than 1 will enable local exposure.
 	 * Good values are usually in the range 0.6 .. 1.0.
 	*/
-	UPROPERTY(interp, BlueprintReadWrite, Category="Lens|Local Exposure", meta=(UIMin = "0.0", UIMax = "1.0", editcondition = "bOverride_LocalExposureHighlightContrast", DisplayName = "Highlight Contrast"))
-	float LocalExposureHighlightContrast;
+	UPROPERTY(interp, BlueprintReadWrite, Category="Lens|Local Exposure", meta=(UIMin = "0.0", UIMax = "1.0", editcondition = "bOverride_LocalExposureHighlightContrastScale", DisplayName = "Highlight Contrast"))
+	float LocalExposureHighlightContrastScale;
 
 	/** 
 	 * Local Exposure decomposes luminance of the frame into a base layer and a detail layer.
@@ -1859,8 +1847,8 @@ struct FPostProcessSettings
 	 * Value less than 1 will enable local exposure.
 	 * Good values are usually in the range 0.6 .. 1.0.
 	*/
-	UPROPERTY(interp, BlueprintReadWrite, Category="Lens|Local Exposure", meta=(UIMin = "0.0", UIMax = "1.0", editcondition = "bOverride_LocalExposureShadowContrast", DisplayName = "Shadow Contrast"))
-	float LocalExposureShadowContrast;
+	UPROPERTY(interp, BlueprintReadWrite, Category="Lens|Local Exposure", meta=(UIMin = "0.0", UIMax = "1.0", editcondition = "bOverride_LocalExposureShadowContrastScale", DisplayName = "Shadow Contrast"))
+	float LocalExposureShadowContrastScale;
 
 	/**
 	 * Local Exposure Highlight Contrast based on the scene EV100.
@@ -2309,6 +2297,17 @@ struct FPostProcessSettings
 		{
 			RayTracingGIType = (ERayTracingGlobalIlluminationType)(RayTracingGI_DEPRECATED == 1);
 		}
+
+		if (bOverride_LocalExposureContrastScale_DEPRECATED)
+		{
+			bOverride_LocalExposureHighlightContrastScale = bOverride_LocalExposureContrastScale_DEPRECATED;
+			bOverride_LocalExposureShadowContrastScale = bOverride_LocalExposureContrastScale_DEPRECATED;
+		}
+		if (LocalExposureContrastScale_DEPRECATED != 1.0f)
+		{
+			LocalExposureHighlightContrastScale = LocalExposureContrastScale_DEPRECATED;
+			LocalExposureShadowContrastScale = LocalExposureContrastScale_DEPRECATED;
+		}
 	}
 #endif
 
diff --git a/Engine/Source/Runtime/Engine/Private/Scene.cpp b/Engine/Source/Runtime/Engine/Private/Scene.cpp
--- a/Engine/Source/Runtime/Engine/Private/Scene.cpp
+++ b/Engine/Source/Runtime/Engine/Private/Scene.cpp
@@ -7,7 +7,6 @@
 #include "UObject/RenderingObjectVersion.h"
 #include "UObject/ReleaseObjectVersion.h"
 #include "UObject/UE5ReleaseStreamObjectVersion.h"
-#include "UObject/FortniteMainBranchObjectVersion.h"
 #include "DataDrivenShaderPlatformInfo.h"
 #include "UObject/UnrealType.h"
 
@@ -510,10 +509,8 @@ FPostProcessSettings::FPostProcessSettings()
 	AutoExposureApplyPhysicalCameraExposure = 1;
 
 	LocalExposureContrastScale_DEPRECATED = 1.0f;
-	LocalExposureHighlightContrastScale_DEPRECATED = 1.0f;
-	LocalExposureShadowContrastScale_DEPRECATED = 1.0f;
-	LocalExposureHighlightContrast = 0.8f;
-	LocalExposureShadowContrast = 0.8f;
+	LocalExposureHighlightContrastScale = 1.0f;
+	LocalExposureShadowContrastScale = 1.0f;
 	LocalExposureDetailStrength = 1.0f;
 	LocalExposureBlurredLuminanceBlend = 0.6f;
 	LocalExposureBlurredLuminanceKernelSizePercent = 50.0f;
@@ -904,26 +901,6 @@ void FPostProcessSettings::PostSerialize(const FArchive& Ar)
 			}
 		}
 
-		if (bOverride_LocalExposureContrastScale_DEPRECATED)
-		{
-			bOverride_LocalExposureHighlightContrastScale_DEPRECATED = bOverride_LocalExposureContrastScale_DEPRECATED;
-			bOverride_LocalExposureShadowContrastScale_DEPRECATED = bOverride_LocalExposureContrastScale_DEPRECATED;
-		}
-		if (LocalExposureContrastScale_DEPRECATED != 1.0f)
-		{
-			LocalExposureHighlightContrastScale_DEPRECATED = LocalExposureContrastScale_DEPRECATED;
-			LocalExposureShadowContrastScale_DEPRECATED = LocalExposureContrastScale_DEPRECATED;
-		}
-
-		if (Ar.CustomVer(FFortniteMainBranchObjectVersion::GUID) < FFortniteMainBranchObjectVersion::LocalExposureDefaultChangeFrom1)
-		{
-			bOverride_LocalExposureHighlightContrast = bOverride_LocalExposureHighlightContrastScale_DEPRECATED;
-			bOverride_LocalExposureShadowContrast = bOverride_LocalExposureShadowContrastScale_DEPRECATED;
-
-			LocalExposureHighlightContrast = LocalExposureHighlightContrastScale_DEPRECATED;
-			LocalExposureShadowContrast = LocalExposureShadowContrastScale_DEPRECATED;
-		}
-
 		if (Ar.CustomVer(FUE5ReleaseStreamObjectVersion::GUID) < FUE5ReleaseStreamObjectVersion::ReflectionMethodEnum)
 		{
 			if (bOverride_RayTracingGI && RayTracingGIType != ERayTracingGlobalIlluminationType::Disabled)
diff --git a/Engine/Source/Runtime/Engine/Private/SceneView.cpp b/Engine/Source/Runtime/Engine/Private/SceneView.cpp
--- a/Engine/Source/Runtime/Engine/Private/SceneView.cpp
+++ b/Engine/Source/Runtime/Engine/Private/SceneView.cpp
@@ -1544,8 +1544,9 @@ void FSceneView::OverridePostProcessSettings(const FPostProcessSettings& Src, fl
 		LERP_PP(AutoExposureBias);
 		LERP_PP(HistogramLogMin);
 		LERP_PP(HistogramLogMax);
-		LERP_PP(LocalExposureHighlightContrast);
-		LERP_PP(LocalExposureShadowContrast);
+		LERP_PP(LocalExposureContrastScale_DEPRECATED);
+		LERP_PP(LocalExposureHighlightContrastScale);
+		LERP_PP(LocalExposureShadowContrastScale);
 		LERP_PP(LocalExposureDetailStrength);
 		LERP_PP(LocalExposureBlurredLuminanceBlend);
 		LERP_PP(LocalExposureBlurredLuminanceKernelSizePercent);
@@ -1857,13 +1858,13 @@ void FSceneView::StartFinalPostprocessSettings(FVector InViewLocation)
 		{
 			const float HighlightContrast = FMath::Clamp(CVarDefaultLocalExposureHighlightContrast.GetValueOnGameThread(), 0.0f, 1.0f);
 
-			FinalPostProcessSettings.LocalExposureHighlightContrast = HighlightContrast;
+			FinalPostProcessSettings.LocalExposureHighlightContrastScale = HighlightContrast;
 		}
 
 		{
 			const float ShadowContrast = FMath::Clamp(CVarDefaultLocalExposureHighlightContrast.GetValueOnGameThread(), 0.0f, 1.0f);
 
-			FinalPostProcessSettings.LocalExposureShadowContrast = ShadowContrast;
+			FinalPostProcessSettings.LocalExposureShadowContrastScale = ShadowContrast;
 		}
 
 		if (!CVarDefaultMotionBlur.GetValueOnGameThread())
@@ -1942,8 +1943,8 @@ void FSceneView::EndFinalPostprocessSettings(const FSceneViewInitOptions& ViewIn
 
 		if (LocalExposureCVarValue <= 0 || !Family->EngineShowFlags.LocalExposure)
 		{
-			FinalPostProcessSettings.LocalExposureHighlightContrast = 1.0f;
-			FinalPostProcessSettings.LocalExposureShadowContrast = 1.0f;
+			FinalPostProcessSettings.LocalExposureHighlightContrastScale = 1.0f;
+			FinalPostProcessSettings.LocalExposureShadowContrastScale = 1.0f;
 			FinalPostProcessSettings.LocalExposureHighlightContrastCurve = nullptr;
 			FinalPostProcessSettings.LocalExposureShadowContrastCurve = nullptr;
 			FinalPostProcessSettings.LocalExposureDetailStrength = 1.0f;
diff --git a/Engine/Source/Runtime/Renderer/Private/PostProcess/PostProcessLocalExposure.cpp b/Engine/Source/Runtime/Renderer/Private/PostProcess/PostProcessLocalExposure.cpp
--- a/Engine/Source/Runtime/Renderer/Private/PostProcess/PostProcessLocalExposure.cpp
+++ b/Engine/Source/Runtime/Renderer/Private/PostProcess/PostProcessLocalExposure.cpp
@@ -111,8 +111,8 @@ FLocalExposureParameters GetLocalExposureParameters(const FViewInfo& View, FIntP
 
 	const FVector2f LocalExposureBilateralGridUVScale = GetLocalExposureBilateralGridUVScale(ViewRectSize);
 
-	float HighlightContrast = Settings.LocalExposureHighlightContrast;
-	float ShadowContrast = Settings.LocalExposureShadowContrast;
+	float HighlightContrast = Settings.LocalExposureHighlightContrastScale;
+	float ShadowContrast = Settings.LocalExposureShadowContrastScale;
 
 	const float AverageSceneLuminance = View.GetLastAverageSceneLuminance();
 	if (AverageSceneLuminance > 0)
diff --git a/Engine/Source/Runtime/Renderer/Private/PostProcess/PostProcessing.cpp b/Engine/Source/Runtime/Renderer/Private/PostProcess/PostProcessing.cpp
--- a/Engine/Source/Runtime/Renderer/Private/PostProcess/PostProcessing.cpp
+++ b/Engine/Source/Runtime/Renderer/Private/PostProcess/PostProcessing.cpp
@@ -613,8 +613,8 @@ void AddPostProcessingPasses(
 
 		const bool bLocalExposureEnabled =
 			EngineShowFlags.VisualizeLocalExposure ||
-			!FMath::IsNearlyEqual(View.FinalPostProcessSettings.LocalExposureHighlightContrast, 1.0f) ||
-			!FMath::IsNearlyEqual(View.FinalPostProcessSettings.LocalExposureShadowContrast, 1.0f) ||
+			!FMath::IsNearlyEqual(View.FinalPostProcessSettings.LocalExposureHighlightContrastScale, 1.0f) ||
+			!FMath::IsNearlyEqual(View.FinalPostProcessSettings.LocalExposureShadowContrastScale, 1.0f) ||
 			View.FinalPostProcessSettings.LocalExposureHighlightContrastCurve ||
 			View.FinalPostProcessSettings.LocalExposureShadowContrastCurve ||
 			!FMath::IsNearlyEqual(View.FinalPostProcessSettings.LocalExposureDetailStrength, 1.0f);
diff --git a/Engine/Source/Runtime/Renderer/Private/SceneRendering.cpp b/Engine/Source/Runtime/Renderer/Private/SceneRendering.cpp
--- a/Engine/Source/Runtime/Renderer/Private/SceneRendering.cpp
+++ b/Engine/Source/Runtime/Renderer/Private/SceneRendering.cpp
@@ -3572,8 +3572,8 @@ void FSceneRenderer::OnRenderFinish(FRDGBuilder& GraphBuilder, FRDGTextureRef Vi
 			bNaniteEnabledButDisabledInProject = bNaniteEnabledButDisabledInProject || (WouldRenderNanite(Scene, View, /*bCheckForAtomicSupport*/ false, /*bCheckForProjectSetting*/ false) && !WouldRenderNanite(Scene, View, /*bCheckForAtomicSupport*/ false, /*bCheckForProjectSetting*/ true));
 
 			if (IsPostProcessingEnabled(View)
-				&& (!FMath::IsNearlyEqual(View.FinalPostProcessSettings.LocalExposureHighlightContrast, 1.0f)
-					|| !FMath::IsNearlyEqual(View.FinalPostProcessSettings.LocalExposureShadowContrast, 1.0f)
+				&& (!FMath::IsNearlyEqual(View.FinalPostProcessSettings.LocalExposureHighlightContrastScale, 1.0f)
+					|| !FMath::IsNearlyEqual(View.FinalPostProcessSettings.LocalExposureShadowContrastScale, 1.0f)
 					|| View.FinalPostProcessSettings.LocalExposureShadowContrastCurve
 					|| View.FinalPostProcessSettings.LocalExposureHighlightContrastCurve
 					|| !FMath::IsNearlyEqual(View.FinalPostProcessSettings.LocalExposureDetailStrength, 1.0f)))
