Expose event in UChildActorComponent to be notified when the ActorChild is created.
Change is needed for edge cases where FComponentReregisterContextBase is used after that the world was init. As an actor using UChildActorComponent will not be able to re-init values on the component child after that it was destroyed and re-constructed via FComponentReregisterContextBase.

diff --git a/Engine/Source/Runtime/Engine/Classes/Components/ChildActorComponent.h b/Engine/Source/Runtime/Engine/Classes/Components/ChildActorComponent.h
--- a/Engine/Source/Runtime/Engine/Classes/Components/ChildActorComponent.h
+++ b/Engine/Source/Runtime/Engine/Classes/Components/ChildActorComponent.h
@@ -192,6 +192,9 @@ public:
 	/** Create the child actor */
 	ENGINE_API virtual void CreateChildActor(TFunction<void(AActor*)> CustomizerFunc = nullptr);
 
+	DECLARE_EVENT_OneParam(UChildActorComponent, FOnChildActorCreated, AActor*);
+	FOnChildActorCreated& OnChildActorCreated() { return OnChildActorCreatedDelegate; }
+
 	AActor* GetChildActor() const { return ChildActor; }
 	AActor* GetChildActorTemplate() const { return ChildActorTemplate; }
 
@@ -231,6 +234,8 @@ private:
 
 	UFUNCTION()
 	ENGINE_API void OnChildActorDestroyed(AActor* Actor);
+
+	FOnChildActorCreated OnChildActorCreatedDelegate;
 };
 
 struct FActorParentComponentSetter
diff --git a/Engine/Source/Runtime/Engine/Private/Components/ChildActorComponent.cpp b/Engine/Source/Runtime/Engine/Private/Components/ChildActorComponent.cpp
--- a/Engine/Source/Runtime/Engine/Private/Components/ChildActorComponent.cpp
+++ b/Engine/Source/Runtime/Engine/Private/Components/ChildActorComponent.cpp
@@ -863,6 +863,11 @@ void UChildActorComponent::CreateChildActor(TFunction<void(AActor*)> CustomizerF
 		delete CachedInstanceData;
 		CachedInstanceData = nullptr;
 	}
+
+	if (ChildActor)
+	{
+		OnChildActorCreatedDelegate.Broadcast(ChildActor);
+	}
 }
 
 void UChildActorComponent::SetChildActorName(const FName InName)
