Made GameModeBase's GameNetDriverReplicationSystem config so it can be changed via ini

#rnx

diff --git a/Engine/Source/Runtime/Engine/Classes/GameFramework/GameModeBase.h b/Engine/Source/Runtime/Engine/Classes/GameFramework/GameModeBase.h
--- a/Engine/Source/Runtime/Engine/Classes/GameFramework/GameModeBase.h
+++ b/Engine/Source/Runtime/Engine/Classes/GameFramework/GameModeBase.h
@@ -597,7 +597,7 @@ protected:
 	 * Leave to Default to use the game engine's preferred system. 
 	 * Useful when migrating from one repsystem to another and a game mode does not fully support both repsystem yet.
 	 */
-	UPROPERTY(EditAnywhere, BlueprintReadWrite, Category = GameMode)
+	UPROPERTY(Config, EditAnywhere, BlueprintReadWrite, Category = GameMode)
 	EReplicationSystem GameNetDriverReplicationSystem = EReplicationSystem::Default;
 
 	/** The list of delegates to check before unpausing a game */
