Fixes sg.ResolutionQuality that can go nuts above 100

diff --git a/Engine/Source/Runtime/Engine/Private/Scalability.cpp b/Engine/Source/Runtime/Engine/Private/Scalability.cpp
--- a/Engine/Source/Runtime/Engine/Private/Scalability.cpp
+++ b/Engine/Source/Runtime/Engine/Private/Scalability.cpp
@@ -477,7 +477,14 @@ FText GetScalabilityNameFromQualityLevel(int32 QualityLevel)
 
 static void SetResolutionQualityLevel(float InResolutionQualityLevel)
 {
-	//InResolutionQualityLevel = FMath::Clamp(InResolutionQualityLevel, Scalability::MinResolutionScale, Scalability::MaxResolutionScale);
+	if (InResolutionQualityLevel == 0.0)
+	{
+		// NOP: just use the project's default screen percentage.
+	}
+	else
+	{
+		InResolutionQualityLevel = FMath::Clamp(InResolutionQualityLevel, Scalability::MinResolutionScale, Scalability::MaxResolutionScale);
+	}
 
 //	UE_LOG(LogConsoleResponse, Display, TEXT("  ResolutionQuality %.2f"), "", InResolutionQualityLevel);
 
