Iris
* The GameInstance can now request a replication system via GetDesiredReplicationSystem

diff --git a/Engine/Source/Runtime/Engine/Classes/Engine/GameInstance.h b/Engine/Source/Runtime/Engine/Classes/Engine/GameInstance.h
--- a/Engine/Source/Runtime/Engine/Classes/Engine/GameInstance.h
+++ b/Engine/Source/Runtime/Engine/Classes/Engine/GameInstance.h
@@ -646,6 +646,9 @@ public:
 	/** Remove a referenced object, this will allow it to GC out */
 	ENGINE_API virtual void UnregisterReferencedObject(UObject* ObjectToReference);
 
+	/** Allows a GameInstance to override the default configured replication system for it's specific conditions. Return Default to use the configured replication system. */
+	ENGINE_API virtual EReplicationSystem GetDesiredReplicationSystem(FName InNetDriverDefinition) const;
+
 protected:
 	/** Non-virtual dispatch for OnStart, also calls the associated global OnStartGameInstance. */
 	ENGINE_API void BroadcastOnStart();
diff --git a/Engine/Source/Runtime/Engine/Private/GameInstance.cpp b/Engine/Source/Runtime/Engine/Private/GameInstance.cpp
--- a/Engine/Source/Runtime/Engine/Private/GameInstance.cpp
+++ b/Engine/Source/Runtime/Engine/Private/GameInstance.cpp
@@ -22,6 +22,7 @@
 #include "Framework/Application/SlateApplication.h"
 #include "Misc/PackageName.h"
 #include "Net/ReplayPlaylistTracker.h"
+#include "Net/Core/Connection/NetEnums.h"
 #include "ReplaySubsystem.h"
 
 #include UE_INLINE_GENERATED_CPP_BY_NAME(GameInstance)
@@ -1429,6 +1430,25 @@ void UGameInstance::NotifyPreClientTravel(const FString& PendingURL, ETravelType
 	OnNotifyPreClientTravel().Broadcast(PendingURL, TravelType, bIsSeamlessTravel);
 }
 
+EReplicationSystem UGameInstance::GetDesiredReplicationSystem(FName InNetDriverDefinition) const
+{
+	EReplicationSystem DesiredRepSystem = EReplicationSystem::Default;
+
+	if (InNetDriverDefinition == NAME_GameNetDriver)
+	{
+		if (UWorld* World = GetWorld())
+		{
+			// If we are the server, return the game mode's desired replication system
+			if (AGameModeBase* ServerGameMode = World->GetAuthGameMode())
+			{
+				return ServerGameMode->GetGameNetDriverReplicationSystem();
+			}
+		}
+	}
+
+	return DesiredRepSystem;
+}
+
 void UGameInstance::ReturnToMainMenu()
 {
 	UWorld* const World = GetWorld();
diff --git a/Engine/Source/Runtime/Engine/Private/UnrealEngine.cpp b/Engine/Source/Runtime/Engine/Private/UnrealEngine.cpp
--- a/Engine/Source/Runtime/Engine/Private/UnrealEngine.cpp
+++ b/Engine/Source/Runtime/Engine/Private/UnrealEngine.cpp
@@ -13614,57 +13614,48 @@ namespace UE::Private
 
 		bool bUseIrisRepSystem = bConfigCanUseIris && bIsEngineDefaultIris;
 
-		if (InNetDriverDefinition == NAME_GameNetDriver)
+		if (UGameInstance* ContextGameInstance = Context.OwningGameInstance)
 		{
-			if (UWorld* World = Context.World())
-			{
-				// If we are the server look if the game mode wanted to override the replication system
-				if (AGameModeBase* GameMode = World->GetAuthGameMode())
-				{
-					EReplicationSystem RequestedRepSystem = GameMode->GetGameNetDriverReplicationSystem();
+			EReplicationSystem GameInstanceDesiredRepSystem = ContextGameInstance->GetDesiredReplicationSystem(InNetDriverDefinition);
 
-					// If the gamemode requested to use the generic repsystem
-					if (RequestedRepSystem == EReplicationSystem::Generic)
-					{
-						UE_CLOG(bUseIrisRepSystem, LogNet, Log, TEXT("GameMode %s is forcing NetDriver %s (NetDefinition %s) to use the Generic replication system."), *GetNameSafe(GameMode), *InNetDriverName.ToString(), *InNetDriverDefinition.ToString());
-						bUseIrisRepSystem = false;
-					}
-					// If the gamemode requested to use the Iris repsystem
-					else if (RequestedRepSystem == EReplicationSystem::Iris)
-					{
-						UE_CLOG(!bUseIrisRepSystem && bConfigCanUseIris, LogNet, Log, TEXT("GameMode %s is forcing NetDriver %s (NetDefinition %s) to use the Iris replication system."), *GetNameSafe(GameMode), *InNetDriverName.ToString(), *InNetDriverDefinition.ToString());
+			// If the game instance requested to use the generic repsystem
+			if (GameInstanceDesiredRepSystem == EReplicationSystem::Generic)
+			{
+				UE_CLOG(bUseIrisRepSystem, LogNet, Log, TEXT("GameInstance %s is forcing NetDriver %s (NetDefinition %s) to use the Generic replication system."), *GetNameSafe(ContextGameInstance), *InNetDriverName.ToString(), *InNetDriverDefinition.ToString());
+				bUseIrisRepSystem = false;
+			}
+			// If the game instance requested to use the Iris repsystem
+			else if (GameInstanceDesiredRepSystem == EReplicationSystem::Iris)
+			{
+				UE_CLOG(!bUseIrisRepSystem && bConfigCanUseIris, LogNet, Log, TEXT("GameInstance %s is forcing NetDriver %s (NetDefinition %s) to use the Iris replication system."), *GetNameSafe(ContextGameInstance), *InNetDriverName.ToString(), *InNetDriverDefinition.ToString());
+				
+				// Enable Iris ONLY if the config supports it.
+				bUseIrisRepSystem = bConfigCanUseIris;
+			}
 
-						// Enable Iris ONLY if the config supports it.
-						bUseIrisRepSystem = bConfigCanUseIris;
-					}
-				}
-				else
-				{
 #if WITH_EDITOR
-					// In PIE let's cheat and make sure the clients follow what the server's net driver is using
-					if (Context.WorldType == EWorldType::PIE)
+			// In PIE let's cheat and make sure the clients follow what the server's net driver is using
+			if (Context.WorldType == EWorldType::PIE && !Context.RunAsDedicated)
+			{
+				if (FWorldContext* ServerPIEContext = GEngine->GetWorldContextFromPIEInstance(0))
+				{
+					if (ServerPIEContext->RunAsDedicated)
 					{
-						if (FWorldContext* ServerPIEContext = GEngine->GetWorldContextFromPIEInstance(0))
+						for (const FNamedNetDriver& PieNetDriver : ServerPIEContext->ActiveNetDrivers)
 						{
-							if (ServerPIEContext->RunAsDedicated)
+							if (PieNetDriver.NetDriverDef->DefName == NAME_GameNetDriver)
 							{
-								for (const FNamedNetDriver& PieNetDriver : ServerPIEContext->ActiveNetDrivers)
-								{
-									if (PieNetDriver.NetDriverDef->DefName == NAME_GameNetDriver)
-									{
-										bUseIrisRepSystem = PieNetDriver.NetDriver->IsUsingIrisReplication();
-										break;
-									}
-								}
+								bUseIrisRepSystem = PieNetDriver.NetDriver->IsUsingIrisReplication();
+								break;
 							}
 						}
 					}
-#endif
 				}
 			}
+#endif
 		}
 
-		// Ignore most of the above if the cmdline is requesting a specific system
+		// Ignore all of the above if the cmdline is requesting a specific system
 		const EReplicationSystem CmdlineRequest = UE::Net::GetUseIrisReplicationCmdlineValue();
 		if (CmdlineRequest == EReplicationSystem::Iris)
 		{
diff --git a/Engine/Source/Runtime/Experimental/Iris/Core/Private/Iris/IrisCoreModule.cpp b/Engine/Source/Runtime/Experimental/Iris/Core/Private/Iris/IrisCoreModule.cpp
--- a/Engine/Source/Runtime/Experimental/Iris/Core/Private/Iris/IrisCoreModule.cpp
+++ b/Engine/Source/Runtime/Experimental/Iris/Core/Private/Iris/IrisCoreModule.cpp
@@ -4,13 +4,16 @@
 #include "Modules/ModuleInterface.h"
 #include "Modules/ModuleManager.h"
 #include "Misc/CoreDelegates.h"
+#include "Misc/CommandLine.h"
+
+#include "Net/Core/Connection/NetEnums.h"
+
 #include "Iris/ReplicationState/DefaultPropertyNetSerializerInfos.h"
 #include "Iris/ReplicationState/PropertyNetSerializerInfoRegistry.h"
 #include "Iris/ReplicationSystem/LegacyPushModel.h"
 #include "Iris/Serialization/InternalNetSerializerDelegates.h"
 #include "Iris/IrisConfigInternal.h"
 
-#include "Misc/CommandLine.h"
 
 class FIrisCoreModule : public IModuleInterface
 {
@@ -38,10 +41,11 @@ private:
 		FCoreDelegates::OnAllModuleLoadingPhasesComplete.AddRaw(this, &FIrisCoreModule::OnAllModuleLoadingPhasesComplete);
 
 		// Check command line for whether we should override the net.Iris.UseIrisReplication cvar, as we need to do that early
-		int32 UseIrisReplication;
-		if(FParse::Value(FCommandLine::Get(), TEXT("UseIrisReplication="), UseIrisReplication))
+		const EReplicationSystem CmdlineRepSystem = UE::Net::GetUseIrisReplicationCmdlineValue();
+		if (CmdlineRepSystem != EReplicationSystem::Default)
 		{
-			UE::Net::SetUseIrisReplication(UseIrisReplication > 0);
+			const bool bEnableIris = CmdlineRepSystem == EReplicationSystem::Iris;
+			UE::Net::SetUseIrisReplication(bEnableIris);
 		}
 
 		RegisterPropertyNetSerializerSelectorTypes();
