Fixing log spam in attenuation settings.

diff --git a/Engine/Source/Runtime/Engine/Private/ActiveSound.cpp b/Engine/Source/Runtime/Engine/Private/ActiveSound.cpp
--- a/Engine/Source/Runtime/Engine/Private/ActiveSound.cpp
+++ b/Engine/Source/Runtime/Engine/Private/ActiveSound.cpp
@@ -1661,7 +1661,8 @@ void FActiveSound::UpdateAttenuation(float DeltaTime, FSoundParseParameters& Par
 
 void FActiveSound::UpdateAttenuation(float DeltaTime, FSoundParseParameters& ParseParams, int32 ListenerIndex, const FSoundAttenuationSettings* SettingsAttenuationNode)
 {
-	const FSoundAttenuationSettings* Settings = nullptr;
+	// We default to using the copied off "overridden" settings (or default constructed settings)
+	const FSoundAttenuationSettings* Settings = &AttenuationSettings;
 
 	// Get the attenuation settings to use for this application to the active sound
 	// Use the passed-in attenuation settings
@@ -1669,23 +1670,12 @@ void FActiveSound::UpdateAttenuation(float DeltaTime, FSoundParseParameters& Par
 	{
 		Settings = SettingsAttenuationNode;
 	}
-	// We use the copied off "overridden" settings
-	else if (bIsAttenuationSettingsOverridden)
-	{
-		Settings = &AttenuationSettings;
-	}
 	// We fallback to using the asset's settings directly
 	else if (SoundAttenuation)
 	{
 		Settings = &SoundAttenuation->Attenuation;
 	}
 
-	if (!Settings)
-	{
-		UE_LOG(LogAudio, Warning, TEXT("No attenuation settings found for active sound."));
-		return;
-	}
-
 	// Reset Focus data and recompute if necessary
 	FAttenuationFocusData FocusDataToApply;
 	FocusDataToApply.PriorityHighest = FocusData.PriorityHighest;
