Fixed cases where PSO precaching was providing incorrect state in non-MVF path

diff --git a/Engine/Plugins/Experimental/Water/Source/Runtime/Private/WaterBodyComponent.cpp b/Engine/Plugins/Experimental/Water/Source/Runtime/Private/WaterBodyComponent.cpp
--- a/Engine/Plugins/Experimental/Water/Source/Runtime/Private/WaterBodyComponent.cpp
+++ b/Engine/Plugins/Experimental/Water/Source/Runtime/Private/WaterBodyComponent.cpp
@@ -31,6 +31,7 @@
 #include "LocalVertexFactory.h"
 #include "UObject/ICookInfo.h"
 #include "UObject/ObjectSaveContext.h"
+#include "DataDrivenShaderPlatformInfo.h"
 
 #include UE_INLINE_GENERATED_CPP_BY_NAME(WaterBodyComponent)
 
@@ -1605,7 +1606,9 @@ void UWaterBodyComponent::PostLoad()
 
 	DeprecateData();
 
-	if (IsComponentPSOPrecachingEnabled())
+	if (IsComponentPSOPrecachingEnabled()
+		// FIXME: need to collect an actual vertex declaration for non-MVF path
+		&& RHISupportsManualVertexFetch(GMaxRHIShaderPlatform))
 	{
 		FPSOPrecacheParams PrecachePSOParams;
 		SetupPrecachePSOParams(PrecachePSOParams);
diff --git a/Engine/Plugins/FX/Niagara/Source/NiagaraVertexFactories/Private/NiagaraMeshVertexFactory.cpp b/Engine/Plugins/FX/Niagara/Source/NiagaraVertexFactories/Private/NiagaraMeshVertexFactory.cpp
--- a/Engine/Plugins/FX/Niagara/Source/NiagaraVertexFactories/Private/NiagaraMeshVertexFactory.cpp
+++ b/Engine/Plugins/FX/Niagara/Source/NiagaraVertexFactories/Private/NiagaraMeshVertexFactory.cpp
@@ -199,9 +199,14 @@ void FNiagaraMeshVertexFactory::ModifyCompilationEnvironment(const FVertexFactor
 
 void FNiagaraMeshVertexFactory::GetPSOPrecacheVertexFetchElements(EVertexInputStreamType VertexInputStreamType, FVertexDeclarationElementList& Elements)
 {
-	Elements.Add(FVertexElement(0, 0, VET_Float3, 0, 0, false));
+	Elements.Add(FVertexElement(0, 0, VET_Float3, 0, sizeof(float)*3u, false));
+
 #if NIAGARA_ENABLE_GPU_SCENE_MESHES
-	Elements.Add(FVertexElement(1, 0, VET_UInt, 13, 0, true));
+	if (UseGPUScene(GMaxRHIShaderPlatform, GMaxRHIFeatureLevel)
+		&& !PlatformGPUSceneUsesUniformBufferView(GMaxRHIShaderPlatform))
+	{
+		Elements.Add(FVertexElement(1, 0, VET_UInt, 13, sizeof(uint32), true));
+	}
 #endif
 }
 
@@ -211,7 +216,8 @@ void FNiagaraMeshVertexFactory::GetVertexElements(ERHIFeatureLevel::Type Feature
 	GetVertexElements(FeatureLevel, bSupportsManualVertexFetch, Data, Elements, InOutStreams);
 
 #if NIAGARA_ENABLE_GPU_SCENE_MESHES
-	if (UseGPUScene(GMaxRHIShaderPlatform, GMaxRHIFeatureLevel))
+	if (UseGPUScene(GMaxRHIShaderPlatform, GMaxRHIFeatureLevel)
+		&& !PlatformGPUSceneUsesUniformBufferView(GMaxRHIShaderPlatform))
 	{
 		Elements.Add(FVertexElement(InOutStreams.Num(), 0, VET_UInt, 13, sizeof(uint32), true));
 	}
diff --git a/Engine/Source/Runtime/Engine/Private/Components/MaterialBillboardComponent.cpp b/Engine/Source/Runtime/Engine/Private/Components/MaterialBillboardComponent.cpp
--- a/Engine/Source/Runtime/Engine/Private/Components/MaterialBillboardComponent.cpp
+++ b/Engine/Source/Runtime/Engine/Private/Components/MaterialBillboardComponent.cpp
@@ -14,6 +14,7 @@
 #include "PrimitiveSceneProxy.h"
 #include "StaticMeshResources.h"
 #include "PSOPrecache.h"
+#include "DataDrivenShaderPlatformInfo.h"
 
 #include UE_INLINE_GENERATED_CPP_BY_NAME(MaterialBillboardComponent)
 
@@ -372,7 +373,9 @@ void UMaterialBillboardComponent::PostLoad()
 {
 	Super::PostLoad();
 
-	if (IsComponentPSOPrecachingEnabled())
+	if (IsComponentPSOPrecachingEnabled()
+		// FIXME: need to collect an actual vertex declaration for non-MVF path
+		&& RHISupportsManualVertexFetch(GMaxRHIShaderPlatform))
 	{
 		FPSOPrecacheParams PrecachePSOParams;
 		SetupPrecachePSOParams(PrecachePSOParams);
diff --git a/Engine/Source/Runtime/Engine/Private/Components/StaticMeshComponent.cpp b/Engine/Source/Runtime/Engine/Private/Components/StaticMeshComponent.cpp
--- a/Engine/Source/Runtime/Engine/Private/Components/StaticMeshComponent.cpp
+++ b/Engine/Source/Runtime/Engine/Private/Components/StaticMeshComponent.cpp
@@ -1566,10 +1566,11 @@ void UStaticMeshComponent::CollectPSOPrecacheDataImpl(
 
 	bool bSupportsManualVertexFetch = VFType->SupportsManualVertexFetch(GMaxRHIFeatureLevel);
 	bool bAnySectionCastsShadows = false;
+	int32 MeshMinLOD = GetStaticMesh()->GetMinLODIdx();
 
 	FPSOPrecacheVertexFactoryDataPerMaterialIndexList VFTypesPerMaterialIndex;
 	FStaticMeshLODResourcesArray& LODResources = GetStaticMesh()->GetRenderData()->LODResources;
-	for (int32 LODIndex = 0; LODIndex < LODResources.Num(); ++LODIndex)
+	for (int32 LODIndex = MeshMinLOD; LODIndex < LODResources.Num(); ++LODIndex)
 	{
 		FStaticMeshLODResources& LODRenderData = LODResources[LODIndex];
 		FVertexDeclarationElementList VertexElements;
diff --git a/Engine/Source/Runtime/Engine/Private/Components/TextRenderComponent.cpp b/Engine/Source/Runtime/Engine/Private/Components/TextRenderComponent.cpp
--- a/Engine/Source/Runtime/Engine/Private/Components/TextRenderComponent.cpp
+++ b/Engine/Source/Runtime/Engine/Private/Components/TextRenderComponent.cpp
@@ -21,6 +21,7 @@
 #include "StaticMeshResources.h"
 #include "Containers/Ticker.h"
 #include "PSOPrecache.h"
+#include "DataDrivenShaderPlatformInfo.h"
 
 static TAutoConsoleVariable<int32> CVarRayTracingTextMeshes(
 	TEXT("r.RayTracing.Geometry.Text"),
@@ -1518,7 +1519,9 @@ void UTextRenderComponent::PostLoad()
 		}
 	}
 
-	if (IsComponentPSOPrecachingEnabled() && TextMaterial)
+	if (IsComponentPSOPrecachingEnabled() && TextMaterial
+		// FIXME: need to collect an actual vertex declaration for non-MVF path
+		&& RHISupportsManualVertexFetch(GMaxRHIShaderPlatform))
 	{
 		FPSOPrecacheParams PrecachePSOParams;
 		SetupPrecachePSOParams(PrecachePSOParams);
diff --git a/Engine/Source/Runtime/Engine/Private/GeometryCollectionRendering.cpp b/Engine/Source/Runtime/Engine/Private/GeometryCollectionRendering.cpp
--- a/Engine/Source/Runtime/Engine/Private/GeometryCollectionRendering.cpp
+++ b/Engine/Source/Runtime/Engine/Private/GeometryCollectionRendering.cpp
@@ -83,28 +83,37 @@ void FGeometryCollectionVertexFactory::ValidateCompiledResult(const FVertexFacto
 
 void FGeometryCollectionVertexFactory::GetPSOPrecacheVertexFetchElements(EVertexInputStreamType VertexInputStreamType, FVertexDeclarationElementList& Elements)
 {
-	Elements.Add(FVertexElement(0, 0, VET_Float3, 0, 0, false));
+	Elements.Add(FVertexElement(0, 0, VET_Float3, 0, sizeof(float)*3u, false));
 
-	switch (VertexInputStreamType)
+	if (VertexInputStreamType == EVertexInputStreamType::PositionAndNormalOnly)
 	{
-	case EVertexInputStreamType::Default:
-	{
-		Elements.Add(FVertexElement(1, 0, VET_UInt, 13, 0, true));
-		break;
-	}
-	case EVertexInputStreamType::PositionOnly:
-	{
-		Elements.Add(FVertexElement(1, 0, VET_UInt, 1, 0, true));
-		break;
+		// 2-axis TangentBasis components in a single buffer, hence *2u
+		Elements.Add(FVertexElement(1, 0, VET_PackedNormal, 2, sizeof(FPackedNormal)*2u, false));
 	}
-	case EVertexInputStreamType::PositionAndNormalOnly:
+
+	if (UseGPUScene(GMaxRHIShaderPlatform, GMaxRHIFeatureLevel)
+		&& !PlatformGPUSceneUsesUniformBufferView(GMaxRHIShaderPlatform))
 	{
-		Elements.Add(FVertexElement(1, 0, VET_PackedNormal, 2, 0, false));
-		Elements.Add(FVertexElement(2, 0, VET_UInt, 1, 0, true));
-		break;
-	}
-	default:
-		checkNoEntry();
+		switch (VertexInputStreamType)
+		{
+		case EVertexInputStreamType::Default:
+		{
+			Elements.Add(FVertexElement(1, 0, VET_UInt, 13, sizeof(uint32), true));
+			break;
+		}
+		case EVertexInputStreamType::PositionOnly:
+		{
+			Elements.Add(FVertexElement(1, 0, VET_UInt, 1, sizeof(uint32), true));
+			break;
+		}
+		case EVertexInputStreamType::PositionAndNormalOnly:
+		{
+			Elements.Add(FVertexElement(2, 0, VET_UInt, 1, sizeof(uint32), true));
+			break;
+		}
+		default:
+			checkNoEntry();
+		}
 	}
 }
 
diff --git a/Engine/Source/Runtime/Engine/Private/SceneTexturesConfig.cpp b/Engine/Source/Runtime/Engine/Private/SceneTexturesConfig.cpp
--- a/Engine/Source/Runtime/Engine/Private/SceneTexturesConfig.cpp
+++ b/Engine/Source/Runtime/Engine/Private/SceneTexturesConfig.cpp
@@ -214,11 +214,16 @@ static uint32 GetEditorPrimitiveNumSamples(ERHIFeatureLevel::Type FeatureLevel)
 static void SetupMobileGBufferFlags(FGBufferBindings GBufferBindings[GBL_Num], bool bRequiresMultiPass)
 {
 	ETextureCreateFlags AddFlags = TexCreate_InputAttachmentRead;
+	
+	// Mobile uses FBF/subpassLoad to fetch data from GBuffer, and FBF does not always work with sRGB targets 
+	ETextureCreateFlags RemoveFlags = TexCreate_SRGB;
 
 	if (!bRequiresMultiPass)
 	{
 		// use memoryless GBuffer when possible
 		AddFlags |= TexCreate_Memoryless;
+		// memoryless GBuffer cant be used with compute
+		RemoveFlags |= TexCreate_UAV;
 	}
 
 	for (uint32 Layout = 0; Layout < GBL_Num; ++Layout)
@@ -230,13 +235,12 @@ static void SetupMobileGBufferFlags(FGBufferBindings GBufferBindings[GBL_Num], b
 		Bindings.GBufferC.Flags |= AddFlags;
 		Bindings.GBufferD.Flags |= AddFlags;
 		Bindings.GBufferE.Flags |= AddFlags;
-
-		// Mobile uses FBF/subpassLoad to fetch data from GBuffer, and FBF does not always work with sRGB targets 
-		Bindings.GBufferA.Flags &= (~TexCreate_SRGB);
-		Bindings.GBufferB.Flags &= (~TexCreate_SRGB);
-		Bindings.GBufferC.Flags &= (~TexCreate_SRGB);
-		Bindings.GBufferD.Flags &= (~TexCreate_SRGB);
-		Bindings.GBufferE.Flags &= (~TexCreate_SRGB);
+		
+		Bindings.GBufferA.Flags &= (~RemoveFlags);
+		Bindings.GBufferB.Flags &= (~RemoveFlags);
+		Bindings.GBufferC.Flags &= (~RemoveFlags);
+		Bindings.GBufferD.Flags &= (~RemoveFlags);
+		Bindings.GBufferE.Flags &= (~RemoveFlags);
 
 		// Input attachments with PF_R8G8B8A8 has better support on mobile than PF_B8G8R8A8
 		auto OverrideB8G8R8A8 = [](FGBufferBinding& Binding) { if (Binding.Format == PF_B8G8R8A8) Binding.Format = PF_R8G8B8A8; };
diff --git a/Engine/Source/Runtime/Landscape/Private/LandscapeRender.cpp b/Engine/Source/Runtime/Landscape/Private/LandscapeRender.cpp
--- a/Engine/Source/Runtime/Landscape/Private/LandscapeRender.cpp
+++ b/Engine/Source/Runtime/Landscape/Private/LandscapeRender.cpp
@@ -3347,7 +3347,8 @@ void FLandscapeVertexFactory::GetPSOPrecacheVertexFetchElements(EVertexInputStre
 {
 	Elements.Add(FVertexElement(0, 0, VET_UByte4, 0, sizeof(FLandscapeVertex), false));
 	
-	if (UseGPUScene(GMaxRHIShaderPlatform, GMaxRHIFeatureLevel))
+	if (UseGPUScene(GMaxRHIShaderPlatform, GMaxRHIFeatureLevel)
+		&& !PlatformGPUSceneUsesUniformBufferView(GMaxRHIShaderPlatform))
 	{
 		Elements.Add(FVertexElement(1, 0, VET_UInt, 1, sizeof(uint32), true));
 	}
