Crash fixed:
- Removed stack heavy recursion
- Make sure that correct actor is ignored to prevent infinite loop when hitting a cluster union

#[fyi] Michael.Bao

diff --git a/Engine/Source/Runtime/Engine/Private/Collision/SceneQuery.cpp b/Engine/Source/Runtime/Engine/Private/Collision/SceneQuery.cpp
--- a/Engine/Source/Runtime/Engine/Private/Collision/SceneQuery.cpp
+++ b/Engine/Source/Runtime/Engine/Private/Collision/SceneQuery.cpp
@@ -417,9 +417,10 @@ private:
 };
 
 template <typename Traits, typename TGeomInputs, typename TAccelContainer>
-bool TSceneCastCommonImp(const UWorld* World, typename Traits::TOutHits& OutHits, const TGeomInputs& GeomInputs, const FVector Start, const FVector End, ECollisionChannel TraceChannel, const struct FCollisionQueryParams& Params, const struct FCollisionResponseParams& ResponseParams, const struct FCollisionObjectQueryParams& ObjectParams, const TAccelContainer& AccelContainer)
+bool TSceneCastCommonImpWithRetryRequest(const UWorld* World, typename Traits::TOutHits& OutHits, const TGeomInputs& GeomInputs, const FVector Start, const FVector End, ECollisionChannel TraceChannel, const struct FCollisionQueryParams& Params, const struct FCollisionResponseParams& ResponseParams, const struct FCollisionObjectQueryParams& ObjectParams, const TAccelContainer& AccelContainer, bool& bOutRequestRetry, FCollisionQueryParams& OutRetryParams)
 {
 	using namespace ChaosInterface;
+	bOutRequestRetry = false;
 
 	FScopeCycleCounter Counter(Params.StatId);
 	STARTQUERYTIMER();
@@ -584,9 +585,16 @@ bool TSceneCastCommonImp(const UWorld* World, typename Traits::TOutHits& OutHits
 								// redo the entire trace and force the SQ trace to ignore the cluster union actor. This is only relevant in the case where we aren't doing a multi-trace since in the
 								// case of a multi-trace, we would've found all the other things we hit as well. In the case of a non-multi (single) trace, the SQ trace determined that the cluster union
 								// is the best hit! But there could be other things we could've hit instead if we ignored the fact that we hit a cluster union.
-								FCollisionQueryParams NewParams = Params;
-								NewParams.AddIgnoredActor(ClusterUnionActor);
-								return TSceneCastCommonImp<Traits, TGeomInputs, TAccelContainer>(World, OutHits, GeomInputs, Start, End, TraceChannel, NewParams, ResponseParams, ObjectParams, AccelContainer);
+								OutRetryParams = Params;
+								// Ignore the actor that was hit (Check the shape data to be sure)
+								uint32 ActorIDFromShape = Traits::GetHits(HitBufferSync)->Shape->GetQueryData().Word0;
+								OutRetryParams.AddIgnoredActor(ActorIDFromShape);
+								if (ActorIDFromShape != ClusterUnionActor->GetUniqueID())
+								{
+									UE_LOG(LogChaos, Warning, TEXT("TSceneCastCommonImpWithRetryRequest: Incorrect Shape Actor ID detected"));
+								}								
+								bOutRequestRetry = true;
+								return false;
 							}
 							else
 							{
@@ -616,6 +624,31 @@ bool TSceneCastCommonImp(const UWorld* World, typename Traits::TOutHits& OutHits
 	return bHaveBlockingHit;
 }
 
+template <typename Traits, typename TGeomInputs, typename TAccelContainer>
+bool TSceneCastCommonImp(const UWorld* World, typename Traits::TOutHits& OutHits, const TGeomInputs& GeomInputs, const FVector Start, const FVector End, ECollisionChannel TraceChannel, const struct FCollisionQueryParams& Params, const struct FCollisionResponseParams& ResponseParams, const struct FCollisionObjectQueryParams& ObjectParams, const TAccelContainer& AccelContainer)
+{
+	bool bRequestRetry = true;
+	bool bReturnResult = false;
+	FCollisionQueryParams RetryParams;
+
+	bReturnResult = TSceneCastCommonImpWithRetryRequest<Traits, TGeomInputs, TAccelContainer>(World, OutHits, GeomInputs, Start, End, TraceChannel, Params, ResponseParams, ObjectParams, AccelContainer, bRequestRetry, RetryParams);	
+	
+	int InfiniteLoopProtection = 10;
+	while (bRequestRetry && InfiniteLoopProtection > 0)
+	{
+		bReturnResult = TSceneCastCommonImpWithRetryRequest<Traits, TGeomInputs, TAccelContainer>(World, OutHits, GeomInputs, Start, End, TraceChannel, RetryParams, ResponseParams, ObjectParams, AccelContainer, bRequestRetry, RetryParams);
+		InfiniteLoopProtection--;
+	}
+
+	if (InfiniteLoopProtection <= 0)
+	{
+		UE_LOG(LogChaos, Warning, TEXT("TSceneCastCommonImp: Potential Infinite Loop Detected"));
+		bReturnResult = false;
+	}
+
+	return bReturnResult;
+}
+
 template <typename Traits, typename PTTraits, typename TGeomInputs, typename TAccelContainer = FDefaultAccelContainer>
 bool TSceneCastCommon(const UWorld* World, typename Traits::TOutHits& OutHits, const TGeomInputs& GeomInputs, const FVector Start, const FVector End, ECollisionChannel TraceChannel, const struct FCollisionQueryParams& Params, const struct FCollisionResponseParams& ResponseParams, const struct FCollisionObjectQueryParams& ObjectParams, const TAccelContainer& AccelContainer = FDefaultAccelContainer{})
 {
@@ -1133,4 +1166,4 @@ template struct FGenericGeomPhysicsInterfaceUsingSpatialAcceleration<TACCEL, FPh
 template struct FGenericGeomPhysicsInterfaceUsingSpatialAcceleration<TACCEL, FPhysicsGeometryCollection>;
 
 DEFINE_GENERIC_GEOM_FOR_ACCEL(Chaos::IDefaultChaosSpatialAcceleration)
-DEFINE_GENERIC_GEOM_FOR_ACCEL(IExternalSpatialAcceleration)
\ No newline at end of file
+DEFINE_GENERIC_GEOM_FOR_ACCEL(IExternalSpatialAcceleration)
diff --git a/Engine/Source/Runtime/Engine/Private/PhysicsEngine/ClusterUnionComponent.cpp b/Engine/Source/Runtime/Engine/Private/PhysicsEngine/ClusterUnionComponent.cpp
--- a/Engine/Source/Runtime/Engine/Private/PhysicsEngine/ClusterUnionComponent.cpp
+++ b/Engine/Source/Runtime/Engine/Private/PhysicsEngine/ClusterUnionComponent.cpp
@@ -521,6 +521,12 @@ void UClusterUnionComponent::ForceRebuildGTParticleGeometry()
 		}
 	}
 
+	// TODO: Make sure that the empty cluster union does not get added to the acceleration structure
+	if (Particles.Num() == 0)
+	{
+		return;
+	}	
+
 	Chaos::FImplicitObjectUnion* NewGeometry = Objects.IsEmpty() ? new Chaos::FImplicitObjectUnionClustered() : new Chaos::FImplicitObjectUnion(MoveTemp(Objects));
 	NewGeometry->SetAllowBVH(true);
 
