Open creating the RepState but then if we abort set it to a nullptr cleaning things up

The default value of RepState is !IsValid so setting back to a nullptr is correct for RepState

diff --git a/Engine/Source/Runtime/Engine/Private/DataReplication.cpp b/Engine/Source/Runtime/Engine/Private/DataReplication.cpp
--- a/Engine/Source/Runtime/Engine/Private/DataReplication.cpp
+++ b/Engine/Source/Runtime/Engine/Private/DataReplication.cpp
@@ -399,7 +399,14 @@ void FObjectReplicator::InitRecentProperties(uint8* Source)
 	// If acting as a server and are IsInternalAck, that means we're recording.
 	// In that case, we don't need to create any receiving state, as no one will be sending data to us.
 	ECreateRepStateFlags Flags = (Connection->IsInternalAck() && bIsServer) ? ECreateRepStateFlags::SkipCreateReceivingState : ECreateRepStateFlags::None;
-	RepState = LocalRepLayout.CreateRepState(Source, RepChangedPropertyTracker, Flags);
+	UE_AUTORTFM_OPEN({
+		RepState = LocalRepLayout.CreateRepState(Source, RepChangedPropertyTracker, Flags);
+	});
+
+	// RepState not valid at the start of this function so just go back to being a nullptr, and let the memory be cleaned up 
+	UE_AUTORTFM_OPENABORT({
+		RepState = nullptr;
+	});
 
 	if (!bCreateSendingState)
 	{
