Add OPEN around FRepLayout::CreateFromFunction and store in a temp pointer so if an abort happens we restore the previous value to that member variable

diff --git a/Engine/Source/Runtime/Engine/Private/NetDriver.cpp b/Engine/Source/Runtime/Engine/Private/NetDriver.cpp
--- a/Engine/Source/Runtime/Engine/Private/NetDriver.cpp
+++ b/Engine/Source/Runtime/Engine/Private/NetDriver.cpp
@@ -6676,7 +6676,14 @@ TSharedPtr<FRepLayout> UNetDriver::GetFunctionRepLayout(UFunction * Function)
 	if (!RepLayoutPtr) 
 	{
 		ECreateRepLayoutFlags Flags = MaySendProperties() ? ECreateRepLayoutFlags::MaySendProperties : ECreateRepLayoutFlags::None;
-		RepLayoutPtr = &RepLayoutMap.Add(Function, FRepLayout::CreateFromFunction(Function, ServerConnection, Flags));
+
+		// Use a temp shared ptr to hold onto a previous value in case we abort. The adding to RepLayoutMap copies this TSharedPtr, and will end up taking a ref count and owning it
+		TSharedPtr<FRepLayout> NewLayoutPtr;
+		UE_AUTORTFM_OPEN({
+			NewLayoutPtr = FRepLayout::CreateFromFunction(Function, ServerConnection, Flags);
+		});
+
+		RepLayoutPtr = &RepLayoutMap.Add(Function, NewLayoutPtr);
 	}
 
 	return *RepLayoutPtr;
