If we Abort when setting the collision type, set it back to the value it was before the we entered

diff --git a/Engine/Source/Runtime/Engine/Private/PrimitiveComponentPhysics.cpp b/Engine/Source/Runtime/Engine/Private/PrimitiveComponentPhysics.cpp
--- a/Engine/Source/Runtime/Engine/Private/PrimitiveComponentPhysics.cpp
+++ b/Engine/Source/Runtime/Engine/Private/PrimitiveComponentPhysics.cpp
@@ -1071,9 +1071,18 @@ void UPrimitiveComponent::SetCollisionResponseToChannels(const FCollisionRespons
 
 void UPrimitiveComponent::SetCollisionEnabled(ECollisionEnabled::Type NewType)
 {
-	if (BodyInstance.GetCollisionEnabled() != NewType)
+	ECollisionEnabled::Type CurrentType = BodyInstance.GetCollisionEnabled();
+
+	if (CurrentType != NewType)
 	{
-		BodyInstance.SetCollisionEnabled(NewType);
+		UE_AUTORTFM_OPEN({
+			BodyInstance.SetCollisionEnabled(NewType);
+		});
+
+		// If we fail set the CollisionEnabled back to the CurrentType
+		UE_AUTORTFM_OPENABORT({
+			BodyInstance.SetCollisionEnabled(CurrentType);
+		});
 
 		EnsurePhysicsStateCreated();
 		OnComponentCollisionSettingsChanged();
