WIP New HLSL translator Material Expressions & bug fixes
Fixed new/old parity issue in texture/texel size preshaders
Disable automatic view mip bias cases to match old translator
Implement MaterialExpressionSphericalParticleOpacity, MaterialExpressionMapARPassthroughCameraUV, MaterialExpressionGIReplace, MaterialExpressionLightmassReplace

#tests materials, editor, PIE, win64

diff --git a/Engine/Source/Runtime/Engine/Classes/Materials/MaterialExpressionGIReplace.h b/Engine/Source/Runtime/Engine/Classes/Materials/MaterialExpressionGIReplace.h
--- a/Engine/Source/Runtime/Engine/Classes/Materials/MaterialExpressionGIReplace.h
+++ b/Engine/Source/Runtime/Engine/Classes/Materials/MaterialExpressionGIReplace.h
@@ -30,6 +30,7 @@ class UMaterialExpressionGIReplace : public UMaterialExpression
 #if WITH_EDITOR
 	virtual int32 Compile(class FMaterialCompiler* Compiler, int32 OutputIndex) override;
 	virtual void GetCaption(TArray<FString>& OutCaptions) const override;
+	virtual bool GenerateHLSLExpression(FMaterialHLSLGenerator& Generator, UE::HLSLTree::FScope& Scope, int32 OutputIndex, UE::HLSLTree::FExpression const*& OutExpression) const override;
 #endif
 	//~ End UMaterialExpression Interface
 };
diff --git a/Engine/Source/Runtime/Engine/Classes/Materials/MaterialExpressionLightmassReplace.h b/Engine/Source/Runtime/Engine/Classes/Materials/MaterialExpressionLightmassReplace.h
--- a/Engine/Source/Runtime/Engine/Classes/Materials/MaterialExpressionLightmassReplace.h
+++ b/Engine/Source/Runtime/Engine/Classes/Materials/MaterialExpressionLightmassReplace.h
@@ -25,6 +25,7 @@ class UMaterialExpressionLightmassReplace : public UMaterialExpression
 #if WITH_EDITOR
 	virtual int32 Compile(class FMaterialCompiler* Compiler, int32 OutputIndex) override;
 	virtual void GetCaption(TArray<FString>& OutCaptions) const override;
+	virtual bool GenerateHLSLExpression(FMaterialHLSLGenerator& Generator, UE::HLSLTree::FScope& Scope, int32 OutputIndex, UE::HLSLTree::FExpression const*& OutExpression) const override;
 #endif
 	//~ End UMaterialExpression Interface
 };
diff --git a/Engine/Source/Runtime/Engine/Classes/Materials/MaterialExpressionMapARPassthroughCameraUV.h b/Engine/Source/Runtime/Engine/Classes/Materials/MaterialExpressionMapARPassthroughCameraUV.h
--- a/Engine/Source/Runtime/Engine/Classes/Materials/MaterialExpressionMapARPassthroughCameraUV.h
+++ b/Engine/Source/Runtime/Engine/Classes/Materials/MaterialExpressionMapARPassthroughCameraUV.h
@@ -24,6 +24,7 @@ class UMaterialExpressionMapARPassthroughCameraUV : public UMaterialExpression
 #if WITH_EDITOR
 	virtual int32 Compile(class FMaterialCompiler* Compiler, int32 OutputIndex) override;
 	virtual void GetCaption(TArray<FString>& OutCaptions) const override;
+	virtual bool GenerateHLSLExpression(FMaterialHLSLGenerator& Generator, UE::HLSLTree::FScope& Scope, int32 OutputIndex, UE::HLSLTree::FExpression const*& OutExpression) const override;
 #endif
 	//~ End UMaterialExpression Interface
 };
diff --git a/Engine/Source/Runtime/Engine/Classes/Materials/MaterialExpressionSphericalParticleOpacity.h b/Engine/Source/Runtime/Engine/Classes/Materials/MaterialExpressionSphericalParticleOpacity.h
--- a/Engine/Source/Runtime/Engine/Classes/Materials/MaterialExpressionSphericalParticleOpacity.h
+++ b/Engine/Source/Runtime/Engine/Classes/Materials/MaterialExpressionSphericalParticleOpacity.h
@@ -30,6 +30,7 @@ class UMaterialExpressionSphericalParticleOpacity : public UMaterialExpression
 	{
 		OutCaptions.Add(TEXT("Spherical Particle Opacity"));
 	}
+	virtual bool GenerateHLSLExpression(FMaterialHLSLGenerator& Generator, UE::HLSLTree::FScope& Scope, int32 OutputIndex, UE::HLSLTree::FExpression const*& OutExpression) const override;
 #endif
 	//~ End UMaterialExpression Interface
 };
diff --git a/Engine/Source/Runtime/Engine/Private/Materials/MaterialExpressionHLSL.cpp b/Engine/Source/Runtime/Engine/Private/Materials/MaterialExpressionHLSL.cpp
--- a/Engine/Source/Runtime/Engine/Private/Materials/MaterialExpressionHLSL.cpp
+++ b/Engine/Source/Runtime/Engine/Private/Materials/MaterialExpressionHLSL.cpp
@@ -94,6 +94,7 @@
 #include "Materials/MaterialExpressionGenericConstant.h"
 #include "Materials/MaterialExpressionGetLocal.h"
 #include "Materials/MaterialExpressionGetMaterialAttributes.h"
+#include "Materials/MaterialExpressionGIReplace.h"
 #include "Materials/MaterialExpressionHairAttributes.h"
 #include "Materials/MaterialExpressionHairColor.h"
 #include "Materials/MaterialExpressionHsvToRgb.h"
@@ -105,10 +106,12 @@
 #include "Materials/MaterialExpressionLightmapUVs.h"
 #include "Materials/MaterialExpressionLightVector.h"
 #include "Materials/MaterialExpressionLinearInterpolate.h"
+#include "Materials/MaterialExpressionLightmassReplace.h"
 #include "Materials/MaterialExpressionLogarithm.h"
 #include "Materials/MaterialExpressionLogarithm10.h"
 #include "Materials/MaterialExpressionLogarithm2.h"
 #include "Materials/MaterialExpressionMakeMaterialAttributes.h"
+#include "Materials/MaterialExpressionMapARPassthroughCameraUV.h"
 #include "Materials/MaterialExpressionMaterialAttributeLayers.h"
 #include "Materials/MaterialExpressionMaterialFunctionCall.h"
 #include "Materials/MaterialExpressionMaterialProxyReplace.h"
@@ -152,6 +155,7 @@
 #include "Materials/MaterialExpressionPreSkinnedPosition.h"
 #include "Materials/MaterialExpressionPreviousFrameSwitch.h"
 #include "Materials/MaterialExpressionSamplePhysicsField.h"
+#include "Materials/MaterialExpressionSphericalParticleOpacity.h"
 #include "Materials/MaterialExpressionQualitySwitch.h"
 #include "Materials/MaterialExpressionRayTracingQualitySwitch.h"
 #include "Materials/MaterialExpressionReflectionCapturePassSwitch.h"
@@ -350,6 +354,36 @@ bool UMaterialExpressionStaticSwitch::GenerateHLSLExpression(FMaterialHLSLGenera
 	return true;
 }
 
+bool UMaterialExpressionMapARPassthroughCameraUV::GenerateHLSLExpression(FMaterialHLSLGenerator& Generator, UE::HLSLTree::FScope& Scope, int32 OutputIndex, UE::HLSLTree::FExpression const*& OutExpression) const
+{
+	using namespace UE::HLSLTree;
+	const FExpression* UVExpression = Coordinates.AcquireHLSLExpression(Generator, Scope);
+	if (!UVExpression)
+	{
+		return Generator.Error(TEXT("UV input missing"));
+	}
+
+	FTree& Tree = Generator.GetTree();
+	const FExpression* UVPair0 = Tree.NewExpression<FExpressionInlineCustomHLSL>(UE::Shader::EValueType::Float4, TEXT("ResolvedView.XRPassthroughCameraUVs[0]"));
+	const FExpression* UVPair1 = Tree.NewExpression<FExpressionInlineCustomHLSL>(UE::Shader::EValueType::Float4, TEXT("ResolvedView.XRPassthroughCameraUVs[1]"));
+	const FExpression* ULerp = Tree.NewLerp(UVPair0, UVPair1, Tree.NewExpression<FExpressionSwizzle>(MakeSwizzleMask(true, false, false, false), UVExpression));
+	OutExpression = Tree.NewLerp(
+		Tree.NewExpression<FExpressionSwizzle>(MakeSwizzleMask(true, true, false, false), ULerp),
+		Tree.NewExpression<FExpressionSwizzle>(MakeSwizzleMask(false, false, true, true), ULerp),
+		Tree.NewExpression<FExpressionSwizzle>(MakeSwizzleMask(false, true, false, false), UVExpression)
+		);
+	return true;
+}
+
+bool UMaterialExpressionSphericalParticleOpacity::GenerateHLSLExpression(FMaterialHLSLGenerator& Generator, UE::HLSLTree::FScope& Scope, int32 OutputIndex, UE::HLSLTree::FExpression const*& OutExpression) const
+{
+	using namespace UE::HLSLTree;
+
+	const FExpression* DensityExpression = Density.AcquireHLSLExpressionOrConstant(Generator, Scope, ConstantDensity);
+	OutExpression = Generator.GetTree().NewExpression<Material::FExpressionSphericalParticleOpacityFunction>(DensityExpression);
+	return true;
+}
+
 bool UMaterialExpressionFeatureLevelSwitch::GenerateHLSLExpression(FMaterialHLSLGenerator& Generator, UE::HLSLTree::FScope& Scope, int32 OutputIndex, UE::HLSLTree::FExpression const*& OutExpression) const
 {
 	using namespace UE::HLSLTree;
@@ -885,6 +919,66 @@ bool UMaterialExpressionStaticSwitchParameter::GenerateHLSLExpression(FMaterialH
 	return true;
 }
 
+bool UMaterialExpressionGIReplace::GenerateHLSLExpression(FMaterialHLSLGenerator& Generator, UE::HLSLTree::FScope& Scope, int32 OutputIndex, UE::HLSLTree::FExpression const*& OutExpression) const
+{
+	using namespace UE::HLSLTree;
+
+	if (!Default.GetTracedInput().Expression)
+	{
+		return Generator.Error(TEXT("Missing GIReplace input 'Default'"));
+	}
+
+	const FExpression* DefaultExpression = Default.AcquireHLSLExpression(Generator, Scope);
+	const FExpression* DynamicIndirectExpression = DynamicIndirect.AcquireHLSLExpressionOrDefaultExpression(Generator, Scope, DefaultExpression);
+	// This is a EMaterialCompilerType::Standard, so does not need to handle StaticIndirect, a different compiler type will handle that
+
+	if (DynamicIndirectExpression == DefaultExpression)
+	{
+		OutExpression = DefaultExpression;
+	}
+	else
+	{
+		const FExpression* ExpressionSwitch = Generator.GetTree().NewExpression<UE::HLSLTree::FExpressionInlineCustomHLSL>(UE::Shader::EValueType::Bool1, TEXT("GetGIReplaceState()"));
+		OutExpression = Generator.GetTree().NewExpression<FExpressionSelect>(ExpressionSwitch, DynamicIndirectExpression, DefaultExpression);
+	}
+
+	return true;
+}
+
+bool UMaterialExpressionLightmassReplace::GenerateHLSLExpression(FMaterialHLSLGenerator& Generator, UE::HLSLTree::FScope& Scope, int32 OutputIndex, UE::HLSLTree::FExpression const*& OutExpression) const
+{
+	using namespace UE::HLSLTree;
+
+	if (!Realtime.GetTracedInput().Expression)
+	{
+		return Generator.Error(TEXT("Missing LightmassReplace input Realtime"));
+	}
+	if (!Lightmass.GetTracedInput().Expression)
+	{
+		return Generator.Error(TEXT("Missing LightmassReplace input Lightmass"));
+	}
+
+	// This is a EMaterialCompilerType::Standard, not a Lightmass compiler, but there is a shader-time switch to handle
+	const FExpression* LightmassExpression = Lightmass.AcquireHLSLExpression(Generator, Scope);
+	const FExpression* RealtimeExpression = Realtime.AcquireHLSLExpression(Generator, Scope);
+
+	if (LightmassExpression && RealtimeExpression)
+	{
+		const FExpression* ExpressionSwitch = Generator.GetTree().NewExpression<UE::HLSLTree::FExpressionInlineCustomHLSL>(UE::Shader::EValueType::Bool1, TEXT("GetLightmassReplaceState()"));
+		OutExpression = Generator.GetTree().NewExpression<FExpressionSelect>(ExpressionSwitch, LightmassExpression, RealtimeExpression);
+	}
+	else if (RealtimeExpression)
+	{
+		OutExpression = RealtimeExpression;
+	}
+	else if (LightmassExpression)
+	{
+		OutExpression = LightmassExpression;
+	}
+
+	return OutExpression != nullptr;
+}
+
 bool UMaterialExpressionStaticComponentMaskParameter::GenerateHLSLExpression(FMaterialHLSLGenerator& Generator, UE::HLSLTree::FScope& Scope, int32 OutputIndex, UE::HLSLTree::FExpression const*& OutExpression) const
 {
 	using namespace UE::HLSLTree;
diff --git a/Engine/Source/Runtime/Engine/Private/Materials/MaterialHLSLEmitter.cpp b/Engine/Source/Runtime/Engine/Private/Materials/MaterialHLSLEmitter.cpp
--- a/Engine/Source/Runtime/Engine/Private/Materials/MaterialHLSLEmitter.cpp
+++ b/Engine/Source/Runtime/Engine/Private/Materials/MaterialHLSLEmitter.cpp
@@ -344,7 +344,7 @@ static void GetMaterialEnvironment(EShaderPlatform InPlatform,
 	const bool bNeedsParticlePosition = EmitMaterialData.IsExternalInputUsed(Material::EExternalInput::ParticleTranslatedWorldPosition) ||
 		EmitMaterialData.IsExternalInputUsed(Material::EExternalInput::ParticleRadius);
 
-	if (bNeedsParticlePosition || InMaterial.ShouldGenerateSphericalParticleNormals()) // || bUsesSphericalParticleOpacity)
+	if (bNeedsParticlePosition || InMaterial.ShouldGenerateSphericalParticleNormals() || EmitContext.bUsesSphericalParticleOpacity)
 	{
 		OutEnvironment.SetDefine(TEXT("NEEDS_PARTICLE_POSITION"), 1);
 	}
@@ -415,6 +415,7 @@ static void GetMaterialEnvironment(EShaderPlatform InPlatform,
 	}
 
 	const bool bNeedsWorldPositionExcludingShaderOffsets =
+		EmitContext.bUsesWorldPositionExcludingShaderOffsets || 
 		EmitMaterialData.IsExternalInputUsed(Material::EExternalInput::WorldPosition_NoOffsets) ||
 		EmitMaterialData.IsExternalInputUsed(Material::EExternalInput::TranslatedWorldPosition_NoOffsets) ||
 		EmitMaterialData.IsExternalInputUsed(Material::EExternalInput::PrevWorldPosition_NoOffsets) ||
diff --git a/Engine/Source/Runtime/Engine/Private/Materials/MaterialHLSLTree.cpp b/Engine/Source/Runtime/Engine/Private/Materials/MaterialHLSLTree.cpp
--- a/Engine/Source/Runtime/Engine/Private/Materials/MaterialHLSLTree.cpp
+++ b/Engine/Source/Runtime/Engine/Private/Materials/MaterialHLSLTree.cpp
@@ -899,6 +899,37 @@ void FExpressionDecalMipmapLevel::EmitValueShader(FEmitContext& Context, FEmitSc
 		EmitSizeExpression);
 }
 
+bool FExpressionSphericalParticleOpacityFunction::PrepareValue(FEmitContext& Context, FEmitScope& Scope, const FRequestedType& RequestedType, FPrepareValueResult& OutResult) const
+{
+	if (Context.ShaderFrequency != SF_Pixel && Context.ShaderFrequency != SF_Compute)
+	{
+		return Context.Error(TEXT("Can only be used in Pixel and Compute shaders."));
+	}
+
+	if (Context.PrepareExpression(DensityExpression, Scope, RequestedType).IsVoid())
+	{
+		return false;
+	}
+
+	if (Context.bMarkLiveValues && Context.MaterialCompilationOutput)
+	{
+		Context.MaterialCompilationOutput->SetIsSceneTextureUsed(PPI_SceneDepth);
+		Context.bUsesSphericalParticleOpacity = true;
+		Context.bUsesWorldPositionExcludingShaderOffsets = true;
+	}
+
+	return OutResult.SetType(Context, RequestedType, EExpressionEvaluation::Shader, Shader::EValueType::Float1);
+}
+
+void FExpressionSphericalParticleOpacityFunction::EmitValueShader(FEmitContext& Context, FEmitScope& Scope, const FRequestedType& RequestedType, FEmitValueShaderResult& OutResult) const
+{
+	FEmitShaderExpression* EmitDensityExpression = DensityExpression->GetValueShader(Context, Scope, Shader::EValueType::Float2);
+	OutResult.Code = Context.EmitInlineExpression(
+		Scope, Shader::EValueType::Float1,
+		TEXT("GetSphericalParticleOpacity(Parameters,%)"),
+		EmitDensityExpression);
+}
+
 bool FExpressionDBufferTexture::PrepareValue(FEmitContext& Context, FEmitScope& Scope, const FRequestedType& RequestedType, FPrepareValueResult& OutResult) const
 {
 	if (UVExpression && Context.PrepareExpression(UVExpression, Scope, RequestedType).IsVoid())
@@ -1420,6 +1451,35 @@ FEmitShaderExpression* EmitTextureSampleShader(
 			break;
 		}
 
+		// If not 2D texture, disable AutomaticViewMipBias.
+		if (TextureType != MCT_Texture2D)
+		{
+			bAutomaticViewMipBias = false;
+		}
+
+		// if we are not in the PS we need a mip level
+		if (Context.ShaderFrequency != SF_Pixel)
+		{
+			MipValueMode = TMVM_MipLevel;
+			bAutomaticViewMipBias = false;
+		}
+
+		const FMaterial* Material = Context.Material;
+		if (Material)
+		{
+			// If mobile, disable AutomaticViewMipBias.
+			if (Material->GetFeatureLevel() < ERHIFeatureLevel::SM5)
+			{
+				bAutomaticViewMipBias = false;
+			}
+
+			// Outside of surface and decal domains, disable AutomaticViewMipBias.
+			if (Material->GetMaterialDomain() != MD_Surface && Material->GetMaterialDomain() != MD_DeferredDecal)
+			{
+				bAutomaticViewMipBias = false;
+			}
+		}
+
 		TStringBuilder<64> FormattedTexture;
 		Private::EmitTextureShader(Context, TextureValue, FormattedTexture);
 
@@ -1786,6 +1846,19 @@ void FExpressionTextureProperty::EmitValuePreshader(FEmitContext& Context, FEmit
 	Context.PreshaderStackPosition++;
 	OutResult.Type = Shader::EValueType::Float3;
 	OutResult.Preshader.WriteOpcode(Op).Write<FMemoryImageMaterialParameterInfo>(TextureValue.ParameterInfo).Write(TextureIndex);
+
+	// Swizzle to two components for texture2d-type textures
+	if (TextureValue.Texture)
+	{
+		EMaterialValueType Type = TextureValue.Texture->GetMaterialType();
+
+		// this follows the old translator's concept of a Float2 texture size, masked to 2 components, .xy
+		if (Type != MCT_VolumeTexture && Type != MCT_Texture2DArray && Type != MCT_SparseVolumeTexture)
+		{
+			OutResult.Preshader.WriteOpcode(Shader::EPreshaderOpcode::ComponentSwizzle).Write((uint8)2).Write((uint8)0).Write((uint8)1).Write((uint8)0).Write((uint8)0);
+			OutResult.Type = Shader::EValueType::Float2;
+		}
+	}
 }
 
 bool FExpressionRuntimeVirtualTextureUniform::PrepareValue(FEmitContext& Context, FEmitScope& Scope, const FRequestedType& RequestedType, FPrepareValueResult& OutResult) const
diff --git a/Engine/Source/Runtime/Engine/Public/HLSLTree/HLSLTreeEmit.h b/Engine/Source/Runtime/Engine/Public/HLSLTree/HLSLTreeEmit.h
--- a/Engine/Source/Runtime/Engine/Public/HLSLTree/HLSLTreeEmit.h
+++ b/Engine/Source/Runtime/Engine/Public/HLSLTree/HLSLTreeEmit.h
@@ -583,6 +583,8 @@ public:
 
 	bool bUsesSkyAtmosphere = false;
 	bool bUsesSpeedTree = false;
+	bool bUsesSphericalParticleOpacity = false;
+	bool bUsesWorldPositionExcludingShaderOffsets = false;
 
 	uint32 DynamicParticleParameterMask = 0u;
 
diff --git a/Engine/Source/Runtime/Engine/Public/MaterialHLSLTree.h b/Engine/Source/Runtime/Engine/Public/MaterialHLSLTree.h
--- a/Engine/Source/Runtime/Engine/Public/MaterialHLSLTree.h
+++ b/Engine/Source/Runtime/Engine/Public/MaterialHLSLTree.h
@@ -401,6 +401,19 @@ public:
 	virtual void EmitValueShader(FEmitContext& Context, FEmitScope& Scope, const FRequestedType& RequestedType, FEmitValueShaderResult& OutResult) const override;
 };
 
+class FExpressionSphericalParticleOpacityFunction : public FExpression
+{
+public:
+	const FExpression* DensityExpression;
+
+	FExpressionSphericalParticleOpacityFunction(const FExpression* InDensityExpression)
+		: DensityExpression(InDensityExpression)
+	{}
+
+	virtual bool PrepareValue(FEmitContext& Context, FEmitScope& Scope, const FRequestedType& RequestedType, FPrepareValueResult& OutResult) const override;
+	virtual void EmitValueShader(FEmitContext& Context, FEmitScope& Scope, const FRequestedType& RequestedType, FEmitValueShaderResult& OutResult) const override;
+};
+
 class FExpressionSobolFunction : public FExpression
 {
 public:
