Wrapped material translation logging behind a CVAR (r.Material.EnableTranslationLogFile).

diff --git a/Engine/Source/Runtime/Engine/Private/Materials/HLSLMaterialTranslator.cpp b/Engine/Source/Runtime/Engine/Private/Materials/HLSLMaterialTranslator.cpp
--- a/Engine/Source/Runtime/Engine/Private/Materials/HLSLMaterialTranslator.cpp
+++ b/Engine/Source/Runtime/Engine/Private/Materials/HLSLMaterialTranslator.cpp
@@ -57,6 +57,14 @@
 #endif
 
 #if ENABLE_COOK_STATS && STATS
+
+static bool GEnableMaterialTranslationLogFile = false;
+static FAutoConsoleVariableRef CVarEnableMaterialTranslationLogFile(
+	TEXT("r.Material.EnableTranslationLogFile"),
+	GEnableMaterialTranslationLogFile,
+	TEXT("Enables material translation log file generation for tracking materials translation time (written to  \"ShaderDebugInfo/MaterialTranslationLog-X.csv\")."),
+	ECVF_ReadOnly);
+
 /**
  * Utility used to create a MaterialTranslationLog.txt file during cooks that contains the list of all translated materials
  * and other info such as how long the translation took.
@@ -79,20 +87,25 @@ struct FCsvLogFile
 
 	void AddEntry(FStringView MaterialName, FDateTime DateTime, float TranslationTime)
 	{
-		FScopeLock Lock{ &CriticalSection };
-		FString DateTimeString = DateTime.ToString(TEXT("%Y-%m-%d %H:%M:%S"));
-		LogContent.Appendf(TEXT("%s,%s,%f\n"), *DateTimeString, MaterialName.GetData(), TranslationTime);
+		if (GEnableMaterialTranslationLogFile)
+		{
+			FScopeLock Lock{ &CriticalSection };
+			FString DateTimeString = DateTime.ToString(TEXT("%Y-%m-%d %H:%M:%S"));
+			LogContent.Appendf(TEXT("%s,%s,%f\n"), *DateTimeString, MaterialName.GetData(), TranslationTime);
+		}
 	}
 
 	void Save()
 	{
-		uint32 MultiprocessId = 0;
-		FParse::Value(FCommandLine::Get(), TEXT("-MultiprocessId="), MultiprocessId);
-		FString FilePath = GShaderCompilingManager->GetAbsoluteShaderDebugInfoDirectory() / FString::Printf(TEXT("MaterialTranslationLog-%d.csv"), MultiprocessId);
-		UE_LOG(LogMaterial, Display, TEXT("Writing out MaterialTranslation log file '%s'"), *FilePath);
-		if (!FFileHelper::SaveStringToFile(LogContent, *FilePath))
+		if (GEnableMaterialTranslationLogFile)
 		{
-			UE_LOG(LogMaterial, Display, TEXT("Cannot open MaterialTranslation log file '%s'"), *FilePath);
+			uint32 MultiprocessId = 0;
+			FParse::Value(FCommandLine::Get(), TEXT("-MultiprocessId="), MultiprocessId);
+			FString FilePath = GShaderCompilingManager->GetAbsoluteShaderDebugInfoDirectory() / FString::Printf(TEXT("MaterialTranslationLog-%d.csv"), MultiprocessId);
+			if (!FFileHelper::SaveStringToFile(LogContent, *FilePath))
+			{
+				UE_LOG(LogMaterial, Display, TEXT("Cannot open MaterialTranslation log file '%s'"), *FilePath);
+			}
 		}
 	}
 
