WIP New HLSL translator Material Expressions
MaterialExpressionRuntimeVirtualTextureReplace, MaterialExpressionPathTracingBufferTexture, MaterialExpressionAntialiasedTextureMask

#tests materials, editor, PIE, win64

diff --git a/Engine/Source/Runtime/Engine/Classes/Materials/MaterialExpressionAntialiasedTextureMask.h b/Engine/Source/Runtime/Engine/Classes/Materials/MaterialExpressionAntialiasedTextureMask.h
--- a/Engine/Source/Runtime/Engine/Classes/Materials/MaterialExpressionAntialiasedTextureMask.h
+++ b/Engine/Source/Runtime/Engine/Classes/Materials/MaterialExpressionAntialiasedTextureMask.h
@@ -40,6 +40,7 @@ class UMaterialExpressionAntialiasedTextureMask : public UMaterialExpressionText
 	//~ Begin UMaterialExpression Interface
 	virtual int32 Compile(class FMaterialCompiler* Compiler, int32 OutputIndex) override;
 	virtual void GetCaption(TArray<FString>& OutCaptions) const override;
+	virtual bool GenerateHLSLExpression(FMaterialHLSLGenerator& Generator, UE::HLSLTree::FScope& Scope, int32 OutputIndex, UE::HLSLTree::FExpression const*& OutExpression) const override;
 	//~ End UMaterialExpression Interface
 #endif // WITH_EDITOR
 };
diff --git a/Engine/Source/Runtime/Engine/Classes/Materials/MaterialExpressionPathTracingBufferTexture.h b/Engine/Source/Runtime/Engine/Classes/Materials/MaterialExpressionPathTracingBufferTexture.h
--- a/Engine/Source/Runtime/Engine/Classes/Materials/MaterialExpressionPathTracingBufferTexture.h
+++ b/Engine/Source/Runtime/Engine/Classes/Materials/MaterialExpressionPathTracingBufferTexture.h
@@ -43,6 +43,7 @@ class UMaterialExpressionPathTracingBufferTexture : public UMaterialExpression
 #if WITH_EDITOR
 	virtual int32 Compile(class FMaterialCompiler* Compiler, int32 OutputIndex) override;
 	virtual void GetCaption(TArray<FString>& OutCaptions) const override;
+	virtual bool GenerateHLSLExpression(FMaterialHLSLGenerator& Generator, UE::HLSLTree::FScope& Scope, int32 OutputIndex, UE::HLSLTree::FExpression const*& OutExpression) const override;
 #endif
 	//~ End UMaterialExpression Interface
 };
diff --git a/Engine/Source/Runtime/Engine/Classes/Materials/MaterialExpressionRuntimeVirtualTextureReplace.h b/Engine/Source/Runtime/Engine/Classes/Materials/MaterialExpressionRuntimeVirtualTextureReplace.h
--- a/Engine/Source/Runtime/Engine/Classes/Materials/MaterialExpressionRuntimeVirtualTextureReplace.h
+++ b/Engine/Source/Runtime/Engine/Classes/Materials/MaterialExpressionRuntimeVirtualTextureReplace.h
@@ -29,6 +29,7 @@ class UMaterialExpressionRuntimeVirtualTextureReplace : public UMaterialExpressi
 	virtual uint32 GetInputType(int32 InputIndex) override { return MCT_Unknown; }
 	virtual uint32 GetOutputType(int32 OutputIndex) override { return MCT_Unknown; }
 	virtual void GetCaption(TArray<FString>& OutCaptions) const override;
+	virtual bool GenerateHLSLExpression(FMaterialHLSLGenerator& Generator, UE::HLSLTree::FScope& Scope, int32 OutputIndex, UE::HLSLTree::FExpression const*& OutExpression) const override;
 #endif
 	//~ End UMaterialExpression Interface
 };
diff --git a/Engine/Source/Runtime/Engine/Private/Materials/MaterialExpressionHLSL.cpp b/Engine/Source/Runtime/Engine/Private/Materials/MaterialExpressionHLSL.cpp
--- a/Engine/Source/Runtime/Engine/Private/Materials/MaterialExpressionHLSL.cpp
+++ b/Engine/Source/Runtime/Engine/Private/Materials/MaterialExpressionHLSL.cpp
@@ -141,6 +141,7 @@
 #include "Materials/MaterialExpressionParticleSize.h"
 #include "Materials/MaterialExpressionParticleSpeed.h"
 #include "Materials/MaterialExpressionParticleSubUVProperties.h"
+#include "Materials/MaterialExpressionPathTracingBufferTexture.h"
 #include "Materials/MaterialExpressionPathTracingQualitySwitch.h"
 #include "Materials/MaterialExpressionPathTracingRayTypeSwitch.h"
 #include "Materials/MaterialExpressionPerInstanceCustomData.h"
@@ -166,6 +167,7 @@
 #include "Materials/MaterialExpressionRotator.h"
 #include "Materials/MaterialExpressionRound.h"
 #include "Materials/MaterialExpressionRuntimeVirtualTextureOutput.h"
+#include "Materials/MaterialExpressionRuntimeVirtualTextureReplace.h"
 #include "Materials/MaterialExpressionRuntimeVirtualTextureSample.h"
 #include "Materials/MaterialExpressionRuntimeVirtualTextureSampleParameter.h"
 #include "Materials/MaterialExpressionSaturate.h"
@@ -1539,6 +1541,27 @@ bool UMaterialExpressionTextureSample::GenerateHLSLExpressionBase(FMaterialHLSLG
 	return true;
 }
 
+bool UMaterialExpressionAntialiasedTextureMask::GenerateHLSLExpression(FMaterialHLSLGenerator& Generator, UE::HLSLTree::FScope& Scope, int32 OutputIndex, UE::HLSLTree::FExpression const*& OutExpression) const
+{
+	using namespace UE::HLSLTree;
+
+	const FExpression* CoordExpression = Coordinates.AcquireHLSLExpressionOrExternalInput(Generator, Scope, Material::MakeInputTexCoord(ConstCoordinate));
+
+	const FExpression* TextureExpression = nullptr;
+	if (TextureObject.GetTracedInput().Expression)
+	{
+		TextureExpression = TextureObject.AcquireHLSLExpression(Generator, Scope);
+	}
+	else if (Texture)
+	{
+		const FMaterialParameterMetadata ParameterMeta(Texture);
+		TextureExpression = Generator.GenerateMaterialParameter(FName(), ParameterMeta, SamplerType);
+	}
+
+	OutExpression = Generator.GetTree().NewExpression<Material::FExpressionAntiAliasedTextureMask>(TextureExpression, CoordExpression, Threshold, Channel);
+	return true;
+}
+
 bool UMaterialExpressionTextureSample::GenerateHLSLExpression(FMaterialHLSLGenerator& Generator, UE::HLSLTree::FScope& Scope, int32 OutputIndex, UE::HLSLTree::FExpression const*& OutExpression) const
 {
 	using namespace UE::HLSLTree;
@@ -1972,6 +1995,27 @@ bool UMaterialExpressionRuntimeVirtualTextureOutput::GenerateHLSLExpression(FMat
 	return true;
 }
 
+bool UMaterialExpressionRuntimeVirtualTextureReplace::GenerateHLSLExpression(FMaterialHLSLGenerator& Generator, UE::HLSLTree::FScope& Scope, int32 OutputIndex, UE::HLSLTree::FExpression const*& OutExpression) const
+{
+	using namespace UE::HLSLTree;
+
+	if (!Default.GetTracedInput().Expression)
+	{
+		return Generator.Errorf(TEXT("Missing RuntimeVirtualTextureReplace input 'Default'"));
+	}
+
+	if (!VirtualTextureOutput.GetTracedInput().Expression)
+	{
+		return Generator.Errorf(TEXT("Missing RuntimeVirtualTextureReplace input 'VirtualTextureOutput'"));
+	}
+
+	const FExpression* DefaultExpression = Default.AcquireHLSLExpression(Generator, Scope);
+	const FExpression* VirtualTextureExpression = VirtualTextureOutput.AcquireHLSLExpression(Generator, Scope);
+	const FExpression* ConditionExpression = Generator.GetTree().NewExpression<FExpressionInlineCustomHLSL>(UE::Shader::EValueType::Bool1, TEXT("GetRuntimeVirtualTextureOutputSwitch()"));
+	OutExpression = Generator.GetTree().NewExpression<FExpressionSelect>(ConditionExpression, VirtualTextureExpression, DefaultExpression);
+	return true;
+}
+
 bool UMaterialExpressionFontSample::GenerateHLSLExpression(FMaterialHLSLGenerator& Generator, UE::HLSLTree::FScope& Scope, int32 OutputIndex, UE::HLSLTree::FExpression const*& OutExpression) const
 {
 	using namespace UE::HLSLTree;
@@ -2473,6 +2517,25 @@ bool UMaterialExpressionPathTracingQualitySwitch::GenerateHLSLExpression(FMateri
 	}
 }
 
+bool UMaterialExpressionPathTracingBufferTexture::GenerateHLSLExpression(FMaterialHLSLGenerator& Generator, UE::HLSLTree::FScope& Scope, int32 OutputIndex, UE::HLSLTree::FExpression const*& OutExpression) const
+{
+	using namespace UE::HLSLTree;
+
+	const FExpression* UVExpression = nullptr;
+	if (Coordinates.GetTracedInput().Expression)
+	{
+		UVExpression = Coordinates.AcquireHLSLExpression(Generator, Scope);
+	}
+	else
+	{
+		UVExpression = Generator.GetTree().NewExpression<UE::HLSLTree::FExpressionInlineCustomHLSL>(UE::Shader::EValueType::Float2, TEXT("GetDefaultPathTracingBufferTextureUV(Parameters, 0)"));
+
+	}
+
+	OutExpression = Generator.GetTree().NewExpression<Material::FPathTracingBufferTextureFunction>(UVExpression, PathTracingBufferTextureId);
+	return true;
+}
+
 bool UMaterialExpressionDepthOfFieldFunction::GenerateHLSLExpression(FMaterialHLSLGenerator& Generator, UE::HLSLTree::FScope& Scope, int32 OutputIndex, UE::HLSLTree::FExpression const*& OutExpression) const
 {
 	using namespace UE::HLSLTree;
diff --git a/Engine/Source/Runtime/Engine/Private/Materials/MaterialHLSLTree.cpp b/Engine/Source/Runtime/Engine/Private/Materials/MaterialHLSLTree.cpp
--- a/Engine/Source/Runtime/Engine/Private/Materials/MaterialHLSLTree.cpp
+++ b/Engine/Source/Runtime/Engine/Private/Materials/MaterialHLSLTree.cpp
@@ -965,6 +965,36 @@ void FExpressionDBufferTexture::EmitValueShader(FEmitContext& Context, FEmitScop
 		EmitTexCoord, DBufferTextureID);
 }
 
+bool FPathTracingBufferTextureFunction::PrepareValue(FEmitContext& Context, FEmitScope& Scope, const FRequestedType& RequestedType, FPrepareValueResult& OutResult) const
+{
+	if (Context.Material->GetMaterialDomain() != MD_PostProcess)
+	{
+		return Context.Error(TEXT("Path tracing buffer textures are only available on post process material."));
+	}
+
+	if (Context.PrepareExpression(UVExpression, Scope, RequestedType).IsVoid())
+	{
+		return false;
+	}
+
+	if (Context.bMarkLiveValues && Context.MaterialCompilationOutput)
+	{
+		Context.MaterialCompilationOutput->SetIsPathTracingBufferTextureUsed(PathTracingBufferTextureID);
+	}
+
+	return OutResult.SetType(Context, RequestedType, EExpressionEvaluation::Shader, Shader::EValueType::Float4);
+}
+
+void FPathTracingBufferTextureFunction::EmitValueShader(FEmitContext& Context, FEmitScope& Scope, const FRequestedType& RequestedType, FEmitValueShaderResult& OutResult) const
+{
+	FEmitShaderExpression* EmitTexCoord = UVExpression->GetValueShader(Context, Scope, Shader::EValueType::Float2);
+
+	OutResult.Code = Context.EmitInlineExpression(
+		Scope, Shader::EValueType::Float4,
+		TEXT("MaterialExpressionPathTracingBufferTextureLookup(Parameters, %, %)"),
+		EmitTexCoord, (int)PathTracingBufferTextureID);
+}
+
 namespace Private
 {
 
@@ -1861,6 +1891,48 @@ void FExpressionTextureProperty::EmitValuePreshader(FEmitContext& Context, FEmit
 	}
 }
 
+bool FExpressionAntiAliasedTextureMask::PrepareValue(FEmitContext& Context, FEmitScope& Scope, const FRequestedType& RequestedType, FPrepareValueResult& OutResult) const
+{
+	if (Context.TargetParameters.FeatureLevel < ERHIFeatureLevel::SM5)
+	{
+		return Context.Errorf(TEXT("Node not supported in feature level %d. %d required."), Context.TargetParameters.FeatureLevel, ERHIFeatureLevel::SM5);
+	}
+
+	const FPreparedType& TextureType = Context.PrepareExpression(TextureExpression, Scope, FMaterialTextureValue::GetTypeName());
+	if (TextureType.Type.ObjectType != FMaterialTextureValue::GetTypeName())
+	{
+		return Context.Error(TEXT("Expected texture"));
+	}
+
+	const FPreparedType& CoordType = Context.PrepareExpression(TexCoordExpression, Scope, Shader::EValueType::Float2);
+	if (CoordType.IsVoid())
+	{
+		return false;
+	}
+
+	return OutResult.SetType(Context, RequestedType, EExpressionEvaluation::Shader, Shader::EValueType::Float1);
+}
+
+void FExpressionAntiAliasedTextureMask::EmitValueShader(FEmitContext& Context, FEmitScope& Scope, const FRequestedType& RequestedType, FEmitValueShaderResult& OutResult) const
+{
+	FEmitShaderExpression* EmitCoord = TexCoordExpression->GetValueShader(Context, Scope, Shader::EValueType::Float2);
+
+	FMaterialTextureValue TextureValue;
+	verify(TextureExpression->GetValueObject(Context, Scope, TextureValue));
+	TStringBuilder<64> FormattedTexture;
+	Private::EmitTextureShader(Context, TextureValue, FormattedTexture);
+
+	OutResult.Code = Context.EmitExpression(
+		Scope,
+		Shader::EValueType::Float1,
+		TEXT("AntialiasedTextureMask(%,%Sampler,%,%,%)"),
+		FormattedTexture.ToString(),
+		FormattedTexture.ToString(),
+		EmitCoord,
+		Threshold,
+		Channel);
+}
+
 bool FExpressionRuntimeVirtualTextureUniform::PrepareValue(FEmitContext& Context, FEmitScope& Scope, const FRequestedType& RequestedType, FPrepareValueResult& OutResult) const
 {
 	const FPreparedType& TextureType = Context.PrepareExpression(TextureExpression, Scope, FMaterialTextureValue::GetTypeName());
diff --git a/Engine/Source/Runtime/Engine/Public/MaterialHLSLTree.h b/Engine/Source/Runtime/Engine/Public/MaterialHLSLTree.h
--- a/Engine/Source/Runtime/Engine/Public/MaterialHLSLTree.h
+++ b/Engine/Source/Runtime/Engine/Public/MaterialHLSLTree.h
@@ -401,6 +401,21 @@ public:
 	virtual void EmitValueShader(FEmitContext& Context, FEmitScope& Scope, const FRequestedType& RequestedType, FEmitValueShaderResult& OutResult) const override;
 };
 
+class FPathTracingBufferTextureFunction : public FExpression
+{
+public:
+	const FExpression* UVExpression;
+	const uint8 PathTracingBufferTextureID;
+
+	FPathTracingBufferTextureFunction(const FExpression* InUVExpression, uint8 InPathTracingBufferTextureID)
+		: UVExpression(InUVExpression)
+		, PathTracingBufferTextureID(InPathTracingBufferTextureID)
+	{}
+
+	virtual bool PrepareValue(FEmitContext& Context, FEmitScope& Scope, const FRequestedType& RequestedType, FPrepareValueResult& OutResult) const override;
+	virtual void EmitValueShader(FEmitContext& Context, FEmitScope& Scope, const FRequestedType& RequestedType, FEmitValueShaderResult& OutResult) const override;
+};
+
 class FExpressionSphericalParticleOpacityFunction : public FExpression
 {
 public:
@@ -571,6 +586,30 @@ public:
 	virtual void EmitValueShader(FEmitContext& Context, FEmitScope& Scope, const FRequestedType& RequestedType, FEmitValueShaderResult& OutResult) const override;
 };
 
+class FExpressionAntiAliasedTextureMask : public FExpression
+{
+public:
+	const FExpression* TextureExpression;
+	const FExpression* TexCoordExpression;
+	float Threshold;
+	uint8 Channel;
+
+	explicit FExpressionAntiAliasedTextureMask(
+		const FExpression* InTextureExpression, 
+		const FExpression* InTexCoordExpression,
+		float InThreshold, 
+		uint8 InChannel)
+		: TextureExpression(InTextureExpression)
+		, TexCoordExpression(InTexCoordExpression)
+		, Threshold(InThreshold)
+		, Channel(InChannel)
+	{}
+
+	virtual bool PrepareValue(FEmitContext& Context, FEmitScope& Scope, const FRequestedType& RequestedType, FPrepareValueResult& OutResult) const override;
+	virtual void EmitValueShader(FEmitContext& Context, FEmitScope& Scope, const FRequestedType& RequestedType, FEmitValueShaderResult& OutResult) const override;
+};
+
+
 class FExpressionStaticTerrainLayerWeight : public FExpression
 {
 public:
