#animation - Added DECLARE_NOTIFY_CONTEXT_INERFACE_API to FAnimNotifyMontageInstanceContext so it can be accessed outside of engine.

diff --git a/Engine/Source/Runtime/Engine/Public/Animation/ActiveMontageInstanceScope.h b/Engine/Source/Runtime/Engine/Public/Animation/ActiveMontageInstanceScope.h
--- a/Engine/Source/Runtime/Engine/Public/Animation/ActiveMontageInstanceScope.h
+++ b/Engine/Source/Runtime/Engine/Public/Animation/ActiveMontageInstanceScope.h
@@ -10,7 +10,7 @@ namespace UE::Anim {
 // This is useful i.e When there's simultaneous instances of the same montage, to trigger start/end events of its notify states.
 class FAnimNotifyMontageInstanceContext : public UE::Anim::IAnimNotifyEventContextDataInterface
 {
-	DECLARE_NOTIFY_CONTEXT_INTERFACE(FAnimNotifyMontageInstanceContext)
+	DECLARE_NOTIFY_CONTEXT_INTERFACE_API(FAnimNotifyMontageInstanceContext, ENGINE_API)
 public:
 	ENGINE_API FAnimNotifyMontageInstanceContext(const int32 InMontageInstanceID);
 	const int32 MontageInstanceID;
diff --git a/Engine/Source/Runtime/Engine/Public/Animation/AnimNodeMessages.h b/Engine/Source/Runtime/Engine/Public/Animation/AnimNodeMessages.h
--- a/Engine/Source/Runtime/Engine/Public/Animation/AnimNodeMessages.h
+++ b/Engine/Source/Runtime/Engine/Public/Animation/AnimNodeMessages.h
@@ -37,6 +37,9 @@ class IAnimNotifyEventContextDataInterface;
 #define DECLARE_NOTIFY_CONTEXT_INTERFACE(ClassName) \
 	DECLARE_ANIMGRAPH_MESSAGE(ClassName)
 
+#define DECLARE_NOTIFY_CONTEXT_INTERFACE_API(ClassName, ModuleApi) \
+	DECLARE_ANIMGRAPH_MESSAGE_API(ClassName, ModuleApi)
+
 namespace UE { namespace Anim {
 
 class IAnimNotifyEventContextDataInterface
