Add -ForceDPCVars option that applies the cvars with 'code' priority.
#reveiw

diff --git a/Engine/Source/Runtime/Engine/Private/DeviceProfiles/DeviceProfileManager.cpp b/Engine/Source/Runtime/Engine/Private/DeviceProfiles/DeviceProfileManager.cpp
--- a/Engine/Source/Runtime/Engine/Private/DeviceProfiles/DeviceProfileManager.cpp
+++ b/Engine/Source/Runtime/Engine/Private/DeviceProfiles/DeviceProfileManager.cpp
@@ -475,9 +475,15 @@ void UDeviceProfileManager::SetDeviceProfileCVars(const FString& DeviceProfileNa
 #endif
 	// pre-apply any -dpcvars= items, so that they override anything in the DPs
 	// Search for all occurrences of dpcvars and dpcvar on the command line.
-	static const TCHAR* DPCVarTags[]{ TEXT("DPCVars="), TEXT("DPCVar=") };
-	for (const TCHAR* Tag : DPCVarTags)
+	static const TCHAR* DPCVarTags[]{ TEXT("DPCVars="), TEXT("DPCVar="), TEXT("ForceDPCVars=") };
+	static const EConsoleVariableFlags DPCVarPri[] = { ECVF_SetByDeviceProfile, ECVF_SetByDeviceProfile, ECVF_SetByCommandline };
+	static_assert(UE_ARRAY_COUNT(DPCVarTags) == UE_ARRAY_COUNT(DPCVarPri));
+
+	for (int i = 0; i<UE_ARRAY_COUNT(DPCVarTags) ; i++)
 	{
+		const EConsoleVariableFlags RequestedPri = DPCVarPri[i];
+		const TCHAR* Tag = DPCVarTags[i];
+		const TCHAR* RequestedPriDesc = GetConsoleVariableSetByName(RequestedPri);
 		FString DPCVarString;
 		for (const TCHAR* Cursor = FCommandLine::Get(); (Cursor != nullptr) && FParse::Value(Cursor, Tag, DPCVarString, false, &Cursor);)
 		{
@@ -493,7 +499,14 @@ void UDeviceProfileManager::SetDeviceProfileCVars(const FString& DeviceProfileNa
 					UE_LOG(LogDeviceProfileManager, Log, TEXT("Setting CommandLine Device Profile CVar: [[%s:%s]]"), *CVarKey, *CVarValue);
 
 					// set it and remember it (no thanks, Ron Popeil)
-					UE::ConfigUtilities::OnSetCVarFromIniEntry(*GDeviceProfilesIni, *CVarKey, *CVarValue, ECVF_SetByDeviceProfile);
+					UE::ConfigUtilities::OnSetCVarFromIniEntry(*GDeviceProfilesIni, *CVarKey, *CVarValue, RequestedPri);
+
+					// Log if the change would not applied.
+					if (IConsoleVariable* CVar = IConsoleManager::Get().FindConsoleVariable(*CVarKey))
+					{
+						const EConsoleVariableFlags ExistingPri = (EConsoleVariableFlags)(CVar->GetFlags() & ECVF_SetByMask);
+						UE_CLOG(RequestedPri < ExistingPri, LogDeviceProfileManager, Warning, TEXT("-%s%s=%s requested priority is too low (%s < %s), value remains %s"), Tag, *CVarKey, *CVarValue, GetConsoleVariableSetByName(RequestedPri), GetConsoleVariableSetByName(ExistingPri), *CVar->GetString() );
+					}
 				}
 			}
 		}
diff --git a/Engine/Source/Runtime/Launch/Private/Android/LaunchAndroid.cpp b/Engine/Source/Runtime/Launch/Private/Android/LaunchAndroid.cpp
--- a/Engine/Source/Runtime/Launch/Private/Android/LaunchAndroid.cpp
+++ b/Engine/Source/Runtime/Launch/Private/Android/LaunchAndroid.cpp
@@ -409,9 +409,9 @@ static void ApplyAndroidCompatConfigRules()
 					auto* CVar = IConsoleManager::Get().FindConsoleVariable(*CVarName);
 					if (CVar)
 					{
-						// set with current priority means that DPs etc can still override anything set here.
-						// e.g. -dpcvars= is expected to work.
-						CVar->SetWithCurrentPriority(*Value); 
+						// set with HF priority
+						// configrules are therefore higher priority than deviceprofiles. (e.g. -dpcvars)
+						CVar->Set(*Value, ECVF_SetByHotfix);
 						UE_LOG(LogAndroid, Log, TEXT("Compat Setting %s = %s"), *CVarName, *Value);
 					}
 				}
