World Partition RuntimeHashSet HLODs: change how we set up HLODs so it's possible to manually add streaming layers and assign multiple HLOD layers to them to support more complex setups like some internal projects do.

diff --git a/Engine/Source/Runtime/Engine/Private/WorldPartition/RuntimeHashSet/RuntimePartition.cpp b/Engine/Source/Runtime/Engine/Private/WorldPartition/RuntimeHashSet/RuntimePartition.cpp
--- a/Engine/Source/Runtime/Engine/Private/WorldPartition/RuntimeHashSet/RuntimePartition.cpp
+++ b/Engine/Source/Runtime/Engine/Private/WorldPartition/RuntimeHashSet/RuntimePartition.cpp
@@ -14,6 +14,13 @@ void URuntimePartition::SetDefaultValues()
 	bIsHLODSetup = false;
 }
 
+void URuntimePartition::InitHLODRuntimePartitionFrom(const URuntimePartition* RuntimePartition, int32 HLODIndex)
+{
+	Name = *FString::Printf(TEXT("%s_HLOD%d"), *RuntimePartition->Name.ToString(), HLODIndex);
+	LoadingRange = RuntimePartition->LoadingRange * 2;
+	bIsHLODSetup = true;
+}
+
 void URuntimePartition::PostEditChangeProperty(FPropertyChangedEvent& InPropertyChangedEvent)
 {
 	const FName PropertyName = InPropertyChangedEvent.GetPropertyName();
@@ -26,6 +33,13 @@ void URuntimePartition::PostEditChangeProperty(FPropertyChangedEvent& InProperty
 	Super::PostEditChangeProperty(InPropertyChangedEvent);
 }
 
+URuntimePartition* URuntimePartition::CreateHLODRuntimePartition(int32 HLODIndex) const
+{
+	URuntimePartition* HLODRuntimePartition = DuplicateObject<URuntimePartition>(this, GetOuter());
+	HLODRuntimePartition->InitHLODRuntimePartitionFrom(this, HLODIndex);
+	return HLODRuntimePartition;
+}
+
 URuntimePartition::FCellDesc URuntimePartition::CreateCellDesc(const FString& InName, bool bInIsSpatiallyLoaded, int32 InLevel, const TArray<const IStreamingGenerationContext::FActorSetInstance*>& InActorSetInstances)
 {
 	FCellDesc CellDesc;
diff --git a/Engine/Source/Runtime/Engine/Private/WorldPartition/RuntimeHashSet/RuntimePartitionLHGrid.cpp b/Engine/Source/Runtime/Engine/Private/WorldPartition/RuntimeHashSet/RuntimePartitionLHGrid.cpp
--- a/Engine/Source/Runtime/Engine/Private/WorldPartition/RuntimeHashSet/RuntimePartitionLHGrid.cpp
+++ b/Engine/Source/Runtime/Engine/Private/WorldPartition/RuntimeHashSet/RuntimePartitionLHGrid.cpp
@@ -68,6 +68,13 @@ void URuntimePartitionLHGrid::PostEditChangeProperty(FPropertyChangedEvent& Prop
 	Super::PostEditChangeProperty(PropertyChangedEvent);
 }
 
+void URuntimePartitionLHGrid::InitHLODRuntimePartitionFrom(const URuntimePartition* RuntimePartition, int32 HLODIndex)
+{
+	Super::InitHLODRuntimePartitionFrom(RuntimePartition, HLODIndex);
+	const URuntimePartitionLHGrid* RuntimePartitionLHGrid = CastChecked<const URuntimePartitionLHGrid>(RuntimePartition);
+	CellSize = RuntimePartitionLHGrid->CellSize * 2;
+}
+
 void URuntimePartitionLHGrid::SetDefaultValues()
 {
 	Super::SetDefaultValues();
diff --git a/Engine/Source/Runtime/Engine/Private/WorldPartition/RuntimeHashSet/WorldPartitionRuntimeHashSet.cpp b/Engine/Source/Runtime/Engine/Private/WorldPartition/RuntimeHashSet/WorldPartitionRuntimeHashSet.cpp
--- a/Engine/Source/Runtime/Engine/Private/WorldPartition/RuntimeHashSet/WorldPartitionRuntimeHashSet.cpp
+++ b/Engine/Source/Runtime/Engine/Private/WorldPartition/RuntimeHashSet/WorldPartitionRuntimeHashSet.cpp
@@ -17,59 +17,6 @@ FAutoConsoleCommand WorldPartitionRuntimeHashSetEnable(
 	})
 );
 
-#if WITH_EDITOR
-void FRuntimePartitionDesc::UpdateHLODPartitionLayers()
-{
-	if (!Class || !MainLayer)
-	{
-		HLODSetups.Empty();
-		return;
-	}
-
-	for (FRuntimePartitionHLODSetup& HLODSetup : HLODSetups)
-	{
-		TSet<const UHLODLayer*> VisitedHLODLayers;
-
-		const UHLODLayer* CurHLODLayer = HLODSetup.HLODLayer;
-		while (CurHLODLayer)
-		{
-			const int32 HLODSetupIndex = VisitedHLODLayers.Num();
-
-			bool bHLODLayerWasAlreadyInSet;
-			VisitedHLODLayers.Add(CurHLODLayer, &bHLODLayerWasAlreadyInSet);
-			if (bHLODLayerWasAlreadyInSet)
-			{
-				// Circular reference
-				break;
-			}
-
-			if (!HLODSetup.PartitionLayers.IsValidIndex(HLODSetupIndex))
-			{
-				HLODSetup.PartitionLayers.AddDefaulted();
-			}
-
-			FRuntimePartitionHLODSetupLayer& HLODSetupLayer = HLODSetup.PartitionLayers[HLODSetupIndex];
-
-			const bool bHLODLayerMatches = HLODSetupLayer.HLODLayer == CurHLODLayer;
-			const UClass* ExpectedHLODPartitionClass = CurHLODLayer->IsSpatiallyLoaded() ? MainLayer->GetClass() : URuntimePartitionPersistent::StaticClass();
-			const bool bHasValidPartitionLayer = HLODSetupLayer.PartitionLayer && (HLODSetupLayer.PartitionLayer->GetClass() == ExpectedHLODPartitionClass);
-			
-			if (!bHLODLayerMatches || !bHasValidPartitionLayer)
-			{
-				HLODSetupLayer.HLODLayer = CurHLODLayer;
-				HLODSetupLayer.PartitionLayer = CurHLODLayer->IsSpatiallyLoaded() ? DuplicateObject<URuntimePartition>(MainLayer, MainLayer->GetOuter()) : NewObject<URuntimePartition>(MainLayer->GetOuter(), ExpectedHLODPartitionClass);
-				HLODSetupLayer.PartitionLayer->Name = CurHLODLayer->GetFName();
-				HLODSetupLayer.PartitionLayer->bIsHLODSetup = true;
-			}
-
-			CurHLODLayer = CurHLODLayer->GetParentLayer();	
-		}
-
-		HLODSetup.PartitionLayers.SetNum(VisitedHLODLayers.Num());
-	}
-}
-#endif
-
 void FRuntimePartitionStreamingData::CreatePartitionsSpatialIndex() const
 {
 	if (!SpatialIndex)
@@ -130,13 +77,7 @@ void UWorldPartitionRuntimeHashSet::PostLoad()
 {
 	Super::PostLoad();
 
-#if WITH_EDITOR
-	if (!GetTypedOuter<UWorld>()->IsGameWorld())
-	{
-		UpdateHLODPartitionLayers();
-	}
-	else
-#endif
+	if (GetTypedOuter<UWorld>()->IsGameWorld())
 	{
 		ForEachStreamingData([](const FRuntimePartitionStreamingData& StreamingData)
 		{
@@ -239,7 +180,7 @@ bool UWorldPartitionRuntimeHashSet::IsValidHLODLayer(FName GridName, const FSoft
 
 		for (const FRuntimePartitionHLODSetup& HLODSetup : RuntimePartitions[RuntimePartitionIndex].HLODSetups)
 		{
-			if (HLODSetup.HLODLayer == HLODLayer)
+			if (HLODSetup.HLODLayers.Contains(HLODLayer))
 			{
 				return true;
 			}
@@ -518,7 +459,7 @@ void UWorldPartitionRuntimeHashSet::PostEditChangeChainProperty(FPropertyChanged
 	Super::PostEditChangeChainProperty(PropertyChangedEvent);
 
 	static FName NAME_RuntimePartitions(TEXT("RuntimePartitions"));
-	static FName NAME_HLODSetups_Key(TEXT("HLODSetups_Key"));
+	static FName NAME_HLODSetups(TEXT("HLODSetups"));
 	static FName NAME_HLODLayer(TEXT("HLODLayer"));
 
 	FName PropertyName = PropertyChangedEvent.Property ? PropertyChangedEvent.Property->GetFName() : NAME_None;
@@ -537,13 +478,6 @@ void UWorldPartitionRuntimeHashSet::PostEditChangeChainProperty(FPropertyChanged
 			RuntimePartitionDesc.Name = RuntimePartitionDesc.Class->GetFName();
 			RuntimePartitionDesc.MainLayer = NewObject<URuntimePartition>(this, RuntimePartitionDesc.Class, NAME_None);
 			RuntimePartitionDesc.MainLayer->SetDefaultValues();
-
-			if (UHLODLayer* DefaultHLODLayer = GetTypedOuter<UWorldPartition>()->GetDefaultHLODLayer())
-			{
-				RuntimePartitionDesc.HLODSetups.Emplace_GetRef().HLODLayer = DefaultHLODLayer;
-			}
-
-			RuntimePartitionDesc.UpdateHLODPartitionLayers();
 		}
 	}
 	else if (PropertyName == GET_MEMBER_NAME_CHECKED(FRuntimePartitionDesc, Name))
@@ -574,20 +508,21 @@ void UWorldPartitionRuntimeHashSet::PostEditChangeChainProperty(FPropertyChanged
 
 		RuntimePartitionDesc.MainLayer->Name = RuntimePartitionDesc.Name;
 	}
-	else if (PropertyName == NAME_HLODLayer)
+	else if (PropertyName == NAME_HLODSetups)
 	{
 		int32 RuntimePartitionIndex = PropertyChangedEvent.GetArrayIndex(NAME_RuntimePartitions.ToString());
-		check(RuntimePartitions.IsValidIndex(RuntimePartitionIndex));
-
-		RuntimePartitions[RuntimePartitionIndex].UpdateHLODPartitionLayers();
-	}
-}
+		if (RuntimePartitions.IsValidIndex(RuntimePartitionIndex))
+		{
+			FRuntimePartitionDesc& RuntimePartitionDesc = RuntimePartitions[RuntimePartitionIndex];
 
-void UWorldPartitionRuntimeHashSet::UpdateHLODPartitionLayers()
-{
-	for (FRuntimePartitionDesc& RuntimePartitionDesc : RuntimePartitions)
-	{
-		RuntimePartitionDesc.UpdateHLODPartitionLayers();
+			int32 HLODSetupsIndex = PropertyChangedEvent.GetArrayIndex(NAME_HLODSetups.ToString());
+			if (RuntimePartitionDesc.HLODSetups.IsValidIndex(HLODSetupsIndex))
+			{
+				FRuntimePartitionHLODSetup& RuntimePartitionHLODSetup = RuntimePartitionDesc.HLODSetups[HLODSetupsIndex];
+				URuntimePartition* ParentRuntimePartition = HLODSetupsIndex ? RuntimePartitionDesc.HLODSetups[HLODSetupsIndex - 1].PartitionLayer : RuntimePartitionDesc.MainLayer;
+				RuntimePartitionHLODSetup.PartitionLayer = ParentRuntimePartition->CreateHLODRuntimePartition(HLODSetupsIndex);
+			}
+		}
 	}
 }
 
@@ -666,4 +601,4 @@ void UWorldPartitionRuntimeHashSet::ForEachStreamingData(TFunctionRef<void(const
 			}
 		}
 	}
-}
\ No newline at end of file
+}
diff --git a/Engine/Source/Runtime/Engine/Private/WorldPartition/RuntimeHashSet/WorldPartitionRuntimeHashSetHLODGeneration.cpp b/Engine/Source/Runtime/Engine/Private/WorldPartition/RuntimeHashSet/WorldPartitionRuntimeHashSetHLODGeneration.cpp
--- a/Engine/Source/Runtime/Engine/Private/WorldPartition/RuntimeHashSet/WorldPartitionRuntimeHashSetHLODGeneration.cpp
+++ b/Engine/Source/Runtime/Engine/Private/WorldPartition/RuntimeHashSet/WorldPartitionRuntimeHashSetHLODGeneration.cpp
@@ -160,18 +160,8 @@ bool UWorldPartitionRuntimeHashSet::SupportsHLODs() const
 
 bool UWorldPartitionRuntimeHashSet::SetupHLODActors(const IStreamingGenerationContext* StreamingGenerationContext, const UWorldPartition::FSetupHLODActorsParams& Params) const
 {
-	auto GetRuntimePartitionDescHLODSetup = [](const FRuntimePartitionDesc* RuntimePartitionDesc, const UHLODLayer* HLODLayer) -> const FRuntimePartitionHLODSetup*
-	{
-		for (const FRuntimePartitionHLODSetup& HLODSetup : RuntimePartitionDesc->HLODSetups)
-		{
-			if (HLODLayer == HLODSetup.HLODLayer)
-			{
-				return &HLODSetup;
-			}
-		}
-
-		return nullptr;
-	};
+	IWorldPartitionHLODUtilities* WPHLODUtilities = FModuleManager::Get().LoadModuleChecked<IWorldPartitionHLODUtilitiesModule>("WorldPartitionHLODUtilities").GetUtilities();
+	check(WPHLODUtilities);
 
 	UWorldPartition* WorldPartition = GetOuterUWorldPartition();
 	const UDataLayerManager* DataLayerManager = WorldPartition->GetDataLayerManager();
@@ -198,9 +188,13 @@ bool UWorldPartitionRuntimeHashSet::SetupHLODActors(const IStreamingGenerationCo
 		TArray<FGuid> HLODActorGuids;
 		for (auto& [RuntimePartition, CellDescInstances] : RuntimePartitionsStreamingDescs)
 		{
+			int32 CellDescInstanceIndex = 0;
 			for (URuntimePartition::FCellDescInstance& CellDescInstance : CellDescInstances)
 			{
-				// Here we split actors into their respective HLOD layers because we want to provide a specific runtime grid name for our HLOD setup to FWorldPartitionHLODUtilities::CreateHLODActors. 
+				const FCellUniqueId CellUniqueId = GetCellUniqueId(CellDescInstance);
+
+				UE_LOG(LogWorldPartition, Display, TEXT("[%d / %d] Processing cell %s..."), ++CellDescInstanceIndex, CellDescInstances.Num(), *CellUniqueId.Name);
+
 				TArray<IStreamingGenerationContext::FActorInstance> ActorInstances;
 				for (const IStreamingGenerationContext::FActorSetInstance* ActorSetInstance : CellDescInstance.ActorSetInstances)
 				{
@@ -213,11 +207,14 @@ bool UWorldPartitionRuntimeHashSet::SetupHLODActors(const IStreamingGenerationCo
 				// Fake tick
 				PrivateUtils::GameTick(WorldPartition->GetWorld());
 
-				const FCellUniqueId CellUniqueId = GetCellUniqueId(CellDescInstance);
-
 				TArray<FName> MainPartitionTokens;
 				TArray<FName> HLODPartitionTokens;
 				verify(ParseGridName(ActorInstances[0].ActorSetInstance->RuntimeGrid, MainPartitionTokens, HLODPartitionTokens));
+
+				if (MainPartitionTokens[0].IsNone())
+				{
+					MainPartitionTokens[0] = RuntimePartitions[0].Name;
+				}
 					
 				FHLODCreationParams HLODCreationParams;
 				HLODCreationParams.WorldPartition = WorldPartition;
@@ -230,7 +227,6 @@ bool UWorldPartitionRuntimeHashSet::SetupHLODActors(const IStreamingGenerationCo
 				HLODCreationParams.ContentBundleGuid = CellDescInstance.ContentBundleID;
 				HLODCreationParams.DataLayerInstances = CellDescInstance.DataLayerInstances;
 
-				IWorldPartitionHLODUtilities* WPHLODUtilities = FModuleManager::Get().LoadModuleChecked<IWorldPartitionHLODUtilitiesModule>("WorldPartitionHLODUtilities").GetUtilities();
 				TArray<AWorldPartitionHLOD*> CellHLODActors = WPHLODUtilities->CreateHLODActors(HLODCreationContext, HLODCreationParams, ActorInstances);
 
 				if (!CellHLODActors.IsEmpty())
diff --git a/Engine/Source/Runtime/Engine/Private/WorldPartition/RuntimeHashSet/WorldPartitionRuntimeHashSetStreamingGeneration.cpp b/Engine/Source/Runtime/Engine/Private/WorldPartition/RuntimeHashSet/WorldPartitionRuntimeHashSetStreamingGeneration.cpp
--- a/Engine/Source/Runtime/Engine/Private/WorldPartition/RuntimeHashSet/WorldPartitionRuntimeHashSetStreamingGeneration.cpp
+++ b/Engine/Source/Runtime/Engine/Private/WorldPartition/RuntimeHashSet/WorldPartitionRuntimeHashSetStreamingGeneration.cpp
@@ -6,6 +6,7 @@
 #include "WorldPartition/ContentBundle/ContentBundleDescriptor.h"
 #include "WorldPartition/WorldPartitionStreamingPolicy.h"
 #include "WorldPartition/DataLayer/DataLayersID.h"
+#include "WorldPartition/HLOD/HLODLayer.h"
 
 #if WITH_EDITOR
 bool UWorldPartitionRuntimeHashSet::GenerateRuntimePartitionsStreamingDescs(const IStreamingGenerationContext* StreamingGenerationContext, TMap<URuntimePartition*, TArray<URuntimePartition::FCellDescInstance>>& OutRuntimeCellDescs) const
@@ -58,11 +59,11 @@ bool UWorldPartitionRuntimeHashSet::GenerateRuntimePartitionsStreamingDescs(cons
 
 				for (const FRuntimePartitionHLODSetup& HLODSetup : (*RuntimePartitionDesc)->HLODSetups)
 				{
-					for (const FRuntimePartitionHLODSetupLayer& HLODPartitionLayer : HLODSetup.PartitionLayers)
+					for (const UHLODLayer* HLODPartitionLayer : HLODSetup.HLODLayers)
 					{
-						if (HLODPartitionLayer.HLODLayer.GetName() == HLODPartitionTokens[0])
+						if (HLODPartitionLayer->GetName() == HLODPartitionTokens[0])
 						{
-							RuntimePartitionsToActorSetMap.FindOrAdd(HLODPartitionLayer.PartitionLayer).Add(&ActorSetInstance);
+							RuntimePartitionsToActorSetMap.FindOrAdd(HLODSetup.PartitionLayer).Add(&ActorSetInstance);
 							bFoundHLODPartitionLayer = true;
 							break;
 						}
@@ -280,4 +281,4 @@ void UWorldPartitionRuntimeHashSet::DumpStateLog(FHierarchicalLogArchive& Ar) co
 
 	Ar.Printf(TEXT(""));
 }
-#endif
\ No newline at end of file
+#endif
diff --git a/Engine/Source/Runtime/Engine/Public/WorldPartition/RuntimeHashSet/RuntimePartition.h b/Engine/Source/Runtime/Engine/Public/WorldPartition/RuntimeHashSet/RuntimePartition.h
--- a/Engine/Source/Runtime/Engine/Public/WorldPartition/RuntimeHashSet/RuntimePartition.h
+++ b/Engine/Source/Runtime/Engine/Public/WorldPartition/RuntimeHashSet/RuntimePartition.h
@@ -16,6 +16,8 @@ public:
 	virtual void PostEditChangeProperty(FPropertyChangedEvent& InPropertyChangedEvent) override;
 	//~ End UObject Interface.
 
+	URuntimePartition* CreateHLODRuntimePartition(int32 HLODIndex) const;
+
 	/*
 	 * Represents a cell descriptor, generated by runtime partitions. This is a streaming cell containing actors, without taking into account data layers, content bundles, etc.
 	 */
@@ -63,6 +65,7 @@ public:
 
 	virtual void SetDefaultValues();
 	virtual bool SupportsHLODs() const PURE_VIRTUAL(URuntimePartition::SupportsHLODs, return true;);
+	virtual void InitHLODRuntimePartitionFrom(const URuntimePartition* RuntimePartition, int32 HLODIndex);
 	virtual bool IsValidPartitionTokens(const TArray<FName>& InPartitionTokens) const PURE_VIRTUAL(URuntimePartition::IsValidPartitionTokens, return false;);
 	virtual bool GenerateStreaming(const FGenerateStreamingParams& InParams, FGenerateStreamingResult& OutResult) PURE_VIRTUAL(URuntimePartition::GenerateStreaming, return false;);
 #endif
diff --git a/Engine/Source/Runtime/Engine/Public/WorldPartition/RuntimeHashSet/RuntimePartitionLHGrid.h b/Engine/Source/Runtime/Engine/Public/WorldPartition/RuntimeHashSet/RuntimePartitionLHGrid.h
--- a/Engine/Source/Runtime/Engine/Public/WorldPartition/RuntimeHashSet/RuntimePartitionLHGrid.h
+++ b/Engine/Source/Runtime/Engine/Public/WorldPartition/RuntimeHashSet/RuntimePartitionLHGrid.h
@@ -18,6 +18,7 @@ public:
 
 	//~ Begin URuntimePartition interface
 	virtual bool SupportsHLODs() const override { return true; }
+	virtual void InitHLODRuntimePartitionFrom(const URuntimePartition* RuntimePartition, int32 HLODIndex);
 	virtual void SetDefaultValues() override;
 	virtual bool IsValidPartitionTokens(const TArray<FName>& InPartitionTokens) const override;
 	virtual bool GenerateStreaming(const FGenerateStreamingParams& InParams, FGenerateStreamingResult& OutResult) override;
diff --git a/Engine/Source/Runtime/Engine/Public/WorldPartition/RuntimeHashSet/WorldPartitionRuntimeHashSet.h b/Engine/Source/Runtime/Engine/Public/WorldPartition/RuntimeHashSet/WorldPartitionRuntimeHashSet.h
--- a/Engine/Source/Runtime/Engine/Public/WorldPartition/RuntimeHashSet/WorldPartitionRuntimeHashSet.h
+++ b/Engine/Source/Runtime/Engine/Public/WorldPartition/RuntimeHashSet/WorldPartitionRuntimeHashSet.h
@@ -13,21 +13,6 @@ struct FPropertyChangedChainEvent;
 
 using FStaticSpatialIndexType = TStaticSpatialIndexRTree<TObjectPtr<UWorldPartitionRuntimeCell>>;
 
-/** Holds settings for an HLOD layer for a particular partition class. */
-USTRUCT()
-struct FRuntimePartitionHLODSetupLayer
-{
-	GENERATED_USTRUCT_BODY()
-
-#if WITH_EDITORONLY_DATA
-	UPROPERTY(VisibleAnywhere, Category = RuntimeSettings, Meta = (DisplayThumbnail = false))
-	TObjectPtr<const UHLODLayer> HLODLayer;
-#endif
-
-	UPROPERTY(EditAnywhere, Category = RuntimeSettings, Instanced)
-	TObjectPtr<URuntimePartition> PartitionLayer;
-};
-
 /** Holds an HLOD setup for a particular partition class. */
 USTRUCT()
 struct FRuntimePartitionHLODSetup
@@ -35,12 +20,11 @@ struct FRuntimePartitionHLODSetup
 	GENERATED_USTRUCT_BODY()
 
 	/** Associated HLOD Layer objects */
-	UPROPERTY(EditAnywhere, Category = RuntimeSettings, Meta = (EditCondition = "Class != nullptr", HideEditConditionToggle, NoResetToDefault, ForceInlineRow))
-	TObjectPtr<const UHLODLayer> HLODLayer;
+	UPROPERTY(EditAnywhere, Category = RuntimeSettings)
+	TArray<TObjectPtr<const UHLODLayer>> HLODLayers;
 
-	/** HLOD setup, one for each layers in the hierarchy */
-	UPROPERTY(VisibleAnywhere, Category = RuntimeSettings, EditFixedSize, Meta = (EditCondition = "HLODLayer != nullptr", HideEditConditionToggle, ForceInlineRow))
-	TArray<FRuntimePartitionHLODSetupLayer> PartitionLayers;
+	UPROPERTY(VisibleAnywhere, Category = RuntimeSettings, Instanced)
+	TObjectPtr<URuntimePartition> PartitionLayer;
 };
 
 /** Holds settings for a runtime partition instance. */
@@ -182,9 +166,6 @@ private:
 	/** Generate the runtime partitions streaming descs. */
 	bool GenerateRuntimePartitionsStreamingDescs(const IStreamingGenerationContext* StreamingGenerationContext, TMap<URuntimePartition*, TArray<URuntimePartition::FCellDescInstance>>& OutRuntimeCellDescs) const;
 
-	/** Update the partition layers to reflect the curent HLOD setups. */
-	void UpdateHLODPartitionLayers();
-
 	struct FCellUniqueId
 	{
 		FString Name;
