Fix instance count tracking for legacy (HISM) path

diff --git a/Engine/Source/Runtime/Engine/Private/InstancedStaticMesh/ISMInstanceDataManager.cpp b/Engine/Source/Runtime/Engine/Private/InstancedStaticMesh/ISMInstanceDataManager.cpp
--- a/Engine/Source/Runtime/Engine/Private/InstancedStaticMesh/ISMInstanceDataManager.cpp
+++ b/Engine/Source/Runtime/Engine/Private/InstancedStaticMesh/ISMInstanceDataManager.cpp
@@ -615,7 +615,15 @@ bool FPrimitiveInstanceDataManager::FlushChanges(FInstanceUpdateComponentDesc &&
 	{
 		LOG_INST_DATA(TEXT("Rebuild with EState::External %s"), TEXT(""));
 
-		check(LegacyBuildData->InstanceIdIndexMap.GetMaxInstanceIndex() == LegacyBuildData->LegacyStaticMeshInstanceData->GetNumInstances());
+		if (Mode == EMode::ExternalLegacyData)
+		{
+			check(LegacyBuildData->LegacyInstanceReorderTable.IsEmpty());
+			check(LegacyBuildData->InstanceIdIndexMap.GetMaxInstanceIndex() == LegacyBuildData->LegacyStaticMeshInstanceData->GetNumInstances());
+		}
+		else
+		{
+			check(LegacyBuildData->InstanceIdIndexMap.GetMaxInstanceIndex() == LegacyBuildData->LegacyInstanceReorderTable.Num());
+		}
 
 		// An "external build" is responsible for setting everything up in the right space.
 		bPrimitiveTransformChanged = false;
@@ -641,7 +649,8 @@ bool FPrimitiveInstanceDataManager::FlushChanges(FInstanceUpdateComponentDesc &&
 
 		// Assemble header info to enable nonblocking primitive update.
 		FInstanceDataBufferHeader InstanceDataBufferHeader;
-		InstanceDataBufferHeader.NumInstances = ExternalChangeSet.InstanceIdIndexMap.GetMaxInstanceIndex();
+		// Note: the buffers will contain the number of instances in the LegacyStaticMeshInstanceData, which may be different from the number in the ISMC & what is tracked (density scaling can do this)
+		InstanceDataBufferHeader.NumInstances = ExternalChangeSet.LegacyStaticMeshInstanceData->GetNumInstances();
 		InstanceDataBufferHeader.PayloadDataStride = FInstanceSceneDataBuffers::CalcPayloadDataStride(ExternalChangeSet.Flags, NumCustomDataFloats, 0);
 		InstanceDataBufferHeader.Flags = ExternalChangeSet.Flags;
 
@@ -649,7 +658,7 @@ bool FPrimitiveInstanceDataManager::FlushChanges(FInstanceUpdateComponentDesc &&
 		{
 			FISMCInstanceDataSceneProxy &ProxyRef = *Proxy;
 			
-			// Forcibly destroy any tracking state, if the external entity manages this it can track the data on its own.
+			// Forcibly destroy any tracking state
 			ProxyRef.ChangeMask.Reset();
 			ProxyRef.InstanceIdIndexMap = MoveTemp(ExternalChangeSet.InstanceIdIndexMap);
 #if WITH_EDITOR
@@ -669,8 +678,9 @@ bool FPrimitiveInstanceDataManager::FlushChanges(FInstanceUpdateComponentDesc &&
 
 			ProxyRef.BuildFromLegacyData(MoveTemp(ExternalChangeSet.LegacyStaticMeshInstanceData), ExternalChangeSet.InstanceLocalBounds, MoveTemp(ExternalChangeSet.LegacyInstanceReorderTable));
 
-			check(ProxyRef.GetData().GetNumInstances() == ProxyRef.InstanceIdIndexMap.GetMaxInstanceIndex());
-			check(ProxyRef.GetUpdateTaskInfo()->GetHeader().NumInstances == ProxyRef.InstanceIdIndexMap.GetMaxInstanceIndex());
+			check(ProxyRef.GetData().GetNumInstances() == ProxyRef.GetUpdateTaskInfo()->GetHeader().NumInstances);
+			// Some instances may not be represented in the data due to density scaling, this should only be true if there is a reorder table that has the same size as the 
+			check(ProxyRef.GetData().GetNumInstances() <= ProxyRef.InstanceIdIndexMap.GetMaxInstanceIndex());
 			check(ProxyRef.GetUpdateTaskInfo()->GetHeader().Flags == ProxyRef.InstanceSceneDataBuffers.GetFlags());
 		});
 
@@ -1083,16 +1093,18 @@ void FPrimitiveInstanceDataManager::ValidateMapping() const
 void FPrimitiveInstanceDataManager::MarkForRebuildFromLegacy(TUniquePtr<FStaticMeshInstanceData> &&InLegacyInstanceData, const TArray<int32> &InstanceReorderTable, const TArray<TRefCountPtr<HHitProxy>> &HitProxies)
 {
 	check(InLegacyInstanceData);
-	// TODO: restore the ID tracking when not in external mode
-	//if (Mode == EMode::ExternalLegacyData)
+	if (Mode == EMode::ExternalLegacyData)
 	{
+		check(InstanceReorderTable.IsEmpty());
 		ClearIdTracking(InLegacyInstanceData->GetNumInstances());
+		check(GetMaxInstanceIndex() == InLegacyInstanceData->GetNumInstances());
+	}
+	else
+	{
+		// TODO: restore the ID tracking when not in external mode
+		ClearIdTracking(InstanceReorderTable.Num());
+		//ClearChangeTracking();
 	}
-	//else
-	//{
-	//	ClearChangeTracking();
-	//}
-	check(GetMaxInstanceIndex() == InLegacyInstanceData->GetNumInstances());
 
 	LegacyBuildData = MakePimpl<FLegacyBuildData>();
 	
