Add extra information to the terrain static permutation data that signifies a layer name is actually a repeat of another layer name.
This does  not happen durring normal landscape generation, but can be used to by external code to reassign layer names, particularly for runtime layer swapping.

diff --git a/Engine/Source/Runtime/Engine/Private/Materials/HLSLMaterialTranslator.cpp b/Engine/Source/Runtime/Engine/Private/Materials/HLSLMaterialTranslator.cpp
--- a/Engine/Source/Runtime/Engine/Private/Materials/HLSLMaterialTranslator.cpp
+++ b/Engine/Source/Runtime/Engine/Private/Materials/HLSLMaterialTranslator.cpp
@@ -8054,7 +8054,10 @@ int32 FHLSLMaterialTranslator::StaticTerrainLayerWeight(FName LayerName,int32 De
 
 		// If there are multiple weight maps with the same name, they should be numbered to allow for unique masks
 		FName LayerNameTest = Parameter.LayerName;
-		LayerNameTest.SetNumber(0);
+		if(Parameter.bIsRepeatedLayer)
+		{
+			LayerNameTest.SetNumber(0);
+		}
 
 		if(LayerNameTest != LayerName)
 		{
diff --git a/Engine/Source/Runtime/Engine/Public/StaticParameterSet.h b/Engine/Source/Runtime/Engine/Public/StaticParameterSet.h
--- a/Engine/Source/Runtime/Engine/Public/StaticParameterSet.h
+++ b/Engine/Source/Runtime/Engine/Public/StaticParameterSet.h
@@ -262,6 +262,9 @@ struct FStaticTerrainLayerWeightParameter
 	UPROPERTY()
 	int32 WeightmapIndex;
 
+	UPROPERTY()
+	bool bIsRepeatedLayer = false;
+
 	UPROPERTY()
 	bool bWeightBasedBlend;
 
@@ -276,6 +279,13 @@ struct FStaticTerrainLayerWeightParameter
 		bWeightBasedBlend(InWeightBasedBlend)
 	{ }
 
+	FStaticTerrainLayerWeightParameter(const FName& InName, int32 InWeightmapIndex, bool InWeightBasedBlend, bool InIsRepeatedLayer) :
+		LayerName(InName),
+		WeightmapIndex(InWeightmapIndex),
+		bIsRepeatedLayer(InIsRepeatedLayer),
+		bWeightBasedBlend(InWeightBasedBlend)
+	{ }
+
 	bool operator==(const FStaticTerrainLayerWeightParameter& Reference) const
 	{
 		return LayerName == Reference.LayerName && WeightmapIndex == Reference.WeightmapIndex && bWeightBasedBlend == Reference.bWeightBasedBlend;
