Expose Substrate MP_FrontMaterial material attribute to make it accessible from Python scripting.

diff --git a/Engine/Source/Runtime/Engine/Public/SceneTypes.h b/Engine/Source/Runtime/Engine/Public/SceneTypes.h
--- a/Engine/Source/Runtime/Engine/Public/SceneTypes.h
+++ b/Engine/Source/Runtime/Engine/Public/SceneTypes.h
@@ -189,7 +189,7 @@ enum EMaterialProperty : int
 	MP_CustomizedUVs7 UMETA(Hidden),
 	MP_PixelDepthOffset UMETA(Hidden),
 	MP_ShadingModel UMETA(Hidden),
-	MP_FrontMaterial UMETA(Hidden),
+	MP_FrontMaterial UMETA(DisplayName = "Front Material"),
 	MP_SurfaceThickness UMETA(Hidden),
 	MP_Displacement UMETA(Hidden),
 
