We have found some component templates that are not outered to the class, lets fix those up on load

diff --git a/Engine/Source/Runtime/Engine/Private/InheritableComponentHandler.cpp b/Engine/Source/Runtime/Engine/Private/InheritableComponentHandler.cpp
--- a/Engine/Source/Runtime/Engine/Private/InheritableComponentHandler.cpp
+++ b/Engine/Source/Runtime/Engine/Private/InheritableComponentHandler.cpp
@@ -45,7 +45,12 @@ void UInheritableComponentHandler::PostLoad()
 						Record.ComponentClass = Record.ComponentTemplate->GetClass();
 					}
 				}
-				
+
+				ensureMsgf(
+					Record.ComponentTemplate->HasAnyFlags(RF_LoadCompleted) && 
+					Record.ComponentTemplate->IsA(Record.ComponentClass),
+					TEXT("Encountered unloaded object while trying to conform component template names"));
+
 				// Fix up component template name on load, if it doesn't match the original template name. Otherwise, archetype lookups will fail for this template.
 				// For example, this can occur after a component variable rename in a parent BP class, but before a child BP class with an override template is loaded.
 				// Note: If the key maps to an SCS node, the node's variable GUID will be used for the lookup instead of the name below (that's only used for UCS keys).
@@ -504,6 +509,13 @@ void UInheritableComponentHandler::PreloadAll()
 		}
 	}
 	PreloadAllTemplates();
+	// This will get component names up to date - since these component 
+	// templates are used by GetArchetypeFromRequiredInfo we want to have real
+	// names as quickly as possible. The circumstances that require
+	// this are not clear to me, but that the logic exists and occasionally
+	// runs means we should run it ASAP - otherwise we may use the wrong
+	// archetype on construction.
+	ConditionalPostLoad();
 }
 
 FComponentKey UInheritableComponentHandler::FindKey(const FName VariableName) const
diff --git a/Engine/Source/Runtime/Engine/Private/SimpleConstructionScript.cpp b/Engine/Source/Runtime/Engine/Private/SimpleConstructionScript.cpp
--- a/Engine/Source/Runtime/Engine/Private/SimpleConstructionScript.cpp
+++ b/Engine/Source/Runtime/Engine/Private/SimpleConstructionScript.cpp
@@ -1508,6 +1508,11 @@ void USimpleConstructionScript::ValidateNodeVariableNames(FCompilerResultsLog& M
 void USimpleConstructionScript::ValidateNodeTemplates(FCompilerResultsLog& MessageLog)
 {
 	TArray<USCS_Node*> Nodes = GetAllNodes();
+	UClass* OwningClass = GetOwnerClass();
+
+	// it's likely that we shouldn't run any of this for diff loaded packages, but don't 
+	// want to destabilize anyone.. so just soldering off my new logic
+	const bool bForDiff = OwningClass->GetOutermost()->HasAnyPackageFlags(PKG_ForDiffing);
 
 	for (USCS_Node* Node : Nodes)
 	{
@@ -1548,6 +1553,34 @@ void USimpleConstructionScript::ValidateNodeTemplates(FCompilerResultsLog& Messa
 				RemoveNodeAndPromoteChildren(Node);
 			}
 		}
+		else if (!bForDiff)
+		{
+			if (Node->ComponentTemplate->GetOuter() != OwningClass)
+			{
+				// this component template is somehow not owned by the class, recreate it:
+				FString VariableName = Node->GetVariableName().ToString();
+				if (Node->ComponentTemplate->HasAnyFlags(RF_ClassDefaultObject))
+				{
+					// duplicate won't work on CDOs because SDO will (for no good reason) reset loaders when duplicating a CDO...
+					Node->ComponentTemplate = NewObject<UActorComponent>(
+						OwningClass, 
+						Node->ComponentTemplate->GetClass(), 
+						*(VariableName + ComponentTemplateNameSuffix), 
+						RF_ArchetypeObject | RF_Transactional | RF_Public);
+				}
+				else
+				{
+					Node->ComponentTemplate = static_cast<UActorComponent*>(StaticDuplicateObject(
+						Node->ComponentTemplate, 
+						OwningClass, 
+						*(VariableName + ComponentTemplateNameSuffix)));
+					Node->ComponentTemplate->SetFlags(RF_ArchetypeObject | RF_Transactional | RF_Public);
+				}
+				MessageLog.Warning(*FText::Format(NSLOCTEXT(
+					"SimpleConstructionScript", "CorruptComponentFixed", 
+					"SCS Node component template {0} has been recreated because it was unexpectedly not owned by the class"), FText::FromString(VariableName)).ToString());
+			}
+		}
 	}
 }
 
