clarify log message.

diff --git a/Engine/Source/Runtime/Engine/Private/TextureDerivedDataTask.cpp b/Engine/Source/Runtime/Engine/Private/TextureDerivedDataTask.cpp
--- a/Engine/Source/Runtime/Engine/Private/TextureDerivedDataTask.cpp
+++ b/Engine/Source/Runtime/Engine/Private/TextureDerivedDataTask.cpp
@@ -1613,7 +1613,7 @@ static void DDC1_FetchAndFillDerivedData(
 
 			if (!bSucceeded)
 			{
-				UE_LOG(LogTexture, Display, TEXT("Texture %s is missing inline mips. The texture will be rebuilt."), *TexturePathName);
+				UE_LOG(LogTexture, Display, TEXT("Texture %s is missing streaming mips when loading for an inline request. The texture will be rebuilt."), *TexturePathName);
 			}
 		}
 		else
