Make CanEvaluateRawAnimationData in UAnimSequnce protected and virtual.

diff --git a/Engine/Source/Runtime/Engine/Classes/Animation/AnimSequence.h b/Engine/Source/Runtime/Engine/Classes/Animation/AnimSequence.h
--- a/Engine/Source/Runtime/Engine/Classes/Animation/AnimSequence.h
+++ b/Engine/Source/Runtime/Engine/Classes/Animation/AnimSequence.h
@@ -768,9 +768,11 @@ private:
 
 	ENGINE_API void GetBonePose_AdditiveMeshRotationOnly(FAnimationPoseData& OutAnimationPoseData, const FAnimExtractContext& ExtractionContext) const;
 
+protected:
 	/** Returns whether or not evaluation of the raw (source) animation data is possible according to whether or not the (editor only) data has been stripped */
-	ENGINE_API bool CanEvaluateRawAnimationData() const;
+	ENGINE_API virtual bool CanEvaluateRawAnimationData() const;
 
+private:
 #if WITH_EDITOR
 	/**
 	 * Remap Tracks to New Skeleton
