Fix XR incorrectly using desktop Screen Percentage mode setting when launching the game, this time modified for minimal impact to avoid UE-197369.

diff --git a/Engine/Source/Runtime/Engine/Private/GameViewportClient.cpp b/Engine/Source/Runtime/Engine/Private/GameViewportClient.cpp
--- a/Engine/Source/Runtime/Engine/Private/GameViewportClient.cpp
+++ b/Engine/Source/Runtime/Engine/Private/GameViewportClient.cpp
@@ -1409,10 +1409,17 @@ void UGameViewportClient::Draw(FViewport* InViewport, FCanvas* SceneCanvas)
 		ViewExt->SetupViewFamily(ViewFamily);
 	}
 
-	if (bStereoRendering && GEngine->XRSystem.IsValid() && GEngine->XRSystem->GetHMDDevice())
+	if (bStereoRendering)
 	{
-		// Allow HMD to modify screen settings
-		GEngine->XRSystem->GetHMDDevice()->UpdateScreenSettings(Viewport);
+		if (GEngine->XRSystem.IsValid() && GEngine->XRSystem->GetHMDDevice())
+		{
+			// Allow HMD to modify screen settings
+			GEngine->XRSystem->GetHMDDevice()->UpdateScreenSettings(Viewport);
+		}
+		
+		// Update stereo flag in viewport client so we can accurately run GetViewStatusForScreenPercentage()
+		static bool bEmulateStereo = FParse::Param(FCommandLine::Get(), TEXT("emulatestereo"));
+		EngineShowFlags.StereoRendering = bEmulateStereo ? true : ViewFamily.EngineShowFlags.StereoRendering;
 	}
 
 	ESplitScreenType::Type SplitScreenConfig = GetCurrentSplitscreenConfiguration();
