Remove new ensure, it is overzealous.. we only need the name at this stage of component conformance

diff --git a/Engine/Source/Runtime/Engine/Private/InheritableComponentHandler.cpp b/Engine/Source/Runtime/Engine/Private/InheritableComponentHandler.cpp
--- a/Engine/Source/Runtime/Engine/Private/InheritableComponentHandler.cpp
+++ b/Engine/Source/Runtime/Engine/Private/InheritableComponentHandler.cpp
@@ -46,11 +46,6 @@ void UInheritableComponentHandler::PostLoad()
 					}
 				}
 
-				ensureMsgf(
-					Record.ComponentTemplate->HasAnyFlags(RF_LoadCompleted) && 
-					Record.ComponentTemplate->IsA(Record.ComponentClass),
-					TEXT("Encountered unloaded object while trying to conform component template names"));
-
 				// Fix up component template name on load, if it doesn't match the original template name. Otherwise, archetype lookups will fail for this template.
 				// For example, this can occur after a component variable rename in a parent BP class, but before a child BP class with an override template is loaded.
 				// Note: If the key maps to an SCS node, the node's variable GUID will be used for the lookup instead of the name below (that's only used for UCS keys).
