Local Exposure with Thresholds
- control over which regions are considered highlights and shadows using thresholds around middle grey luminance.

diff --git a/Engine/Shaders/Private/PostProcessHistogramCommon.ush b/Engine/Shaders/Private/PostProcessHistogramCommon.ush
--- a/Engine/Shaders/Private/PostProcessHistogramCommon.ush
+++ b/Engine/Shaders/Private/PostProcessHistogramCommon.ush
@@ -38,6 +38,9 @@ float LocalExposure_BlurredLuminanceBlend;
 float LocalExposure_MiddleGreyExposureCompensation;
 float2 LocalExposure_BilateralGridUVScale;
 
+float LocalExposure_HighlightThreshold;
+float LocalExposure_ShadowThreshold;
+
 float CalculateEyeAdaptationLuminance(float3 Color)
 {
 	return max(dot(Color, float3(1.0f, 1.0f, 1.0f) / 3.0f), EyeAdaptation_LuminanceMin);
@@ -245,7 +248,24 @@ float CalculateLogLocalExposure(float LogLum, float BaseLogLum, float LogMiddleG
 	float BaseCentered = (BaseLogLum - LogMiddleGrey);
 	float ContrastScale = BaseCentered > 0 ? HighlightContrastScale : ShadowContrastScale;
 
-	float LogLocalLum = LogMiddleGrey + BaseCentered * ContrastScale + DetailLogLum * DetailStrength;
+	// Apply threshold
+
+	float ThresholdOffset;
+	{
+		if (BaseCentered > 0)
+		{
+			ThresholdOffset = BaseCentered - max(0.0f, BaseCentered - LocalExposure_HighlightThreshold);
+		}
+		else
+		{
+			ThresholdOffset = BaseCentered - min(0.0f, BaseCentered + LocalExposure_ShadowThreshold);
+		}
+
+		BaseCentered -= ThresholdOffset;
+	}
+
+	float LogLocalLum = LogMiddleGrey + ThresholdOffset + BaseCentered * ContrastScale + DetailLogLum * DetailStrength;
+
 	return LogLocalLum - LogLum;
 }
 
diff --git a/Engine/Source/Runtime/Engine/Classes/Engine/Scene.h b/Engine/Source/Runtime/Engine/Classes/Engine/Scene.h
--- a/Engine/Source/Runtime/Engine/Classes/Engine/Scene.h
+++ b/Engine/Source/Runtime/Engine/Classes/Engine/Scene.h
@@ -913,6 +913,12 @@ struct FPostProcessSettings
 	UPROPERTY(EditAnywhere, BlueprintReadWrite, Category=Overrides, meta=(PinHiddenByDefault, InlineEditConditionToggle))
 	uint8 bOverride_LocalExposureShadowContrastCurve:1;
 
+	UPROPERTY(EditAnywhere, BlueprintReadWrite, Category=Overrides, meta=(PinHiddenByDefault, InlineEditConditionToggle))
+	uint8 bOverride_LocalExposureHighlightThreshold:1;
+
+	UPROPERTY(EditAnywhere, BlueprintReadWrite, Category=Overrides, meta=(PinHiddenByDefault, InlineEditConditionToggle))
+	uint8 bOverride_LocalExposureShadowThreshold :1;
+
 	UPROPERTY(EditAnywhere, BlueprintReadWrite, Category=Overrides, meta=(PinHiddenByDefault, InlineEditConditionToggle))
 	uint8 bOverride_LocalExposureDetailStrength:1;
 
@@ -1864,6 +1870,18 @@ struct FPostProcessSettings
 	UPROPERTY(EditAnywhere, BlueprintReadWrite, Category = "Lens|Local Exposure", meta = (editcondition = "bOverride_LocalExposureShadowContrastCurve", DisplayName = "Shadow Contrast Curve"))
 	TObjectPtr<class UCurveFloat> LocalExposureShadowContrastCurve = nullptr;
 
+	/** 
+	 * Threshold used to determine which regions of the screen are considered highlights.
+	*/
+	UPROPERTY(interp, BlueprintReadWrite, Category="Lens|Local Exposure", meta=(UIMin = "0.0", UIMax = "4.0", editcondition = "bOverride_LocalExposureHighlightThreshold", DisplayName = "Highlight Threshold"))
+	float LocalExposureHighlightThreshold;
+
+	/**
+	 * Threshold used to determine which regions of the screen are considered shadows.
+	*/
+	UPROPERTY(interp, BlueprintReadWrite, Category="Lens|Local Exposure", meta=(UIMin = "0.0", UIMax = "4.0", editcondition = "bOverride_LocalExposureShadowThreshold", DisplayName = "Shadow Threshold"))
+	float LocalExposureShadowThreshold;
+
 	/**
 	 * Local Exposure decomposes luminance of the frame into a base layer and a detail layer.
 	 * Value different than 1 will enable local exposure.
diff --git a/Engine/Source/Runtime/Engine/Private/SceneView.cpp b/Engine/Source/Runtime/Engine/Private/SceneView.cpp
--- a/Engine/Source/Runtime/Engine/Private/SceneView.cpp
+++ b/Engine/Source/Runtime/Engine/Private/SceneView.cpp
@@ -1547,6 +1547,8 @@ void FSceneView::OverridePostProcessSettings(const FPostProcessSettings& Src, fl
 		LERP_PP(LocalExposureContrastScale_DEPRECATED);
 		LERP_PP(LocalExposureHighlightContrastScale);
 		LERP_PP(LocalExposureShadowContrastScale);
+		LERP_PP(LocalExposureHighlightThreshold);
+		LERP_PP(LocalExposureShadowThreshold);
 		LERP_PP(LocalExposureDetailStrength);
 		LERP_PP(LocalExposureBlurredLuminanceBlend);
 		LERP_PP(LocalExposureBlurredLuminanceKernelSizePercent);
diff --git a/Engine/Source/Runtime/Renderer/Private/PostProcess/PostProcessLocalExposure.cpp b/Engine/Source/Runtime/Renderer/Private/PostProcess/PostProcessLocalExposure.cpp
--- a/Engine/Source/Runtime/Renderer/Private/PostProcess/PostProcessLocalExposure.cpp
+++ b/Engine/Source/Runtime/Renderer/Private/PostProcess/PostProcessLocalExposure.cpp
@@ -139,6 +139,8 @@ FLocalExposureParameters GetLocalExposureParameters(const FViewInfo& View, FIntP
 	Parameters.BlurredLuminanceBlend = Settings.LocalExposureBlurredLuminanceBlend;
 	Parameters.MiddleGreyExposureCompensation = LocalExposureMiddleGreyExposureCompensation;
 	Parameters.BilateralGridUVScale = LocalExposureBilateralGridUVScale;
+	Parameters.HighlightThreshold = Settings.LocalExposureHighlightThreshold;
+	Parameters.ShadowThreshold = Settings.LocalExposureShadowThreshold;
 	return Parameters;
 }
 
diff --git a/Engine/Source/Runtime/Renderer/Private/PostProcess/PostProcessLocalExposure.h b/Engine/Source/Runtime/Renderer/Private/PostProcess/PostProcessLocalExposure.h
--- a/Engine/Source/Runtime/Renderer/Private/PostProcess/PostProcessLocalExposure.h
+++ b/Engine/Source/Runtime/Renderer/Private/PostProcess/PostProcessLocalExposure.h
@@ -13,6 +13,8 @@ BEGIN_SHADER_PARAMETER_STRUCT(FLocalExposureParameters, )
 	SHADER_PARAMETER(float, BlurredLuminanceBlend)
 	SHADER_PARAMETER(float, MiddleGreyExposureCompensation)
 	SHADER_PARAMETER(FVector2f, BilateralGridUVScale)
+	SHADER_PARAMETER(float, HighlightThreshold)
+	SHADER_PARAMETER(float, ShadowThreshold)
 END_SHADER_PARAMETER_STRUCT()
 
 FLocalExposureParameters GetLocalExposureParameters(const FViewInfo& View, FIntPoint ViewRectSize, const FEyeAdaptationParameters& EyeAdaptationParameters);
