Added logs when UActorDescContainer::Initialize finds invalid actors.

diff --git a/Engine/Source/Runtime/Engine/Private/WorldPartition/ActorDescContainer.cpp b/Engine/Source/Runtime/Engine/Private/WorldPartition/ActorDescContainer.cpp
--- a/Engine/Source/Runtime/Engine/Private/WorldPartition/ActorDescContainer.cpp
+++ b/Engine/Source/Runtime/Engine/Private/WorldPartition/ActorDescContainer.cpp
@@ -97,9 +97,24 @@ void UActorDescContainer::Initialize(const FInitializeParams& InitParams)
 					*ExistingDesc->GetActorPackage().ToString());
 				bValid = false;
 			}
-			else if(!ActorDesc->GetNativeClass().IsValid() || 
-					(ActorDesc->GetBaseClass().IsValid() && !ClassDescRegistry.IsRegisteredClass(ActorDesc->GetBaseClass())) || 
-					(InitParams.FilterActorDesc && !InitParams.FilterActorDesc(ActorDesc.Get())))
+			else if (!ActorDesc->GetNativeClass().IsValid())
+			{
+				UE_LOG(LogWorldPartition, Warning, TEXT("Invalid actor native class: Actor: '%s' (guid '%s') from package '%s'"),
+					*ActorDesc->GetActorName().ToString(),
+					*ActorDesc->GetGuid().ToString(),
+					*ActorDesc->GetActorPackage().ToString());
+				bValid = false;
+			}
+			else if (ActorDesc->GetBaseClass().IsValid() && !ClassDescRegistry.IsRegisteredClass(ActorDesc->GetBaseClass()))
+			{
+				UE_LOG(LogWorldPartition, Warning, TEXT("Unknown actor base class `%s`: Actor: '%s' (guid '%s') from package '%s'"),
+					*ActorDesc->GetBaseClass().ToString(),
+					*ActorDesc->GetActorName().ToString(),
+					*ActorDesc->GetGuid().ToString(),
+					*ActorDesc->GetActorPackage().ToString());
+				bValid = false;
+			}
+			else if (InitParams.FilterActorDesc && !InitParams.FilterActorDesc(ActorDesc.Get()))
 			{
 				bValid = false;
 			}
diff --git a/Engine/Source/Runtime/Engine/Private/WorldPartition/WorldPartition.cpp b/Engine/Source/Runtime/Engine/Private/WorldPartition/WorldPartition.cpp
--- a/Engine/Source/Runtime/Engine/Private/WorldPartition/WorldPartition.cpp
+++ b/Engine/Source/Runtime/Engine/Private/WorldPartition/WorldPartition.cpp
@@ -2061,7 +2061,16 @@ UActorDescContainer* UWorldPartition::RegisterActorDescContainer(const FContaine
 			}
 
 			// Filter actors with duplicated GUID in WorldPartition
-			return !GetActorDesc(ActorDesc->GetGuid());
+			if (GetActorDesc(ActorDesc->GetGuid()))
+			{
+				UE_LOG(LogWorldPartition, Warning, TEXT("Found existing actor descriptor guid `%s`: Actor: '%s' from package '%s'"),
+					*ActorDesc->GetGuid().ToString(),
+					*ActorDesc->GetActorName().ToString(),
+					*ActorDesc->GetActorPackage().ToString());
+				return false;
+			}
+
+			return true;
 		};
 
 		UActorDescContainer* ContainerToRegister = NewObject<UActorDescContainer>(this);
