Shader Compiler:  generate normalized line endings up front (CRLF -> LF) for procedural shader includes, including Niagara, /Engine/Generated/Material.ush, and /Engine/Generated/VertexFactory.ush, to support an optimized version of ShaderConvertAndStripComments (4x faster, ready to go in a follow up CL).

#rnx

diff --git a/Engine/Plugins/FX/Niagara/Source/NiagaraEditor/Private/NiagaraHlslTranslator.cpp b/Engine/Plugins/FX/Niagara/Source/NiagaraEditor/Private/NiagaraHlslTranslator.cpp
--- a/Engine/Plugins/FX/Niagara/Source/NiagaraEditor/Private/NiagaraHlslTranslator.cpp
+++ b/Engine/Plugins/FX/Niagara/Source/NiagaraEditor/Private/NiagaraHlslTranslator.cpp
@@ -8371,6 +8371,7 @@ void TNiagaraHlslTranslator<GraphBridge>::HandleCustomHlslNode(const FCustomHlsl
 	OutSignature.Name = *OutName; // Force the name to be set to include the node guid for safety...
 	bOutCustomHlsl = true;
 	OutCustomHlsl = GraphBridge::GetCustomHlslString(CustomFunctionHlsl);
+	OutCustomHlsl.ReplaceInline(TEXT("\r\n"), TEXT("\n"), ESearchCase::CaseSensitive);
 	GraphBridge::GetCustomHlslIncludePaths(CustomFunctionHlsl, OutCustomHlslIncludeFilePaths);
 
 	FNiagaraFunctionSignature InSignature = CustomFunctionHlsl->Signature;
diff --git a/Engine/Plugins/Tests/RuntimeTests/Source/RuntimeTests/Private/Shaders/ShaderComparisonTests.cpp b/Engine/Plugins/Tests/RuntimeTests/Source/RuntimeTests/Private/Shaders/ShaderComparisonTests.cpp
--- a/Engine/Plugins/Tests/RuntimeTests/Source/RuntimeTests/Private/Shaders/ShaderComparisonTests.cpp
+++ b/Engine/Plugins/Tests/RuntimeTests/Source/RuntimeTests/Private/Shaders/ShaderComparisonTests.cpp
@@ -186,7 +186,7 @@ bool FCompareBasepassShaders::RunTest(const FString& Parameters)
 
     AddCommand(new FDelayedFunctionLatentCommand([=] {
         FString str;
-        FString DefineStr = FString("#define ") + Parameters + " 1\r\n";
+        FString DefineStr = FString("#define ") + Parameters + " 1\n";
         for (int i = 0; i < ShaderNames.Num(); ++i) {
             FFileHelper::LoadFileToString((*FileStrings)[i], *ShaderNames[i]);
             str = DefineStr + (*FileStrings)[i];
diff --git a/Engine/Source/Runtime/Engine/Private/Materials/HLSLMaterialDerivativeAutogen.cpp b/Engine/Source/Runtime/Engine/Private/Materials/HLSLMaterialDerivativeAutogen.cpp
--- a/Engine/Source/Runtime/Engine/Private/Materials/HLSLMaterialDerivativeAutogen.cpp
+++ b/Engine/Source/Runtime/Engine/Private/Materials/HLSLMaterialDerivativeAutogen.cpp
@@ -1365,15 +1365,15 @@ FString FMaterialDerivativeAutogen::GenerateUsedFunctions(FHLSLMaterialTranslato
 			FString FieldName = GetFloatVectorName(DerivType);
 			FString FieldNameDDXY = GetFloatVectorDDXYName(DerivType);
 
-			Ret += BaseName + TEXT(" Construct") + BaseName + TEXT("(") + FieldName + TEXT(" InValue,") + FieldNameDDXY + TEXT(" InDdx,") + FieldNameDDXY + TEXT(" InDdy)") LINE_TERMINATOR;
-			Ret += TEXT("{") LINE_TERMINATOR;
-			Ret += TEXT("\t") + BaseName + TEXT(" Ret;") LINE_TERMINATOR;
-			Ret += TEXT("\tRet.Value = InValue;") LINE_TERMINATOR;
-			Ret += TEXT("\tRet.Ddx = InDdx;") LINE_TERMINATOR;
-			Ret += TEXT("\tRet.Ddy = InDdy;") LINE_TERMINATOR;
-			Ret += TEXT("\treturn Ret;") LINE_TERMINATOR;
-			Ret += TEXT("}") LINE_TERMINATOR;
-			Ret += TEXT("") LINE_TERMINATOR;
+			Ret += BaseName + TEXT(" Construct") + BaseName + TEXT("(") + FieldName + TEXT(" InValue,") + FieldNameDDXY + TEXT(" InDdx,") + FieldNameDDXY + TEXT(" InDdy)") HLSL_LINE_TERMINATOR;
+			Ret += TEXT("{") HLSL_LINE_TERMINATOR;
+			Ret += TEXT("\t") + BaseName + TEXT(" Ret;") HLSL_LINE_TERMINATOR;
+			Ret += TEXT("\tRet.Value = InValue;") HLSL_LINE_TERMINATOR;
+			Ret += TEXT("\tRet.Ddx = InDdx;") HLSL_LINE_TERMINATOR;
+			Ret += TEXT("\tRet.Ddy = InDdy;") HLSL_LINE_TERMINATOR;
+			Ret += TEXT("\treturn Ret;") HLSL_LINE_TERMINATOR;
+			Ret += TEXT("}") HLSL_LINE_TERMINATOR;
+			Ret += TEXT("") HLSL_LINE_TERMINATOR;
 		}
 	}
 
@@ -1386,15 +1386,15 @@ FString FMaterialDerivativeAutogen::GenerateUsedFunctions(FHLSLMaterialTranslato
 			FString BaseName = GetDerivVectorName(DerivType);
 			FString FieldName = GetFloatVectorName(DerivType);
 
-			Ret += BaseName + TEXT(" ConstructConstant") + BaseName + TEXT("(") + FieldName + TEXT(" Value)") LINE_TERMINATOR;
-			Ret += TEXT("{") LINE_TERMINATOR;
-			Ret += TEXT("\t") + BaseName + TEXT(" Ret;") LINE_TERMINATOR;
-			Ret += TEXT("\tRet.Value = Value;") LINE_TERMINATOR;
-			Ret += TEXT("\tRet.Ddx = 0;") LINE_TERMINATOR;
-			Ret += TEXT("\tRet.Ddy = 0;") LINE_TERMINATOR;
-			Ret += TEXT("\treturn Ret;") LINE_TERMINATOR;
-			Ret += TEXT("}") LINE_TERMINATOR;
-			Ret += TEXT("") LINE_TERMINATOR;
+			Ret += BaseName + TEXT(" ConstructConstant") + BaseName + TEXT("(") + FieldName + TEXT(" Value)") HLSL_LINE_TERMINATOR;
+			Ret += TEXT("{") HLSL_LINE_TERMINATOR;
+			Ret += TEXT("\t") + BaseName + TEXT(" Ret;") HLSL_LINE_TERMINATOR;
+			Ret += TEXT("\tRet.Value = Value;") HLSL_LINE_TERMINATOR;
+			Ret += TEXT("\tRet.Ddx = 0;") HLSL_LINE_TERMINATOR;
+			Ret += TEXT("\tRet.Ddy = 0;") HLSL_LINE_TERMINATOR;
+			Ret += TEXT("\treturn Ret;") HLSL_LINE_TERMINATOR;
+			Ret += TEXT("}") HLSL_LINE_TERMINATOR;
+			Ret += TEXT("") HLSL_LINE_TERMINATOR;
 		}
 	}
 
@@ -1407,24 +1407,24 @@ FString FMaterialDerivativeAutogen::GenerateUsedFunctions(FHLSLMaterialTranslato
 			FString BaseName = GetDerivVectorName(DerivType);
 			FString FieldName = GetFloatVectorName(DerivType);
 
-			Ret += BaseName + TEXT(" ConstructFinite") + BaseName + TEXT("(") + FieldName + TEXT(" InValue)") LINE_TERMINATOR;
-			Ret += TEXT("{") LINE_TERMINATOR;
-			Ret += TEXT("\t") + BaseName + TEXT(" Ret;") LINE_TERMINATOR;
-			Ret += TEXT("\tRet.Value = InValue;") LINE_TERMINATOR;
+			Ret += BaseName + TEXT(" ConstructFinite") + BaseName + TEXT("(") + FieldName + TEXT(" InValue)") HLSL_LINE_TERMINATOR;
+			Ret += TEXT("{") HLSL_LINE_TERMINATOR;
+			Ret += TEXT("\t") + BaseName + TEXT(" Ret;") HLSL_LINE_TERMINATOR;
+			Ret += TEXT("\tRet.Value = InValue;") HLSL_LINE_TERMINATOR;
 			if (Index == 4)
 			{
 				// LWC_TODO: Better way to do this??
-				Ret += TEXT("\tRet.Ddx = ddx(LWCHackToFloat(InValue));") LINE_TERMINATOR;
-				Ret += TEXT("\tRet.Ddy = ddy(LWCHackToFloat(InValue));") LINE_TERMINATOR;
+				Ret += TEXT("\tRet.Ddx = ddx(LWCHackToFloat(InValue));") HLSL_LINE_TERMINATOR;
+				Ret += TEXT("\tRet.Ddy = ddy(LWCHackToFloat(InValue));") HLSL_LINE_TERMINATOR;
 			}
 			else
 			{
-				Ret += TEXT("\tRet.Ddx = ddx(InValue);") LINE_TERMINATOR;
-				Ret += TEXT("\tRet.Ddy = ddy(InValue);") LINE_TERMINATOR;
+				Ret += TEXT("\tRet.Ddx = ddx(InValue);") HLSL_LINE_TERMINATOR;
+				Ret += TEXT("\tRet.Ddy = ddy(InValue);") HLSL_LINE_TERMINATOR;
 			}
-			Ret += TEXT("\treturn Ret;") LINE_TERMINATOR;
-			Ret += TEXT("}") LINE_TERMINATOR;
-			Ret += TEXT("") LINE_TERMINATOR;
+			Ret += TEXT("\treturn Ret;") HLSL_LINE_TERMINATOR;
+			Ret += TEXT("}") HLSL_LINE_TERMINATOR;
+			Ret += TEXT("") HLSL_LINE_TERMINATOR;
 		}
 	}
 
@@ -1446,15 +1446,15 @@ FString FMaterialDerivativeAutogen::GenerateUsedFunctions(FHLSLMaterialTranslato
 				const EDerivativeType SrcTypeDDXY = MakeNonLWCType(SrcDerivType);
 				const EDerivativeType DstTypeDDXY = MakeNonLWCType(DstDerivType);
 
-				Ret += DstBaseName + TEXT(" Convert_") + DstBaseName + TEXT("_") + SrcBaseName + TEXT("(") + SrcBaseName + TEXT(" Src)") LINE_TERMINATOR;
-				Ret += TEXT("{") LINE_TERMINATOR;
-				Ret += TEXT("\t") + DstBaseName + TEXT(" Ret;") LINE_TERMINATOR;
-				Ret += TEXT("\tRet.Value = ") + CoerceFloat(Translator, TEXT("Src.Value"), DstDerivType, SrcDerivType) + TEXT(";") LINE_TERMINATOR;
-				Ret += TEXT("\tRet.Ddx = ") + CoerceFloat(Translator, TEXT("Src.Ddx"), DstTypeDDXY, SrcTypeDDXY) + TEXT(";") LINE_TERMINATOR;
-				Ret += TEXT("\tRet.Ddy = ") + CoerceFloat(Translator, TEXT("Src.Ddy"), DstTypeDDXY, SrcTypeDDXY) + TEXT(";") LINE_TERMINATOR;
-				Ret += TEXT("\treturn Ret;") LINE_TERMINATOR;
-				Ret += TEXT("}") LINE_TERMINATOR;
-				Ret += TEXT("") LINE_TERMINATOR;
+				Ret += DstBaseName + TEXT(" Convert_") + DstBaseName + TEXT("_") + SrcBaseName + TEXT("(") + SrcBaseName + TEXT(" Src)") HLSL_LINE_TERMINATOR;
+				Ret += TEXT("{") HLSL_LINE_TERMINATOR;
+				Ret += TEXT("\t") + DstBaseName + TEXT(" Ret;") HLSL_LINE_TERMINATOR;
+				Ret += TEXT("\tRet.Value = ") + CoerceFloat(Translator, TEXT("Src.Value"), DstDerivType, SrcDerivType) + TEXT(";") HLSL_LINE_TERMINATOR;
+				Ret += TEXT("\tRet.Ddx = ") + CoerceFloat(Translator, TEXT("Src.Ddx"), DstTypeDDXY, SrcTypeDDXY) + TEXT(";") HLSL_LINE_TERMINATOR;
+				Ret += TEXT("\tRet.Ddy = ") + CoerceFloat(Translator, TEXT("Src.Ddy"), DstTypeDDXY, SrcTypeDDXY) + TEXT(";") HLSL_LINE_TERMINATOR;
+				Ret += TEXT("\treturn Ret;") HLSL_LINE_TERMINATOR;
+				Ret += TEXT("}") HLSL_LINE_TERMINATOR;
+				Ret += TEXT("") HLSL_LINE_TERMINATOR;
 			}
 		}
 	}
@@ -1482,22 +1482,22 @@ FString FMaterialDerivativeAutogen::GenerateUsedFunctions(FHLSLMaterialTranslato
 			{
 				FString ElemStr = FString::Printf(TEXT("%d"), ElemIndex + 1);
 
-				Ret += ScalarName + TEXT(" Extract") + BaseName + TEXT("_") + ElemStr + TEXT("(") + BaseName + TEXT(" InValue)") LINE_TERMINATOR;
-				Ret += TEXT("{") LINE_TERMINATOR;
-				Ret += TEXT("\t") + ScalarName + TEXT(" Ret;") LINE_TERMINATOR;
+				Ret += ScalarName + TEXT(" Extract") + BaseName + TEXT("_") + ElemStr + TEXT("(") + BaseName + TEXT(" InValue)") HLSL_LINE_TERMINATOR;
+				Ret += TEXT("{") HLSL_LINE_TERMINATOR;
+				Ret += TEXT("\t") + ScalarName + TEXT(" Ret;") HLSL_LINE_TERMINATOR;
 				if (IsLWCType(DerivType))
 				{
-					Ret += FString::Printf(TEXT("\tRet.Value = LWCGetComponent(InValue.Value, %d);") LINE_TERMINATOR, ElemIndex);
+					Ret += FString::Printf(TEXT("\tRet.Value = LWCGetComponent(InValue.Value, %d);") HLSL_LINE_TERMINATOR, ElemIndex);
 				}
 				else
 				{
-					Ret += TEXT("\tRet.Value = InValue.Value.") + FString(SwizzleList[ElemIndex]) + TEXT(";") LINE_TERMINATOR;
+					Ret += TEXT("\tRet.Value = InValue.Value.") + FString(SwizzleList[ElemIndex]) + TEXT(";") HLSL_LINE_TERMINATOR;
 				}
-				Ret += TEXT("\tRet.Ddx = InValue.Ddx.") + FString(SwizzleList[ElemIndex]) + TEXT(";") LINE_TERMINATOR;
-				Ret += TEXT("\tRet.Ddy = InValue.Ddy.") + FString(SwizzleList[ElemIndex]) + TEXT(";") LINE_TERMINATOR;
-				Ret += TEXT("\treturn Ret;") LINE_TERMINATOR;
-				Ret += TEXT("}") LINE_TERMINATOR;
-				Ret += TEXT("") LINE_TERMINATOR;
+				Ret += TEXT("\tRet.Ddx = InValue.Ddx.") + FString(SwizzleList[ElemIndex]) + TEXT(";") HLSL_LINE_TERMINATOR;
+				Ret += TEXT("\tRet.Ddy = InValue.Ddy.") + FString(SwizzleList[ElemIndex]) + TEXT(";") HLSL_LINE_TERMINATOR;
+				Ret += TEXT("\treturn Ret;") HLSL_LINE_TERMINATOR;
+				Ret += TEXT("}") HLSL_LINE_TERMINATOR;
+				Ret += TEXT("") HLSL_LINE_TERMINATOR;
 			}
 		}
 	}
@@ -1522,195 +1522,195 @@ FString FMaterialDerivativeAutogen::GenerateUsedFunctions(FHLSLMaterialTranslato
 				switch ((EFunc2)Op)
 				{
 				case EFunc2::Add:
-					Ret += BaseName + TEXT(" AddDeriv") + Suffix + TEXT("(") + BaseName + TEXT(" A, ") + BaseName + TEXT(" B)") LINE_TERMINATOR;
-					Ret += TEXT("{") LINE_TERMINATOR;
-					Ret += TEXT("\t") + BaseName + TEXT(" Ret;") LINE_TERMINATOR;
+					Ret += BaseName + TEXT(" AddDeriv") + Suffix + TEXT("(") + BaseName + TEXT(" A, ") + BaseName + TEXT(" B)") HLSL_LINE_TERMINATOR;
+					Ret += TEXT("{") HLSL_LINE_TERMINATOR;
+					Ret += TEXT("\t") + BaseName + TEXT(" Ret;") HLSL_LINE_TERMINATOR;
 					if (IsLWCType(DerivType))
 					{
 						Translator.AddLWCFuncUsage(ELWCFunctionKind::Add);
-						Ret += TEXT("\tRet.Value = LWCAdd(A.Value, B.Value);") LINE_TERMINATOR;
+						Ret += TEXT("\tRet.Value = LWCAdd(A.Value, B.Value);") HLSL_LINE_TERMINATOR;
 					}
 					else
 					{
-						Ret += TEXT("\tRet.Value = A.Value + B.Value;") LINE_TERMINATOR;
+						Ret += TEXT("\tRet.Value = A.Value + B.Value;") HLSL_LINE_TERMINATOR;
 					}
-					Ret += TEXT("\tRet.Ddx = A.Ddx + B.Ddx;") LINE_TERMINATOR;
-					Ret += TEXT("\tRet.Ddy = A.Ddy + B.Ddy;") LINE_TERMINATOR;
-					Ret += TEXT("\treturn Ret;") LINE_TERMINATOR;
-					Ret += TEXT("}") LINE_TERMINATOR;
-					Ret += TEXT("") LINE_TERMINATOR;
+					Ret += TEXT("\tRet.Ddx = A.Ddx + B.Ddx;") HLSL_LINE_TERMINATOR;
+					Ret += TEXT("\tRet.Ddy = A.Ddy + B.Ddy;") HLSL_LINE_TERMINATOR;
+					Ret += TEXT("\treturn Ret;") HLSL_LINE_TERMINATOR;
+					Ret += TEXT("}") HLSL_LINE_TERMINATOR;
+					Ret += TEXT("") HLSL_LINE_TERMINATOR;
 					break;
 				case EFunc2::Sub:
-					Ret += BaseName + TEXT(" SubDeriv") + Suffix + TEXT("(") + BaseName + TEXT(" A, ") + BaseName + TEXT(" B)") LINE_TERMINATOR;
-					Ret += TEXT("{") LINE_TERMINATOR;
-					Ret += TEXT("\t") + BaseName + TEXT(" Ret;") LINE_TERMINATOR;
+					Ret += BaseName + TEXT(" SubDeriv") + Suffix + TEXT("(") + BaseName + TEXT(" A, ") + BaseName + TEXT(" B)") HLSL_LINE_TERMINATOR;
+					Ret += TEXT("{") HLSL_LINE_TERMINATOR;
+					Ret += TEXT("\t") + BaseName + TEXT(" Ret;") HLSL_LINE_TERMINATOR;
 					if (IsLWCType(DerivType))
 					{
 						Translator.AddLWCFuncUsage(ELWCFunctionKind::Subtract);
-						Ret += TEXT("\tRet.Value = LWCSubtract(A.Value, B.Value);") LINE_TERMINATOR;
+						Ret += TEXT("\tRet.Value = LWCSubtract(A.Value, B.Value);") HLSL_LINE_TERMINATOR;
 					}
 					else
 					{
-						Ret += TEXT("\tRet.Value = A.Value - B.Value;") LINE_TERMINATOR;
+						Ret += TEXT("\tRet.Value = A.Value - B.Value;") HLSL_LINE_TERMINATOR;
 					}
-					Ret += TEXT("\tRet.Ddx = A.Ddx - B.Ddx;") LINE_TERMINATOR;
-					Ret += TEXT("\tRet.Ddy = A.Ddy - B.Ddy;") LINE_TERMINATOR;
-					Ret += TEXT("\treturn Ret;") LINE_TERMINATOR;
-					Ret += TEXT("}") LINE_TERMINATOR;
-					Ret += TEXT("") LINE_TERMINATOR;
+					Ret += TEXT("\tRet.Ddx = A.Ddx - B.Ddx;") HLSL_LINE_TERMINATOR;
+					Ret += TEXT("\tRet.Ddy = A.Ddy - B.Ddy;") HLSL_LINE_TERMINATOR;
+					Ret += TEXT("\treturn Ret;") HLSL_LINE_TERMINATOR;
+					Ret += TEXT("}") HLSL_LINE_TERMINATOR;
+					Ret += TEXT("") HLSL_LINE_TERMINATOR;
 					break;
 				case EFunc2::Mul:
-					Ret += BaseName + TEXT(" MulDeriv") + Suffix + TEXT("(") + BaseName + TEXT(" A, ") + BaseName + TEXT(" B)") LINE_TERMINATOR;
-					Ret += TEXT("{") LINE_TERMINATOR;
-					Ret += TEXT("\t") + BaseName + TEXT(" Ret;") LINE_TERMINATOR;
+					Ret += BaseName + TEXT(" MulDeriv") + Suffix + TEXT("(") + BaseName + TEXT(" A, ") + BaseName + TEXT(" B)") HLSL_LINE_TERMINATOR;
+					Ret += TEXT("{") HLSL_LINE_TERMINATOR;
+					Ret += TEXT("\t") + BaseName + TEXT(" Ret;") HLSL_LINE_TERMINATOR;
 					if (IsLWCType(DerivType))
 					{
 						Translator.AddLWCFuncUsage(ELWCFunctionKind::MultiplyVectorVector);
-						Ret += TEXT("\tRet.Value = LWCMultiply(A.Value, B.Value);") LINE_TERMINATOR;
-						Ret += TEXT("\tRet.Ddx = A.Ddx * LWCToFloat(B.Value) + LWCToFloat(A.Value) * B.Ddx;") LINE_TERMINATOR;
-						Ret += TEXT("\tRet.Ddy = A.Ddy * LWCToFloat(B.Value) + LWCToFloat(A.Value) * B.Ddy;") LINE_TERMINATOR;
+						Ret += TEXT("\tRet.Value = LWCMultiply(A.Value, B.Value);") HLSL_LINE_TERMINATOR;
+						Ret += TEXT("\tRet.Ddx = A.Ddx * LWCToFloat(B.Value) + LWCToFloat(A.Value) * B.Ddx;") HLSL_LINE_TERMINATOR;
+						Ret += TEXT("\tRet.Ddy = A.Ddy * LWCToFloat(B.Value) + LWCToFloat(A.Value) * B.Ddy;") HLSL_LINE_TERMINATOR;
 					}
 					else
 					{
-						Ret += TEXT("\tRet.Value = A.Value * B.Value;") LINE_TERMINATOR;
-						Ret += TEXT("\tRet.Ddx = A.Ddx * B.Value + A.Value * B.Ddx;") LINE_TERMINATOR;
-						Ret += TEXT("\tRet.Ddy = A.Ddy * B.Value + A.Value * B.Ddy;") LINE_TERMINATOR;
+						Ret += TEXT("\tRet.Value = A.Value * B.Value;") HLSL_LINE_TERMINATOR;
+						Ret += TEXT("\tRet.Ddx = A.Ddx * B.Value + A.Value * B.Ddx;") HLSL_LINE_TERMINATOR;
+						Ret += TEXT("\tRet.Ddy = A.Ddy * B.Value + A.Value * B.Ddy;") HLSL_LINE_TERMINATOR;
 					}
-					Ret += TEXT("\treturn Ret;") LINE_TERMINATOR;
-					Ret += TEXT("}") LINE_TERMINATOR;
-					Ret += TEXT("") LINE_TERMINATOR;
+					Ret += TEXT("\treturn Ret;") HLSL_LINE_TERMINATOR;
+					Ret += TEXT("}") HLSL_LINE_TERMINATOR;
+					Ret += TEXT("") HLSL_LINE_TERMINATOR;
 					break;
 				case EFunc2::Div:
-					Ret += BaseName + TEXT(" DivDeriv") + Suffix + TEXT("(") + BaseName + TEXT(" A, ") + BaseName + TEXT(" B)") LINE_TERMINATOR;
-					Ret += TEXT("{") LINE_TERMINATOR;
-					Ret += TEXT("\t") + BaseName + TEXT(" Ret;") LINE_TERMINATOR;
+					Ret += BaseName + TEXT(" DivDeriv") + Suffix + TEXT("(") + BaseName + TEXT(" A, ") + BaseName + TEXT(" B)") HLSL_LINE_TERMINATOR;
+					Ret += TEXT("{") HLSL_LINE_TERMINATOR;
+					Ret += TEXT("\t") + BaseName + TEXT(" Ret;") HLSL_LINE_TERMINATOR;
 					if (IsLWCType(DerivType))
 					{
 						Translator.AddLWCFuncUsage(ELWCFunctionKind::Divide);
-						Ret += TEXT("\tRet.Value = LWCDivide(A.Value, B.Value);") LINE_TERMINATOR;
+						Ret += TEXT("\tRet.Value = LWCDivide(A.Value, B.Value);") HLSL_LINE_TERMINATOR;
 						Translator.AddLWCFuncUsage(ELWCFunctionKind::Other);
 						Translator.AddLWCFuncUsage(ELWCFunctionKind::Demote, 2);
 						Translator.AddLWCFuncUsage(ELWCFunctionKind::MultiplyVectorVector, 3);
-						Ret += TEXT("\t") + FieldName + TEXT(" Denom = LWCRcp(LWCMultiply(B.Value, B.Value));") LINE_TERMINATOR;
-						Ret += TEXT("\t") + FieldName + TEXT(" dFdA =  LWCToFloat(LWCMultiply(B.Value, Denom));") LINE_TERMINATOR;
-						Ret += TEXT("\t") + FieldName + TEXT(" dFdB = -LWCToFloat(LWCMultiply(A.Value, Denom));") LINE_TERMINATOR;
-						Ret += TEXT("\tRet.Ddx = dFdA * A.Ddx + dFdB * B.Ddx;") LINE_TERMINATOR;
-						Ret += TEXT("\tRet.Ddy = dFdA * A.Ddy + dFdB * B.Ddy;") LINE_TERMINATOR;
+						Ret += TEXT("\t") + FieldName + TEXT(" Denom = LWCRcp(LWCMultiply(B.Value, B.Value));") HLSL_LINE_TERMINATOR;
+						Ret += TEXT("\t") + FieldName + TEXT(" dFdA =  LWCToFloat(LWCMultiply(B.Value, Denom));") HLSL_LINE_TERMINATOR;
+						Ret += TEXT("\t") + FieldName + TEXT(" dFdB = -LWCToFloat(LWCMultiply(A.Value, Denom));") HLSL_LINE_TERMINATOR;
+						Ret += TEXT("\tRet.Ddx = dFdA * A.Ddx + dFdB * B.Ddx;") HLSL_LINE_TERMINATOR;
+						Ret += TEXT("\tRet.Ddy = dFdA * A.Ddy + dFdB * B.Ddy;") HLSL_LINE_TERMINATOR;
 					}
 					else
 					{
-						Ret += TEXT("\tRet.Value = A.Value / B.Value;") LINE_TERMINATOR;
-						Ret += TEXT("\t") + FieldName + TEXT(" Denom = rcp(B.Value * B.Value);") LINE_TERMINATOR;
-						Ret += TEXT("\t") + FieldName + TEXT(" dFdA =  B.Value * Denom;") LINE_TERMINATOR;
-						Ret += TEXT("\t") + FieldName + TEXT(" dFdB = -A.Value * Denom;") LINE_TERMINATOR;
-						Ret += TEXT("\tRet.Ddx = dFdA * A.Ddx + dFdB * B.Ddx;") LINE_TERMINATOR;
-						Ret += TEXT("\tRet.Ddy = dFdA * A.Ddy + dFdB * B.Ddy;") LINE_TERMINATOR;
+						Ret += TEXT("\tRet.Value = A.Value / B.Value;") HLSL_LINE_TERMINATOR;
+						Ret += TEXT("\t") + FieldName + TEXT(" Denom = rcp(B.Value * B.Value);") HLSL_LINE_TERMINATOR;
+						Ret += TEXT("\t") + FieldName + TEXT(" dFdA =  B.Value * Denom;") HLSL_LINE_TERMINATOR;
+						Ret += TEXT("\t") + FieldName + TEXT(" dFdB = -A.Value * Denom;") HLSL_LINE_TERMINATOR;
+						Ret += TEXT("\tRet.Ddx = dFdA * A.Ddx + dFdB * B.Ddx;") HLSL_LINE_TERMINATOR;
+						Ret += TEXT("\tRet.Ddy = dFdA * A.Ddy + dFdB * B.Ddy;") HLSL_LINE_TERMINATOR;
 					}
-					Ret += TEXT("\treturn Ret;") LINE_TERMINATOR;
-					Ret += TEXT("}") LINE_TERMINATOR;
-					Ret += TEXT("") LINE_TERMINATOR;
+					Ret += TEXT("\treturn Ret;") HLSL_LINE_TERMINATOR;
+					Ret += TEXT("}") HLSL_LINE_TERMINATOR;
+					Ret += TEXT("") HLSL_LINE_TERMINATOR;
 
 					if (IsLWCType(DerivType))
 					{
 						// Add an overload to divide LWC by non-LWC value
-						Ret += BaseName + TEXT(" DivDeriv") + Suffix + TEXT("(") + BaseName + TEXT(" A, ") + NonLWCBaseName + TEXT(" B)") LINE_TERMINATOR;
-						Ret += TEXT("{") LINE_TERMINATOR;
-						Ret += TEXT("\t") + BaseName + TEXT(" Ret;") LINE_TERMINATOR;
+						Ret += BaseName + TEXT(" DivDeriv") + Suffix + TEXT("(") + BaseName + TEXT(" A, ") + NonLWCBaseName + TEXT(" B)") HLSL_LINE_TERMINATOR;
+						Ret += TEXT("{") HLSL_LINE_TERMINATOR;
+						Ret += TEXT("\t") + BaseName + TEXT(" Ret;") HLSL_LINE_TERMINATOR;
 						Translator.AddLWCFuncUsage(ELWCFunctionKind::Divide);
-						Ret += TEXT("\tRet.Value = LWCDivide(A.Value, B.Value);") LINE_TERMINATOR;
-						Ret += TEXT("\t") + FieldName + TEXT(" Denom = rcp(B.Value * B.Value);") LINE_TERMINATOR;
-						Ret += TEXT("\t") + FieldName + TEXT(" dFdA =  B.Value * Denom;") LINE_TERMINATOR;
+						Ret += TEXT("\tRet.Value = LWCDivide(A.Value, B.Value);") HLSL_LINE_TERMINATOR;
+						Ret += TEXT("\t") + FieldName + TEXT(" Denom = rcp(B.Value * B.Value);") HLSL_LINE_TERMINATOR;
+						Ret += TEXT("\t") + FieldName + TEXT(" dFdA =  B.Value * Denom;") HLSL_LINE_TERMINATOR;
 						Translator.AddLWCFuncUsage(ELWCFunctionKind::Demote);
 						Translator.AddLWCFuncUsage(ELWCFunctionKind::MultiplyVectorVector);
-						Ret += TEXT("\t") + FieldName + TEXT(" dFdB = -LWCToFloat(LWCMultiply(A.Value, Denom));") LINE_TERMINATOR;
-						Ret += TEXT("\tRet.Ddx = dFdA * A.Ddx + dFdB * B.Ddx;") LINE_TERMINATOR;
-						Ret += TEXT("\tRet.Ddy = dFdA * A.Ddy + dFdB * B.Ddy;") LINE_TERMINATOR;
-						Ret += TEXT("\treturn Ret;") LINE_TERMINATOR;
-						Ret += TEXT("}") LINE_TERMINATOR;
-						Ret += TEXT("") LINE_TERMINATOR;
+						Ret += TEXT("\t") + FieldName + TEXT(" dFdB = -LWCToFloat(LWCMultiply(A.Value, Denom));") HLSL_LINE_TERMINATOR;
+						Ret += TEXT("\tRet.Ddx = dFdA * A.Ddx + dFdB * B.Ddx;") HLSL_LINE_TERMINATOR;
+						Ret += TEXT("\tRet.Ddy = dFdA * A.Ddy + dFdB * B.Ddy;") HLSL_LINE_TERMINATOR;
+						Ret += TEXT("\treturn Ret;") HLSL_LINE_TERMINATOR;
+						Ret += TEXT("}") HLSL_LINE_TERMINATOR;
+						Ret += TEXT("") HLSL_LINE_TERMINATOR;
 					}
 					break;
 				case EFunc2::Fmod:
 					// Only valid when B derivatives are zero.
 					// We can't really do anything meaningful in the non-zero case.
-					Ret += NonLWCBaseName + TEXT(" FmodDeriv") + Suffix + TEXT("(") + BaseName + TEXT(" A, ") + NonLWCBaseName + TEXT(" B)") LINE_TERMINATOR;
-					Ret += TEXT("{") LINE_TERMINATOR;
-					Ret += TEXT("\t") + NonLWCBaseName + TEXT(" Ret;") LINE_TERMINATOR;
+					Ret += NonLWCBaseName + TEXT(" FmodDeriv") + Suffix + TEXT("(") + BaseName + TEXT(" A, ") + NonLWCBaseName + TEXT(" B)") HLSL_LINE_TERMINATOR;
+					Ret += TEXT("{") HLSL_LINE_TERMINATOR;
+					Ret += TEXT("\t") + NonLWCBaseName + TEXT(" Ret;") HLSL_LINE_TERMINATOR;
 					if (IsLWCType(DerivType))
 					{
 						Translator.AddLWCFuncUsage(ELWCFunctionKind::Other);
-						Ret += TEXT("\tRet.Value = LWCFmod(A.Value, B.Value);") LINE_TERMINATOR;
+						Ret += TEXT("\tRet.Value = LWCFmod(A.Value, B.Value);") HLSL_LINE_TERMINATOR;
 					}
 					else
 					{
-						Ret += TEXT("\tRet.Value = fmod(A.Value, B.Value);") LINE_TERMINATOR;
+						Ret += TEXT("\tRet.Value = fmod(A.Value, B.Value);") HLSL_LINE_TERMINATOR;
 					}
-					Ret += TEXT("\tRet.Ddx = A.Ddx;") LINE_TERMINATOR;
-					Ret += TEXT("\tRet.Ddy = A.Ddy;") LINE_TERMINATOR;
-					Ret += TEXT("\treturn Ret;") LINE_TERMINATOR;
-					Ret += TEXT("}") LINE_TERMINATOR;
-					Ret += TEXT("") LINE_TERMINATOR;
+					Ret += TEXT("\tRet.Ddx = A.Ddx;") HLSL_LINE_TERMINATOR;
+					Ret += TEXT("\tRet.Ddy = A.Ddy;") HLSL_LINE_TERMINATOR;
+					Ret += TEXT("\treturn Ret;") HLSL_LINE_TERMINATOR;
+					Ret += TEXT("}") HLSL_LINE_TERMINATOR;
+					Ret += TEXT("") HLSL_LINE_TERMINATOR;
 					break;
 				case EFunc2::Min:
-					Ret += BaseName + TEXT(" MinDeriv") + Suffix + TEXT("(") + BaseName + TEXT(" A, ") + BaseName + TEXT(" B)") LINE_TERMINATOR;
-					Ret += TEXT("{") LINE_TERMINATOR;
-					Ret += TEXT("\t") + BaseName + TEXT(" Ret;") LINE_TERMINATOR;
+					Ret += BaseName + TEXT(" MinDeriv") + Suffix + TEXT("(") + BaseName + TEXT(" A, ") + BaseName + TEXT(" B)") HLSL_LINE_TERMINATOR;
+					Ret += TEXT("{") HLSL_LINE_TERMINATOR;
+					Ret += TEXT("\t") + BaseName + TEXT(" Ret;") HLSL_LINE_TERMINATOR;
 					if (IsLWCType(DerivType))
 					{
 						Translator.AddLWCFuncUsage(ELWCFunctionKind::Other);
-						Ret += TEXT("\t") + BoolName + TEXT(" Cmp = LWCLess(A.Value, B.Value);") LINE_TERMINATOR;
-						Ret += TEXT("\tRet.Value = LWCSelect(Cmp, A.Value, B.Value);") LINE_TERMINATOR;
+						Ret += TEXT("\t") + BoolName + TEXT(" Cmp = LWCLess(A.Value, B.Value);") HLSL_LINE_TERMINATOR;
+						Ret += TEXT("\tRet.Value = LWCSelect(Cmp, A.Value, B.Value);") HLSL_LINE_TERMINATOR;
 					}
 					else
 					{
-						Ret += TEXT("\t") + BoolName + TEXT(" Cmp = A.Value < B.Value;") LINE_TERMINATOR;
-						Ret += TEXT("\tRet.Value = select(Cmp, A.Value, B.Value);") LINE_TERMINATOR;
+						Ret += TEXT("\t") + BoolName + TEXT(" Cmp = A.Value < B.Value;") HLSL_LINE_TERMINATOR;
+						Ret += TEXT("\tRet.Value = select(Cmp, A.Value, B.Value);") HLSL_LINE_TERMINATOR;
 					}
-					Ret += TEXT("\tRet.Ddx = select(Cmp, A.Ddx, B.Ddx);") LINE_TERMINATOR;
-					Ret += TEXT("\tRet.Ddy = select(Cmp, A.Ddy, B.Ddy);") LINE_TERMINATOR;
-					Ret += TEXT("\treturn Ret;") LINE_TERMINATOR;
-					Ret += TEXT("}") LINE_TERMINATOR;
-					Ret += TEXT("") LINE_TERMINATOR;
+					Ret += TEXT("\tRet.Ddx = select(Cmp, A.Ddx, B.Ddx);") HLSL_LINE_TERMINATOR;
+					Ret += TEXT("\tRet.Ddy = select(Cmp, A.Ddy, B.Ddy);") HLSL_LINE_TERMINATOR;
+					Ret += TEXT("\treturn Ret;") HLSL_LINE_TERMINATOR;
+					Ret += TEXT("}") HLSL_LINE_TERMINATOR;
+					Ret += TEXT("") HLSL_LINE_TERMINATOR;
 					break;
 				case EFunc2::Max:
-					Ret += BaseName + TEXT(" MaxDeriv") + Suffix + TEXT("(") + BaseName + TEXT(" A, ") + BaseName + TEXT(" B)") LINE_TERMINATOR;
-					Ret += TEXT("{") LINE_TERMINATOR;
-					Ret += TEXT("\t") + BaseName + TEXT(" Ret;") LINE_TERMINATOR;
+					Ret += BaseName + TEXT(" MaxDeriv") + Suffix + TEXT("(") + BaseName + TEXT(" A, ") + BaseName + TEXT(" B)") HLSL_LINE_TERMINATOR;
+					Ret += TEXT("{") HLSL_LINE_TERMINATOR;
+					Ret += TEXT("\t") + BaseName + TEXT(" Ret;") HLSL_LINE_TERMINATOR;
 					if (IsLWCType(DerivType))
 					{
 						Translator.AddLWCFuncUsage(ELWCFunctionKind::Other);
-						Ret += TEXT("\t") + BoolName + TEXT(" Cmp = LWCGreater(A.Value, B.Value);") LINE_TERMINATOR;
-						Ret += TEXT("\tRet.Value = LWCSelect(Cmp, A.Value, B.Value);") LINE_TERMINATOR;
+						Ret += TEXT("\t") + BoolName + TEXT(" Cmp = LWCGreater(A.Value, B.Value);") HLSL_LINE_TERMINATOR;
+						Ret += TEXT("\tRet.Value = LWCSelect(Cmp, A.Value, B.Value);") HLSL_LINE_TERMINATOR;
 					}
 					else
 					{
-						Ret += TEXT("\t") + BoolName + TEXT(" Cmp = A.Value > B.Value;") LINE_TERMINATOR;
-						Ret += TEXT("\tRet.Value = select(Cmp, A.Value, B.Value);") LINE_TERMINATOR;
+						Ret += TEXT("\t") + BoolName + TEXT(" Cmp = A.Value > B.Value;") HLSL_LINE_TERMINATOR;
+						Ret += TEXT("\tRet.Value = select(Cmp, A.Value, B.Value);") HLSL_LINE_TERMINATOR;
 					}
-					Ret += TEXT("\tRet.Ddx = select(Cmp, A.Ddx, B.Ddx);") LINE_TERMINATOR;
-					Ret += TEXT("\tRet.Ddy = select(Cmp, A.Ddy, B.Ddy);") LINE_TERMINATOR;
-					Ret += TEXT("\treturn Ret;") LINE_TERMINATOR;
-					Ret += TEXT("}") LINE_TERMINATOR;
-					Ret += TEXT("") LINE_TERMINATOR;
+					Ret += TEXT("\tRet.Ddx = select(Cmp, A.Ddx, B.Ddx);") HLSL_LINE_TERMINATOR;
+					Ret += TEXT("\tRet.Ddy = select(Cmp, A.Ddy, B.Ddy);") HLSL_LINE_TERMINATOR;
+					Ret += TEXT("\treturn Ret;") HLSL_LINE_TERMINATOR;
+					Ret += TEXT("}") HLSL_LINE_TERMINATOR;
+					Ret += TEXT("") HLSL_LINE_TERMINATOR;
 					break;
 				case EFunc2::Dot:
-					Ret += ScalarName + TEXT(" DotDeriv") + Suffix + TEXT("(") + BaseName + TEXT(" A, ") + BaseName + TEXT(" B)") LINE_TERMINATOR;
-					Ret += TEXT("{") LINE_TERMINATOR;
+					Ret += ScalarName + TEXT(" DotDeriv") + Suffix + TEXT("(") + BaseName + TEXT(" A, ") + BaseName + TEXT(" B)") HLSL_LINE_TERMINATOR;
+					Ret += TEXT("{") HLSL_LINE_TERMINATOR;
 					if (IsLWCType(DerivType))
 					{
-						Ret += TEXT("\t") + ScalarName + TEXT(" Ret = ConstructConstant") + ScalarName + TEXT("(LWCPromote(0.0f));") LINE_TERMINATOR;
+						Ret += TEXT("\t") + ScalarName + TEXT(" Ret = ConstructConstant") + ScalarName + TEXT("(LWCPromote(0.0f));") HLSL_LINE_TERMINATOR;
 					}
 					else
 					{
-						Ret += TEXT("\t") + ScalarName + TEXT(" Ret = ConstructConstant") + ScalarName + TEXT("(0);") LINE_TERMINATOR;
+						Ret += TEXT("\t") + ScalarName + TEXT(" Ret = ConstructConstant") + ScalarName + TEXT("(0);") HLSL_LINE_TERMINATOR;
 					}
 					for (uint32 Component = 0; Component < NumComponents; Component++)
 					{
-						Ret += FString::Printf(TEXT("\tRet = AddDeriv%s(Ret,MulDeriv%s(Extract%s_%d(A),Extract%s_%d(B)));"), *Suffix, *Suffix, *BaseName, Component + 1, *BaseName, Component + 1) + LINE_TERMINATOR;
+						Ret += FString::Printf(TEXT("\tRet = AddDeriv%s(Ret,MulDeriv%s(Extract%s_%d(A),Extract%s_%d(B)));"), *Suffix, *Suffix, *BaseName, Component + 1, *BaseName, Component + 1) + HLSL_LINE_TERMINATOR;
 					}
-					Ret += TEXT("\treturn Ret;") LINE_TERMINATOR;
-					Ret += TEXT("}") LINE_TERMINATOR;
-					Ret += TEXT("") LINE_TERMINATOR;
+					Ret += TEXT("\treturn Ret;") HLSL_LINE_TERMINATOR;
+					Ret += TEXT("}") HLSL_LINE_TERMINATOR;
+					Ret += TEXT("") HLSL_LINE_TERMINATOR;
 					break;
 				case EFunc2::Pow:
 					// pow(A,B) = exp(B*log(A))
@@ -1719,58 +1719,58 @@ FString FMaterialDerivativeAutogen::GenerateUsedFunctions(FHLSLMaterialTranslato
 					// sanity check when B is constant and A is a linear function (B'=0,A'=1)
 					//     pow'(A,B) = pow(A,B) * (0*log(A) + (B/A)*1)
 					//     pow'(A,B) = B * pow(A,B-1)
-					Ret += BaseName + TEXT(" PowDeriv") + Suffix + TEXT("(") + BaseName + TEXT(" A, ") + BaseName + TEXT(" B)") LINE_TERMINATOR;
-					Ret += TEXT("{") LINE_TERMINATOR;
-					Ret += TEXT("\t") + BaseName + TEXT(" Ret;") LINE_TERMINATOR;
-					Ret += TEXT("\tRet.Value = pow(A.Value, B.Value);") LINE_TERMINATOR;
-					Ret += TEXT("\tRet.Ddx = Ret.Value * (B.Ddx * log(A.Value) + (B.Value/A.Value)*A.Ddx);") LINE_TERMINATOR;
-					Ret += TEXT("\tRet.Ddy = Ret.Value * (B.Ddy * log(A.Value) + (B.Value/A.Value)*A.Ddy);") LINE_TERMINATOR;
-					Ret += TEXT("\treturn Ret;") LINE_TERMINATOR;
-					Ret += TEXT("}") LINE_TERMINATOR;
-					Ret += TEXT("") LINE_TERMINATOR;
+					Ret += BaseName + TEXT(" PowDeriv") + Suffix + TEXT("(") + BaseName + TEXT(" A, ") + BaseName + TEXT(" B)") HLSL_LINE_TERMINATOR;
+					Ret += TEXT("{") HLSL_LINE_TERMINATOR;
+					Ret += TEXT("\t") + BaseName + TEXT(" Ret;") HLSL_LINE_TERMINATOR;
+					Ret += TEXT("\tRet.Value = pow(A.Value, B.Value);") HLSL_LINE_TERMINATOR;
+					Ret += TEXT("\tRet.Ddx = Ret.Value * (B.Ddx * log(A.Value) + (B.Value/A.Value)*A.Ddx);") HLSL_LINE_TERMINATOR;
+					Ret += TEXT("\tRet.Ddy = Ret.Value * (B.Ddy * log(A.Value) + (B.Value/A.Value)*A.Ddy);") HLSL_LINE_TERMINATOR;
+					Ret += TEXT("\treturn Ret;") HLSL_LINE_TERMINATOR;
+					Ret += TEXT("}") HLSL_LINE_TERMINATOR;
+					Ret += TEXT("") HLSL_LINE_TERMINATOR;
 					break;
 				case EFunc2::PowPositiveClamped:
-					Ret += BaseName + TEXT(" PowPositiveClampedDeriv") + Suffix + TEXT("(") + BaseName + TEXT(" A, ") + BaseName + TEXT(" B)") LINE_TERMINATOR;
-					Ret += TEXT("{") LINE_TERMINATOR;
-					Ret += TEXT("\t") + BaseName + TEXT(" Ret;") LINE_TERMINATOR;
-					Ret += TEXT("\t") + BoolName + TEXT(" InRange = (0.0 < B.Value);") LINE_TERMINATOR; // should we check for A as well?
-					Ret += TEXT("\t") + FieldName + TEXT(" Zero = 0.0;") LINE_TERMINATOR;
-					Ret += TEXT("\tRet.Value = PositiveClampedPow(A.Value, B.Value);") LINE_TERMINATOR;
-					Ret += TEXT("\tRet.Ddx = Ret.Value * (B.Ddx * log(A.Value) + (B.Value/A.Value)*A.Ddx);") LINE_TERMINATOR;
-					Ret += TEXT("\tRet.Ddy = Ret.Value * (B.Ddy * log(A.Value) + (B.Value/A.Value)*A.Ddy);") LINE_TERMINATOR;
-					Ret += TEXT("\tRet.Ddx = select(InRange, Ret.Ddx, Zero);") LINE_TERMINATOR;
-					Ret += TEXT("\tRet.Ddy = select(InRange, Ret.Ddy, Zero);") LINE_TERMINATOR;
-					Ret += TEXT("\treturn Ret;") LINE_TERMINATOR;
-					Ret += TEXT("}") LINE_TERMINATOR;
-					Ret += TEXT("") LINE_TERMINATOR;
+					Ret += BaseName + TEXT(" PowPositiveClampedDeriv") + Suffix + TEXT("(") + BaseName + TEXT(" A, ") + BaseName + TEXT(" B)") HLSL_LINE_TERMINATOR;
+					Ret += TEXT("{") HLSL_LINE_TERMINATOR;
+					Ret += TEXT("\t") + BaseName + TEXT(" Ret;") HLSL_LINE_TERMINATOR;
+					Ret += TEXT("\t") + BoolName + TEXT(" InRange = (0.0 < B.Value);") HLSL_LINE_TERMINATOR; // should we check for A as well?
+					Ret += TEXT("\t") + FieldName + TEXT(" Zero = 0.0;") HLSL_LINE_TERMINATOR;
+					Ret += TEXT("\tRet.Value = PositiveClampedPow(A.Value, B.Value);") HLSL_LINE_TERMINATOR;
+					Ret += TEXT("\tRet.Ddx = Ret.Value * (B.Ddx * log(A.Value) + (B.Value/A.Value)*A.Ddx);") HLSL_LINE_TERMINATOR;
+					Ret += TEXT("\tRet.Ddy = Ret.Value * (B.Ddy * log(A.Value) + (B.Value/A.Value)*A.Ddy);") HLSL_LINE_TERMINATOR;
+					Ret += TEXT("\tRet.Ddx = select(InRange, Ret.Ddx, Zero);") HLSL_LINE_TERMINATOR;
+					Ret += TEXT("\tRet.Ddy = select(InRange, Ret.Ddy, Zero);") HLSL_LINE_TERMINATOR;
+					Ret += TEXT("\treturn Ret;") HLSL_LINE_TERMINATOR;
+					Ret += TEXT("}") HLSL_LINE_TERMINATOR;
+					Ret += TEXT("") HLSL_LINE_TERMINATOR;
 					break;
 				case EFunc2::Atan2:
-					Ret += BaseName + TEXT(" Atan2Deriv") + Suffix + TEXT("(") + BaseName + TEXT(" A, ") + BaseName + TEXT(" B)") LINE_TERMINATOR;
-					Ret += TEXT("{") LINE_TERMINATOR;
-					Ret += TEXT("\t") + BaseName + TEXT(" Ret;") LINE_TERMINATOR;
-					Ret += TEXT("\tRet.Value = atan2(A.Value, B.Value);") LINE_TERMINATOR;
-					Ret += TEXT("\t") + FieldName + TEXT(" Denom = rcp(A.Value * A.Value + B.Value * B.Value);") LINE_TERMINATOR;
-					Ret += TEXT("\t") + FieldName + TEXT(" dFdA =  B.Value * Denom;") LINE_TERMINATOR;
-					Ret += TEXT("\t") + FieldName + TEXT(" dFdB = -A.Value * Denom;") LINE_TERMINATOR;
-					Ret += TEXT("\tRet.Ddx = dFdA * A.Ddx + dFdB * B.Ddx;") LINE_TERMINATOR;
-					Ret += TEXT("\tRet.Ddy = dFdA * A.Ddy + dFdB * B.Ddy;") LINE_TERMINATOR;
-					Ret += TEXT("\treturn Ret;") LINE_TERMINATOR;
-					Ret += TEXT("}") LINE_TERMINATOR;
-					Ret += TEXT("") LINE_TERMINATOR;
+					Ret += BaseName + TEXT(" Atan2Deriv") + Suffix + TEXT("(") + BaseName + TEXT(" A, ") + BaseName + TEXT(" B)") HLSL_LINE_TERMINATOR;
+					Ret += TEXT("{") HLSL_LINE_TERMINATOR;
+					Ret += TEXT("\t") + BaseName + TEXT(" Ret;") HLSL_LINE_TERMINATOR;
+					Ret += TEXT("\tRet.Value = atan2(A.Value, B.Value);") HLSL_LINE_TERMINATOR;
+					Ret += TEXT("\t") + FieldName + TEXT(" Denom = rcp(A.Value * A.Value + B.Value * B.Value);") HLSL_LINE_TERMINATOR;
+					Ret += TEXT("\t") + FieldName + TEXT(" dFdA =  B.Value * Denom;") HLSL_LINE_TERMINATOR;
+					Ret += TEXT("\t") + FieldName + TEXT(" dFdB = -A.Value * Denom;") HLSL_LINE_TERMINATOR;
+					Ret += TEXT("\tRet.Ddx = dFdA * A.Ddx + dFdB * B.Ddx;") HLSL_LINE_TERMINATOR;
+					Ret += TEXT("\tRet.Ddy = dFdA * A.Ddy + dFdB * B.Ddy;") HLSL_LINE_TERMINATOR;
+					Ret += TEXT("\treturn Ret;") HLSL_LINE_TERMINATOR;
+					Ret += TEXT("}") HLSL_LINE_TERMINATOR;
+					Ret += TEXT("") HLSL_LINE_TERMINATOR;
 					break;
 				case EFunc2::Atan2Fast:
-					Ret += BaseName + TEXT(" Atan2FastDeriv") + Suffix + TEXT("(") + BaseName + TEXT(" A, ") + BaseName + TEXT(" B)") LINE_TERMINATOR;
-					Ret += TEXT("{") LINE_TERMINATOR;
-					Ret += TEXT("\t") + BaseName + TEXT(" Ret;") LINE_TERMINATOR;
-					Ret += TEXT("\tRet.Value = atan2Fast(A.Value, B.Value);") LINE_TERMINATOR;
-					Ret += TEXT("\t") + FieldName + TEXT(" Denom = rcp(A.Value * A.Value + B.Value * B.Value);") LINE_TERMINATOR;
-					Ret += TEXT("\t") + FieldName + TEXT(" dFdA =  B.Value * Denom;") LINE_TERMINATOR;
-					Ret += TEXT("\t") + FieldName + TEXT(" dFdB = -A.Value * Denom;") LINE_TERMINATOR;
-					Ret += TEXT("\tRet.Ddx = dFdA * A.Ddx + dFdB * B.Ddx;") LINE_TERMINATOR;
-					Ret += TEXT("\tRet.Ddy = dFdA * A.Ddy + dFdB * B.Ddy;") LINE_TERMINATOR;
-					Ret += TEXT("\treturn Ret;") LINE_TERMINATOR;
-					Ret += TEXT("}") LINE_TERMINATOR;
-					Ret += TEXT("") LINE_TERMINATOR;
+					Ret += BaseName + TEXT(" Atan2FastDeriv") + Suffix + TEXT("(") + BaseName + TEXT(" A, ") + BaseName + TEXT(" B)") HLSL_LINE_TERMINATOR;
+					Ret += TEXT("{") HLSL_LINE_TERMINATOR;
+					Ret += TEXT("\t") + BaseName + TEXT(" Ret;") HLSL_LINE_TERMINATOR;
+					Ret += TEXT("\tRet.Value = atan2Fast(A.Value, B.Value);") HLSL_LINE_TERMINATOR;
+					Ret += TEXT("\t") + FieldName + TEXT(" Denom = rcp(A.Value * A.Value + B.Value * B.Value);") HLSL_LINE_TERMINATOR;
+					Ret += TEXT("\t") + FieldName + TEXT(" dFdA =  B.Value * Denom;") HLSL_LINE_TERMINATOR;
+					Ret += TEXT("\t") + FieldName + TEXT(" dFdB = -A.Value * Denom;") HLSL_LINE_TERMINATOR;
+					Ret += TEXT("\tRet.Ddx = dFdA * A.Ddx + dFdB * B.Ddx;") HLSL_LINE_TERMINATOR;
+					Ret += TEXT("\tRet.Ddy = dFdA * A.Ddy + dFdB * B.Ddy;") HLSL_LINE_TERMINATOR;
+					Ret += TEXT("\treturn Ret;") HLSL_LINE_TERMINATOR;
+					Ret += TEXT("}") HLSL_LINE_TERMINATOR;
+					Ret += TEXT("") HLSL_LINE_TERMINATOR;
 					break;
 				case EFunc2::Cross:
 					if (DerivType == EDerivativeType::Float3)
@@ -1778,15 +1778,15 @@ FString FMaterialDerivativeAutogen::GenerateUsedFunctions(FHLSLMaterialTranslato
 						// (A*B)' = A' * B + A * B'
 						// Cross(A, B) = A.yzx * B.zxy - A.zxy * B.yzx;
 						// Cross(A, B)' = A.yzx' * B.zxy + A.yzx * B.zxy' - A.zxy' * B.yzx - A.zxy * B.yzx';
-						Ret += BaseName + TEXT(" CrossDeriv") + Suffix + TEXT("(") + BaseName + TEXT(" A, ") + BaseName + TEXT(" B)") LINE_TERMINATOR;
-						Ret += TEXT("{") LINE_TERMINATOR;
-						Ret += TEXT("\t") + BaseName + TEXT(" Ret;") LINE_TERMINATOR;
-						Ret += TEXT("\tRet.Value = cross(A.Value, B.Value);") LINE_TERMINATOR;
-						Ret += TEXT("\tRet.Ddx = A.Ddx.yzx * B.Value.zxy + A.Value.yzx * B.Ddx.zxy - A.Ddx.zxy * B.Value.yzx - A.Value.zxy * B.Ddx.yzx;") LINE_TERMINATOR;
-						Ret += TEXT("\tRet.Ddy = A.Ddy.yzx * B.Value.zxy + A.Value.yzx * B.Ddy.zxy - A.Ddy.zxy * B.Value.yzx - A.Value.zxy * B.Ddy.yzx;") LINE_TERMINATOR;
-						Ret += TEXT("\treturn Ret;") LINE_TERMINATOR;
-						Ret += TEXT("}") LINE_TERMINATOR;
-						Ret += TEXT("") LINE_TERMINATOR;
+						Ret += BaseName + TEXT(" CrossDeriv") + Suffix + TEXT("(") + BaseName + TEXT(" A, ") + BaseName + TEXT(" B)") HLSL_LINE_TERMINATOR;
+						Ret += TEXT("{") HLSL_LINE_TERMINATOR;
+						Ret += TEXT("\t") + BaseName + TEXT(" Ret;") HLSL_LINE_TERMINATOR;
+						Ret += TEXT("\tRet.Value = cross(A.Value, B.Value);") HLSL_LINE_TERMINATOR;
+						Ret += TEXT("\tRet.Ddx = A.Ddx.yzx * B.Value.zxy + A.Value.yzx * B.Ddx.zxy - A.Ddx.zxy * B.Value.yzx - A.Value.zxy * B.Ddx.yzx;") HLSL_LINE_TERMINATOR;
+						Ret += TEXT("\tRet.Ddy = A.Ddy.yzx * B.Value.zxy + A.Value.yzx * B.Ddy.zxy - A.Ddy.zxy * B.Value.yzx - A.Value.zxy * B.Ddy.yzx;") HLSL_LINE_TERMINATOR;
+						Ret += TEXT("\treturn Ret;") HLSL_LINE_TERMINATOR;
+						Ret += TEXT("}") HLSL_LINE_TERMINATOR;
+						Ret += TEXT("") HLSL_LINE_TERMINATOR;
 					}
 					break;
 				default:
@@ -1816,383 +1816,383 @@ FString FMaterialDerivativeAutogen::GenerateUsedFunctions(FHLSLMaterialTranslato
 				switch((EFunc1)Op)
 				{
 				case EFunc1::Abs:
-					Ret += BaseName + TEXT(" AbsDeriv") + Suffix + TEXT("(") + BaseName + TEXT(" A)") LINE_TERMINATOR;
-					Ret += TEXT("{") LINE_TERMINATOR;
-					Ret += TEXT("\t") + NonLWCFieldName + TEXT(" One = 1.0f;") LINE_TERMINATOR;
-					Ret += TEXT("\t") + BaseName + TEXT(" Ret;") LINE_TERMINATOR;
+					Ret += BaseName + TEXT(" AbsDeriv") + Suffix + TEXT("(") + BaseName + TEXT(" A)") HLSL_LINE_TERMINATOR;
+					Ret += TEXT("{") HLSL_LINE_TERMINATOR;
+					Ret += TEXT("\t") + NonLWCFieldName + TEXT(" One = 1.0f;") HLSL_LINE_TERMINATOR;
+					Ret += TEXT("\t") + BaseName + TEXT(" Ret;") HLSL_LINE_TERMINATOR;
 					if (IsLWCType(DerivType))
 					{
 						Translator.AddLWCFuncUsage(ELWCFunctionKind::Other);
-						Ret += TEXT("\tRet.Value = LWCAbs(A.Value);") LINE_TERMINATOR;
-						Ret += TEXT("\t") + NonLWCFieldName + TEXT(" dFdA = select(LWCGreaterEqual(A.Value, 0.0f), One, -One);") LINE_TERMINATOR;
+						Ret += TEXT("\tRet.Value = LWCAbs(A.Value);") HLSL_LINE_TERMINATOR;
+						Ret += TEXT("\t") + NonLWCFieldName + TEXT(" dFdA = select(LWCGreaterEqual(A.Value, 0.0f), One, -One);") HLSL_LINE_TERMINATOR;
 					}
 					else
 					{
-						Ret += TEXT("\tRet.Value = abs(A.Value);") LINE_TERMINATOR;
-						Ret += TEXT("\t") + NonLWCFieldName + TEXT(" dFdA = select(A.Value >= 0.0f, One, -One);") LINE_TERMINATOR;
+						Ret += TEXT("\tRet.Value = abs(A.Value);") HLSL_LINE_TERMINATOR;
+						Ret += TEXT("\t") + NonLWCFieldName + TEXT(" dFdA = select(A.Value >= 0.0f, One, -One);") HLSL_LINE_TERMINATOR;
 					}
 
-					Ret += TEXT("\tRet.Ddx = dFdA * A.Ddx;") LINE_TERMINATOR;
-					Ret += TEXT("\tRet.Ddy = dFdA * A.Ddy;") LINE_TERMINATOR;
-					Ret += TEXT("\treturn Ret;") LINE_TERMINATOR;
-					Ret += TEXT("}") LINE_TERMINATOR;
-					Ret += TEXT("") LINE_TERMINATOR;
+					Ret += TEXT("\tRet.Ddx = dFdA * A.Ddx;") HLSL_LINE_TERMINATOR;
+					Ret += TEXT("\tRet.Ddy = dFdA * A.Ddy;") HLSL_LINE_TERMINATOR;
+					Ret += TEXT("\treturn Ret;") HLSL_LINE_TERMINATOR;
+					Ret += TEXT("}") HLSL_LINE_TERMINATOR;
+					Ret += TEXT("") HLSL_LINE_TERMINATOR;
 					break;
 				case EFunc1::Sin:
-					Ret += NonLWCBaseName + TEXT(" SinDeriv") + Suffix + TEXT("(") + BaseName + TEXT(" A)") LINE_TERMINATOR;
-					Ret += TEXT("{") LINE_TERMINATOR;
-					Ret += TEXT("\t") + NonLWCBaseName + TEXT(" Ret;") LINE_TERMINATOR;
+					Ret += NonLWCBaseName + TEXT(" SinDeriv") + Suffix + TEXT("(") + BaseName + TEXT(" A)") HLSL_LINE_TERMINATOR;
+					Ret += TEXT("{") HLSL_LINE_TERMINATOR;
+					Ret += TEXT("\t") + NonLWCBaseName + TEXT(" Ret;") HLSL_LINE_TERMINATOR;
 					if (IsLWCType(DerivType))
 					{
 						Translator.AddLWCFuncUsage(ELWCFunctionKind::Other);
-						Ret += TEXT("\tRet.Value = LWCSin(A.Value);") LINE_TERMINATOR;
+						Ret += TEXT("\tRet.Value = LWCSin(A.Value);") HLSL_LINE_TERMINATOR;
 						Ret += TEXT("\t") + FieldName + TEXT(" dFdA = LWCCos(A.Value);");
 					}
 					else
 					{
-						Ret += TEXT("\tRet.Value = sin(A.Value);") LINE_TERMINATOR;
+						Ret += TEXT("\tRet.Value = sin(A.Value);") HLSL_LINE_TERMINATOR;
 						Ret += TEXT("\t") + FieldName + TEXT(" dFdA = cos(A.Value);");
 					}
-					Ret += TEXT("\tRet.Ddx = dFdA * A.Ddx;") LINE_TERMINATOR;
-					Ret += TEXT("\tRet.Ddy = dFdA * A.Ddy;") LINE_TERMINATOR;
-					Ret += TEXT("\treturn Ret;") LINE_TERMINATOR;
-					Ret += TEXT("}") LINE_TERMINATOR;
-					Ret += TEXT("") LINE_TERMINATOR;
+					Ret += TEXT("\tRet.Ddx = dFdA * A.Ddx;") HLSL_LINE_TERMINATOR;
+					Ret += TEXT("\tRet.Ddy = dFdA * A.Ddy;") HLSL_LINE_TERMINATOR;
+					Ret += TEXT("\treturn Ret;") HLSL_LINE_TERMINATOR;
+					Ret += TEXT("}") HLSL_LINE_TERMINATOR;
+					Ret += TEXT("") HLSL_LINE_TERMINATOR;
 					break;
 				case EFunc1::Cos:
-					Ret += NonLWCBaseName + TEXT(" CosDeriv") + Suffix + TEXT("(") + BaseName + TEXT(" A)") LINE_TERMINATOR;
-					Ret += TEXT("{") LINE_TERMINATOR;
-					Ret += TEXT("\t") + NonLWCBaseName + TEXT(" Ret;") LINE_TERMINATOR;
+					Ret += NonLWCBaseName + TEXT(" CosDeriv") + Suffix + TEXT("(") + BaseName + TEXT(" A)") HLSL_LINE_TERMINATOR;
+					Ret += TEXT("{") HLSL_LINE_TERMINATOR;
+					Ret += TEXT("\t") + NonLWCBaseName + TEXT(" Ret;") HLSL_LINE_TERMINATOR;
 					if (IsLWCType(DerivType))
 					{
 						Translator.AddLWCFuncUsage(ELWCFunctionKind::Other);
-						Ret += TEXT("\tRet.Value = LWCCos(A.Value);") LINE_TERMINATOR;
+						Ret += TEXT("\tRet.Value = LWCCos(A.Value);") HLSL_LINE_TERMINATOR;
 						Ret += TEXT("\t") + FieldName + TEXT(" dFdA = -LWCSin(A.Value);");
 					}
 					else
 					{
-						Ret += TEXT("\tRet.Value = cos(A.Value);") LINE_TERMINATOR;
+						Ret += TEXT("\tRet.Value = cos(A.Value);") HLSL_LINE_TERMINATOR;
 						Ret += TEXT("\t") + FieldName + TEXT(" dFdA = -sin(A.Value);");
 					}
-					Ret += TEXT("\tRet.Ddx = dFdA * A.Ddx;") LINE_TERMINATOR;
-					Ret += TEXT("\tRet.Ddy = dFdA * A.Ddy;") LINE_TERMINATOR;
-					Ret += TEXT("\treturn Ret;") LINE_TERMINATOR;
-					Ret += TEXT("}") LINE_TERMINATOR;
-					Ret += TEXT("") LINE_TERMINATOR;
+					Ret += TEXT("\tRet.Ddx = dFdA * A.Ddx;") HLSL_LINE_TERMINATOR;
+					Ret += TEXT("\tRet.Ddy = dFdA * A.Ddy;") HLSL_LINE_TERMINATOR;
+					Ret += TEXT("\treturn Ret;") HLSL_LINE_TERMINATOR;
+					Ret += TEXT("}") HLSL_LINE_TERMINATOR;
+					Ret += TEXT("") HLSL_LINE_TERMINATOR;
 					break;
 				case EFunc1::Tan:
-					Ret += NonLWCBaseName + TEXT(" TanDeriv") + Suffix + TEXT("(") + BaseName + TEXT(" A)") LINE_TERMINATOR;
-					Ret += TEXT("{") LINE_TERMINATOR;
-					Ret += TEXT("\t") + NonLWCBaseName + TEXT(" Ret;") LINE_TERMINATOR;
+					Ret += NonLWCBaseName + TEXT(" TanDeriv") + Suffix + TEXT("(") + BaseName + TEXT(" A)") HLSL_LINE_TERMINATOR;
+					Ret += TEXT("{") HLSL_LINE_TERMINATOR;
+					Ret += TEXT("\t") + NonLWCBaseName + TEXT(" Ret;") HLSL_LINE_TERMINATOR;
 					if (IsLWCType(DerivType))
 					{
 						Translator.AddLWCFuncUsage(ELWCFunctionKind::Other);
-						Ret += TEXT("\tRet.Value = LWCTan(A.Value);") LINE_TERMINATOR;
+						Ret += TEXT("\tRet.Value = LWCTan(A.Value);") HLSL_LINE_TERMINATOR;
 						Ret += TEXT("\t") + FieldName + TEXT(" dFdA = rcp(Pow2(LWCCos(A.Value)));");
 					}
 					else
 					{
-						Ret += TEXT("\tRet.Value = tan(A.Value);") LINE_TERMINATOR;
+						Ret += TEXT("\tRet.Value = tan(A.Value);") HLSL_LINE_TERMINATOR;
 						Ret += TEXT("\t") + FieldName + TEXT(" dFdA = rcp(Pow2(cos(A.Value)));");
 					}
-					Ret += TEXT("\tRet.Ddx = dFdA * A.Ddx;") LINE_TERMINATOR;
-					Ret += TEXT("\tRet.Ddy = dFdA * A.Ddy;") LINE_TERMINATOR;
-					Ret += TEXT("\treturn Ret;") LINE_TERMINATOR;
-					Ret += TEXT("}") LINE_TERMINATOR;
-					Ret += TEXT("") LINE_TERMINATOR;
+					Ret += TEXT("\tRet.Ddx = dFdA * A.Ddx;") HLSL_LINE_TERMINATOR;
+					Ret += TEXT("\tRet.Ddy = dFdA * A.Ddy;") HLSL_LINE_TERMINATOR;
+					Ret += TEXT("\treturn Ret;") HLSL_LINE_TERMINATOR;
+					Ret += TEXT("}") HLSL_LINE_TERMINATOR;
+					Ret += TEXT("") HLSL_LINE_TERMINATOR;
 					break;
 				case EFunc1::Asin:
-					Ret += BaseName + TEXT(" AsinDeriv") + Suffix + TEXT("(") + BaseName + TEXT(" A)") LINE_TERMINATOR;
-					Ret += TEXT("{") LINE_TERMINATOR;
-					Ret += TEXT("\t") + BaseName + TEXT(" Ret;") LINE_TERMINATOR;
-					Ret += TEXT("\tRet.Value = asin(A.Value);") LINE_TERMINATOR;
+					Ret += BaseName + TEXT(" AsinDeriv") + Suffix + TEXT("(") + BaseName + TEXT(" A)") HLSL_LINE_TERMINATOR;
+					Ret += TEXT("{") HLSL_LINE_TERMINATOR;
+					Ret += TEXT("\t") + BaseName + TEXT(" Ret;") HLSL_LINE_TERMINATOR;
+					Ret += TEXT("\tRet.Value = asin(A.Value);") HLSL_LINE_TERMINATOR;
 					Ret += TEXT("\t") + FieldName + TEXT(" dFdA = rsqrt(max(1.0f - A.Value * A.Value, 0.00001f));");
-					Ret += TEXT("\tRet.Ddx = dFdA * A.Ddx;") LINE_TERMINATOR;
-					Ret += TEXT("\tRet.Ddy = dFdA * A.Ddy;") LINE_TERMINATOR;
-					Ret += TEXT("\treturn Ret;") LINE_TERMINATOR;
-					Ret += TEXT("}") LINE_TERMINATOR;
-					Ret += TEXT("") LINE_TERMINATOR;
+					Ret += TEXT("\tRet.Ddx = dFdA * A.Ddx;") HLSL_LINE_TERMINATOR;
+					Ret += TEXT("\tRet.Ddy = dFdA * A.Ddy;") HLSL_LINE_TERMINATOR;
+					Ret += TEXT("\treturn Ret;") HLSL_LINE_TERMINATOR;
+					Ret += TEXT("}") HLSL_LINE_TERMINATOR;
+					Ret += TEXT("") HLSL_LINE_TERMINATOR;
 					break;
 				case EFunc1::AsinFast:
-					Ret += BaseName + TEXT(" AsinFastDeriv") + Suffix + TEXT("(") + BaseName + TEXT(" A)") LINE_TERMINATOR;
-					Ret += TEXT("{") LINE_TERMINATOR;
-					Ret += TEXT("\t") + BaseName + TEXT(" Ret;") LINE_TERMINATOR;
-					Ret += TEXT("\tRet.Value = asinFast(A.Value);") LINE_TERMINATOR;
+					Ret += BaseName + TEXT(" AsinFastDeriv") + Suffix + TEXT("(") + BaseName + TEXT(" A)") HLSL_LINE_TERMINATOR;
+					Ret += TEXT("{") HLSL_LINE_TERMINATOR;
+					Ret += TEXT("\t") + BaseName + TEXT(" Ret;") HLSL_LINE_TERMINATOR;
+					Ret += TEXT("\tRet.Value = asinFast(A.Value);") HLSL_LINE_TERMINATOR;
 					Ret += TEXT("\t") + FieldName + TEXT(" dFdA = rsqrt(max(1.0f - A.Value * A.Value, 0.00001f));");
-					Ret += TEXT("\tRet.Ddx = dFdA * A.Ddx;") LINE_TERMINATOR;
-					Ret += TEXT("\tRet.Ddy = dFdA * A.Ddy;") LINE_TERMINATOR;
-					Ret += TEXT("\treturn Ret;") LINE_TERMINATOR;
-					Ret += TEXT("}") LINE_TERMINATOR;
-					Ret += TEXT("") LINE_TERMINATOR;
+					Ret += TEXT("\tRet.Ddx = dFdA * A.Ddx;") HLSL_LINE_TERMINATOR;
+					Ret += TEXT("\tRet.Ddy = dFdA * A.Ddy;") HLSL_LINE_TERMINATOR;
+					Ret += TEXT("\treturn Ret;") HLSL_LINE_TERMINATOR;
+					Ret += TEXT("}") HLSL_LINE_TERMINATOR;
+					Ret += TEXT("") HLSL_LINE_TERMINATOR;
 					break;
 				case EFunc1::Acos:
-					Ret += BaseName + TEXT(" AcosDeriv") + Suffix + TEXT("(") + BaseName + TEXT(" A)") LINE_TERMINATOR;
-					Ret += TEXT("{") LINE_TERMINATOR;
-					Ret += TEXT("\t") + BaseName + TEXT(" Ret;") LINE_TERMINATOR;
-					Ret += TEXT("\tRet.Value = acos(A.Value);") LINE_TERMINATOR;
+					Ret += BaseName + TEXT(" AcosDeriv") + Suffix + TEXT("(") + BaseName + TEXT(" A)") HLSL_LINE_TERMINATOR;
+					Ret += TEXT("{") HLSL_LINE_TERMINATOR;
+					Ret += TEXT("\t") + BaseName + TEXT(" Ret;") HLSL_LINE_TERMINATOR;
+					Ret += TEXT("\tRet.Value = acos(A.Value);") HLSL_LINE_TERMINATOR;
 					Ret += TEXT("\t") + FieldName + TEXT(" dFdA = -rsqrt(max(1.0f - A.Value * A.Value, 0.00001f));");
-					Ret += TEXT("\tRet.Ddx = dFdA * A.Ddx;") LINE_TERMINATOR;
-					Ret += TEXT("\tRet.Ddy = dFdA * A.Ddy;") LINE_TERMINATOR;
-					Ret += TEXT("\treturn Ret;") LINE_TERMINATOR;
-					Ret += TEXT("}") LINE_TERMINATOR;
-					Ret += TEXT("") LINE_TERMINATOR;
+					Ret += TEXT("\tRet.Ddx = dFdA * A.Ddx;") HLSL_LINE_TERMINATOR;
+					Ret += TEXT("\tRet.Ddy = dFdA * A.Ddy;") HLSL_LINE_TERMINATOR;
+					Ret += TEXT("\treturn Ret;") HLSL_LINE_TERMINATOR;
+					Ret += TEXT("}") HLSL_LINE_TERMINATOR;
+					Ret += TEXT("") HLSL_LINE_TERMINATOR;
 					break;
 				case EFunc1::AcosFast:
-					Ret += BaseName + TEXT(" AcosFastDeriv") + Suffix + TEXT("(") + BaseName + TEXT(" A)") LINE_TERMINATOR;
-					Ret += TEXT("{") LINE_TERMINATOR;
-					Ret += TEXT("\t") + BaseName + TEXT(" Ret;") LINE_TERMINATOR;
-					Ret += TEXT("\tRet.Value = acosFast(A.Value);") LINE_TERMINATOR;
+					Ret += BaseName + TEXT(" AcosFastDeriv") + Suffix + TEXT("(") + BaseName + TEXT(" A)") HLSL_LINE_TERMINATOR;
+					Ret += TEXT("{") HLSL_LINE_TERMINATOR;
+					Ret += TEXT("\t") + BaseName + TEXT(" Ret;") HLSL_LINE_TERMINATOR;
+					Ret += TEXT("\tRet.Value = acosFast(A.Value);") HLSL_LINE_TERMINATOR;
 					Ret += TEXT("\t") + FieldName + TEXT(" dFdA = -rsqrt(max(1.0f - A.Value * A.Value, 0.00001f));");
-					Ret += TEXT("\tRet.Ddx = dFdA * A.Ddx;") LINE_TERMINATOR;
-					Ret += TEXT("\tRet.Ddy = dFdA * A.Ddy;") LINE_TERMINATOR;
-					Ret += TEXT("\treturn Ret;") LINE_TERMINATOR;
-					Ret += TEXT("}") LINE_TERMINATOR;
-					Ret += TEXT("") LINE_TERMINATOR;
+					Ret += TEXT("\tRet.Ddx = dFdA * A.Ddx;") HLSL_LINE_TERMINATOR;
+					Ret += TEXT("\tRet.Ddy = dFdA * A.Ddy;") HLSL_LINE_TERMINATOR;
+					Ret += TEXT("\treturn Ret;") HLSL_LINE_TERMINATOR;
+					Ret += TEXT("}") HLSL_LINE_TERMINATOR;
+					Ret += TEXT("") HLSL_LINE_TERMINATOR;
 					break;
 				case EFunc1::Atan:
-					Ret += BaseName + TEXT(" AtanDeriv") + Suffix + TEXT("(") + BaseName + TEXT(" A)") LINE_TERMINATOR;
-					Ret += TEXT("{") LINE_TERMINATOR;
-					Ret += TEXT("\t") + BaseName + TEXT(" Ret;") LINE_TERMINATOR;
-					Ret += TEXT("\tRet.Value = atan(A.Value);") LINE_TERMINATOR;
+					Ret += BaseName + TEXT(" AtanDeriv") + Suffix + TEXT("(") + BaseName + TEXT(" A)") HLSL_LINE_TERMINATOR;
+					Ret += TEXT("{") HLSL_LINE_TERMINATOR;
+					Ret += TEXT("\t") + BaseName + TEXT(" Ret;") HLSL_LINE_TERMINATOR;
+					Ret += TEXT("\tRet.Value = atan(A.Value);") HLSL_LINE_TERMINATOR;
 					Ret += TEXT("\t") + FieldName + TEXT(" dFdA = rcp(A.Value * A.Value + 1.0f);");
-					Ret += TEXT("\tRet.Ddx = dFdA * A.Ddx;") LINE_TERMINATOR;
-					Ret += TEXT("\tRet.Ddy = dFdA * A.Ddy;") LINE_TERMINATOR;
-					Ret += TEXT("\treturn Ret;") LINE_TERMINATOR;
-					Ret += TEXT("}") LINE_TERMINATOR;
-					Ret += TEXT("") LINE_TERMINATOR;
+					Ret += TEXT("\tRet.Ddx = dFdA * A.Ddx;") HLSL_LINE_TERMINATOR;
+					Ret += TEXT("\tRet.Ddy = dFdA * A.Ddy;") HLSL_LINE_TERMINATOR;
+					Ret += TEXT("\treturn Ret;") HLSL_LINE_TERMINATOR;
+					Ret += TEXT("}") HLSL_LINE_TERMINATOR;
+					Ret += TEXT("") HLSL_LINE_TERMINATOR;
 					break;
 				case EFunc1::AtanFast:
-					Ret += BaseName + TEXT(" AtanFastDeriv") + Suffix + TEXT("(") + BaseName + TEXT(" A)") LINE_TERMINATOR;
-					Ret += TEXT("{") LINE_TERMINATOR;
-					Ret += TEXT("\t") + BaseName + TEXT(" Ret;") LINE_TERMINATOR;
-					Ret += TEXT("\tRet.Value = atanFast(A.Value);") LINE_TERMINATOR;
+					Ret += BaseName + TEXT(" AtanFastDeriv") + Suffix + TEXT("(") + BaseName + TEXT(" A)") HLSL_LINE_TERMINATOR;
+					Ret += TEXT("{") HLSL_LINE_TERMINATOR;
+					Ret += TEXT("\t") + BaseName + TEXT(" Ret;") HLSL_LINE_TERMINATOR;
+					Ret += TEXT("\tRet.Value = atanFast(A.Value);") HLSL_LINE_TERMINATOR;
 					Ret += TEXT("\t") + FieldName + TEXT(" dFdA = rcp(A.Value * A.Value + 1.0f);");
-					Ret += TEXT("\tRet.Ddx = dFdA * A.Ddx;") LINE_TERMINATOR;
-					Ret += TEXT("\tRet.Ddy = dFdA * A.Ddy;") LINE_TERMINATOR;
-					Ret += TEXT("\treturn Ret;") LINE_TERMINATOR;
-					Ret += TEXT("}") LINE_TERMINATOR;
-					Ret += TEXT("") LINE_TERMINATOR;
+					Ret += TEXT("\tRet.Ddx = dFdA * A.Ddx;") HLSL_LINE_TERMINATOR;
+					Ret += TEXT("\tRet.Ddy = dFdA * A.Ddy;") HLSL_LINE_TERMINATOR;
+					Ret += TEXT("\treturn Ret;") HLSL_LINE_TERMINATOR;
+					Ret += TEXT("}") HLSL_LINE_TERMINATOR;
+					Ret += TEXT("") HLSL_LINE_TERMINATOR;
 					break;
 				case EFunc1::Sqrt:
-					Ret += NonLWCBaseName + TEXT(" SqrtDeriv") + Suffix + TEXT("(") + BaseName + TEXT(" A)") LINE_TERMINATOR;
-					Ret += TEXT("{") LINE_TERMINATOR;
-					Ret += TEXT("\t") + NonLWCBaseName + TEXT(" Ret;") LINE_TERMINATOR;
+					Ret += NonLWCBaseName + TEXT(" SqrtDeriv") + Suffix + TEXT("(") + BaseName + TEXT(" A)") HLSL_LINE_TERMINATOR;
+					Ret += TEXT("{") HLSL_LINE_TERMINATOR;
+					Ret += TEXT("\t") + NonLWCBaseName + TEXT(" Ret;") HLSL_LINE_TERMINATOR;
 					if (IsLWCType(DerivType))
 					{
 						Translator.AddLWCFuncUsage(ELWCFunctionKind::Other);
-						Ret += TEXT("\tRet.Value = LWCSqrt(A.Value);") LINE_TERMINATOR;
-						Ret += TEXT("\t") + NonLWCFieldName + TEXT(" dFdA = 0.5f * LWCRsqrt(LWCMax(A.Value, 0.00001f));") LINE_TERMINATOR;
+						Ret += TEXT("\tRet.Value = LWCSqrt(A.Value);") HLSL_LINE_TERMINATOR;
+						Ret += TEXT("\t") + NonLWCFieldName + TEXT(" dFdA = 0.5f * LWCRsqrt(LWCMax(A.Value, 0.00001f));") HLSL_LINE_TERMINATOR;
 					}
 					else
 					{
-						Ret += TEXT("\tRet.Value = sqrt(A.Value);") LINE_TERMINATOR;
-						Ret += TEXT("\t") + NonLWCFieldName + TEXT(" dFdA = 0.5f * rsqrt(max(A.Value, 0.00001f));") LINE_TERMINATOR;
+						Ret += TEXT("\tRet.Value = sqrt(A.Value);") HLSL_LINE_TERMINATOR;
+						Ret += TEXT("\t") + NonLWCFieldName + TEXT(" dFdA = 0.5f * rsqrt(max(A.Value, 0.00001f));") HLSL_LINE_TERMINATOR;
 					}
 					
-					Ret += TEXT("\tRet.Ddx = dFdA * A.Ddx;") LINE_TERMINATOR;
-					Ret += TEXT("\tRet.Ddy = dFdA * A.Ddy;") LINE_TERMINATOR;
-					Ret += TEXT("\treturn Ret;") LINE_TERMINATOR;
-					Ret += TEXT("}") LINE_TERMINATOR;
-					Ret += TEXT("") LINE_TERMINATOR;
+					Ret += TEXT("\tRet.Ddx = dFdA * A.Ddx;") HLSL_LINE_TERMINATOR;
+					Ret += TEXT("\tRet.Ddy = dFdA * A.Ddy;") HLSL_LINE_TERMINATOR;
+					Ret += TEXT("\treturn Ret;") HLSL_LINE_TERMINATOR;
+					Ret += TEXT("}") HLSL_LINE_TERMINATOR;
+					Ret += TEXT("") HLSL_LINE_TERMINATOR;
 					break;
 				case EFunc1::Rcp:
-					Ret += NonLWCBaseName + TEXT(" RcpDeriv") + Suffix + TEXT("(") + BaseName + TEXT(" A)") LINE_TERMINATOR;
-					Ret += TEXT("{") LINE_TERMINATOR;
-					Ret += TEXT("\t") + NonLWCBaseName + TEXT(" Ret;") LINE_TERMINATOR;
+					Ret += NonLWCBaseName + TEXT(" RcpDeriv") + Suffix + TEXT("(") + BaseName + TEXT(" A)") HLSL_LINE_TERMINATOR;
+					Ret += TEXT("{") HLSL_LINE_TERMINATOR;
+					Ret += TEXT("\t") + NonLWCBaseName + TEXT(" Ret;") HLSL_LINE_TERMINATOR;
 					if (IsLWCType(DerivType))
 					{
 						Translator.AddLWCFuncUsage(ELWCFunctionKind::Other);
-						Ret += TEXT("\tRet.Value = LWCRcp(A.Value);") LINE_TERMINATOR;
+						Ret += TEXT("\tRet.Value = LWCRcp(A.Value);") HLSL_LINE_TERMINATOR;
 					}
 					else
 					{
-						Ret += TEXT("\tRet.Value = rcp(A.Value);") LINE_TERMINATOR;
+						Ret += TEXT("\tRet.Value = rcp(A.Value);") HLSL_LINE_TERMINATOR;
 					}
 					Ret += TEXT("\t") + NonLWCFieldName + TEXT(" dFdA = -Ret.Value * Ret.Value;");
-					Ret += TEXT("\tRet.Ddx = dFdA * A.Ddx;") LINE_TERMINATOR;
-					Ret += TEXT("\tRet.Ddy = dFdA * A.Ddy;") LINE_TERMINATOR;
-					Ret += TEXT("\treturn Ret;") LINE_TERMINATOR;
-					Ret += TEXT("}") LINE_TERMINATOR;
-					Ret += TEXT("") LINE_TERMINATOR;
+					Ret += TEXT("\tRet.Ddx = dFdA * A.Ddx;") HLSL_LINE_TERMINATOR;
+					Ret += TEXT("\tRet.Ddy = dFdA * A.Ddy;") HLSL_LINE_TERMINATOR;
+					Ret += TEXT("\treturn Ret;") HLSL_LINE_TERMINATOR;
+					Ret += TEXT("}") HLSL_LINE_TERMINATOR;
+					Ret += TEXT("") HLSL_LINE_TERMINATOR;
 					break;
 				case EFunc1::Rsqrt:
-					Ret += NonLWCBaseName + TEXT(" RsqrtDeriv") + Suffix + TEXT("(") + BaseName + TEXT(" A)") LINE_TERMINATOR;
-					Ret += TEXT("{") LINE_TERMINATOR;
-					Ret += TEXT("\t") + NonLWCBaseName + TEXT(" Ret;") LINE_TERMINATOR;
+					Ret += NonLWCBaseName + TEXT(" RsqrtDeriv") + Suffix + TEXT("(") + BaseName + TEXT(" A)") HLSL_LINE_TERMINATOR;
+					Ret += TEXT("{") HLSL_LINE_TERMINATOR;
+					Ret += TEXT("\t") + NonLWCBaseName + TEXT(" Ret;") HLSL_LINE_TERMINATOR;
 					if (IsLWCType(DerivType))
 					{
 						Translator.AddLWCFuncUsage(ELWCFunctionKind::Other);
-						Ret += TEXT("\tRet.Value = LWCRsqrt(A.Value);") LINE_TERMINATOR;
+						Ret += TEXT("\tRet.Value = LWCRsqrt(A.Value);") HLSL_LINE_TERMINATOR;
 						Ret += TEXT("\t") + NonLWCFieldName + TEXT(" dFdA = -0.5f * LWCRsqrt(A.Value) * LWCRcp(A.Value);");
 					}
 					else
 					{
-						Ret += TEXT("\tRet.Value = rsqrt(A.Value);") LINE_TERMINATOR;
+						Ret += TEXT("\tRet.Value = rsqrt(A.Value);") HLSL_LINE_TERMINATOR;
 						Ret += TEXT("\t") + NonLWCFieldName + TEXT(" dFdA = -0.5f * rsqrt(A.Value) * rcp(A.Value);");
 					}
 					
-					Ret += TEXT("\tRet.Ddx = dFdA * A.Ddx;") LINE_TERMINATOR;
-					Ret += TEXT("\tRet.Ddy = dFdA * A.Ddy;") LINE_TERMINATOR;
-					Ret += TEXT("\treturn Ret;") LINE_TERMINATOR;
-					Ret += TEXT("}") LINE_TERMINATOR;
-					Ret += TEXT("") LINE_TERMINATOR;
+					Ret += TEXT("\tRet.Ddx = dFdA * A.Ddx;") HLSL_LINE_TERMINATOR;
+					Ret += TEXT("\tRet.Ddy = dFdA * A.Ddy;") HLSL_LINE_TERMINATOR;
+					Ret += TEXT("\treturn Ret;") HLSL_LINE_TERMINATOR;
+					Ret += TEXT("}") HLSL_LINE_TERMINATOR;
+					Ret += TEXT("") HLSL_LINE_TERMINATOR;
 					break;
 				case EFunc1::Saturate:
-					Ret += NonLWCBaseName + TEXT(" SaturateDeriv") + Suffix + TEXT("(") + BaseName + TEXT(" A)") LINE_TERMINATOR;
-					Ret += TEXT("{") LINE_TERMINATOR;
-					Ret += TEXT("\t") + NonLWCFieldName + TEXT(" Zero = 0.0f;") LINE_TERMINATOR;
-					Ret += TEXT("\t") + NonLWCBaseName + TEXT(" Ret;") LINE_TERMINATOR;
+					Ret += NonLWCBaseName + TEXT(" SaturateDeriv") + Suffix + TEXT("(") + BaseName + TEXT(" A)") HLSL_LINE_TERMINATOR;
+					Ret += TEXT("{") HLSL_LINE_TERMINATOR;
+					Ret += TEXT("\t") + NonLWCFieldName + TEXT(" Zero = 0.0f;") HLSL_LINE_TERMINATOR;
+					Ret += TEXT("\t") + NonLWCBaseName + TEXT(" Ret;") HLSL_LINE_TERMINATOR;
 					if (IsLWCType(DerivType))
 					{
 						Translator.AddLWCFuncUsage(ELWCFunctionKind::Other);
-						Ret += TEXT("\tRet.Value = LWCSaturate(A.Value);") LINE_TERMINATOR;
-						Ret += TEXT("\t") + BoolName + TEXT(" InRange = LWCEquals(Ret.Value, A.Value);") LINE_TERMINATOR;
+						Ret += TEXT("\tRet.Value = LWCSaturate(A.Value);") HLSL_LINE_TERMINATOR;
+						Ret += TEXT("\t") + BoolName + TEXT(" InRange = LWCEquals(Ret.Value, A.Value);") HLSL_LINE_TERMINATOR;
 					}
 					else
 					{
-						Ret += TEXT("\tRet.Value = saturate(A.Value);") LINE_TERMINATOR;
-						Ret += TEXT("\t") + BoolName + TEXT(" InRange = and(0.0 < A.Value, A.Value < 1.0);") LINE_TERMINATOR;
+						Ret += TEXT("\tRet.Value = saturate(A.Value);") HLSL_LINE_TERMINATOR;
+						Ret += TEXT("\t") + BoolName + TEXT(" InRange = and(0.0 < A.Value, A.Value < 1.0);") HLSL_LINE_TERMINATOR;
 					}
-					Ret += TEXT("\tRet.Ddx = select(InRange, A.Ddx, Zero);") LINE_TERMINATOR;
-					Ret += TEXT("\tRet.Ddy = select(InRange, A.Ddy, Zero);") LINE_TERMINATOR;
-					Ret += TEXT("\treturn Ret;") LINE_TERMINATOR;
-					Ret += TEXT("}") LINE_TERMINATOR;
-					Ret += TEXT("") LINE_TERMINATOR;
+					Ret += TEXT("\tRet.Ddx = select(InRange, A.Ddx, Zero);") HLSL_LINE_TERMINATOR;
+					Ret += TEXT("\tRet.Ddy = select(InRange, A.Ddy, Zero);") HLSL_LINE_TERMINATOR;
+					Ret += TEXT("\treturn Ret;") HLSL_LINE_TERMINATOR;
+					Ret += TEXT("}") HLSL_LINE_TERMINATOR;
+					Ret += TEXT("") HLSL_LINE_TERMINATOR;
 					break;
 				case EFunc1::Frac:
-					Ret += NonLWCBaseName + TEXT(" FracDeriv") + Suffix + TEXT("(") + BaseName + TEXT(" A)") LINE_TERMINATOR;
-					Ret += TEXT("{") LINE_TERMINATOR;
-					Ret += TEXT("\t") + NonLWCBaseName + TEXT(" Ret;") LINE_TERMINATOR;
+					Ret += NonLWCBaseName + TEXT(" FracDeriv") + Suffix + TEXT("(") + BaseName + TEXT(" A)") HLSL_LINE_TERMINATOR;
+					Ret += TEXT("{") HLSL_LINE_TERMINATOR;
+					Ret += TEXT("\t") + NonLWCBaseName + TEXT(" Ret;") HLSL_LINE_TERMINATOR;
 					if (IsLWCType(DerivType))
 					{
 						Translator.AddLWCFuncUsage(ELWCFunctionKind::Other);
-						Ret += TEXT("\tRet.Value = LWCFrac(A.Value);") LINE_TERMINATOR;
+						Ret += TEXT("\tRet.Value = LWCFrac(A.Value);") HLSL_LINE_TERMINATOR;
 					}
 					else
 					{
-						Ret += TEXT("\tRet.Value = frac(A.Value);") LINE_TERMINATOR;
+						Ret += TEXT("\tRet.Value = frac(A.Value);") HLSL_LINE_TERMINATOR;
 					}
-					Ret += TEXT("\tRet.Ddx = A.Ddx;") LINE_TERMINATOR;
-					Ret += TEXT("\tRet.Ddy = A.Ddy;") LINE_TERMINATOR;
-					Ret += TEXT("\treturn Ret;") LINE_TERMINATOR;
-					Ret += TEXT("}") LINE_TERMINATOR;
-					Ret += TEXT("") LINE_TERMINATOR;
+					Ret += TEXT("\tRet.Ddx = A.Ddx;") HLSL_LINE_TERMINATOR;
+					Ret += TEXT("\tRet.Ddy = A.Ddy;") HLSL_LINE_TERMINATOR;
+					Ret += TEXT("\treturn Ret;") HLSL_LINE_TERMINATOR;
+					Ret += TEXT("}") HLSL_LINE_TERMINATOR;
+					Ret += TEXT("") HLSL_LINE_TERMINATOR;
 					break;
 				case EFunc1::Log:
-					Ret += BaseName + TEXT(" LogDeriv") + Suffix + TEXT("(") + BaseName + TEXT(" A)") LINE_TERMINATOR;
-					Ret += TEXT("{") LINE_TERMINATOR;
-					Ret += TEXT("\t") + BaseName + TEXT(" Ret;") LINE_TERMINATOR;
-					Ret += TEXT("\tRet.Value = log(A.Value);") LINE_TERMINATOR;
+					Ret += BaseName + TEXT(" LogDeriv") + Suffix + TEXT("(") + BaseName + TEXT(" A)") HLSL_LINE_TERMINATOR;
+					Ret += TEXT("{") HLSL_LINE_TERMINATOR;
+					Ret += TEXT("\t") + BaseName + TEXT(" Ret;") HLSL_LINE_TERMINATOR;
+					Ret += TEXT("\tRet.Value = log(A.Value);") HLSL_LINE_TERMINATOR;
 					Ret += TEXT("\t") + FieldName + TEXT(" dFdA = rcp(A.Value);");
-					Ret += TEXT("\tRet.Ddx = dFdA * A.Ddx ;") LINE_TERMINATOR;
-					Ret += TEXT("\tRet.Ddy = dFdA * A.Ddy;") LINE_TERMINATOR;
-					Ret += TEXT("\treturn Ret;") LINE_TERMINATOR;
-					Ret += TEXT("}") LINE_TERMINATOR;
-					Ret += TEXT("") LINE_TERMINATOR;
+					Ret += TEXT("\tRet.Ddx = dFdA * A.Ddx ;") HLSL_LINE_TERMINATOR;
+					Ret += TEXT("\tRet.Ddy = dFdA * A.Ddy;") HLSL_LINE_TERMINATOR;
+					Ret += TEXT("\treturn Ret;") HLSL_LINE_TERMINATOR;
+					Ret += TEXT("}") HLSL_LINE_TERMINATOR;
+					Ret += TEXT("") HLSL_LINE_TERMINATOR;
 					break;
 				case EFunc1::Log2:
-					Ret += BaseName + TEXT(" Log2Deriv") + Suffix + TEXT("(") + BaseName + TEXT(" A)") LINE_TERMINATOR;
-					Ret += TEXT("{") LINE_TERMINATOR;
-					Ret += TEXT("\t") + BaseName + TEXT(" Ret;") LINE_TERMINATOR;
-					Ret += TEXT("\tRet.Value = log2(A.Value);") LINE_TERMINATOR;
+					Ret += BaseName + TEXT(" Log2Deriv") + Suffix + TEXT("(") + BaseName + TEXT(" A)") HLSL_LINE_TERMINATOR;
+					Ret += TEXT("{") HLSL_LINE_TERMINATOR;
+					Ret += TEXT("\t") + BaseName + TEXT(" Ret;") HLSL_LINE_TERMINATOR;
+					Ret += TEXT("\tRet.Value = log2(A.Value);") HLSL_LINE_TERMINATOR;
 					Ret += TEXT("\t") + FieldName + TEXT(" dFdA = rcp(A.Value) * 1.442695f;");
-					Ret += TEXT("\tRet.Ddx = dFdA * A.Ddx ;") LINE_TERMINATOR;
-					Ret += TEXT("\tRet.Ddy = dFdA * A.Ddy;") LINE_TERMINATOR;
-					Ret += TEXT("\treturn Ret;") LINE_TERMINATOR;
-					Ret += TEXT("}") LINE_TERMINATOR;
-					Ret += TEXT("") LINE_TERMINATOR;
+					Ret += TEXT("\tRet.Ddx = dFdA * A.Ddx ;") HLSL_LINE_TERMINATOR;
+					Ret += TEXT("\tRet.Ddy = dFdA * A.Ddy;") HLSL_LINE_TERMINATOR;
+					Ret += TEXT("\treturn Ret;") HLSL_LINE_TERMINATOR;
+					Ret += TEXT("}") HLSL_LINE_TERMINATOR;
+					Ret += TEXT("") HLSL_LINE_TERMINATOR;
 					break;
 				case EFunc1::Log10:
-					Ret += BaseName + TEXT(" Log10Deriv") + Suffix + TEXT("(") + BaseName + TEXT(" A)") LINE_TERMINATOR;
-					Ret += TEXT("{") LINE_TERMINATOR;
-					Ret += TEXT("\t") + BaseName + TEXT(" Ret;") LINE_TERMINATOR;
-					Ret += TEXT("\tRet.Value = log10(A.Value);") LINE_TERMINATOR;
+					Ret += BaseName + TEXT(" Log10Deriv") + Suffix + TEXT("(") + BaseName + TEXT(" A)") HLSL_LINE_TERMINATOR;
+					Ret += TEXT("{") HLSL_LINE_TERMINATOR;
+					Ret += TEXT("\t") + BaseName + TEXT(" Ret;") HLSL_LINE_TERMINATOR;
+					Ret += TEXT("\tRet.Value = log10(A.Value);") HLSL_LINE_TERMINATOR;
 					Ret += TEXT("\t") + FieldName + TEXT(" dFdA = rcp(A.Value) * 0.4342945f;");
-					Ret += TEXT("\tRet.Ddx = dFdA * A.Ddx;") LINE_TERMINATOR;
-					Ret += TEXT("\tRet.Ddy = dFdA * A.Ddy;") LINE_TERMINATOR;
-					Ret += TEXT("\treturn Ret;") LINE_TERMINATOR;
-					Ret += TEXT("}") LINE_TERMINATOR;
-					Ret += TEXT("") LINE_TERMINATOR;
+					Ret += TEXT("\tRet.Ddx = dFdA * A.Ddx;") HLSL_LINE_TERMINATOR;
+					Ret += TEXT("\tRet.Ddy = dFdA * A.Ddy;") HLSL_LINE_TERMINATOR;
+					Ret += TEXT("\treturn Ret;") HLSL_LINE_TERMINATOR;
+					Ret += TEXT("}") HLSL_LINE_TERMINATOR;
+					Ret += TEXT("") HLSL_LINE_TERMINATOR;
 					break;
 				case EFunc1::Exp:
-					Ret += BaseName + TEXT(" ExpDeriv") + Suffix + TEXT("(") + BaseName + TEXT(" A)") LINE_TERMINATOR;
-					Ret += TEXT("{") LINE_TERMINATOR;
-					Ret += TEXT("\t") + BaseName + TEXT(" Ret;") LINE_TERMINATOR;
-					Ret += TEXT("\tRet.Value = exp(A.Value);") LINE_TERMINATOR;
-					Ret += TEXT("\tRet.Ddx = exp(A.Value) * A.Ddx;") LINE_TERMINATOR;
-					Ret += TEXT("\tRet.Ddy = exp(A.Value) * A.Ddy;") LINE_TERMINATOR;
-					Ret += TEXT("\treturn Ret;") LINE_TERMINATOR;
-					Ret += TEXT("}") LINE_TERMINATOR;
-					Ret += TEXT("") LINE_TERMINATOR;
+					Ret += BaseName + TEXT(" ExpDeriv") + Suffix + TEXT("(") + BaseName + TEXT(" A)") HLSL_LINE_TERMINATOR;
+					Ret += TEXT("{") HLSL_LINE_TERMINATOR;
+					Ret += TEXT("\t") + BaseName + TEXT(" Ret;") HLSL_LINE_TERMINATOR;
+					Ret += TEXT("\tRet.Value = exp(A.Value);") HLSL_LINE_TERMINATOR;
+					Ret += TEXT("\tRet.Ddx = exp(A.Value) * A.Ddx;") HLSL_LINE_TERMINATOR;
+					Ret += TEXT("\tRet.Ddy = exp(A.Value) * A.Ddy;") HLSL_LINE_TERMINATOR;
+					Ret += TEXT("\treturn Ret;") HLSL_LINE_TERMINATOR;
+					Ret += TEXT("}") HLSL_LINE_TERMINATOR;
+					Ret += TEXT("") HLSL_LINE_TERMINATOR;
 					break;
 				case EFunc1::Exp2:
-					Ret += BaseName + TEXT(" Exp2Deriv") + Suffix + TEXT("(") + BaseName + TEXT(" A)") LINE_TERMINATOR;
-					Ret += TEXT("{") LINE_TERMINATOR;
-					Ret += TEXT("\t") + BaseName + TEXT(" Ret;") LINE_TERMINATOR;
-					Ret += TEXT("\tRet.Value = exp2(A.Value);") LINE_TERMINATOR;
+					Ret += BaseName + TEXT(" Exp2Deriv") + Suffix + TEXT("(") + BaseName + TEXT(" A)") HLSL_LINE_TERMINATOR;
+					Ret += TEXT("{") HLSL_LINE_TERMINATOR;
+					Ret += TEXT("\t") + BaseName + TEXT(" Ret;") HLSL_LINE_TERMINATOR;
+					Ret += TEXT("\tRet.Value = exp2(A.Value);") HLSL_LINE_TERMINATOR;
 					Ret += TEXT("\t") + FieldName + TEXT(" dFdA = exp2(A.Value) * 0.693147f;");
-					Ret += TEXT("\tRet.Ddx = dFdA * A.Ddx;") LINE_TERMINATOR;
-					Ret += TEXT("\tRet.Ddy = dFdA * A.Ddy;") LINE_TERMINATOR;
-					Ret += TEXT("\treturn Ret;") LINE_TERMINATOR;
-					Ret += TEXT("}") LINE_TERMINATOR;
-					Ret += TEXT("") LINE_TERMINATOR;
+					Ret += TEXT("\tRet.Ddx = dFdA * A.Ddx;") HLSL_LINE_TERMINATOR;
+					Ret += TEXT("\tRet.Ddy = dFdA * A.Ddy;") HLSL_LINE_TERMINATOR;
+					Ret += TEXT("\treturn Ret;") HLSL_LINE_TERMINATOR;
+					Ret += TEXT("}") HLSL_LINE_TERMINATOR;
+					Ret += TEXT("") HLSL_LINE_TERMINATOR;
 					break;
 				case EFunc1::Length:
-					Ret += ScalarName + TEXT(" LengthDeriv") + Suffix + TEXT("(") + BaseName + TEXT(" A)") LINE_TERMINATOR;
-					Ret += TEXT("{") LINE_TERMINATOR;
-					Ret += TEXT("\t") + ScalarName + TEXT(" Ret;") LINE_TERMINATOR;
+					Ret += ScalarName + TEXT(" LengthDeriv") + Suffix + TEXT("(") + BaseName + TEXT(" A)") HLSL_LINE_TERMINATOR;
+					Ret += TEXT("{") HLSL_LINE_TERMINATOR;
+					Ret += TEXT("\t") + ScalarName + TEXT(" Ret;") HLSL_LINE_TERMINATOR;
 					if (IsLWCType(DerivType))
 					{
 						Translator.AddLWCFuncUsage(ELWCFunctionKind::Other);
-						Ret += TEXT("\tRet.Value = LWCLength(A.Value);") LINE_TERMINATOR;
+						Ret += TEXT("\tRet.Value = LWCLength(A.Value);") HLSL_LINE_TERMINATOR;
 					}
 					else
 					{
-						Ret += TEXT("\tRet.Value = length(A.Value);") LINE_TERMINATOR;
+						Ret += TEXT("\tRet.Value = length(A.Value);") HLSL_LINE_TERMINATOR;
 					}
-					Ret += FString::Printf(TEXT("\tFloatDeriv Deriv = SqrtDeriv%s(DotDeriv%s(A,A));") LINE_TERMINATOR, * Suffix, * Suffix);
-					Ret += TEXT("\tRet.Ddx = Deriv.Ddx;") LINE_TERMINATOR;
-					Ret += TEXT("\tRet.Ddy = Deriv.Ddy;") LINE_TERMINATOR;
-					Ret += TEXT("\treturn Ret;") LINE_TERMINATOR;
-					Ret += TEXT("}") LINE_TERMINATOR;
-					Ret += TEXT("") LINE_TERMINATOR;
+					Ret += FString::Printf(TEXT("\tFloatDeriv Deriv = SqrtDeriv%s(DotDeriv%s(A,A));") HLSL_LINE_TERMINATOR, * Suffix, * Suffix);
+					Ret += TEXT("\tRet.Ddx = Deriv.Ddx;") HLSL_LINE_TERMINATOR;
+					Ret += TEXT("\tRet.Ddy = Deriv.Ddy;") HLSL_LINE_TERMINATOR;
+					Ret += TEXT("\treturn Ret;") HLSL_LINE_TERMINATOR;
+					Ret += TEXT("}") HLSL_LINE_TERMINATOR;
+					Ret += TEXT("") HLSL_LINE_TERMINATOR;
 					break;
 				case EFunc1::InvLength:
-					Ret += TEXT("FloatDeriv InvLengthDeriv") + Suffix + TEXT("(") + BaseName + TEXT(" A)") LINE_TERMINATOR;
-					Ret += TEXT("{") LINE_TERMINATOR;
-					Ret += TEXT("\tFloatDeriv Ret;") LINE_TERMINATOR;
+					Ret += TEXT("FloatDeriv InvLengthDeriv") + Suffix + TEXT("(") + BaseName + TEXT(" A)") HLSL_LINE_TERMINATOR;
+					Ret += TEXT("{") HLSL_LINE_TERMINATOR;
+					Ret += TEXT("\tFloatDeriv Ret;") HLSL_LINE_TERMINATOR;
 					if (IsLWCType(DerivType))
 					{
 						Translator.AddLWCFuncUsage(ELWCFunctionKind::Other);
-						Ret += TEXT("\tRet.Value = LWCRcpLength(A.Value);") LINE_TERMINATOR;
+						Ret += TEXT("\tRet.Value = LWCRcpLength(A.Value);") HLSL_LINE_TERMINATOR;
 					}
 					else
 					{
-						Ret += TEXT("\tRet.Value = rcp(length(A.Value));") LINE_TERMINATOR;
+						Ret += TEXT("\tRet.Value = rcp(length(A.Value));") HLSL_LINE_TERMINATOR;
 					}
-					Ret += FString::Printf(TEXT("\tFloatDeriv Deriv = RsqrtDeriv%s(DotDeriv%s(A,A));") LINE_TERMINATOR, *Suffix, *Suffix);
-					Ret += TEXT("\tRet.Ddx = Deriv.Ddx;") LINE_TERMINATOR;
-					Ret += TEXT("\tRet.Ddy = Deriv.Ddy;") LINE_TERMINATOR;
-					Ret += TEXT("\treturn Ret;") LINE_TERMINATOR;
-					Ret += TEXT("}") LINE_TERMINATOR;
-					Ret += TEXT("") LINE_TERMINATOR;
+					Ret += FString::Printf(TEXT("\tFloatDeriv Deriv = RsqrtDeriv%s(DotDeriv%s(A,A));") HLSL_LINE_TERMINATOR, *Suffix, *Suffix);
+					Ret += TEXT("\tRet.Ddx = Deriv.Ddx;") HLSL_LINE_TERMINATOR;
+					Ret += TEXT("\tRet.Ddy = Deriv.Ddy;") HLSL_LINE_TERMINATOR;
+					Ret += TEXT("\treturn Ret;") HLSL_LINE_TERMINATOR;
+					Ret += TEXT("}") HLSL_LINE_TERMINATOR;
+					Ret += TEXT("") HLSL_LINE_TERMINATOR;
 					break;
 				case EFunc1::Normalize:
-					Ret += NonLWCBaseName + TEXT(" NormalizeDeriv") + Suffix + TEXT("(") + BaseName + TEXT(" A)") LINE_TERMINATOR;
-					Ret += TEXT("{") LINE_TERMINATOR;
-					Ret += FString::Printf(TEXT("\tFloatDeriv InvLen = RsqrtDeriv%s(DotDeriv%s(A,A));") LINE_TERMINATOR, *Suffix, *Suffix);
-					Ret += TEXT("\t") + BaseName + TEXT(" Ret = MulDeriv") + Suffix + TEXT("(") + ConvertDeriv(TEXT("InvLen"), DerivType, EDerivativeType::Float1) + TEXT(", A);") LINE_TERMINATOR;
+					Ret += NonLWCBaseName + TEXT(" NormalizeDeriv") + Suffix + TEXT("(") + BaseName + TEXT(" A)") HLSL_LINE_TERMINATOR;
+					Ret += TEXT("{") HLSL_LINE_TERMINATOR;
+					Ret += FString::Printf(TEXT("\tFloatDeriv InvLen = RsqrtDeriv%s(DotDeriv%s(A,A));") HLSL_LINE_TERMINATOR, *Suffix, *Suffix);
+					Ret += TEXT("\t") + BaseName + TEXT(" Ret = MulDeriv") + Suffix + TEXT("(") + ConvertDeriv(TEXT("InvLen"), DerivType, EDerivativeType::Float1) + TEXT(", A);") HLSL_LINE_TERMINATOR;
 					if (IsLWCType(DerivType))
 					{
 						// Convert the result to non-LWC
-						Ret += TEXT("\treturn ") + ConvertDeriv(TEXT("Ret"), MakeNonLWCType(DerivType), DerivType) + TEXT(";") LINE_TERMINATOR;
+						Ret += TEXT("\treturn ") + ConvertDeriv(TEXT("Ret"), MakeNonLWCType(DerivType), DerivType) + TEXT(";") HLSL_LINE_TERMINATOR;
 					}
 					else
 					{
-						Ret += TEXT("\treturn Ret;") LINE_TERMINATOR;
+						Ret += TEXT("\treturn Ret;") HLSL_LINE_TERMINATOR;
 					}
-					Ret += TEXT("}") LINE_TERMINATOR;
-					Ret += TEXT("") LINE_TERMINATOR;
+					Ret += TEXT("}") HLSL_LINE_TERMINATOR;
+					Ret += TEXT("") HLSL_LINE_TERMINATOR;
 					break;
 				default:
 					check(0);
@@ -2215,53 +2215,53 @@ FString FMaterialDerivativeAutogen::GenerateUsedFunctions(FHLSLMaterialTranslato
 			
 			// lerp(a,b,s) = a*(1-s) + b*s
 			// lerp(a,b,s)' = a' * (1 - s') + b' * s + s' * (b - a)
-			Ret += FString::Printf(TEXT("%s LerpDeriv(%s A, %s B, %s S)"), *BaseName, *BaseName, *BaseName, *NonLWCBaseName) + LINE_TERMINATOR;
-			Ret += TEXT("{") LINE_TERMINATOR;
-			Ret += TEXT("\t") + BaseName + TEXT(" Ret;") LINE_TERMINATOR;
+			Ret += FString::Printf(TEXT("%s LerpDeriv(%s A, %s B, %s S)"), *BaseName, *BaseName, *BaseName, *NonLWCBaseName) + HLSL_LINE_TERMINATOR;
+			Ret += TEXT("{") HLSL_LINE_TERMINATOR;
+			Ret += TEXT("\t") + BaseName + TEXT(" Ret;") HLSL_LINE_TERMINATOR;
 			if (IsLWCType(DerivType))
 			{
 				Translator.AddLWCFuncUsage(ELWCFunctionKind::Other);
-				Ret += TEXT("\tRet.Value = LWCLerp(A.Value, B.Value, S.Value);") LINE_TERMINATOR;
+				Ret += TEXT("\tRet.Value = LWCLerp(A.Value, B.Value, S.Value);") HLSL_LINE_TERMINATOR;
 				Translator.AddLWCFuncUsage(ELWCFunctionKind::Subtract, 2);
-				Ret += TEXT("\tRet.Ddx = lerp(A.Ddx, B.Ddx, S.Value) + S.Ddx * LWCToFloat(LWCSubtract(B.Value, A.Value));") LINE_TERMINATOR;
-				Ret += TEXT("\tRet.Ddy = lerp(A.Ddy, B.Ddy, S.Value) + S.Ddy * LWCToFloat(LWCSubtract(B.Value, A.Value));") LINE_TERMINATOR;
+				Ret += TEXT("\tRet.Ddx = lerp(A.Ddx, B.Ddx, S.Value) + S.Ddx * LWCToFloat(LWCSubtract(B.Value, A.Value));") HLSL_LINE_TERMINATOR;
+				Ret += TEXT("\tRet.Ddy = lerp(A.Ddy, B.Ddy, S.Value) + S.Ddy * LWCToFloat(LWCSubtract(B.Value, A.Value));") HLSL_LINE_TERMINATOR;
 			}
 			else
 			{
-				Ret += TEXT("\tRet.Value = lerp(A.Value, B.Value, S.Value);") LINE_TERMINATOR;
-				Ret += TEXT("\tRet.Ddx = lerp(A.Ddx, B.Ddx, S.Value) + S.Ddx * (B.Value - A.Value);") LINE_TERMINATOR;
-				Ret += TEXT("\tRet.Ddy = lerp(A.Ddy, B.Ddy, S.Value) + S.Ddy * (B.Value - A.Value);") LINE_TERMINATOR;
+				Ret += TEXT("\tRet.Value = lerp(A.Value, B.Value, S.Value);") HLSL_LINE_TERMINATOR;
+				Ret += TEXT("\tRet.Ddx = lerp(A.Ddx, B.Ddx, S.Value) + S.Ddx * (B.Value - A.Value);") HLSL_LINE_TERMINATOR;
+				Ret += TEXT("\tRet.Ddy = lerp(A.Ddy, B.Ddy, S.Value) + S.Ddy * (B.Value - A.Value);") HLSL_LINE_TERMINATOR;
 			}
-			Ret += TEXT("\treturn Ret;") LINE_TERMINATOR;
-			Ret += TEXT("}") LINE_TERMINATOR;
-			Ret += TEXT("") LINE_TERMINATOR;
+			Ret += TEXT("\treturn Ret;") HLSL_LINE_TERMINATOR;
+			Ret += TEXT("}") HLSL_LINE_TERMINATOR;
+			Ret += TEXT("") HLSL_LINE_TERMINATOR;
 		}
 
 		if (bIfEnabled[Index] || IsDebugGenerateAllFunctionsEnabled())
 		{
 			// TODO - do we need to make these bool-vectors?
-			Ret += FString::Printf(TEXT("%s IfDeriv(bool bGreaterEqual, %s Greater, %s Less)"), *BaseName, *BaseName, *BaseName) + LINE_TERMINATOR;
-			Ret += TEXT("{") LINE_TERMINATOR;
-			Ret += TEXT("\tif(bGreaterEqual)") LINE_TERMINATOR;
-			Ret += TEXT("\t\treturn Greater;") LINE_TERMINATOR;
-			Ret += TEXT("\telse") LINE_TERMINATOR;
-			Ret += TEXT("\t\treturn Less;") LINE_TERMINATOR;
-			Ret += TEXT("}") LINE_TERMINATOR;
-			Ret += TEXT("") LINE_TERMINATOR;
+			Ret += FString::Printf(TEXT("%s IfDeriv(bool bGreaterEqual, %s Greater, %s Less)"), *BaseName, *BaseName, *BaseName) + HLSL_LINE_TERMINATOR;
+			Ret += TEXT("{") HLSL_LINE_TERMINATOR;
+			Ret += TEXT("\tif(bGreaterEqual)") HLSL_LINE_TERMINATOR;
+			Ret += TEXT("\t\treturn Greater;") HLSL_LINE_TERMINATOR;
+			Ret += TEXT("\telse") HLSL_LINE_TERMINATOR;
+			Ret += TEXT("\t\treturn Less;") HLSL_LINE_TERMINATOR;
+			Ret += TEXT("}") HLSL_LINE_TERMINATOR;
+			Ret += TEXT("") HLSL_LINE_TERMINATOR;
 		}
 
 		if (bIf2Enabled[Index] || IsDebugGenerateAllFunctionsEnabled())
 		{
-			Ret += FString::Printf(TEXT("%s IfDeriv(bool bNotEqual, bool bGreaterEqual, %s Greater, %s Less, %s Equal)"), *BaseName, *BaseName, *BaseName, *BaseName) + LINE_TERMINATOR;
-			Ret += TEXT("{") LINE_TERMINATOR;
-			Ret += TEXT("\tif(!bNotEqual)") LINE_TERMINATOR;	// Written like this to preserve NaN behavior of original code.
-			Ret += TEXT("\t\treturn Equal;") LINE_TERMINATOR;
-			Ret += TEXT("\tif(bGreaterEqual)") LINE_TERMINATOR;
-			Ret += TEXT("\t\treturn Greater;") LINE_TERMINATOR;
-			Ret += TEXT("\telse") LINE_TERMINATOR;
-			Ret += TEXT("\t\treturn Less;") LINE_TERMINATOR;
-			Ret += TEXT("}") LINE_TERMINATOR;
-			Ret += TEXT("") LINE_TERMINATOR;
+			Ret += FString::Printf(TEXT("%s IfDeriv(bool bNotEqual, bool bGreaterEqual, %s Greater, %s Less, %s Equal)"), *BaseName, *BaseName, *BaseName, *BaseName) + HLSL_LINE_TERMINATOR;
+			Ret += TEXT("{") HLSL_LINE_TERMINATOR;
+			Ret += TEXT("\tif(!bNotEqual)") HLSL_LINE_TERMINATOR;	// Written like this to preserve NaN behavior of original code.
+			Ret += TEXT("\t\treturn Equal;") HLSL_LINE_TERMINATOR;
+			Ret += TEXT("\tif(bGreaterEqual)") HLSL_LINE_TERMINATOR;
+			Ret += TEXT("\t\treturn Greater;") HLSL_LINE_TERMINATOR;
+			Ret += TEXT("\telse") HLSL_LINE_TERMINATOR;
+			Ret += TEXT("\t\treturn Less;") HLSL_LINE_TERMINATOR;
+			Ret += TEXT("}") HLSL_LINE_TERMINATOR;
+			Ret += TEXT("") HLSL_LINE_TERMINATOR;
 		}
 	}
 
@@ -2269,56 +2269,56 @@ FString FMaterialDerivativeAutogen::GenerateUsedFunctions(FHLSLMaterialTranslato
 	{
 		// float2(dot(InTexCoords, InRotationScale.xy), dot(InTexCoords, InRotationScale.zw)) + InOffset;
 		// InTexCoords.xy * InRotationScale.xw + InTexCoords.yx * InRotationScale.yz + InOffset;
-		Ret += TEXT("FloatDeriv2 RotateScaleOffsetTexCoordsDeriv(FloatDeriv2 TexCoord, FloatDeriv4 RotationScale, FloatDeriv2 Offset)") LINE_TERMINATOR;
-		Ret += TEXT("{") LINE_TERMINATOR;
-		Ret += TEXT("\tFloatDeriv2 Ret = Offset;") LINE_TERMINATOR;
-		Ret += TEXT("\tRet = AddDeriv(Ret, MulDeriv(TexCoord, SwizzleDeriv2(RotationScale, xw)));") LINE_TERMINATOR;
-		Ret += TEXT("\tRet = AddDeriv(Ret, MulDeriv(SwizzleDeriv2(TexCoord, yx), SwizzleDeriv2(RotationScale, yz)));") LINE_TERMINATOR;
-		Ret += TEXT("\treturn Ret;") LINE_TERMINATOR;
-		Ret += TEXT("}") LINE_TERMINATOR;
-		Ret += TEXT("") LINE_TERMINATOR;
+		Ret += TEXT("FloatDeriv2 RotateScaleOffsetTexCoordsDeriv(FloatDeriv2 TexCoord, FloatDeriv4 RotationScale, FloatDeriv2 Offset)") HLSL_LINE_TERMINATOR;
+		Ret += TEXT("{") HLSL_LINE_TERMINATOR;
+		Ret += TEXT("\tFloatDeriv2 Ret = Offset;") HLSL_LINE_TERMINATOR;
+		Ret += TEXT("\tRet = AddDeriv(Ret, MulDeriv(TexCoord, SwizzleDeriv2(RotationScale, xw)));") HLSL_LINE_TERMINATOR;
+		Ret += TEXT("\tRet = AddDeriv(Ret, MulDeriv(SwizzleDeriv2(TexCoord, yx), SwizzleDeriv2(RotationScale, yz)));") HLSL_LINE_TERMINATOR;
+		Ret += TEXT("\treturn Ret;") HLSL_LINE_TERMINATOR;
+		Ret += TEXT("}") HLSL_LINE_TERMINATOR;
+		Ret += TEXT("") HLSL_LINE_TERMINATOR;
 	}
 	
 	if (bUnMirrorEnabled[1][1] || IsDebugGenerateAllFunctionsEnabled())
 	{
 		// UnMirrorUV
-		Ret += TEXT("FloatDeriv2 UnMirrorUV(FloatDeriv2 UV, FMaterialPixelParameters Parameters)") LINE_TERMINATOR;
-		Ret += TEXT("{") LINE_TERMINATOR;
-		Ret += TEXT("\tconst MaterialFloat Scale = (Parameters.UnMirrored * 0.5f);") LINE_TERMINATOR;
-		Ret += TEXT("\tUV.Value = UV.Value * Scale + 0.5f;") LINE_TERMINATOR;
-		Ret += TEXT("\tUV.Ddx *= Scale;") LINE_TERMINATOR;
-		Ret += TEXT("\tUV.Ddy *= Scale;") LINE_TERMINATOR;
-		Ret += TEXT("\treturn UV;") LINE_TERMINATOR;
-		Ret += TEXT("}") LINE_TERMINATOR;
-		Ret += TEXT("") LINE_TERMINATOR;
+		Ret += TEXT("FloatDeriv2 UnMirrorUV(FloatDeriv2 UV, FMaterialPixelParameters Parameters)") HLSL_LINE_TERMINATOR;
+		Ret += TEXT("{") HLSL_LINE_TERMINATOR;
+		Ret += TEXT("\tconst MaterialFloat Scale = (Parameters.UnMirrored * 0.5f);") HLSL_LINE_TERMINATOR;
+		Ret += TEXT("\tUV.Value = UV.Value * Scale + 0.5f;") HLSL_LINE_TERMINATOR;
+		Ret += TEXT("\tUV.Ddx *= Scale;") HLSL_LINE_TERMINATOR;
+		Ret += TEXT("\tUV.Ddy *= Scale;") HLSL_LINE_TERMINATOR;
+		Ret += TEXT("\treturn UV;") HLSL_LINE_TERMINATOR;
+		Ret += TEXT("}") HLSL_LINE_TERMINATOR;
+		Ret += TEXT("") HLSL_LINE_TERMINATOR;
 	}
 	
 	if(bUnMirrorEnabled[1][0] || IsDebugGenerateAllFunctionsEnabled())
 	{
 		// UnMirrorU
-		Ret += TEXT("FloatDeriv2 UnMirrorU(FloatDeriv2 UV, FMaterialPixelParameters Parameters)") LINE_TERMINATOR;
-		Ret += TEXT("{") LINE_TERMINATOR;
-		Ret += TEXT("\tconst MaterialFloat Scale = (Parameters.UnMirrored * 0.5f);") LINE_TERMINATOR;
-		Ret += TEXT("\tUV.Value.x = UV.Value.x * Scale + 0.5f;") LINE_TERMINATOR;
-		Ret += TEXT("\tUV.Ddx.x *= Scale;") LINE_TERMINATOR;
-		Ret += TEXT("\tUV.Ddy.x *= Scale;") LINE_TERMINATOR;
-		Ret += TEXT("\treturn UV;") LINE_TERMINATOR;
-		Ret += TEXT("}") LINE_TERMINATOR;
-		Ret += TEXT("") LINE_TERMINATOR;
+		Ret += TEXT("FloatDeriv2 UnMirrorU(FloatDeriv2 UV, FMaterialPixelParameters Parameters)") HLSL_LINE_TERMINATOR;
+		Ret += TEXT("{") HLSL_LINE_TERMINATOR;
+		Ret += TEXT("\tconst MaterialFloat Scale = (Parameters.UnMirrored * 0.5f);") HLSL_LINE_TERMINATOR;
+		Ret += TEXT("\tUV.Value.x = UV.Value.x * Scale + 0.5f;") HLSL_LINE_TERMINATOR;
+		Ret += TEXT("\tUV.Ddx.x *= Scale;") HLSL_LINE_TERMINATOR;
+		Ret += TEXT("\tUV.Ddy.x *= Scale;") HLSL_LINE_TERMINATOR;
+		Ret += TEXT("\treturn UV;") HLSL_LINE_TERMINATOR;
+		Ret += TEXT("}") HLSL_LINE_TERMINATOR;
+		Ret += TEXT("") HLSL_LINE_TERMINATOR;
 	}
 	
 	if (bUnMirrorEnabled[0][1] || IsDebugGenerateAllFunctionsEnabled())
 	{
 		// UnMirrorV
-		Ret += TEXT("FloatDeriv2 UnMirrorV(FloatDeriv2 UV, FMaterialPixelParameters Parameters)") LINE_TERMINATOR;
-		Ret += TEXT("{") LINE_TERMINATOR;
-		Ret += TEXT("\tconst MaterialFloat Scale = (Parameters.UnMirrored * 0.5f);") LINE_TERMINATOR;
-		Ret += TEXT("\tUV.Value.y = UV.Value.y * Scale + 0.5f;") LINE_TERMINATOR;
-		Ret += TEXT("\tUV.Ddx.y *= Scale;") LINE_TERMINATOR;
-		Ret += TEXT("\tUV.Ddy.y *= Scale;") LINE_TERMINATOR;
-		Ret += TEXT("\treturn UV;") LINE_TERMINATOR;
-		Ret += TEXT("}") LINE_TERMINATOR;
-		Ret += TEXT("") LINE_TERMINATOR;
+		Ret += TEXT("FloatDeriv2 UnMirrorV(FloatDeriv2 UV, FMaterialPixelParameters Parameters)") HLSL_LINE_TERMINATOR;
+		Ret += TEXT("{") HLSL_LINE_TERMINATOR;
+		Ret += TEXT("\tconst MaterialFloat Scale = (Parameters.UnMirrored * 0.5f);") HLSL_LINE_TERMINATOR;
+		Ret += TEXT("\tUV.Value.y = UV.Value.y * Scale + 0.5f;") HLSL_LINE_TERMINATOR;
+		Ret += TEXT("\tUV.Ddx.y *= Scale;") HLSL_LINE_TERMINATOR;
+		Ret += TEXT("\tUV.Ddy.y *= Scale;") HLSL_LINE_TERMINATOR;
+		Ret += TEXT("\treturn UV;") HLSL_LINE_TERMINATOR;
+		Ret += TEXT("}") HLSL_LINE_TERMINATOR;
+		Ret += TEXT("") HLSL_LINE_TERMINATOR;
 	}
 
 	return Ret;
diff --git a/Engine/Source/Runtime/Engine/Private/Materials/HLSLMaterialTranslator.cpp b/Engine/Source/Runtime/Engine/Private/Materials/HLSLMaterialTranslator.cpp
--- a/Engine/Source/Runtime/Engine/Private/Materials/HLSLMaterialTranslator.cpp
+++ b/Engine/Source/Runtime/Engine/Private/Materials/HLSLMaterialTranslator.cpp
@@ -676,7 +676,7 @@ void FHLSLMaterialTranslator::CompileCustomOutputs(TArray<UMaterialExpressionCus
 			{
 				if (CustomOutput->NeedsCustomOutputDefines())
 				{
-					ResourcesString += FString::Printf(TEXT("#define NUM_MATERIAL_OUTPUTS_%s %d\r\n"), *CustomOutput->GetFunctionName().ToUpper(), NumOutputs);
+					ResourcesString += FString::Printf(TEXT("#define NUM_MATERIAL_OUTPUTS_%s %d\n"), *CustomOutput->GetFunctionName().ToUpper(), NumOutputs);
 				}
 
 				if (NumOutputs > 0)
@@ -1259,7 +1259,7 @@ bool FHLSLMaterialTranslator::Translate()
 				// If used, add compile data
 				if (bValidResultCompiled)
 				{
-					ResourcesString += FString::Printf(TEXT("#define NUM_MATERIAL_OUTPUTS_%s %d\r\n"), *Attribute.FunctionName.ToUpper(), NumOutputs);
+					ResourcesString += FString::Printf(TEXT("#define NUM_MATERIAL_OUTPUTS_%s %d\n"), *Attribute.FunctionName.ToUpper(), NumOutputs);
 				}
 			}
 		}
@@ -1289,7 +1289,7 @@ bool FHLSLMaterialTranslator::Translate()
 		// Output the implementation for any custom expressions we will call below.
 		for (int32 ExpressionIndex = 0; ExpressionIndex < CustomExpressions.Num(); ExpressionIndex++)
 		{
-			ResourcesString += CustomExpressions[ExpressionIndex].Implementation + "\r\n\r\n";
+			ResourcesString += CustomExpressions[ExpressionIndex].Implementation + "\n\n";
 		}
 
 		// Translation is designed to have a code chunk generation phase followed by several passes that only has readonly access to the code chunks.
@@ -1722,7 +1722,7 @@ bool FHLSLMaterialTranslator::Translate()
 		// Output the implementation for any custom output expressions
 				for (int32 ExpressionIndex = 0; ExpressionIndex < DerivativeVariations[Variation].CustomOutputImplementations.Num(); ExpressionIndex++)
 		{
-					ResourcesString += DerivativeVariations[Variation].CustomOutputImplementations[ExpressionIndex] + "\r\n\r\n";
+					ResourcesString += DerivativeVariations[Variation].CustomOutputImplementations[ExpressionIndex] + "\n\n";
 				}
 			}
 		}
@@ -1955,7 +1955,7 @@ bool FHLSLMaterialTranslator::Translate()
 			else
 			{
 				MaterialCompilationOutput.SubstrateMaterialCompilationOutput.SubstrateMaterialDescription = "";
-				ResourcesString += "// No Substrate material provided\r\n";
+				ResourcesString += "// No Substrate material provided\n";
 
 				// Adde default Substrate functions
 				ResourcesString += "#if TEMPLATE_USES_SUBSTRATE\n";
@@ -2424,23 +2424,23 @@ void FHLSLMaterialTranslator::GetMaterialEnvironment(EShaderPlatform InPlatform,
 					return TEXT("ERROR");
 				};
 				FString SubstrateCompilationContextName = GetSubstrateCompilationContextName(SubstrateCompilationContextIndex);
-				SubstrateMaterialContextDescription += FString::Printf(TEXT("----- SUBSTRATE - %s -----\r\n"), *SubstrateCompilationContextName);
-				SubstrateMaterialContextDescription += FString::Printf(TEXT("SubstrateCompilationInfo -\r\n"));
-				SubstrateMaterialContextDescription += FString::Printf(TEXT(" - Byte Per Pixel Budget                           %u\r\n"), SubstrateBytePerPixel_Platform);
-				SubstrateMaterialContextDescription += FString::Printf(TEXT(" - Closure Per Pixel Budget                        %u\r\n"), SubstrateClosurePerPixel_Platform);
-				SubstrateMaterialContextDescription += FString::Printf(TEXT(" - Requested Byte Size before simplification       %u (%d UINT32)\r\n"), SubstrateSimplificationStatus.OriginalRequestedByteSize, SubstrateSimplificationStatus.OriginalRequestedByteSize / 4);
-				SubstrateMaterialContextDescription += FString::Printf(TEXT(" - Requested Byte Size after simplification        %u (%d UINT32)\r\n"), SubstrateCtx.SubstrateMaterialRequestedSizeByte, SubstrateCtx.SubstrateMaterialRequestedSizeByte / 4);
-				SubstrateMaterialContextDescription += FString::Printf(TEXT(" - Requested Closure Count before simplification   %u\r\n"), SubstrateSimplificationStatus.OriginalRequestedClosureCount);
-				SubstrateMaterialContextDescription += FString::Printf(TEXT(" - Requested Closure Count after simplification    %u\r\n"), SubstrateCtx.SubstrateMaterialClosureCount);
-				SubstrateMaterialContextDescription += FString::Printf(TEXT(" - Material complexity                             %s\r\n"), *FSubstrateMaterialComplexity::ToString(SubstrateCtx.SubstrateMaterialComplexity.SubstrateMaterialType(), true /* Upper case */));
-				SubstrateMaterialContextDescription += FString::Printf(TEXT(" - BSDF Count                                      %i\r\n"), SubstrateCtx.SubstrateMaterialEffectiveClosureCount); // REMOVE?
+				SubstrateMaterialContextDescription += FString::Printf(TEXT("----- SUBSTRATE - %s -----\n"), *SubstrateCompilationContextName);
+				SubstrateMaterialContextDescription += FString::Printf(TEXT("SubstrateCompilationInfo -\n"));
+				SubstrateMaterialContextDescription += FString::Printf(TEXT(" - Byte Per Pixel Budget                           %u\n"), SubstrateBytePerPixel_Platform);
+				SubstrateMaterialContextDescription += FString::Printf(TEXT(" - Closure Per Pixel Budget                        %u\n"), SubstrateClosurePerPixel_Platform);
+				SubstrateMaterialContextDescription += FString::Printf(TEXT(" - Requested Byte Size before simplification       %u (%d UINT32)\n"), SubstrateSimplificationStatus.OriginalRequestedByteSize, SubstrateSimplificationStatus.OriginalRequestedByteSize / 4);
+				SubstrateMaterialContextDescription += FString::Printf(TEXT(" - Requested Byte Size after simplification        %u (%d UINT32)\n"), SubstrateCtx.SubstrateMaterialRequestedSizeByte, SubstrateCtx.SubstrateMaterialRequestedSizeByte / 4);
+				SubstrateMaterialContextDescription += FString::Printf(TEXT(" - Requested Closure Count before simplification   %u\n"), SubstrateSimplificationStatus.OriginalRequestedClosureCount);
+				SubstrateMaterialContextDescription += FString::Printf(TEXT(" - Requested Closure Count after simplification    %u\n"), SubstrateCtx.SubstrateMaterialClosureCount);
+				SubstrateMaterialContextDescription += FString::Printf(TEXT(" - Material complexity                             %s\n"), *FSubstrateMaterialComplexity::ToString(SubstrateCtx.SubstrateMaterialComplexity.SubstrateMaterialType(), true /* Upper case */));
+				SubstrateMaterialContextDescription += FString::Printf(TEXT(" - BSDF Count                                      %i\n"), SubstrateCtx.SubstrateMaterialEffectiveClosureCount); // REMOVE?
 				if (RequestedSharedLocalBasesCount > SUBSTRATE_MAX_SHAREDLOCALBASES_REGISTERS)
 				{
-					SubstrateMaterialDescription += FString::Printf(TEXT(" - SharedLocalBasesCount                      %i (Requested:%i)\r\n"), SubstrateCtx.FinalUsedSharedLocalBasesCount, RequestedSharedLocalBasesCount);
+					SubstrateMaterialDescription += FString::Printf(TEXT(" - SharedLocalBasesCount                      %i (Requested:%i)\n"), SubstrateCtx.FinalUsedSharedLocalBasesCount, RequestedSharedLocalBasesCount);
 				}
 				else
 				{
-					SubstrateMaterialDescription += FString::Printf(TEXT(" - SharedLocalBasesCount                      %i\r\n"), SubstrateCtx.FinalUsedSharedLocalBasesCount);
+					SubstrateMaterialDescription += FString::Printf(TEXT(" - SharedLocalBasesCount                      %i\n"), SubstrateCtx.FinalUsedSharedLocalBasesCount);
 				}
 
 
@@ -2457,22 +2457,22 @@ void FHLSLMaterialTranslator::GetMaterialEnvironment(EShaderPlatform InPlatform,
 					}
 					const FSubstrateSharedLocalBasesInfo& SubstrateSharedLocalBasesInfo = SubstrateCtx.SubstrateCompilationInfoGetMatchingSharedLocalBasisInfo(BSDFOperator.BSDFRegisteredSharedLocalBasis);
 
-					SubstrateMaterialContextDescription += FString::Printf(TEXT("     - %s - SharedLocalBasisIndexMacro = %s \r\n"), *GetSubstrateBSDFName(BSDFOperator.BSDFType), *GetSubstrateSharedLocalBasisIndexMacroInner(BSDFOperator.BSDFRegisteredSharedLocalBasis, SubstrateCompilationContextIndex));
+					SubstrateMaterialContextDescription += FString::Printf(TEXT("     - %s - SharedLocalBasisIndexMacro = %s \n"), *GetSubstrateBSDFName(BSDFOperator.BSDFType), *GetSubstrateSharedLocalBasisIndexMacroInner(BSDFOperator.BSDFRegisteredSharedLocalBasis, SubstrateCompilationContextIndex));
 				}
 
-				SubstrateMaterialContextDescription += FString::Printf(TEXT("----------- SUBSTRATE TREE - %s -----------\r\n"), *SubstrateCompilationContextName);
-				SubstrateMaterialContextDescription += FString::Printf(TEXT("Graph maximum distance to leaves %u\r\n"), RootMaximumDistanceToLeaves);
+				SubstrateMaterialContextDescription += FString::Printf(TEXT("----------- SUBSTRATE TREE - %s -----------\n"), *SubstrateCompilationContextName);
+				SubstrateMaterialContextDescription += FString::Printf(TEXT("Graph maximum distance to leaves %u\n"), RootMaximumDistanceToLeaves);
 				// Debug print operators according to depth from root.
 				{
 
 					for (int32 DistanceToLeaves = RootMaximumDistanceToLeaves; DistanceToLeaves >= 0; --DistanceToLeaves)
 					{
-						SubstrateMaterialContextDescription += FString::Printf(TEXT("----- DistanceFromLeaves = %d -----\r\n"), DistanceToLeaves);
+						SubstrateMaterialContextDescription += FString::Printf(TEXT("----- DistanceFromLeaves = %d -----\n"), DistanceToLeaves);
 						for (auto& It : SubstrateCtx.SubstrateMaterialExpressionRegisteredOperators)
 						{
 							if (!It.IsDiscarded() && It.MaxDistanceFromLeaves == DistanceToLeaves)
 							{
-								SubstrateMaterialContextDescription += FString::Printf(TEXT("\tIdx=%d Op=%s ParentIdx=%d LeftIndex=%d RightIndex=%d BSDFIdx=%d LayerDepth=%d IsTop=%d IsBot=%d BSDFType=%s SSS=%d MFP=%d F90=%d Rough2=%d Fuzz=%d Aniso=%d Glint=%d SpecularProfile=%d\r\n"),
+								SubstrateMaterialContextDescription += FString::Printf(TEXT("\tIdx=%d Op=%s ParentIdx=%d LeftIndex=%d RightIndex=%d BSDFIdx=%d LayerDepth=%d IsTop=%d IsBot=%d BSDFType=%s SSS=%d MFP=%d F90=%d Rough2=%d Fuzz=%d Aniso=%d Glint=%d SpecularProfile=%d\n"),
 									It.Index, GetSubstrateOperatorStr(It.OperatorType), It.ParentIndex, It.LeftIndex, It.RightIndex, It.BSDFIndex, It.LayerDepth, It.bIsTop, It.bIsBottom,
 									*GetSubstrateBSDFName(It.BSDFType), It.bBSDFHasSSS, It.bBSDFHasMFPPluggedIn, It.bBSDFHasEdgeColor, It.bBSDFHasSecondRoughnessOrSimpleClearCoat, It.bBSDFHasFuzz, It.bBSDFHasAnisotropy, It.bBSDFHasGlint, It.bBSDFHasSpecularProfile);
 							}
@@ -2543,25 +2543,25 @@ TBitArray<> FHLSLMaterialTranslator::GetVertexInterpolatorsOffsets(FString& Vert
 
 		const EMaterialValueType Type = Interpolator->InterpolatedType == MCT_Float ? MCT_Float1 : Interpolator->InterpolatedType;
 
-		VertexInterpolatorsOffsetsDefinitionCode += LINE_TERMINATOR;
-		VertexInterpolatorsOffsetsDefinitionCode += FString::Printf(TEXT("#define VERTEX_INTERPOLATOR_%i_TEXCOORDS_X\t%i") LINE_TERMINATOR, Index, GetNextUVSlot());
+		VertexInterpolatorsOffsetsDefinitionCode += HLSL_LINE_TERMINATOR;
+		VertexInterpolatorsOffsetsDefinitionCode += FString::Printf(TEXT("#define VERTEX_INTERPOLATOR_%i_TEXCOORDS_X\t%i") HLSL_LINE_TERMINATOR, Index, GetNextUVSlot());
 
 		if (Type >= MCT_Float2)
 		{
-			VertexInterpolatorsOffsetsDefinitionCode += FString::Printf(TEXT("#define VERTEX_INTERPOLATOR_%i_TEXCOORDS_Y\t%i") LINE_TERMINATOR, Index, GetNextUVSlot());
+			VertexInterpolatorsOffsetsDefinitionCode += FString::Printf(TEXT("#define VERTEX_INTERPOLATOR_%i_TEXCOORDS_Y\t%i") HLSL_LINE_TERMINATOR, Index, GetNextUVSlot());
 
 			if (Type >= MCT_Float3)
 			{
-				VertexInterpolatorsOffsetsDefinitionCode += FString::Printf(TEXT("#define VERTEX_INTERPOLATOR_%i_TEXCOORDS_Z\t%i") LINE_TERMINATOR, Index, GetNextUVSlot());
+				VertexInterpolatorsOffsetsDefinitionCode += FString::Printf(TEXT("#define VERTEX_INTERPOLATOR_%i_TEXCOORDS_Z\t%i") HLSL_LINE_TERMINATOR, Index, GetNextUVSlot());
 
 				if (Type == MCT_Float4)
 				{
-					VertexInterpolatorsOffsetsDefinitionCode += FString::Printf(TEXT("#define VERTEX_INTERPOLATOR_%i_TEXCOORDS_W\t%i") LINE_TERMINATOR, Index, GetNextUVSlot());
+					VertexInterpolatorsOffsetsDefinitionCode += FString::Printf(TEXT("#define VERTEX_INTERPOLATOR_%i_TEXCOORDS_W\t%i") HLSL_LINE_TERMINATOR, Index, GetNextUVSlot());
 				}
 			}
 		}
 			
-		VertexInterpolatorsOffsetsDefinitionCode += LINE_TERMINATOR;
+		VertexInterpolatorsOffsetsDefinitionCode += HLSL_LINE_TERMINATOR;
 	}
 
 	return AllocatedCoords;
@@ -2705,10 +2705,10 @@ FString FHLSLMaterialTranslator::GetMaterialShaderCode()
 		{
 			const FVector4f DefaultValue = FMaterialAttributeDefinitionMap::GetDefaultValue(AttributeID);
 
-			MaterialAttributesDeclaration.Appendf(TEXT("\t%s %s;") LINE_TERMINATOR, HLSLType, *PropertyName);
+			MaterialAttributesDeclaration.Appendf(TEXT("\t%s %s;") HLSL_LINE_TERMINATOR, HLSLType, *PropertyName);
 
 			// Chainable method to set the attribute
-			MaterialAttributesUtilities.Appendf(TEXT("FMaterialAttributes FMaterialAttributes_Set%s(FMaterialAttributes InAttributes, %s InValue) { InAttributes.%s = InValue; return InAttributes; }") LINE_TERMINATOR,
+			MaterialAttributesUtilities.Appendf(TEXT("FMaterialAttributes FMaterialAttributes_Set%s(FMaterialAttributes InAttributes, %s InValue) { InAttributes.%s = InValue; return InAttributes; }") HLSL_LINE_TERMINATOR,
 				*PropertyName, HLSLType, *PropertyName);
 		}
 	}
@@ -2791,7 +2791,7 @@ FString FHLSLMaterialTranslator::GetMaterialShaderCode()
 		if (bEnableExecutionFlow)
 		{
 			const FString AttributeName = FMaterialAttributeDefinitionMap::GetAttributeName((EMaterialProperty)(MP_CustomizedUVs0 + CustomUVIndex));
-			CustomUVAssignments.Appendf(TEXT("\tOutTexCoords[%u] = Parameters.MaterialAttributes.%s;") LINE_TERMINATOR, CustomUVIndex, *AttributeName);
+			CustomUVAssignments.Appendf(TEXT("\tOutTexCoords[%u] = Parameters.MaterialAttributes.%s;") HLSL_LINE_TERMINATOR, CustomUVIndex, *AttributeName);
 		}
 		else
 		{
@@ -2805,7 +2805,7 @@ FString FHLSLMaterialTranslator::GetMaterialShaderCode()
 				LastProperty = MP_CustomizedUVs0 + CustomUVIndex;
 			}
 
-			CustomUVAssignments.Appendf(TEXT("\tOutTexCoords[%u] = %s;") LINE_TERMINATOR, CustomUVIndex, *DerivativeVariations[BaseDerivativeVariation].TranslatedCodeChunks[MP_CustomizedUVs0 + CustomUVIndex]);
+			CustomUVAssignments.Appendf(TEXT("\tOutTexCoords[%u] = %s;") HLSL_LINE_TERMINATOR, CustomUVIndex, *DerivativeVariations[BaseDerivativeVariation].TranslatedCodeChunks[MP_CustomizedUVs0 + CustomUVIndex]);
 		}
 	}
 
@@ -2827,19 +2827,19 @@ FString FHLSLMaterialTranslator::GetMaterialShaderCode()
 			const int32 Index = Interpolator->InterpolatorIndex;
 
 			// Note: We reference the UV define directly to avoid having to pre-accumulate UV counts before property translation
-			CustomInterpolatorAssignments.Appendf(TEXT("\tOutTexCoords[VERTEX_INTERPOLATOR_%i_TEXCOORDS_X].%s = VertexInterpolator%i(Parameters).x;") LINE_TERMINATOR, Index, Swizzle[Offset%2], Index);
+			CustomInterpolatorAssignments.Appendf(TEXT("\tOutTexCoords[VERTEX_INTERPOLATOR_%i_TEXCOORDS_X].%s = VertexInterpolator%i(Parameters).x;") HLSL_LINE_TERMINATOR, Index, Swizzle[Offset%2], Index);
 				
 			if (Type >= MCT_Float2)
 			{
-				CustomInterpolatorAssignments.Appendf(TEXT("\tOutTexCoords[VERTEX_INTERPOLATOR_%i_TEXCOORDS_Y].%s = VertexInterpolator%i(Parameters).y;") LINE_TERMINATOR, Index, Swizzle[(Offset+1)%2], Index);
+				CustomInterpolatorAssignments.Appendf(TEXT("\tOutTexCoords[VERTEX_INTERPOLATOR_%i_TEXCOORDS_Y].%s = VertexInterpolator%i(Parameters).y;") HLSL_LINE_TERMINATOR, Index, Swizzle[(Offset+1)%2], Index);
 
 				if (Type >= MCT_Float3)
 				{
-					CustomInterpolatorAssignments.Appendf(TEXT("\tOutTexCoords[VERTEX_INTERPOLATOR_%i_TEXCOORDS_Z].%s = VertexInterpolator%i(Parameters).z;") LINE_TERMINATOR, Index, Swizzle[(Offset+2)%2], Index);
+					CustomInterpolatorAssignments.Appendf(TEXT("\tOutTexCoords[VERTEX_INTERPOLATOR_%i_TEXCOORDS_Z].%s = VertexInterpolator%i(Parameters).z;") HLSL_LINE_TERMINATOR, Index, Swizzle[(Offset+2)%2], Index);
 
 					if (Type == MCT_Float4)
 					{
-						CustomInterpolatorAssignments.Appendf(TEXT("\tOutTexCoords[VERTEX_INTERPOLATOR_%i_TEXCOORDS_W].%s = VertexInterpolator%i(Parameters).w;") LINE_TERMINATOR, Index, Swizzle[(Offset+3)%2], Index);
+						CustomInterpolatorAssignments.Appendf(TEXT("\tOutTexCoords[VERTEX_INTERPOLATOR_%i_TEXCOORDS_W].%s = VertexInterpolator%i(Parameters).w;") HLSL_LINE_TERMINATOR, Index, Swizzle[(Offset+3)%2], Index);
 					}
 				}
 			}
@@ -2857,14 +2857,14 @@ FString FHLSLMaterialTranslator::GetMaterialShaderCode()
 		{
 			const FString AttributeName = FMaterialAttributeDefinitionMap::GetAttributeName((EMaterialProperty)(MP_CustomizedUVs0 + TexCoordIndex));
 
-			EvaluateVertexCode.Appendf(TEXT("\tDefaultMaterialAttributes.%s = Parameters.TexCoords[%d];") LINE_TERMINATOR, *AttributeName, TexCoordIndex);
+			EvaluateVertexCode.Appendf(TEXT("\tDefaultMaterialAttributes.%s = Parameters.TexCoords[%d];") HLSL_LINE_TERMINATOR, *AttributeName, TexCoordIndex);
 		}
 
 		EvaluateVertexCode.Append(TranslatedAttributesCodeChunks[SF_Vertex]);
 
 		MaterialSourceTemplateParams.Add({ TEXT("evaluate_material_attributes"), TranslatedAttributesCodeChunks[SF_Pixel] });
 
-		FString EvaluateMaterialAttributesCode = TEXT("    FMaterialAttributes MaterialAttributes = EvaluatePixelMaterialAttributes(Parameters);" LINE_TERMINATOR);
+		FString EvaluateMaterialAttributesCode = TEXT("    FMaterialAttributes MaterialAttributes = EvaluatePixelMaterialAttributes(Parameters);" HLSL_LINE_TERMINATOR);
 
 		for (int32 PropertyIndex = 0; PropertyIndex < MP_MAX; ++PropertyIndex)
 		{
@@ -2882,11 +2882,11 @@ FString FHLSLMaterialTranslator::GetMaterialShaderCode()
 			if (PropertyIndex == MP_SubsurfaceColor)
 			{
 				// TODO - properly handle subsurface profile
-				EvaluateMaterialAttributesCode.Appendf("    PixelMaterialInputs.Subsurface = float4(MaterialAttributes.%s, 0.0f);" LINE_TERMINATOR, *PropertyName);
+				EvaluateMaterialAttributesCode.Appendf("    PixelMaterialInputs.Subsurface = float4(MaterialAttributes.%s, 0.0f);" HLSL_LINE_TERMINATOR, *PropertyName);
 			}
 			else
 			{
-				EvaluateMaterialAttributesCode.Appendf("    PixelMaterialInputs.%s = MaterialAttributes.%s;" LINE_TERMINATOR, *PropertyName, *PropertyName);
+				EvaluateMaterialAttributesCode.Appendf("    PixelMaterialInputs.%s = MaterialAttributes.%s;" HLSL_LINE_TERMINATOR, *PropertyName, *PropertyName);
 			}
 		}
 
@@ -3475,7 +3475,7 @@ int32 FHLSLMaterialTranslator::AddCodeChunkInner(uint64 Hash, const TCHAR* Forma
 	if (Type == MCT_VoidStatement)
 	{
 		CodeIndex = CurrentScopeChunks->Num();
-		const FString Statement = FString("") + FormattedCode + LINE_TERMINATOR;
+		const FString Statement = FString("") + FormattedCode + HLSL_LINE_TERMINATOR;
 		new(*CurrentScopeChunks) FShaderCodeChunk(Hash, *Statement, *Statement, TEXT(""), Type, DerivativeStatus, true);
 	}
 	else if (bInlined)
@@ -3503,9 +3503,9 @@ int32 FHLSLMaterialTranslator::AddCodeChunkInner(uint64 Hash, const TCHAR* Forma
 			// Allocate a local variable name
 			const FString SymbolName = CreateSymbolName(TEXT("Local"));
 			// Construct the definition string which stores the result in a temporary and adds a newline for readability
-			const FString LocalVariableDefinitionFinite = FString("	") + HLSLTypeString(Type) + TEXT(" ") + SymbolName + TEXT(" = ") + FormattedCode + TEXT(";") + LINE_TERMINATOR;
+			const FString LocalVariableDefinitionFinite = FString("	") + HLSLTypeString(Type) + TEXT(" ") + SymbolName + TEXT(" = ") + FormattedCode + TEXT(";") + HLSL_LINE_TERMINATOR;
 			// Construct the definition string which stores the result in a temporary and adds a newline for readability
-			const FString LocalVariableDefinitionAnalytic = FString("	") + HLSLTypeString(Type) + TEXT(" ") + SymbolName + TEXT(" = ") + FormattedCode + TEXT(";") + LINE_TERMINATOR;
+			const FString LocalVariableDefinitionAnalytic = FString("	") + HLSLTypeString(Type) + TEXT(" ") + SymbolName + TEXT(" = ") + FormattedCode + TEXT(";") + HLSL_LINE_TERMINATOR;
 			// Adding a code chunk that creates a local variable
 			new(*CurrentScopeChunks) FShaderCodeChunk(Hash, *LocalVariableDefinitionFinite, *LocalVariableDefinitionAnalytic, SymbolName, Type, DerivativeStatus, false);
 		}
@@ -3583,9 +3583,9 @@ int32 FHLSLMaterialTranslator::AddCodeChunkInnerDeriv(const TCHAR* FormattedCode
 			// Allocate a local variable name
 			const FString SymbolName = CreateSymbolName(TEXT("Local"));
 			// Construct the definition string which stores the result in a temporary and adds a newline for readability
-			const FString LocalVariableDefinitionFinite = FString("	") + HLSLTypeString(Type) + TEXT(" ") + SymbolName + TEXT(" = ") + FormattedCodeFinite + TEXT(";") + LINE_TERMINATOR;
+			const FString LocalVariableDefinitionFinite = FString("	") + HLSLTypeString(Type) + TEXT(" ") + SymbolName + TEXT(" = ") + FormattedCodeFinite + TEXT(";") + HLSL_LINE_TERMINATOR;
 			// Analytic version too
-			const FString LocalVariableDefinitionAnalytic = FString("	") + (bEmitInvalidDerivToken ? "$" : "") + HLSLTypeStringDeriv(Type, DerivativeStatus) + TEXT(" ") + SymbolName + TEXT(" = ") + FormattedCodeAnalytic + TEXT(";") + LINE_TERMINATOR;
+			const FString LocalVariableDefinitionAnalytic = FString("	") + (bEmitInvalidDerivToken ? "$" : "") + HLSLTypeStringDeriv(Type, DerivativeStatus) + TEXT(" ") + SymbolName + TEXT(" = ") + FormattedCodeAnalytic + TEXT(";") + HLSL_LINE_TERMINATOR;
 			// Adding a code chunk that creates a local variable
 			new(*CurrentScopeChunks) FShaderCodeChunk(Hash, *LocalVariableDefinitionFinite, *LocalVariableDefinitionAnalytic, SymbolName, Type, DerivativeStatus, false);
 		}
@@ -13585,7 +13585,7 @@ int32 FHLSLMaterialTranslator::CustomExpression( class UMaterialExpressionCustom
 		{
 			Code = FString(TEXT("return ")) + Code + TEXT(";");
 		}
-		Code.ReplaceInline(TEXT("\n"), TEXT("\r\n"), ESearchCase::CaseSensitive);
+		Code.ReplaceInline(TEXT("\r\n"), TEXT("\n"), ESearchCase::CaseSensitive);
 
 		FString ParametersType = ShaderFrequency == SF_Vertex ? TEXT("Vertex") : TEXT("Pixel");
 
@@ -13602,9 +13602,9 @@ int32 FHLSLMaterialTranslator::CustomExpression( class UMaterialExpressionCustom
 
 		for (FCustomDefine DefineEntry : Custom->AdditionalDefines)
 		{
-			FString DefineStatement = TEXT("#ifndef ") + DefineEntry.DefineName + LINE_TERMINATOR;
-			DefineStatement += TEXT("#define ") + DefineEntry.DefineName + TEXT(" ") + DefineEntry.DefineValue + LINE_TERMINATOR;
-			DefineStatement += TEXT("#endif//") + DefineEntry.DefineName + LINE_TERMINATOR;
+			FString DefineStatement = TEXT("#ifndef ") + DefineEntry.DefineName + HLSL_LINE_TERMINATOR;
+			DefineStatement += TEXT("#define ") + DefineEntry.DefineName + TEXT(" ") + DefineEntry.DefineValue + HLSL_LINE_TERMINATOR;
+			DefineStatement += TEXT("#endif//") + DefineEntry.DefineName + HLSL_LINE_TERMINATOR;
 
 			Entry.Implementation += DefineStatement;
 		}
@@ -13615,12 +13615,12 @@ int32 FHLSLMaterialTranslator::CustomExpression( class UMaterialExpressionCustom
 			IncludeStatement += TEXT("\"");
 			IncludeStatement += IncludeFile;
 			IncludeStatement += TEXT("\"");
-			IncludeStatement += LINE_TERMINATOR;
+			IncludeStatement += HLSL_LINE_TERMINATOR;
 
 			Entry.Implementation += IncludeStatement;
 		}
 
-		Entry.Implementation += FString::Printf(TEXT("%s CustomExpression%d(FMaterial%sParameters Parameters%s)\r\n{\r\n%s\r\n}\r\n"), *OutputTypeString, CustomExpressionIndex, *ParametersType, *InputParamDecl, *Code);
+		Entry.Implementation += FString::Printf(TEXT("%s CustomExpression%d(FMaterial%sParameters Parameters%s)\n{\n%s\n}\n"), *OutputTypeString, CustomExpressionIndex, *ParametersType, *InputParamDecl, *Code);
 		const uint64 ImplementationHash = CityHash64((char*)*Entry.Implementation, Entry.Implementation.Len() * sizeof(TCHAR));
 
 		Entry.OutputCodeIndex.Empty(Custom->AdditionalOutputs.Num() + 1);
@@ -13750,7 +13750,7 @@ int32 FHLSLMaterialTranslator::CustomOutput(class UMaterialExpressionCustomOutpu
 	// Primary function will have _LWC suffix if it returns an LWC type
 	// We also define a pre-processor symbol to indicate the custom output function is available so that
 	// shaders can implement some fallback behavior for cases where the custom output has not been added
-	FString ImplementationCodeFinite = FString::Printf(TEXT("#define HAVE_%s 1\r\n%s %s%s(%s FMaterial%sParameters Parameters)\r\n{\r\n%s return %s;\r\n}\r\n"),
+	FString ImplementationCodeFinite = FString::Printf(TEXT("#define HAVE_%s 1\n%s %s%s(%s FMaterial%sParameters Parameters)\n{\n%s return %s;\n}\n"),
 		*FunctionNameBase,
 		*OutputTypeString,
 		*FunctionNameBase,
@@ -13763,7 +13763,7 @@ int32 FHLSLMaterialTranslator::CustomOutput(class UMaterialExpressionCustomOutpu
 	{
 		// Add a wrapper with no suffix to return a non-LWC type
 		AddLWCFuncUsage(ELWCFunctionKind::Demote);
-		ImplementationCodeFinite += FString::Printf(TEXT("%s %s(%s FMaterial%sParameters Parameters) { return LWCToFloat(%s_LWC(Parameters)); }\r\n"),
+		ImplementationCodeFinite += FString::Printf(TEXT("%s %s(%s FMaterial%sParameters Parameters) { return LWCToFloat(%s_LWC(Parameters)); }\n"),
 			HLSLTypeString(MakeNonLWCType(OutputType)),
 			*FunctionNameBase,
 			ShaderFrequency == SF_Vertex ? TEXT("const") : TEXT("inout"),
@@ -13774,7 +13774,7 @@ int32 FHLSLMaterialTranslator::CustomOutput(class UMaterialExpressionCustomOutpu
 	{
 		// Add a wrapper with LWC suffix to return a LWC type
 		AddLWCFuncUsage(ELWCFunctionKind::Promote);
-		ImplementationCodeFinite += FString::Printf(TEXT("%s %s_LWC(%s FMaterial%sParameters Parameters) { return LWCPromote(%s(Parameters)); }\r\n"),
+		ImplementationCodeFinite += FString::Printf(TEXT("%s %s_LWC(%s FMaterial%sParameters Parameters) { return LWCPromote(%s(Parameters)); }\n"),
 			HLSLTypeString(MakeLWCType(OutputType)),
 			*FunctionNameBase,
 			ShaderFrequency == SF_Vertex ? TEXT("const") : TEXT("inout"),
@@ -13785,7 +13785,7 @@ int32 FHLSLMaterialTranslator::CustomOutput(class UMaterialExpressionCustomOutpu
 	DerivativeVariations[CompiledPDV_FiniteDifferences].CustomOutputImplementations.Add(ImplementationCodeFinite);
 
 	// FIXEM: for now just skipping these, as they create invalid code.
-	//FString ImplementationCodeAnalytic = FString::Printf(TEXT("%s %s%d_Analytic(FMaterial%sParameters Parameters)\r\n{\r\n%s return %s;\r\n}\r\n"), *OutputTypeString, *Custom->GetFunctionName(), OutputIndex, ShaderFrequency == SF_Vertex ? TEXT("Vertex") : TEXT("Pixel"), *Definitions, *Body);
+	//FString ImplementationCodeAnalytic = FString::Printf(TEXT("%s %s%d_Analytic(FMaterial%sParameters Parameters)\n{\n%s return %s;\n}\n"), *OutputTypeString, *Custom->GetFunctionName(), OutputIndex, ShaderFrequency == SF_Vertex ? TEXT("Vertex") : TEXT("Pixel"), *Definitions, *Body);
 	//DerivativeVariations[CompiledPDV_Analytic].CustomOutputImplementations.Add(ImplementationCodeAnalytic);
 
 	// return value is not used
@@ -13840,7 +13840,7 @@ void FHLSLMaterialTranslator::GenerateCustomAttributeCode(int32 OutputIndex, int
 		GetFixedParameterCode(OutputCode, *CurrentScopeChunks, Definitions, Body);
 	}
 
-	FString ImplementationCode = FString::Printf(TEXT("%s %s%d(FMaterial%sParameters Parameters)\r\n{\r\n%s return %s;\r\n}\r\n"), *OutputTypeString, *DisplayName, OutputIndex, ShaderFrequency == SF_Vertex ? TEXT("Vertex") : TEXT("Pixel"), *Definitions, *Body);
+	FString ImplementationCode = FString::Printf(TEXT("%s %s%d(FMaterial%sParameters Parameters)\n{\n%s return %s;\n}\n"), *OutputTypeString, *DisplayName, OutputIndex, ShaderFrequency == SF_Vertex ? TEXT("Vertex") : TEXT("Pixel"), *Definitions, *Body);
 	CustomOutputImplementations.Add(ImplementationCode);
 }
 #endif
diff --git a/Engine/Source/Runtime/Engine/Private/Materials/MaterialHLSLEmitter.cpp b/Engine/Source/Runtime/Engine/Private/Materials/MaterialHLSLEmitter.cpp
--- a/Engine/Source/Runtime/Engine/Private/Materials/MaterialHLSLEmitter.cpp
+++ b/Engine/Source/Runtime/Engine/Private/Materials/MaterialHLSLEmitter.cpp
@@ -196,7 +196,7 @@ static FString GenerateMaterialTemplateHLSL(EShaderPlatform ShaderPlatform,
 		for (uint32 CustomUVIndex = 0; CustomUVIndex < NumPixelTexCoords; CustomUVIndex++)
 		{
 			const FString AttributeName = FMaterialAttributeDefinitionMap::GetAttributeName((EMaterialProperty)(MP_CustomizedUVs0 + CustomUVIndex));
-			CustomUVAssignments += FString::Printf(TEXT("\tOutTexCoords[%u] = Parameters.MaterialAttributes.%s;") LINE_TERMINATOR, CustomUVIndex, *AttributeName);
+			CustomUVAssignments += FString::Printf(TEXT("\tOutTexCoords[%u] = Parameters.MaterialAttributes.%s;") HLSL_LINE_TERMINATOR, CustomUVIndex, *AttributeName);
 		}
 		MaterialSourceTemplateParams.Add({ TEXT("get_material_customized_u_vs"), CustomUVAssignments });
 	}
@@ -206,7 +206,7 @@ static FString GenerateMaterialTemplateHLSL(EShaderPlatform ShaderPlatform,
 		FString CustomInterpolatorAssignments;
 		for (uint32 InterpolatorIndex = 0; InterpolatorIndex < NumCustomVectors; ++InterpolatorIndex)
 		{
-			CustomInterpolatorAssignments += FString::Printf(TEXT("\tOutTexCoords[NUM_MATERIAL_TEXCOORDS + %u] = Parameters.CustomInterpolators[%u];") LINE_TERMINATOR, InterpolatorIndex, InterpolatorIndex);
+			CustomInterpolatorAssignments += FString::Printf(TEXT("\tOutTexCoords[NUM_MATERIAL_TEXCOORDS + %u] = Parameters.CustomInterpolators[%u];") HLSL_LINE_TERMINATOR, InterpolatorIndex, InterpolatorIndex);
 		}
 		MaterialSourceTemplateParams.Add({ TEXT("get_custom_interpolators"), CustomInterpolatorAssignments });
 	}
@@ -215,14 +215,14 @@ static FString GenerateMaterialTemplateHLSL(EShaderPlatform ShaderPlatform,
 	{
 		FString EvaluateMaterialDeclaration;
 
-		EvaluateMaterialDeclaration += TEXT("void EvaluateVertexMaterialAttributesInternal(in out FMaterialVertexParameters Parameters)" LINE_TERMINATOR);
-		EvaluateMaterialDeclaration += TEXT("{" LINE_TERMINATOR);
+		EvaluateMaterialDeclaration += TEXT("void EvaluateVertexMaterialAttributesInternal(in out FMaterialVertexParameters Parameters)" HLSL_LINE_TERMINATOR);
+		EvaluateMaterialDeclaration += TEXT("{" HLSL_LINE_TERMINATOR);
 		EvaluateMaterialDeclaration += VertexShaderCode;
-		EvaluateMaterialDeclaration += TEXT("}" LINE_TERMINATOR);
+		EvaluateMaterialDeclaration += TEXT("}" HLSL_LINE_TERMINATOR);
 
-		EvaluateMaterialDeclaration += TEXT("void EvaluateVertexMaterialAttributes(in out FMaterialVertexParameters Parameters)" LINE_TERMINATOR);
-		EvaluateMaterialDeclaration += TEXT("{" LINE_TERMINATOR);
-		EvaluateMaterialDeclaration += TEXT("    EvaluateVertexMaterialAttributesInternal(Parameters);" LINE_TERMINATOR);
+		EvaluateMaterialDeclaration += TEXT("void EvaluateVertexMaterialAttributes(in out FMaterialVertexParameters Parameters)" HLSL_LINE_TERMINATOR);
+		EvaluateMaterialDeclaration += TEXT("{" HLSL_LINE_TERMINATOR);
+		EvaluateMaterialDeclaration += TEXT("    EvaluateVertexMaterialAttributesInternal(Parameters);" HLSL_LINE_TERMINATOR);
 		for (uint32 TexCoordIndex = 0; TexCoordIndex < NumPixelTexCoords; ++TexCoordIndex)
 		{
 			const Material::EExternalInput TexCoordInput = Material::MakeInputTexCoord(TexCoordIndex);
@@ -230,10 +230,10 @@ static FString GenerateMaterialTemplateHLSL(EShaderPlatform ShaderPlatform,
 			if (EmitMaterialData.IsExternalInputUsed(SF_Pixel, TexCoordInput) && !EmitMaterialData.IsExternalInputUsed(SF_Vertex, TexCoordInput))
 			{
 				const FString AttributeName = FMaterialAttributeDefinitionMap::GetAttributeName((EMaterialProperty)(MP_CustomizedUVs0 + TexCoordIndex));
-				EvaluateMaterialDeclaration += FString::Printf(TEXT("    Parameters.MaterialAttributes.%s = Parameters.TexCoords[%d];") LINE_TERMINATOR, *AttributeName, TexCoordIndex);
+				EvaluateMaterialDeclaration += FString::Printf(TEXT("    Parameters.MaterialAttributes.%s = Parameters.TexCoords[%d];") HLSL_LINE_TERMINATOR, *AttributeName, TexCoordIndex);
 			}
 		}
-		EvaluateMaterialDeclaration += TEXT("}" LINE_TERMINATOR);
+		EvaluateMaterialDeclaration += TEXT("}" HLSL_LINE_TERMINATOR);
 
 		FString EvaluateMaterialAttributesPhase0[2];
 		FString EvaluateMaterialAttributesPhase1[2];
@@ -241,28 +241,28 @@ static FString GenerateMaterialTemplateHLSL(EShaderPlatform ShaderPlatform,
 		{
 			const TCHAR* DerivativeName = (DerivativeIndex == 0) ? TEXT("HWDerivative") : TEXT("AnalyticDerivative");
 
-			EvaluateMaterialAttributesPhase0[DerivativeIndex] =  FString::Printf(TEXT("    FMaterialAttributes MaterialAttributesPhase0;" LINE_TERMINATOR), DerivativeName);
-			EvaluateMaterialAttributesPhase0[DerivativeIndex] += FString::Printf(TEXT("    EvaluatePixelMaterialAttributesPhase0_%s(Parameters, MaterialAttributesPhase0);" LINE_TERMINATOR), DerivativeName);
+			EvaluateMaterialAttributesPhase0[DerivativeIndex] =  FString::Printf(TEXT("    FMaterialAttributes MaterialAttributesPhase0;" HLSL_LINE_TERMINATOR), DerivativeName);
+			EvaluateMaterialAttributesPhase0[DerivativeIndex] += FString::Printf(TEXT("    EvaluatePixelMaterialAttributesPhase0_%s(Parameters, MaterialAttributesPhase0);" HLSL_LINE_TERMINATOR), DerivativeName);
 
-			EvaluateMaterialDeclaration += FString::Printf(TEXT("void EvaluatePixelMaterialAttributesPhase0_%s(in out FMaterialPixelParameters Parameters, out FMaterialAttributes OutResult)" LINE_TERMINATOR), DerivativeName);
-			EvaluateMaterialDeclaration += TEXT("{" LINE_TERMINATOR);
+			EvaluateMaterialDeclaration += FString::Printf(TEXT("void EvaluatePixelMaterialAttributesPhase0_%s(in out FMaterialPixelParameters Parameters, out FMaterialAttributes OutResult)" HLSL_LINE_TERMINATOR), DerivativeName);
+			EvaluateMaterialDeclaration += TEXT("{" HLSL_LINE_TERMINATOR);
 			EvaluateMaterialDeclaration += PixelShaderCodePhase0[DerivativeIndex];
-			EvaluateMaterialDeclaration += TEXT("}" LINE_TERMINATOR);
+			EvaluateMaterialDeclaration += TEXT("}" HLSL_LINE_TERMINATOR);
 
 			if (PixelShaderCodePhase1[DerivativeIndex])
 			{
-				EvaluateMaterialDeclaration += FString::Printf(TEXT("void EvaluatePixelMaterialAttributesPhase1_%s(in out FMaterialPixelParameters Parameters, out FMaterialAttributes OutResult)" LINE_TERMINATOR), DerivativeName);
-				EvaluateMaterialDeclaration += TEXT("{" LINE_TERMINATOR);
+				EvaluateMaterialDeclaration += FString::Printf(TEXT("void EvaluatePixelMaterialAttributesPhase1_%s(in out FMaterialPixelParameters Parameters, out FMaterialAttributes OutResult)" HLSL_LINE_TERMINATOR), DerivativeName);
+				EvaluateMaterialDeclaration += TEXT("{" HLSL_LINE_TERMINATOR);
 				EvaluateMaterialDeclaration += PixelShaderCodePhase1[DerivativeIndex];
-				EvaluateMaterialDeclaration += TEXT("}" LINE_TERMINATOR);
+				EvaluateMaterialDeclaration += TEXT("}" HLSL_LINE_TERMINATOR);
 
-				EvaluateMaterialAttributesPhase1[DerivativeIndex] =  FString::Printf(TEXT("    FMaterialAttributes MaterialAttributesPhase1;" LINE_TERMINATOR), DerivativeName);
-				EvaluateMaterialAttributesPhase1[DerivativeIndex] += FString::Printf(TEXT("    EvaluatePixelMaterialAttributesPhase1_%s(Parameters, MaterialAttributesPhase1);" LINE_TERMINATOR), DerivativeName);
-				EvaluateMaterialAttributesPhase1[DerivativeIndex] += TEXT("    Parameters.MaterialAttributes = MaterialAttributesPhase1;" LINE_TERMINATOR);
+				EvaluateMaterialAttributesPhase1[DerivativeIndex] =  FString::Printf(TEXT("    FMaterialAttributes MaterialAttributesPhase1;" HLSL_LINE_TERMINATOR), DerivativeName);
+				EvaluateMaterialAttributesPhase1[DerivativeIndex] += FString::Printf(TEXT("    EvaluatePixelMaterialAttributesPhase1_%s(Parameters, MaterialAttributesPhase1);" HLSL_LINE_TERMINATOR), DerivativeName);
+				EvaluateMaterialAttributesPhase1[DerivativeIndex] += TEXT("    Parameters.MaterialAttributes = MaterialAttributesPhase1;" HLSL_LINE_TERMINATOR);
 			}
 			else
 			{
-				EvaluateMaterialAttributesPhase0[DerivativeIndex] += TEXT("    Parameters.MaterialAttributes = MaterialAttributesPhase0;" LINE_TERMINATOR);
+				EvaluateMaterialAttributesPhase0[DerivativeIndex] += TEXT("    Parameters.MaterialAttributes = MaterialAttributesPhase0;" HLSL_LINE_TERMINATOR);
 			}
 		}
 
@@ -290,19 +290,19 @@ static FString GenerateMaterialTemplateHLSL(EShaderPlatform ShaderPlatform,
 
 				if (PropertyIndex == MP_FrontMaterial)
 				{
-					EvaluateMaterialAttributesCode += FString::Printf("    PixelMaterialInputs.%s = GetInitialisedSubstrateData();" LINE_TERMINATOR, *PropertyName);
+					EvaluateMaterialAttributesCode += FString::Printf("    PixelMaterialInputs.%s = GetInitialisedSubstrateData();" HLSL_LINE_TERMINATOR, *PropertyName);
 				}
 				else if (PropertyIndex == MP_SubsurfaceColor)
 				{
 					EvaluateMaterialAttributesCode += FString::Printf(
-						"    PixelMaterialInputs.Subsurface = float4(MaterialAttributesPhase%d.%s, %s);" LINE_TERMINATOR,
+						"    PixelMaterialInputs.Subsurface = float4(MaterialAttributesPhase%d.%s, %s);" HLSL_LINE_TERMINATOR,
 						EvaluatePhase,
 						*PropertyName,
 						SubsurfaceProfileCode ? SubsurfaceProfileCode : TEXT("1.0f"));
 				}
 				else
 				{
-					EvaluateMaterialAttributesCode += FString::Printf("    PixelMaterialInputs.%s = MaterialAttributesPhase%d.%s;" LINE_TERMINATOR, *PropertyName, EvaluatePhase, *PropertyName);
+					EvaluateMaterialAttributesCode += FString::Printf("    PixelMaterialInputs.%s = MaterialAttributesPhase%d.%s;" HLSL_LINE_TERMINATOR, *PropertyName, EvaluatePhase, *PropertyName);
 				}
 			}
 
diff --git a/Engine/Source/Runtime/Engine/Private/Materials/MaterialSourceTemplate.cpp b/Engine/Source/Runtime/Engine/Private/Materials/MaterialSourceTemplate.cpp
--- a/Engine/Source/Runtime/Engine/Private/Materials/MaterialSourceTemplate.cpp
+++ b/Engine/Source/Runtime/Engine/Private/Materials/MaterialSourceTemplate.cpp
@@ -52,6 +52,10 @@ bool FMaterialSourceTemplate::Preload(EShaderPlatform ShaderPlatform)
 	FString MaterialTemplateString;
 	LoadShaderSourceFileChecked(TEXT("/Engine/Private/MaterialTemplate.ush"), ShaderPlatform, MaterialTemplateString);
 
+	// Normalize line endings -- preprocessor does this later if necessary, but that can run faster if it's already done, and doing it here
+	// means it only happens once when the template gets loaded, rather than for every Material shader.
+	MaterialTemplateString.ReplaceInline(TEXT("\r\n"), TEXT("\n"), ESearchCase::CaseSensitive);
+
 	// Find the string index of the '#line' statement in MaterialTemplate.usf
 	const int32 LineIndex = MaterialTemplateString.Find(TEXT("#line"), ESearchCase::CaseSensitive);
 	check(LineIndex != INDEX_NONE);
@@ -62,7 +66,6 @@ bool FMaterialSourceTemplate::Preload(EShaderPlatform ShaderPlatform)
 	do
 	{
 		TemplateLineNumber++;
-		// Using \n instead of LINE_TERMINATOR as not all of the lines are terminated consistently
 		// Subtract one from the last found line ending index to make sure we skip over it
 		StartPosition = MaterialTemplateString.Find(TEXT("\n"), ESearchCase::CaseSensitive, ESearchDir::FromEnd, StartPosition - 1);
 	} while (StartPosition != INDEX_NONE);
diff --git a/Engine/Source/Runtime/RenderCore/Private/ShaderParameterMetadata.cpp b/Engine/Source/Runtime/RenderCore/Private/ShaderParameterMetadata.cpp
--- a/Engine/Source/Runtime/RenderCore/Private/ShaderParameterMetadata.cpp
+++ b/Engine/Source/Runtime/RenderCore/Private/ShaderParameterMetadata.cpp
@@ -947,7 +947,7 @@ void FShaderParametersMetadata::InitializeUniformBufferDeclaration()
 
 		// Cache strings for uniform buffer generated path and include
 		UniformBufferPath = FString::Printf(TEXT("/Engine/Generated/UniformBuffers/%s.ush"), ShaderVariableName);
-		UniformBufferInclude = FString::Printf(TEXT("#include \"/Engine/Generated/UniformBuffers/%s.ush\"") LINE_TERMINATOR, ShaderVariableName);
+		UniformBufferInclude = FString::Printf(TEXT("#include \"/Engine/Generated/UniformBuffers/%s.ush\"") HLSL_LINE_TERMINATOR, ShaderVariableName);
 
 		// Cache some frequently used hashes
 		UniformBufferPathHash = GetTypeHash(UniformBufferPath);
diff --git a/Engine/Source/Runtime/RenderCore/Private/ShaderParameterParser.cpp b/Engine/Source/Runtime/RenderCore/Private/ShaderParameterParser.cpp
--- a/Engine/Source/Runtime/RenderCore/Private/ShaderParameterParser.cpp
+++ b/Engine/Source/Runtime/RenderCore/Private/ShaderParameterParser.cpp
@@ -1068,7 +1068,7 @@ bool FShaderParameterParser::MoveShaderParametersToRootConstantBuffer(
 				if (ParsedParameter->ConstantBufferParameterType == EShaderParameterType::BindlessResourceIndex)
 				{
 					RootCBufferContent.Append(FString::Printf(
-						TEXT("uint BindlessResource_%.*s : packoffset(c%d%s);\r\n"),
+						TEXT("uint BindlessResource_%.*s : packoffset(c%d%s);\n"),
 						SVARG(ParsedParameter->ParsedName),
 						ConstantRegister,
 						ConstantSwizzle));
@@ -1076,7 +1076,7 @@ bool FShaderParameterParser::MoveShaderParametersToRootConstantBuffer(
 				else if (ParsedParameter->ConstantBufferParameterType == EShaderParameterType::BindlessSamplerIndex)
 				{
 					RootCBufferContent.Append(FString::Printf(
-						TEXT("uint BindlessSampler_%.*s : packoffset(c%d%s);\r\n"),
+						TEXT("uint BindlessSampler_%.*s : packoffset(c%d%s);\n"),
 						SVARG(ParsedParameter->ParsedName),
 						ConstantRegister,
 						ConstantSwizzle));
@@ -1086,7 +1086,7 @@ bool FShaderParameterParser::MoveShaderParametersToRootConstantBuffer(
 					if (!ParsedParameter->ParsedArraySize.IsEmpty())
 					{
 						RootCBufferContent.Append(FString::Printf(
-							TEXT("%.*s %s[%.*s] : packoffset(c%d%s);\r\n"),
+							TEXT("%.*s %s[%.*s] : packoffset(c%d%s);\n"),
 							SVARG(ParsedParameter->ParsedType),
 							ShaderBindingName,
 							SVARG(ParsedParameter->ParsedArraySize),
@@ -1096,7 +1096,7 @@ bool FShaderParameterParser::MoveShaderParametersToRootConstantBuffer(
 					else
 					{
 						RootCBufferContent.Append(FString::Printf(
-							TEXT("%.*s %s : packoffset(c%d%s);\r\n"),
+							TEXT("%.*s %s : packoffset(c%d%s);\n"),
 							SVARG(ParsedParameter->ParsedType),
 							ShaderBindingName,
 							ConstantRegister,
@@ -1109,10 +1109,10 @@ bool FShaderParameterParser::MoveShaderParametersToRootConstantBuffer(
 		});
 
 		FString CBufferCodeBlock = FString::Printf(
-			TEXT("%s %s\r\n")
-			TEXT("{\r\n")
+			TEXT("%s %s\n")
+			TEXT("{\n")
 			TEXT("%s")
-			TEXT("}\r\n\r\n"),
+			TEXT("}\n\n"),
 			ConstantBufferType,
 			FShaderParametersMetadata::kRootUniformBufferBindingName,
 			*RootCBufferContent);
diff --git a/Engine/Source/Runtime/RenderCore/Private/ShaderParameters.cpp b/Engine/Source/Runtime/RenderCore/Private/ShaderParameters.cpp
--- a/Engine/Source/Runtime/RenderCore/Private/ShaderParameters.cpp
+++ b/Engine/Source/Runtime/RenderCore/Private/ShaderParameters.cpp
@@ -96,7 +96,7 @@ void FShaderUniformBufferParameter::ModifyCompilationEnvironment(const TCHAR* Pa
 	}
 
 	FString& GeneratedUniformBuffersInclude = OutEnvironment.IncludeVirtualPathToContentsMap.FindOrAdd("/Engine/Generated/GeneratedUniformBuffers.ush");
-	const FString Include = FString::Printf(TEXT("#include \"%s\"") LINE_TERMINATOR, *IncludeName);
+	const FString Include = FString::Printf(TEXT("#include \"%s\"") HLSL_LINE_TERMINATOR, *IncludeName);
 
 	GeneratedUniformBuffersInclude.Append(Include);
 	Struct.AddResourceTableEntries(OutEnvironment.ResourceTableMap, OutEnvironment.UniformBufferMap);
@@ -171,7 +171,7 @@ static void CreateHLSLUniformBufferStructMembersDeclaration(
 		{
 			Decl.StructMembers << GlobalPrefix << TEXT("_");
 		}
-		Decl.StructMembers << Member.GetName() << TEXT(");\r\n");
+		Decl.StructMembers << Member.GetName() << TEXT(");\n");
 	};
 
 	if ((UniformBufferStruct.GetUsageFlags() & (uint32)FShaderParametersMetadata::EUsageFlags::UniformView) != 0)
@@ -188,7 +188,7 @@ static void CreateHLSLUniformBufferStructMembersDeclaration(
 		}
 		ParameterName << Member.GetName();
 
-		Decl.ConstantBufferMembers.Appendf(TEXT("UB_CB_UNIFORM_BLOCK(%s, %s);\r\n"), *UniformBufferName, *ParameterName);
+		Decl.ConstantBufferMembers.Appendf(TEXT("UB_CB_UNIFORM_BLOCK(%s, %s);\n"), *UniformBufferName, *ParameterName);
 
 		AddStructMember(Member, false);
 		return;
@@ -232,7 +232,7 @@ static void CreateHLSLUniformBufferStructMembersDeclaration(
 				check(HLSLBaseOffset < AbsoluteMemberOffset);
 				while (HLSLBaseOffset < AbsoluteMemberOffset)
 				{
-					Decl.ConstantBufferMembers.Appendf(TEXT("\t%s() UB_CB_MEMBER_NAME(%s, Padding%u);\r\n"), PreviousBaseTypeName, *UniformBufferName, HLSLBaseOffset);
+					Decl.ConstantBufferMembers.Appendf(TEXT("\t%s() UB_CB_MEMBER_NAME(%s, Padding%u);\n"), PreviousBaseTypeName, *UniformBufferName, HLSLBaseOffset);
 					HLSLBaseOffset += 4;
 				};
 				check(HLSLBaseOffset == AbsoluteMemberOffset);
@@ -250,7 +250,7 @@ static void CreateHLSLUniformBufferStructMembersDeclaration(
 			}
 			ParameterName << Member.GetName();
 
-			Decl.ConstantBufferMembers.Appendf(TEXT("\tUB_UINT() UB_CB_PREFIXED_MEMBER_NAME(%s, %s, %s);\r\n"), *UniformBufferName, MemberPrefix, *ParameterName);
+			Decl.ConstantBufferMembers.Appendf(TEXT("\tUB_UINT() UB_CB_PREFIXED_MEMBER_NAME(%s, %s, %s);\n"), *UniformBufferName, MemberPrefix, *ParameterName);
 		}
 		else
 		{
@@ -307,7 +307,7 @@ static void CreateHLSLUniformBufferStructMembersDeclaration(
 				check(HLSLBaseOffset < AbsoluteMemberOffset);
 				while (HLSLBaseOffset < AbsoluteMemberOffset)
 				{
-					Decl.ConstantBufferMembers.Appendf(TEXT("\t%s() UB_CB_MEMBER_NAME(%s, Padding%u);\r\n"), PreviousBaseTypeName, *UniformBufferName, HLSLBaseOffset);
+					Decl.ConstantBufferMembers.Appendf(TEXT("\t%s() UB_CB_MEMBER_NAME(%s, Padding%u);\n"), PreviousBaseTypeName, *UniformBufferName, HLSLBaseOffset);
 					HLSLBaseOffset += 4;
 				};
 				check(HLSLBaseOffset == AbsoluteMemberOffset);
@@ -322,7 +322,7 @@ static void CreateHLSLUniformBufferStructMembersDeclaration(
 			}
 			ParameterName << Member.GetName();
 
-			Decl.ConstantBufferMembers.Appendf(TEXT("\t%s(%s) UB_CB_MEMBER_NAME(%s, %s%s);\r\n"), BaseTypeName, *TypeDim, *UniformBufferName, *ParameterName, *ArrayDim);
+			Decl.ConstantBufferMembers.Appendf(TEXT("\t%s(%s) UB_CB_MEMBER_NAME(%s, %s%s);\n"), BaseTypeName, *TypeDim, *UniformBufferName, *ParameterName, *ArrayDim);
 
 			AddStructMember(Member, false);
 		}
@@ -346,17 +346,17 @@ static void CreateHLSLUniformBufferStructMembersDeclaration(
 
 			if (Member.GetBaseType() == UBMT_SAMPLER)
 			{
-				Decl.ResourceMembers.Appendf(TEXT("UB_RESOURCE_MEMBER_SAMPLER(%s, %s, %s);\r\n"), Member.GetShaderType(), *UniformBufferName, *ParameterName);
+				Decl.ResourceMembers.Appendf(TEXT("UB_RESOURCE_MEMBER_SAMPLER(%s, %s, %s);\n"), Member.GetShaderType(), *UniformBufferName, *ParameterName);
 				AddStructMember(Member, true);
 			}
 			else if (Member.GetBaseType() == UBMT_SRV)
 			{
-				Decl.ResourceMembers.Appendf(TEXT("UB_RESOURCE_MEMBER_RESOURCE(%s, %s, %s);\r\n"), Member.GetShaderType(), *UniformBufferName, *ParameterName);
+				Decl.ResourceMembers.Appendf(TEXT("UB_RESOURCE_MEMBER_RESOURCE(%s, %s, %s);\n"), Member.GetShaderType(), *UniformBufferName, *ParameterName);
 				AddStructMember(Member, true);
 			}
 			else
 			{
-				Decl.ResourceMembers.Appendf(TEXT("UB_RESOURCE_MEMBER_RESOURCE(%s, %s, %s);\r\n"), Member.GetShaderType(), *UniformBufferName, *ParameterName);
+				Decl.ResourceMembers.Appendf(TEXT("UB_RESOURCE_MEMBER_RESOURCE(%s, %s, %s);\n"), Member.GetShaderType(), *UniformBufferName, *ParameterName);
 				AddStructMember(Member, true);
 			}
 		}
@@ -374,15 +374,15 @@ static FString CreateHLSLUniformBufferDeclaration(const TCHAR* UniformBufferName
 		CreateHLSLUniformBufferStructMembersDeclaration(UniformBufferStruct, UniformBufferName, TEXT(""), TEXT(""), 0, Decl, HLSLBaseOffset);
 
 		return FString::Printf(
-			TEXT("#pragma once\r\n")
-			TEXT("UB_CB_DEFINITION_START(%s)\r\n")
+			TEXT("#pragma once\n")
+			TEXT("UB_CB_DEFINITION_START(%s)\n")
 			TEXT("%s")
-			TEXT("UB_CB_DEFINITION_END(%s)\r\n")
+			TEXT("UB_CB_DEFINITION_END(%s)\n")
 			TEXT("%s")
-			TEXT("UniformBuffer %s\r\n")
-			TEXT("{\r\n")
+			TEXT("UniformBuffer %s\n")
+			TEXT("{\n")
 			TEXT("%s")
-			TEXT("};\r\n"),
+			TEXT("};\n"),
 			UniformBufferName,
 			*Decl.ConstantBufferMembers,
 			UniformBufferName,
diff --git a/Engine/Source/Runtime/RenderCore/Public/ShaderCore.h b/Engine/Source/Runtime/RenderCore/Public/ShaderCore.h
--- a/Engine/Source/Runtime/RenderCore/Public/ShaderCore.h
+++ b/Engine/Source/Runtime/RenderCore/Public/ShaderCore.h
@@ -1200,6 +1200,9 @@ class FShaderType;
 class FVertexFactoryType;
 class FShaderPipelineType;
 
+// Text to use as line terminator for HLSL files (may differ from platform LINE_TERMINATOR)
+#define HLSL_LINE_TERMINATOR TEXT("\n")
+
 /** Force updates each shader/pipeline type provided to update their list of referenced uniform buffers. */
 RENDERCORE_API void UpdateReferencedUniformBufferNames(
 	TArrayView<const FShaderType*> OutdatedShaderTypes,
diff --git a/Engine/Source/Runtime/RenderCore/Public/ShaderParameterMetadata.h b/Engine/Source/Runtime/RenderCore/Public/ShaderParameterMetadata.h
--- a/Engine/Source/Runtime/RenderCore/Public/ShaderParameterMetadata.h
+++ b/Engine/Source/Runtime/RenderCore/Public/ShaderParameterMetadata.h
@@ -452,7 +452,7 @@ private:
 
 	/** Strings for uniform buffer generated path and include, created once */
 	FString UniformBufferPath;		// Format:  "/Engine/Generated/UniformBuffers/%s.ush"
-	FString UniformBufferInclude;	// Format:  "#include \"/Engine/Generated/UniformBuffers/%s.ush\"" LINE_TERMINATOR
+	FString UniformBufferInclude;	// Format:  "#include \"/Engine/Generated/UniformBuffers/%s.ush\"" HLSL_LINE_TERMINATOR
 
 	/** Hashes for frequently used strings */
 	uint32 UniformBufferPathHash;
