cluster unions should only add themselves into the SQ if they are non-empty (i.e. have a non-zero # of shapes and the shape in the shapes array isn't because of a shape created by an FImplicitObjectUnionClustered)

diff --git a/Engine/Source/Runtime/Engine/Private/PhysicsEngine/Experimental/PhysScene_Chaos.cpp b/Engine/Source/Runtime/Engine/Private/PhysicsEngine/Experimental/PhysScene_Chaos.cpp
--- a/Engine/Source/Runtime/Engine/Private/PhysicsEngine/Experimental/PhysScene_Chaos.cpp
+++ b/Engine/Source/Runtime/Engine/Private/PhysicsEngine/Experimental/PhysScene_Chaos.cpp
@@ -2351,9 +2351,21 @@ void FPhysScene_Chaos::OnSyncBodies(Chaos::FPhysicsSolverBase* Solver)
 					PendingTransforms.Add(FPhysScenePendingComponentTransform_Chaos(ParentComponent, DirtyParticle->GetWakeEvent()));
 				}
 
+				ParentComponent->SyncClusterUnionFromProxy();
+
 				// make sure we have at least a child to be added to the acceleration structure 
 				// this avoid the invalid bounds to cause the particle to be added to the global acceleration structure array
-				if (ParentComponent->NumChildClusterComponents() > 0)
+				bool bShouldBeInSQ = false;
+
+				if (FImplicitObjectRef GeometryRef = DirtyParticle->GetGeometry())
+				{
+					if (FImplicitObjectUnion* Union = GeometryRef->AsA<FImplicitObjectUnion>())
+					{
+						bShouldBeInSQ = (Union->GetNumRootObjects() > 0);
+					}
+				}
+
+				if (bShouldBeInSQ)
 				{
 					Interface->AddToSpatialAcceleration({ &Handle, 1 }, Outer->GetSpacialAcceleration());
 				}
@@ -2361,7 +2373,7 @@ void FPhysScene_Chaos::OnSyncBodies(Chaos::FPhysicsSolverBase* Solver)
 				{
 					Interface->RemoveFromSpatialAcceleration({ &Handle, 1 }, Outer->GetSpacialAcceleration());
 				}
-				ParentComponent->SyncClusterUnionFromProxy();
+
 				DirtyParticle->ClearEvents();
 			}
 		}
