[SplineMeshComponent] added linear extrapolation support
- added support to distort bounds that could exceed the mesh bounds used to define the spline local range [0, 1]
- in case point on spline is outside default range we use linear extrapolation for position or bounds
- this change also improves positioning of objects attached to sockets outside the mesh bounds
- navigation bounds are now computed and cached to match distorted navcollsion

diff --git a/Engine/Source/Runtime/Engine/Classes/Components/SplineMeshComponent.h b/Engine/Source/Runtime/Engine/Classes/Components/SplineMeshComponent.h
--- a/Engine/Source/Runtime/Engine/Classes/Components/SplineMeshComponent.h
+++ b/Engine/Source/Runtime/Engine/Classes/Components/SplineMeshComponent.h
@@ -130,6 +130,9 @@ class USplineMeshComponent : public UStaticMeshComponent, public IInterface_Coll
 	UPROPERTY()
 	FGuid CachedMeshBodySetupGuid;
 
+	// Navigation bounds can differ from primitive bounds since NavCollision can hold more geometry
+	FBox CachedNavigationBounds;
+
 	// Physics data.
 	UPROPERTY()
 	TObjectPtr<UBodySetup> BodySetup;
@@ -194,11 +197,30 @@ public:
 	//Begin USceneComponent Interface
 	ENGINE_API virtual FBoxSphereBounds CalcBounds(const FTransform& LocalToWorld) const override;
 	ENGINE_API virtual FTransform GetSocketTransform(FName InSocketName, ERelativeTransformSpace TransformSpace = RTS_World) const override;
+	ENGINE_API virtual void UpdateBounds() override;
 	//End USceneComponent Interface
 
 	//Begin UPrimitiveComponent Interface
 protected:
 	ENGINE_API virtual void OnCreatePhysicsState() override;
+
+	float ComputeRatioAlongSpline(float DistanceAlong) const;
+
+	/** Returns the normalized range on the spline where the visual mesh is located taking custom range into account. */
+	void ComputeVisualMeshSplineTRange(float& MinT, float& MaxT) const;
+
+	/**
+	 * Computes the bounding box, in world space, for a given bounding box distorted by the spline in local space.
+	 * By default this method uses the provided mesh bounds that were used to define the spline range [0,1] so all points are expected
+	 * to stay in that range. In case the bounds to deform are overriden by the optional parameter then linear extrapolation
+	 * will be applied at the beginning and at the end of the spline for the exceeding part.
+	 * @param InLocalToWorld Transformation to apply to the computed bounds to convert them from local space to world space.
+	 * @param InMeshBounds Bounds of the static mesh that get distorted by the spline.
+	 * @param InBoundsToDistort Optional bounds to distort instead of using the mesh bounds.
+	 * @return Bounds, in world space, of the provided bounds distorted by the spline.
+	 */
+	FBox ComputeDistortedBounds(const FTransform& InLocalToWorld, const FBoxSphereBounds& InMeshBounds, const FBoxSphereBounds* InBoundsToDistort = nullptr) const;
+
 public:
 	ENGINE_API virtual class UBodySetup* GetBodySetup() override;
 #if WITH_EDITOR
@@ -388,7 +410,7 @@ public:
 	 * Calculates the spline transform, including roll, scale, and offset along the spline at a specified alpha interpolation parameter along the spline
 	 * @Note:  This is mirrored to Lightmass::CalcSliceTransform() and LocalVertexShader.usf.  If you update one of these, please update them all!
 	 */
-	ENGINE_API FTransform CalcSliceTransformAtSplineOffset(const float Alpha) const;
+	ENGINE_API FTransform CalcSliceTransformAtSplineOffset(const float Alpha, const float MinT=0.f, const float MaxT=1.0f) const;
 
 	UE_DEPRECATED(5.2, "Use GetAxisValueRef() instead.")
 	static const double& GetAxisValue(const FVector3d& InVector, ESplineMeshAxis::Type InAxis) { return GetAxisValueRef(InVector, InAxis); }
diff --git a/Engine/Source/Runtime/Engine/Private/Components/SplineMeshComponent.cpp b/Engine/Source/Runtime/Engine/Private/Components/SplineMeshComponent.cpp
--- a/Engine/Source/Runtime/Engine/Private/Components/SplineMeshComponent.cpp
+++ b/Engine/Source/Runtime/Engine/Private/Components/SplineMeshComponent.cpp
@@ -29,6 +29,15 @@
 #include "LandscapeSplineActor.h"
 #endif // WITH_EDITOR
 
+namespace UE::SplineMesh
+{
+	float RealToFloatChecked(const double Value)
+	{
+		ensureMsgf(Value >= TNumericLimits<float>::Lowest() && Value <= TNumericLimits<float>::Max(), TEXT("Value %f exceeds float limits"), Value);
+		return static_cast<float>(Value);
+	}
+}
+
 int32 GNoRecreateSplineMeshProxy = 1;
 static FAutoConsoleVariableRef CVarNoRecreateSplineMeshProxy(
 	TEXT("r.SplineMesh.NoRecreateProxy"),
@@ -61,7 +70,7 @@ void PackSplineMeshParams(const FSplineMeshShaderParams& Params, const TArrayVie
 	};
 	auto PackSNorm16 = [](float Value, uint32 Shift = 0) -> uint32
 	{
-		float N = FMath::Clamp(Value, -1.0f, 1.0f) * 0.5f + 0.5f;
+		const float N = FMath::Clamp(Value, -1.0f, 1.0f) * 0.5f + 0.5f;
 		return uint32(N * 65535.0f) << Shift;
 	};
 
@@ -130,7 +139,7 @@ static FVector3f SplineEvalPos(const FSplineMeshParams& Params, float A)
 	return SplineEvalPos(StartPos, StartTangent, EndPos, EndTangent, A);
 }
 
-static FVector3f SplineEvalDir(const FVector3f& StartPos, const FVector3f& StartTangent, const FVector3f& EndPos, const FVector3f& EndTangent, float A)
+static FVector3f SplineEvalTangent(const FVector3f& StartPos, const FVector3f& StartTangent, const FVector3f& EndPos, const FVector3f& EndTangent, const float A)
 {
 	const FVector3f C = (6 * StartPos) + (3 * StartTangent) + (3 * EndTangent) - (6 * EndPos);
 	const FVector3f D = (-6 * StartPos) - (4 * StartTangent) - (2 * EndTangent) + (6 * EndPos);
@@ -138,10 +147,10 @@ static FVector3f SplineEvalDir(const FVector3f& StartPos, const FVector3f& Start
 
 	const float A2 = A * A;
 
-	return ((C * A2) + (D * A) + E).GetSafeNormal();
+	return (C * A2) + (D * A) + E;
 }
 
-static FVector3f SplineEvalDir(const FSplineMeshParams& Params, float A)
+static FVector3f SplineEvalTangent(const FSplineMeshParams& Params, const float A)
 {
 	// TODO: these don't need to be doubles!
 	const FVector3f StartPos = FVector3f(Params.StartPos);
@@ -149,7 +158,12 @@ static FVector3f SplineEvalDir(const FSplineMeshParams& Params, float A)
 	const FVector3f EndPos = FVector3f(Params.EndPos);
 	const FVector3f EndTangent = FVector3f(Params.EndTangent);
 
-	return SplineEvalDir(StartPos, StartTangent, EndPos, EndTangent, A);
+	return SplineEvalTangent(StartPos, StartTangent, EndPos, EndTangent, A);
+}
+
+static FVector3f SplineEvalDir(const FSplineMeshParams& Params, const float A)
+{
+	return SplineEvalTangent(Params, A).GetSafeNormal();
 }
 
 FSplineMeshInstanceData::FSplineMeshInstanceData(const USplineMeshComponent* SourceComponent)
@@ -186,7 +200,7 @@ void FSplineMeshVertexFactoryShaderParameters::GetElementShaderBindings(
 	
 	if (BatchElement.bUserDataIsColorVertexBuffer)
 	{
-		FColorVertexBuffer* OverrideColorVertexBuffer = (FColorVertexBuffer*)BatchElement.UserData;
+		const FColorVertexBuffer* OverrideColorVertexBuffer = (FColorVertexBuffer*)BatchElement.UserData;
 		check(OverrideColorVertexBuffer);
 
 		if (!LocalVertexFactory->SupportsManualVertexFetch(FeatureLevel))
@@ -275,9 +289,9 @@ void FSplineMeshSceneProxy::InitVertexFactory(USplineMeshComponent* InComponent,
 		return;
 	}
 
-	UStaticMesh* Parent = InComponent->GetStaticMesh();
+	const UStaticMesh* Parent = InComponent->GetStaticMesh();
 	bool bOverrideColorVertexBuffer = !!InOverrideColorVertexBuffer;
-	ERHIFeatureLevel::Type FeatureLevel = GetScene().GetFeatureLevel();
+	const ERHIFeatureLevel::Type FeatureLevel = GetScene().GetFeatureLevel();
 
 	if ((VertexFactories->SplineVertexFactory && !bOverrideColorVertexBuffer) || (VertexFactories->SplineVertexFactoryOverrideColorVertexBuffer && bOverrideColorVertexBuffer))
 	{
@@ -650,16 +664,15 @@ FSplineMeshShaderParams USplineMeshComponent::CalculateShaderParams() const
 	Output.MeshZScale = 1.0f;
 	Output.MeshZOffset = 0.0f;
 
-	float BoundsXYRadius = 0.0f;
 	if (GetStaticMesh())
 	{
 		const FBoxSphereBounds StaticMeshBounds = GetStaticMesh()->GetBounds();
-		BoundsXYRadius = FVector3f(StaticMeshBounds.BoxExtent).Dot((Output.MeshX + Output.MeshY).GetUnsafeNormal());
+		const float BoundsXYRadius = FVector3f(StaticMeshBounds.BoxExtent).Dot((Output.MeshX + Output.MeshY).GetUnsafeNormal());
 
-		const float MeshMinZ = GetAxisValueRef(StaticMeshBounds.Origin - StaticMeshBounds.BoxExtent, ForwardAxis);
-		const float MeshZLen = 2.0f * GetAxisValueRef(StaticMeshBounds.BoxExtent, ForwardAxis);
+		const float MeshMinZ = UE::SplineMesh::RealToFloatChecked(GetAxisValueRef(StaticMeshBounds.Origin - StaticMeshBounds.BoxExtent, ForwardAxis));
+		const float MeshZLen = UE::SplineMesh::RealToFloatChecked(2 * GetAxisValueRef(StaticMeshBounds.BoxExtent, ForwardAxis));
 		const float InvMeshZLen = (MeshZLen <= 0.0f) ? 1.0f : 1.0f / MeshZLen;
-		const float MeshTexelLen = float(SPLINE_MESH_TEXEL_WIDTH - 1);
+		constexpr float MeshTexelLen = float(SPLINE_MESH_TEXEL_WIDTH - 1);
 
 		if (FMath::IsNearlyEqual(SplineBoundaryMin, SplineBoundaryMax))
 		{
@@ -706,7 +719,7 @@ FSplineMeshShaderParams USplineMeshComponent::CalculateShaderParams() const
 
 			// Take the mid-point scale in X/Y to balance out LOD selection in case either of them are extreme.
 			auto AvgAbs = [](float A, float B) { return (FMath::Abs(A) + FMath::Abs(B)) * 0.5f; };
-			FVector3f DeformScale = FVector3f(
+			const FVector3f DeformScale = FVector3f(
 				SplineLength * Output.MeshZScale,
 				AvgAbs(Output.StartScale.X, Output.EndScale.X),
 				AvgAbs(Output.StartScale.Y, Output.EndScale.Y)
@@ -772,7 +785,7 @@ void USplineMeshComponent::Serialize(FArchive& Ar)
 		SplineParams.StartRoll -= UE_HALF_PI;
 		SplineParams.EndRoll -= UE_HALF_PI;
 
-		float Temp = SplineParams.StartOffset.X;
+		double Temp = SplineParams.StartOffset.X;
 		SplineParams.StartOffset.X = -SplineParams.StartOffset.Y;
 		SplineParams.StartOffset.Y = Temp;
 		Temp = SplineParams.EndOffset.X;
@@ -797,7 +810,7 @@ bool USplineMeshComponent::IsEditorOnly() const
 	}
 
 	// If Landscape uses generated LandscapeSplineMeshesActors, SplineMeshComponents is removed from cooked build  
-	ALandscapeSplineActor* SplineActor = Cast<ALandscapeSplineActor>(GetOwner());
+	const ALandscapeSplineActor* SplineActor = Cast<ALandscapeSplineActor>(GetOwner());
 	if (SplineActor && SplineActor->HasGeneratedLandscapeSplineMeshesActors())
 	{
 		return true;
@@ -808,7 +821,7 @@ bool USplineMeshComponent::IsEditorOnly() const
 
 bool USplineMeshComponent::Modify(bool bAlwaysMarkDirty)
 {
-	bool bSavedToTransactionBuffer = Super::Modify(bAlwaysMarkDirty);
+	const bool bSavedToTransactionBuffer = Super::Modify(bAlwaysMarkDirty);
 
 	if (BodySetup != nullptr)
 	{
@@ -826,13 +839,13 @@ void USplineMeshComponent::CollectPSOPrecacheData(const FPSOPrecacheParams& Base
 		return;
 	}
 
-	FVertexFactoryType* VertexFactoryType = &FSplineMeshVertexFactory::StaticType;
+	const FVertexFactoryType* VertexFactoryType = &FSplineMeshVertexFactory::StaticType;
 	int32 LightMapCoordinateIndex = GetStaticMesh()->GetLightMapCoordinateIndex();
 
 	auto SMC_GetElements = [LightMapCoordinateIndex](const FStaticMeshLODResources& LODRenderData, int32 LODIndex, bool bSupportsManualVertexFetch, FVertexDeclarationElementList& Elements)
 	{
 		// FIXME: This will miss when SM component overrides vertex colors and source StaticMesh does not have vertex colors
-		bool bOverrideColorVertexBuffer = false;
+		constexpr bool bOverrideColorVertexBuffer = false;
 		FLocalVertexFactory::FDataType Data;
 		InitSplineMeshVertexFactoryComponents(LODRenderData.VertexBuffers, nullptr /*VertexFactory*/, LightMapCoordinateIndex, bOverrideColorVertexBuffer, Data);
 		FLocalVertexFactory::GetVertexElements(GMaxRHIFeatureLevel, EVertexInputStreamType::Default, bSupportsManualVertexFetch, Data, Elements);
@@ -860,109 +873,156 @@ FPrimitiveSceneProxy* USplineMeshComponent::CreateStaticMeshSceneProxy(Nanite::F
 
 FBoxSphereBounds USplineMeshComponent::CalcBounds(const FTransform& LocalToWorld) const
 {
-	if (!GetStaticMesh())
+	const UStaticMesh* Mesh = GetStaticMesh();
+	if (Mesh == nullptr)
 	{
-		return FBoxSphereBounds(FBox(ForceInit));
+		return FBox();
 	}
 
-	float MinT = 0.0f;
-	float MaxT = 1.0f;
+	const FBox ComputedBounds = ComputeDistortedBounds(LocalToWorld, Mesh->GetBounds());
+	return ComputedBounds;
+}
 
-	const FBoxSphereBounds MeshBounds = GetStaticMesh()->GetBounds();
+void USplineMeshComponent::UpdateBounds()
+{
+	Super::UpdateBounds();
 
-	const bool bHasCustomBoundary = !FMath::IsNearlyEqual(SplineBoundaryMin, SplineBoundaryMax);
-	if (bHasCustomBoundary)
+	CachedNavigationBounds = Bounds.GetBox();
+
+	if (const UStaticMesh* Mesh = GetStaticMesh())
 	{
-		// If there's a custom boundary, alter the min/max of the spline we need to evaluate
-		const float MeshMin = GetAxisValueRef(MeshBounds.Origin - MeshBounds.BoxExtent, ForwardAxis);
-		const float MeshMax = GetAxisValueRef(MeshBounds.Origin + MeshBounds.BoxExtent, ForwardAxis);
+		if (const UNavCollisionBase* NavCollision = Mesh->GetNavCollision())
+		{
+			// Match condition in DoCustomNavigableGeometryExport
+			const FBox NavCollisionBounds = NavCollision->GetBounds();
+			if (ensure(!NavCollision->IsDynamicObstacle())
+				&& NavCollision->HasConvexGeometry()
+				&& NavCollisionBounds.IsValid)
+			{
+				const FBoxSphereBounds NavCollisionBoxSphereBounds(NavCollisionBounds);
+				CachedNavigationBounds = ComputeDistortedBounds(GetComponentTransform(), Mesh->GetBounds(), &NavCollisionBoxSphereBounds);
+			}
+		}
+	}
+}
 
-		const float MeshMinT = (MeshMin - SplineBoundaryMin) / (SplineBoundaryMax - SplineBoundaryMin);
-		const float MeshMaxT = (MeshMax - SplineBoundaryMin) / (SplineBoundaryMax - SplineBoundaryMin);
+float USplineMeshComponent::ComputeRatioAlongSpline(const float DistanceAlong) const 
+{
+	// Find how far 'along' mesh (or custom boundaries) we are
+	float Alpha = 0.f;
 
-		// Disallow extrapolation beyond a certain value; enormous bounding boxes cause the render thread to crash
-		const float MaxSplineExtrapolation = 4.0f;
-		if (FMath::Abs(MeshMinT) < MaxSplineExtrapolation && FMath::Abs(MeshMaxT) < MaxSplineExtrapolation)
+	const bool bHasCustomBoundary = !FMath::IsNearlyEqual(SplineBoundaryMin, SplineBoundaryMax);
+	if (bHasCustomBoundary)
+	{
+		Alpha = (DistanceAlong - SplineBoundaryMin) / (SplineBoundaryMax - SplineBoundaryMin);
+	}
+	else if (GetStaticMesh())
+	{
+		const FBoxSphereBounds StaticMeshBounds = GetStaticMesh()->GetBounds();
+		const double MeshMinZ = GetAxisValueRef(StaticMeshBounds.Origin, ForwardAxis) - GetAxisValueRef(StaticMeshBounds.BoxExtent, ForwardAxis);
+		const double MeshRangeZ = 2 * GetAxisValueRef(StaticMeshBounds.BoxExtent, ForwardAxis);
+		if (MeshRangeZ > UE_SMALL_NUMBER)
 		{
-			MinT = MeshMinT;
-			MaxT = MeshMaxT;
+			Alpha = UE::SplineMesh::RealToFloatChecked((DistanceAlong - MeshMinZ) / MeshRangeZ);
 		}
 	}
+	return Alpha;
+}
+
+void USplineMeshComponent::ComputeVisualMeshSplineTRange(float& MinT, float& MaxT) const
+{
+    MinT = 0.0;
+    MaxT = 1.0;
+
+    const bool bHasCustomBoundary = !FMath::IsNearlyEqual(SplineBoundaryMin, SplineBoundaryMax);
+    if (bHasCustomBoundary)
+    {
+        const FBoxSphereBounds& MeshBounds = GetStaticMesh()->GetBounds();
+        // If there's a custom boundary, alter the min/max of the spline we need to evaluate
+        const float BoundsMin = UE::SplineMesh::RealToFloatChecked(GetAxisValueRef(MeshBounds.Origin - MeshBounds.BoxExtent, ForwardAxis));
+        const float BoundsMax = UE::SplineMesh::RealToFloatChecked(GetAxisValueRef(MeshBounds.Origin + MeshBounds.BoxExtent, ForwardAxis));
+        const float BoundsMinT = (BoundsMin - SplineBoundaryMin) / (SplineBoundaryMax - SplineBoundaryMin);
+        const float BoundsMaxT = (BoundsMax - SplineBoundaryMin) / (SplineBoundaryMax - SplineBoundaryMin);
+        // Disallow extrapolation beyond a certain value; enormous bounding boxes cause the render thread to crash
+        constexpr float MaxSplineExtrapolation = 4.0f;
+        MinT = FMath::Max(-MaxSplineExtrapolation, BoundsMinT);
+        MaxT = FMath::Min(BoundsMaxT, MaxSplineExtrapolation);
+    }
+}
+
+FBox USplineMeshComponent::ComputeDistortedBounds(const FTransform& InLocalToWorld, const FBoxSphereBounds& InMeshBounds, const FBoxSphereBounds* InBoundsToDistort) const
+{
+	float MinT = 0.0f;
+	float MaxT = 1.0f;
+	ComputeVisualMeshSplineTRange(MinT, MaxT);
+	const FBoxSphereBounds& BoundsToDistort = InBoundsToDistort ? *InBoundsToDistort : InMeshBounds;
+
 
 	const FVector AxisMask = GetAxisMask(ForwardAxis);
-	const FVector FlattenedMeshOrigin = MeshBounds.Origin * AxisMask;
-	const FVector FlattenedMeshExtent = MeshBounds.BoxExtent * AxisMask;
-	const FBox MeshBoundingBox = FBox(FlattenedMeshOrigin - FlattenedMeshExtent, FlattenedMeshOrigin + FlattenedMeshExtent);
+	const FVector FlattenedBoundsOrigin = BoundsToDistort.Origin * AxisMask;
+	const FVector FlattenedBoundsExtent = BoundsToDistort.BoxExtent * AxisMask;
+	const FBox FlattenedBounds = FBox(FlattenedBoundsOrigin - FlattenedBoundsExtent, FlattenedBoundsOrigin + FlattenedBoundsExtent);
 
 	FBox BoundingBox(ForceInit);
-	BoundingBox += MeshBoundingBox.TransformBy(CalcSliceTransformAtSplineOffset(MinT));
-	BoundingBox += MeshBoundingBox.TransformBy(CalcSliceTransformAtSplineOffset(MaxT));
+	BoundingBox += FlattenedBounds.TransformBy(CalcSliceTransformAtSplineOffset(MinT, MinT, MaxT));
+	BoundingBox += FlattenedBounds.TransformBy(CalcSliceTransformAtSplineOffset(MaxT, MinT, MaxT));
 
 	// Work out coefficients of the cubic spline derivative equation dx/dt
-	const FVector A = 6.0f * SplineParams.StartPos + 3.0f * SplineParams.StartTangent + 3.0f * SplineParams.EndTangent - 6.0f * SplineParams.EndPos;
-	const FVector B = -6.0f * SplineParams.StartPos - 4.0f * SplineParams.StartTangent - 2.0f * SplineParams.EndTangent + 6.0f * SplineParams.EndPos;
-	const FVector C = SplineParams.StartTangent;
-
-	// Minima/maxima happen where dx/dt == 0, calculate t values
-	const FVector Discriminant = B * B - 4.0f * A * C;
+	const FVector A(6 * SplineParams.StartPos + 3 * SplineParams.StartTangent + 3 * SplineParams.EndTangent - 6 * SplineParams.EndPos);
+	const FVector B(-6 * SplineParams.StartPos - 4 * SplineParams.StartTangent - 2 * SplineParams.EndTangent + 6 * SplineParams.EndPos);
+	const FVector C(SplineParams.StartTangent);
 
-	// Work out minima/maxima component-by-component.
-	// Negative discriminant means no solution; A == 0 implies coincident start/end points
-	if (Discriminant.X > 0.0f && !FMath::IsNearlyZero(A.X))
-	{
-		const float SqrtDiscriminant = FMath::Sqrt(Discriminant.X);
-		const float Denominator = 0.5f / A.X;
-		const float T0 = (-B.X + SqrtDiscriminant) * Denominator;
-		const float T1 = (-B.X - SqrtDiscriminant) * Denominator;
-
-		if (T0 >= MinT && T0 <= MaxT)
+	auto AppendAxisExtrema = [&BoundingBox, &FlattenedBounds, MinT, MaxT, this](const double Discriminant, const double A, const double B)
 		{
-			BoundingBox += MeshBoundingBox.TransformBy(CalcSliceTransformAtSplineOffset(T0));
-		}
+			// Negative discriminant means no solution; A == 0 implies coincident start/end points
+			if (Discriminant > 0 && !FMath::IsNearlyZero(A))
+			{
+				const double SqrtDiscriminant = FMath::Sqrt(Discriminant);
+				const double Denominator = 0.5 / A;
+				const double T0 = (-B + SqrtDiscriminant) * Denominator;
+				const double T1 = (-B - SqrtDiscriminant) * Denominator;
 
-		if (T1 >= MinT && T1 <= MaxT)
-		{
-			BoundingBox += MeshBoundingBox.TransformBy(CalcSliceTransformAtSplineOffset(T1));
-		}
-	}
+				if (T0 >= MinT && T0 <= MaxT)
+				{
+					BoundingBox += FlattenedBounds.TransformBy(CalcSliceTransformAtSplineOffset(UE::SplineMesh::RealToFloatChecked(T0), MinT, MaxT));
+				}
 
-	if (Discriminant.Y > 0.0f && !FMath::IsNearlyZero(A.Y))
-	{
-		const float SqrtDiscriminant = FMath::Sqrt(Discriminant.Y);
-		const float Denominator = 0.5f / A.Y;
-		const float T0 = (-B.Y + SqrtDiscriminant) * Denominator;
-		const float T1 = (-B.Y - SqrtDiscriminant) * Denominator;
+				if (T1 >= MinT && T1 <= MaxT)
+				{
+					BoundingBox += FlattenedBounds.TransformBy(CalcSliceTransformAtSplineOffset(UE::SplineMesh::RealToFloatChecked(T1), MinT, MaxT));
+				}
+			}
+		};
 
-		if (T0 >= MinT && T0 <= MaxT)
-		{
-			BoundingBox += MeshBoundingBox.TransformBy(CalcSliceTransformAtSplineOffset(T0));
-		}
+	// Minima/maxima happen where dx/dt == 0, calculate t values
+	const FVector Discriminant = B * B - 4 * A * C;
 
-		if (T1 >= MinT && T1 <= MaxT)
-		{
-			BoundingBox += MeshBoundingBox.TransformBy(CalcSliceTransformAtSplineOffset(T1));
-		}
-	}
+	// Work out minima/maxima component-by-component.
+	AppendAxisExtrema(Discriminant.X, A.X, B.X);
+	AppendAxisExtrema(Discriminant.Y, A.Y, B.Y);
+	AppendAxisExtrema(Discriminant.Z, A.Z, B.Z);
 
-	if (Discriminant.Z > 0.0f && !FMath::IsNearlyZero(A.Z))
+	// Applying extrapolation if bounds to apply on spline are different than the mesh bounds used
+	// to define the spline range [0,1]
+	if (InBoundsToDistort != nullptr && InBoundsToDistort != &InMeshBounds)
 	{
-		const float SqrtDiscriminant = FMath::Sqrt(Discriminant.Z);
-		const float Denominator = 0.5f / A.Z;
-		const float T0 = (-B.Z + SqrtDiscriminant) * Denominator;
-		const float T1 = (-B.Z - SqrtDiscriminant) * Denominator;
+		const double BoundsMin = GetAxisValueRef(BoundsToDistort.Origin - BoundsToDistort.BoxExtent, ForwardAxis);
+		const double BoundsMax = GetAxisValueRef(BoundsToDistort.Origin + BoundsToDistort.BoxExtent, ForwardAxis);
 
-		if (T0 >= MinT && T0 <= MaxT)
+		float Alpha = ComputeRatioAlongSpline(UE::SplineMesh::RealToFloatChecked(BoundsMin));
+		if (Alpha < MinT)
 		{
-			BoundingBox += MeshBoundingBox.TransformBy(CalcSliceTransformAtSplineOffset(T0));
+			BoundingBox += FlattenedBounds.TransformBy(CalcSliceTransformAtSplineOffset(Alpha, MinT, MaxT));
 		}
 
-		if (T1 >= MinT && T1 <= MaxT)
+		Alpha = ComputeRatioAlongSpline(UE::SplineMesh::RealToFloatChecked(BoundsMax));
+		if (Alpha > MaxT)
 		{
-			BoundingBox += MeshBoundingBox.TransformBy(CalcSliceTransformAtSplineOffset(T1));
+			BoundingBox += FlattenedBounds.TransformBy(CalcSliceTransformAtSplineOffset(Alpha, MinT, MaxT));
 		}
 	}
 
-	return FBoxSphereBounds(BoundingBox.TransformBy(LocalToWorld));
+	return BoundingBox.TransformBy(InLocalToWorld);
 }
 
 FTransform USplineMeshComponent::GetSocketTransform(FName InSocketName, ERelativeTransformSpace TransformSpace) const
@@ -972,9 +1032,8 @@ FTransform USplineMeshComponent::GetSocketTransform(FName InSocketName, ERelativ
 		UStaticMeshSocket const* const Socket = GetSocketByName(InSocketName);
 		if (Socket)
 		{
-			FTransform SocketTransform;
-			SocketTransform = FTransform(Socket->RelativeRotation, Socket->RelativeLocation * GetAxisMask(ForwardAxis), Socket->RelativeScale);
-			SocketTransform = SocketTransform * CalcSliceTransform(GetAxisValueRef(Socket->RelativeLocation, ForwardAxis));
+			FTransform SocketTransform(Socket->RelativeRotation, Socket->RelativeLocation * GetAxisMask(ForwardAxis), Socket->RelativeScale);
+			SocketTransform = SocketTransform * CalcSliceTransform(UE::SplineMesh::RealToFloatChecked(GetAxisValueRef(Socket->RelativeLocation, ForwardAxis)));
 
 			switch (TransformSpace)
 			{
@@ -986,7 +1045,7 @@ FTransform USplineMeshComponent::GetSocketTransform(FName InSocketName, ERelativ
 			{
 				if (const AActor* Actor = GetOwner())
 				{
-					return (SocketTransform * GetComponentToWorld()).GetRelativeTransform(GetOwner()->GetTransform());
+					return (SocketTransform * GetComponentToWorld()).GetRelativeTransform(Actor->GetTransform());
 				}
 				break;
 			}
@@ -1004,36 +1063,42 @@ FTransform USplineMeshComponent::GetSocketTransform(FName InSocketName, ERelativ
 
 FTransform USplineMeshComponent::CalcSliceTransform(const float DistanceAlong) const
 {
-	const bool bHasCustomBoundary = !FMath::IsNearlyEqual(SplineBoundaryMin, SplineBoundaryMax);
+	const float Alpha = ComputeRatioAlongSpline(DistanceAlong);
 
-	// Find how far 'along' mesh we are
-	float Alpha = 0.f;
-	if (bHasCustomBoundary)
-	{
-		Alpha = (DistanceAlong - SplineBoundaryMin) / (SplineBoundaryMax - SplineBoundaryMin);
-	}
-	else if (GetStaticMesh())
-	{
-		const FBoxSphereBounds StaticMeshBounds = GetStaticMesh()->GetBounds();
-		const float MeshMinZ = GetAxisValueRef(StaticMeshBounds.Origin, ForwardAxis) - GetAxisValueRef(StaticMeshBounds.BoxExtent, ForwardAxis);
-		const float MeshRangeZ = 2.0f * GetAxisValueRef(StaticMeshBounds.BoxExtent, ForwardAxis);
-		if (MeshRangeZ > UE_SMALL_NUMBER)
-		{
-			Alpha = (DistanceAlong - MeshMinZ) / MeshRangeZ;
-		}
-	}
+	float MinT = 0.f;
+	float MaxT = 1.f;
+	ComputeVisualMeshSplineTRange(MinT, MaxT);
 
-	return CalcSliceTransformAtSplineOffset(Alpha);
+	return CalcSliceTransformAtSplineOffset(Alpha, MinT, MaxT);
 }
 
-FTransform USplineMeshComponent::CalcSliceTransformAtSplineOffset(const float Alpha) const
+FTransform USplineMeshComponent::CalcSliceTransformAtSplineOffset(const float Alpha, const float MinT, const float MaxT) const
 {
 	// Apply hermite interp to Alpha if desired
 	const float HermiteAlpha = bSmoothInterpRollScale ? SmoothStep(0.0, 1.0, Alpha) : Alpha;
 
 	// Then find the point and direction of the spline at this point along
-	FVector3f SplinePos = SplineEvalPos(SplineParams, Alpha);
-	const FVector3f SplineDir = SplineEvalDir(SplineParams, Alpha);
+	FVector3f SplinePos;
+	FVector3f SplineDir;
+
+	// Use linear extrapolation
+	if (Alpha < MinT)
+	{
+		const FVector3f StartTangent(SplineEvalTangent(SplineParams, MinT));
+		SplinePos = FVector3f(SplineParams.StartPos) + (StartTangent * (Alpha - MinT));
+		SplineDir = StartTangent.GetSafeNormal();
+	}
+	else if (Alpha > MaxT)
+	{
+		const FVector3f EndTangent(SplineEvalTangent(SplineParams, MaxT));
+		SplinePos = FVector3f(SplineParams.EndPos) + (EndTangent * (Alpha - MaxT));
+		SplineDir = EndTangent.GetSafeNormal();
+	}
+	else
+	{
+		SplinePos = SplineEvalPos(SplineParams, Alpha);
+		SplineDir = SplineEvalDir(SplineParams, Alpha);
+	}
 
 	// Find base frenet frame
 	const FVector3f BaseXVec = (FVector3f(SplineUpDir) ^ SplineDir).GetSafeNormal();
@@ -1123,10 +1188,7 @@ bool USplineMeshComponent::GetTriMeshSizeEstimates(struct FTriMeshCollisionDataE
 
 FBox USplineMeshComponent::GetNavigationBounds() const
 {
-	// @todo Revisit how to collect navigation bounds, but for now behave like UPrimitiveComponent::GetNavigationBounds(), same as before the recent addition of UStaticMeshComponent::GetNavigationBounds().
-	
-	// Return invalid box when retrieving NavigationBounds before they are being computed at component registration
-	return bRegistered ? Bounds.GetBox() : FBox(ForceInit);
+	return CachedNavigationBounds;
 }
 
 void USplineMeshComponent::GetMeshId(FString& OutMeshId)
@@ -1223,7 +1285,7 @@ bool USplineMeshComponent::DoCustomNavigableGeometryExport(FNavigableGeometryExp
 				for (int32 i = 0; i < NavCollision->GetConvexCollision().VertexBuffer.Num(); ++i)
 				{
 					FVector Vertex = NavCollision->GetConvexCollision().VertexBuffer[i];
-					Vertex = CalcSliceTransform(GetAxisValueRef(Vertex, ForwardAxis)).TransformPosition(Vertex * Mask);
+					Vertex = CalcSliceTransform(UE::SplineMesh::RealToFloatChecked(GetAxisValueRef(Vertex, ForwardAxis))).TransformPosition(Vertex * Mask);
 					VertexBuffer.Add(Vertex);
 				}
 				GeomExport.ExportCustomMesh(VertexBuffer.GetData(), VertexBuffer.Num(),
@@ -1234,7 +1296,7 @@ bool USplineMeshComponent::DoCustomNavigableGeometryExport(FNavigableGeometryExp
 				for (int32 i = 0; i < NavCollision->GetTriMeshCollision().VertexBuffer.Num(); ++i)
 				{
 					FVector Vertex = NavCollision->GetTriMeshCollision().VertexBuffer[i];
-					Vertex = CalcSliceTransform(GetAxisValueRef(Vertex, ForwardAxis)).TransformPosition(Vertex * Mask);
+					Vertex = CalcSliceTransform(UE::SplineMesh::RealToFloatChecked(GetAxisValueRef(Vertex, ForwardAxis))).TransformPosition(Vertex * Mask);
 					VertexBuffer.Add(Vertex);
 				}
 				GeomExport.ExportCustomMesh(VertexBuffer.GetData(), VertexBuffer.Num(),
@@ -1296,7 +1358,7 @@ void USplineMeshComponent::RecreateCollision()
 			// distortion of a sphere can't be done nicely, so we just transform the origin and size
 			for (FKSphereElem& SphereElem : BodySetup->AggGeom.SphereElems)
 			{
-				const float Z = GetAxisValueRef(SphereElem.Center, ForwardAxis);
+				const float Z = UE::SplineMesh::RealToFloatChecked(GetAxisValueRef(SphereElem.Center, ForwardAxis));
 				FTransform SliceTransform = CalcSliceTransform(Z);
 				SphereElem.Center *= Mask;
 
@@ -1307,12 +1369,12 @@ void USplineMeshComponent::RecreateCollision()
 			// distortion of a sphyl can't be done nicely, so we just transform the origin and size
 			for (FKSphylElem& SphylElem : BodySetup->AggGeom.SphylElems)
 			{
-				const float Z = GetAxisValueRef(SphylElem.Center, ForwardAxis);
+				const float Z = UE::SplineMesh::RealToFloatChecked(GetAxisValueRef(SphylElem.Center, ForwardAxis));
 				FTransform SliceTransform = CalcSliceTransform(Z);
 				SphylElem.Center *= Mask;
 
 				FTransform TM = SphylElem.GetTransform();
-				SphylElem.Length = (TM * SliceTransform).TransformVector(FVector(0, 0, SphylElem.Length)).Size();
+				SphylElem.Length = UE::SplineMesh::RealToFloatChecked((TM * SliceTransform).TransformVector(FVector(0, 0, SphylElem.Length)).Size());
 				SphylElem.Radius *= SliceTransform.GetMaximumAxisScale();
 
 				SphylElem.SetTransform(TM * SliceTransform);
@@ -1348,7 +1410,7 @@ void USplineMeshComponent::RecreateCollision()
 					// pretransform the point by its local transform so we are working in untransformed local space
 					FVector TransformedPoint = TM.TransformPosition(Point);
 					// apply the transform to spline space
-					Point = CalcSliceTransform(GetAxisValueRef(TransformedPoint, ForwardAxis)).TransformPosition(TransformedPoint * Mask);
+					Point = CalcSliceTransform(UE::SplineMesh::RealToFloatChecked(GetAxisValueRef(TransformedPoint, ForwardAxis))).TransformPosition(TransformedPoint * Mask);
 				}
 
 				// Set the local transform as an identity as points have already been transformed
@@ -1435,8 +1497,8 @@ float USplineMeshComponent::GetTextureStreamingTransformScale() const
 	if (GetStaticMesh())
 	{
 		// We do this by looking at the ratio between current bounds (including deformation) and undeformed (straight from staticmesh)
-		const float MinExtent = 1.0f;
-		FBoxSphereBounds UndeformedBounds = GetStaticMesh()->GetBounds().TransformBy(GetComponentTransform());
+		constexpr float MinExtent = 1.0f;
+		const FBoxSphereBounds UndeformedBounds = GetStaticMesh()->GetBounds().TransformBy(GetComponentTransform());
 		if (UndeformedBounds.BoxExtent.X >= MinExtent)
 		{
 			SplineDeformFactor = FMath::Max(SplineDeformFactor, Bounds.BoxExtent.X / UndeformedBounds.BoxExtent.X);
@@ -1451,14 +1513,14 @@ float USplineMeshComponent::GetTextureStreamingTransformScale() const
 		}
 	}
 
-	return SplineDeformFactor * Super::GetTextureStreamingTransformScale();
+	return UE::SplineMesh::RealToFloatChecked(SplineDeformFactor) * Super::GetTextureStreamingTransformScale();
 }
 
 #if WITH_EDITOR
 void USplineMeshComponent::PostEditChangeProperty(FPropertyChangedEvent& PropertyChangedEvent)
 {
-	FProperty* MemberPropertyThatChanged = PropertyChangedEvent.MemberProperty;
-	bool bIsSplineParamsChange = MemberPropertyThatChanged && MemberPropertyThatChanged->GetNameCPP() == TEXT("SplineParams");
+	const FProperty* MemberPropertyThatChanged = PropertyChangedEvent.MemberProperty;
+	const bool bIsSplineParamsChange = MemberPropertyThatChanged && MemberPropertyThatChanged->GetNameCPP() == TEXT("SplineParams");
 	if (bIsSplineParamsChange)
 	{
 		SetEndTangent(SplineParams.EndTangent, false);
