NetDriver will attempt to update unmapped references to stably-named subobjects when net.OptimizedRemapping is true. Previously it was possible that some of these references would never map, due to their NetGUIDs never being added to the ImportedGuids list.
Can be disabled by setting net.RemapStableSubobjects to 0.
Added an EngineTest for this case.

diff --git a/Engine/Source/Runtime/Engine/Private/DataChannel.cpp b/Engine/Source/Runtime/Engine/Private/DataChannel.cpp
--- a/Engine/Source/Runtime/Engine/Private/DataChannel.cpp
+++ b/Engine/Source/Runtime/Engine/Private/DataChannel.cpp
@@ -111,6 +111,7 @@ namespace UE::Net
 {
 	extern int32 FilterGuidRemapping;
 	extern bool bDiscardTornOffActorRPCs;
+	extern bool bRemapStableSubobjects;
 
 	static float QueuedBunchTimeoutSeconds = 30.0f;
 	static FAutoConsoleVariableRef CVarQueuedBunchTimeoutSeconds(
@@ -2383,6 +2384,20 @@ void UActorChannel::DestroyActorAndComponents()
 			MoveMappedObjectToUnmapped(Component);
 		}
 
+		// Also unmap any stably-named subobjects we didn't create
+		if (UE::Net::bRemapStableSubobjects)
+		{
+			TArray<UObject*> Inners;
+			GetObjectsWithOuter(Actor, Inners);
+			for (UObject* Inner : Inners)
+			{
+				if (Inner->IsNameStableForNetworking())
+				{
+					MoveMappedObjectToUnmapped(Inner);
+				}
+			}
+		}
+
 		// Unmap this object so we can remap it if it becomes relevant again in the future
 		MoveMappedObjectToUnmapped(Actor);
 
diff --git a/Engine/Source/Runtime/Engine/Private/NetDriver.cpp b/Engine/Source/Runtime/Engine/Private/NetDriver.cpp
--- a/Engine/Source/Runtime/Engine/Private/NetDriver.cpp
+++ b/Engine/Source/Runtime/Engine/Private/NetDriver.cpp
@@ -311,6 +311,21 @@ namespace UE::Net::Private
 		}
 	}
 
+	bool IsGuidInOuterChain(const FNetGUIDCache& GuidCache, const FNetGuidCacheObject* CacheObj, FNetworkGUID GuidMatch)
+	{
+		while (CacheObj && CacheObj->OuterGUID.IsValid())
+		{
+			if (CacheObj->OuterGUID == GuidMatch)
+			{
+				return true;
+			}
+
+			CacheObj = GuidCache.GetCacheObject(CacheObj->OuterGUID);
+		}
+
+		return false;
+	}
+
 } //namespace UE::Net::Private
 
 namespace UE::Net
@@ -904,6 +919,9 @@ namespace UE
 	{
 		int32 FilterGuidRemapping = 1;
 		static FAutoConsoleVariableRef CVarFilterGuidRemapping(TEXT("net.FilterGuidRemapping"), FilterGuidRemapping, TEXT("Remove destroyed and parent guids from unmapped list"));
+
+		bool bRemapStableSubobjects = true;
+		static FAutoConsoleVariableRef CVarNetRemapStableSubobjects(TEXT("net.RemapStableSubobjects"), bRemapStableSubobjects, TEXT("If enabled, attempts to remap stable subobjects when net.OptimizedRemapping is also enabled."));
 	};
 };
 
@@ -1054,6 +1072,36 @@ void UNetDriver::TickFlush(float DeltaSeconds)
 
 					if (GuidCache->GetObjectFromNetGUID(NetworkGuid, false) != nullptr)
 					{
+						if (UE::Net::bRemapStableSubobjects)
+						{
+							QUICK_SCOPE_CYCLE_COUNTER(STAT_NetRemapStableSubobjects);
+
+							// Stably-named/net addressable subobjects are created by user code on clients and not by replication.
+							// Normally this optimized remapping path relies on adding NetGuids of spawned actors/objects to the
+							// GuidCache->ImportedNetGuids list, but for net addressable subobjects there's no hook to do this.
+							// In order to actually remap them so references will be valid, this code will find any potentially
+							// unmapped addressable subobjects of an owning object that is in the import list (like a replicated actor).
+							// It looks through the unmapped replicators' reference list for any guids that have the imported
+							// guid in their outer chain.
+							TSet<FNetworkGUID>& PendingGuidsRef = PendingOuterNetGuidsRef.FindOrAdd(NetworkGuid);
+							
+							for (FObjectReplicator* Replicator : UnmappedReplicators)
+							{
+								for (const FNetworkGUID ReferencedGuid : Replicator->ReferencedGuids)
+								{
+									const FNetGuidCacheObject* const CacheObjectPtr = GuidCache->GetCacheObject(ReferencedGuid);
+
+									if (CacheObjectPtr && !CacheObjectPtr->PathName.IsNone())
+									{
+										if (UE::Net::Private::IsGuidInOuterChain(*GuidCache, CacheObjectPtr, NetworkGuid))
+										{
+											PendingGuidsRef.Add(ReferencedGuid);
+										}
+									}
+								}
+							}
+						}
+
 						NewlyMappedGuids.Add(NetworkGuid);
 						bMappedOrBroken = true;
 					}
