[ZenLoader] Fix missing include

#rnx
#prefligth none

diff --git a/Engine/Source/Runtime/Engine/Private/Tests/Loading/AsyncLoadingTests_Shared.cpp b/Engine/Source/Runtime/Engine/Private/Tests/Loading/AsyncLoadingTests_Shared.cpp
--- a/Engine/Source/Runtime/Engine/Private/Tests/Loading/AsyncLoadingTests_Shared.cpp
+++ b/Engine/Source/Runtime/Engine/Private/Tests/Loading/AsyncLoadingTests_Shared.cpp
@@ -3,6 +3,7 @@
 #include "AsyncLoadingTests_Shared.h"
 #include "Misc/AutomationTest.h"
 #include "Misc/PackageName.h"
+#include "UObject/MetaData.h"
 #include "UObject/SavePackage.h"
 
 UAsyncLoadingTests_Shared::FOnPostLoadDelegate UAsyncLoadingTests_Shared::OnPostLoad;
