[NetPhysics] Character basing while jumping

Adds functionality to allow a character to stay based while jumping on a base, as long as the base is under the character.

diff --git a/Engine/Source/Runtime/Engine/Classes/GameFramework/CharacterMovementComponent.h b/Engine/Source/Runtime/Engine/Classes/GameFramework/CharacterMovementComponent.h
--- a/Engine/Source/Runtime/Engine/Classes/GameFramework/CharacterMovementComponent.h
+++ b/Engine/Source/Runtime/Engine/Classes/GameFramework/CharacterMovementComponent.h
@@ -1373,6 +1373,17 @@ public:
 	/** Update OldBaseLocation and OldBaseQuat if there is a valid movement base, and store the relative location/rotation if necessary. Ignores bDeferUpdateBasedMovement and forces the update. */
 	ENGINE_API virtual void SaveBaseLocation();
 
+	/** Apply inherited velocity when leaving base, for example from jumping off it */
+	ENGINE_API virtual void ApplyImpartedMovementBaseVelocity();
+	
+	/** Property used to set if characters should stay based on objects while jumping */
+	UPROPERTY(Category = "Character Movement: Jumping / Falling", EditAnywhere, BlueprintReadWrite)
+	bool bStayBasedInAir = false;
+	
+	/** Property used to set how high above base characters should stay based on objects while jumping if bStayBasedInAir is set */
+	UPROPERTY(Category = "Character Movement: Jumping / Falling", EditAnywhere, BlueprintReadWrite, meta = (editcondition = "bStayBasedInAir"))
+	float StayBasedInAirHeight = 1000.0f;
+
 	/** changes physics based on MovementMode */
 	ENGINE_API virtual void StartNewPhysics(float deltaTime, int32 Iterations);
 	
diff --git a/Engine/Source/Runtime/Engine/Private/Components/CharacterMovementComponent.cpp b/Engine/Source/Runtime/Engine/Private/Components/CharacterMovementComponent.cpp
--- a/Engine/Source/Runtime/Engine/Private/Components/CharacterMovementComponent.cpp
+++ b/Engine/Source/Runtime/Engine/Private/Components/CharacterMovementComponent.cpp
@@ -1288,21 +1288,17 @@ void UCharacterMovementComponent::OnMovementModeChanged(EMovementMode PreviousMo
 
 		if (MovementMode == MOVE_Falling)
 		{
-			DecayingFormerBaseVelocity = GetImpartedMovementBaseVelocity();
-			Velocity += DecayingFormerBaseVelocity;
-			if (bMovementInProgress && CurrentRootMotion.HasAdditiveVelocity())
+			if (bStayBasedInAir == false)
 			{
-				// If we leave a base during movement and we have additive root motion, we need to add the imparted velocity so that it retains it next tick
-				CurrentRootMotion.LastPreAdditiveVelocity += DecayingFormerBaseVelocity;
-			}
-			if (!CharacterMovementCVars::bAddFormerBaseVelocityToRootMotionOverrideWhenFalling || FormerBaseVelocityDecayHalfLife == 0.f)
-			{
-				DecayingFormerBaseVelocity = FVector::ZeroVector;
+				ApplyImpartedMovementBaseVelocity();
 			}
 			CharacterOwner->Falling();
 		}
 
-		SetBase(NULL);
+		if (!IsFalling() || bStayBasedInAir == false)
+		{
+			SetBase(NULL);
+		}
 
 		if (MovementMode == MOVE_None)
 		{
@@ -2305,6 +2301,23 @@ void UCharacterMovementComponent::UpdateBasedMovement(float DeltaSeconds)
 		return;
 	}
 
+	// Check if falling above current base
+	if (IsFalling() && bStayBasedInAir)
+	{
+		const FVector PawnLocation = UpdatedComponent->GetComponentLocation();
+		FFindFloorResult OutFloorResult;
+		ComputeFloorDist(PawnLocation, StayBasedInAirHeight, StayBasedInAirHeight, OutFloorResult, CharacterOwner->GetCapsuleComponent()->GetScaledCapsuleRadius(), NULL);
+
+		UPrimitiveComponent* HitComponent = OutFloorResult.HitResult.Component.Get();
+		if (HitComponent && HitComponent->GetAttachmentRoot() != MovementBase->GetAttachmentRoot())
+		{
+			// New or no base under the character
+			ApplyImpartedMovementBaseVelocity();
+			SetBase(NULL);
+			return;
+		}
+	}
+
 	// Ignore collision with bases during these movements.
 	TGuardValue<EMoveComponentFlags> ScopedFlagRestore(MoveComponentFlags, MoveComponentFlags | MOVECOMP_IgnoreBases);
 
@@ -2438,6 +2451,21 @@ void UCharacterMovementComponent::UpdateBasedRotation(FRotator& FinalRotation, c
 	}
 }
 
+void UCharacterMovementComponent::ApplyImpartedMovementBaseVelocity()
+{
+	DecayingFormerBaseVelocity = GetImpartedMovementBaseVelocity();
+	Velocity += DecayingFormerBaseVelocity;
+	if (bMovementInProgress && CurrentRootMotion.HasAdditiveVelocity())
+	{
+		// If we leave a base during movement and we have additive root motion, we need to add the imparted velocity so that it retains it next tick
+		CurrentRootMotion.LastPreAdditiveVelocity += DecayingFormerBaseVelocity;
+	}
+	if (!CharacterMovementCVars::bAddFormerBaseVelocityToRootMotionOverrideWhenFalling || FormerBaseVelocityDecayHalfLife == 0.f)
+	{
+		DecayingFormerBaseVelocity = FVector::ZeroVector;
+	}
+}
+
 void UCharacterMovementComponent::DisableMovement()
 {
 	if (CharacterOwner)
