Fixed ShouldCache call being bypassed when getting shaders, which could lead to missing shader types on cooked build while ODSC would instead let them build if requested.
Now ShouldCache should be respected in all cases : if it returns false, the shader will be missing both with and without ODSC.

#tests editor

diff --git a/Engine/Source/Runtime/Engine/Private/Materials/MaterialShared.cpp b/Engine/Source/Runtime/Engine/Private/Materials/MaterialShared.cpp
--- a/Engine/Source/Runtime/Engine/Private/Materials/MaterialShared.cpp
+++ b/Engine/Source/Runtime/Engine/Private/Materials/MaterialShared.cpp
@@ -3598,18 +3598,21 @@ bool FMaterial::TryGetShaders(const FMaterialShaderTypes& InTypes, const FVertex
 			{
 				bMissingShader = true;
 
+#if WITH_EDITOR || WITH_ODSC
+				for (const FShaderType* ShaderType : InTypes.PipelineType->GetStages())
+				{
+					if (!ShouldCacheShaderType(ShaderType, InVertexFactoryType))
+					{
+						return false;
+					}
+				}
+#endif // WITH_EDITOR || WITH_ODSC
+
 #if WITH_ODSC
 				if (FPlatformProperties::RequiresCookedData())
 				{
 					if (GODSCManager->IsHandlingRequests())
 					{
-				        for (const FShaderType* ShaderType : InTypes.PipelineType->GetStages())
-				        {
-					        if (!ShouldCacheShaderType(ShaderType, InVertexFactoryType))
-					        {
-						        return false;
-					        }
-				        }
 						const FString MaterialName = GetFullPath();
 						const FString VFTypeName(InVertexFactoryType ? InVertexFactoryType->GetName() : TEXT(""));
 						const FString PipelineName(InTypes.PipelineType->GetName());
@@ -3690,15 +3693,18 @@ bool FMaterial::TryGetShaders(const FMaterialShaderTypes& InTypes, const FVertex
 				{
 					bMissingShader = true;
 
+#if WITH_EDITOR || WITH_ODSC
+					if (!ShouldCacheShaderType(ShaderType, InVertexFactoryType))
+					{
+						return false;
+					}
+#endif // WITH_EDITOR || WITH_ODSC
+
 #if WITH_ODSC
 					if (FPlatformProperties::RequiresCookedData())
 					{
 						if (GODSCManager->IsHandlingRequests())
 						{
-					        if (!ShouldCacheShaderType(ShaderType, InVertexFactoryType))
-					        {
-						        return false;
-					        }
 							const FString MaterialName = GetFullPath();
 							const FString VFTypeName(InVertexFactoryType ? InVertexFactoryType->GetName() : TEXT(""));
 							const FString PipelineName;
@@ -5373,4 +5379,4 @@ FMaterialShaderParameters::FMaterialShaderParameters(const FMaterial* InMaterial
 	bAlwaysEvaluateWorldPositionOffset = InMaterial->ShouldAlwaysEvaluateWorldPositionOffset();
 }
 
-#undef LOCTEXT_NAMESPACE
+#undef LOCTEXT_NAMESPACE
\ No newline at end of file
