Fixes a bug where UMaterialInstanceDynamic::HasPixelAnimation() always returned false.

#tests win64

diff --git a/Engine/Source/Runtime/Engine/Classes/Materials/MaterialInstanceDynamic.h b/Engine/Source/Runtime/Engine/Classes/Materials/MaterialInstanceDynamic.h
--- a/Engine/Source/Runtime/Engine/Classes/Materials/MaterialInstanceDynamic.h
+++ b/Engine/Source/Runtime/Engine/Classes/Materials/MaterialInstanceDynamic.h
@@ -198,6 +198,7 @@ class UMaterialInstanceDynamic : public UMaterialInstance
 	ENGINE_API virtual bool IsMasked() const override;
 	ENGINE_API virtual FDisplacementScaling GetDisplacementScaling() const override;
 	ENGINE_API virtual float GetMaxWorldPositionOffsetDisplacement() const override;
+	ENGINE_API virtual bool HasPixelAnimation() const override;
 
 	/**
 	 * In order to remap to the correct texture streaming data, we must keep track of each texture renamed.
diff --git a/Engine/Source/Runtime/Engine/Private/Materials/MaterialInstance.cpp b/Engine/Source/Runtime/Engine/Private/Materials/MaterialInstance.cpp
--- a/Engine/Source/Runtime/Engine/Private/Materials/MaterialInstance.cpp
+++ b/Engine/Source/Runtime/Engine/Private/Materials/MaterialInstance.cpp
@@ -1721,6 +1721,11 @@ float UMaterialInstanceDynamic::GetMaxWorldPositionOffsetDisplacement() const
 	return Parent ? Parent->GetMaxWorldPositionOffsetDisplacement() : 0.0f;
 }
 
+bool UMaterialInstanceDynamic::HasPixelAnimation() const
+{
+	return Parent ? Parent->HasPixelAnimation() : false;
+}
+
 FMaterialShadingModelField UMaterialInstanceDynamic::GetShadingModels() const
 {
 	return Parent ? Parent->GetShadingModels() : MSM_DefaultLit;
