Add separate cvars controlling validation of each detail level for scene components as legacy content for some projects may only be valid for certain detail levels and not cause problems at cook time because those projects only enable unused components for certain detail levels.

diff --git a/Engine/Source/Runtime/Engine/Private/Components/SceneComponent.cpp b/Engine/Source/Runtime/Engine/Private/Components/SceneComponent.cpp
--- a/Engine/Source/Runtime/Engine/Private/Components/SceneComponent.cpp
+++ b/Engine/Source/Runtime/Engine/Private/Components/SceneComponent.cpp
@@ -33,6 +33,7 @@
 #include "UObject/Class.h"
 #include "UObject/ICookInfo.h"
 #include "Misc/DataValidation.h"
+#include "Misc/EnumRange.h"
 
 #define LOCTEXT_NAMESPACE "SceneComponent"
 
@@ -474,10 +475,21 @@ static bool CheckDescendantsAreAlsoCulledForTarget(USceneComponent const* SceneC
 
 	for (USceneComponent* ChildSceneComponent : AttachedChildren)
 	{
-		if (ChildSceneComponent && SceneComponentNeedsLoadForTarget(ChildSceneComponent, TargetPlatform))
+		if (!ChildSceneComponent)
 		{
+			continue;
+		}
+
+		UE_LOG(LogSceneComponent, Display, TEXT("Checking attached component %s for culling"), *GetPathNameSafe(ChildSceneComponent));
+		if (SceneComponentNeedsLoadForTarget(ChildSceneComponent, TargetPlatform))
+		{
+			UE_LOG(LogSceneComponent, Display, TEXT("Scene component %s will not be culled"), *GetPathNameSafe(ChildSceneComponent));
 			return false;
 		}
+		else
+		{
+			UE_LOG(LogSceneComponent, Display, TEXT("Scene component %s will eot be culled"), *GetPathNameSafe(ChildSceneComponent));
+		}
 	}
 
 	return true;
@@ -487,6 +499,7 @@ bool USceneComponent::NeedsLoadForTargetPlatform(const ITargetPlatform* TargetPl
 {
 	if(!SceneComponentNeedsLoadForTarget(this, TargetPlatform))
 	{
+		UE_LOG(LogSceneComponent, Display, TEXT("Scene component %s will be culled for target, checking children"), *GetPathName());
 		// Also check whether any of our children are culled.
 		bool bDescendantsCulled = CheckDescendantsAreAlsoCulledForTarget(this, TargetPlatform);
 
@@ -510,11 +523,25 @@ static FAutoConsoleVariableRef CVarValidateSceneComponentAttachmentEditorOnlySet
 	ECVF_Default
 );
 
-bool GValidateSceneComponentAttachmentDetailLevel = true;
-static FAutoConsoleVariableRef CVarValidateSceneComponentAttachmentDetailLevel (
-	TEXT("p.ValidateSceneComponentAttachmentDetailLevel"),
-	GValidateSceneComponentAttachmentDetailLevel,
-	TEXT("If enabled, checks that components don't have attached components of higher detail level during data validation/"),
+bool GValidateSceneComponentAttachmentDetailLevel_Low = true;
+static FAutoConsoleVariableRef CVarValidateSceneComponentAttachmentDetailLevel_Low (
+	TEXT("p.ValidateSceneComponentAttachmentDetailLevel_Low"),
+	GValidateSceneComponentAttachmentDetailLevel_Low,
+	TEXT("If enabled, checks that cooking for a target detail level of Low and removing unneeded components will not remove the parents of any components."),
+	ECVF_Default
+);
+bool GValidateSceneComponentAttachmentDetailLevel_Medium = true;
+static FAutoConsoleVariableRef CVarValidateSceneComponentAttachmentDetailLevel_Medium (
+	TEXT("p.ValidateSceneComponentAttachmentDetailLevel_Medium"),
+	GValidateSceneComponentAttachmentDetailLevel_Medium,
+	TEXT("If enabled, checks that cooking for a target detail level of Medium and removing unneeded components will not remove the parents of any components."),
+	ECVF_Default
+);
+bool GValidateSceneComponentAttachmentDetailLevel_High = true;
+static FAutoConsoleVariableRef CVarValidateSceneComponentAttachmentDetailLevel_High (
+	TEXT("p.ValidateSceneComponentAttachmentDetailLevel_High"),
+	GValidateSceneComponentAttachmentDetailLevel_High,
+	TEXT("If enabled, checks that cooking for a target detail level of High and removing unneeded components will not remove the parents of any components."),
 	ECVF_Default
 );
 
@@ -536,7 +563,9 @@ EDataValidationResult USceneComponent::IsDataValid(FDataValidationContext& Conte
 			}	
 		}
 	}
-	if (GValidateSceneComponentAttachmentDetailLevel)
+
+	if (   DetailMode != EDetailMode::DM_Low 
+		&& (GValidateSceneComponentAttachmentDetailLevel_Low || GValidateSceneComponentAttachmentDetailLevel_Medium || GValidateSceneComponentAttachmentDetailLevel_High))
 	{
 		for (USceneComponent* ChildSceneComponent : GetAttachChildren())
 		{
@@ -544,9 +573,15 @@ EDataValidationResult USceneComponent::IsDataValid(FDataValidationContext& Conte
 			{
 				continue;
 			}
-
-			if (ChildSceneComponent->DetailMode < DetailMode)
+			
+			bool bBrokenAtLow = DetailMode > EDetailMode::DM_Low && ChildSceneComponent->DetailMode <= EDetailMode::DM_Low;
+			bool bBrokenAtMedium = DetailMode > EDetailMode::DM_Medium && ChildSceneComponent->DetailMode <= EDetailMode::DM_Medium;
+			bool bBrokenAtHigh = DetailMode > EDetailMode::DM_High && ChildSceneComponent->DetailMode <= EDetailMode::DM_High;
+			if((GValidateSceneComponentAttachmentDetailLevel_Low && bBrokenAtLow)
+			|| (GValidateSceneComponentAttachmentDetailLevel_Medium && bBrokenAtMedium)
+			|| (GValidateSceneComponentAttachmentDetailLevel_High && bBrokenAtHigh))
 			{
+				// This child is culled at this detail mode, even though we aren't
 				Context.AddError(FText::Format(LOCTEXT("SceneComponent_AttachmentDetailLevelMismatch",
 					"Component {0} of detail level {1} cannot be removed because it has an attached child {2} of detail level {3}"),
 					FText::FromString(GetPathName()),
@@ -558,7 +593,6 @@ EDataValidationResult USceneComponent::IsDataValid(FDataValidationContext& Conte
 			}	
 		}
 	}
-	
 	return CombineDataValidationResults(Result, Super::IsDataValid(Context));
 }
 
