Fix ray tracing geometry initializers with incorrect OwnerName in FStaticMeshStreamIn.
- LODResourcesView is a subset of RenderData->LODResources based on AssetLODBias so the same bias needs to be applied when generating the owner name.

diff --git a/Engine/Source/Runtime/Engine/Private/Streaming/StaticMeshUpdate.cpp b/Engine/Source/Runtime/Engine/Private/Streaming/StaticMeshUpdate.cpp
--- a/Engine/Source/Runtime/Engine/Private/Streaming/StaticMeshUpdate.cpp
+++ b/Engine/Source/Runtime/Engine/Private/Streaming/StaticMeshUpdate.cpp
@@ -223,7 +223,7 @@ void FStaticMeshStreamIn::CreateBuffers_Internal(const FContext& Context)
 			if (IsRayTracingEnabled() && Context.Mesh->bSupportRayTracing &&
 				LODResource.VertexBuffers.StaticMeshVertexBuffer.GetNumVertices() > 0)
 			{
-				const FName OwnerName = UStaticMesh::GetLODPathName(Context.Mesh, LODIdx);
+				const FName OwnerName = UStaticMesh::GetLODPathName(Context.Mesh, LODIdx + Context.Mesh->GetStreamableResourceState().AssetLODBias);
 
 				FRayTracingGeometryInitializer Initializer;
 				if (Context.Mesh->HasValidNaniteData() && Nanite::GetSupportsRayTracingProceduralPrimitive(GMaxRHIShaderPlatform))
@@ -309,7 +309,7 @@ void FStaticMeshStreamIn::DoFinishUpdate(const FContext& Context)
 				// Skip LODs that have their render data stripped
 				if (LODResource.VertexBuffers.StaticMeshVertexBuffer.GetNumVertices() > 0)
 				{
-					const FName OwnerName = UStaticMesh::GetLODPathName(Context.Mesh, LODIndex);
+					const FName OwnerName = UStaticMesh::GetLODPathName(Context.Mesh, LODIndex + Context.Mesh->GetStreamableResourceState().AssetLODBias);
 
 					// Rebuild the initializer because it could have been reset during a previous release
 					FRayTracingGeometryInitializer Initializer;
