Small changes to FStaticMeshLODResources
- make SetupRayTracingGeometryInitializer(...) function const.
- Use UStaticMesh::GetLODPathName(...) instead of duplicate logic.

diff --git a/Engine/Source/Runtime/Engine/Private/StaticMesh.cpp b/Engine/Source/Runtime/Engine/Private/StaticMesh.cpp
--- a/Engine/Source/Runtime/Engine/Private/StaticMesh.cpp
+++ b/Engine/Source/Runtime/Engine/Private/StaticMesh.cpp
@@ -1323,7 +1323,7 @@ void FStaticMeshLODResources::UpdateVertexMemoryStats() const
 
 void FStaticMeshLODResources::InitResources(UStaticMesh* Parent, int32 LODIndex)
 {
-	const FName OwnerName(FString::Printf(TEXT("%s [LOD%d]"), Parent ? *Parent->GetPathName() : TEXT("UnknownStaticMesh"), LODIndex));
+	const FName OwnerName = UStaticMesh::GetLODPathName(Parent, LODIndex);
 
 	if (bBuffersInlined)
 	{
@@ -1440,7 +1440,7 @@ void FStaticMeshLODResources::InitResources(UStaticMesh* Parent, int32 LODIndex)
 }
 
 #if RHI_RAYTRACING
-void FStaticMeshLODResources::SetupRayTracingGeometryInitializer(FRayTracingGeometryInitializer& Initializer, const FName& DebugName, const FName& OwnerName)
+void FStaticMeshLODResources::SetupRayTracingGeometryInitializer(FRayTracingGeometryInitializer& Initializer, const FName& DebugName, const FName& OwnerName) const
 {
 	Initializer.DebugName = DebugName;
 	Initializer.OwnerName = OwnerName;
diff --git a/Engine/Source/Runtime/Engine/Public/StaticMeshResources.h b/Engine/Source/Runtime/Engine/Public/StaticMeshResources.h
--- a/Engine/Source/Runtime/Engine/Public/StaticMeshResources.h
+++ b/Engine/Source/Runtime/Engine/Public/StaticMeshResources.h
@@ -515,7 +515,7 @@ public:
 	void GetResourceSizeEx(FResourceSizeEx& CumulativeResourceSize) const;
 
 #if RHI_RAYTRACING
-	void SetupRayTracingGeometryInitializer(FRayTracingGeometryInitializer& Initializer, const FName& DebugName, const FName& OwnerName);
+	void SetupRayTracingGeometryInitializer(FRayTracingGeometryInitializer& Initializer, const FName& DebugName, const FName& OwnerName) const;
 	static void SetupRayTracingProceduralGeometryInitializer(FRayTracingGeometryInitializer& Initializer, const FName& DebugName, const FName& OwnerName);
 #endif // RHI_RAYTRACING
 
