Oodle WAVE: store Raw wave editor bulk data in a better way to improve compression performance.
#p4v-preflight-copy 28810120

diff --git a/Engine/Config/BaseEditor.ini b/Engine/Config/BaseEditor.ini
--- a/Engine/Config/BaseEditor.ini
+++ b/Engine/Config/BaseEditor.ini
@@ -254,6 +254,9 @@ bExpandClassPickerClassList=False
 [Cooker.CleanStaticMeshMtrlSkip]
 Class=UnrealEd.ActorFactory
 
+[AudioImporter]
+EnableOodleWAV=False
+
 [TextureImporter]
 AllowNonPowerOfTwoTextures=True
 ; Imported jpegs will be stored in packages in jpeg format until edited, this will help keep texture asset sizes down
diff --git a/Engine/Source/Runtime/Engine/Classes/Sound/SoundWave.h b/Engine/Source/Runtime/Engine/Classes/Sound/SoundWave.h
--- a/Engine/Source/Runtime/Engine/Classes/Sound/SoundWave.h
+++ b/Engine/Source/Runtime/Engine/Classes/Sound/SoundWave.h
@@ -879,7 +879,136 @@ public:
 
 #if WITH_EDITORONLY_DATA
 	/** Uncompressed wav data 16 bit in mono or stereo - stereo not allowed for multichannel data */
-	UE::Serialization::FEditorBulkData RawData;
+	struct FEditorAudioBulkData
+	{
+		UE::Serialization::FEditorBulkData RawData;
+
+		ENGINE_API void CreateFromBulkData(FBulkData& InBulkData, const FGuid& InGuid, UObject* Owner);
+		ENGINE_API void Serialize(FArchive& Ar, UObject* Owner, bool bAllowRegister=true);
+		ENGINE_API TFuture<FSharedBuffer> GetPayload() const;
+		ENGINE_API bool HasPayloadData() const;
+		ENGINE_API void UpdatePayload(FSharedBuffer InPayload, UObject* Owner = nullptr);
+
+		//
+		// Deprecated unused API forwarding for potential backwards compatability issues.
+		//
+
+		UE_DEPRECATED(5.4, "CreateLegacyUniqueIdentifier is provided just for API backwards compatibility.")
+		void CreateLegacyUniqueIdentifier(UObject* Owner)
+		{
+			RawData.CreateLegacyUniqueIdentifier(Owner);
+		}
+		UE_DEPRECATED(5.4, "Reset is provided just for API backwards compatibility.")
+		void Reset() 
+		{ 
+			RawData.Reset(); 
+		}
+		UE_DEPRECATED(5.4, "UnloadData is provided just for API backwards compatibility.")
+		void UnloadData()
+		{
+			RawData.UnloadData();
+		}
+		UE_DEPRECATED(5.4, "DetachFromDisk is provided just for API backwards compatibility.")
+		void DetachFromDisk(FArchive* Ar, bool bEnsurePayloadIsLoaded)
+		{
+			RawData.DetachFromDisk(Ar, bEnsurePayloadIsLoaded);
+		}
+		UE_DEPRECATED(5.4, "GetIdentifier is provided just for API backwards compatibility.")
+		FGuid GetIdentifier() const
+		{
+			return RawData.GetIdentifier();
+		}
+		UE_DEPRECATED(5.4, "GetPayloadId is provided just for API backwards compatibility.")
+		const FIoHash& GetPayloadId() const
+		{
+			return RawData.GetPayloadId();
+		}
+		UE_DEPRECATED(5.4, "GetPayloadSize is provided just for API backwards compatibility.")
+		int64 GetPayloadSize() const
+		{
+			return RawData.GetPayloadSize();
+		}
+		UE_DEPRECATED(5.4, "DoesPayloadNeedLoading is provided just for API backwards compatibility.")
+		bool DoesPayloadNeedLoading() const
+		{
+			return RawData.DoesPayloadNeedLoading();
+		}
+		UE_DEPRECATED(5.4, "GetCompressedPayload is provided just for API backwards compatibility.")
+		TFuture<FCompressedBuffer> GetCompressedPayload() const
+		{
+			return RawData.GetCompressedPayload();
+		}
+		UE_DEPRECATED(5.4, "UpdatePayload is provided just for API backwards compatibility.")
+		void UpdatePayload(FCompressedBuffer InPayload, UObject* Owner = nullptr)
+		{
+			RawData.UpdatePayload(InPayload, Owner);
+		}
+		UE_DEPRECATED(5.4, "UpdatePayload is provided just for API backwards compatibility.")
+		void UpdatePayload(UE::Serialization::FEditorBulkData::FSharedBufferWithID InPayload, UObject* Owner = nullptr)
+		{
+			RawData.UpdatePayload(MoveTemp(InPayload), Owner);
+		}
+		UE_DEPRECATED(5.4, "SetCompressionOptions is provided just for API backwards compatibility.")
+		void SetCompressionOptions(UE::Serialization::ECompressionOptions Option)
+		{
+			RawData.SetCompressionOptions(Option);
+		}
+		UE_DEPRECATED(5.4, "SetCompressionOptions is provided just for API backwards compatibility.")
+		void SetCompressionOptions(ECompressedBufferCompressor Compressor, ECompressedBufferCompressionLevel CompressionLevel)
+		{
+			RawData.SetCompressionOptions(Compressor, CompressionLevel);
+		}
+		UE_DEPRECATED(5.4, "GetBulkDataVersions is provided just for API backwards compatibility.")
+		void GetBulkDataVersions(FArchive& InlineArchive, FPackageFileVersion& OutUEVersion, int32& OutLicenseeUEVersion, FCustomVersionContainer& OutCustomVersions) const
+		{
+			RawData.GetBulkDataVersions(InlineArchive, OutUEVersion, OutLicenseeUEVersion, OutCustomVersions);
+		}
+		UE_DEPRECATED(5.4, "TearOff is provided just for API backwards compatibility.")
+		void TearOff()
+		{
+			RawData.TearOff();
+		}
+		UE_DEPRECATED(5.4, "CopyTornOff is provided just for API backwards compatibility.")
+		UE::Serialization::FEditorBulkData CopyTornOff() const
+		{
+			return RawData.CopyTornOff();
+		}
+		UE_DEPRECATED(5.4, "SerializeForRegistry is provided just for API backwards compatibility.")
+		void SerializeForRegistry(FArchive& Ar)
+		{
+			RawData.SerializeForRegistry(Ar);
+		}
+		UE_DEPRECATED(5.4, "CanSaveForRegistry is provided just for API backwards compatibility.")
+		bool CanSaveForRegistry() const
+		{
+			return RawData.CanSaveForRegistry();
+		}
+		UE_DEPRECATED(5.4, "HasPlaceholderPayloadId is provided just for API backwards compatibility.")
+		bool HasPlaceholderPayloadId() const
+		{
+			return RawData.HasPlaceholderPayloadId();
+		}
+		UE_DEPRECATED(5.4, "IsMemoryOnlyPayload is provided just for API backwards compatibility.")
+		bool IsMemoryOnlyPayload() const
+		{
+			return RawData.IsMemoryOnlyPayload();
+		}
+		UE_DEPRECATED(5.4, "UpdatePayloadId is provided just for API backwards compatibility.")
+		void UpdatePayloadId()
+		{
+			return RawData.UpdatePayloadId();
+		}
+		UE_DEPRECATED(5.4, "LocationMatches is provided just for API backwards compatibility.")
+		bool LocationMatches(const UE::Serialization::FEditorBulkData& Other) const
+		{
+			return RawData.LocationMatches(Other);
+		}
+		UE_DEPRECATED(5.4, "UpdateRegistrationOwner is provided just for API backwards compatibility.")
+		void UpdateRegistrationOwner(UObject* Owner)
+		{
+			RawData.UpdateRegistrationOwner(Owner);
+		}
+	} RawData;
 
 	/** Waveform edits to be applied to this SoundWave on cook (editing transformations will trigger a cook) */
 	UPROPERTY(EditAnywhere, Instanced, Category = "Waveform Processing")
diff --git a/Engine/Source/Runtime/Engine/Engine.Build.cs b/Engine/Source/Runtime/Engine/Engine.Build.cs
--- a/Engine/Source/Runtime/Engine/Engine.Build.cs
+++ b/Engine/Source/Runtime/Engine/Engine.Build.cs
@@ -172,7 +172,8 @@ public class Engine : ModuleRules
 			"UEOgg",
 			"Vorbis",
 			"VorbisFile",
-			"libOpus"
+			"libOpus",
+			"OodleWAV"
 			);
 
 		DynamicallyLoadedModuleNames.Add("EyeTracker");
diff --git a/Engine/Source/Runtime/Engine/Private/Audio.cpp b/Engine/Source/Runtime/Engine/Private/Audio.cpp
--- a/Engine/Source/Runtime/Engine/Private/Audio.cpp
+++ b/Engine/Source/Runtime/Engine/Private/Audio.cpp
@@ -1918,7 +1918,8 @@ bool FWaveModInfo::IsFormatSupported() const
 	return (*pFormatTag == WAVE_INFO_FORMAT_PCM
 		|| *pFormatTag == WAVE_INFO_FORMAT_ADPCM
 		|| *pFormatTag == WAVE_INFO_FORMAT_DVI_ADPCM
-		|| *pFormatTag == WAVE_INFO_FORMAT_IEEE_FLOAT);
+		|| *pFormatTag == WAVE_INFO_FORMAT_IEEE_FLOAT
+		|| *pFormatTag == WAVE_INFO_FORMAT_OODLE_WAVE);
 }
 
 bool FWaveModInfo::IsFormatUncompressed() const
diff --git a/Engine/Source/Runtime/Engine/Private/AudioDerivedData.cpp b/Engine/Source/Runtime/Engine/Private/AudioDerivedData.cpp
--- a/Engine/Source/Runtime/Engine/Private/AudioDerivedData.cpp
+++ b/Engine/Source/Runtime/Engine/Private/AudioDerivedData.cpp
@@ -1607,7 +1607,7 @@ struct FAudioCookInputs
 	// copy-on-write mechanism that doesn't require us to make
 	// a copy in memory to get immutability.
 	FCriticalSection& BulkDataCriticalSection;
-	UE::Serialization::FEditorBulkData& BulkData;
+	USoundWave::FEditorAudioBulkData& BulkData;
 #endif
 
 	FAudioCookInputs(USoundWave* InSoundWave, FName InBaseFormat, FName InHashFormat, const FPlatformAudioCookOverrides* InCookOverrides, const ITargetPlatform* InTargetPlatform)
diff --git a/Engine/Source/Runtime/Engine/Private/SoundWave.cpp b/Engine/Source/Runtime/Engine/Private/SoundWave.cpp
--- a/Engine/Source/Runtime/Engine/Private/SoundWave.cpp
+++ b/Engine/Source/Runtime/Engine/Private/SoundWave.cpp
@@ -18,6 +18,7 @@
 #include "Interfaces/ITargetPlatform.h"
 #include "Interfaces/ITargetPlatformManagerModule.h"
 #include "Misc/App.h"
+#include "Misc/ConfigCacheIni.h"
 #include "SubtitleManager.h"
 #include "DerivedDataCacheInterface.h"
 #include "EditorFramework/AssetImportData.h"
@@ -32,6 +33,9 @@
 #include "Templates/UnrealTemplate.h"
 #include "UObject/ObjectSaveContext.h"
 #include "ISoundWaveCloudStreaming.h"
+#if WITH_EDITORONLY_DATA
+#include "oowav.h"
+#endif
 
 static int32 SoundWaveDefaultLoadingBehaviorCVar = static_cast<int32>(ESoundWaveLoadingBehavior::LoadOnDemand);
 FAutoConsoleVariableRef CVarSoundWaveDefaultLoadingBehavior(
@@ -4729,3 +4733,79 @@ bool USoundWave::HasError() const
 	return SoundWaveDataPtr->HasError();
 }
 
+#if WITH_EDITORONLY_DATA
+TFuture<FSharedBuffer> USoundWave::FEditorAudioBulkData::GetPayload() const
+{
+	TFuture<FSharedBuffer> BufferFuture = RawData.GetPayload();
+	const uint8* Data = (const uint8*) BufferFuture.Get().GetData(); // Will block.
+	uint64 DataSize = BufferFuture.Get().GetSize();
+	FSharedBuffer Buffer = FSharedBuffer::Clone(Data, DataSize);
+	Data = (const uint8*) Buffer.GetData();
+	DataSize = Buffer.GetSize();
+	FWaveModInfo WaveInfo;
+	if (!WaveInfo.ReadWaveInfo(Data, DataSize, 0)) 
+	{
+		UE_LOG(LogAudio, Warning, TEXT("Failed to read wave data."));
+	}
+	if (*WaveInfo.pFormatTag == FWaveModInfo::WAVE_INFO_FORMAT_OODLE_WAVE)
+	{
+		// Convert OodleWave data back to a WAV file
+		*WaveInfo.pFormatTag = FWaveModInfo::WAVE_INFO_FORMAT_PCM;
+		int16* samples = (int16*)WaveInfo.SampleDataStart;
+		int64 num_samples = WaveInfo.GetNumSamples();
+		int64 num_channels = *WaveInfo.pChannels;
+		TArray<int16> scratch_buffer;
+		scratch_buffer.AddUninitialized(num_samples);
+		oowav_decode16(samples, scratch_buffer.GetData(), num_samples, num_channels);
+	}
+	TPromise<FSharedBuffer> promise;
+	promise.EmplaceValue(Buffer);
+	return promise.GetFuture();
+}
+
+void USoundWave::FEditorAudioBulkData::UpdatePayload(FSharedBuffer InPayload, UObject* Owner)
+{
+	const uint8* Data = (const uint8*) InPayload.GetData(); // Will block.
+	uint64 DataSize = InPayload.GetSize();
+	FSharedBuffer Buffer = FSharedBuffer::Clone(Data, DataSize);
+	Data = (const uint8*) Buffer.GetData();
+	DataSize = Buffer.GetSize();
+	FWaveModInfo WaveInfo;
+	if (!WaveInfo.ReadWaveInfo(Data, DataSize, 0)) 
+	{
+		UE_LOG(LogAudio, Warning, TEXT("Failed to read wave data."));
+	}
+	if (*WaveInfo.pFormatTag == FWaveModInfo::WAVE_INFO_FORMAT_PCM)
+	{
+		bool bEnableOodleWAV = false;
+		GConfig->GetBool(TEXT("AudioImporter"), TEXT("EnableOodleWAV"), bEnableOodleWAV, GEditorIni);
+		if (bEnableOodleWAV)
+		{
+			int16* samples = (int16*)WaveInfo.SampleDataStart;
+			int64 num_samples = WaveInfo.GetNumSamples();
+			int64 num_channels = *WaveInfo.pChannels;
+			TArray<int16> scratch_buffer;
+			scratch_buffer.AddUninitialized(num_samples);
+			oowav_encode16(samples, scratch_buffer.GetData(), num_samples, num_channels);
+			*WaveInfo.pFormatTag = FWaveModInfo::WAVE_INFO_FORMAT_OODLE_WAVE;
+		}
+	}
+	RawData.UpdatePayload(Buffer, Owner);
+}
+
+bool USoundWave::FEditorAudioBulkData::HasPayloadData() const
+{
+	return RawData.HasPayloadData();
+}
+
+void USoundWave::FEditorAudioBulkData::Serialize(FArchive& Ar, UObject* Owner, bool bAllowRegister)
+{
+	RawData.Serialize(Ar, Owner, bAllowRegister);
+}
+
+void USoundWave::FEditorAudioBulkData::CreateFromBulkData(FBulkData& InBulkData, const FGuid& InGuid, UObject* Owner)
+{
+	RawData.CreateFromBulkData(InBulkData, InGuid, Owner);
+}
+#endif
+
diff --git a/Engine/Source/Runtime/Engine/Public/Audio.h b/Engine/Source/Runtime/Engine/Public/Audio.h
--- a/Engine/Source/Runtime/Engine/Public/Audio.h
+++ b/Engine/Source/Runtime/Engine/Public/Audio.h
@@ -889,6 +889,7 @@ public:
 	static constexpr uint16 WAVE_INFO_FORMAT_ADPCM = 0x0002;
 	static constexpr uint16 WAVE_INFO_FORMAT_IEEE_FLOAT = 0x0003;
 	static constexpr uint16 WAVE_INFO_FORMAT_DVI_ADPCM = 0x0011;
+	static constexpr uint16 WAVE_INFO_FORMAT_OODLE_WAVE = 0xFFFF;
 
 	// Pointers to variables in the in-memory WAVE file.
 	const uint32* pSamplesPerSec;
@@ -896,7 +897,7 @@ public:
 	const uint16* pBlockAlign;
 	const uint16* pBitsPerSample;
 	const uint16* pChannels;
-	const uint16* pFormatTag;
+	uint16* pFormatTag;
 
 	const uint32* pWaveDataSize;
 	const uint32* pMasterSize;
diff --git a/Engine/Source/Runtime/OodleWAV/OodleWAV.Build.cs b/Engine/Source/Runtime/OodleWAV/OodleWAV.Build.cs
new file mode 100644
--- /dev/null
+++ b/Engine/Source/Runtime/OodleWAV/OodleWAV.Build.cs
@@ -0,0 +1,38 @@
+// Copyright Epic Games, Inc. All Rights Reserved.
+using UnrealBuildTool;
+using System.IO;
+
+public class OodleWAV : ModuleRules
+{
+	public OodleWAV(ReadOnlyTargetRules Target) : base(Target)
+	{
+		Type = ModuleType.External;
+
+		bool bWithLib = false;
+		if (Target.Platform == UnrealTargetPlatform.Win64 && Target.WindowsPlatform.Architecture != UnrealArch.Arm64)
+		{
+			string LibPath = Path.Combine(ModuleDirectory, "Lib");
+			PublicAdditionalLibraries.Add(Path.Combine(LibPath, "oowav_w64.lib"));
+			bWithLib = true;
+		}
+		else if (Target.Platform == UnrealTargetPlatform.Mac)
+		{
+			string LibPath = Path.Combine(ModuleDirectory, "Lib");
+			PublicAdditionalLibraries.Add(Path.Combine(LibPath, "oowav_mac.a"));
+			bWithLib = true;
+		}
+		else if (Target.IsInPlatformGroup(UnrealPlatformGroup.Unix))
+		{
+			string LibPath = Path.Combine(ModuleDirectory, "Lib/Unix", Target.Architecture.LinuxName);
+			PublicAdditionalLibraries.Add(Path.Combine(LibPath, "oowav_linux.a"));
+			bWithLib = true;
+		}
+
+		if (bWithLib)
+		{
+			PublicSystemIncludePaths.Add(Path.Combine(ModuleDirectory, "Source"));
+		}
+
+		PublicDefinitions.Add("WITH_OOWAV=" + (bWithLib ? '1' : '0'));
+	}
+}
diff --git a/Engine/Source/Runtime/OodleWAV/Source/oowav.h b/Engine/Source/Runtime/OodleWAV/Source/oowav.h
new file mode 100644
--- /dev/null
+++ b/Engine/Source/Runtime/OodleWAV/Source/oowav.h
@@ -0,0 +1,21 @@
+// Copyright Epic Games, Inc. All Rights Reserved.
+/*
+    Oodle WAVE is a simple data transform on PCM audio data that is modifies the data to be more compressible by Oodle. 
+*/
+#pragma once
+
+#ifdef __cplusplus
+extern "C" {
+#endif
+
+// Encodes/decodes 16-bit samples. samples is the input/output buffer, scratch_buffer is a temporary buffer.  
+// operates in place. samples is both input and output.
+// scratch_buffer should be the same size as samples.
+// num_samples is the number of samples in the buffer
+// num_channels is the number of channels in the buffer.
+void oowav_encode16(short *samples, short *scratch_buffer, long long num_samples, long long num_channels);
+void oowav_decode16(short *samples, short *scratch_buffer, long long num_samples, long long num_channels);
+
+#ifdef __cplusplus
+}
+#endif
\ No newline at end of file
