Add GetInstanceForMontage to UAnimInstance which returns the Anim Montage Instance even for montages that are currently blending out.

UAnimInstance already provides GetActiveInstanceForMontage but this returns null as soon as the montage is blending out, which doesn't suit our needs.

diff --git a/Engine/Source/Runtime/Engine/Classes/Animation/AnimInstance.h b/Engine/Source/Runtime/Engine/Classes/Animation/AnimInstance.h
--- a/Engine/Source/Runtime/Engine/Classes/Animation/AnimInstance.h
+++ b/Engine/Source/Runtime/Engine/Classes/Animation/AnimInstance.h
@@ -729,6 +729,9 @@ public:
 	/** Get Active FAnimMontageInstance for given Montage asset. Will return NULL if Montage is not currently Active. */
 	ENGINE_API FAnimMontageInstance* GetActiveInstanceForMontage(const UAnimMontage* Montage) const;
 
+	/** Get FAnimMontageInstance for given Montage asset (even when blending out). Will return NULL if Montage is not currently playing. */
+	ENGINE_API FAnimMontageInstance* GetInstanceForMontage(const UAnimMontage* Montage) const;
+
 	/** Get the FAnimMontageInstance currently running that matches this ID.  Will return NULL if no instance is found. */
 	ENGINE_API FAnimMontageInstance* GetMontageInstanceForID(int32 MontageInstanceID);
 
diff --git a/Engine/Source/Runtime/Engine/Private/Animation/AnimInstance.cpp b/Engine/Source/Runtime/Engine/Private/Animation/AnimInstance.cpp
--- a/Engine/Source/Runtime/Engine/Private/Animation/AnimInstance.cpp
+++ b/Engine/Source/Runtime/Engine/Private/Animation/AnimInstance.cpp
@@ -3557,6 +3557,19 @@ FAnimMontageInstance* UAnimInstance::GetActiveInstanceForMontage(const UAnimMont
 	return FoundInstancePtr ? *FoundInstancePtr : nullptr;
 }
 
+FAnimMontageInstance* UAnimInstance::GetInstanceForMontage(const UAnimMontage* Montage) const
+{
+	for (FAnimMontageInstance* MontageInstance : MontageInstances)
+	{
+		if (MontageInstance && MontageInstance->Montage == Montage)
+		{
+			return MontageInstance;
+		}
+	}
+
+	return nullptr;
+}
+
 FAnimMontageInstance* UAnimInstance::GetMontageInstanceForID(int32 MontageInstanceID)
 {
 	for (FAnimMontageInstance* MontageInstance : MontageInstances)
