Chaos - support per-particle InitialOverlapDepenetrationVelocity
- remove the now-superfluous on/off flag

This connects FBodyInstance::MaxDepenetrationVelocity (previously unused in Chaos) to InitialOverlapDepenetrationVelocity on rigid particles.

diff --git a/Engine/Source/Editor/DetailCustomizations/Private/BodyInstanceCustomization.cpp b/Engine/Source/Editor/DetailCustomizations/Private/BodyInstanceCustomization.cpp
--- a/Engine/Source/Editor/DetailCustomizations/Private/BodyInstanceCustomization.cpp
+++ b/Engine/Source/Editor/DetailCustomizations/Private/BodyInstanceCustomization.cpp
@@ -1356,7 +1356,6 @@ void FBodyInstanceCustomizationHelper::CustomizeDetails( IDetailLayoutBuilder& D
 		EnablePhysicsProperty(GET_MEMBER_NAME_CHECKED(FBodyInstance, AngularDamping), bDisplayAngularDamping);
 		EnablePhysicsProperty(GET_MEMBER_NAME_CHECKED(FBodyInstance, bEnableGravity), bDisplayEnableGravity);
 		EnablePhysicsProperty(GET_MEMBER_NAME_CHECKED(FBodyInstance, bInertiaConditioning), bDisplayInertiaConditioning);
-		EnablePhysicsProperty(GET_MEMBER_NAME_CHECKED(FBodyInstance, bInitialOverlapDepenetration), bDisplayInitialOverlapDepenetration);
 		EnablePhysicsProperty(GET_MEMBER_NAME_CHECKED(FBodyInstance, WalkableSlopeOverride), bDisplayWalkableSlopeOverride);
 
 		AddBodyConstraint(PhysicsCategory, BodyInstanceHandler);
@@ -1386,7 +1385,6 @@ void FBodyInstanceCustomizationHelper::CustomizeDetails( IDetailLayoutBuilder& D
 		BodyInstanceHandler->GetChildHandle(GET_MEMBER_NAME_CHECKED(FBodyInstance, bLockZRotation))->MarkHiddenByCustomization();
 		BodyInstanceHandler->GetChildHandle(GET_MEMBER_NAME_CHECKED(FBodyInstance, PositionSolverIterationCount))->MarkHiddenByCustomization();
 		BodyInstanceHandler->GetChildHandle(GET_MEMBER_NAME_CHECKED(FBodyInstance, VelocitySolverIterationCount))->MarkHiddenByCustomization();
-		BodyInstanceHandler->GetChildHandle(GET_MEMBER_NAME_CHECKED(FBodyInstance, MaxDepenetrationVelocity))->MarkHiddenByCustomization();
 		BodyInstanceHandler->GetChildHandle(GET_MEMBER_NAME_CHECKED(FBodyInstance, CustomDOFPlaneNormal))->MarkHiddenByCustomization();
 
 		//Add the rest
diff --git a/Engine/Source/Runtime/Engine/Classes/PhysicsEngine/BodyInstance.h b/Engine/Source/Runtime/Engine/Classes/PhysicsEngine/BodyInstance.h
--- a/Engine/Source/Runtime/Engine/Classes/PhysicsEngine/BodyInstance.h
+++ b/Engine/Source/Runtime/Engine/Classes/PhysicsEngine/BodyInstance.h
@@ -403,7 +403,7 @@ public:
 
 protected:
 
-	/** [PhysX Only] Whether this body instance has its own custom MaxDepenetrationVelocity*/
+	/** Whether this body instance has its own custom MaxDepenetrationVelocity*/
 	UPROPERTY(EditAnywhere, Category = Physics, meta=(InlineEditConditionToggle))
 	uint8 bOverrideMaxDepenetrationVelocity : 1;
 
@@ -440,15 +440,6 @@ protected:
 	UPROPERTY(EditAnywhere, AdvancedDisplay, BlueprintReadWrite, Category = Physics)
 	uint8 bInertiaConditioning : 1;
 
-	/**
-	 * @brief Enable the initial-overlap de-penetration mechanism
-	 * 
-	 * When enabled, initial overlaps are de-penetrated at a speed that is limited by the physics project configuration.
-	 * If disabled, initial overlaps are not tracked and bodies will attempt to de-penetrate in one tick which can lead to large separating velocities.
-	 */
-	UPROPERTY(EditAnywhere, AdvancedDisplay, BlueprintReadWrite, Category = Physics)
-	uint8 bInitialOverlapDepenetration : 1;
-
 
 public:
 	/** Current scale of physics - used to know when and how physics must be rescaled to match current transform of OwnerComponent. */
@@ -485,7 +476,12 @@ private:
 	struct FCollisionResponse CollisionResponses;
 
 protected:
-	/** [PhysX Only] The maximum velocity used to depenetrate this object*/
+	/** 
+	 * The maximum velocity used to depenetrate this object from others when spawned with initial overlaps or teleports (does not affect overlaps as a result of normal movement).
+	 * A value of zero will allow objects that are spawned overlapping to go to sleep as they are rather than pop out of each other.
+	 * A negative value is equivalent to bOverrideMaxDepenetrationVelocity = false, meaning use the project setting.
+	 * When set, this overrides the CollisionInitialOverlapDepenetrationVelocity project setting (and not the MaxDepenetrationVelocity solver setting that will be deprecated)
+	*/
 	UPROPERTY(EditAnywhere, AdvancedDisplay, BlueprintReadOnly, Category = Physics, meta = (editcondition = "bOverrideMaxDepenetrationVelocity", ClampMin = "0.0", UIMin = "0.0"))
 	float MaxDepenetrationVelocity;
 
@@ -821,12 +817,6 @@ public:
 	/** If set to true, this body will treat bodies that do not have the flag set as having infinite mass */
 	ENGINE_API void SetOneWayInteraction(bool InOneWayInteraction = true);
 
-	/** Is the controlled inital-overlap depenetration system enabled for this body */
-	ENGINE_API bool IsInitialOverlapDepenetrationEnabled() const;
-
-	/** Enable/disable the controlled inital-overlap depenetration system for this body */
-	ENGINE_API void SetInitialOverlapDepenetrationEnabled(bool bInEnabled);
-
 	/** Add a torque to this body */
 	ENGINE_API void AddTorqueInRadians(const FVector& Torque, bool bAllowSubstepping = true, bool bAccelChange = false, const FAsyncPhysicsTimestamp TimeStamp = FAsyncPhysicsTimestamp(), APlayerController* PlayerController = nullptr);
 	/** Clear accumulated torques on this body */
@@ -856,8 +846,25 @@ public:
 	/** Get the maximum angular velocity of this body */
 	ENGINE_API float GetMaxAngularVelocityInRadians() const;
 
-	/** Set the maximum depenetration velocity the physics simulation will introduce */
+	/** Are we overriding the MaxDepenetrationVelocity. See SetMaxDepenetrationVelocity */
+	ENGINE_API bool GetOverrideMaxDepenetrationVelocity() const { return bOverrideMaxDepenetrationVelocity; }
+
+	/** Enable/Disable override of MaxDepenetrationVelocity */
+	ENGINE_API void SetOverrideMaxDepenetrationVelocity(bool bInEnabled);
+
+	/**
+	 * Set the maximum velocity used to depenetrate this object from others when spawned with initial overlaps or teleports (does not affect overlaps as a result of normal movement).
+	 * A value of zero will allow objects that are spawned overlapping to go to sleep as they are rather than pop out of each other.
+	 * Note: implicitly calls SetOverrideMaxDepenetrationVelocity(true)
+	 * Note: MaxDepenetration overrides the CollisionInitialOverlapDepenetrationVelocity project setting (and not the MaxDepenetrationVelocity solver setting that will be deprecated)
+	*/
 	ENGINE_API void SetMaxDepenetrationVelocity(float MaxVelocity);
+
+private:
+	// Send the MaxDepenetrationVelocity setting to the solver
+	ENGINE_API void UpdateMaxDepenetrationVelocity();
+
+public:
 	/** Set whether we should get a notification about physics collisions */
 	ENGINE_API void SetInstanceNotifyRBCollision(bool bNewNotifyCollision);
 	/** Enables/disables whether this body is affected by gravity. */
diff --git a/Engine/Source/Runtime/Engine/Private/PhysicsEngine/BodyInstance.cpp b/Engine/Source/Runtime/Engine/Private/PhysicsEngine/BodyInstance.cpp
--- a/Engine/Source/Runtime/Engine/Private/PhysicsEngine/BodyInstance.cpp
+++ b/Engine/Source/Runtime/Engine/Private/PhysicsEngine/BodyInstance.cpp
@@ -336,7 +336,6 @@ FBodyInstance::FBodyInstance()
 	, bInterpolateWhenSubStepping(true)
 	, bPendingCollisionProfileSetup(false)
 	, bInertiaConditioning(true)
-	, bInitialOverlapDepenetration(true)
 	, Scale3D(1.0f)
 	, CollisionProfileName(UCollisionProfile::CustomCollisionProfileName)
 	, PositionSolverIterationCount(8)
@@ -1211,7 +1210,6 @@ void FInitBodiesHelperBase::CreateActor_AssumesLocked(FBodyInstance* Instance, c
 		FPhysicsInterface::SetMaxLinearVelocity_AssumesLocked(Instance->ActorHandle, TNumericLimits<float>::Max());
 		FPhysicsInterface::SetSmoothEdgeCollisionsEnabled_AssumesLocked(Instance->ActorHandle, Instance->bSmoothEdgeCollisions);
 		FPhysicsInterface::SetInertiaConditioningEnabled_AssumesLocked(Instance->ActorHandle, Instance->bInertiaConditioning);
-		FPhysicsInterface::SetInitialOverlapDepenetrationEnabled_AssumesLocked(Instance->ActorHandle, Instance->bInitialOverlapDepenetration);
 
 		// Set sleep event notification
 		FPhysicsInterface::SetSendsSleepNotifies_AssumesLocked(Instance->ActorHandle, Instance->bGenerateWakeEvents);
@@ -3280,17 +3278,34 @@ void FBodyInstance::SetMaxAngularVelocityInRadians(float NewMaxAngVel, bool bAdd
 	}
 }
 
+void FBodyInstance::SetOverrideMaxDepenetrationVelocity(bool bInEnabled)
+{
+	bOverrideMaxDepenetrationVelocity = bInEnabled;
+
+	UpdateMaxDepenetrationVelocity();
+}
+
 void FBodyInstance::SetMaxDepenetrationVelocity(float MaxVelocity)
 {
+	bOverrideMaxDepenetrationVelocity = true;
 	MaxDepenetrationVelocity = MaxVelocity;
 
-	FPhysicsCommand::ExecuteWrite(ActorHandle, [&](const FPhysicsActorHandle& Actor)
+	UpdateMaxDepenetrationVelocity();
+}
+
+void FBodyInstance::UpdateMaxDepenetrationVelocity()
+{
+	// Negative values mean do not use max depenetration velocity (equivalent to large number)
+	const float UsedMaxDepenetrationVelocity = (bOverrideMaxDepenetrationVelocity) ? MaxDepenetrationVelocity : -1.0f;
+
+	// NOTE: FBodyInstance::MaxDepentrationVelocity now means initial depenetration velocity, 
+	// and not the general solver depentration velocity limit (which will probably be removed)
+	FPhysicsCommand::ExecuteWrite(ActorHandle, [this, UsedMaxDepenetrationVelocity](const FPhysicsActorHandle& Actor)
 	{
-		FPhysicsInterface::SetMaxDepenetrationVelocity_AssumesLocked(Actor, MaxDepenetrationVelocity);
+		FPhysicsInterface::SetMaxDepenetrationVelocity_AssumesLocked(Actor, UsedMaxDepenetrationVelocity);
 	});
 }
 
-
 void FBodyInstance::AddCustomPhysics(FCalculateCustomPhysics& CalculateCustomPhysics)
 {
 	FPhysicsCommand::ExecuteWrite(ActorHandle, [&](const FPhysicsActorHandle& Actor)
@@ -3383,27 +3398,6 @@ void FBodyInstance::SetOneWayInteraction(bool InOneWayInteraction /*= true*/)
 		});
 }
 
-bool FBodyInstance::IsInitialOverlapDepenetrationEnabled() const
-{
-	return bInitialOverlapDepenetration;
-}
-
-void FBodyInstance::SetInitialOverlapDepenetrationEnabled(bool bInEnabled)
-{
-	if (bInEnabled != bInitialOverlapDepenetration)
-	{
-		bInitialOverlapDepenetration = bInEnabled;
-
-		FPhysicsCommand::ExecuteWrite(ActorHandle, [bInEnabled](const FPhysicsActorHandle& Actor)
-			{
-				if (FChaosEngineInterface::IsValid(Actor))
-				{
-					Actor->GetGameThreadAPI().SetInitialOverlapDepenetrationEnabled(bInEnabled);
-				}
-			});
-	}
-}
-
 void FBodyInstance::AddTorqueInRadians(const FVector& Torque, bool bAllowSubstepping, bool bAccelChange, const FAsyncPhysicsTimestamp TimeStamp, APlayerController* PlayerController)
 {
 	const bool bIsInternal = TimeStamp.IsValid();
@@ -4097,7 +4091,7 @@ void FBodyInstance::InitDynamicProperties_AssumesLocked()
 			UpdateMassProperties();
 			UpdateDampingProperties();
 			SetMaxAngularVelocityInRadians(GetMaxAngularVelocityInRadians(), false, false);
-			SetMaxDepenetrationVelocity(bOverrideMaxDepenetrationVelocity ? MaxDepenetrationVelocity : UPhysicsSettings::Get()->MaxDepenetrationVelocity);
+			UpdateMaxDepenetrationVelocity();
 		}
 		else
 		{
diff --git a/Engine/Source/Runtime/Experimental/Chaos/Private/Chaos/Collision/PBDCollisionConstraint.cpp b/Engine/Source/Runtime/Experimental/Chaos/Private/Chaos/Collision/PBDCollisionConstraint.cpp
--- a/Engine/Source/Runtime/Experimental/Chaos/Private/Chaos/Collision/PBDCollisionConstraint.cpp
+++ b/Engine/Source/Runtime/Experimental/Chaos/Private/Chaos/Collision/PBDCollisionConstraint.cpp
@@ -216,6 +216,7 @@ namespace Chaos
 		, SavedManifoldPoints()
 		, ManifoldPoints()
 		, MinInitialPhi(0)
+		, InitialOverlapDepenetrationVelocity(0)
 		, CCDTimeOfImpact(0)
 		, SolverBodies{ nullptr, nullptr }
 		, CCDEnablePenetration(0)
@@ -255,6 +256,7 @@ namespace Chaos
 		, SavedManifoldPoints()
 		, ManifoldPoints()
 		, MinInitialPhi(0)
+		, InitialOverlapDepenetrationVelocity(0)
 		, CCDTimeOfImpact(0)
 		, SolverBodies{ nullptr, nullptr }
 		, CCDEnablePenetration(0)
@@ -323,8 +325,12 @@ namespace Chaos
 		Flags.bUseManifold = bInUseManifold;
 		Flags.bUseIncrementalManifold = false;
 
-		// Should we use the initial overlap depenetration mechanism
-		Flags.bOverlapDepenetrationEnabled = FConstGenericParticleHandle(GetParticle0())->InitialOverlapDepentrationEnabled() && FConstGenericParticleHandle(GetParticle1())->InitialOverlapDepentrationEnabled();
+		// The initial overlap depenetration velocity is that max of the two body settings. A negative value means use the system default.
+		// If either body has a zero of positive value, it will be used instead of the system setting. If both have a value, the max will be used.
+		// If both bodies have negative values, the solver settings will be used (search for GetInitialOverlapDepentrationVelocity to see where)
+		const FRealSingle InitialOverlapDepenetrationVelocity0 = FConstGenericParticleHandle(GetParticle0())->InitialOverlapDepenetrationVelocity();
+		const FRealSingle InitialOverlapDepenetrationVelocity1 = FConstGenericParticleHandle(GetParticle1())->InitialOverlapDepenetrationVelocity();
+		InitialOverlapDepenetrationVelocity = FMath::Max(InitialOverlapDepenetrationVelocity0, InitialOverlapDepenetrationVelocity1);
 
 		// Is this a one-way interaction?
 		const bool bDynamic0 = FConstGenericParticleHandle(GetParticle0())->IsDynamic();
diff --git a/Engine/Source/Runtime/Experimental/Chaos/Private/Chaos/Collision/PBDCollisionContainerSolver.cpp b/Engine/Source/Runtime/Experimental/Chaos/Private/Chaos/Collision/PBDCollisionContainerSolver.cpp
--- a/Engine/Source/Runtime/Experimental/Chaos/Private/Chaos/Collision/PBDCollisionContainerSolver.cpp
+++ b/Engine/Source/Runtime/Experimental/Chaos/Private/Chaos/Collision/PBDCollisionContainerSolver.cpp
@@ -236,10 +236,13 @@ namespace Chaos
 		const FConstraintSolverBody& Body0 = Solver.SolverBody0();
 		const FConstraintSolverBody& Body1 = Solver.SolverBody1();
 
-		// Negative DepenetrationVelocity means depenetrate immediately
-		const bool bEnableDepenetrationVelocity = CVars::bChaos_Collision_EnableInitialDepenetration && Constraint->GetInitialOverlapDepentrationEnabled();
-		const FSolverReal MaxDepenetrationVelocity = bEnableDepenetrationVelocity ? SolverSettings.DepenetrationVelocity : FSolverReal(-1);
+		// MaxDepenetrationVelocity controls the rate at which initial-overlaps are resolved
+		// If the constraint has a non-negative MaxDepenetrationVelocity we use it, otherwise use the solver setting.
+		// If resultant MaxDepenetrationVelocity is negative, it means depenetrate immediately
+		const FSolverReal MaxDepenetrationVelocity = (Constraint->GetInitialOverlapDepentrationVelocity() >= 0) ? Constraint->GetInitialOverlapDepentrationVelocity() : SolverSettings.DepenetrationVelocity;
 
+		// The maximum correction we can apply in one frame
+		// @todo(chaos): consider removing this functionality?
 		const FSolverReal MaxPushOut = (SolverSettings.MaxPushOutVelocity > 0) ? (FSolverReal(SolverSettings.MaxPushOutVelocity) * Dt) : FSolverReal(0);
 
 		// Only calculate state for newly added contacts. Normally this is all of them, but maybe not if incremental collision is used by RBAN.
diff --git a/Engine/Source/Runtime/Experimental/Chaos/Private/PBDRigidsSolver.cpp b/Engine/Source/Runtime/Experimental/Chaos/Private/PBDRigidsSolver.cpp
--- a/Engine/Source/Runtime/Experimental/Chaos/Private/PBDRigidsSolver.cpp
+++ b/Engine/Source/Runtime/Experimental/Chaos/Private/PBDRigidsSolver.cpp
@@ -2610,6 +2610,7 @@ TRACE_COUNTER_SET(ChaosTraceCounter_##Name, Value)
 		Rigid->SetAngularEtherDrag(DynamicMisc.AngularEtherDrag());
 		Rigid->SetMaxLinearSpeedSq(DynamicMisc.MaxLinearSpeedSq());
 		Rigid->SetMaxAngularSpeedSq(DynamicMisc.MaxAngularSpeedSq());
+		Rigid->SetInitialOverlapDepenetrationVelocity(DynamicMisc.InitialOverlapDepenetrationVelocity());
 		Rigid->SetCollisionGroup(DynamicMisc.CollisionGroup());
 		Rigid->SetDisabled(DynamicMisc.Disabled());
 		Rigid->SetCollisionConstraintFlags(DynamicMisc.CollisionConstraintFlags());
diff --git a/Engine/Source/Runtime/Experimental/Chaos/Public/Chaos/Collision/PBDCollisionConstraint.h b/Engine/Source/Runtime/Experimental/Chaos/Public/Chaos/Collision/PBDCollisionConstraint.h
--- a/Engine/Source/Runtime/Experimental/Chaos/Public/Chaos/Collision/PBDCollisionConstraint.h
+++ b/Engine/Source/Runtime/Experimental/Chaos/Public/Chaos/Collision/PBDCollisionConstraint.h
@@ -460,8 +460,8 @@ namespace Chaos
 		// NOTE: This is initially set based on whether are allowing incremental manifolds
 		bool GetUseIncrementalCollisionDetection() const { return !Flags.bUseManifold || Flags.bUseIncrementalManifold; }
 
-		// Does this contact support the initial overlap depenetration. If not, bodies are depenetrated in one frame as opposed to at a fixed speed
-		bool GetInitialOverlapDepentrationEnabled() const { return Flags.bOverlapDepenetrationEnabled; }
+		// Initial overlap depenetration velocity from the maximum of the two bodies
+		FRealSingle GetInitialOverlapDepentrationVelocity() const { return InitialOverlapDepenetrationVelocity; }
 
 		/**
 		* Reset the material properties to those from the shape materials. Called each frame to reset contact modifications to the material.
@@ -889,7 +889,6 @@ namespace Chaos
 				uint16 bModifierApplied : 1;			// Was a constraint modifier applied this tick
 				uint16 bMaterialSet : 1;				// Has the material been set (or does it need to be reset)
 				uint16 bInitialContact : 1;				// Is this contact considered an initial contact
-				uint16 bOverlapDepenetrationEnabled : 1;// Does this contact support the initial overlap depenetration system
 				uint16 bIsOneWayInteraction : 1;		// Does one of the bodies have the one-way interaction bit set?
 			};
 			uint16 Bits;
@@ -959,6 +958,9 @@ namespace Chaos
 		// when we were already separated at some point in the past.
 		FRealSingle MinInitialPhi;
 
+		// If we have initial overlap, the depenetration speed
+		FRealSingle InitialOverlapDepenetrationVelocity;
+
 		// Value in range [0,1] used to interpolate P between [X,P] that we will rollback to when solving at time of impact.
 		FRealSingle CCDTimeOfImpact;
 
diff --git a/Engine/Source/Runtime/Experimental/Chaos/Public/Chaos/ParticleDirtyFlags.h b/Engine/Source/Runtime/Experimental/Chaos/Public/Chaos/ParticleDirtyFlags.h
--- a/Engine/Source/Runtime/Experimental/Chaos/Public/Chaos/ParticleDirtyFlags.h
+++ b/Engine/Source/Runtime/Experimental/Chaos/Public/Chaos/ParticleDirtyFlags.h
@@ -380,6 +380,13 @@ public:
 			Ar << MMaxAngularSpeedSq;
 		}
 
+		// @todo(chaos): add this
+		//Ar << MInitialOverlapDepenetrationVelocity
+		if (Ar.IsLoading())
+		{
+			MInitialOverlapDepenetrationVelocity = 0;
+		}
+
 		if (bAddControlFlags)
 		{
 			Ar << MControlFlags;
@@ -393,6 +400,7 @@ public:
 		SetAngularEtherDrag(Other.AngularEtherDrag());
 		SetMaxLinearSpeedSq(Other.MaxLinearSpeedSq());
 		SetMaxAngularSpeedSq(Other.MaxAngularSpeedSq());
+		SetInitialOverlapDepenetrationVelocity(Other.InitialOverlapDepenetrationVelocity());
 		SetObjectState(Other.ObjectState());
 		SetCollisionGroup(Other.CollisionGroup());
 		SetSleepType(Other.SleepType());
@@ -409,6 +417,7 @@ public:
 			&& AngularEtherDrag() == Other.AngularEtherDrag()
 			&& MaxLinearSpeedSq() == Other.MaxLinearSpeedSq()
 			&& MaxAngularSpeedSq() == Other.MaxAngularSpeedSq()
+			&& InitialOverlapDepenetrationVelocity() == Other.InitialOverlapDepenetrationVelocity()
 			&& CollisionGroup() == Other.CollisionGroup()
 			&& SleepType() == Other.SleepType()
 			&& CollisionConstraintFlags() == Other.CollisionConstraintFlags()
@@ -433,6 +442,9 @@ public:
 	FReal MaxAngularSpeedSq() const { return MMaxAngularSpeedSq; }
 	void SetMaxAngularSpeedSq(FReal InMaxAngularSpeed) { MMaxAngularSpeedSq = InMaxAngularSpeed; }
 
+	FRealSingle InitialOverlapDepenetrationVelocity() const { return MInitialOverlapDepenetrationVelocity; }
+	void SetInitialOverlapDepenetrationVelocity(FRealSingle InVel) { MInitialOverlapDepenetrationVelocity = InVel; }
+
 	EObjectStateType ObjectState() const { return MObjectState; }
 	void SetObjectState(EObjectStateType InState){ MObjectState = InState; }
 
@@ -458,9 +470,6 @@ public:
 	void SetCollisionConstraintFlags(uint32 InCollisionConstraintFlag) { MCollisionConstraintFlag = InCollisionConstraintFlag; }
 	void AddCollisionConstraintFlag(const ECollisionConstraintFlags Flag) { MCollisionConstraintFlag |= uint32(Flag); }
 	void RemoveCollisionConstraintFlag(const ECollisionConstraintFlags Flag) { MCollisionConstraintFlag &= ~uint32(Flag); }
-	
-	bool InitialOverlapDepenetrationEnabled() const { return MControlFlags.GetInitialOverlapDepenetrationEnabled(); }
-	void SetInitialOverlapDepenetrationEnabled(const bool bEnabled) { MControlFlags.SetInitialOverlapDepenetrationEnabled(bEnabled); }
 
 	bool OneWayInteraction() const { return MControlFlags.GetOneWayInteractionEnabled(); }
 	void SetOneWayInteraction(bool bInOneWayInteraction) { MControlFlags.SetOneWayInteractionEnabled(bInOneWayInteraction); }
@@ -479,6 +488,7 @@ private:
 	FReal MAngularEtherDrag;
 	FReal MMaxLinearSpeedSq;
 	FReal MMaxAngularSpeedSq;
+	FRealSingle MInitialOverlapDepenetrationVelocity = 0;
 	int32 MCollisionGroup;
 
 	EObjectStateType MObjectState;
diff --git a/Engine/Source/Runtime/Experimental/Chaos/Public/Chaos/ParticleHandle.h b/Engine/Source/Runtime/Experimental/Chaos/Public/Chaos/ParticleHandle.h
--- a/Engine/Source/Runtime/Experimental/Chaos/Public/Chaos/ParticleHandle.h
+++ b/Engine/Source/Runtime/Experimental/Chaos/Public/Chaos/ParticleHandle.h
@@ -104,6 +104,7 @@ void PBDRigidParticleDefaultConstruct(FConcrete& Concrete, const FPBDRigidPartic
 	Concrete.SetAngularImpulseVelocity(TVector<T, d>(0));
 	Concrete.SetMaxLinearSpeedSq(TNumericLimits<T>::Max());
 	Concrete.SetMaxAngularSpeedSq(TNumericLimits<T>::Max());
+	Concrete.SetInitialOverlapDepenetrationVelocity(-1.0f);
 	Concrete.SetM(1);
 	Concrete.SetInvM(1);
 	Concrete.SetCenterOfMass(TVector<T,d>(0));
@@ -113,7 +114,6 @@ void PBDRigidParticleDefaultConstruct(FConcrete& Concrete, const FPBDRigidPartic
 	Concrete.SetLinearEtherDrag(0.f);
 	Concrete.SetAngularEtherDrag(0.f);
 	Concrete.SetGravityEnabled(Params.bGravityEnabled);
-	Concrete.SetInitialOverlapDepenetrationEnabled(true);
 	Concrete.SetCCDEnabled(Params.bCCDEnabled);
 	Concrete.SetDisabled(Params.bDisabled);
 	Concrete.SetSleepType(ESleepType::MaterialSleep);
@@ -1164,6 +1164,9 @@ public:
 	T& MaxAngularSpeedSq() { return PBDRigidParticles->MaxAngularSpeedSq(ParticleIdx); }
 	void SetMaxAngularSpeedSq(const T& InMaxAngularSpeed) { PBDRigidParticles->MaxAngularSpeedSq(ParticleIdx) = InMaxAngularSpeed; }
 
+	FRealSingle InitialOverlapDepenetrationVelocity() const { return PBDRigidParticles->InitialOverlapDepenetrationVelocity(ParticleIdx); }
+	void SetInitialOverlapDepenetrationVelocity(FRealSingle InVel) { PBDRigidParticles->InitialOverlapDepenetrationVelocity(ParticleIdx) = InVel; }
+
 	EObjectStateType ObjectState() const { return PBDRigidParticles->ObjectState(ParticleIdx); }
 	EObjectStateType PreObjectState() const { return PBDRigidParticles->PreObjectState(ParticleIdx); }
 
@@ -1231,16 +1234,6 @@ public:
 		PBDRigidParticles->ControlFlags(ParticleIdx).SetCCDEnabled(bEnabled);
 	}
 
-	bool InitialOverlapDepenetrationEnabled() const 
-	{ 
-		return ControlFlags().GetInitialOverlapDepenetrationEnabled();
-	}
-
-	void SetInitialOverlapDepenetrationEnabled(const bool bEnabled)
-	{ 
-		PBDRigidParticles->ControlFlags(ParticleIdx).SetInitialOverlapDepenetrationEnabled(bEnabled);
-	}
-
 	inline bool OneWayInteraction() const
 	{ 
 		return ControlFlags().GetOneWayInteractionEnabled();
@@ -2270,13 +2263,13 @@ public:
 		return FReal(0);
 	}
 
-	bool InitialOverlapDepentrationEnabled() const
+	FRealSingle InitialOverlapDepenetrationVelocity() const
 	{
 		if (auto RigidHandle = MHandle->CastToRigidParticle())
 		{
-			return RigidHandle->InitialOverlapDepenetrationEnabled();
+			return RigidHandle->InitialOverlapDepenetrationVelocity();
 		}
-		return true;
+		return 0;
 	}
 
 	bool OneWayInteraction() const
@@ -3143,12 +3136,6 @@ public:
 		MMiscData.Modify(true, MDirtyFlags, Proxy, [bUpdateKinematicFromSimulation](auto& Data) { Data.SetUpdateKinematicFromSimulation(bUpdateKinematicFromSimulation); });
 	}
 
-	bool InitialOverlapDepenetrationEnabled() const { return MMiscData.Read().InitialOverlapDepenetrationEnabled(); }
-	void SetInitialOverlapDepenetrationEnabled(const bool bInEnabled)
-	{
-		MMiscData.Modify(true, MDirtyFlags, Proxy, [bInEnabled](auto& Data) { Data.SetInitialOverlapDepenetrationEnabled(bInEnabled); });
-	}
-
 	bool OneWayInteraction() const { return MMiscData.Read().OneWayInteraction(); }
 	void SetOneWayInteraction(const bool bInEnabled)
 	{
@@ -3332,6 +3319,13 @@ public:
 		MMiscData.Modify(true,MDirtyFlags,Proxy,[&InAngularSpeed](auto& Data){ Data.SetMaxAngularSpeedSq(InAngularSpeed);});
 	}
 
+	FRealSingle InitialOverlapDepenetrationVelocity() const { return MMiscData.Read().InitialOverlapDepenetrationVelocity(); }
+	void SetInitialOverlapDepenetrationVelocity(FRealSingle InVel)
+	{
+		MMiscData.Modify(true, MDirtyFlags, Proxy, [InVel](auto& Data) { Data.SetInitialOverlapDepenetrationVelocity(InVel); });
+		
+	}
+
 	int32 Island() const { return MIsland; }
 	// TODO(stett): Make the setter private. It is public right now to provide access to proxies.
 	void SetIsland(const int32 InIsland)
diff --git a/Engine/Source/Runtime/Experimental/Chaos/Public/Chaos/RigidParticleControlFlags.h b/Engine/Source/Runtime/Experimental/Chaos/Public/Chaos/RigidParticleControlFlags.h
--- a/Engine/Source/Runtime/Experimental/Chaos/Public/Chaos/RigidParticleControlFlags.h
+++ b/Engine/Source/Runtime/Experimental/Chaos/Public/Chaos/RigidParticleControlFlags.h
@@ -47,10 +47,6 @@ namespace Chaos
 		bool GetOneWayInteractionEnabled() const { return Flags.bOneWayInteractionEnabled; }
 		FRigidParticleControlFlags& SetOneWayInteractionEnabled(const bool bEnabled) { Flags.bOneWayInteractionEnabled = bEnabled; return *this; }
 
-		// If enabled initial overlaps will be depenetrated at a rate that depends on physics settings
-		bool GetInitialOverlapDepenetrationEnabled() const { return Flags.bEnableInitialOverlapDepenetration; }
-		FRigidParticleControlFlags& SetInitialOverlapDepenetrationEnabled(const bool bEnabled) { Flags.bEnableInitialOverlapDepenetration = bEnabled; return *this; }
-
 		// If enabled, inertia may be increased to improve stability
 		bool GetInertiaConditioningEnabled() const { return Flags.bInertiaConditioningEnabled; }
 		FRigidParticleControlFlags& SetInertiaConditioningEnabled(const bool bEnabled) { Flags.bInertiaConditioningEnabled = bEnabled; return *this; }
diff --git a/Engine/Source/Runtime/Experimental/Chaos/Public/Chaos/RigidParticles.h b/Engine/Source/Runtime/Experimental/Chaos/Public/Chaos/RigidParticles.h
--- a/Engine/Source/Runtime/Experimental/Chaos/Public/Chaos/RigidParticles.h
+++ b/Engine/Source/Runtime/Experimental/Chaos/Public/Chaos/RigidParticles.h
@@ -107,6 +107,7 @@ public:
 		, MAngularEtherDrag(MoveTemp(Other.MAngularEtherDrag))
 		, MaxLinearSpeedsSq(MoveTemp(Other.MaxLinearSpeedsSq))
 		, MaxAngularSpeedsSq(MoveTemp(Other.MaxAngularSpeedsSq))
+		, MInitialOverlapDepenetrationVelocity(MoveTemp(Other.MInitialOverlapDepenetrationVelocity))
 		, MCollisionParticles(MoveTemp(Other.MCollisionParticles))
 		, MSleepType(MoveTemp(Other.MSleepType))
 		, MSleepCounter(MoveTemp(Other.MSleepCounter))
@@ -136,6 +137,7 @@ public:
 		TArrayCollection::AddArray(&MAngularEtherDrag);
 		TArrayCollection::AddArray(&MaxLinearSpeedsSq);
 		TArrayCollection::AddArray(&MaxAngularSpeedsSq);
+		TArrayCollection::AddArray(&MInitialOverlapDepenetrationVelocity);
 		TArrayCollection::AddArray(&MCollisionParticles);
 		TArrayCollection::AddArray(&MSleepType);
 		TArrayCollection::AddArray(&MSleepCounter);
@@ -197,6 +199,9 @@ public:
 	FORCEINLINE const T& MaxAngularSpeedSq(const int32 index) const { return MaxAngularSpeedsSq[index]; }
 	FORCEINLINE T& MaxAngularSpeedSq(const int32 index) { return MaxAngularSpeedsSq[index]; }
 
+	FORCEINLINE const FRealSingle& InitialOverlapDepenetrationVelocity(const int32 index) const { return MInitialOverlapDepenetrationVelocity[index]; }
+	FORCEINLINE FRealSingle& InitialOverlapDepenetrationVelocity(const int32 index) { return MInitialOverlapDepenetrationVelocity[index]; }
+
 	FORCEINLINE int32 CollisionParticlesSize(int32 Index) const { return MCollisionParticles[Index] == nullptr ? 0 : MCollisionParticles[Index]->Size(); }
 
 	void CollisionParticlesInitIfNeeded(const int32 Index);
@@ -407,6 +412,7 @@ private:
 	TArrayCollectionArray<T> MAngularEtherDrag;
 	TArrayCollectionArray<T> MaxLinearSpeedsSq;
 	TArrayCollectionArray<T> MaxAngularSpeedsSq;
+	TArrayCollectionArray<FRealSingle> MInitialOverlapDepenetrationVelocity;
 	TArrayCollectionArray<TUniquePtr<TBVHParticles<T, d>>> MCollisionParticles;
 	TArrayCollectionArray<ESleepType> MSleepType;
 	TArrayCollectionArray<int8> MSleepCounter;
diff --git a/Engine/Source/Runtime/Experimental/Chaos/Public/PhysicsProxy/SingleParticlePhysicsProxy.h b/Engine/Source/Runtime/Experimental/Chaos/Public/PhysicsProxy/SingleParticlePhysicsProxy.h
--- a/Engine/Source/Runtime/Experimental/Chaos/Public/PhysicsProxy/SingleParticlePhysicsProxy.h
+++ b/Engine/Source/Runtime/Experimental/Chaos/Public/PhysicsProxy/SingleParticlePhysicsProxy.h
@@ -458,30 +458,6 @@ public:
 		});
 	}
 
-	bool InitialOverlapDepenetrationEnabled() const
-	{
-		return Read([](auto* Particle)
-			{
-				if (auto Rigid = Particle->CastToRigidParticle())
-				{
-					return Rigid->InitialOverlapDepenetrationEnabled();
-				}
-
-				return true;
-			});
-	}
-
-	void SetInitialOverlapDepenetrationEnabled(const bool bEnabled)
-	{
-		Write([bEnabled](auto* Particle)
-			{
-				if (auto Rigid = Particle->CastToRigidParticle())
-				{
-					return Rigid->SetInitialOverlapDepenetrationEnabled(bEnabled);
-				}
-			});
-	}
-
 	bool InertiaConditioningEnabled() const
 	{
 		return Read([](auto* Particle)
@@ -1300,6 +1276,26 @@ public:
 		return TNumericLimits<FReal>::Max();
 	}
 
+	FRealSingle GetInitialOverlapDepenetrationVelocity()
+	{
+		VerifyContext();
+		if (const TPBDRigidParticle<FReal, 3>*Rigid = GetParticle_LowLevel()->CastToRigidParticle())
+		{
+			return Rigid->InitialOverlapDepenetrationVelocity();
+		}
+
+		return 0;
+	}
+
+	void SetInitialOverlapDepenetrationVelocity(FRealSingle InNewSpeed)
+	{
+		VerifyContext();
+		if (TPBDRigidParticle<FReal, 3>*Rigid = GetParticle_LowLevel()->CastToRigidParticle())
+		{
+			Rigid->SetInitialOverlapDepenetrationVelocity(InNewSpeed);
+		}
+	}
+
 	void SetDisabled(bool bDisable)
 	{
 		VerifyContext();
diff --git a/Engine/Source/Runtime/Experimental/Chaos/Public/RewindData.h b/Engine/Source/Runtime/Experimental/Chaos/Public/RewindData.h
--- a/Engine/Source/Runtime/Experimental/Chaos/Public/RewindData.h
+++ b/Engine/Source/Runtime/Experimental/Chaos/Public/RewindData.h
@@ -728,6 +728,11 @@ inline FString ToStringHelper(const FReal Val)
 	return FString::Printf(TEXT("%f"), Val);
 }
 
+inline FString ToStringHelper(const FRealSingle Val)
+{
+	return FString::Printf(TEXT("%f"), Val);
+}
+
 inline FString ToStringHelper(const EObjectStateType Val)
 {
 	return FString::Printf(TEXT("%d"), Val);
@@ -978,6 +983,7 @@ public:
 	REWIND_PARTICLE_RIGID_PROPERTY(DynamicsMisc, AngularEtherDrag)
 	REWIND_PARTICLE_RIGID_PROPERTY(DynamicsMisc, MaxLinearSpeedSq)
 	REWIND_PARTICLE_RIGID_PROPERTY(DynamicsMisc, MaxAngularSpeedSq)
+	REWIND_PARTICLE_RIGID_PROPERTY(DynamicsMisc, InitialOverlapDepenetrationVelocity)
 	REWIND_PARTICLE_RIGID_PROPERTY(DynamicsMisc, ObjectState)
 	REWIND_PARTICLE_RIGID_PROPERTY(DynamicsMisc, CollisionGroup)
 	REWIND_PARTICLE_RIGID_PROPERTY(DynamicsMisc, ControlFlags)
@@ -1040,6 +1046,8 @@ public:
 			REWIND_PARTICLE_TO_STR(AngularEtherDrag)
 			REWIND_PARTICLE_TO_STR(MaxLinearSpeedSq)
 			REWIND_PARTICLE_TO_STR(MaxAngularSpeedSq)
+			REWIND_PARTICLE_TO_STR(InitialOverlapDepenetrationVelocity)
+
 			REWIND_PARTICLE_TO_STR(ObjectState)
 			REWIND_PARTICLE_TO_STR(CollisionGroup)
 			REWIND_PARTICLE_TO_STR(ControlFlags)
diff --git a/Engine/Source/Runtime/Experimental/ChaosVisualDebugger/Private/DataWrappers/ChaosVDParticleDataWrapper.cpp b/Engine/Source/Runtime/Experimental/ChaosVisualDebugger/Private/DataWrappers/ChaosVDParticleDataWrapper.cpp
--- a/Engine/Source/Runtime/Experimental/ChaosVisualDebugger/Private/DataWrappers/ChaosVDParticleDataWrapper.cpp
+++ b/Engine/Source/Runtime/Experimental/ChaosVisualDebugger/Private/DataWrappers/ChaosVDParticleDataWrapper.cpp
@@ -11,7 +11,6 @@ bool FChaosVDFRigidParticleControlFlags::Serialize(FArchive& Ar)
 	Ar << bGravityEnabled;
 	Ar << bCCDEnabled;
 	Ar << bOneWayInteractionEnabled;
-	Ar << bEnableInitialOverlapDepenetration;
 	Ar << bInertiaConditioningEnabled;
 	Ar << GravityGroupIndex;
 
@@ -65,6 +64,7 @@ bool FChaosVDParticleDynamicMisc::Serialize(FArchive& Ar)
 	Ar << MAngularEtherDrag;
 	Ar << MMaxLinearSpeedSq;
 	Ar << MMaxAngularSpeedSq;
+	Ar << MInitialOverlapDepenetrationVelocity;
 	Ar << MCollisionGroup;
 	Ar << MObjectState;
 	Ar << MSleepType;
diff --git a/Engine/Source/Runtime/Experimental/ChaosVisualDebugger/Public/DataWrappers/ChaosVDParticleDataWrapper.h b/Engine/Source/Runtime/Experimental/ChaosVisualDebugger/Public/DataWrappers/ChaosVDParticleDataWrapper.h
--- a/Engine/Source/Runtime/Experimental/ChaosVisualDebugger/Public/DataWrappers/ChaosVDParticleDataWrapper.h
+++ b/Engine/Source/Runtime/Experimental/ChaosVisualDebugger/Public/DataWrappers/ChaosVDParticleDataWrapper.h
@@ -63,7 +63,6 @@ struct FChaosVDFRigidParticleControlFlags : public FChaosVDParticleDataBase
 		: bGravityEnabled(false),
 		  bCCDEnabled(false),
 		  bOneWayInteractionEnabled(false),
-		  bEnableInitialOverlapDepenetration(false),
 		  bInertiaConditioningEnabled(false), 
 		  GravityGroupIndex(0)
 	{
@@ -77,7 +76,6 @@ struct FChaosVDFRigidParticleControlFlags : public FChaosVDParticleDataBase
 		bGravityEnabled = Other.GetGravityEnabled();
 		bCCDEnabled = Other.GetCCDEnabled();
 		bOneWayInteractionEnabled = Other.GetOneWayInteractionEnabled();
-		bEnableInitialOverlapDepenetration = Other.GetInitialOverlapDepenetrationEnabled();
 		bInertiaConditioningEnabled = Other.GetInertiaConditioningEnabled();
 		GravityGroupIndex = Other.GetGravityGroupIndex();
 
@@ -91,8 +89,6 @@ struct FChaosVDFRigidParticleControlFlags : public FChaosVDParticleDataBase
 	UPROPERTY(EditAnywhere, Category= "Particle Control Flags")
 	bool bOneWayInteractionEnabled;
 	UPROPERTY(EditAnywhere, Category = "Particle Control Flags")
-	bool bEnableInitialOverlapDepenetration;
-	UPROPERTY(EditAnywhere, Category = "Particle Control Flags")
 	bool bInertiaConditioningEnabled;
 	UPROPERTY(EditAnywhere, Category= "Particle Control Flags")
 	int32 GravityGroupIndex;
@@ -325,7 +321,7 @@ struct FChaosVDParticleDynamicMisc : public FChaosVDParticleDataBase
 {
 	GENERATED_BODY()
 
-	FChaosVDParticleDynamicMisc(): MAngularEtherDrag(0), MMaxLinearSpeedSq(0), MMaxAngularSpeedSq(0),
+	FChaosVDParticleDynamicMisc(): MAngularEtherDrag(0), MMaxLinearSpeedSq(0), MMaxAngularSpeedSq(0), MInitialOverlapDepenetrationVelocity(0),
 	                               MCollisionGroup(0), MObjectState(), MSleepType(), bDisabled(false)
 	{
 	}
@@ -339,6 +335,7 @@ struct FChaosVDParticleDynamicMisc : public FChaosVDParticleDataBase
 		MAngularEtherDrag = Other.AngularEtherDrag();
 		MMaxLinearSpeedSq = Other.MaxLinearSpeedSq();
 		MMaxAngularSpeedSq = Other.MaxAngularSpeedSq();
+		MInitialOverlapDepenetrationVelocity = Other.InitialOverlapDepenetrationVelocity();
 		MObjectState = static_cast<EChaosVDObjectStateType>(Other.ObjectState());
 		MCollisionGroup = Other.CollisionGroup();
 		MSleepType =  static_cast<EChaosVDSleepType>(Other.SleepType());
@@ -360,6 +357,9 @@ struct FChaosVDParticleDynamicMisc : public FChaosVDParticleDataBase
 	UPROPERTY(EditAnywhere, Category= "Particle Dynamic Misc")
 	double MMaxAngularSpeedSq;
 
+	UPROPERTY(EditAnywhere, Category = "Particle Dynamic Misc")
+	float MInitialOverlapDepenetrationVelocity;
+
 	UPROPERTY(EditAnywhere, Category= "Particle Dynamic Misc")
 	int32 MCollisionGroup;
 
diff --git a/Engine/Source/Runtime/PhysicsCore/Private/ChaosEngineInterface.cpp b/Engine/Source/Runtime/PhysicsCore/Private/ChaosEngineInterface.cpp
--- a/Engine/Source/Runtime/PhysicsCore/Private/ChaosEngineInterface.cpp
+++ b/Engine/Source/Runtime/PhysicsCore/Private/ChaosEngineInterface.cpp
@@ -667,13 +667,20 @@ void FChaosEngineInterface::SetMaxLinearVelocity_AssumesLocked(const FPhysicsAct
 
 float FChaosEngineInterface::GetMaxDepenetrationVelocity_AssumesLocked(const FPhysicsActorHandle& InActorReference)
 {
-	CHAOS_ENSURE(false);
-	return FLT_MAX;
+	if (ensure(FChaosEngineInterface::IsValid(InActorReference)))
+	{
+		return FMath::Sqrt(InActorReference->GetGameThreadAPI().GetInitialOverlapDepenetrationVelocity());
+	}
+
+	return 0;
 }
 
 void FChaosEngineInterface::SetMaxDepenetrationVelocity_AssumesLocked(const FPhysicsActorHandle& InActorReference,float InMaxDepenetrationVelocity)
 {
-	CHAOS_ENSURE(false);
+	if (ensure(FChaosEngineInterface::IsValid(InActorReference)))
+	{
+		InActorReference->GetGameThreadAPI().SetInitialOverlapDepenetrationVelocity(InMaxDepenetrationVelocity);
+	}
 }
 
 FVector FChaosEngineInterface::GetWorldVelocityAtPoint_AssumesLocked(const FPhysicsActorHandle& InActorReference,const FVector& InPoint)
@@ -1193,16 +1200,6 @@ void FChaosEngineInterface::SetInertiaConditioningEnabled_AssumesLocked(const FP
 	InActorReference->GetGameThreadAPI().SetInertiaConditioningEnabled(bEnabled);
 }
 
-bool FChaosEngineInterface::IsInitialOverlapDepenetrationEnabled_AssumesLocked(const FPhysicsActorHandle& InActorReference)
-{
-	return InActorReference->GetGameThreadAPI().InitialOverlapDepenetrationEnabled();
-}
-
-void FChaosEngineInterface::SetInitialOverlapDepenetrationEnabled_AssumesLocked(const FPhysicsActorHandle& InActorReference, bool bEnabled)
-{
-	InActorReference->GetGameThreadAPI().SetInitialOverlapDepenetrationEnabled(bEnabled);
-}
-
 void FChaosEngineInterface::SetIsSimulationShape(const FPhysicsShapeHandle& InShape,bool bIsSimShape)
 {
 	InShape.Shape->SetSimEnabled(bIsSimShape);
diff --git a/Engine/Source/Runtime/PhysicsCore/Public/Chaos/ChaosEngineInterface.h b/Engine/Source/Runtime/PhysicsCore/Public/Chaos/ChaosEngineInterface.h
--- a/Engine/Source/Runtime/PhysicsCore/Public/Chaos/ChaosEngineInterface.h
+++ b/Engine/Source/Runtime/PhysicsCore/Public/Chaos/ChaosEngineInterface.h
@@ -472,9 +472,6 @@ public:
 	static PHYSICSCORE_API bool IsInertiaConditioningEnabled_AssumesLocked(const FPhysicsActorHandle& InActorReference);
 	static PHYSICSCORE_API void SetInertiaConditioningEnabled_AssumesLocked(const FPhysicsActorHandle& InActorReference, bool bEnabled);
 
-	static PHYSICSCORE_API bool IsInitialOverlapDepenetrationEnabled_AssumesLocked(const FPhysicsActorHandle& InActorReference);
-	static PHYSICSCORE_API void SetInitialOverlapDepenetrationEnabled_AssumesLocked(const FPhysicsActorHandle& InActorReference, bool bEnabled);
-
 	static PHYSICSCORE_API float GetStabilizationEnergyThreshold_AssumesLocked(const FPhysicsActorHandle& InHandle);
 	static PHYSICSCORE_API void SetStabilizationEnergyThreshold_AssumesLocked(const FPhysicsActorHandle& InHandle,float InThreshold);
 	static PHYSICSCORE_API uint32 GetSolverPositionIterationCount_AssumesLocked(const FPhysicsActorHandle& InHandle);
