#ue5 - Scoped movement update for character mesh to avoid redundant movement of mesh children during multiple scoped moves.
Added cvars:
p.DeferCharacterMeshMovement (true/false)
p.DeferCharacterMeshMovementForAllCorrections (true/false)

- Also cleaned up previous scoped update into it's own wrapper struct to make it easier to use and find uses of that versus other scoped moves.
- Added a small logic fix to FScopedMovementUpdate when using a null component to avoid overhead in SceneComponent when checking for deferred behavior.
- Changed CharacterMovement existing scopes to use null component rather than immediate updates when disabled (saves some transform copies, happens on server dual moves sometimes).
- Added scoped movement for capsule within ServerMove_PerformMovement() to combine the rotation change with the movement update.

#tests PIE w/ multiple players with cvars on/off. Moved, fired, etc with medium network profile to check validity with corrections.

diff --git a/Engine/Source/Runtime/Engine/Private/Components/CharacterMovementComponent.cpp b/Engine/Source/Runtime/Engine/Private/Components/CharacterMovementComponent.cpp
--- a/Engine/Source/Runtime/Engine/Private/Components/CharacterMovementComponent.cpp
+++ b/Engine/Source/Runtime/Engine/Private/Components/CharacterMovementComponent.cpp
@@ -318,6 +318,20 @@ namespace CharacterMovementCVars
 		TEXT("When enabled, this allows a character that's supposed to remain vertical to snap to a vertical orientation even if RotationRate settings would block it. See @ShouldRemainVertical and @RotationRate."),
 		ECVF_Default);
 
+	static bool bDeferCharacterMeshMovement = false;
+	FAutoConsoleVariableRef CVarDeferCharacterMeshMovement(
+		TEXT("p.DeferCharacterMeshMovement"),
+		bDeferCharacterMeshMovement,
+		TEXT("Optimization - When enabled, defers CharacterMesh move propagation until the end of larger scoped moves. The mesh will still move, but all attached components will wait until all mesh movement is done within the scope."),
+		ECVF_Default);
+
+	static bool bDeferCharacterMeshMovementForAllCorrections = true;
+	FAutoConsoleVariableRef CVarDeferCharacterMeshMovementForAllCorrections(
+		TEXT("p.DeferCharacterMeshMovementForAllCorrections"),
+		bDeferCharacterMeshMovementForAllCorrections,
+		TEXT("Optimization - When enabled, defers CharacterMesh move propagation for all corrections until the end of larger scoped moves. Requires `bDeferCharacterMeshMovement=true'."),
+		ECVF_Default);
+
 #if !UE_BUILD_SHIPPING
 
 	int32 NetShowCorrections = 0;
@@ -401,6 +415,7 @@ namespace CharacterMovementCVars
  *	{
  *		FScopedPreventMeshBoneUpdate ScopedNoMeshBoneUpdate(CharacterOwner->GetMesh(), EKinematicBonesUpdateToPhysics::SkipAllBones);
  *		// Do something to move mesh, bones will not update
+ *		// Do something again to move mesh, bones will not update
  *	}
  *	// Movement of mesh at this point will use previous setting.
  */
@@ -433,6 +448,43 @@ private:
 };
 
 
+/**
+ * Helper to change updated component transform updates within a scope, optionally allowing it to revert back without side effects.
+ * Example usage:
+ *	{
+ *		FScopedCapsuleMovementUpdate ScopedDeferredUpdate(UpdatedComponent, bDeferredUpdatesEnabled);
+ *		// Do something to move capsule
+ *		// Do something again to move capsule
+ *	}
+ *	// Movement of mesh at this point will use previous setting.
+ */
+struct FScopedCapsuleMovementUpdate : public FScopedMovementUpdate
+{
+	typedef FScopedMovementUpdate Super;
+
+	FScopedCapsuleMovementUpdate(USceneComponent* UpdatedComponent, bool bEnabled)
+	: Super(bEnabled ? UpdatedComponent : nullptr, EScopedUpdate::DeferredUpdates)
+	{
+	}
+};
+
+
+/**
+ * Similar to FScopedCapsuleMovementUpdate, but intended for the character mesh instead.
+ * @see FScopedCapsuleMovementUpdate
+ */
+struct FScopedMeshMovementUpdate
+{
+	FScopedMeshMovementUpdate(USkeletalMeshComponent* Mesh, bool bEnabled = true)
+	: ScopedMoveUpdate(bEnabled && CharacterMovementCVars::bDeferCharacterMeshMovement ? Mesh : nullptr, EScopedUpdate::DeferredUpdates)
+	{
+	}
+
+private:
+	FScopedMovementUpdate ScopedMoveUpdate;
+};
+
+
 
 void FFindFloorResult::SetFromSweep(const FHitResult& InHit, const float InSweepFloorDist, const bool bIsWalkableFloor)
 {
@@ -1543,12 +1595,15 @@ void UCharacterMovementComponent::TickComponent(float DeltaTime, enum ELevelTick
 			}
 		}
 
+		// Defer all mesh child updates until all movement completes. Avoiding combining this with the corrections above, as those may cause a lot of significant movement.
+		FScopedMeshMovementUpdate ScopedMeshUpdate(CharacterOwner->GetMesh());
+
 		// Perform input-driven move for any locally-controlled character, and also
 		// allow animation root motion or physics to move characters even if they have no controller
 		const bool bShouldPerformControlledCharMove = CharacterOwner->IsLocallyControlled() 
 													  || (!CharacterOwner->Controller && bRunPhysicsWithNoController)		
 													  || (!CharacterOwner->Controller && CharacterOwner->IsPlayingRootMotion());
-		
+
 		if (bShouldPerformControlledCharMove)
 		{
 			ControlledCharacterMove(InputVector, DeltaTime);
@@ -1580,6 +1635,9 @@ void UCharacterMovementComponent::TickComponent(float DeltaTime, enum ELevelTick
 	}
 	else if (CharacterOwner->GetLocalRole() == ROLE_SimulatedProxy)
 	{
+		// Defer all mesh child updates until all movement completes.
+		FScopedMeshMovementUpdate ScopedMeshUpdate(CharacterOwner->GetMesh());
+
 		if (bShrinkProxyCapsule)
 		{
 			AdjustProxyCapsuleSize();
@@ -1624,7 +1682,7 @@ void UCharacterMovementComponent::PostPhysicsTickComponent(float DeltaTime, FCha
 		{
 			ensure(false); // Not supported
 		}
-		FScopedMovementUpdate ScopedMovementUpdate(UpdatedComponent, bEnableScopedMovementUpdates ? EScopedUpdate::DeferredUpdates : EScopedUpdate::ImmediateUpdates);
+		FScopedCapsuleMovementUpdate ScopedMovementUpdate(UpdatedComponent, bEnableScopedMovementUpdates);
 		UpdateBasedMovement(DeltaTime);
 		SaveBaseLocation();
 		bDeferUpdateBasedMovement = false;
@@ -1893,7 +1951,7 @@ void UCharacterMovementComponent::SimulateRootMotion(float DeltaSeconds, const F
 {
 	if( CharacterOwner && CharacterOwner->GetMesh() && (DeltaSeconds > 0.f) )
 	{
-		FScopedMovementUpdate ScopedMovementUpdate(UpdatedComponent, bEnableScopedMovementUpdates ? EScopedUpdate::DeferredUpdates : EScopedUpdate::ImmediateUpdates);
+		FScopedCapsuleMovementUpdate ScopedMovementUpdate(UpdatedComponent, bEnableScopedMovementUpdates);
 
 		// Convert Local Space Root Motion to world space. Do it right before used by physics to make sure we use up to date transforms, as translation is relative to rotation.
 		const FTransform WorldSpaceRootMotionTransform = ConvertLocalRootMotionToWorld(LocalRootMotionTransform, DeltaSeconds);
@@ -2004,7 +2062,7 @@ void UCharacterMovementComponent::SimulateMovement(float DeltaSeconds)
 
 	// Scoped updates can improve performance of multiple MoveComponent calls.
 	{
-		FScopedMovementUpdate ScopedMovementUpdate(UpdatedComponent, bEnableScopedMovementUpdates ? EScopedUpdate::DeferredUpdates : EScopedUpdate::ImmediateUpdates);
+		FScopedCapsuleMovementUpdate ScopedMovementUpdate(UpdatedComponent, bEnableScopedMovementUpdates);
 
 		bool bHandledNetUpdate = false;
 		if (bIsSimulatedProxy)
@@ -2539,7 +2597,7 @@ void UCharacterMovementComponent::PerformMovement(float DeltaSeconds)
 
 	// Scoped updates can improve performance of multiple MoveComponent calls.
 	{
-		FScopedMovementUpdate ScopedMovementUpdate(UpdatedComponent, bEnableScopedMovementUpdates ? EScopedUpdate::DeferredUpdates : EScopedUpdate::ImmediateUpdates);
+		FScopedCapsuleMovementUpdate ScopedMovementUpdate(UpdatedComponent, bEnableScopedMovementUpdates);
 
 		MaybeUpdateBasedMovement(DeltaSeconds);
 
@@ -6448,7 +6506,7 @@ void UCharacterMovementComponent::MoveSmooth(const FVector& InVelocity, const fl
 	// Custom movement may need an update even if there is zero velocity.
 	if (MovementMode == MOVE_Custom)
 	{
-		FScopedMovementUpdate ScopedMovementUpdate(UpdatedComponent, bEnableScopedMovementUpdates ? EScopedUpdate::DeferredUpdates : EScopedUpdate::ImmediateUpdates);
+		FScopedCapsuleMovementUpdate ScopedMovementUpdate(UpdatedComponent, bEnableScopedMovementUpdates);
 		PhysCustom(DeltaSeconds, 0);
 		return;
 	}
@@ -6459,7 +6517,7 @@ void UCharacterMovementComponent::MoveSmooth(const FVector& InVelocity, const fl
 		return;
 	}
 
-	FScopedMovementUpdate ScopedMovementUpdate(UpdatedComponent, bEnableScopedMovementUpdates ? EScopedUpdate::DeferredUpdates : EScopedUpdate::ImmediateUpdates);
+	FScopedCapsuleMovementUpdate ScopedMovementUpdate(UpdatedComponent, bEnableScopedMovementUpdates);
 
 	if (IsMovingOnGround())
 	{
@@ -7142,7 +7200,7 @@ bool UCharacterMovementComponent::StepUp(const FVector& GravDir, const FVector&
 	}
 
 	// Scope our movement updates, and do not apply them until all intermediate moves are completed.
-	FScopedMovementUpdate ScopedStepUpMovement(UpdatedComponent, EScopedUpdate::DeferredUpdates);
+	FScopedCapsuleMovementUpdate ScopedStepUpMovement(UpdatedComponent, true);
 
 	// step up - treat as vertical wall
 	FHitResult SweepUpHit(1.f);
@@ -8169,6 +8227,9 @@ bool UCharacterMovementComponent::ClientUpdatePositionAfterServerUpdate()
 	CharacterOwner->bClientUpdating = true;
 	bForceNextFloorCheck = true;
 
+	// Defer all mesh child updates until all movement completes.
+	FScopedMeshMovementUpdate ScopedMeshUpdate(CharacterOwner->GetMesh(), CharacterMovementCVars::bDeferCharacterMeshMovementForAllCorrections);
+
 	// Replay moves that have not yet been acked.
 	UE_LOG(LogNetPlayerMovement, Verbose, TEXT("ClientUpdatePositionAfterServerUpdate Replaying %d Moves, starting at Timestamp %f"), ClientData->SavedMoves.Num(), ClientData->SavedMoves[0]->TimeStamp);
 	for (int32 i=0; i<ClientData->SavedMoves.Num(); i++)
@@ -8483,7 +8544,7 @@ void UCharacterMovementComponent::ReplicateMoveToServer(float DeltaTime, const F
 				FScopedMeshBoneUpdateOverride ScopedNoMeshBoneUpdate(CharacterOwner->GetMesh(), EKinematicBonesUpdateToPhysics::SkipAllBones);
 
 				// Accumulate multiple transform updates until scope ends.
-				FScopedMovementUpdate ScopedMovementUpdate(UpdatedComponent, EScopedUpdate::DeferredUpdates);
+				FScopedCapsuleMovementUpdate ScopedMovementUpdate(UpdatedComponent, true);
 				UE_LOG(LogNetPlayerMovement, VeryVerbose, TEXT("CombineMove: add delta %f + %f and revert from %f %f to %f %f"), DeltaTime, PendingMove->DeltaTime, UpdatedComponent->GetComponentLocation().X, UpdatedComponent->GetComponentLocation().Y, OldStartLocation.X, OldStartLocation.Y);
 
 				NewMove->CombineWith(PendingMove, CharacterOwner, PC, OldStartLocation);
@@ -8689,6 +8750,7 @@ void UCharacterMovementComponent::CallServerMovePacked(const FSavedMove_Characte
 }
 
 
+///// DEPRECATED /////
 void UCharacterMovementComponent::CallServerMove
 	(
 	const FSavedMove_Character* NewMove,
@@ -8779,7 +8841,7 @@ void UCharacterMovementComponent::CallServerMove
 }
 
 
-
+///// DEPRECATED /////
 void UCharacterMovementComponent::ServerMoveOld_Implementation
 	(
 	float OldTimeStamp,
@@ -8825,6 +8887,7 @@ void UCharacterMovementComponent::ServerMoveOld_Implementation
 }
 
 
+///// DEPRECATED /////
 void UCharacterMovementComponent::ServerMoveDual_Implementation(
 	float TimeStamp0,
 	FVector_NetQuantize10 InAccel0,
@@ -8841,12 +8904,13 @@ void UCharacterMovementComponent::ServerMoveDual_Implementation(
 	uint8 ClientMovementMode)
 {
 	// Optional scoped movement update to combine moves for cheaper performance on the server.
-	FScopedMovementUpdate ScopedMovementUpdate(UpdatedComponent, bEnableServerDualMoveScopedMovementUpdates ? EScopedUpdate::DeferredUpdates : EScopedUpdate::ImmediateUpdates);
+	FScopedCapsuleMovementUpdate ScopedDualMovementUpdate(UpdatedComponent, bEnableServerDualMoveScopedMovementUpdates);
 
 	ServerMove_Implementation(TimeStamp0, InAccel0, FVector(1.f,2.f,3.f), PendingFlags, ClientRoll, View0, ClientMovementBase, ClientBaseBone, ClientMovementMode);
 	ServerMove_Implementation(TimeStamp, InAccel, ClientLoc, NewFlags, ClientRoll, View, ClientMovementBase, ClientBaseBone, ClientMovementMode);
 }
 
+///// DEPRECATED /////
 void UCharacterMovementComponent::ServerMoveDualHybridRootMotion_Implementation(
 	float TimeStamp0,
 	FVector_NetQuantize10 InAccel0,
@@ -9407,7 +9471,7 @@ void UCharacterMovementComponent::ServerMove_HandleMoveData(const FCharacterNetw
 
 	// Optional scoped movement update for dual moves to combine moves for cheaper performance on the server.
 	const bool bMoveAllowsScopedDualMove = MoveDataContainer.bHasPendingMove && !MoveDataContainer.bDisableCombinedScopedMove;
-	FScopedMovementUpdate ScopedMovementUpdate(UpdatedComponent, (bMoveAllowsScopedDualMove && bEnableServerDualMoveScopedMovementUpdates && bEnableScopedMovementUpdates) ? EScopedUpdate::DeferredUpdates : EScopedUpdate::ImmediateUpdates);
+	FScopedCapsuleMovementUpdate ScopedMovementUpdate(UpdatedComponent, bMoveAllowsScopedDualMove && bEnableServerDualMoveScopedMovementUpdates && bEnableScopedMovementUpdates);
 
 	// Optional pending move as part of "dual move"
 	if (MoveDataContainer.bHasPendingMove)
@@ -9487,6 +9551,9 @@ void UCharacterMovementComponent::ServerMove_PerformMovement(const FCharacterNet
 	const UWorld* MyWorld = GetWorld();
 	const float DeltaTime = ServerData->GetServerMoveDeltaTime(ClientTimeStamp, CharacterOwner->GetActorTimeDilation(*MyWorld));
 
+	// Defer all mesh child updates until all movement completes.
+	FScopedMeshMovementUpdate ScopedMeshUpdate(CharacterOwner->GetMesh());
+
 	if (DeltaTime > 0.f)
 	{
 		ServerData->CurrentClientTimeStamp = ClientTimeStamp;
@@ -9507,6 +9574,7 @@ void UCharacterMovementComponent::ServerMove_PerformMovement(const FCharacterNet
 		// Perform actual movement
 		if ((MyWorld->GetWorldSettings()->GetPauserPlayerState() == NULL))
 		{
+			FScopedCapsuleMovementUpdate ScopedMovementUpdate(UpdatedComponent, bEnableScopedMovementUpdates);
 			if (PC)
 			{
 				PC->UpdateRotation(DeltaTime);
@@ -9528,6 +9596,7 @@ void UCharacterMovementComponent::ServerMove_PerformMovement(const FCharacterNet
 }
 
 
+///// DEPRECATED /////
 void UCharacterMovementComponent::ServerMove(float TimeStamp, FVector_NetQuantize10 InAccel, FVector_NetQuantize100 ClientLoc, uint8 CompressedMoveFlags, uint8 ClientRoll, uint32 View, UPrimitiveComponent* ClientMovementBase, FName ClientBaseBoneName, uint8 ClientMovementMode)
 {
 	if (MovementBaseUtility::IsDynamicBase(ClientMovementBase))
@@ -9542,6 +9611,7 @@ void UCharacterMovementComponent::ServerMove(float TimeStamp, FVector_NetQuantiz
 	}
 }
 
+///// DEPRECATED /////
 void UCharacterMovementComponent::ServerMove_Implementation(
 	float TimeStamp,
 	FVector_NetQuantize10 InAccel,
@@ -9637,6 +9707,7 @@ void UCharacterMovementComponent::ServerMove_Implementation(
 }
 
 
+///// DEPRECATED /////
 void UCharacterMovementComponent::ServerMoveHandleClientError(float ClientTimeStamp, float DeltaTime, const FVector& Accel, const FVector& RelativeClientLoc, UPrimitiveComponent* ClientMovementBase, FName ClientBaseBoneName, uint8 ClientMovementMode)
 {
 	if (!ShouldUsePackedMovementRPCs())
@@ -10101,6 +10172,9 @@ void UCharacterMovementComponent::MoveAutonomous
 
 	const bool bWasPlayingRootMotion = CharacterOwner->IsPlayingRootMotion();
 
+	// Defer all mesh child updates until all movement completes.
+	FScopedMeshMovementUpdate ScopedMeshUpdate(CharacterOwner->GetMesh());
+
 	PerformMovement(DeltaTime);
 
 	// Check if data is valid as PerformMovement can mark character for pending kill
diff --git a/Engine/Source/Runtime/Engine/Private/Engine/ScopedMovementUpdate.cpp b/Engine/Source/Runtime/Engine/Private/Engine/ScopedMovementUpdate.cpp
--- a/Engine/Source/Runtime/Engine/Private/Engine/ScopedMovementUpdate.cpp
+++ b/Engine/Source/Runtime/Engine/Private/Engine/ScopedMovementUpdate.cpp
@@ -48,6 +48,9 @@ FScopedMovementUpdate::FScopedMovementUpdate( class USceneComponent* Component,
 	else
 	{
 		Owner = nullptr;
+		// Having no owner should be a no-op on all moves, but if someone requests deferred updates on a null component
+		// we should not return true for IsDeferringUpdates() so force it off here.
+		bDeferUpdates = false;
 	}
 }
 
