Shader Compiler:  missing file from previous CL 28741202 (normalized line endings)

diff --git a/Engine/Source/Runtime/Engine/Private/ShaderCompiler/ShaderCompiler.cpp b/Engine/Source/Runtime/Engine/Private/ShaderCompiler/ShaderCompiler.cpp
--- a/Engine/Source/Runtime/Engine/Private/ShaderCompiler/ShaderCompiler.cpp
+++ b/Engine/Source/Runtime/Engine/Private/ShaderCompiler/ShaderCompiler.cpp
@@ -7230,8 +7230,8 @@ void GenerateInstancedStereoCode(FString& Result, EShaderPlatform ShaderPlatform
 	const TArray<FShaderParametersMetadata::FMember>& StructMembersInstanced = InstancedView->GetMembers();
 
 	// ViewState definition
-	Result =  "struct ViewState\r\n";
-	Result += "{\r\n";
+	Result =  "struct ViewState\n";
+	Result += "{\n";
 	for (int32 MemberIndex = 0; MemberIndex < StructMembersInstanced.Num(); ++MemberIndex)
 	{
 		FString MemberDecl;
@@ -7239,40 +7239,40 @@ void GenerateInstancedStereoCode(FString& Result, EShaderPlatform ShaderPlatform
 		// GPUSceneViewId for example needs to return 	uint GPUSceneViewId; and not uint4 InstancedView_GPUSceneViewId[2];
 		// and that initial representation is in StructMembersView
 		GenerateUniformBufferStructMember(MemberDecl, StructMembersView[MemberIndex], ShaderPlatform);
-		Result += FString::Printf(TEXT("\t%s;\r\n"), *MemberDecl);
-	}
-	Result += "\tFLWCInverseMatrix WorldToClip;\r\n";
-	Result += "\tFLWCMatrix ClipToWorld;\r\n";
-	Result += "\tFLWCMatrix ScreenToWorld;\r\n";
-	Result += "\tFLWCMatrix PrevClipToWorld;\r\n";
-	Result += "\tFLWCVector3 WorldCameraOrigin;\r\n";
-	Result += "\tFLWCVector3 WorldViewOrigin;\r\n";
-	Result += "\tFLWCVector3 PrevWorldCameraOrigin;\r\n";
-	Result += "\tFLWCVector3 PrevWorldViewOrigin;\r\n";
-	Result += "\tFLWCVector3 PreViewTranslation;\r\n";
-	Result += "\tFLWCVector3 PrevPreViewTranslation;\r\n";
-	Result += "};\r\n";
-
-	Result += "\tvoid FinalizeViewState(inout ViewState InOutView);\r\n";
+		Result += FString::Printf(TEXT("\t%s;\n"), *MemberDecl);
+	}
+	Result += "\tFLWCInverseMatrix WorldToClip;\n";
+	Result += "\tFLWCMatrix ClipToWorld;\n";
+	Result += "\tFLWCMatrix ScreenToWorld;\n";
+	Result += "\tFLWCMatrix PrevClipToWorld;\n";
+	Result += "\tFLWCVector3 WorldCameraOrigin;\n";
+	Result += "\tFLWCVector3 WorldViewOrigin;\n";
+	Result += "\tFLWCVector3 PrevWorldCameraOrigin;\n";
+	Result += "\tFLWCVector3 PrevWorldViewOrigin;\n";
+	Result += "\tFLWCVector3 PreViewTranslation;\n";
+	Result += "\tFLWCVector3 PrevPreViewTranslation;\n";
+	Result += "};\n";
+
+	Result += "\tvoid FinalizeViewState(inout ViewState InOutView);\n";
 
 	// GetPrimaryView definition
-	Result += "ViewState GetPrimaryView()\r\n";
-	Result += "{\r\n";
-	Result += "\tViewState Result;\r\n";
+	Result += "ViewState GetPrimaryView()\n";
+	Result += "{\n";
+	Result += "\tViewState Result;\n";
 	for (int32 MemberIndex = 0; MemberIndex < StructMembersInstanced.Num(); ++MemberIndex)
 	{
 		const FShaderParametersMetadata::FMember& Member = StructMembersView[MemberIndex];
-		Result += FString::Printf(TEXT("\tResult.%s = View.%s;\r\n"), Member.GetName(), Member.GetName());
+		Result += FString::Printf(TEXT("\tResult.%s = View.%s;\n"), Member.GetName(), Member.GetName());
 	}
-	Result += "\tFinalizeViewState(Result);\r\n";
-	Result += "\treturn Result;\r\n";
-	Result += "}\r\n";
+	Result += "\tFinalizeViewState(Result);\n";
+	Result += "\treturn Result;\n";
+	Result += "}\n";
 
 	// GetInstancedView definition
-	Result += "#if (INSTANCED_STEREO || MOBILE_MULTI_VIEW)\r\n";
-	Result += "ViewState GetInstancedView(uint ViewIndex)\r\n";
-	Result += "{\r\n";
-	Result += "\tViewState Result;\r\n";
+	Result += "#if (INSTANCED_STEREO || MOBILE_MULTI_VIEW)\n";
+	Result += "ViewState GetInstancedView(uint ViewIndex)\n";
+	Result += "{\n";
+	Result += "\tViewState Result;\n";
 	for (int32 MemberIndex = 0; MemberIndex < StructMembersInstanced.Num(); ++MemberIndex)
 	{
 		const FShaderParametersMetadata::FMember& ViewMember = StructMembersView[MemberIndex];
@@ -7289,7 +7289,7 @@ void GenerateInstancedStereoCode(FString& Result, EShaderPlatform ShaderPlatform
 			checkf((InstancedViewMember.GetNumElements() % ViewMember.GetNumElements()) == 0, TEXT("Per-view arrays are expected to be stored in an array that is an exact multiple of the original array."));
 			for (uint32 ElementIndex = 0; ElementIndex < ViewMember.GetNumElements(); ElementIndex++)
 			{
-				Result += FString::Printf(TEXT("\tResult.%s[%u] = (%s) InstancedView.%s[ViewIndex * %u + %u];\r\n"),
+				Result += FString::Printf(TEXT("\tResult.%s[%u] = (%s) InstancedView.%s[ViewIndex * %u + %u];\n"),
 					ViewMember.GetName(), ElementIndex, *ViewMemberTypeName, InstancedViewMember.GetName(), ViewMember.GetNumElements(), ElementIndex);
 			}
 		}
@@ -7297,14 +7297,14 @@ void GenerateInstancedStereoCode(FString& Result, EShaderPlatform ShaderPlatform
 		{
 			// if View has a scalar field for this index, and InstancedView has an array with >1 elements -> per-view scalar
 			// 	Result.TranslatedWorldToClip = (float4x4) InstancedView_TranslatedWorldToClip[ViewIndex];
-			Result += FString::Printf(TEXT("\tResult.%s = (%s) InstancedView.%s[ViewIndex];\r\n"),
+			Result += FString::Printf(TEXT("\tResult.%s = (%s) InstancedView.%s[ViewIndex];\n"),
 				ViewMember.GetName(), *ViewMemberTypeName, InstancedViewMember.GetName());
 		}
 		else if (InstancedViewMember.GetNumElements() == ViewMember.GetNumElements())
 		{
 			// if View has the same number of elements for this index as InstancedView, it's backed by a view-dependent array, assume a view-independent field
 			// 	Result.TemporalAAParams = InstancedView_TemporalAAParams;
-			Result += FString::Printf(TEXT("\tResult.%s = InstancedView.%s;\r\n"),
+			Result += FString::Printf(TEXT("\tResult.%s = InstancedView.%s;\n"),
 				ViewMember.GetName(), InstancedViewMember.GetName());
 		}
 		else
@@ -7315,10 +7315,10 @@ void GenerateInstancedStereoCode(FString& Result, EShaderPlatform ShaderPlatform
 				);
 		}
 	}
-	Result += "\tFinalizeViewState(Result);\r\n";
-	Result += "\treturn Result;\r\n";
-	Result += "}\r\n";
-	Result += "#endif\r\n";
+	Result += "\tFinalizeViewState(Result);\n";
+	Result += "\treturn Result;\n";
+	Result += "}\n";
+	Result += "#endif\n";
 }
 
 void ValidateShaderFilePath(const FString& VirtualShaderFilePath, const FString& VirtualSourceFilePath)
