Remove check that is invalid since CL 25128474

It is now possible for FFolder::GetRootObjectAssociatedLevel() to change between calls to FActorFolders::RebuildFolderListForWorld().
The Guid will be updated by the call to FFolder::GetActorFolder()

This was found by switching the world asset on a LevelInstance between two levels sharing the same folder names

#rnx

diff --git a/Engine/Source/Runtime/Engine/Private/Folder.cpp b/Engine/Source/Runtime/Engine/Private/Folder.cpp
--- a/Engine/Source/Runtime/Engine/Private/Folder.cpp
+++ b/Engine/Source/Runtime/Engine/Private/Folder.cpp
@@ -153,9 +153,8 @@ UActorFolder* FFolder::GetActorFolder() const
 		ActorFolder = Level->GetActorFolder(Path);
 		if (ActorFolder)
 		{
-			// Cache ActorFolderGuid
- 			check(!ActorFolderGuid.IsValid());
-			ActorFolderGuid = ActorFolder->GetGuid();
+			// Cache ActorFolderGuid (GetRootObjectAssociatedLevel() can change between calls so avoid checking that it has already been set)
+ 			ActorFolderGuid = ActorFolder->GetGuid();
 			return ActorFolder;
 		}
 	}
