Fixed cooking indeterminism caused by code assuming that IAssetRegistry::GetAssets returns a deterministic ordered list (which is not true).

diff --git a/Engine/Source/Runtime/Engine/Private/Level.cpp b/Engine/Source/Runtime/Engine/Private/Level.cpp
--- a/Engine/Source/Runtime/Engine/Private/Level.cpp
+++ b/Engine/Source/Runtime/Engine/Private/Level.cpp
@@ -3966,6 +3966,9 @@ TArray<FString> ULevel::GetOnDiskExternalActorPackages(const FString& ExternalAc
 		IAssetRegistry& AssetRegistry = FModuleManager::LoadModuleChecked<FAssetRegistryModule>(TEXT("AssetRegistry")).Get();
 		AssetRegistry.ScanSynchronous({ ExternalActorsPath }, TArray<FString>());
 		AssetRegistry.GetAssets(Filter, ActorAssets);
+		// Sort the list as the order affects the cooking which needs to be deterministic.
+		// @todo_ow: Remove this once UE-198035 is fixed.
+		ActorAssets.Sort();
 
 		ActorPackageNames.Reserve(ActorAssets.Num());
 		Algo::Transform(ActorAssets, ActorPackageNames, [](const FAssetData& ActorAssetData) { return ActorAssetData.PackageName.ToString(); });
diff --git a/Engine/Source/Runtime/Engine/Private/WorldPartition/ActorDescContainer.cpp b/Engine/Source/Runtime/Engine/Private/WorldPartition/ActorDescContainer.cpp
--- a/Engine/Source/Runtime/Engine/Private/WorldPartition/ActorDescContainer.cpp
+++ b/Engine/Source/Runtime/Engine/Private/WorldPartition/ActorDescContainer.cpp
@@ -59,6 +59,9 @@ void UActorDescContainer::Initialize(const FInitializeParams& InitParams)
 
 		TRACE_CPUPROFILER_EVENT_SCOPE(GetAssets);
 		AssetRegistry.GetAssets(Filter, Assets);
+		// Sort the list as the order affects the cooking which needs to be deterministic.
+		// @todo_ow: Remove this once UE-198035 is fixed.
+		Assets.Sort();
 	}
 
 	FWorldPartitionClassDescRegistry& ClassDescRegistry = FWorldPartitionClassDescRegistry::Get();
diff --git a/Engine/Source/Runtime/Engine/Private/WorldPartition/DataLayer/WorldDataLayersActorDesc.cpp b/Engine/Source/Runtime/Engine/Private/WorldPartition/DataLayer/WorldDataLayersActorDesc.cpp
--- a/Engine/Source/Runtime/Engine/Private/WorldPartition/DataLayer/WorldDataLayersActorDesc.cpp
+++ b/Engine/Source/Runtime/Engine/Private/WorldPartition/DataLayer/WorldDataLayersActorDesc.cpp
@@ -280,6 +280,9 @@ const TArray<FDataLayerInstanceDesc>& FWorldDataLayersActorDesc::GetExternalPack
 			Filter.PackagePaths.Add(*ExternalObjectsPath);
 			TArray<FAssetData> Assets;
 			AssetRegistry.GetAssets(Filter, Assets);
+			// Sort the list as the order affects the cooking which needs to be deterministic.
+			// @todo_ow: Remove this once UE-198035 is fixed.
+			Assets.Sort();
 
 			for (const FAssetData& Asset : Assets)
 			{
diff --git a/Engine/Source/Runtime/Engine/Public/ExternalPackageHelper.h b/Engine/Source/Runtime/Engine/Public/ExternalPackageHelper.h
--- a/Engine/Source/Runtime/Engine/Public/ExternalPackageHelper.h
+++ b/Engine/Source/Runtime/Engine/Public/ExternalPackageHelper.h
@@ -117,6 +117,9 @@ void FExternalPackageHelper::LoadObjectsFromExternalPackages(UObject* InOuter, T
 	Filter.PackagePaths.Add(*ExternalObjectsPath);
 	TArray<FAssetData> Assets;
 	AssetRegistry.GetAssets(Filter, Assets);
+	// Sort the list as the order affects the cooking which needs to be deterministic.
+	// @todo_ow: Remove this once UE-198035 is fixed.
+	Assets.Sort();
 
 	ObjectPackageNames.Reserve(Assets.Num());
 	for (const FAssetData& Asset : Assets)
