World Partition RuntimeHashSet HLODs:
- Added manual naming of HLOD runtime partition layers to be able to match with streaming sources target grids inclusion/exclusion.
- Added detection of duplicated runtime partition layers names.
- Added detection of duplicated HLOD layers in HLOD setup layers.

diff --git a/Engine/Source/Runtime/Engine/Private/WorldPartition/RuntimeHashSet/RuntimePartition.cpp b/Engine/Source/Runtime/Engine/Private/WorldPartition/RuntimeHashSet/RuntimePartition.cpp
--- a/Engine/Source/Runtime/Engine/Private/WorldPartition/RuntimeHashSet/RuntimePartition.cpp
+++ b/Engine/Source/Runtime/Engine/Private/WorldPartition/RuntimeHashSet/RuntimePartition.cpp
@@ -3,22 +3,28 @@
 #include "WorldPartition/RuntimeHashSet/RuntimePartition.h"
 #include "WorldPartition/RuntimeHashSet/WorldPartitionRuntimeHashSet.h"
 
+URuntimePartition::URuntimePartition(const FObjectInitializer& ObjectInitializer)
+	: Super(ObjectInitializer)
+{
+#if WITH_EDITOR
+	SetDefaultValues();
+#endif
+}
+
 #if WITH_EDITOR
 void URuntimePartition::SetDefaultValues()
 {
-	Name = GetClass()->GetFName();
 	bBlockOnSlowStreaming = false;
 	bClientOnlyVisible = false;
 	Priority = 0;
 	LoadingRange = 25600;
-	bIsHLODSetup = false;
+	HLODIndex = INDEX_NONE;
 }
 
-void URuntimePartition::InitHLODRuntimePartitionFrom(const URuntimePartition* RuntimePartition, int32 HLODIndex)
+void URuntimePartition::InitHLODRuntimePartitionFrom(const URuntimePartition* InRuntimePartition, int32 InHLODIndex)
 {
-	Name = *FString::Printf(TEXT("%s_HLOD%d"), *RuntimePartition->Name.ToString(), HLODIndex);
-	LoadingRange = RuntimePartition->LoadingRange * 2;
-	bIsHLODSetup = true;
+	LoadingRange = InRuntimePartition->LoadingRange * 2;
+	HLODIndex = InHLODIndex;
 }
 
 void URuntimePartition::PostEditChangeProperty(FPropertyChangedEvent& InPropertyChangedEvent)
@@ -33,10 +39,10 @@ void URuntimePartition::PostEditChangeProperty(FPropertyChangedEvent& InProperty
 	Super::PostEditChangeProperty(InPropertyChangedEvent);
 }
 
-URuntimePartition* URuntimePartition::CreateHLODRuntimePartition(int32 HLODIndex) const
+URuntimePartition* URuntimePartition::CreateHLODRuntimePartition(int32 InHLODIndex) const
 {
 	URuntimePartition* HLODRuntimePartition = DuplicateObject<URuntimePartition>(this, GetOuter());
-	HLODRuntimePartition->InitHLODRuntimePartitionFrom(this, HLODIndex);
+	HLODRuntimePartition->InitHLODRuntimePartitionFrom(this, InHLODIndex);
 	return HLODRuntimePartition;
 }
 
diff --git a/Engine/Source/Runtime/Engine/Private/WorldPartition/RuntimeHashSet/RuntimePartitionLHGrid.cpp b/Engine/Source/Runtime/Engine/Private/WorldPartition/RuntimeHashSet/RuntimePartitionLHGrid.cpp
--- a/Engine/Source/Runtime/Engine/Private/WorldPartition/RuntimeHashSet/RuntimePartitionLHGrid.cpp
+++ b/Engine/Source/Runtime/Engine/Private/WorldPartition/RuntimeHashSet/RuntimePartitionLHGrid.cpp
@@ -68,10 +68,10 @@ void URuntimePartitionLHGrid::PostEditChangeProperty(FPropertyChangedEvent& Prop
 	Super::PostEditChangeProperty(PropertyChangedEvent);
 }
 
-void URuntimePartitionLHGrid::InitHLODRuntimePartitionFrom(const URuntimePartition* RuntimePartition, int32 HLODIndex)
+void URuntimePartitionLHGrid::InitHLODRuntimePartitionFrom(const URuntimePartition* InRuntimePartition, int32 InHLODIndex)
 {
-	Super::InitHLODRuntimePartitionFrom(RuntimePartition, HLODIndex);
-	const URuntimePartitionLHGrid* RuntimePartitionLHGrid = CastChecked<const URuntimePartitionLHGrid>(RuntimePartition);
+	Super::InitHLODRuntimePartitionFrom(InRuntimePartition, InHLODIndex);
+	const URuntimePartitionLHGrid* RuntimePartitionLHGrid = CastChecked<const URuntimePartitionLHGrid>(InRuntimePartition);
 	CellSize = RuntimePartitionLHGrid->CellSize * 2;
 }
 
diff --git a/Engine/Source/Runtime/Engine/Private/WorldPartition/RuntimeHashSet/WorldPartitionRuntimeHashSet.cpp b/Engine/Source/Runtime/Engine/Private/WorldPartition/RuntimeHashSet/WorldPartitionRuntimeHashSet.cpp
--- a/Engine/Source/Runtime/Engine/Private/WorldPartition/RuntimeHashSet/WorldPartitionRuntimeHashSet.cpp
+++ b/Engine/Source/Runtime/Engine/Private/WorldPartition/RuntimeHashSet/WorldPartitionRuntimeHashSet.cpp
@@ -460,7 +460,7 @@ void UWorldPartitionRuntimeHashSet::PostEditChangeChainProperty(FPropertyChanged
 
 	static FName NAME_RuntimePartitions(TEXT("RuntimePartitions"));
 	static FName NAME_HLODSetups(TEXT("HLODSetups"));
-	static FName NAME_HLODLayer(TEXT("HLODLayer"));
+	static FName NAME_HLODLayers(TEXT("HLODLayers"));
 
 	FName PropertyName = PropertyChangedEvent.Property ? PropertyChangedEvent.Property->GetFName() : NAME_None;
 
@@ -475,9 +475,10 @@ void UWorldPartitionRuntimeHashSet::PostEditChangeChainProperty(FPropertyChanged
 
 		if (RuntimePartitionDesc.Class)
 		{
-			RuntimePartitionDesc.Name = RuntimePartitionDesc.Class->GetFName();
+			RuntimePartitionDesc.Name = *FString::Printf(TEXT("%s_%d"), *RuntimePartitionDesc.Class->GetName(), RuntimePartitionIndex);
 			RuntimePartitionDesc.MainLayer = NewObject<URuntimePartition>(this, RuntimePartitionDesc.Class, NAME_None);
 			RuntimePartitionDesc.MainLayer->SetDefaultValues();
+			RuntimePartitionDesc.MainLayer->Name = RuntimePartitionDesc.Name;
 		}
 	}
 	else if (PropertyName == GET_MEMBER_NAME_CHECKED(FRuntimePartitionDesc, Name))
@@ -487,26 +488,55 @@ void UWorldPartitionRuntimeHashSet::PostEditChangeChainProperty(FPropertyChanged
 
 		FRuntimePartitionDesc& RuntimePartitionDesc = RuntimePartitions[RuntimePartitionIndex];
 
-		if (RuntimePartitionDesc.Name == NAME_PersistentLevel)
+		int32 HLODSetupsIndex = PropertyChangedEvent.GetArrayIndex(NAME_HLODSetups.ToString());
+		if (RuntimePartitionDesc.HLODSetups.IsValidIndex(HLODSetupsIndex))
 		{
-			RuntimePartitionDesc.Name = RuntimePartitionDesc.Class->GetFName();
+			FRuntimePartitionHLODSetup& RuntimePartitionHLODSetup = RuntimePartitionDesc.HLODSetups[HLODSetupsIndex];
+
+			if (RuntimePartitionDesc.Name == NAME_PersistentLevel)
+			{
+				RuntimePartitionHLODSetup.Name = *FString::Printf(TEXT("HLOD_%d"), RuntimePartitionHLODSetup.PartitionLayer->HLODIndex);
+			}
+			else
+			{
+				for (int32 CurHLODSetupsIndex = 0; CurHLODSetupsIndex < RuntimePartitionDesc.HLODSetups.Num(); CurHLODSetupsIndex++)
+				{
+					if (CurHLODSetupsIndex != HLODSetupsIndex)
+					{
+						if (RuntimePartitionHLODSetup.Name == RuntimePartitionDesc.HLODSetups[CurHLODSetupsIndex].Name)
+						{
+							RuntimePartitionHLODSetup.Name = *FString::Printf(TEXT("HLOD_%d"), RuntimePartitionHLODSetup.PartitionLayer->HLODIndex);
+							break;
+						}
+					}
+				}
+			}
+
+			RuntimePartitionHLODSetup.PartitionLayer->Name = RuntimePartitionHLODSetup.Name;
 		}
 		else
 		{
-			for (int32 CurRuntimePartitionIndex = 0; CurRuntimePartitionIndex < RuntimePartitions.Num(); CurRuntimePartitionIndex++)
+			if (RuntimePartitionDesc.Name == NAME_PersistentLevel)
+			{
+				RuntimePartitionDesc.Name = RuntimePartitionDesc.Class->GetFName();
+			}
+			else
 			{
-				if (CurRuntimePartitionIndex != RuntimePartitionIndex)
+				for (int32 CurRuntimePartitionIndex = 0; CurRuntimePartitionIndex < RuntimePartitions.Num(); CurRuntimePartitionIndex++)
 				{
-					if (RuntimePartitionDesc.Name == RuntimePartitions[CurRuntimePartitionIndex].Name)
+					if (CurRuntimePartitionIndex != RuntimePartitionIndex)
 					{
-						RuntimePartitionDesc.Name = RuntimePartitionDesc.Class->GetFName();
-						break;
+						if (RuntimePartitionDesc.Name == RuntimePartitions[CurRuntimePartitionIndex].Name)
+						{
+							RuntimePartitionDesc.Name = RuntimePartitionDesc.Class->GetFName();
+							break;
+						}
 					}
 				}
 			}
-		}
 
-		RuntimePartitionDesc.MainLayer->Name = RuntimePartitionDesc.Name;
+			RuntimePartitionDesc.MainLayer->Name = RuntimePartitionDesc.Name;
+		}
 	}
 	else if (PropertyName == NAME_HLODSetups)
 	{
@@ -520,7 +550,44 @@ void UWorldPartitionRuntimeHashSet::PostEditChangeChainProperty(FPropertyChanged
 			{
 				FRuntimePartitionHLODSetup& RuntimePartitionHLODSetup = RuntimePartitionDesc.HLODSetups[HLODSetupsIndex];
 				URuntimePartition* ParentRuntimePartition = HLODSetupsIndex ? RuntimePartitionDesc.HLODSetups[HLODSetupsIndex - 1].PartitionLayer : RuntimePartitionDesc.MainLayer;
+				RuntimePartitionHLODSetup.Name = *FString::Printf(TEXT("HLOD_%d"), HLODSetupsIndex);
 				RuntimePartitionHLODSetup.PartitionLayer = ParentRuntimePartition->CreateHLODRuntimePartition(HLODSetupsIndex);
+				RuntimePartitionHLODSetup.PartitionLayer->Name = RuntimePartitionHLODSetup.Name;
+			}
+		}
+	}
+	else if (PropertyName == NAME_HLODLayers)
+	{
+		int32 RuntimePartitionIndex = PropertyChangedEvent.GetArrayIndex(NAME_RuntimePartitions.ToString());
+		if (RuntimePartitions.IsValidIndex(RuntimePartitionIndex))
+		{
+			FRuntimePartitionDesc& RuntimePartitionDesc = RuntimePartitions[RuntimePartitionIndex];
+
+			int32 HLODSetupsIndex = PropertyChangedEvent.GetArrayIndex(NAME_HLODSetups.ToString());
+			if (RuntimePartitionDesc.HLODSetups.IsValidIndex(HLODSetupsIndex))
+			{
+				FRuntimePartitionHLODSetup& RuntimePartitionHLODSetup = RuntimePartitionDesc.HLODSetups[HLODSetupsIndex];
+
+				int32 HLODLayersIndex = PropertyChangedEvent.GetArrayIndex(NAME_HLODLayers.ToString());
+				if (RuntimePartitionHLODSetup.HLODLayers.IsValidIndex(HLODLayersIndex))
+				{
+					const UHLODLayer* HLODLayer = RuntimePartitionHLODSetup.HLODLayers[HLODLayersIndex];
+
+					// Remove duplicated entries
+					for (int32 CurrentHLODSetupsIndex = 0; CurrentHLODSetupsIndex < RuntimePartitionDesc.HLODSetups.Num(); CurrentHLODSetupsIndex++)
+					{
+						FRuntimePartitionHLODSetup& CurrentRuntimePartitionHLODSetup = RuntimePartitionDesc.HLODSetups[CurrentHLODSetupsIndex];
+						for (int32 CurrentHLODLayerIndex = 0; CurrentHLODLayerIndex < CurrentRuntimePartitionHLODSetup.HLODLayers.Num(); CurrentHLODLayerIndex++)
+						{
+							const UHLODLayer* CurrentHLODLayer = CurrentRuntimePartitionHLODSetup.HLODLayers[CurrentHLODLayerIndex];
+							if (((CurrentHLODSetupsIndex != HLODSetupsIndex) || (CurrentHLODLayerIndex != HLODLayersIndex)) && (CurrentHLODLayer == HLODLayer))
+							{
+								CurrentRuntimePartitionHLODSetup.HLODLayers.RemoveAt(CurrentHLODLayerIndex--);
+								break;
+							}
+						}
+					}
+				}
 			}
 		}
 	}
diff --git a/Engine/Source/Runtime/Engine/Private/WorldPartition/RuntimeHashSet/WorldPartitionRuntimeHashSetStreamingGeneration.cpp b/Engine/Source/Runtime/Engine/Private/WorldPartition/RuntimeHashSet/WorldPartitionRuntimeHashSetStreamingGeneration.cpp
--- a/Engine/Source/Runtime/Engine/Private/WorldPartition/RuntimeHashSet/WorldPartitionRuntimeHashSetStreamingGeneration.cpp
+++ b/Engine/Source/Runtime/Engine/Private/WorldPartition/RuntimeHashSet/WorldPartitionRuntimeHashSetStreamingGeneration.cpp
@@ -179,7 +179,7 @@ bool UWorldPartitionRuntimeHashSet::GenerateStreaming(UWorldPartitionStreamingPo
 		RuntimeCell->SetPriority(CellDescInstance.Priority);
 		RuntimeCell->SetClientOnlyVisible(CellDescInstance.bClientOnlyVisible);
 		RuntimeCell->SetBlockOnSlowLoading(CellDescInstance.bBlockOnSlowStreaming);
-		RuntimeCell->SetIsHLOD(CellDescInstance.SourcePartition->bIsHLODSetup);
+		RuntimeCell->SetIsHLOD(CellDescInstance.SourcePartition->HLODIndex != INDEX_NONE);
 		RuntimeCell->SetGuid(CellUniqueId.Guid);
 
 		UWorldPartitionRuntimeCellDataSpatialHashSet* RuntimeCellDataHashSet = CastChecked<UWorldPartitionRuntimeCellDataSpatialHashSet>(RuntimeCell->RuntimeCellData);
diff --git a/Engine/Source/Runtime/Engine/Public/WorldPartition/RuntimeHashSet/RuntimePartition.h b/Engine/Source/Runtime/Engine/Public/WorldPartition/RuntimeHashSet/RuntimePartition.h
--- a/Engine/Source/Runtime/Engine/Public/WorldPartition/RuntimeHashSet/RuntimePartition.h
+++ b/Engine/Source/Runtime/Engine/Public/WorldPartition/RuntimeHashSet/RuntimePartition.h
@@ -8,7 +8,7 @@
 UCLASS(Abstract, CollapseCategories)
 class URuntimePartition : public UObject
 {
-	GENERATED_BODY()
+	GENERATED_UCLASS_BODY()
 
 public:
 #if WITH_EDITOR
@@ -16,7 +16,7 @@ public:
 	virtual void PostEditChangeProperty(FPropertyChangedEvent& InPropertyChangedEvent) override;
 	//~ End UObject Interface.
 
-	URuntimePartition* CreateHLODRuntimePartition(int32 HLODIndex) const;
+	URuntimePartition* CreateHLODRuntimePartition(int32 InHLODIndex) const;
 
 	/*
 	 * Represents a cell descriptor, generated by runtime partitions. This is a streaming cell containing actors, without taking into account data layers, content bundles, etc.
@@ -65,7 +65,7 @@ public:
 
 	virtual void SetDefaultValues();
 	virtual bool SupportsHLODs() const PURE_VIRTUAL(URuntimePartition::SupportsHLODs, return true;);
-	virtual void InitHLODRuntimePartitionFrom(const URuntimePartition* RuntimePartition, int32 HLODIndex);
+	virtual void InitHLODRuntimePartitionFrom(const URuntimePartition* InRuntimePartition, int32 InHLODIndex);
 	virtual bool IsValidPartitionTokens(const TArray<FName>& InPartitionTokens) const PURE_VIRTUAL(URuntimePartition::IsValidPartitionTokens, return false;);
 	virtual bool GenerateStreaming(const FGenerateStreamingParams& InParams, FGenerateStreamingResult& OutResult) PURE_VIRTUAL(URuntimePartition::GenerateStreaming, return false;);
 #endif
@@ -74,20 +74,20 @@ public:
 	UPROPERTY()
 	FName Name;
 
-	UPROPERTY(EditAnywhere, Category = RuntimeSettings, Meta = (EditCondition = "!bIsHLODSetup", EditConditionHides, HideEditConditionToggle))
+	UPROPERTY(EditAnywhere, Category = RuntimeSettings, Meta = (EditCondition = "HLODSetup != INDEX_NONE", EditConditionHides, HideEditConditionToggle))
 	bool bBlockOnSlowStreaming;
 
-	UPROPERTY(EditAnywhere, Category = RuntimeSettings, Meta = (EditCondition = "!bIsHLODSetup", EditConditionHides, HideEditConditionToggle))
+	UPROPERTY(EditAnywhere, Category = RuntimeSettings, Meta = (EditCondition = "HLODSetup != INDEX_NONE", EditConditionHides, HideEditConditionToggle))
 	bool bClientOnlyVisible;
 
-	UPROPERTY(EditAnywhere, Category = RuntimeSettings, Meta = (EditCondition = "!bIsHLODSetup", EditConditionHides, HideEditConditionToggle))
+	UPROPERTY(EditAnywhere, Category = RuntimeSettings, Meta = (EditCondition = "HLODSetup != INDEX_NONE", EditConditionHides, HideEditConditionToggle))
 	int32 Priority;
 
 	UPROPERTY(EditAnywhere, Category = RuntimeSettings)
 	int32 LoadingRange;
 
 	UPROPERTY()
-	bool bIsHLODSetup;
+	int32 HLODIndex;
 #endif
 
 protected:
diff --git a/Engine/Source/Runtime/Engine/Public/WorldPartition/RuntimeHashSet/RuntimePartitionLHGrid.h b/Engine/Source/Runtime/Engine/Public/WorldPartition/RuntimeHashSet/RuntimePartitionLHGrid.h
--- a/Engine/Source/Runtime/Engine/Public/WorldPartition/RuntimeHashSet/RuntimePartitionLHGrid.h
+++ b/Engine/Source/Runtime/Engine/Public/WorldPartition/RuntimeHashSet/RuntimePartitionLHGrid.h
@@ -18,7 +18,7 @@ public:
 
 	//~ Begin URuntimePartition interface
 	virtual bool SupportsHLODs() const override { return true; }
-	virtual void InitHLODRuntimePartitionFrom(const URuntimePartition* RuntimePartition, int32 HLODIndex);
+	virtual void InitHLODRuntimePartitionFrom(const URuntimePartition* InRuntimePartition, int32 InHLODIndex);
 	virtual void SetDefaultValues() override;
 	virtual bool IsValidPartitionTokens(const TArray<FName>& InPartitionTokens) const override;
 	virtual bool GenerateStreaming(const FGenerateStreamingParams& InParams, FGenerateStreamingResult& OutResult) override;
diff --git a/Engine/Source/Runtime/Engine/Public/WorldPartition/RuntimeHashSet/WorldPartitionRuntimeHashSet.h b/Engine/Source/Runtime/Engine/Public/WorldPartition/RuntimeHashSet/WorldPartitionRuntimeHashSet.h
--- a/Engine/Source/Runtime/Engine/Public/WorldPartition/RuntimeHashSet/WorldPartitionRuntimeHashSet.h
+++ b/Engine/Source/Runtime/Engine/Public/WorldPartition/RuntimeHashSet/WorldPartitionRuntimeHashSet.h
@@ -19,11 +19,15 @@ struct FRuntimePartitionHLODSetup
 {
 	GENERATED_USTRUCT_BODY()
 
+	/** Name for this HLOD layer setup */
+	UPROPERTY(EditAnywhere, Category = RuntimeSettings)
+	FName Name;
+
 	/** Associated HLOD Layer objects */
 	UPROPERTY(EditAnywhere, Category = RuntimeSettings)
 	TArray<TObjectPtr<const UHLODLayer>> HLODLayers;
 
-	UPROPERTY(VisibleAnywhere, Category = RuntimeSettings, Instanced)
+	UPROPERTY(VisibleAnywhere, Category = RuntimeSettings, Instanced, Meta = (NoResetToDefault, TitleProperty = "Name"))
 	TObjectPtr<URuntimePartition> PartitionLayer;
 };
 
