World Partition: remove 5.1 deprecated functions from FStreamingSourceShapeHelper and FWorldPartitionStreamingQuerySource interfaces,

diff --git a/Engine/Source/Runtime/Engine/Public/WorldPartition/WorldPartitionStreamingSource.h b/Engine/Source/Runtime/Engine/Public/WorldPartition/WorldPartitionStreamingSource.h
--- a/Engine/Source/Runtime/Engine/Public/WorldPartition/WorldPartitionStreamingSource.h
+++ b/Engine/Source/Runtime/Engine/Public/WorldPartition/WorldPartitionStreamingSource.h
@@ -169,18 +169,6 @@ public:
 		return false;
 	}
 
-	UE_DEPRECATED(5.1, "Use IsSourceAffectingGrid that takes Target Sets and Behavior instead.")
-	FORCEINLINE static bool IsSourceAffectingGrid(FName InSourceTargetGrid, const FSoftObjectPath& InSourceTargetHLODLayer, FName InGridName, const FSoftObjectPath& InGridHLODLayer)
-	{
-		if ((InSourceTargetGrid.IsNone() && InSourceTargetHLODLayer.IsNull()) ||
-			(InSourceTargetHLODLayer.IsNull() && (InSourceTargetGrid == InGridName)) ||
-			(!InSourceTargetHLODLayer.IsNull() && (InSourceTargetHLODLayer == InGridHLODLayer)))
-		{
-			return true;
-		}
-		return false;
-	}
-
 	FORCEINLINE static void ForEachShape(float InGridLoadingRange, float InDefaultRadius, bool bInProjectIn2D, const FVector& InLocation, const FRotator& InRotation, const TArray<FStreamingSourceShape>& InShapes, TFunctionRef<void(const FSphericalSector&)> InOperation, float InExtraRadius = 0.f, float InExtraAngle = 0.f)
 	{
 		const FTransform Transform(bInProjectIn2D ? FRotator(0, InRotation.Yaw, 0) : InRotation, InLocation);
@@ -311,14 +299,8 @@ struct FWorldPartitionStreamingQuerySource
 	/** Defines how TargetGrids/TargetHLODLayers will be applied to this streaming source. */
 	EStreamingSourceTargetBehavior TargetBehavior;
 
-	UE_DEPRECATED(5.1, "Use TargetGrids instead.")
-	FName TargetGrid;
-
 	TSet<FName> TargetGrids;
 
-	UE_DEPRECATED(5.1, "Use TargetHLODLayers")
-	FSoftObjectPath TargetHLODLayer;
-
 	TSet<FSoftObjectPath> TargetHLODLayers;
 	
 	TArray<FStreamingSourceShape> Shapes;
