World Partition: force streaming sources logs to verbose.
Can be reactivated with "wp.runtime.setlogworldpartitionverbosity verbose" command.
Deprecated "wp.editor.setlogworldpartitionverbosity" command in favor of the new runtime version.

diff --git a/Engine/Source/Runtime/Engine/Private/WorldPartition/WorldPartition.cpp b/Engine/Source/Runtime/Engine/Private/WorldPartition/WorldPartition.cpp
--- a/Engine/Source/Runtime/Engine/Private/WorldPartition/WorldPartition.cpp
+++ b/Engine/Source/Runtime/Engine/Private/WorldPartition/WorldPartition.cpp
@@ -232,9 +232,10 @@ static FAutoConsoleCommand DumpActorDescs(
 	})
 );
 
-static FAutoConsoleCommand SetLogWorldPartitionVerbosity(
+UE_DEPRECATED(5.4, "wp.Editor.SetLogWorldPartitionVerbosity is deprecated, use wp.Runtime.SetLogWorldPartitionVerbosity instead")
+static FAutoConsoleCommand SetLogWorldPartitionEditorVerbosity(
 	TEXT("wp.Editor.SetLogWorldPartitionVerbosity"),
-	TEXT("Change the WorldPartition verbosity log verbosity."),
+	TEXT("Change the WorldPartition editor log verbosity."),
 	FConsoleCommandWithArgsDelegate::CreateLambda([](const TArray<FString>& Args)
 	{
 		if (Args.Num() == 1)
@@ -270,6 +271,27 @@ public:
 UWorldPartition::FWorldPartitionChangedEvent UWorldPartition::WorldPartitionChangedEvent;
 #endif
 
+#if !NO_LOGGING
+static FAutoConsoleCommand SetLogWorldPartitionVerbosity(
+	TEXT("wp.Runtime.SetLogWorldPartitionVerbosity"),
+	TEXT("Change the WorldPartition log verbosity."),
+	FConsoleCommandWithArgsDelegate::CreateLambda([](const TArray<FString>& Args)
+	{
+		if (Args.Num() == 1)
+		{
+			if (Args[0].Contains(TEXT("Verbose")))
+			{
+				LogWorldPartition.SetVerbosity(ELogVerbosity::Verbose);
+			}
+			else
+			{
+				LogWorldPartition.SetVerbosity(LogWorldPartition.GetCompileTimeVerbosity());
+			}
+		}
+	})
+);
+#endif
+
 UWorldPartition::UWorldPartition(const FObjectInitializer& ObjectInitializer)
 	: Super(ObjectInitializer)
 #if WITH_EDITOR
diff --git a/Engine/Source/Runtime/Engine/Private/WorldPartition/WorldPartitionStreamingPolicy.cpp b/Engine/Source/Runtime/Engine/Private/WorldPartition/WorldPartitionStreamingPolicy.cpp
--- a/Engine/Source/Runtime/Engine/Private/WorldPartition/WorldPartitionStreamingPolicy.cpp
+++ b/Engine/Source/Runtime/Engine/Private/WorldPartition/WorldPartitionStreamingPolicy.cpp
@@ -108,22 +108,6 @@ void UWorldPartitionStreamingPolicy::UpdateStreamingSources(bool bCanOptimizeUpd
 	UpdateStreamingSourcesHash = NewUpdateStreamingSourcesHash;
 }
 
-#define WORLDPARTITION_LOG_UPDATESTREAMINGSTATE(Verbosity)\
-UE_SUPPRESS(LogWorldPartition, Verbosity, \
-{ \
-	if ((bIsStreamingInEnabled && (ToActivateCells.Num() > 0 || ToLoadCells.Num() > 0)) || ToUnloadCells.Num() > 0) \
-	{ \
-		UE_LOG(LogWorldPartition, Verbosity, TEXT("UWorldPartitionStreamingPolicy: CellsToActivate(%d), CellsToLoad(%d), CellsToUnload(%d)"), ToActivateCells.Num(), ToLoadCells.Num(), ToUnloadCells.Num()); \
-		FTransform LocalToWorld = WorldPartition->GetInstanceTransform(); \
-		for (int i = 0; i < StreamingSources.Num(); ++i) \
-		{ \
-			FVector ViewLocation = LocalToWorld.TransformPosition(StreamingSources[i].Location); \
-			FRotator ViewRotation = LocalToWorld.TransformRotation(StreamingSources[i].Rotation.Quaternion()).Rotator(); \
-			UE_LOG(LogWorldPartition, Verbosity, TEXT("UWorldPartitionStreamingPolicy: Sources[%d] = %s,%s"), i, *ViewLocation.ToString(), *ViewRotation.ToString()); \
-		} \
-	} \
-}) \
-
 bool UWorldPartitionStreamingPolicy::IsInBlockTillLevelStreamingCompleted(bool bIsCausedByBadStreamingPerformance /* = false*/) const
 {
 	const UWorld* World = GetWorld();
@@ -557,14 +541,18 @@ void UWorldPartitionStreamingPolicy::UpdateStreamingState()
 		BuildCellsToUnload(LoadedCells);
 	}
 
-	if(World->bMatchStarted)
-	{
-		WORLDPARTITION_LOG_UPDATESTREAMINGSTATE(Verbose);
-	}
-	else
-	{
-		WORLDPARTITION_LOG_UPDATESTREAMINGSTATE(Log);
-	}
+	UE_SUPPRESS(LogWorldPartition, Verbose,
+		if ((bIsStreamingInEnabled && (ToActivateCells.Num() > 0 || ToLoadCells.Num() > 0)) || ToUnloadCells.Num() > 0)
+		{
+			UE_LOG(LogWorldPartition, Verbose, TEXT("UWorldPartitionStreamingPolicy: CellsToActivate(%d), CellsToLoad(%d), CellsToUnload(%d)"), ToActivateCells.Num(), ToLoadCells.Num(), ToUnloadCells.Num());
+			FTransform LocalToWorld = WorldPartition->GetInstanceTransform();
+			for (int i = 0; i < StreamingSources.Num(); ++i)
+			{
+				FVector ViewLocation = LocalToWorld.TransformPosition(StreamingSources[i].Location);
+				FRotator ViewRotation = LocalToWorld.TransformRotation(StreamingSources[i].Rotation.Quaternion()).Rotator();
+				UE_LOG(LogWorldPartition, Verbose, TEXT("UWorldPartitionStreamingPolicy: Sources[%d] = %s,%s"), i, *ViewLocation.ToString(), *ViewRotation.ToString());
+			}
+		});
 
 #if !UE_BUILD_SHIPPING
 	UpdateDebugCellsStreamingPriority(FrameActivateCells, FrameLoadCells);
