Fix deprecation warning.

diff --git a/Engine/Source/Runtime/Engine/Private/WorldPartition/WorldPartition.cpp b/Engine/Source/Runtime/Engine/Private/WorldPartition/WorldPartition.cpp
--- a/Engine/Source/Runtime/Engine/Private/WorldPartition/WorldPartition.cpp
+++ b/Engine/Source/Runtime/Engine/Private/WorldPartition/WorldPartition.cpp
@@ -232,26 +232,6 @@ static FAutoConsoleCommand DumpActorDescs(
 	})
 );
 
-UE_DEPRECATED(5.4, "wp.Editor.SetLogWorldPartitionVerbosity is deprecated, use wp.Runtime.SetLogWorldPartitionVerbosity instead")
-static FAutoConsoleCommand SetLogWorldPartitionEditorVerbosity(
-	TEXT("wp.Editor.SetLogWorldPartitionVerbosity"),
-	TEXT("Change the WorldPartition editor log verbosity."),
-	FConsoleCommandWithArgsDelegate::CreateLambda([](const TArray<FString>& Args)
-	{
-		if (Args.Num() == 1)
-		{
-			if (Args[0].Contains(TEXT("Verbose")))
-			{
-				LogWorldPartition.SetVerbosity(ELogVerbosity::Verbose);
-			}
-			else
-			{
-				LogWorldPartition.SetVerbosity(LogWorldPartition.GetCompileTimeVerbosity());
-			}
-		}
-	})
-);
-
 class FLoaderAdapterAlwaysLoadedActors : public FLoaderAdapterShape
 {
 public:
