Increase net reliable buffer to 512.

diff --git a/Engine/Source/Runtime/Engine/Classes/Engine/NetConnection.h b/Engine/Source/Runtime/Engine/Classes/Engine/NetConnection.h
--- a/Engine/Source/Runtime/Engine/Classes/Engine/NetConnection.h
+++ b/Engine/Source/Runtime/Engine/Classes/Engine/NetConnection.h
@@ -77,7 +77,7 @@ extern ENGINE_API TAutoConsoleVariable<int32> CVarNetEnableCongestionControl;
 /*-----------------------------------------------------------------------------
 	Types.
 -----------------------------------------------------------------------------*/
-enum { RELIABLE_BUFFER = 256 }; // Power of 2 >= 1.
+enum { RELIABLE_BUFFER = 512 }; // Power of 2 >= 1.
 enum { MAX_PACKETID = FNetPacketNotify::SequenceNumberT::SeqNumberCount };  // Power of 2 >= 1, covering guaranteed loss/misorder time.
 enum { MAX_CHSEQUENCE = 1024 }; // Power of 2 >RELIABLE_BUFFER, covering loss/misorder time.
 enum { MAX_BUNCH_HEADER_BITS = 256 };
