[StaticMeshComponent] No longer relevant to navigation until compilation of the associated mesh is completed. This approach replaces the previous one consisting in returning either the render bounds or invalid bounds which could produce log spam when loading level and compiling multiple meshes.

diff --git a/Engine/Source/Runtime/Engine/Private/Components/StaticMeshComponent.cpp b/Engine/Source/Runtime/Engine/Private/Components/StaticMeshComponent.cpp
--- a/Engine/Source/Runtime/Engine/Private/Components/StaticMeshComponent.cpp
+++ b/Engine/Source/Runtime/Engine/Private/Components/StaticMeshComponent.cpp
@@ -3011,21 +3011,28 @@ bool UStaticMeshComponent::DoCustomNavigableGeometryExport(FNavigableGeometryExp
 {
 	const FVector Scale3D = GetComponentToWorld().GetScale3D();
 
-	// Pending compilation, RecreatePhysicsState will be called to update the navigation system once compilation finishes.
-	if (!Scale3D.IsZero() && GetStaticMesh() && !GetStaticMesh()->IsCompiling() && GetStaticMesh()->GetNavCollision())
+	if (!Scale3D.IsZero())
 	{
-		const UNavCollisionBase* NavCollision = GetStaticMesh()->GetNavCollision();
-		if (ShouldExportAsObstacle(*NavCollision))
+		if (const UStaticMesh* Mesh = GetStaticMesh())
 		{
-			// skip default export
-			return false;
-		}
+			if (ensureMsgf(!Mesh->IsCompiling(), TEXT("%s is not considered relevant to navigation until associated mesh is compiled."), *GetFullName()))
+			{
+				if (const UNavCollisionBase* NavCollision = Mesh->GetNavCollision())
+				{
+					if (ShouldExportAsObstacle(*NavCollision))
+					{
+						// skip default export
+						return false;
+					}
 
-		const bool bHasData = NavCollision->ExportGeometry(GetComponentToWorld(), GeomExport);
-		if (bHasData)
-		{
-			// skip default export
-			return false;
+					const bool bHasData = NavCollision->ExportGeometry(GetComponentToWorld(), GeomExport);
+					if (bHasData)
+					{
+						// skip default export
+						return false;
+					}
+				}
+			}
 		}
 	}
 
@@ -3075,27 +3082,30 @@ UMaterialInterface* UStaticMeshComponent::GetMaterialFromCollisionFaceIndex(int3
 
 bool UStaticMeshComponent::IsNavigationRelevant() const
 {
-	return GetStaticMesh() != nullptr && GetStaticMesh()->IsNavigationRelevant() && Super::IsNavigationRelevant();
+	if (const UStaticMesh* Mesh = GetStaticMesh())
+	{
+		// Pending compilation, update to the the navigation system will be done once compilation finishes.
+		return !Mesh->IsCompiling() && Mesh->IsNavigationRelevant() && Super::IsNavigationRelevant();
+	}
+
+	return false;
 }
 
 FBox UStaticMeshComponent::GetNavigationBounds() const
 {
-	if (GetStaticMesh())
+	if (const UStaticMesh* Mesh = GetStaticMesh())
 	{
-		if (GetStaticMesh()->IsCompiling())
-		{
-			// Navigation bounds will get queried again once async static mesh compilation finishes
-			return FBox();
-		}
-
-		if (const UNavCollisionBase* NavCollision = GetStaticMesh()->GetNavCollision())
+		if (ensureMsgf(!Mesh->IsCompiling(), TEXT("%s is not considered relevant to navigation until associated mesh is compiled."), *GetFullName()))
 		{
-			FBox NavBounds = NavCollision->GetBounds();
-			if (!NavBounds.IsValid)
+			if (const UNavCollisionBase* NavCollision = Mesh->GetNavCollision())
 			{
-				NavBounds = GetStaticMesh()->GetBounds().GetBox();
+				FBox NavBounds = NavCollision->GetBounds();
+				if (!NavBounds.IsValid)
+				{
+					NavBounds = Mesh->GetBounds().GetBox();
+				}
+				return NavBounds.TransformBy(GetComponentTransform());
 			}
-			return NavBounds.TransformBy(GetComponentTransform());
 		}
 	}
 
@@ -3107,14 +3117,20 @@ void UStaticMeshComponent::GetNavigationData(FNavigationRelevantData& Data) cons
 	Super::GetNavigationData(Data);
 
 	const FVector Scale3D = GetComponentToWorld().GetScale3D();
-	
-	// Navigation data will get refreshed once async static mesh compilation finishes
-	if (!Scale3D.IsZero() && GetStaticMesh() && !GetStaticMesh()->IsCompiling() && GetStaticMesh()->GetNavCollision())
+	if (!Scale3D.IsZero())
 	{
-		UNavCollisionBase* NavCollision = GetStaticMesh()->GetNavCollision();
-		if (ShouldExportAsObstacle(*NavCollision))
+		if (const UStaticMesh* Mesh = GetStaticMesh())
 		{
-			NavCollision->GetNavigationModifier(Data.Modifiers, GetComponentTransform());
+			if (ensureMsgf(!Mesh->IsCompiling(), TEXT("%s is not considered relevant to navigation until associated mesh is compiled."), *GetFullName()))
+			{
+				if (UNavCollisionBase* NavCollision = Mesh->GetNavCollision())
+				{
+					if (ShouldExportAsObstacle(*NavCollision))
+					{
+						NavCollision->GetNavigationModifier(Data.Modifiers, GetComponentTransform());
+					}
+				}
+			}
 		}
 	}
 }
diff --git a/Engine/Source/Runtime/NavigationSystem/Private/NavigationOctree.cpp b/Engine/Source/Runtime/NavigationSystem/Private/NavigationOctree.cpp
--- a/Engine/Source/Runtime/NavigationSystem/Private/NavigationOctree.cpp
+++ b/Engine/Source/Runtime/NavigationSystem/Private/NavigationOctree.cpp
@@ -163,6 +163,9 @@ void FNavigationOctree::AddNode(UObject* ElementOb, INavRelevantInterface* NavEl
 #if WITH_EDITOR
 			const IInterface_AsyncCompilation* AsyncCompiledObject = Cast<IInterface_AsyncCompilation>(ElementOb);
 			bIsCompiling = AsyncCompiledObject && AsyncCompiledObject->IsCompiling();
+
+			UE_CLOG(bIsCompiling, LogNavigation, Warning, TEXT("%hs: Objects %s should not be considered relevant to navigation until associated asset compilation is completed."),
+				__FUNCTION__, *GetFullNameSafe(ElementOb));
 #endif
 
 			// Skip custom navigation export during async compilation, the node will be invalidated once
