Don't clear Initializer in FRayTracingGeometry::ReleaseRHIForStreaming(...)
- clearing initializer causes streamed out geometries to not be properly initialized when raytracing is toggled on.
- static mesh ray tracing geometry initializer doesn't change at runtime, so instead of clearing initializer on stream-out and setting it up again during stream-in, we only need to set Initializer.Type as appropriate.
- added CheckRayTracingGeometryInitializer(...) to verify that initializer is matches previous logic.

Repro steps for the issue this change fixes:
1 - Run project with mesh streaming and dynamic ray tracing enabled.
2 - show RayTracingDebug (enable some raytracing scene visualization)
3 - r.RayTracing.Enable 0 (disable raytracing)
4 - r.Streaming.StressTest 1 (enable streaming stress test which causes some LODs to be streamed out)
5 - r.RayTracing.Enable 1 (re-enable raytracing)
6 - r.Streaming.StressTest 0 (disable streaming stress test which causes LODs to be streamed in)
7 - Notice that, without this change, ray tracing instances are still using the lower LOD geometry.

diff --git a/Engine/Source/Runtime/Engine/Private/Streaming/StaticMeshUpdate.cpp b/Engine/Source/Runtime/Engine/Private/Streaming/StaticMeshUpdate.cpp
--- a/Engine/Source/Runtime/Engine/Private/Streaming/StaticMeshUpdate.cpp
+++ b/Engine/Source/Runtime/Engine/Private/Streaming/StaticMeshUpdate.cpp
@@ -142,14 +142,15 @@ void FStaticMeshStreamIn::FIntermediateBuffers::CheckIsNull() const
 
 #if RHI_RAYTRACING
 
-void FStaticMeshStreamIn::FIntermediateRayTracingGeometry::CreateFromCPUData(FRHICommandList& RHICmdList, FRayTracingGeometryInitializer InInitializer, TResourceArray<uint8>& OfflineData)
+void FStaticMeshStreamIn::FIntermediateRayTracingGeometry::CreateFromCPUData(FRHICommandList& RHICmdList, FRayTracingGeometry& RayTracingGeometry)
 {
-	Initializer = MoveTemp(InInitializer);
+	Initializer = RayTracingGeometry.Initializer;
+	Initializer.Type = ERayTracingGeometryInitializerType::StreamingSource;
 
-	if (OfflineData.Num())
+	if (RayTracingGeometry.RawData.Num())
 	{
 		check(Initializer.OfflineData == nullptr);
-		Initializer.OfflineData = &OfflineData;
+		Initializer.OfflineData = &RayTracingGeometry.RawData;
 	}
 
 	static const auto CVarDebugForceRuntimeBLAS = IConsoleManager::Get().FindTConsoleVariableDataInt(TEXT("r.Raytracing.DebugForceRuntimeBLAS"));
@@ -173,15 +174,70 @@ void FStaticMeshStreamIn::FIntermediateRayTracingGeometry::SafeRelease()
 
 void FStaticMeshStreamIn::FIntermediateRayTracingGeometry::TransferRayTracingGeometry(FRayTracingGeometry& RayTracingGeometry, FRHIResourceUpdateBatcher& Batcher)
 {
-	check(RayTracingGeometryRHI != nullptr);
-
-	// Should also set initializer?? ie: RayTracingGeometry.SetInitializer(Initializer);
 	RayTracingGeometry.InitRHIForStreaming(RayTracingGeometryRHI, Batcher);
 	RayTracingGeometry.SetRequiresBuild(bRequiresBuild);
 
 	SafeRelease();
 }
 
+#if DO_CHECK
+static void CheckRayTracingGeometryInitializer(
+	const UStaticMesh* Mesh,
+	int32 LODIdx,
+	const FStaticMeshLODResources& LODResource,
+	ERayTracingGeometryInitializerType ExpectedInitializerType,
+	const FRayTracingGeometryInitializer& Initializer)
+{
+	const FName OwnerName = UStaticMesh::GetLODPathName(Mesh, LODIdx);
+
+	FRayTracingGeometryInitializer TmpInitializer;
+	if (Mesh->HasValidNaniteData() && Nanite::GetSupportsRayTracingProceduralPrimitive(GMaxRHIShaderPlatform))
+	{
+		FStaticMeshLODResources::SetupRayTracingProceduralGeometryInitializer(TmpInitializer, Mesh->GetFName(), OwnerName);
+	}
+	else
+	{
+		LODResource.SetupRayTracingGeometryInitializer(TmpInitializer, Mesh->GetFName(), OwnerName);
+	}
+
+	TmpInitializer.Type = ExpectedInitializerType;
+
+	// Can't compare TmpInitializer == Initializer directly due to some members not having equality operators
+
+	check(TmpInitializer.IndexBuffer == Initializer.IndexBuffer);
+	check(TmpInitializer.IndexBufferOffset == Initializer.IndexBufferOffset);
+	check(TmpInitializer.GeometryType == Initializer.GeometryType);
+	check(TmpInitializer.TotalPrimitiveCount == Initializer.TotalPrimitiveCount);
+
+	// Can't compare Segments directly due to some members not having equality operators
+	check(TmpInitializer.Segments.Num() == Initializer.Segments.Num());
+
+	for (int32 SegmentIndex = 0; SegmentIndex < TmpInitializer.Segments.Num(); ++SegmentIndex)
+	{
+		//check(TmpInitializer.Segments[SegmentIndex] == Initializer.Segments[SegmentIndex]);
+		check(TmpInitializer.Segments[SegmentIndex].VertexBuffer == Initializer.Segments[SegmentIndex].VertexBuffer);
+		check(TmpInitializer.Segments[SegmentIndex].VertexBufferElementType == Initializer.Segments[SegmentIndex].VertexBufferElementType);
+		check(TmpInitializer.Segments[SegmentIndex].VertexBufferOffset == Initializer.Segments[SegmentIndex].VertexBufferOffset);
+		check(TmpInitializer.Segments[SegmentIndex].VertexBufferStride == Initializer.Segments[SegmentIndex].VertexBufferStride);
+		check(TmpInitializer.Segments[SegmentIndex].MaxVertices == Initializer.Segments[SegmentIndex].MaxVertices);
+		check(TmpInitializer.Segments[SegmentIndex].FirstPrimitive == Initializer.Segments[SegmentIndex].FirstPrimitive);
+		check(TmpInitializer.Segments[SegmentIndex].NumPrimitives == Initializer.Segments[SegmentIndex].NumPrimitives);
+		check(TmpInitializer.Segments[SegmentIndex].bForceOpaque == Initializer.Segments[SegmentIndex].bForceOpaque);
+		check(TmpInitializer.Segments[SegmentIndex].bAllowDuplicateAnyHitShaderInvocation == Initializer.Segments[SegmentIndex].bAllowDuplicateAnyHitShaderInvocation);
+		check(TmpInitializer.Segments[SegmentIndex].bEnabled == Initializer.Segments[SegmentIndex].bEnabled);
+	}
+
+	check(TmpInitializer.OfflineData == Initializer.OfflineData);
+	check(TmpInitializer.SourceGeometry == Initializer.SourceGeometry);
+	check(TmpInitializer.bFastBuild == Initializer.bFastBuild);
+	check(TmpInitializer.bAllowUpdate == Initializer.bAllowUpdate);
+	check(TmpInitializer.bAllowCompaction == Initializer.bAllowCompaction);
+	check(TmpInitializer.Type == Initializer.Type);
+	// Can't compare DebugName directly due to FDebugName not having equality operator
+	check(TmpInitializer.OwnerName == Initializer.OwnerName);
+}
+#endif
+
 #endif
 
 FStaticMeshStreamIn::FStaticMeshStreamIn(const UStaticMesh* InMesh)
@@ -223,23 +279,19 @@ void FStaticMeshStreamIn::CreateBuffers_Internal(const FContext& Context)
 			if (IsRayTracingEnabled() && Context.Mesh->bSupportRayTracing &&
 				LODResource.VertexBuffers.StaticMeshVertexBuffer.GetNumVertices() > 0)
 			{
-				const FName OwnerName = UStaticMesh::GetLODPathName(Context.Mesh, LODIdx + Context.Mesh->GetStreamableResourceState().AssetLODBias);
-
-				FRayTracingGeometryInitializer Initializer;
-				if (Context.Mesh->HasValidNaniteData() && Nanite::GetSupportsRayTracingProceduralPrimitive(GMaxRHIShaderPlatform))
-				{
-					FStaticMeshLODResources::SetupRayTracingProceduralGeometryInitializer(Initializer, Context.Mesh->GetFName(), OwnerName);
-				}
-				else
-				{
-					Context.LODResourcesView[LODIdx]->SetupRayTracingGeometryInitializer(Initializer, Context.Mesh->GetFName(), OwnerName);
-				}
-				Initializer.Type = ERayTracingGeometryInitializerType::StreamingSource;
+#if DO_CHECK
+				CheckRayTracingGeometryInitializer(
+					Context.Mesh,
+					LODIdx + Context.Mesh->GetStreamableResourceState().AssetLODBias,
+					LODResource,
+					ERayTracingGeometryInitializerType::StreamingDestination,
+					LODResource.RayTracingGeometry.Initializer);
+#endif
 
 				FRHIAsyncCommandList AsyncCommandList;
 				FRHICommandList& RHICmdList = bRenderThread ? FRHICommandListImmediate::Get() : *AsyncCommandList;
 
-				IntermediateRayTracingGeometry[LODIdx].CreateFromCPUData(RHICmdList, MoveTemp(Initializer), LODResource.RayTracingGeometry.RawData);
+				IntermediateRayTracingGeometry[LODIdx].CreateFromCPUData(RHICmdList, LODResource.RayTracingGeometry);
 			}
 #endif
 		}
@@ -288,9 +340,7 @@ void FStaticMeshStreamIn::DoFinishUpdate(const FContext& Context)
 				IntermediateBuffersArray[LODIdx].TransferBuffers(LODResource, Batcher);
 
 #if RHI_RAYTRACING
-				if (IsRayTracingEnabled()
-					&& Context.Mesh->bSupportRayTracing
-					&& LODResource.VertexBuffers.StaticMeshVertexBuffer.GetNumVertices() > 0)
+				if (IsRayTracingAllowed() && Context.Mesh->bSupportRayTracing && LODResource.VertexBuffers.StaticMeshVertexBuffer.GetNumVertices() > 0)
 				{
 					IntermediateRayTracingGeometry[LODIdx].TransferRayTracingGeometry(LODResource.RayTracingGeometry, Batcher);
 				}
@@ -309,23 +359,21 @@ void FStaticMeshStreamIn::DoFinishUpdate(const FContext& Context)
 				// Skip LODs that have their render data stripped
 				if (LODResource.VertexBuffers.StaticMeshVertexBuffer.GetNumVertices() > 0)
 				{
-					const FName OwnerName = UStaticMesh::GetLODPathName(Context.Mesh, LODIndex + Context.Mesh->GetStreamableResourceState().AssetLODBias);
-
-					// Rebuild the initializer because it could have been reset during a previous release
-					FRayTracingGeometryInitializer Initializer;
-					if (Context.Mesh->HasValidNaniteData() && Nanite::GetSupportsRayTracingProceduralPrimitive(GMaxRHIShaderPlatform))
-					{
-						FStaticMeshLODResources::SetupRayTracingProceduralGeometryInitializer(Initializer, Context.Mesh->GetFName(), OwnerName);
-					}
-					else
-					{
-						LODResource.SetupRayTracingGeometryInitializer(Initializer, Context.Mesh->GetFName(), OwnerName);
-					}
-					LODResource.RayTracingGeometry.SetInitializer(Initializer);
-					LODResource.RayTracingGeometry.SetAsStreamedIn();			
+#if DO_CHECK
+					// Streaming LODs in/out shouldn't affect the ray tracing geometry initializer
+					// Here we check that assumption
+					CheckRayTracingGeometryInitializer(
+						Context.Mesh,
+						LODIndex + Context.Mesh->GetStreamableResourceState().AssetLODBias,
+						LODResource,
+						ERayTracingGeometryInitializerType::Rendering,
+						LODResource.RayTracingGeometry.Initializer);
+
+					check(EnumHasAllFlags(LODResource.RayTracingGeometry.GetGeometryState(), FRayTracingGeometry::EGeometryStateFlags::StreamedIn));
+#endif
 
 					// Under very rare circumstances that we switch ray tracing on/off right in the middle of streaming RayTracingGeometryRHI might not be valid.
-					if (IsRayTracingEnabled() && LODResource.RayTracingGeometry.RayTracingGeometryRHI.IsValid())
+					if (IsRayTracingEnabled() && ensure(LODResource.RayTracingGeometry.RayTracingGeometryRHI.IsValid()))
 					{
 						LODResource.RayTracingGeometry.RequestBuildIfNeeded(FRHICommandListImmediate::Get(), ERTAccelerationStructureBuildPriority::Normal);
 					}
diff --git a/Engine/Source/Runtime/Engine/Private/Streaming/StaticMeshUpdate.h b/Engine/Source/Runtime/Engine/Private/Streaming/StaticMeshUpdate.h
--- a/Engine/Source/Runtime/Engine/Private/Streaming/StaticMeshUpdate.h
+++ b/Engine/Source/Runtime/Engine/Private/Streaming/StaticMeshUpdate.h
@@ -106,7 +106,7 @@ protected:
 		bool bRequiresBuild = false;
 
 	public:
-		void CreateFromCPUData(FRHICommandList& RHICmdList, FRayTracingGeometryInitializer Initializer, TResourceArray<uint8>& OfflineData);
+		void CreateFromCPUData(FRHICommandList& RHICmdList, FRayTracingGeometry& RayTracingGeometry);
 
 		void SafeRelease();
 
diff --git a/Engine/Source/Runtime/RenderCore/Private/RayTracingGeometry.cpp b/Engine/Source/Runtime/RenderCore/Private/RayTracingGeometry.cpp
--- a/Engine/Source/Runtime/RenderCore/Private/RayTracingGeometry.cpp
+++ b/Engine/Source/Runtime/RenderCore/Private/RayTracingGeometry.cpp
@@ -23,6 +23,13 @@ FRayTracingGeometry::~FRayTracingGeometry() = default;
 
 void FRayTracingGeometry::InitRHIForStreaming(FRHIRayTracingGeometry* IntermediateGeometry, FRHIResourceUpdateBatcher& Batcher)
 {
+	ensureMsgf(RayTracingGeometryRHI || !IsRayTracingEnabled(),
+		TEXT("RayTracingGeometryRHI should be valid when ray tracing is enabled.\n")
+		TEXT("This check failing points to a race condition between FRayTracingGeometryManager::Tick(...) and FStaticMeshStreamIn processing.\n")
+	);
+
+	Initializer.Type = ERayTracingGeometryInitializerType::Rendering;
+
 	EnumAddFlags(GeometryState, EGeometryStateFlags::StreamedIn);
 
 	if (RayTracingGeometryRHI && IntermediateGeometry)
@@ -30,17 +37,21 @@ void FRayTracingGeometry::InitRHIForStreaming(FRHIRayTracingGeometry* Intermedia
 		Batcher.QueueUpdateRequest(RayTracingGeometryRHI, IntermediateGeometry);
 		EnumAddFlags(GeometryState, EGeometryStateFlags::Valid);
 	}
+	else
+	{
+		check(GetRayTracingMode() == ERayTracingMode::Dynamic);
+	}
 }
 
 void FRayTracingGeometry::ReleaseRHIForStreaming(FRHIResourceUpdateBatcher& Batcher)
 {
 	RemoveBuildRequest();
 
-	Initializer = {};
-
 	EnumRemoveFlags(GeometryState, EGeometryStateFlags::StreamedIn);
 	EnumRemoveFlags(GeometryState, EGeometryStateFlags::Valid);
 
+	Initializer.Type = ERayTracingGeometryInitializerType::StreamingDestination;
+
 	if (RayTracingGeometryRHI)
 	{
 		Batcher.QueueUpdateRequest(RayTracingGeometryRHI, nullptr);
