Cameras: downgrade camera shake asserts and add edit condition to BP parameter struct

#noqatesting

diff --git a/Engine/Source/Runtime/Engine/Classes/Camera/CameraShakeBase.h b/Engine/Source/Runtime/Engine/Classes/Camera/CameraShakeBase.h
--- a/Engine/Source/Runtime/Engine/Classes/Camera/CameraShakeBase.h
+++ b/Engine/Source/Runtime/Engine/Classes/Camera/CameraShakeBase.h
@@ -32,7 +32,7 @@ struct FCameraShakePatternStartParams
 	bool bOverrideDuration = false;
 
 	/** An optional override for the camera shake's duration */
-	UPROPERTY(EditAnywhere, BlueprintReadWrite, Category=CameraShake)
+	UPROPERTY(EditAnywhere, BlueprintReadWrite, Category=CameraShake, meta=(EditCondition="bOverrideDuration"))
 	float DurationOverride = 0.f;
 };
 
diff --git a/Engine/Source/Runtime/Engine/Private/Camera/CameraShakeBase.cpp b/Engine/Source/Runtime/Engine/Private/Camera/CameraShakeBase.cpp
--- a/Engine/Source/Runtime/Engine/Private/Camera/CameraShakeBase.cpp
+++ b/Engine/Source/Runtime/Engine/Private/Camera/CameraShakeBase.cpp
@@ -285,7 +285,10 @@ void UCameraShakeBase::StartShake(const FCameraShakeBaseStartParams& Params)
 
 	// Check that we were correctly stopped before we are asked to play again.
 	// Note that single-instance shakes can be restarted while they're running.
-	checkf(!bIsActive || bSingleInstance, TEXT("Starting to play a shake that was already playing."));
+	if (!ensureMsgf(!bIsActive || bSingleInstance, TEXT("Starting to play a shake that was already playing.")))
+	{
+		return;
+	}
 
 	// Remember the various settings for this run.
 	// Note that the camera manager can be null, for example in unit tests.
@@ -314,7 +317,10 @@ void UCameraShakeBase::UpdateAndApplyCameraShake(float DeltaTime, float Alpha, F
 	SCOPE_CYCLE_COUNTER(STAT_UpdateShake);
 	SCOPE_CYCLE_UOBJECT(This, this);
 
-	checkf(bIsActive, TEXT("Updating a camera shake that wasn't started with a call to StartShake!"));
+	if (!ensureMsgf(bIsActive, TEXT("Updating a camera shake that wasn't started with a call to StartShake!")))
+	{
+		return;
+	}
 
 	if (RootShakePattern)
 	{
@@ -351,7 +357,10 @@ void UCameraShakeBase::ScrubAndApplyCameraShake(float AbsoluteTime, float Alpha,
 	SCOPE_CYCLE_COUNTER(STAT_UpdateShake);
 	SCOPE_CYCLE_UOBJECT(This, this);
 
-	checkf(bIsActive, TEXT("Updating a camera shake that wasn't started with a call to StartShake!"));
+	if (!ensureMsgf(bIsActive, TEXT("Updating a camera shake that wasn't started with a call to StartShake!")))
+	{
+		return;
+	}
 
 	if (RootShakePattern)
 	{
