Fixed check in the instance tracking tripped when editing blueprint template objects by disabling tracking for such objects.

#rnx

diff --git a/Engine/Source/Runtime/Engine/Private/InstancedStaticMesh/ISMInstanceDataManager.cpp b/Engine/Source/Runtime/Engine/Private/InstancedStaticMesh/ISMInstanceDataManager.cpp
--- a/Engine/Source/Runtime/Engine/Private/InstancedStaticMesh/ISMInstanceDataManager.cpp
+++ b/Engine/Source/Runtime/Engine/Private/InstancedStaticMesh/ISMInstanceDataManager.cpp
@@ -25,6 +25,11 @@ TAutoConsoleVariable<float> CVarInstanceUpdateTaskDebugDelay(
 FPrimitiveInstanceDataManager::FPrimitiveInstanceDataManager(UPrimitiveComponent* InPrimitiveComponent) 
 	: PrimitiveComponent(InPrimitiveComponent) 
 {
+	// Don't do anything if this is not a "real" ISM being tracked (this logic shopuld move out).
+	if (PrimitiveComponent.IsValid() && PrimitiveComponent->HasAnyFlags(RF_ClassDefaultObject | RF_ArchetypeObject))
+	{
+		TrackingState = ETrackingState::Disabled;
+	}
 }
 
 void FPrimitiveInstanceDataManager::SetMode(EMode InMode)
@@ -40,6 +45,11 @@ void FPrimitiveInstanceDataManager::SetMode(EMode InMode)
 
 void FPrimitiveInstanceDataManager::Add(int32 InInstanceAddAtIndex, bool bInsert)
 {
+	if (GetState() == ETrackingState::Disabled)
+	{
+		return;
+	}
+
 	ValidateMapping();
 
 	// If the manager is marked for legacy-only mode, we should not see any tracking calls!
@@ -111,6 +121,11 @@ void FPrimitiveInstanceDataManager::Add(int32 InInstanceAddAtIndex, bool bInsert
 
 void FPrimitiveInstanceDataManager::RemoveAtSwap(int32 InstanceIndex)
 {
+	if (GetState() == ETrackingState::Disabled)
+	{
+		return;
+	}
+
 	ValidateMapping();
 
 	check(Mode != EMode::ExternalLegacyData);
@@ -155,6 +170,11 @@ void FPrimitiveInstanceDataManager::RemoveAtSwap(int32 InstanceIndex)
 	
 void FPrimitiveInstanceDataManager::RemoveAt(int32 InstanceIndex)
 {
+	if (GetState() == ETrackingState::Disabled)
+	{
+		return;
+	}
+
 	ValidateMapping();
 
 	check(Mode != EMode::ExternalLegacyData);
@@ -216,6 +236,11 @@ void FPrimitiveInstanceDataManager::TransformChanged(FPrimitiveInstanceId Instan
 }
 void FPrimitiveInstanceDataManager::TransformsChangedAll()
 {
+	if (GetState() == ETrackingState::Disabled)
+	{
+		return;
+	}
+
 	LOG_INST_DATA(TEXT("TransformsChangedAll(%s)"), TEXT(""));
 	bTransformChangedAllInstances = true;
 	MarkComponentRenderInstancesDirty();
@@ -235,6 +260,11 @@ void FPrimitiveInstanceDataManager::BakedLightingDataChanged(int32 InstanceIndex
 
 void FPrimitiveInstanceDataManager::BakedLightingDataChangedAll()
 {
+	if (GetState() == ETrackingState::Disabled)
+	{
+		return;
+	}
+
 	LOG_INST_DATA(TEXT("BakedLightingDataChangedAll(%s)"), TEXT(""));
 	bBakedLightingDataChanged = true;
 	MarkComponentRenderInstancesDirty();
@@ -242,6 +272,11 @@ void FPrimitiveInstanceDataManager::BakedLightingDataChangedAll()
 
 void FPrimitiveInstanceDataManager::NumCustomDataChanged()
 {
+	if (GetState() == ETrackingState::Disabled)
+	{
+		return;
+	}
+
 	LOG_INST_DATA(TEXT("NumCustomDataChanged(%s)"), TEXT(""));
 	bNumCustomDataChanged = true;
 	MarkComponentRenderInstancesDirty();
@@ -251,6 +286,11 @@ void FPrimitiveInstanceDataManager::NumCustomDataChanged()
 
 void FPrimitiveInstanceDataManager::EditorDataChangedAll()
 {
+	if (GetState() == ETrackingState::Disabled)
+	{
+		return;
+	}
+
 	LOG_INST_DATA(TEXT("EditorDataChangedAll(%s)"), TEXT(""));
 	bAnyEditorDataChanged = true;
 	MarkComponentRenderInstancesDirty();
@@ -260,6 +300,11 @@ void FPrimitiveInstanceDataManager::EditorDataChangedAll()
 
 void FPrimitiveInstanceDataManager::PrimitiveTransformChanged()
 {
+	if (GetState() == ETrackingState::Disabled)
+	{
+		return;
+	}
+
 	LOG_INST_DATA(TEXT("PrimitiveTransformChanged(%s)"), TEXT(""));
 	bPrimitiveTransformChanged = true;
 	MarkComponentRenderInstancesDirty();
@@ -575,6 +620,11 @@ void FPrimitiveInstanceDataManager::InitChangeSet(const FChangeDesc &ChangeDesc,
 
 bool FPrimitiveInstanceDataManager::FlushChanges(FInstanceUpdateComponentDesc &&ComponentData, bool bNewPrimitiveProxy)
 {
+	if (GetState() == ETrackingState::Disabled)
+	{
+		return false;
+	}
+
 	// Always clear the flag such that any subsequent change marks it as needing update again.
 	bComponentMarkedDirty = false;
 
@@ -825,6 +875,11 @@ bool FPrimitiveInstanceDataManager::FlushChanges(FInstanceUpdateComponentDesc &&
 
 void FPrimitiveInstanceDataManager::PostLoad(int32 InNumInstances, TUniquePtr<FStaticMeshInstanceData> &&InStaticMeshInstanceData)
 {
+	if (GetState() == ETrackingState::Disabled)
+	{
+		return;
+	}
+
 	if (LegacyStaticMeshInstanceData.IsValid())
 	{
 		check(NumInstances == LegacyStaticMeshInstanceData->GetNumInstances());
@@ -855,6 +910,11 @@ void FPrimitiveInstanceDataManager::PostLoad(int32 InNumInstances, TUniquePtr<FS
 
 void FPrimitiveInstanceDataManager::ClearIdTracking(int32 InNumInstances)
 {
+	if (GetState() == ETrackingState::Disabled)
+	{
+		return;
+	}
+
 	// Reset the mapping to identity & clear allocator, this looses all association with ID:Index that existed before
 	IndexToIdMap.Empty();
 	IdToIndexMap.Empty();
@@ -868,6 +928,11 @@ void FPrimitiveInstanceDataManager::ClearIdTracking(int32 InNumInstances)
 
 void FPrimitiveInstanceDataManager::ClearChangeTracking()
 {
+	if (GetState() == ETrackingState::Disabled)
+	{
+		return;
+	}
+
 	// When tracking data is cleared, we loose connection to previously tracked state until the next update is sent.
 	TrackingState = ETrackingState::Initial;
 
@@ -942,6 +1007,11 @@ void FPrimitiveInstanceDataManager::CreateExplicitIdentityMapping()
 
 void FPrimitiveInstanceDataManager::MarkChangeHelper(FChangeBitArray& TrackingArray, int32 InstanceIndex)
 {
+	if (GetState() == ETrackingState::Disabled)
+	{
+		return;
+	}
+
 	check(Mode != EMode::ExternalLegacyData);
 
 	if (GetState() != ETrackingState::Tracked)
@@ -1092,6 +1162,11 @@ void FPrimitiveInstanceDataManager::ValidateMapping() const
 
 void FPrimitiveInstanceDataManager::MarkForRebuildFromLegacy(TUniquePtr<FStaticMeshInstanceData> &&InLegacyInstanceData, const TArray<int32> &InstanceReorderTable, const TArray<TRefCountPtr<HHitProxy>> &HitProxies)
 {
+	if (GetState() == ETrackingState::Disabled)
+	{
+		return;
+	}
+
 	check(InLegacyInstanceData);
 	if (Mode == EMode::ExternalLegacyData)
 	{
diff --git a/Engine/Source/Runtime/Engine/Public/InstancedStaticMesh/ISMInstanceDataManager.h b/Engine/Source/Runtime/Engine/Public/InstancedStaticMesh/ISMInstanceDataManager.h
--- a/Engine/Source/Runtime/Engine/Public/InstancedStaticMesh/ISMInstanceDataManager.h
+++ b/Engine/Source/Runtime/Engine/Public/InstancedStaticMesh/ISMInstanceDataManager.h
@@ -53,6 +53,7 @@ public:
 	{
 		Initial, // In the initial state, there is no proxy and therefore changes do not need to be tracked, e.g., during initial setup of an ISM component.
 		Tracked,
+		Disabled,
 	};
 
 	enum class EMode
