Temp fix crash when validating a content bundle container.

diff --git a/Engine/Source/Runtime/Engine/Private/WorldPartition/DataLayer/DataLayerUtils.cpp b/Engine/Source/Runtime/Engine/Private/WorldPartition/DataLayer/DataLayerUtils.cpp
--- a/Engine/Source/Runtime/Engine/Private/WorldPartition/DataLayer/DataLayerUtils.cpp
+++ b/Engine/Source/Runtime/Engine/Private/WorldPartition/DataLayer/DataLayerUtils.cpp
@@ -81,10 +81,18 @@ TArray<FName> FDataLayerUtils::ResolvedDataLayerInstanceNames(const UDataLayerMa
 
 // For performance reasons, this function assumes that InActorDesc's DataLayerInstanceNames was already resolved.
 bool FDataLayerUtils::ResolveRuntimeDataLayerInstanceNames(const UDataLayerManager* InDataLayerManager, const FWorldPartitionActorDescView& InActorDescView, const FActorDescViewMap& ActorDescViewMap, TArray<FName>& OutRuntimeDataLayerInstanceNames)
-{	
+{
+	const TArray<FName>& ActorDescViewDataLayerInstanceNames = InActorDescView.GetDataLayerInstanceNames();
+
+	// @todo_ow : temp fix
+	if (!ActorDescViewDataLayerInstanceNames.Num())
+	{
+		return true;
+	}
+	
 	if (InDataLayerManager && InDataLayerManager->CanResolveDataLayers())
 	{
-		for (FName DataLayerInstanceName : InActorDescView.GetDataLayerInstanceNames())
+		for (FName DataLayerInstanceName : ActorDescViewDataLayerInstanceNames)
 		{
 			const UDataLayerInstance* DataLayerInstance = InDataLayerManager->GetDataLayerInstanceFromName(DataLayerInstanceName);
 			if (DataLayerInstance && DataLayerInstance->IsRuntime())
@@ -95,27 +103,24 @@ bool FDataLayerUtils::ResolveRuntimeDataLayerInstanceNames(const UDataLayerManag
 
 		return true;
 	}
-	else
-	{
-		TArray<const FWorldDataLayersActorDesc*> WorldDataLayersActorDescs = FindWorldDataLayerActorDescs(ActorDescViewMap);
 
-		// Fallback on FWorldDataLayersActorDesc
-		if (WorldDataLayersActorDescs.Num())
+	// Fallback on FWorldDataLayersActorDesc
+	TArray<const FWorldDataLayersActorDesc*> WorldDataLayersActorDescs = FindWorldDataLayerActorDescs(ActorDescViewMap);
+	if (WorldDataLayersActorDescs.Num())
+	{
+		check(AreWorldDataLayersActorDescsSane(WorldDataLayersActorDescs));
+		for (FName DataLayerInstanceName : ActorDescViewDataLayerInstanceNames)
 		{
-			check(AreWorldDataLayersActorDescsSane(WorldDataLayersActorDescs));
-			for (FName DataLayerInstanceName : InActorDescView.GetDataLayerInstanceNames())
+			if (const FDataLayerInstanceDesc* DataLayerInstanceDesc = GetDataLayerInstanceDescFromInstanceName(WorldDataLayersActorDescs, DataLayerInstanceName))
 			{
-				if (const FDataLayerInstanceDesc* DataLayerInstanceDesc = GetDataLayerInstanceDescFromInstanceName(WorldDataLayersActorDescs, DataLayerInstanceName))
+				if (DataLayerInstanceDesc->GetDataLayerType() == EDataLayerType::Runtime)
 				{
-					if (DataLayerInstanceDesc->GetDataLayerType() == EDataLayerType::Runtime)
-					{
-						OutRuntimeDataLayerInstanceNames.Add(DataLayerInstanceName);
-					}
+					OutRuntimeDataLayerInstanceNames.Add(DataLayerInstanceName);
 				}
 			}
-
-			return true;
 		}
+
+		return true;
 	}
 
 	return false;
