Ensure any in-flight animation sequence compression tasks are finished bfore rebuilding the reference skeleton

diff --git a/Engine/Source/Runtime/Engine/Private/Animation/AnimationSequenceCompiler.cpp b/Engine/Source/Runtime/Engine/Private/Animation/AnimationSequenceCompiler.cpp
--- a/Engine/Source/Runtime/Engine/Private/Animation/AnimationSequenceCompiler.cpp
+++ b/Engine/Source/Runtime/Engine/Private/Animation/AnimationSequenceCompiler.cpp
@@ -199,7 +199,7 @@ namespace UE::Anim
 		
 	}
 
-	void FAnimSequenceCompilingManager::FinishCompilation(TArrayView<USkeleton* const> InSkeletons)
+	void FAnimSequenceCompilingManager::FinishCompilation(TArrayView<const USkeleton* const> InSkeletons)
 	{
 		TArray<UAnimSequence*> PendingAnimationSequences;	
 		for (const TWeakObjectPtr<UAnimSequence>& WeakAnimSequence : RegisteredAnimSequences)
diff --git a/Engine/Source/Runtime/Engine/Private/Animation/AnimationSequenceCompiler.h b/Engine/Source/Runtime/Engine/Private/Animation/AnimationSequenceCompiler.h
--- a/Engine/Source/Runtime/Engine/Private/Animation/AnimationSequenceCompiler.h
+++ b/Engine/Source/Runtime/Engine/Private/Animation/AnimationSequenceCompiler.h
@@ -33,7 +33,7 @@ namespace UE::Anim
 
 		void AddAnimSequences(TArrayView<UAnimSequence* const> InAnimSequences);
 		void FinishCompilation(TArrayView<UAnimSequence* const> InAnimSequences);
-		void FinishCompilation(TArrayView<USkeleton* const> InSkeletons);
+		void FinishCompilation(TArrayView<const USkeleton* const> InSkeletons);
 
 	protected:
 		virtual void ProcessAsyncTasks(bool bLimitExecutionTime = false) override;
diff --git a/Engine/Source/Runtime/Engine/Private/Animation/ReferenceSkeleton.cpp b/Engine/Source/Runtime/Engine/Private/Animation/ReferenceSkeleton.cpp
--- a/Engine/Source/Runtime/Engine/Private/Animation/ReferenceSkeleton.cpp
+++ b/Engine/Source/Runtime/Engine/Private/Animation/ReferenceSkeleton.cpp
@@ -4,6 +4,7 @@
 #include "Animation/Skeleton.h"
 #include "EngineLogs.h"
 #include "Engine/SkeletalMesh.h"
+#include "AnimationSequenceCompiler.h"
 
 FReferenceSkeletonModifier::FReferenceSkeletonModifier(USkeleton* InSkeleton)
 	: RefSkeleton(InSkeleton->ReferenceSkeleton),
@@ -454,6 +455,10 @@ int32 FReferenceSkeleton::GetRawSourceBoneIndex(const USkeleton* Skeleton, const
 
 void FReferenceSkeleton::RebuildRefSkeleton(const USkeleton* Skeleton, bool bRebuildNameMap)
 {
+#if WITH_EDITOR
+	UE::Anim::FAnimSequenceCompilingManager::Get().FinishCompilation({Skeleton});
+#endif // WITH_EDITOR
+	
 	if (bRebuildNameMap)
 	{
 		//On loading FinalRefBone data wont exist but NameToIndexMap will and will be valid
