World Partition - HLOD: HLOD warmup
* Fixed issue with wrong bounds used when warming up actors whose components haven't been registered yet
* Run warmup until the HLOD is actually shown - previously, warmup could have run for CVarHLODWarmupNumFrames but still be required until the server visibility ack which can be received a few frames later, possibly leading to the eviction of our warmed up data from Nanite active pages set.
* Don't rely on frame numbers in UWorldPartitionHLODRuntimeSubsystem::OnBeginRenderViews(). just decrease a warmup counter instead
* Reset warmup if skipping warmup during one or more frames
* Don't reset warmup if GFrameNumber isn't incremented between calls to UWorldPartitionHLODRuntimeSubsystem::PrepareToWarmup()

diff --git a/Engine/Source/Runtime/Engine/Private/WorldPartition/HLOD/HLODRuntimeSubsystem.cpp b/Engine/Source/Runtime/Engine/Private/WorldPartition/HLOD/HLODRuntimeSubsystem.cpp
--- a/Engine/Source/Runtime/Engine/Private/WorldPartition/HLOD/HLODRuntimeSubsystem.cpp
+++ b/Engine/Source/Runtime/Engine/Private/WorldPartition/HLOD/HLODRuntimeSubsystem.cpp
@@ -544,12 +544,13 @@ static void PrepareVTRequests(TMap<UMaterialInterface*, float>& InOutVTRequests,
 	}
 }
 
-static void PrepareNaniteRequests(TSet<Nanite::FResources*>& InOutNaniteRequests, UStaticMeshComponent* InStaticMeshComponent)
+static void PrepareNaniteRequests(TMap<Nanite::FResources*, int32>& InOutNaniteRequests, UStaticMeshComponent* InStaticMeshComponent, int32 InNumFramesUntilRender)
 {
 	UStaticMesh* StaticMesh = InStaticMeshComponent->GetStaticMesh();
 	if (StaticMesh && StaticMesh->HasValidNaniteData())
 	{
-		InOutNaniteRequests.Add(StaticMesh->GetRenderData()->NaniteResourcesPtr.Get());
+		int32& NumFramesUntilRender = InOutNaniteRequests.FindOrAdd(StaticMesh->GetRenderData()->NaniteResourcesPtr.Get());
+		NumFramesUntilRender = FMath::Max(InNumFramesUntilRender, 1);
 	}
 }
 
@@ -563,23 +564,18 @@ bool UWorldPartitionHLODRuntimeSubsystem::PrepareToWarmup(const UWorldPartitionR
 	{
 		FWorldPartitionHLODWarmupState& WarmupState = HLODActorsToWarmup.FindOrAdd(InHLODActor);
 
-		// In case a previous request to unload was aborted and the cell never unloaded... assume warmup requests are expired after a given amount of frames.
-		const uint32 WarmupExpiredFrames = 30;
-		const uint32 CurrentFrameNumber = GFrameNumberRenderThread;
-
-		// Trigger warmup on the first request to unload, or if a warmup request expired
-		const bool bInitiateWarmup = WarmupState.WarmupEndFrame == INDEX_NONE || CurrentFrameNumber > (WarmupState.WarmupEndFrame + WarmupExpiredFrames);
-		const bool bWarmupCompleted = !bInitiateWarmup && CurrentFrameNumber >= WarmupState.WarmupEndFrame;
-		
-		if (bInitiateWarmup)
+		// Trigger warmup for CVarHLODWarmupNumFrames frames on the first request, or if a warmup wasn't requested in the last frame
+		const bool bResetWarmup =  WarmupState.WarmupLastRequestedFrame == INDEX_NONE ||
+								  (WarmupState.WarmupLastRequestedFrame + 1) < GFrameNumber;
+				
+		if (bResetWarmup)
 		{
-			// Warmup will be triggered in the next BeginRenderView() call, at which point the frame number will have been incremented.
-			WarmupState.WarmupStartFrame = CurrentFrameNumber + 1;
-			WarmupState.WarmupEndFrame = WarmupState.WarmupStartFrame + CVarHLODWarmupNumFrames.GetValueOnGameThread();
-			WarmupState.Location = InCell->GetCellBounds().GetCenter();
+			WarmupState.WarmupCallsUntilReady = CVarHLODWarmupNumFrames.GetValueOnGameThread();
+			WarmupState.WarmupBounds = InCell->GetContentBounds();
 		}
-
-		bHLODActorNeedsWarmUp = !bWarmupCompleted;
+		
+		bHLODActorNeedsWarmUp = WarmupState.WarmupCallsUntilReady != 0;
+		WarmupState.WarmupLastRequestedFrame = GFrameNumber;
 	}
 
 	return bHLODActorNeedsWarmUp;
@@ -733,7 +729,7 @@ static float IsInView(const FVector& BoundsOrigin, const FVector& BoundsExtent,
 	return MaxScreenSizePixels > 0;
 }
 
-static void MakeHLODRenderResourcesResident(TMap<UMaterialInterface*, float>& VTRequests, TSet<Nanite::FResources*>& NaniteRequests, const FSceneViewFamily& InViewFamily)
+static void MakeHLODRenderResourcesResident(TMap<UMaterialInterface*, float>& VTRequests, TMap<Nanite::FResources*, int32>& NaniteRequests, const FSceneViewFamily& InViewFamily)
 {
 	if (!VTRequests.IsEmpty() || !NaniteRequests.IsEmpty())
 	{
@@ -748,10 +744,12 @@ static void MakeHLODRenderResourcesResident(TMap<UMaterialInterface*, float>& VT
 				GetRendererModule().RequestVirtualTextureTiles(MaterialRenderProxy, FVector2D(VTRequest.Value, VTRequest.Value), FeatureLevel);
 			}
 
-			const uint32 NumFramesBeforeRender = CVarHLODWarmupNumFrames.GetValueOnRenderThread();
-			for (const Nanite::FResources* Resource : NaniteRequests)
+			for (const TPair<Nanite::FResources*, int32> NaniteRequest : NaniteRequests)
 			{
-				GetRendererModule().PrefetchNaniteResource(Resource, NumFramesBeforeRender);
+				const Nanite::FResources* NaniteResource = NaniteRequest.Key;
+				const int32 NumFramesUntilRender = NaniteRequest.Value;
+
+				GetRendererModule().PrefetchNaniteResource(NaniteResource, NumFramesUntilRender);
 			}
 		});
 	}
@@ -762,7 +760,7 @@ void UWorldPartitionHLODRuntimeSubsystem::OnBeginRenderViews(const FSceneViewFam
 	TRACE_CPUPROFILER_EVENT_SCOPE(UWorldPartitionHLODRuntimeSubsystem::OnBeginRenderViews)
 
 	TMap<UMaterialInterface*, float> VTRequests;
-	TSet<Nanite::FResources*> NaniteRequests;
+	TMap<Nanite::FResources*, int32> NaniteRequests;
 
 	const bool bWarmupNanite = CVarHLODWarmupNanite.GetValueOnGameThread() != 0;
 	const bool bWarmupVT = CVarHLODWarmupVT.GetValueOnGameThread() != 0;
@@ -770,26 +768,27 @@ void UWorldPartitionHLODRuntimeSubsystem::OnBeginRenderViews(const FSceneViewFam
 	for (FHLODWarmupStateMap::TIterator HLODActorWarmupStateIt(HLODActorsToWarmup); HLODActorWarmupStateIt; ++HLODActorWarmupStateIt)
 	{
 		const FObjectKey& HLODActorObjectKey = HLODActorWarmupStateIt.Key();
-		const FWorldPartitionHLODWarmupState& HLODWarmupState = HLODActorWarmupStateIt.Value();
+		FWorldPartitionHLODWarmupState& HLODWarmupState = HLODActorWarmupStateIt.Value();
 		const AWorldPartitionHLOD* HLODActor = Cast<AWorldPartitionHLOD>(HLODActorObjectKey.ResolveObjectPtr());
 		
-		if (HLODActor && InViewFamily.FrameNumber < HLODWarmupState.WarmupEndFrame)
+		if (HLODActor)
 		{
-			HLODActor->ForEachComponent<UStaticMeshComponent>(false, [&](UStaticMeshComponent* SMC)
+			// Retrieve this component's bound - we must support getting the bounds before the component is even registered.
+			FVector BoundsOrigin;
+			FVector BoundsExtent;
+			HLODWarmupState.WarmupBounds.GetCenterAndExtents(BoundsOrigin, BoundsExtent);
+
+			float ScreenSizePixels = 0;
+			if (IsInView(BoundsOrigin, BoundsExtent, InViewFamily, bWarmupVT, ScreenSizePixels))
 			{
-				// Assume ISM HLOD don't need warmup, as they are actually found in the source level
-				if (SMC->IsA<UInstancedStaticMeshComponent>() || !SMC->GetStaticMesh())
+				HLODActor->ForEachComponent<UStaticMeshComponent>(false, [&](UStaticMeshComponent* SMC)
 				{
-					return;
-				}
-
-				// Retrieve this component's bound - we must support getting the bounds before the component is even registered.
-				FVector BoundsOrigin = SMC->GetStaticMesh()->GetBounds().Origin + SMC->GetRelativeLocation();
-				FVector BoundsExtent = SMC->GetStaticMesh()->GetBounds().BoxExtent;
+					// Assume ISM HLOD don't need warmup, as they are actually found in the source level
+					if (SMC->IsA<UInstancedStaticMeshComponent>())
+					{
+						return;
+					}
 				
-				float ScreenSizePixels = 0;
-				if (IsInView(BoundsOrigin, BoundsExtent, InViewFamily, bWarmupVT, ScreenSizePixels))
-				{
 					if (bWarmupVT)
 					{
 						PrepareVTRequests(VTRequests, SMC, ScreenSizePixels);
@@ -797,21 +796,32 @@ void UWorldPartitionHLODRuntimeSubsystem::OnBeginRenderViews(const FSceneViewFam
 
 					if (bWarmupNanite)
 					{
-						// Only issue Nanite requests on the first warmup frame
-						if (HLODWarmupState.WarmupStartFrame == InViewFamily.FrameNumber)
+						if (HLODWarmupState.WarmupCallsUntilReady == CVarHLODWarmupNumFrames.GetValueOnGameThread())
 						{
-							PrepareNaniteRequests(NaniteRequests, SMC);
+							// Send a nanite request to prepare for visibility in CVarHLODWarmupNumFrames frames
+							PrepareNaniteRequests(NaniteRequests, SMC, CVarHLODWarmupNumFrames.GetValueOnGameThread());
+						}
+						else if (HLODWarmupState.WarmupCallsUntilReady == 0)
+						{
+							// We expect HLOD to be visible at any moment (likely waiting for server visibility ack)
+							PrepareNaniteRequests(NaniteRequests, SMC, 1);
 						}
 					}
+				});
 
 #if ENABLE_DRAW_DEBUG
-					if (CVarHLODWarmupDebugDraw.GetValueOnAnyThread())
-					{
-						DrawDebugBox(HLODActor->GetWorld(), BoundsOrigin, BoundsExtent, FColor::Yellow, /*bPersistentLine*/ false, /*Lifetime*/ 1.0f);
-					}
-#endif
+				if (CVarHLODWarmupDebugDraw.GetValueOnAnyThread())
+				{
+					DrawDebugBox(HLODActor->GetWorld(), BoundsOrigin, BoundsExtent, FColor::Yellow, /*bPersistentLine*/ false, /*Lifetime*/ 1.0f);
 				}
-			});
+#endif
+			}
+		}
+
+		// Progress toward warmup readiness
+		if (HLODWarmupState.WarmupCallsUntilReady != 0)
+		{
+			HLODWarmupState.WarmupCallsUntilReady--;
 		}
 
 		if (!HLODActor)
diff --git a/Engine/Source/Runtime/Engine/Public/WorldPartition/HLOD/HLODRuntimeSubsystem.h b/Engine/Source/Runtime/Engine/Public/WorldPartition/HLOD/HLODRuntimeSubsystem.h
--- a/Engine/Source/Runtime/Engine/Public/WorldPartition/HLOD/HLODRuntimeSubsystem.h
+++ b/Engine/Source/Runtime/Engine/Public/WorldPartition/HLOD/HLODRuntimeSubsystem.h
@@ -100,9 +100,9 @@ private:
 
 	struct FWorldPartitionHLODWarmupState
 	{
-		uint32 WarmupStartFrame = INDEX_NONE;
-		uint32 WarmupEndFrame = INDEX_NONE;
-		FVector Location;
+		uint32 WarmupLastRequestedFrame = INDEX_NONE;
+		uint32 WarmupCallsUntilReady = INDEX_NONE;
+		FBox WarmupBounds;
 	};
 
 	typedef TMap<FObjectKey, FWorldPartitionHLODWarmupState> FHLODWarmupStateMap;
