When deciding if a shader should be compiled for editor ODSC take into account the shaders permutation ID.

- This was causing missing shaders for nanite objects with WPO.

diff --git a/Engine/Source/Runtime/Engine/Private/Materials/MaterialShared.cpp b/Engine/Source/Runtime/Engine/Private/Materials/MaterialShared.cpp
--- a/Engine/Source/Runtime/Engine/Private/Materials/MaterialShared.cpp
+++ b/Engine/Source/Runtime/Engine/Private/Materials/MaterialShared.cpp
@@ -3547,7 +3547,7 @@ bool FMaterial::TryGetShaders(const FMaterialShaderTypes& InTypes, const FVertex
 	TArray<FShaderCommonCompileJobPtr> CompileJobs;
 	bool bMissingShader = false;
 
-	auto ShouldCacheShaderType = [&ShaderPlatform, this, &PermutationFlags](const FShaderType* ShaderType, const FVertexFactoryType* VertexFactoryType) -> bool {
+	auto ShouldCacheShaderType = [&ShaderPlatform, this, &PermutationFlags](const FShaderType* ShaderType, const FVertexFactoryType* VertexFactoryType, const int32 PermutationId) -> bool {
 		// Check to see if the FMaterial should cache these types.
 		if (!ShouldCache(ShaderPlatform, ShaderType, VertexFactoryType))
 		{
@@ -3557,14 +3557,14 @@ bool FMaterial::TryGetShaders(const FMaterialShaderTypes& InTypes, const FVertex
 		// if we are just a MaterialShaderType (not associated with a mesh)
 		if (const FMaterialShaderType* MaterialShader = ShaderType->GetMaterialShaderType())
 		{
-			return MaterialShader->ShouldCompilePermutation(ShaderPlatform, this, kUniqueShaderPermutationId, PermutationFlags);
+			return MaterialShader->ShouldCompilePermutation(ShaderPlatform, this, PermutationId, PermutationFlags);
 		}
 
 		// if we are a MeshMaterialShader
 		if (const FMeshMaterialShaderType* MeshMaterialShader = ShaderType->GetMeshMaterialShaderType())
 		{
 			const bool bVFShouldCache = FMeshMaterialShaderType::ShouldCompileVertexFactoryPermutation(ShaderPlatform, this, VertexFactoryType, ShaderType, PermutationFlags);
-			const bool bShaderShouldCache = MeshMaterialShader->ShouldCompilePermutation(ShaderPlatform, this, VertexFactoryType, kUniqueShaderPermutationId, PermutationFlags);
+			const bool bShaderShouldCache = MeshMaterialShader->ShouldCompilePermutation(ShaderPlatform, this, VertexFactoryType, PermutationId, PermutationFlags);
 			return bVFShouldCache && bShaderShouldCache;
 		}
 		
@@ -3598,18 +3598,23 @@ bool FMaterial::TryGetShaders(const FMaterialShaderTypes& InTypes, const FVertex
 			{
 				bMissingShader = true;
 
+#if WITH_EDITOR || WITH_ODSC
+				for (const FShaderType* ShaderType : InTypes.PipelineType->GetStages())
+				{
+					const int32 PermutationId = InTypes.PermutationId[ShaderType->GetFrequency()];
+
+					if (!ShouldCacheShaderType(ShaderType, InVertexFactoryType, PermutationId))
+					{
+						return false;
+					}
+				}
+#endif // WITH_EDITOR || WITH_ODSC
+
 #if WITH_ODSC
 				if (FPlatformProperties::RequiresCookedData())
 				{
 					if (GODSCManager->IsHandlingRequests())
 					{
-				        for (const FShaderType* ShaderType : InTypes.PipelineType->GetStages())
-				        {
-					        if (!ShouldCacheShaderType(ShaderType, InVertexFactoryType))
-					        {
-						        return false;
-					        }
-				        }
 						const FString MaterialName = GetFullPath();
 						const FString VFTypeName(InVertexFactoryType ? InVertexFactoryType->GetName() : TEXT(""));
 						const FString PipelineName(InTypes.PipelineType->GetName());
@@ -3690,15 +3695,18 @@ bool FMaterial::TryGetShaders(const FMaterialShaderTypes& InTypes, const FVertex
 				{
 					bMissingShader = true;
 
+#if WITH_EDITOR || WITH_ODSC
+					if (!ShouldCacheShaderType(ShaderType, InVertexFactoryType, PermutationId))
+					{
+						return false;
+					}
+#endif // WITH_EDITOR || WITH_ODSC
+
 #if WITH_ODSC
 					if (FPlatformProperties::RequiresCookedData())
 					{
 						if (GODSCManager->IsHandlingRequests())
 						{
-					        if (!ShouldCacheShaderType(ShaderType, InVertexFactoryType))
-					        {
-						        return false;
-					        }
 							const FString MaterialName = GetFullPath();
 							const FString VFTypeName(InVertexFactoryType ? InVertexFactoryType->GetName() : TEXT(""));
 							const FString PipelineName;
@@ -5373,4 +5381,4 @@ FMaterialShaderParameters::FMaterialShaderParameters(const FMaterial* InMaterial
 	bAlwaysEvaluateWorldPositionOffset = InMaterial->ShouldAlwaysEvaluateWorldPositionOffset();
 }
 
-#undef LOCTEXT_NAMESPACE
+#undef LOCTEXT_NAMESPACE
\ No newline at end of file
