Override for solver delta time on FBodyInstance

diff --git a/Engine/Source/Runtime/Engine/Classes/PhysicsEngine/BodyInstance.h b/Engine/Source/Runtime/Engine/Classes/PhysicsEngine/BodyInstance.h
--- a/Engine/Source/Runtime/Engine/Classes/PhysicsEngine/BodyInstance.h
+++ b/Engine/Source/Runtime/Engine/Classes/PhysicsEngine/BodyInstance.h
@@ -440,6 +440,22 @@ protected:
 	UPROPERTY(EditAnywhere, AdvancedDisplay, BlueprintReadWrite, Category = Physics)
 	uint8 bInertiaConditioning : 1;
 
+public:
+	/** Set the desired delta time for the body. **/
+	UPROPERTY(EditAnywhere, AdvancedDisplay, BlueprintReadOnly, Category = Physics)
+	uint8 bOverrideSolverAsyncDeltaTime : 1;
+
+	/** Override value for physics solver async delta time.  With multiple actors specifying this, the solver will use the smallest delta time **/
+	UPROPERTY(EditAnywhere, AdvancedDisplay, BlueprintReadOnly, Category = Physics, meta = (editcondition = "bOverrideSolverAsyncDeltaTime"))
+	float SolverAsyncDeltaTime;
+
+	float GetSolverAsyncDeltaTime() const { return SolverAsyncDeltaTime; }
+	bool IsSolverAsyncDeltaTimeSet() const { return bOverrideSolverAsyncDeltaTime && SolverAsyncDeltaTime > 0.0; }
+
+	void SetSolverAsyncDeltaTime(const float NewSolverAsyncDeltaTime);
+
+private:
+	void UpdateSolverAsyncDeltaTime();
 
 public:
 	/** Current scale of physics - used to know when and how physics must be rescaled to match current transform of OwnerComponent. */
diff --git a/Engine/Source/Runtime/Engine/Private/PhysicsEngine/BodyInstance.cpp b/Engine/Source/Runtime/Engine/Private/PhysicsEngine/BodyInstance.cpp
--- a/Engine/Source/Runtime/Engine/Private/PhysicsEngine/BodyInstance.cpp
+++ b/Engine/Source/Runtime/Engine/Private/PhysicsEngine/BodyInstance.cpp
@@ -84,6 +84,11 @@ FAutoConsoleVariableRef CVarbPreventInvalidBodyInstanceTransforms(
 	bPreventInvalidBodyInstanceTransforms, 
 	TEXT("If true, an attempt to create a BodyInstance with an invalid transform will fail with a warning"));
 
+bool bEnableOverrideSolverDeltaTime = true;
+FAutoConsoleVariableRef CVarbEnableOverrideSolverDeltaTime(
+	TEXT("p.EnableOverrideSolverDeltaTime"),
+	bEnableOverrideSolverDeltaTime,
+	TEXT("If true, setting for override solver delta time can be used.  False will disable this feature."));
 
 using namespace PhysicsInterfaceTypes;
 
@@ -335,7 +340,9 @@ FBodyInstance::FBodyInstance()
 	, bOverrideWalkableSlopeOnInstance(false)
 	, bInterpolateWhenSubStepping(true)
 	, bPendingCollisionProfileSetup(false)
-	, bInertiaConditioning(true)
+	, bInertiaConditioning(true)	
+	, bOverrideSolverAsyncDeltaTime(false)
+	, SolverAsyncDeltaTime(1.f / 60)
 	, Scale3D(1.0f)
 	, CollisionProfileName(UCollisionProfile::CustomCollisionProfileName)
 	, PositionSolverIterationCount(8)
@@ -1506,6 +1513,12 @@ void FInitBodiesHelperBase::InitBodies()
 				}
 			}
 
+			// set solver async delta time if any bodies are overriding the sim delta time
+			for (FBodyInstance* BI : Bodies)
+			{				
+				BI->UpdateSolverAsyncDeltaTime();			
+			}
+
 			// Set up dynamic instance data
 			if (!IsStatic())
 			{
@@ -4333,6 +4346,46 @@ void FBodyInstance::BuildBodyCollisionFlags(FBodyCollisionFlags& OutFlags, EColl
 	}
 }
 
+inline void FBodyInstance::SetSolverAsyncDeltaTime(const float NewSolverAsyncDeltaTime)
+{
+	bOverrideSolverAsyncDeltaTime = true;
+	SolverAsyncDeltaTime = NewSolverAsyncDeltaTime;
+	UpdateSolverAsyncDeltaTime();
+}
+
+void FBodyInstance::UpdateSolverAsyncDeltaTime()
+{
+	if (bOverrideSolverAsyncDeltaTime && !bEnableOverrideSolverDeltaTime)
+	{
+		UE_LOG(LogPhysics, Warning, TEXT("FBodyInstance::SolverAsyncDeltaTime : Ignoring parameter because overriden by p.EnableOverrideSolverDeltaTime"));
+		return;
+	}
+
+	if (bOverrideSolverAsyncDeltaTime)
+	{
+		if (SolverAsyncDeltaTime < 0.f)
+		{
+			UE_LOG(LogPhysics, Error, TEXT("FBodyInstance::SolverAsyncDeltaTime : Value must be greater than 0"));			
+			return;
+		}
+
+		Chaos::FPhysicsSolverBase* Solver = ActorHandle->GetSolverBase();
+
+		if (!Solver->IsUsingAsyncResults())
+		{
+			UE_LOG(LogPhysics, Warning, TEXT("FBodyInstance::SolverAsyncDeltaTime : Ignoring parameter because solver is not setup to use async results"));
+			return;
+		}
+		else
+		{
+			float OldAsyncDeltaTime = Solver->GetAsyncDeltaTime();
+
+			// set async delta time to minimum of delta time and what this body wants
+			Solver->EnableAsyncMode(FMath::Min<float>(Solver->GetAsyncDeltaTime(), SolverAsyncDeltaTime));
+		}
+	}
+}
+
 void FBodyInstance::UpdateInterpolateWhenSubStepping()
 {
 	if(UPhysicsSettings::Get()->bSubstepping)
