World Partition Changelist Validation: when validating content bundles containers, create an actor descriptors container collection and always add the corresponding world to it so we can validate data layers.
Remove temp fix crash when validating a content bundle container submitted earlier to avoid this error.

diff --git a/Engine/Plugins/Editor/DataValidation/Source/DataValidation/Private/WorldPartitionChangelistValidator.cpp b/Engine/Plugins/Editor/DataValidation/Source/DataValidation/Private/WorldPartitionChangelistValidator.cpp
--- a/Engine/Plugins/Editor/DataValidation/Source/DataValidation/Private/WorldPartitionChangelistValidator.cpp
+++ b/Engine/Plugins/Editor/DataValidation/Source/DataValidation/Private/WorldPartitionChangelistValidator.cpp
@@ -136,7 +136,7 @@ void UWorldPartitionChangelistValidator::ValidateActorsAndDataLayersFromChangeLi
 		}
 	}
 
-	auto RegisterContainerToValidate = [](UWorld* InWorld, FName InContainerPackageName, FActorDescContainerCollection& OutRegisteredContainers)
+	auto RegisterContainerToValidate = [](UWorld* InWorld, FName InContainerPackageName, FActorDescContainerCollection& OutRegisteredContainers, const FGuid& InContentBundleGuid)
 	{
 		if (OutRegisteredContainers.Contains(InContainerPackageName))
 		{
@@ -156,6 +156,11 @@ void UWorldPartitionChangelistValidator::ValidateActorsAndDataLayersFromChangeLi
 			// Find in memory failed, load the ActorDescContainer
 			ActorDescContainer = NewObject<UActorDescContainer>();
 			ActorDescContainer->Initialize({ nullptr, InContainerPackageName });
+			ActorDescContainer->SetContentBundleGuid(InContentBundleGuid);
+		}
+		else
+		{
+			check(ActorDescContainer->GetContentBundleGuid() == InContentBundleGuid);
 		}
 
 		OutRegisteredContainers.AddContainer(ActorDescContainer);
@@ -173,6 +178,10 @@ void UWorldPartitionChangelistValidator::ValidateActorsAndDataLayersFromChangeLi
 		TGuardValue<UObject*> GuardCurrentAsset(CurrentAsset, World);
 		
 		FActorDescContainerCollection ContainersToValidate;
+
+		// Always register the main world container because content bundle containers can't be validated separately
+		RegisterContainerToValidate(World, MapPath.GetPackageName(), ContainersToValidate, FGuid());
+
 		for (const FAssetData& ActorData : ActorsData)
 		{
 			FString ActorPackagePath = ActorData.PackagePath.ToString();
@@ -184,11 +193,7 @@ void UWorldPartitionChangelistValidator::ValidateActorsAndDataLayersFromChangeLi
 				FString ContentBundleContainerPackagePath;
 				verify(ContentBundlePaths::BuildActorDescContainerPackagePath(FString(ContentBundleMountPoint), ContentBundleGuid, MapPath.GetPackageName().ToString(), ContentBundleContainerPackagePath));
 
-				RegisterContainerToValidate(World, FName(*ContentBundleContainerPackagePath), ContainersToValidate);
-			}
-			else
-			{
-				RegisterContainerToValidate(World, MapPath.GetPackageName(), ContainersToValidate);
+				RegisterContainerToValidate(World, FName(*ContentBundleContainerPackagePath), ContainersToValidate, ContentBundleGuid);
 			}
 		}
 
@@ -206,9 +211,9 @@ void UWorldPartitionChangelistValidator::ValidateActorsAndDataLayersFromChangeLi
 		}
 
 		// Invoke static WorldPartition Validation from the ActorDescContainer
-		UWorldPartition::FCheckForErrorsParams Params;
-		Params.ErrorHandler = this;
-		Params.bEnableStreaming = !ULevel::GetIsStreamingDisabledFromPackage(MapPath.GetPackageName());
+		UWorldPartition::FCheckForErrorsParams Params = UWorldPartition::FCheckForErrorsParams()
+			.SetErrorHandler(this)
+			.SetEnableStreaming(!ULevel::GetIsStreamingDisabledFromPackage(MapPath.GetPackageName()));
 
 		ContainersToValidate.ForEachActorDescContainer([&Params](const UActorDescContainer* ActorDescContainer)
 		{
@@ -219,11 +224,8 @@ void UWorldPartitionChangelistValidator::ValidateActorsAndDataLayersFromChangeLi
 			}
 		});
 
-		ContainersToValidate.ForEachActorDescContainer([&Params](const UActorDescContainer* ActorDescContainer)
-		{
-			Params.ActorDescContainer = ActorDescContainer;
-			UWorldPartition::CheckForErrors(Params);
-		});
+		Params.ActorDescContainerCollection = &ContainersToValidate;
+		UWorldPartition::CheckForErrors(Params);
 	}
 }
 
diff --git a/Engine/Source/Runtime/Engine/Private/WorldPartition/DataLayer/DataLayerUtils.cpp b/Engine/Source/Runtime/Engine/Private/WorldPartition/DataLayer/DataLayerUtils.cpp
--- a/Engine/Source/Runtime/Engine/Private/WorldPartition/DataLayer/DataLayerUtils.cpp
+++ b/Engine/Source/Runtime/Engine/Private/WorldPartition/DataLayer/DataLayerUtils.cpp
@@ -84,12 +84,6 @@ bool FDataLayerUtils::ResolveRuntimeDataLayerInstanceNames(const UDataLayerManag
 {
 	const TArray<FName>& ActorDescViewDataLayerInstanceNames = InActorDescView.GetDataLayerInstanceNames();
 
-	// @todo_ow : temp fix
-	if (!ActorDescViewDataLayerInstanceNames.Num())
-	{
-		return true;
-	}
-	
 	if (InDataLayerManager && InDataLayerManager->CanResolveDataLayers())
 	{
 		for (FName DataLayerInstanceName : ActorDescViewDataLayerInstanceNames)
diff --git a/Engine/Source/Runtime/Engine/Private/WorldPartition/WorldPartitionStreamingGeneration.cpp b/Engine/Source/Runtime/Engine/Private/WorldPartition/WorldPartitionStreamingGeneration.cpp
--- a/Engine/Source/Runtime/Engine/Private/WorldPartition/WorldPartitionStreamingGeneration.cpp
+++ b/Engine/Source/Runtime/Engine/Private/WorldPartition/WorldPartitionStreamingGeneration.cpp
@@ -86,7 +86,7 @@
 
 UWorldPartition::FCheckForErrorsParams::FCheckForErrorsParams()
 	: ErrorHandler(nullptr)
-	, ActorDescContainer(nullptr)
+	, ActorDescContainerCollection(nullptr)
 	, bEnableStreaming(false)
 {}
 
@@ -1497,39 +1497,14 @@
 	});	
 }
 
-void UWorldPartition::CheckForErrors(IStreamingGenerationErrorHandler* ErrorHandler) const
-{
-	FActorDescList ModifiedActorDescList;
-
-	FWorldPartitionStreamingGenerator::FWorldPartitionStreamingGeneratorParams StreamingGeneratorParams = FWorldPartitionStreamingGenerator::FWorldPartitionStreamingGeneratorParams()
-		.SetWorldPartitionContext(this)
-		.SetModifiedActorsDescList(&ModifiedActorDescList)
-		.SetErrorHandler(StreamingGenerationErrorHandlerOverride ? (*StreamingGenerationErrorHandlerOverride)(ErrorHandler) : ErrorHandler)
-		.SetIsValidGrid([this](FName GridName) { return RuntimeHash->IsValidGrid(GridName); })
-		.SetIsValidHLODLayer([this](FName GridName, const FSoftObjectPath& HLODLayerPath) { return RuntimeHash->IsValidHLODLayer(GridName, HLODLayerPath); })
-		.SetEnableStreaming(IsStreamingEnabled());
-
-	ForEachActorDescContainer([&StreamingGeneratorParams](const UActorDescContainer* InActorDescContainer)
-	{
-		for (FActorDescList::TConstIterator<> ActorDescIt(InActorDescContainer); ActorDescIt; ++ActorDescIt)
-		{
-			check(!StreamingGeneratorParams.ActorGuidsToContainerMap.Contains(ActorDescIt->GetGuid()));
-			StreamingGeneratorParams.ActorGuidsToContainerMap.Add(ActorDescIt->GetGuid(), InActorDescContainer);
-		}
-	});
-
-	ForEachActorDescContainer([this, &StreamingGeneratorParams](const UActorDescContainer* InActorDescContainer)
-	{
-		check(StreamingGeneratorParams.WorldPartitionContext == InActorDescContainer->GetWorldPartition());
-		FWorldPartitionStreamingGenerator StreamingGenerator(StreamingGeneratorParams);
-		FStreamingGenerationActorDescCollection Collection{InActorDescContainer};
-		StreamingGenerator.PreparationPhase(Collection);
-	});
-	
-}
-
 FStreamingGenerationActorDescCollection::FStreamingGenerationActorDescCollection(std::initializer_list<TObjectPtr<const UActorDescContainer>> ActorDescContainerArray)
 	: TActorDescContainerCollection<TObjectPtr<const UActorDescContainer>>(ActorDescContainerArray)
 {
 	SortCollection();
 }
+
+FStreamingGenerationActorDescCollection::FStreamingGenerationActorDescCollection(const TArray<const UActorDescContainer*>& ActorDescContainers)
+	: TActorDescContainerCollection<TObjectPtr<const UActorDescContainer>>(ActorDescContainers)
+{
+	SortCollection();
+}
@@ -1628,24 +1603,67 @@
 #endif
 }
 
+void UWorldPartition::CheckForErrors(IStreamingGenerationErrorHandler* ErrorHandler) const
+{
+	FActorDescList ModifiedActorDescList;
+
+	FWorldPartitionStreamingGenerator::FWorldPartitionStreamingGeneratorParams StreamingGeneratorParams = FWorldPartitionStreamingGenerator::FWorldPartitionStreamingGeneratorParams()
+		.SetWorldPartitionContext(this)
+		.SetModifiedActorsDescList(&ModifiedActorDescList)
+		.SetErrorHandler(StreamingGenerationErrorHandlerOverride ? (*StreamingGenerationErrorHandlerOverride)(ErrorHandler) : ErrorHandler)
+		.SetIsValidGrid([this](FName GridName) { return RuntimeHash->IsValidGrid(GridName); })
+		.SetIsValidHLODLayer([this](FName GridName, const FSoftObjectPath& HLODLayerPath) { return RuntimeHash->IsValidHLODLayer(GridName, HLODLayerPath); })
+		.SetEnableStreaming(IsStreamingEnabled());
+
+	ForEachActorDescContainer([&StreamingGeneratorParams](const UActorDescContainer* InActorDescContainer)
+	{
+		for (FActorDescList::TConstIterator<> ActorDescIt(InActorDescContainer); ActorDescIt; ++ActorDescIt)
+		{
+			check(!StreamingGeneratorParams.ActorGuidsToContainerMap.Contains(ActorDescIt->GetGuid()));
+			StreamingGeneratorParams.ActorGuidsToContainerMap.Add(ActorDescIt->GetGuid(), InActorDescContainer);
+		}
+	});
+
+	ForEachActorDescContainer([this, &StreamingGeneratorParams](const UActorDescContainer* InActorDescContainer)
+	{
+		check(StreamingGeneratorParams.WorldPartitionContext == InActorDescContainer->GetWorldPartition());
+		FWorldPartitionStreamingGenerator StreamingGenerator(StreamingGeneratorParams);
+		FStreamingGenerationActorDescCollection Collection{InActorDescContainer};
+		StreamingGenerator.PreparationPhase(Collection);
+	});	
+}
+
 /* Deprecated */
 void UWorldPartition::CheckForErrors(IStreamingGenerationErrorHandler* ErrorHandler, const UActorDescContainer* ActorDescContainer, bool bEnableStreaming, bool)
 {
-	FCheckForErrorsParams Params;
-	Params.ErrorHandler = ErrorHandler;
-	Params.ActorDescContainer = ActorDescContainer;
-	Params.bEnableStreaming = bEnableStreaming;
+	FActorDescContainerCollection ActorDescContainerCollection;
+	ActorDescContainerCollection.AddContainer(const_cast<UActorDescContainer*>(ActorDescContainer));
 
+	FCheckForErrorsParams Params = FCheckForErrorsParams()
+		.SetErrorHandler(ErrorHandler)
+		.SetActorDescContainerCollection(&ActorDescContainerCollection)
+		.SetEnableStreaming(bEnableStreaming);
+
 	CheckForErrors(Params);
 }
 
+/* Static version, mainly used by changelist validation */
 void UWorldPartition::CheckForErrors(const FCheckForErrorsParams& Params)
 {
+	check(Params.ErrorHandler);
+	check(Params.ActorDescContainerCollection);
+
 	FActorDescList ModifiedActorDescList;
 
+	TArray<const UActorDescContainer*> Containers;
+	Params.ActorDescContainerCollection->ForEachActorDescContainer([&Containers](UActorDescContainer* InActorDescContainer) { Containers.Add(InActorDescContainer); });
+
+	FStreamingGenerationActorDescCollection Collection { Containers };
+	const UActorDescContainer* MainActorDescContainer = Collection.GetMainActorDescContainer();
+
 	FWorldPartitionStreamingGenerator::FWorldPartitionStreamingGeneratorParams StreamingGeneratorParams = FWorldPartitionStreamingGenerator::FWorldPartitionStreamingGeneratorParams()
-		.SetWorldPartitionContext(Params.ActorDescContainer->GetWorldPartition())
-		.SetModifiedActorsDescList(!Params.ActorDescContainer->IsTemplateContainer() ? &ModifiedActorDescList : nullptr)
+		.SetWorldPartitionContext(MainActorDescContainer->GetWorldPartition())
+		.SetModifiedActorsDescList(!MainActorDescContainer->IsTemplateContainer() ? &ModifiedActorDescList : nullptr)
 		.SetErrorHandler(StreamingGenerationErrorHandlerOverride ? (*StreamingGenerationErrorHandlerOverride)(Params.ErrorHandler) : Params.ErrorHandler)
 		.SetIsValidGrid([](FName) { return true; })
 		.SetIsValidHLODLayer([](FName, const FSoftObjectPath&) { return true; })
@@ -1653,11 +1671,8 @@
 		.SetActorGuidsToContainerMap(Params.ActorGuidsToContainerMap);
 
 	FWorldPartitionStreamingGenerator StreamingGenerator(StreamingGeneratorParams);
-
-	FStreamingGenerationActorDescCollection Collection{Params.ActorDescContainer};
 	StreamingGenerator.PreparationPhase(Collection);
 }
-
 #endif // WITH_EDITOR
 
 #undef LOCTEXT_NAMESPACE
diff --git a/Engine/Source/Runtime/Engine/Public/WorldPartition/ActorDescContainerCollection.h b/Engine/Source/Runtime/Engine/Public/WorldPartition/ActorDescContainerCollection.h
--- a/Engine/Source/Runtime/Engine/Public/WorldPartition/ActorDescContainerCollection.h
+++ b/Engine/Source/Runtime/Engine/Public/WorldPartition/ActorDescContainerCollection.h
@@ -20,6 +20,9 @@ public:
 	template<class U>
 	TActorDescContainerCollection(std::initializer_list<U> ActorDescContainerArray);
 
+	template<class U>
+	TActorDescContainerCollection(const TArray<U>& ActorDescContainers);
+
 	virtual ~TActorDescContainerCollection();
 
 	void AddContainer(ActorDescContPtrType Container);
@@ -201,7 +204,18 @@ public:
 template<class ActorDescContPtrType>
 template<class U>
 TActorDescContainerCollection<ActorDescContPtrType>::TActorDescContainerCollection(std::initializer_list<U> ActorDescContainerArray)
-	:ActorDescContainerCollection(ActorDescContainerArray)
+	: ActorDescContainerCollection(ActorDescContainerArray)
+{
+	ForEachActorDescContainer([this](ActorDescContPtrType ActorDescContainer)
+	{
+		RegisterDelegates(ActorDescContainer);
+	});
+}
+
+template<class ActorDescContPtrType>
+template<class U>
+TActorDescContainerCollection<ActorDescContPtrType>::TActorDescContainerCollection(const TArray<U>& ActorDescContainers)
+	: ActorDescContainerCollection(ActorDescContainers)
 {
 	ForEachActorDescContainer([this](ActorDescContPtrType ActorDescContainer)
 	{
diff --git a/Engine/Source/Runtime/Engine/Public/WorldPartition/WorldPartition.h b/Engine/Source/Runtime/Engine/Public/WorldPartition/WorldPartition.h
--- a/Engine/Source/Runtime/Engine/Public/WorldPartition/WorldPartition.h
+++ b/Engine/Source/Runtime/Engine/Public/WorldPartition/WorldPartition.h
@@ -331,9 +331,14 @@ public:
 		ENGINE_API FCheckForErrorsParams();
 
 		IStreamingGenerationErrorHandler* ErrorHandler;
-		const UActorDescContainer* ActorDescContainer;
+		const FActorDescContainerCollection* ActorDescContainerCollection;
 		bool bEnableStreaming;
 		TMap<FGuid, const UActorDescContainer*> ActorGuidsToContainerMap;
+
+		FCheckForErrorsParams& SetErrorHandler(IStreamingGenerationErrorHandler* InErrorHandler) { ErrorHandler = InErrorHandler; return *this; }
+		FCheckForErrorsParams& SetActorDescContainerCollection(const FActorDescContainerCollection* InActorDescContainerCollection) { ActorDescContainerCollection = InActorDescContainerCollection; return *this; }
+		FCheckForErrorsParams& SetEnableStreaming(bool bInEnableStreaming) { bEnableStreaming = bInEnableStreaming; return *this; }
+		FCheckForErrorsParams& SetActorGuidsToContainerMap(const TMap<FGuid, const UActorDescContainer*>& InActorGuidsToContainerMap) { ActorGuidsToContainerMap = InActorGuidsToContainerMap; return *this; }
 	};
 
 	UE_DEPRECATED(5.2, "CheckForErrors is deprecated, CheckForErrors with FCheckForErrorsParams should be used instead.")
diff --git a/Engine/Source/Runtime/Engine/Public/WorldPartition/WorldPartitionStreamingGeneration.h b/Engine/Source/Runtime/Engine/Public/WorldPartition/WorldPartitionStreamingGeneration.h
--- a/Engine/Source/Runtime/Engine/Public/WorldPartition/WorldPartitionStreamingGeneration.h
+++ b/Engine/Source/Runtime/Engine/Public/WorldPartition/WorldPartitionStreamingGeneration.h
@@ -107,6 +107,7 @@ class FStreamingGenerationActorDescCollection : public TActorDescContainerCollec
 public:
 	FStreamingGenerationActorDescCollection() = default;
 	ENGINE_API FStreamingGenerationActorDescCollection(std::initializer_list<TObjectPtr<const UActorDescContainer>> ActorDescContainerArray);
+	ENGINE_API FStreamingGenerationActorDescCollection(const TArray<const UActorDescContainer*>& ActorDescContainers);
 
 	ENGINE_API UWorld* GetWorld() const;
 
