Fix Asset Manager bugs when making changes after initial startup:
Change RegisterSpecificPrimaryAsset to add scan paths so game feature data assets are not lost while editing assets and using PIE
Fix ReinitializeFromConfig to copy over dynamic assets and rescan any paths that were added after initial load
Clean up access to config fields on the type info so they can be correctly used by game features at runtime and remove unnecessary hack that was adding a fake scan directory

diff --git a/Engine/Plugins/Experimental/GameFeatures/Source/GameFeatures/Private/GameFeaturesSubsystem.cpp b/Engine/Plugins/Experimental/GameFeatures/Source/GameFeatures/Private/GameFeaturesSubsystem.cpp
--- a/Engine/Plugins/Experimental/GameFeatures/Source/GameFeatures/Private/GameFeaturesSubsystem.cpp
+++ b/Engine/Plugins/Experimental/GameFeatures/Source/GameFeatures/Private/GameFeaturesSubsystem.cpp
@@ -539,8 +539,7 @@ void UGameFeaturesSubsystem::AddGameFeatureToAssetManager(const UGameFeatureData
 
 	for (FPrimaryAssetTypeInfo TypeInfo : GameFeatureToAdd->GetPrimaryAssetTypesToScan())
 	{
-		// @TODO: we shouldn't be accessing private data here. Need a better way to do this
-		for (FDirectoryPath& Path : TypeInfo.Directories)
+		for (FDirectoryPath& Path : TypeInfo.GetDirectories())
 		{
 			// Convert plugin-relative paths to full package paths
 			FixPluginPackagePath(Path.Path, PluginRootPath, false);
@@ -590,7 +589,7 @@ void UGameFeaturesSubsystem::RemoveGameFeatureFromAssetManager(const UGameFeatur
 			continue;
 		}
 
-		for (FDirectoryPath& Path : TypeInfo.Directories)
+		for (FDirectoryPath& Path : TypeInfo.GetDirectories())
 		{
 			FixPluginPackagePath(Path.Path, PluginRootPath, false);
 		}
diff --git a/Engine/Plugins/Experimental/GameFeatures/Source/GameFeaturesEditor/Private/GameFeaturesEditorModule.cpp b/Engine/Plugins/Experimental/GameFeatures/Source/GameFeaturesEditor/Private/GameFeaturesEditorModule.cpp
--- a/Engine/Plugins/Experimental/GameFeatures/Source/GameFeaturesEditor/Private/GameFeaturesEditorModule.cpp
+++ b/Engine/Plugins/Experimental/GameFeatures/Source/GameFeaturesEditor/Private/GameFeaturesEditorModule.cpp
@@ -341,16 +341,11 @@ class FGameFeaturesEditorModule : public FDefaultModuleImpl
 		// Add the rule to project settings
 		if (bSuccess)
 		{
-			FDirectoryPath DummyPath;
-			DummyPath.Path = TEXT("/Game/Unused");
-
 			FPrimaryAssetTypeInfo NewTypeInfo(
 				UGameFeatureData::StaticClass()->GetFName(),
 				UGameFeatureData::StaticClass(),
 				false,
-				false,
-				{ DummyPath },
-				{});
+				false);
 			NewTypeInfo.Rules.CookRule = EPrimaryAssetCookRule::AlwaysCook;
 
 			Settings->Modify(true);
diff --git a/Engine/Source/Runtime/Engine/Classes/Engine/AssetManager.h b/Engine/Source/Runtime/Engine/Classes/Engine/AssetManager.h
--- a/Engine/Source/Runtime/Engine/Classes/Engine/AssetManager.h
+++ b/Engine/Source/Runtime/Engine/Classes/Engine/AssetManager.h
@@ -687,6 +687,9 @@ protected:
 	/** Returns true if path should be excluded from primary asset scans, called from ShouldIncludeInAssetSearch and in the editor */
 	ENGINE_API virtual bool IsPathExcludedFromScan(const FString& Path) const;
 
+	/** Returns true if we're in the middle of handling the initial config, false if this is being called from something else like a plugin */
+	ENGINE_API bool IsScanningFromInitialConfig() const;
+
 	/** Filter function that is called from SearchAssetRegistryPaths, returns true if asset data should be included in search results */
 	ENGINE_API virtual bool ShouldIncludeInAssetSearch(const FAssetData& AssetData, const FAssetManagerSearchRules& SearchRules) const;
 
@@ -869,6 +872,11 @@ private:
 	/** Provide proper reentrancy for AssetRegistry temporary caching */
 	bool bOldTemporaryCachingMode = false;
 
+	/** True if we're doing an initial scan, private because this may be replaced by a different data structure */
+	bool bScanningFromInitialConfig = false;
+
+	void InternalAddAssetScanPath(FPrimaryAssetTypeData& TypeData, const FString& AssetScanPath);
+
 #if WITH_EDITOR
 	/** Recursive handler for InitializeAssetBundlesFromMetadata */
 	ENGINE_API virtual void InitializeAssetBundlesFromMetadata_Recursive(const UStruct* Struct, const void* StructValue, FAssetBundleData& AssetBundle, FName DebugName, TSet<const void*>& AllVisitedStructValues) const;
diff --git a/Engine/Source/Runtime/Engine/Classes/Engine/AssetManagerTypes.h b/Engine/Source/Runtime/Engine/Classes/Engine/AssetManagerTypes.h
--- a/Engine/Source/Runtime/Engine/Classes/Engine/AssetManagerTypes.h
+++ b/Engine/Source/Runtime/Engine/Classes/Engine/AssetManagerTypes.h
@@ -131,30 +131,25 @@ struct FPrimaryAssetRulesExplicitOverride
 	ENGINE_API void OverrideRulesExplicitly(FPrimaryAssetRules& RulesToOverride) const;
 };
 
-/** Structure with publicly exposed information about an asset type. These can be loaded out of a config file. */
+/** Structure with publicly exposed information about an asset type. These can be loaded out of a config file or constructed at runtime */
 USTRUCT()
 struct FPrimaryAssetTypeInfo
 {
 	GENERATED_BODY()
 
-	// Loaded out of ini or set via ScanPathsForPrimaryAssets
-
 	/** The logical name for this type of Primary Asset */
 	UPROPERTY(EditAnywhere, Category = AssetType)
 	FName PrimaryAssetType;
 
-#if WITH_EDITOR
-	const TSoftClassPtr<UObject>& GetAssetBaseClass() const { return AssetBaseClass; }
-	void SetAssetBaseClass(const TSoftClassPtr<UObject>& InAssetBaseClass) { AssetBaseClass = InAssetBaseClass; }
-#endif //if WITH_EDITOR
-
 private:
+	// Use accessors below to modify config data
+
 	/** Base Class of all assets of this type */
 	UPROPERTY(EditAnywhere, Category = AssetType, meta = (AllowAbstract))
 	TSoftClassPtr<UObject> AssetBaseClass;
 
 public:
-	/** Base Class of all assets of this type */
+	/** Runtime cached copy of asset base class, this will only be correct if FillRuntimeData has been called */
 	UPROPERTY(Transient)
 	TObjectPtr<UClass> AssetBaseClassLoaded;
 
@@ -171,14 +166,9 @@ public:
 	UPROPERTY(EditAnywhere, Category = AssetType)
 	bool bIsEditorOnly;
 
-#if WITH_EDITOR
-	const TArray<FDirectoryPath>& GetDirectories() const { return Directories; }
-	TArray<FDirectoryPath>& GetDirectories() { return Directories; }
-
-	const TArray<FSoftObjectPath>& GetSpecificAssets() const { return SpecificAssets; }
-#endif //if WITH_EDITOR
-
 private:
+	// Use accessors below to modify config data
+
 	/** Directories to search for this asset type */
 	UPROPERTY(EditAnywhere, Category = AssetType, meta = (RelativeToGameContentDir, LongPackageName))
 	TArray<FDirectoryPath> Directories;
@@ -206,6 +196,7 @@ public:
 
 	FPrimaryAssetTypeInfo() : AssetBaseClass(UObject::StaticClass()), AssetBaseClassLoaded(UObject::StaticClass()), bHasBlueprintClasses(false), bIsEditorOnly(false),  bIsDynamicAsset(false), NumberOfAssets(0) {}
 
+	/** Initializes a runtime version of the struct */
 	FPrimaryAssetTypeInfo(FName InPrimaryAssetType, UClass* InAssetBaseClass, bool bInHasBlueprintClasses, bool bInIsEditorOnly)
 		: PrimaryAssetType(InPrimaryAssetType), AssetBaseClass(InAssetBaseClass), AssetBaseClassLoaded(InAssetBaseClass), bHasBlueprintClasses(bInHasBlueprintClasses), bIsEditorOnly(bInIsEditorOnly), bIsDynamicAsset(false), NumberOfAssets(0)
 	{
@@ -224,11 +215,58 @@ public:
 	{
 	}
 
+	/** Gets the config version of the asset base class */
+	const TSoftClassPtr<UObject>& GetAssetBaseClass() const
+	{
+		return AssetBaseClass;
+	}
+
+	/** Set the base class, only possible before runtime data is filled in */
+	void SetAssetBaseClass(const TSoftClassPtr<UObject>& InAssetBaseClass)
+	{
+		if (ensure(CanModifyConfigData()))
+		{
+			AssetBaseClass = InAssetBaseClass;
+		}
+	}
+
+	/** Access the config version of directories to scan, may be empty for runtime-added types */
+	const TArray<FDirectoryPath>& GetDirectories() const 
+	{
+		return Directories;
+	}
+
+	/** Modify the config version of directories to scan, only possible before runtime data is filled in */
+	TArray<FDirectoryPath>& GetDirectories()
+	{
+		ensure(CanModifyConfigData());
+		return Directories; 
+	}
+
+	/** Access the config version of specific assets to scan, may be empty for runtime-added types */
+	const TArray<FSoftObjectPath>& GetSpecificAssets() const 
+	{
+		return SpecificAssets;
+	}
+	
+	/** Modify the config version of specific assets to scan, only possible before runtime data is filled in */
+	TArray<FSoftObjectPath>& GetSpecificAssets() 
+	{
+		ensure(CanModifyConfigData());
+		return SpecificAssets;
+	}
+
+	/** Returns true if this has valid config data that can be converted into runtime data as needed, this may be false for runtime only info */
+	ENGINE_API bool HasValidConfigData() const;
+
+	/** Returns true if this has valid runtime data because FillRuntimeData was previously called */
+	ENGINE_API bool HasValidRuntimeData() const;
+
+	/** Returns true if it is safe to modify config data as it has not been turned into runtime data yet */
+	ENGINE_API bool CanModifyConfigData() const;
+
 	/** Fills out transient variables based on parsed ones. Sets status bools saying rather data is valid, and rather it had to synchronously load the base class */
 	ENGINE_API void FillRuntimeData(bool& bIsValid, bool& bBaseClassWasLoaded);
-
-	/** Very temporary workaround for the gamefeatures subsystem needing to modify directories **/
-	friend class UGameFeaturesSubsystem;
 };
 
 /** Information about a package chunk, computed by the asset manager or read out of the cooked asset registry */
diff --git a/Engine/Source/Runtime/Engine/Private/AssetManager.cpp b/Engine/Source/Runtime/Engine/Private/AssetManager.cpp
--- a/Engine/Source/Runtime/Engine/Private/AssetManager.cpp
+++ b/Engine/Source/Runtime/Engine/Private/AssetManager.cpp
@@ -131,10 +131,13 @@ struct FPrimaryAssetTypeData
 	void ShrinkAssets();
 
 	/** In the editor, paths that we need to scan once asset registry is done loading */
-	TArray<FString> DeferredAssetScanPaths;
+	TSet<FString> DeferredAssetScanPaths;
+
+	/** List of paths that were explicitly requested by other systems and not loaded from the default config */
+	TSet<FString> AdditionalAssetScanPaths;
 
 	/** Expanded list of asset scan paths and package names, will not include virtual paths */
-	TArray<FString> RealAssetScanPaths;
+	TSet<FString> RealAssetScanPaths;
 
 	FPrimaryAssetTypeData() {}
 	~FPrimaryAssetTypeData();
@@ -928,7 +931,7 @@ int32 UAssetManager::ScanPathsForPrimaryAssets(FPrimaryAssetType PrimaryAssetTyp
 	// Add path info
 	for (const FString& Path : Paths)
 	{
-		TypeData.Info.AssetScanPaths.AddUnique(Path);
+		InternalAddAssetScanPath(TypeData, Path);
 	}
 
 #if WITH_EDITOR
@@ -939,7 +942,7 @@ int32 UAssetManager::ScanPathsForPrimaryAssets(FPrimaryAssetType PrimaryAssetTyp
 		// Keep track of the paths we asked for so once assets are discovered we will refresh the list
 		for (const FString& Path : Paths)
 		{
-			TypeData.DeferredAssetScanPaths.AddUnique(Path);
+			TypeData.DeferredAssetScanPaths.Add(Path);
 		}
 
 		// Since we are still asynchronously discovering assets, we'll wait until that is done before populating with primary assets
@@ -958,7 +961,7 @@ int32 UAssetManager::ScanPathsForPrimaryAssets(FPrimaryAssetType PrimaryAssetTyp
 	SearchRules.bSkipVirtualPathExpansion = true;
 	for (const FString& Path : SearchRules.AssetScanPaths)
 	{
-		TypeData.RealAssetScanPaths.AddUnique(Path);
+		TypeData.RealAssetScanPaths.Add(Path);
 	}
 
 	TArray<FAssetData> AssetDataList;
@@ -1060,6 +1063,7 @@ void UAssetManager::RemoveScanPathsForPrimaryAssets(FPrimaryAssetType PrimaryAss
 		}
 
 		TypeData.DeferredAssetScanPaths.Remove(Path);
+		TypeData.AdditionalAssetScanPaths.Remove(Path);
 	}
 
 	// Expand paths so we can record them for later
@@ -1070,6 +1074,16 @@ void UAssetManager::RemoveScanPathsForPrimaryAssets(FPrimaryAssetType PrimaryAss
 	}
 }
 
+void UAssetManager::InternalAddAssetScanPath(FPrimaryAssetTypeData& TypeData, const FString& AssetScanPath)
+{
+	TypeData.Info.AssetScanPaths.AddUnique(AssetScanPath);
+
+	if (!IsScanningFromInitialConfig())
+	{
+		TypeData.AdditionalAssetScanPaths.Add(AssetScanPath);
+	}
+}
+
 void UAssetManager::RemovePrimaryAssetType(FPrimaryAssetType PrimaryAssetType)
 {
 	TSharedRef<FPrimaryAssetTypeData>* TypeDataRef;
@@ -1118,17 +1132,21 @@ bool UAssetManager::RegisterSpecificPrimaryAsset(const FPrimaryAssetId& PrimaryA
 		return false;
 	}
 
-	const TSharedRef<FPrimaryAssetTypeData>* FoundType = AssetTypeMap.Find(PrimaryAssetId.PrimaryAssetType);
+	TSharedRef<FPrimaryAssetTypeData>* FoundType = AssetTypeMap.Find(PrimaryAssetId.PrimaryAssetType);
 	if (!FoundType)
 	{
 		return false;
 	}
 
+	FPrimaryAssetTypeData& TypeData = FoundType->Get();
 	if (!TryUpdateCachedAssetData(PrimaryAssetId, NewAssetData, false))
 	{
 		return false;
 	}
 
+	// Add to the list of scan paths so this will be found on refresh
+	InternalAddAssetScanPath(TypeData, NewAssetData.GetSoftObjectPath().ToString());
+
 	OnObjectReferenceListInvalidated();
 
 	return true;
@@ -3305,6 +3323,7 @@ void UAssetManager::ScanPrimaryAssetTypesFromConfig()
 	const UAssetManagerSettings& Settings = GetSettings();
 
 	PushBulkScanning();
+	TGuardValue<bool> ScopeGuard(bScanningFromInitialConfig, true);
 
 	double LastPumpTime = FPlatformTime::Seconds();
 	for (FPrimaryAssetTypeInfo TypeInfo : Settings.PrimaryAssetTypesToScan)
@@ -3661,6 +3680,11 @@ bool UAssetManager::IsPathExcludedFromScan(const FString& Path) const
 	return false;
 }
 
+bool UAssetManager::IsScanningFromInitialConfig() const
+{
+	return bScanningFromInitialConfig;
+}
+
 bool UAssetManager::GetContentRootPathFromPackageName(const FString& PackageName, FString& OutContentRootPath)
 {
 	if (PackageName.StartsWith(TEXT("/"), ESearchCase::CaseSensitive))
@@ -3734,6 +3758,7 @@ void UAssetManager::RefreshPrimaryAssetDirectory(bool bForceRefresh)
 	if (bForceRefresh || !bIsPrimaryAssetDirectoryCurrent)
 	{
 		PushBulkScanning();
+		TGuardValue<bool> ScopeGuard(bScanningFromInitialConfig, true);
 
 		for (TPair<FName, TSharedRef<FPrimaryAssetTypeData>>& TypePair : AssetTypeMap)
 		{
@@ -3855,6 +3880,7 @@ EAssetSetManagerResult::Type UAssetManager::ShouldSetManager(const FAssetIdentif
 void UAssetManager::OnAssetRegistryFilesLoaded()
 {
 	PushBulkScanning();
+	TGuardValue<bool> ScopeGuard(bScanningFromInitialConfig, true);
 
 	for (TPair<FName, TSharedRef<FPrimaryAssetTypeData>>& TypePair : AssetTypeMap)
 	{
@@ -3863,7 +3889,7 @@ void UAssetManager::OnAssetRegistryFilesLoaded()
 		if (TypeData.DeferredAssetScanPaths.Num())
 		{
 			// File scan finished, now scan for assets. Maps are sorted so this will be in the order of original scan requests
-			ScanPathsForPrimaryAssets(TypePair.Key, TypeData.DeferredAssetScanPaths, TypeData.Info.AssetBaseClassLoaded, TypeData.Info.bHasBlueprintClasses, TypeData.Info.bIsEditorOnly, false);
+			ScanPathsForPrimaryAssets(TypePair.Key, TypeData.DeferredAssetScanPaths.Array(), TypeData.Info.AssetBaseClassLoaded, TypeData.Info.bHasBlueprintClasses, TypeData.Info.bIsEditorOnly, false);
 
 			TypeData.DeferredAssetScanPaths.Empty();
 		}
@@ -4641,12 +4667,17 @@ void UAssetManager::ReinitializeFromConfig()
 	// We specifically do not reset AssetRuleOverrides as those can be set by something other than inis
 	for (TPair<FName, TSharedRef<FPrimaryAssetTypeData>>& Pair : AssetTypeMap)
 	{
-		Pair.Value->ResetAssets(AssetPathMap);
+		if (!Pair.Value->Info.bIsDynamicAsset)
+		{
+			Pair.Value->ResetAssets(AssetPathMap);
+		}
 	}
 	check(AssetPathMap.IsEmpty()); // Should have been emptied by the ResetAssets calls
 	ManagementParentMap.Reset();
 	CachedAssetBundles.Reset();
 	AlreadyScannedDirectories.Reset();
+
+	TMap<FName, TSharedRef<FPrimaryAssetTypeData>> OldAssetTypeMap = MoveTemp(AssetTypeMap);
 	AssetTypeMap.Reset();
 
 	// This code is editor only, so reinitialize globals
@@ -4666,6 +4697,23 @@ void UAssetManager::ReinitializeFromConfig()
 
 	LoadRedirectorMaps();
 	ScanPrimaryAssetTypesFromConfig();
+
+	// Go through old list and restore data that was added after the initial config load
+	for (TPair<FName, TSharedRef<FPrimaryAssetTypeData>>& TypePair : OldAssetTypeMap)
+	{
+		FPrimaryAssetTypeData& TypeData = TypePair.Value.Get();
+
+		if (TypeData.Info.bIsDynamicAsset)
+		{
+			// Restore dynamic assets as they were before
+			AssetTypeMap.Add(TypePair.Key, TypePair.Value);
+		}
+		else if (TypeData.AdditionalAssetScanPaths.Num())
+		{
+			// Rescan any paths added after initial scan
+			ScanPathsForPrimaryAssets(TypePair.Key, TypeData.AdditionalAssetScanPaths.Array(), TypeData.Info.AssetBaseClassLoaded, TypeData.Info.bHasBlueprintClasses, TypeData.Info.bIsEditorOnly, false);
+		}
+	}
 }
 
 void UAssetManager::OnInMemoryAssetCreated(UObject *Object)
diff --git a/Engine/Source/Runtime/Engine/Private/AssetManagerTypes.cpp b/Engine/Source/Runtime/Engine/Private/AssetManagerTypes.cpp
--- a/Engine/Source/Runtime/Engine/Private/AssetManagerTypes.cpp
+++ b/Engine/Source/Runtime/Engine/Private/AssetManagerTypes.cpp
@@ -8,6 +8,50 @@
 
 #include UE_INLINE_GENERATED_CPP_BY_NAME(AssetManagerTypes)
 
+bool FPrimaryAssetTypeInfo::HasValidConfigData() const
+{
+	if (PrimaryAssetType == NAME_None)
+	{
+		return false;
+	}
+
+	if (AssetBaseClass.IsNull())
+	{
+		return false;
+	}
+
+	// No paths are required
+
+	return true;
+}
+
+bool FPrimaryAssetTypeInfo::CanModifyConfigData() const
+{
+	// Can't modify config data after adding paths
+	return AssetScanPaths.Num() == 0;
+}
+
+bool FPrimaryAssetTypeInfo::HasValidRuntimeData() const
+{
+	if (PrimaryAssetType == NAME_None)
+	{
+		return false;
+	}
+
+	if (AssetBaseClassLoaded == nullptr)
+	{
+		return false;
+	}
+
+	// Invalid if the paths haven't been copied over yet, this is valid if all paths are empty
+	if ((AssetScanPaths.Num() == 0) && (Directories.Num() + SpecificAssets.Num() > 0))
+	{
+		return false;
+	}
+
+	return true;
+}
+
 void FPrimaryAssetTypeInfo::FillRuntimeData(bool& bIsValid, bool& bBaseClassWasLoaded)
 {
 	bBaseClassWasLoaded = false;
@@ -57,7 +101,7 @@ void FPrimaryAssetTypeInfo::FillRuntimeData(bool& bIsValid, bool& bBaseClassWasL
 	}
 
 	// Valid data, it's fine for a type to have no scan directories
-	bIsValid = true;
+	bIsValid = ensureMsgf(HasValidRuntimeData(), TEXT("Failed to FillRuntimeData for Primary Asset Type %s"), *PrimaryAssetType.ToString());
 }
 
 bool FPrimaryAssetRules::IsDefault() const
