Fixes for DDC2 texture build path.

diff --git a/Engine/Source/Runtime/Engine/Private/TextureDerivedDataTask.cpp b/Engine/Source/Runtime/Engine/Private/TextureDerivedDataTask.cpp
--- a/Engine/Source/Runtime/Engine/Private/TextureDerivedDataTask.cpp
+++ b/Engine/Source/Runtime/Engine/Private/TextureDerivedDataTask.cpp
@@ -2154,6 +2154,7 @@ static bool UnpackPlatformDataFromBuild(FTexturePlatformData& OutPlatformData, U
 	UE::DerivedData::FBuildOutput& BuildOutput = InBuildCompleteParams.Output;
 
 	FImage CPUCopy;
+	CPUCopy.Format = ERawImageFormat::Invalid;
 	{
 		// CPUCopy might not exist if the build didn't request it
 		const FValueWithId& MetadataValue = BuildOutput.GetValue(FValueId::FromName(ANSITEXTVIEW("CPUCopyImageInfo")));
@@ -2765,8 +2766,8 @@ private:
 		ANSITEXTVIEW("EncodedTextureExtendedData"),
 		ANSITEXTVIEW("MipTail"),
 		ANSITEXTVIEW("TextureBuildMetadata"),
-		ANSITEXTVIEW("CPUCopyMetadata"),
-		ANSITEXTVIEW("CPUCopyData")
+		ANSITEXTVIEW("CPUCopyImageInfo"),
+		ANSITEXTVIEW("CPUCopyRawData")
 	};
 
 
