Cheats: Fixed crash when calling UCheatManager::FreezeFrame() on a client.

#rnx

diff --git a/Engine/Source/Runtime/Engine/Private/CheatManager.cpp b/Engine/Source/Runtime/Engine/Private/CheatManager.cpp
--- a/Engine/Source/Runtime/Engine/Private/CheatManager.cpp
+++ b/Engine/Source/Runtime/Engine/Private/CheatManager.cpp
@@ -144,10 +144,18 @@ bool UCheatManager::ProcessConsoleExec(const TCHAR* Cmd, FOutputDevice& Ar, UObj
 
 void UCheatManager::FreezeFrame(float delay)
 {
-	FCanUnpause DefaultCanUnpause;
-	DefaultCanUnpause.BindUObject( GetOuterAPlayerController(), &APlayerController::DefaultCanUnpause );
-	GetWorld()->GetAuthGameMode()->SetPause(GetOuterAPlayerController(),DefaultCanUnpause);
-	GetWorld()->PauseDelay = GetWorld()->TimeSeconds + delay;
+	if (UWorld* World = GetWorld())
+	{
+		if (AGameModeBase* GameMode = World->GetAuthGameMode())
+		{
+			check(GetOuterAPlayerController() != NULL);
+
+			FCanUnpause DefaultCanUnpause;
+			DefaultCanUnpause.BindUObject(GetOuterAPlayerController(), &APlayerController::DefaultCanUnpause);
+			GameMode->SetPause(GetOuterAPlayerController(), DefaultCanUnpause);
+			World->PauseDelay = World->TimeSeconds + delay;
+		}
+	}
 }
 
 void UCheatManager::Teleport()
