Display human readable DDC key for the DefaultMaterial.

- This is to help debug and diagnose missmatched DDC issues between streams.

diff --git a/Engine/Source/Runtime/Engine/Private/Materials/MaterialShader.cpp b/Engine/Source/Runtime/Engine/Private/Materials/MaterialShader.cpp
--- a/Engine/Source/Runtime/Engine/Private/Materials/MaterialShader.cpp
+++ b/Engine/Source/Runtime/Engine/Private/Materials/MaterialShader.cpp
@@ -1489,7 +1489,7 @@ TSharedRef<FMaterialShaderMap::FAsyncLoadContext> FMaterialShaderMap::BeginLoadF
 			{
 				FString SpecialEngineDDCKey = Result->DataKey;
 				SpecialEngineDDCKey.RemoveFromEnd(TEXT("\n"));
-				UE_LOG(LogMaterial, Log, TEXT("%s-%s-%s: %s"), *Material->GetAssetName(), *LexToString(ShaderMapId.FeatureLevel), *LexToString(ShaderMapId.QualityLevel), *SpecialEngineDDCKey);
+				UE_LOG(LogMaterial, Display, TEXT("%s-%s-%s: %s"), *Material->GetAssetName(), *LexToString(ShaderMapId.FeatureLevel), *LexToString(ShaderMapId.QualityLevel), *SpecialEngineDDCKey);
 			}
 
 			if (UNLIKELY(!CVarShaderCompilerDebugDDCKeyAsset.GetValueOnAnyThread().IsEmpty()))
