Expose some FBox math utils

diff --git a/Engine/Source/Runtime/Engine/Classes/Kismet/KismetMathLibrary.h b/Engine/Source/Runtime/Engine/Classes/Kismet/KismetMathLibrary.h
--- a/Engine/Source/Runtime/Engine/Classes/Kismet/KismetMathLibrary.h
+++ b/Engine/Source/Runtime/Engine/Classes/Kismet/KismetMathLibrary.h
@@ -4293,6 +4293,30 @@ class UKismetMathLibrary : public UBlueprintFunctionLibrary
     UFUNCTION(BlueprintPure, Category = "Math|Geometry", meta=(DisplayName = "Is Point In Box (Box)"))
     static ENGINE_API bool IsPointInBox_Box(FVector Point, FBox Box);
 
+	/**
+	 * Gets the volume of this box.
+	 *
+	 * @return The box volume.
+	 */
+	UFUNCTION(BlueprintPure, Category = "Math|Geometry", meta=(DisplayName = "Get Volume (Box)"))
+	static ENGINE_API double GetBoxVolume(const FBox& InBox);
+
+	/**
+	 * Gets the size of this box.
+	 *
+	 * @return The box size.
+	 */
+	UFUNCTION(BlueprintPure, Category = "Math|Geometry", meta=(DisplayName = "Get Size (Box)"))
+	static ENGINE_API FVector GetBoxSize(const FBox& InBox);
+
+	/**
+	 * Gets the center point of this box.
+	 *
+	 * @return The center point.
+	 */
+	UFUNCTION(BlueprintPure, Category = "Math|Geometry", meta=(DisplayName = "Get Center (Box)"))
+	static ENGINE_API FVector GetBoxCenter(const FBox& InBox);
+
 	/**
 	* Determines whether a given point is in a box with a given transform. Includes points on the box.
 	*
diff --git a/Engine/Source/Runtime/Engine/Private/KismetMathLibrary.cpp b/Engine/Source/Runtime/Engine/Private/KismetMathLibrary.cpp
--- a/Engine/Source/Runtime/Engine/Private/KismetMathLibrary.cpp
+++ b/Engine/Source/Runtime/Engine/Private/KismetMathLibrary.cpp
@@ -7,6 +7,7 @@
 #include "Blueprint/BlueprintSupport.h"
 #include "Math/ConvexHull2d.h"
 #include "Math/DualQuat.h"
+#include "Math/Box.h"
 
 #include "Math/RandomStream.h"
 #include "Misc/RuntimeErrors.h"
@@ -1371,6 +1372,21 @@ bool UKismetMathLibrary::IsPointInBox_Box(FVector Point, FBox Box)
 	return Box.IsInsideOrOn(Point);
 }
 
+double UKismetMathLibrary::GetBoxVolume(const FBox& InBox)
+{
+	return InBox.GetVolume();
+}
+
+FVector UKismetMathLibrary::GetBoxSize(const FBox& InBox)
+{
+	return InBox.GetSize();
+}
+
+FVector UKismetMathLibrary::GetBoxCenter(const FBox& InBox)
+{
+	return InBox.GetCenter();
+}
+
 bool UKismetMathLibrary::IsPointInBoxWithTransform(FVector Point, const FTransform& BoxWorldTransform, FVector BoxExtent)
 {
 	// Put point in component space
