Apply localization macro to string literals in GetAttributeOverrideForMaterial for consistent language localisation in the material editor's material attribute node.

diff --git a/Engine/Source/Runtime/Engine/Private/Materials/MaterialAttributeDefinitionMap.cpp b/Engine/Source/Runtime/Engine/Private/Materials/MaterialAttributeDefinitionMap.cpp
--- a/Engine/Source/Runtime/Engine/Private/Materials/MaterialAttributeDefinitionMap.cpp
+++ b/Engine/Source/Runtime/Engine/Private/Materials/MaterialAttributeDefinitionMap.cpp
@@ -371,9 +371,9 @@ FText FMaterialAttributeDefinitionMap::GetAttributeOverrideForMaterial(const FGu
 	case MP_BaseColor:
 		return Material->MaterialDomain == MD_Volume ? LOCTEXT("Albedo", "Albedo") : LOCTEXT("BaseColor", "Base Color");
 	case MP_Metallic:
-		CustomPinNames.Add({ MSM_Hair, "Scatter" });
-		CustomPinNames.Add({ MSM_Eye, "Curvature" });
-		return FText::FromString(GetPinNameFromShadingModelField(Material->GetShadingModels(), CustomPinNames, "Metallic"));
+		CustomPinNames.Add({ MSM_Hair, LOCTEXT("Scatter", "Scatter").ToString()});
+		CustomPinNames.Add({ MSM_Eye, LOCTEXT("Curvature", "Curvature").ToString()});
+		return FText::FromString(GetPinNameFromShadingModelField(Material->GetShadingModels(), CustomPinNames, LOCTEXT("Metallic", "Metallic").ToString()));
 	case MP_Specular:
 		return LOCTEXT("Specular", "Specular");
 	case MP_Roughness:
@@ -381,8 +381,8 @@ FText FMaterialAttributeDefinitionMap::GetAttributeOverrideForMaterial(const FGu
 	case MP_Anisotropy:
 		return LOCTEXT("Anisotropy", "Anisotropy");
 	case MP_Normal:
-		CustomPinNames.Add({ MSM_Hair, "Tangent" });
-		return FText::FromString(GetPinNameFromShadingModelField(Material->GetShadingModels(), CustomPinNames, "Normal"));
+		CustomPinNames.Add({ MSM_Hair, LOCTEXT("Tangent", "Tangent").ToString() });
+		return FText::FromString(GetPinNameFromShadingModelField(Material->GetShadingModels(), CustomPinNames, LOCTEXT("Normal", "Normal").ToString()));
 	case MP_Tangent:
 		return LOCTEXT("Tangent", "Tangent");
 	case MP_WorldPositionOffset:
@@ -398,19 +398,19 @@ FText FMaterialAttributeDefinitionMap::GetAttributeOverrideForMaterial(const FGu
 		{
 			return LOCTEXT("Extinction", "Extinction");
 		}
-		CustomPinNames.Add({ MSM_Cloth, "Fuzz Color" });
-		return FText::FromString(GetPinNameFromShadingModelField(Material->GetShadingModels(), CustomPinNames, "Subsurface Color"));
+		CustomPinNames.Add({ MSM_Cloth, LOCTEXT("FuzzColor", "Fuzz Color").ToString() });
+		return FText::FromString(GetPinNameFromShadingModelField(Material->GetShadingModels(), CustomPinNames, LOCTEXT("SubsurfaceColor", "Subsurface Color").ToString()));
 	case MP_CustomData0:
-		CustomPinNames.Add({ MSM_ClearCoat, "Clear Coat" });
-		CustomPinNames.Add({ MSM_Hair, "Backlit" });
-		CustomPinNames.Add({ MSM_Cloth, "Cloth" });
-		CustomPinNames.Add({ MSM_Eye, "Iris Mask" });
-		CustomPinNames.Add({ MSM_SubsurfaceProfile, "Curvature" });
-		return FText::FromString(GetPinNameFromShadingModelField(Material->GetShadingModels(), CustomPinNames, "Custom Data 0"));
+		CustomPinNames.Add({ MSM_ClearCoat, LOCTEXT("ClearCoat", "Clear Coat").ToString() });
+		CustomPinNames.Add({ MSM_Hair, LOCTEXT("Backlit", "Backlit").ToString() });
+		CustomPinNames.Add({ MSM_Cloth, LOCTEXT("Cloth", "Cloth").ToString() });
+		CustomPinNames.Add({ MSM_Eye, LOCTEXT("IrisMask", "Iris Mask").ToString() });
+		CustomPinNames.Add({ MSM_SubsurfaceProfile, LOCTEXT("Curvature", "Curvature").ToString() });
+		return FText::FromString(GetPinNameFromShadingModelField(Material->GetShadingModels(), CustomPinNames, LOCTEXT("CustomData0", "Custom Data 0").ToString()));
 	case MP_CustomData1:
-		CustomPinNames.Add({ MSM_ClearCoat, "Clear Coat Roughness" });
-		CustomPinNames.Add({ MSM_Eye, "Iris Distance" });
-		return FText::FromString(GetPinNameFromShadingModelField(Material->GetShadingModels(), CustomPinNames, "Custom Data 1"));
+		CustomPinNames.Add({ MSM_ClearCoat, LOCTEXT("ClearCoatRoughness", "Clear Coat Roughness").ToString() });
+		CustomPinNames.Add({ MSM_Eye, LOCTEXT("IrisDistance", "Iris Distance").ToString() });
+		return FText::FromString(GetPinNameFromShadingModelField(Material->GetShadingModels(), CustomPinNames, LOCTEXT("CustomData1", "Custom Data 1").ToString()));
 	case MP_AmbientOcclusion:
 		return LOCTEXT("AmbientOcclusion", "Ambient Occlusion");
 	case MP_Refraction:
