Added comments (from nathaniel.morgan) to post process settings.

diff --git a/Engine/Source/Runtime/Engine/Classes/Engine/Scene.h b/Engine/Source/Runtime/Engine/Classes/Engine/Scene.h
--- a/Engine/Source/Runtime/Engine/Classes/Engine/Scene.h
+++ b/Engine/Source/Runtime/Engine/Classes/Engine/Scene.h
@@ -1343,61 +1343,84 @@ struct FPostProcessSettings
 	*/
 	UPROPERTY(interp, BlueprintReadWrite, Category = "Color Grading|Temperature", meta = (DisplayName = "Temperature Type", editcondition = "bOverride_TemperatureType" ))
 	TEnumAsByte<enum ETemperatureMethod> TemperatureType;
+	/** Controls the color temperature or white balance in degrees Kelvin which the scene considers as white light. */
 	UPROPERTY(interp, BlueprintReadWrite, Category="Color Grading|Temperature", meta=(UIMin = "1500.0", UIMax = "15000.0", editcondition = "bOverride_WhiteTemp", DisplayName = "Temp"))
 	float WhiteTemp;
+	/** Controls the color of the scene along the magenta - green axis (orthogonal to the color temperature).  This feature is equivalent to color tint in digital cameras. */
 	UPROPERTY(interp, BlueprintReadWrite, Category="Color Grading|Temperature", meta=(UIMin = "-1.0", UIMax = "1.0", editcondition = "bOverride_WhiteTint", DisplayName = "Tint"))
 	float WhiteTint;
 
-	// Color Correction controls
+	/** Control the intensity of the color(hue) for the entire image.Higher values will result in more vibrant colors. */
 	UPROPERTY(interp, BlueprintReadWrite, Category = "Color Grading|Global", meta = (UIMin = "0.0", UIMax = "2.0", Delta = "0.01", ColorGradingMode = "saturation", ShiftMouseMovePixelPerDelta = "10", SupportDynamicSliderMaxValue = "true", editcondition = "bOverride_ColorSaturation", DisplayName = "Saturation"))
 	FVector4 ColorSaturation;
+	/** Control the range of light and dark values in your scene. Lower values will reduce the difference between bright and dark areas while higher values will increase the difference between the bright and dark areas. */
 	UPROPERTY(interp, BlueprintReadWrite, Category = "Color Grading|Global", meta = (UIMin = "0.0", UIMax = "2.0", Delta = "0.01", ColorGradingMode = "contrast", ShiftMouseMovePixelPerDelta = "10", SupportDynamicSliderMaxValue = "true", editcondition = "bOverride_ColorContrast", DisplayName = "Contrast"))
 	FVector4 ColorContrast;
+	/** Control the luminance curve of the scene. Raising or lowering this value will result brightening or darkening the mid-tones of the entire image. */
 	UPROPERTY(interp, BlueprintReadWrite, Category = "Color Grading|Global", meta = (UIMin = "0.0", UIMax = "2.0", Delta = "0.01", ColorGradingMode = "gamma", ShiftMouseMovePixelPerDelta = "10", SupportDynamicSliderMaxValue = "true", editcondition = "bOverride_ColorGamma", DisplayName = "Gamma"))
 	FVector4 ColorGamma;
+	/** This value multiplies the colors of the image.  Raising or lowering this value will result in brightening or darkening the entire scene. */
 	UPROPERTY(interp, BlueprintReadWrite, Category = "Color Grading|Global", meta = (UIMin = "0.0", UIMax = "2.0", Delta = "0.01", ColorGradingMode = "gain", ShiftMouseMovePixelPerDelta = "10", SupportDynamicSliderMaxValue = "true", editcondition = "bOverride_ColorGain", DisplayName = "Gain"))
 	FVector4 ColorGain;
+	/** This value is added to the colors of the scene.  Raising or lowering this value will result in the image being more or less washed-out. */
 	UPROPERTY(interp, BlueprintReadWrite, Category = "Color Grading|Global", meta = (UIMin = "-1.0", UIMax = "1.0", Delta = "0.001", ColorGradingMode = "offset", ShiftMouseMovePixelPerDelta = "20", SupportDynamicSliderMaxValue = "true", SupportDynamicSliderMinValue = "true", editcondition = "bOverride_ColorOffset", DisplayName = "Offset"))
 	FVector4 ColorOffset;
 
+	/** Control the intensity of the colors (hue) in the shadow region of the image.  Higher values will result in more vibrant colors. */
 	UPROPERTY(interp, BlueprintReadWrite, Category = "Color Grading|Shadows", meta = (UIMin = "0.0", UIMax = "2.0", Delta = "0.01", ColorGradingMode = "saturation", ShiftMouseMovePixelPerDelta = "10", SupportDynamicSliderMaxValue = "true", editcondition = "bOverride_ColorSaturationShadows", DisplayName = "Saturation"))
 	FVector4 ColorSaturationShadows;
+	/** Control the range of light and dark values in your scene. Lower values will reduce the difference between bright and dark areas while higher values will increase the difference between the bright and dark areas. */
 	UPROPERTY(interp, BlueprintReadWrite, Category = "Color Grading|Shadows", meta = (UIMin = "0.0", UIMax = "2.0", Delta = "0.01", ColorGradingMode = "contrast", ShiftMouseMovePixelPerDelta = "10", SupportDynamicSliderMaxValue = "true", editcondition = "bOverride_ColorContrastShadows", DisplayName = "Contrast"))
 	FVector4 ColorContrastShadows;
+	/** Control the luminance curve of the shadow region. Raising or lowering this value will result brightening or darkening the mid-tones of the shadow region. */
 	UPROPERTY(interp, BlueprintReadWrite, Category = "Color Grading|Shadows", meta = (UIMin = "0.0", UIMax = "2.0", Delta = "0.01", ColorGradingMode = "gamma", ShiftMouseMovePixelPerDelta = "10", SupportDynamicSliderMaxValue = "true", editcondition = "bOverride_ColorGammaShadows", DisplayName = "Gamma"))
 	FVector4 ColorGammaShadows;
+	/** This value multiplies the colors in the shadow region.  Raising or lowering this value will result in brightening or darkening the affected region. */
 	UPROPERTY(interp, BlueprintReadWrite, Category = "Color Grading|Shadows", meta = (UIMin = "0.0", UIMax = "2.0", Delta = "0.01", ColorGradingMode = "gain", ShiftMouseMovePixelPerDelta = "10", SupportDynamicSliderMaxValue = "true", editcondition = "bOverride_ColorGainShadows", DisplayName = "Gain"))
 	FVector4 ColorGainShadows;
+	/** This value is added to the colors in the shadow region.  Raising or lowering this value will result in the shadows being more or less washed-out. */
 	UPROPERTY(interp, BlueprintReadWrite, Category = "Color Grading|Shadows", meta = (UIMin = "-1.0", UIMax = "1.0", Delta = "0.001", ColorGradingMode = "offset", ShiftMouseMovePixelPerDelta = "20", SupportDynamicSliderMaxValue = "true", SupportDynamicSliderMinValue = "true", editcondition = "bOverride_ColorOffsetShadows", DisplayName = "Offset"))
 	FVector4 ColorOffsetShadows;
 
+	/** Control the intensity of the colors (hue) in the mid-tone region of the image.  Higher values will result in more vibrant colors. */
 	UPROPERTY(interp, BlueprintReadWrite, Category = "Color Grading|Midtones", meta = (UIMin = "0.0", UIMax = "2.0", Delta = "0.01", ColorGradingMode = "saturation", ShiftMouseMovePixelPerDelta = "10", SupportDynamicSliderMaxValue = "true", editcondition = "bOverride_ColorSaturationMidtones", DisplayName = "Saturation"))
 	FVector4 ColorSaturationMidtones;
+	/** Control the range of light and dark values in the mid-tone region. Lower values will reduce the difference between bright and dark areas while higher values will increase the difference between the bright and dark areas. */
 	UPROPERTY(interp, BlueprintReadWrite, Category = "Color Grading|Midtones", meta = (UIMin = "0.0", UIMax = "2.0", Delta = "0.01", ColorGradingMode = "contrast", ShiftMouseMovePixelPerDelta = "10", SupportDynamicSliderMaxValue = "true", editcondition = "bOverride_ColorContrastMidtones", DisplayName = "Contrast"))
 	FVector4 ColorContrastMidtones;
+	/** Control the luminance curve of the mid-tone region of the image. Raising or lowering this value will result brightening or darkening the mid-tones of the image. */
 	UPROPERTY(interp, BlueprintReadWrite, Category = "Color Grading|Midtones", meta = (UIMin = "0.0", UIMax = "2.0", Delta = "0.01", ColorGradingMode = "gamma", ShiftMouseMovePixelPerDelta = "10", SupportDynamicSliderMaxValue = "true", editcondition = "bOverride_ColorGammaMidtones", DisplayName = "Gamma"))
 	FVector4 ColorGammaMidtones;
+	/** This value multiplies the colors in the mid-tone region of the image.  Raising or lowering this value will result in brightening or darkening the affected region. */
 	UPROPERTY(interp, BlueprintReadWrite, Category = "Color Grading|Midtones", meta = (UIMin = "0.0", UIMax = "2.0", Delta = "0.01", ColorGradingMode = "gain", ShiftMouseMovePixelPerDelta = "10", SupportDynamicSliderMaxValue = "true", editcondition = "bOverride_ColorGainMidtones", DisplayName = "Gain"))
 	FVector4 ColorGainMidtones;
+	/** This value is added to the colors in the mid-tone region of the image.  Raising or lowering this value will result in the mid-tones being more or less washed-out. */
 	UPROPERTY(interp, BlueprintReadWrite, Category = "Color Grading|Midtones", meta = (UIMin = "-1.0", UIMax = "1.0", Delta = "0.001", ColorGradingMode = "offset", ShiftMouseMovePixelPerDelta = "20", SupportDynamicSliderMaxValue = "true", SupportDynamicSliderMinValue = "true", editcondition = "bOverride_ColorOffsetMidtones", DisplayName = "Offset"))
 	FVector4 ColorOffsetMidtones;
 
+	/** Control the intensity of the color (hue) for the highlights region of the image.  Higher values will result in more vibrant colors. */
 	UPROPERTY(interp, BlueprintReadWrite, Category = "Color Grading|Highlights", meta = (UIMin = "0.0", UIMax = "2.0", Delta = "0.01", ColorGradingMode = "saturation", ShiftMouseMovePixelPerDelta = "10", SupportDynamicSliderMaxValue = "true", editcondition = "bOverride_ColorSaturationHighlights", DisplayName = "Saturation"))
 	FVector4 ColorSaturationHighlights;
+	/** Control the range of light and dark values in the highlights region. Lower values will reduce the difference between bright and dark areas while higher values will increase the difference between the bright and dark areas. */
 	UPROPERTY(interp, BlueprintReadWrite, Category = "Color Grading|Highlights", meta = (UIMin = "0.0", UIMax = "2.0", Delta = "0.01", ColorGradingMode = "contrast", ShiftMouseMovePixelPerDelta = "10", SupportDynamicSliderMaxValue = "true", editcondition = "bOverride_ColorContrastHighlights", DisplayName = "Contrast"))
 	FVector4 ColorContrastHighlights;
+	/** Control the luminance curve of the highlight region. Raising or lowering this value will result brightening or darkening the mid-tones of the highlight region. */
 	UPROPERTY(interp, BlueprintReadWrite, Category = "Color Grading|Highlights", meta = (UIMin = "0.0", UIMax = "2.0", Delta = "0.01", ColorGradingMode = "gamma", ShiftMouseMovePixelPerDelta = "10", SupportDynamicSliderMaxValue = "true", editcondition = "bOverride_ColorGammaHighlights", DisplayName = "Gamma"))
 	FVector4 ColorGammaHighlights;
+	/** This value multiplies the colors in the highlight region.  Raising or lowering this value will result in brightening or darkening the affected region. */
 	UPROPERTY(interp, BlueprintReadWrite, Category = "Color Grading|Highlights", meta = (UIMin = "0.0", UIMax = "2.0", Delta = "0.01", ColorGradingMode = "gain", ShiftMouseMovePixelPerDelta = "10", SupportDynamicSliderMaxValue = "true", editcondition = "bOverride_ColorGainHighlights", DisplayName = "Gain"))
 	FVector4 ColorGainHighlights;
+	/** This value is added to the colors in the highlight region.  Raising or lowering this value will result in the highlights being more or less washed-out. */
 	UPROPERTY(interp, BlueprintReadWrite, Category = "Color Grading|Highlights", meta = (UIMin = "-1.0", UIMax = "1.0", Delta = "0.001", ColorGradingMode = "offset", ShiftMouseMovePixelPerDelta = "20", SupportDynamicSliderMaxValue = "true", SupportDynamicSliderMinValue = "true", editcondition = "bOverride_ColorOffsetHighlights", DisplayName = "Offset"))
 	FVector4 ColorOffsetHighlights;
+	/** This value sets the lower threshold for what is considered to be the highlight region of the image. */
 	UPROPERTY(interp, BlueprintReadWrite, Category = "Color Grading|Highlights", meta = (UIMin = "-1.0", UIMax = "1.0", editcondition = "bOverride_ColorCorrectionHighlightsMin", DisplayName = "HighlightsMin"))
 	float ColorCorrectionHighlightsMin;
-	/** Should be set larger than HighlightsMin. Default is 1.0, for backwards compatibility */
+	/** This value sets the upper threshold for what is considered to be the highlight region of the image.  This value should be larger than HighlightsMin. Default is 1.0, for backwards compatibility */
 	UPROPERTY(interp, BlueprintReadWrite, Category = "Color Grading|Highlights", meta = (UIMin = "1.0", UIMax = "10.0", editcondition = "bOverride_ColorCorrectionHighlightsMax", DisplayName = "HighlightsMax"))
 	float ColorCorrectionHighlightsMax;
 
+	/** This value sets the threshold for what is considered to be the shadow region of the image. */
 	UPROPERTY(interp, BlueprintReadWrite, Category = "Color Grading|Shadows", meta = (UIMin = "-1.0", UIMax = "1.0", editcondition = "bOverride_ColorCorrectionShadowsMax", DisplayName = "ShadowsMax"))
 	float ColorCorrectionShadowsMax;
 
@@ -1411,14 +1434,19 @@ struct FPostProcessSettings
 	UPROPERTY(interp, BlueprintReadWrite, Category = "Color Grading|Misc", meta = (ClampMin = "0.0", UIMax = "1.0", editcondition = "bOverride_ToneCurveAmount"))
 	float ToneCurveAmount;
 
+	/** Controls the overall steepness of the tonemapper curve.  Larger values increase scene contrast and smaller values reduce contrast. */
 	UPROPERTY(interp, BlueprintReadWrite, Category="Film", meta=(UIMin = "0.0", UIMax = "1.0", editcondition = "bOverride_FilmSlope", DisplayName = "Slope"))
 	float FilmSlope;
+	/** Controls the contrast of the dark end of the tonemapper curve. Larger values increase contrast and smaller values decrease contrast. */
 	UPROPERTY(interp, BlueprintReadWrite, Category="Film", meta=(UIMin = "0.0", UIMax = "1.0", editcondition = "bOverride_FilmToe", DisplayName = "Toe"))
 	float FilmToe;
+	/**  Sometimes referred to as highlight rolloff.  Controls the contrast of the bright end of the tonemapper curve. Larger values increase contrast and smaller values decrease contrast.  */
 	UPROPERTY(interp, BlueprintReadWrite, Category="Film", meta=(UIMin = "0.0", UIMax = "1.0", editcondition = "bOverride_FilmShoulder", DisplayName = "Shoulder"))
 	float FilmShoulder;
+	/** Lowers the toe of the tonemapper curve by this amount. Increasing this value causes more of the scene to clip to black.  For most purposes, this property should remain 0 */
 	UPROPERTY(interp, BlueprintReadWrite, Category="Film", meta=(UIMin = "0.0", UIMax = "1.0", editcondition = "bOverride_FilmBlackClip", DisplayName = "Black clip"))
 	float FilmBlackClip;
+	/** Controls the height of the tonemapper curve.  Raising this value can cause bright values to more quickly approach fully-saturated white. */
 	UPROPERTY(interp, BlueprintReadWrite, Category="Film", meta=(UIMin = "0.0", UIMax = "1.0", editcondition = "bOverride_FilmWhiteClip", DisplayName = "White clip"))
 	float FilmWhiteClip;
 
@@ -1909,26 +1937,31 @@ struct FPostProcessSettings
 	UPROPERTY(interp, BlueprintReadWrite, Category = "Film Grain", meta = (UIMin = "0.0", UIMax = "1.0", editcondition = "bOverride_FilmGrainIntensity"))
 	float FilmGrainIntensity;
 
+	/** Control over the grain intensity in the regions of the image considered shadow areas. */
 	UPROPERTY(interp, BlueprintReadWrite, Category = "Film Grain", meta = (UIMin = "0.0", UIMax = "1.0", editcondition = "bOverride_FilmGrainIntensityShadows"))
 	float FilmGrainIntensityShadows;
 
+	/** Control over the grain intensity in the mid-tone region of the image. */
 	UPROPERTY(interp, BlueprintReadWrite, Category = "Film Grain", meta = (UIMin = "0.0", UIMax = "1.0", editcondition = "bOverride_FilmGrainIntensityMidtones"))
 	float FilmGrainIntensityMidtones;
 
+	/** Control over the grain intensity in the regions of the image considered highlight areas. */
 	UPROPERTY(interp, BlueprintReadWrite, Category = "Film Grain", meta = (UIMin = "0.0", UIMax = "1.0", editcondition = "bOverride_FilmGrainIntensityHighlights"))
 	float FilmGrainIntensityHighlights;
 
+	/** Sets the upper bound used for Film Grain Shadow Intensity. */
 	UPROPERTY(interp, BlueprintReadWrite, Category = "Film Grain", meta = (UIMin = "0.0", UIMax = "1.0", editcondition = "bOverride_FilmGrainShadowsMax"))
 	float FilmGrainShadowsMax;
 	
+	/** Sets the lower bound used for Film Grain Highlight Intensity. */
 	UPROPERTY(interp, BlueprintReadWrite, Category = "Film Grain", meta = (UIMin = "0.0", UIMax = "1.0", editcondition = "bOverride_FilmGrainHighlightsMin"))
 	float FilmGrainHighlightsMin;
 
-	/** Should be set larger than HighlightsMin. Default is 1.0, for backwards compatibility */
+	/** Sets the upper bound used for Film Grain Highlight Intensity. This value should be larger than HighlightsMin.. Default is 1.0, for backwards compatibility */
 	UPROPERTY(interp, BlueprintReadWrite, Category = "Film Grain", meta = (UIMin = "1.0", UIMax = "10.0", editcondition = "bOverride_FilmGrainHighlightsMax"))
 	float FilmGrainHighlightsMax;
 
-	/** Size of texel of FilmGrainTexture on screen. */
+	/** Controls the size of the film grain. Size of texel of FilmGrainTexture on screen. */
 	UPROPERTY(interp, BlueprintReadWrite, Category = "Film Grain", meta = (UIMin = "0.0", UIMax = "4.0", editcondition = "bOverride_FilmGrainTexelSize"))
 	float FilmGrainTexelSize;
 
@@ -2020,7 +2053,7 @@ struct FPostProcessSettings
 	UPROPERTY(EditAnywhere, BlueprintReadWrite, Category="Lens|Depth of Field", meta=(ForceUnits=mm, ClampMin = "0.1", UIMin="0.1", UIMax= "1000.0", editcondition = "bOverride_DepthOfFieldSensorWidth", DisplayName = "Sensor Width (mm)"))
 	float DepthOfFieldSensorWidth;
 
-	/** Width of the camera sensor to assume, in mm. */
+	/** This is the squeeze factor for the DOF, which emulates the properties of anamorphic lenses. */
 	UPROPERTY(EditAnywhere, BlueprintReadWrite, Category = "Lens|Depth of Field", meta = (ClampMin = "1.0", ClampMax = "2.0", editcondition = "bOverride_DepthOfFieldSqueezeFactor", DisplayName = "Squeeze Factor"))
 	float DepthOfFieldSqueezeFactor;
 
