Add path to allow the instance manager to reset tracked state when the component is registered, to handle properties being manipulated crucially replacing the instance data making the tracked count go out of sync. The issue manifested when applying changes in the editor to an ISM BP template.
- old behavior can be restored by disabling r.InstanceData.ResetTrackingOnRegister in case of fallout.

diff --git a/Engine/Source/Runtime/Engine/Private/HierarchicalInstancedStaticMesh.cpp b/Engine/Source/Runtime/Engine/Private/HierarchicalInstancedStaticMesh.cpp
--- a/Engine/Source/Runtime/Engine/Private/HierarchicalInstancedStaticMesh.cpp
+++ b/Engine/Source/Runtime/Engine/Private/HierarchicalInstancedStaticMesh.cpp
@@ -2957,8 +2957,8 @@ FPrimitiveSceneProxy* UHierarchicalInstancedStaticMeshComponent::CreateSceneProx
 
 	// Verify that the mesh is valid & that we have instances before creating a proxy.
 	const bool bAreMeshAndInstancesValid =
-		// Make sure we have instances, or an update with instances on the way (for the external data (landscape grass) mode).
-		(GetNumInstances() > 0 ||  PrimitiveInstanceDataManager.GetMaxInstanceIndex() > 0) &&
+		// Make sure we have instances, or an update with instances on the way, or already built instances (for the external data (landscape grass) mode).
+		(GetNumInstances() > 0 ||  PrimitiveInstanceDataManager.GetMaxInstanceIndex() > 0 || NumBuiltRenderInstances > 0) &&
 		// Make sure we have an actual static mesh.
 		GetStaticMesh() &&
 		!GetStaticMesh()->IsCompiling() &&
diff --git a/Engine/Source/Runtime/Engine/Private/InstancedStaticMesh.cpp b/Engine/Source/Runtime/Engine/Private/InstancedStaticMesh.cpp
--- a/Engine/Source/Runtime/Engine/Private/InstancedStaticMesh.cpp
+++ b/Engine/Source/Runtime/Engine/Private/InstancedStaticMesh.cpp
@@ -4328,6 +4328,10 @@ void UInstancedStaticMeshComponent::OnRegister()
 		{
 			InstancingRandomSeed = FMath::Rand();
 		}
+
+		// Propagate the current number of instances to the manager, because of the multifarious ways the engine shoves data into the properties it is possible for the count to get out of sync.
+		// At this point we need to let the manager reset if needed, also at this point we may assume that we have no idea of the state of individual members.
+		PrimitiveInstanceDataManager.OnRegister(PerInstanceSMData.Num());
 	}
 }
 		
diff --git a/Engine/Source/Runtime/Engine/Private/InstancedStaticMesh/ISMInstanceDataManager.cpp b/Engine/Source/Runtime/Engine/Private/InstancedStaticMesh/ISMInstanceDataManager.cpp
--- a/Engine/Source/Runtime/Engine/Private/InstancedStaticMesh/ISMInstanceDataManager.cpp
+++ b/Engine/Source/Runtime/Engine/Private/InstancedStaticMesh/ISMInstanceDataManager.cpp
@@ -11,17 +11,29 @@
 #define IDPROXY_ENABLE_ASYNC_TASK 1
 
 #if 0
-#define LOG_INST_DATA(_Format_, ...) UE_LOG(LogInstanceProxy, Log, _Format_, ##__VA_ARGS__)
+
+#define LOG_INST_DATA(_Format_, ...) \
+{ \
+	FString Tmp = FString::Printf(_Format_, ##__VA_ARGS__);\
+	UE_LOG(LogInstanceProxy, Log, TEXT("%p, %s"), PrimitiveComponent.IsValid() ? PrimitiveComponent.Get() : nullptr, *Tmp);\
+}
+
 #else
 	#define LOG_INST_DATA(_Format_, ...) 
 #endif
 
-TAutoConsoleVariable<float> CVarInstanceUpdateTaskDebugDelay(
+static TAutoConsoleVariable<float> CVarInstanceUpdateTaskDebugDelay(
 	TEXT("r.InstanceUpdateTaskDebugDelay"),
 	0.0f,
 	TEXT("Instance update debug delay in seconds."),
 	ECVF_RenderThreadSafe);
 
+static TAutoConsoleVariable<int32> CVarInstanceDataResetTrackingOnRegister(
+	TEXT("r.InstanceData.ResetTrackingOnRegister"),
+	1,
+	TEXT("Chicken switch to disable the new code to reset tracking & instance count during OnRegister, if this causes problems.\nTODO: Remove."));
+
+
 FPrimitiveInstanceDataManager::FPrimitiveInstanceDataManager(UPrimitiveComponent* InPrimitiveComponent) 
 	: PrimitiveComponent(InPrimitiveComponent) 
 {
@@ -30,6 +42,7 @@ FPrimitiveInstanceDataManager::FPrimitiveInstanceDataManager(UPrimitiveComponent
 	{
 		TrackingState = ETrackingState::Disabled;
 	}
+	LOG_INST_DATA(TEXT("FPrimitiveInstanceDataManager %s, TrackingState=%s"), *PrimitiveComponent->GetFullName(), TrackingState == ETrackingState::Disabled ? TEXT("Disabled") : TEXT("Initial"));
 }
 
 void FPrimitiveInstanceDataManager::SetMode(EMode InMode)
@@ -1093,6 +1106,7 @@ void FPrimitiveInstanceDataManager::MarkIndexChanged(FPrimitiveInstanceId Instan
 
 TSharedPtr<FISMCInstanceDataSceneProxy, ESPMode::ThreadSafe> FPrimitiveInstanceDataManager::GetOrCreateProxy(FStaticShaderPlatform InShaderPlatform, ERHIFeatureLevel::Type InFeatureLevel)
 {
+	LOG_INST_DATA(TEXT("GetOrCreateProxy"));
 	if (Proxy && !Proxy->CheckPlatformFeatureLevel(InShaderPlatform, InFeatureLevel))
 	{
 		// TODO: May need to add some attachment counter checks here?
@@ -1122,6 +1136,7 @@ TSharedPtr<FISMCInstanceDataSceneProxy, ESPMode::ThreadSafe> FPrimitiveInstanceD
 
 void FPrimitiveInstanceDataManager::Invalidate(int32 InNumInstances)
 {
+	LOG_INST_DATA(TEXT("Invalidate"));
 	Proxy.Reset();
 	ClearIdTracking(InNumInstances);
 	// When the proxy is being replaced, we want to invalidate the owning proxy thing as the continuity is lost.
@@ -1209,3 +1224,20 @@ SIZE_T FPrimitiveInstanceDataManager::GetAllocatedSize() const
 		BakedLightingDataChangedInstances.GetAllocatedSize() +
 		IndexChangeInstances.GetAllocatedSize();
 }
+
+void FPrimitiveInstanceDataManager::OnRegister(int32 InNumInstances)
+{
+	LOG_INST_DATA(TEXT("OnRegister(InNumInstances : %d) NumInstances: %d"), InNumInstances, NumInstances);
+
+	if (CVarInstanceDataResetTrackingOnRegister.GetValueOnGameThread())
+	{
+		if (InNumInstances != NumInstances)
+		{
+			ClearIdTracking(InNumInstances);
+		}
+		else
+		{
+			ClearChangeTracking();
+		}
+	}
+}
diff --git a/Engine/Source/Runtime/Engine/Public/InstancedStaticMesh/ISMInstanceDataManager.h b/Engine/Source/Runtime/Engine/Public/InstancedStaticMesh/ISMInstanceDataManager.h
--- a/Engine/Source/Runtime/Engine/Public/InstancedStaticMesh/ISMInstanceDataManager.h
+++ b/Engine/Source/Runtime/Engine/Public/InstancedStaticMesh/ISMInstanceDataManager.h
@@ -149,6 +149,13 @@ public:
 
 	inline void ResetComponentDirtyTracking() { bComponentMarkedDirty = false; }
 
+	/**
+	 * Called by the corresponding function in the owner UPrimitiveComponent
+	 * Because of the multifarious ways the engine shoves data into the properties it is possible for the count to get out of sync. Also at this point we may assume that we have no idea of the state of individual members.
+	 * Thus, this function will reset tracking state to force a full update as well as, conditionally - if the counts are mismatched - reset the ID mapping.
+	 */
+	void OnRegister(int32 InNumInstances);
+
 private:
 	template <typename TaskLambdaType>
 	static void BeginUpdateTask(FInstanceDataUpdateTaskInfo &InstanceDataUpdateTaskInfo, TaskLambdaType &&TaskLambda, const FInstanceDataBufferHeader &InInstanceDataBufferHeader);
