Add CVar "a.AnimNode.DeadBlending.Enable" to disabled dead blending

diff --git a/Engine/Source/Runtime/Engine/Private/Animation/AnimNode_DeadBlending.cpp b/Engine/Source/Runtime/Engine/Private/Animation/AnimNode_DeadBlending.cpp
--- a/Engine/Source/Runtime/Engine/Private/Animation/AnimNode_DeadBlending.cpp
+++ b/Engine/Source/Runtime/Engine/Private/Animation/AnimNode_DeadBlending.cpp
@@ -16,6 +16,8 @@ LLM_DEFINE_TAG(Animation_DeadBlending);
 
 #define LOCTEXT_NAMESPACE "AnimNode_DeadBlending"
 
+TAutoConsoleVariable<int32> CVarAnimDeadBlendingEnable(TEXT("a.AnimNode.DeadBlending.Enable"), 1, TEXT("Enable / Disable DeadBlending"));
+
 namespace UE::Anim {
 
 	// Inertialization request event bound to a node
@@ -699,6 +701,24 @@ void FAnimNode_DeadBlending::Evaluate_AnyThread(FPoseContext& Output)
 
 	Source.Evaluate(Output);
 
+	// Disable inertialization if requested (for testing / debugging)
+	if (!CVarAnimDeadBlendingEnable.GetValueOnAnyThread())
+	{
+		// Clear any pending inertialization requests
+		RequestQueue.Reset();
+
+		// Clear the inertialization state
+		Deactivate();
+
+		// Clear the pose history
+		PoseSnapshots.Reset();
+
+		// Reset the cached time accumulator
+		DeltaTime = 0.0f;
+
+		return;
+	}
+
 	// Automatically detect teleports... note that we do the teleport distance check against the root bone's location (world space) rather
 	// than the mesh component's location because we still want to inertialize instances where the skeletal mesh component has been moved
 	// while simultaneously counter-moving the root bone (as is the case when mounting and dismounting vehicles for example)
