LevelInstance: Distinguish between IsInLevelInstance/IsInEditLevelInstance/IsInEditLevelInstanceHierarchy
- IsInLevelInstance should always be true (even if Level Instance is being edited)
- IsInEditLevelInstance should only be true if actor is directly in a edited level instance
- IsInEditLevelInstanceHierarchy should be true if actor is in an edited level instance or any parent level instance is currently being edited (used for PostProcess)
- Deprecate IsInEditingLevelInstance virtual method in favor of non virtual IsInEditLevelInstanceHierarchy

#rnx

diff --git a/Engine/Source/Runtime/Engine/Classes/GameFramework/Actor.h b/Engine/Source/Runtime/Engine/Classes/GameFramework/Actor.h
--- a/Engine/Source/Runtime/Engine/Classes/GameFramework/Actor.h
+++ b/Engine/Source/Runtime/Engine/Classes/GameFramework/Actor.h
@@ -273,12 +273,15 @@ private:
 	uint8 bForceNetAddressable:1;
 
 #if WITH_EDITORONLY_DATA
-	/**
-	 * Whether this actor belongs to a level instance which is currently being edited.
-	 */
+	/** Whether this actor belongs to a level instance which is currently being edited. */
 	UPROPERTY(Transient)
-	uint8 bIsInEditingLevelInstance:1;
+	uint8 bIsInEditLevelInstance:1;
 
+	/** Whether this actor belongs to a level instance in a level instance hierarchy currently being edited. Itself or its parent level instances. */
+	UPROPERTY(Transient)
+	uint8 bIsInEditLevelInstanceHierarchy:1;
+
+	/** Whether this actor belongs to a level instance  */
 	UPROPERTY(Transient)
 	uint8 bIsInLevelInstance:1;
 
@@ -310,16 +313,29 @@ public:
 	}
 
 #if WITH_EDITOR
-	/** If true, the actor belongs to a level instance which is currently being edited. */
+	/** Deprecated for a non virtual version. */
+	UE_DEPRECATED(5.4, "Call IsInEditLevelInstanceHierarchy/IsInEditLevelInstance instead.")
 	virtual bool IsInEditingLevelInstance() const
 	{
-		return bIsInEditingLevelInstance;
+		return bIsInEditLevelInstance;
+	}
+
+	/** If true, the actor belongs to a level instance which is currently being edited */
+	bool IsInEditLevelInstance() const
+	{
+		return bIsInEditLevelInstance;
+	}
+
+	/** If true, the actor belongs to a level instance which is currently being edited or a parent level instance being edited. */
+	bool IsInEditLevelInstanceHierarchy() const
+	{
+		return bIsInEditLevelInstanceHierarchy;
 	}
 
 	/** If true, the actor belongs to a level instance. */
-	virtual bool IsInLevelInstance() const
+	bool IsInLevelInstance() const
 	{
-		return bIsInEditingLevelInstance || bIsInLevelInstance;
+		return bIsInLevelInstance;
 	}
 #endif
 
@@ -4468,12 +4484,14 @@ private:
 struct FSetActorIsInLevelInstance
 {
 private:
-	FSetActorIsInLevelInstance(AActor* InActor)
+	FSetActorIsInLevelInstance(AActor* InActor, bool bIsEditing = false)
 	{
 		InActor->bIsInLevelInstance = true;
+		InActor->bIsInEditLevelInstance = bIsEditing;
 	}
 
 	friend class ULevelStreamingLevelInstance;
+	friend class ULevelStreamingLevelInstanceEditor;
 };
 #endif
 
diff --git a/Engine/Source/Runtime/Engine/Private/Actor.cpp b/Engine/Source/Runtime/Engine/Private/Actor.cpp
--- a/Engine/Source/Runtime/Engine/Private/Actor.cpp
+++ b/Engine/Source/Runtime/Engine/Private/Actor.cpp
@@ -189,7 +189,8 @@ void AActor::InitializeDefaults()
 	bHasRegisteredAllComponents = false;
 
 #if WITH_EDITORONLY_DATA
-	bIsInEditingLevelInstance = false;
+	bIsInEditLevelInstanceHierarchy = false;
+	bIsInEditLevelInstance = false;
 	bIsInLevelInstance = false;
 	PivotOffset = FVector::ZeroVector;
 #endif
@@ -5228,17 +5229,13 @@ void AActor::PushLevelInstanceEditingStateToProxies(bool bInEditingState)
 	TInlineComponentArray<UPrimitiveComponent*> PrimComponents;
 	GetComponents(PrimComponents);
 
-
-	bIsInEditingLevelInstance = bInEditingState;
-
-	// We call the virtual method here to allow subclasses to extend the logic for determining level instance state.
-	const bool bPushEditingLevelInstanceState = IsInEditingLevelInstance();
+	bIsInEditLevelInstanceHierarchy = bInEditingState;
 
 	for (const auto& PrimComponent : PrimComponents)
 	{
 		if (PrimComponent->IsRegistered())
 		{
-			PrimComponent->PushLevelInstanceEditingStateToProxy(bPushEditingLevelInstanceState);
+			PrimComponent->PushLevelInstanceEditingStateToProxy(bIsInEditLevelInstanceHierarchy);
 		}
 	}
 
diff --git a/Engine/Source/Runtime/Engine/Private/Components/PrimitiveComponent.cpp b/Engine/Source/Runtime/Engine/Private/Components/PrimitiveComponent.cpp
--- a/Engine/Source/Runtime/Engine/Private/Components/PrimitiveComponent.cpp
+++ b/Engine/Source/Runtime/Engine/Private/Components/PrimitiveComponent.cpp
@@ -1694,7 +1694,7 @@ bool UPrimitiveComponent::GetLevelInstanceEditingState() const
 #if WITH_EDITOR
 	if (const AActor* Owner = GetOwner())
 	{
-		return Owner->IsInEditingLevelInstance();
+		return Owner->IsInEditLevelInstanceHierarchy();
 	}
 #endif
 
diff --git a/Engine/Source/Runtime/Engine/Private/LevelInstance/LevelInstanceEditorLevelStreaming.cpp b/Engine/Source/Runtime/Engine/Private/LevelInstance/LevelInstanceEditorLevelStreaming.cpp
--- a/Engine/Source/Runtime/Engine/Private/LevelInstance/LevelInstanceEditorLevelStreaming.cpp
+++ b/Engine/Source/Runtime/Engine/Private/LevelInstance/LevelInstanceEditorLevelStreaming.cpp
@@ -112,8 +112,6 @@ ULevelStreamingLevelInstanceEditor* ULevelStreamingLevelInstanceEditor::Load(ILe
 
 		if (ULevel* LoadedLevel = LevelStreaming->GetLoadedLevel())
 		{
-			LoadedLevel->OnLoadedActorAddedToLevelEvent.AddUObject(LevelStreaming, &ULevelStreamingLevelInstanceEditor::OnLoadedActorAddedToLevel);
-		
 			// Create special actor that will handle changing the pivot of this level
 			FLevelInstanceEditorPivotHelper::Create(LevelInstance, LevelStreaming);
 		}
@@ -130,22 +128,27 @@ void ULevelStreamingLevelInstanceEditor::Unload(ULevelStreamingLevelInstanceEdit
 	{
 		if (ULevel* LoadedLevel = LevelStreaming->GetLoadedLevel())
 		{
-			LoadedLevel->OnLoadedActorAddedToLevelEvent.RemoveAll(LevelStreaming);
+			LoadedLevel->OnLoadedActorAddedToLevelPreEvent.RemoveAll(LevelStreaming);
 			LevelInstanceSubsystem->RemoveLevelsFromWorld({ LoadedLevel });
 		}
 	}
 }
 
-void ULevelStreamingLevelInstanceEditor::OnLoadedActorAddedToLevel(AActor& InActor)
+void ULevelStreamingLevelInstanceEditor::OnLoadedActorsAddedToLevelPreEvent(const TArray<AActor*>& InActors)
 {
-	OnLevelActorAdded(&InActor);
+	for (AActor* Actor : InActors)
+	{
+		OnLevelActorAdded(Actor);
+	}
 }
 
 void ULevelStreamingLevelInstanceEditor::OnLevelActorAdded(AActor* InActor)
 {
 	if (InActor && InActor->GetLevel() == LoadedLevel)
 	{
-		InActor->PushLevelInstanceEditingStateToProxies(true);
+		const bool bIsEditing = true;
+		FSetActorIsInLevelInstance SetIsInLevelInstance(InActor, bIsEditing);
+		InActor->PushLevelInstanceEditingStateToProxies(bIsEditing);
 	}
 }
 
@@ -157,6 +160,9 @@ void ULevelStreamingLevelInstanceEditor::OnLevelLoadedChanged(ULevel* InLevel)
 	{
 		check(InLevel == NewLoadedLevel);
 
+		OnLoadedActorsAddedToLevelPreEvent(NewLoadedLevel->Actors);
+		NewLoadedLevel->OnLoadedActorAddedToLevelPreEvent.AddUObject(this, &ULevelStreamingLevelInstanceEditor::OnLoadedActorsAddedToLevelPreEvent);
+
 		// Avoid prompts for Level Instance editing
 		NewLoadedLevel->bPromptWhenAddingToLevelBeforeCheckout = false;
 		NewLoadedLevel->bPromptWhenAddingToLevelOutsideBounds = false;
diff --git a/Engine/Source/Runtime/Engine/Private/LevelInstance/LevelInstanceLevelStreaming.cpp b/Engine/Source/Runtime/Engine/Private/LevelInstance/LevelInstanceLevelStreaming.cpp
--- a/Engine/Source/Runtime/Engine/Private/LevelInstance/LevelInstanceLevelStreaming.cpp
+++ b/Engine/Source/Runtime/Engine/Private/LevelInstance/LevelInstanceLevelStreaming.cpp
@@ -160,7 +160,7 @@ void ULevelStreamingLevelInstance::OnLoadedActorsAddedToLevelPostEvent(const TAr
 						Actor->PushSelectionToProxies();
 						if (LevelInstance)
 						{
-							Actor->PushLevelInstanceEditingStateToProxies(CastChecked<AActor>(LevelInstance)->IsInEditingLevelInstance());
+							Actor->PushLevelInstanceEditingStateToProxies(CastChecked<AActor>(LevelInstance)->IsInEditLevelInstanceHierarchy());
 						}
 
 						if (LevelInstanceEditorInstanceActor.IsValid())
@@ -293,7 +293,7 @@ ULevelStreamingLevelInstance* ULevelStreamingLevelInstance::LoadInstance(ILevelI
 				LevelStreaming->LevelInstanceEditorInstanceActor = ALevelInstanceEditorInstanceActor::Create(LevelInstance, Level);
 
 				// Push editing state to child actors
-				LevelInstanceActor->PushLevelInstanceEditingStateToProxies(LevelInstanceActor->IsInEditingLevelInstance());
+				LevelInstanceActor->PushLevelInstanceEditingStateToProxies(LevelInstanceActor->IsInEditLevelInstanceHierarchy());
 			}
 			else
 			{
diff --git a/Engine/Source/Runtime/Engine/Public/LevelInstance/LevelInstanceEditorLevelStreaming.h b/Engine/Source/Runtime/Engine/Public/LevelInstance/LevelInstanceEditorLevelStreaming.h
--- a/Engine/Source/Runtime/Engine/Public/LevelInstance/LevelInstanceEditorLevelStreaming.h
+++ b/Engine/Source/Runtime/Engine/Public/LevelInstance/LevelInstanceEditorLevelStreaming.h
@@ -26,7 +26,7 @@ public:
 	virtual TOptional<FFolder::FRootObject> GetFolderRootObject() const override;
 protected:
 	void OnLevelActorAdded(AActor* InActor);
-	void OnLoadedActorAddedToLevel(AActor& InActor);
+	void OnLoadedActorsAddedToLevelPreEvent(const TArray<AActor*>& InActors);
 
 	friend class ULevelInstanceSubsystem;
 
diff --git a/Engine/Source/Runtime/Landscape/Private/LandscapeEdit.cpp b/Engine/Source/Runtime/Landscape/Private/LandscapeEdit.cpp
--- a/Engine/Source/Runtime/Landscape/Private/LandscapeEdit.cpp
+++ b/Engine/Source/Runtime/Landscape/Private/LandscapeEdit.cpp
@@ -2569,7 +2569,7 @@ bool ULandscapeInfo::SupportsLandscapeEditing() const
 {
 	// Don't let landscapes from level instances be edited : they can only be edited in their source level (note that technically, the IsEditing test is not necessary as they cannot be edited 
 	//  in level instance mode, since it's mutually exclusive with landscape mode, but let's keep it for describing the intention here) :
-	if (LandscapeActor.IsValid() && LandscapeActor->IsInLevelInstance() && !LandscapeActor->IsInEditingLevelInstance())
+	if (LandscapeActor.IsValid() && LandscapeActor->IsInLevelInstance() && !LandscapeActor->IsInEditLevelInstance())
 	{
 		return false;
 	}
