[NetPhysics] Replication stuttering fixes, even out the target state when we receive replicated state at unexpected intervals.

Extrapolate the target state used in predictive interplation both forward and backwards from the replicated state to make sure we don't leave gaps when receiving a state that is 4 ticks ahead 2 ticks after the last or vice versa.

diff --git a/Engine/Source/Runtime/Engine/Private/PhysicsEngine/Experimental/PhysScene_Chaos.cpp b/Engine/Source/Runtime/Engine/Private/PhysicsEngine/Experimental/PhysScene_Chaos.cpp
--- a/Engine/Source/Runtime/Engine/Private/PhysicsEngine/Experimental/PhysScene_Chaos.cpp
+++ b/Engine/Source/Runtime/Engine/Private/PhysicsEngine/Experimental/PhysScene_Chaos.cpp
@@ -1320,7 +1320,7 @@ const FRigidBodyState* FPhysScene_Chaos::GetStateFromReplicationCache(UPrimitive
 	if (!GetSolver()->GetRewindCallback())
 	{
 		// We only populate replication cache through the RewindCallback
-		ServerFrame = 0;
+		ServerFrame = GetSolver()->GetCurrentFrame();
 		return nullptr;
 	}
 
diff --git a/Engine/Source/Runtime/Engine/Private/PhysicsEngine/PhysicsReplication.cpp b/Engine/Source/Runtime/Engine/Private/PhysicsEngine/PhysicsReplication.cpp
--- a/Engine/Source/Runtime/Engine/Private/PhysicsEngine/PhysicsReplication.cpp
+++ b/Engine/Source/Runtime/Engine/Private/PhysicsEngine/PhysicsReplication.cpp
@@ -122,8 +122,11 @@ namespace PhysicsReplicationCVars
 		static float RotCorrectionTimeMultiplier = 1.0f;
 		static FAutoConsoleVariableRef CVarRotCorrectionTimeMultiplier(TEXT("np2.PredictiveInterpolation.RotCorrectionTimeMultiplier"), RotCorrectionTimeMultiplier, TEXT("Multiplier to adjust how much of RTT (network Round Trip Time) to add to positional offset correction."));
 
-		static float InterpolationTimeMultiplier = 1.25f;
-		static FAutoConsoleVariableRef CVarInterpolationTimeMultiplier(TEXT("np2.PredictiveInterpolation.InterpolationTimeMultiplier"), InterpolationTimeMultiplier, TEXT("Multiplier to adjust the replication interpolation time which is based on the sendrate of replication data from the server."));
+		static float PosInterpolationTimeMultiplier = 1.1f;
+		static FAutoConsoleVariableRef CVarInterpolationTimeMultiplier(TEXT("np2.PredictiveInterpolation.InterpolationTimeMultiplier"), PosInterpolationTimeMultiplier, TEXT("Multiplier to adjust the replication interpolation time which is based on the sendrate of replication data from the server."));
+		
+		static float RotInterpolationTimeMultiplier = 1.25f;
+		static FAutoConsoleVariableRef CVarRotInterpolationTimeMultiplier(TEXT("np2.PredictiveInterpolation.RotInterpolationTimeMultiplier"), RotInterpolationTimeMultiplier, TEXT("Multiplier to adjust the rotational replication interpolation time which is based on the sendrate of replication data from the server."));
 		
 		static float AverageReceiveIntervalSmoothing = 3.0f;
 		static FAutoConsoleVariableRef CVarAverageReceiveIntervalSmoothing(TEXT("np2.PredictiveInterpolation.AverageReceiveIntervalSmoothing"), AverageReceiveIntervalSmoothing, TEXT("Recommended range: 1.0 - 5.0. Higher value makes the average receive interval adjust itself slower, reducing spikes in InterpolationTime."));
@@ -140,8 +143,8 @@ namespace PhysicsReplicationCVars
 		static float ErrorAccumulationSeconds = 3.0f;
 		static FAutoConsoleVariableRef CVarErrorAccumulationSeconds(TEXT("np2.PredictiveInterpolation.ErrorAccumulationSeconds"), ErrorAccumulationSeconds, TEXT("Perform a reposition if replication have not been able to cover the min expected distance towards the target for this amount of time."));
 		
-		static bool DisableErrorVelocityLimits = false;
-		static FAutoConsoleVariableRef CVarDisableErrorVelocityLimits(TEXT("np2.PredictiveInterpolation.DisableErrorVelocityLimits"), DisableErrorVelocityLimits, TEXT("Disable the velocity limit and allow error accumulation at any velocity."));
+		static bool bDisableErrorVelocityLimits = false;
+		static FAutoConsoleVariableRef CVarDisableErrorVelocityLimits(TEXT("np2.PredictiveInterpolation.DisableErrorVelocityLimits"), bDisableErrorVelocityLimits, TEXT("Disable the velocity limit and allow error accumulation at any velocity."));
 		
 		static float ErrorAccLinVelMaxLimit = 50.0f;
 		static FAutoConsoleVariableRef CVarErrorAccLinVelMaxLimit(TEXT("np2.PredictiveInterpolation.ErrorAccLinVelMaxLimit"), ErrorAccLinVelMaxLimit, TEXT("If target velocity is below this limit we check for desync to trigger softsnap and accumulate time to build up to a hardsnap."));
@@ -161,8 +164,8 @@ namespace PhysicsReplicationCVars
 		static float EarlyOutAngle = 1.f;
 		static FAutoConsoleVariableRef CVarEarlyOutAngle(TEXT("np2.PredictiveInterpolation.EarlyOutAngle"), EarlyOutAngle, TEXT("If object is within this rotational angle (in degrees) from the source target, early out from replication and apply sleep if replicated."));
 		
-		static bool PostResimWaitForUpdate = false;
-		static FAutoConsoleVariableRef CVarPostResimWaitForUpdate(TEXT("np2.PredictiveInterpolation.PostResimWaitForUpdate"), PostResimWaitForUpdate, TEXT("After a resimulation, wait for replicated states that correspond to post-resim state before processing replication again."));
+		static bool bPostResimWaitForUpdate = false;
+		static FAutoConsoleVariableRef CVarPostResimWaitForUpdate(TEXT("np2.PredictiveInterpolation.PostResimWaitForUpdate"), bPostResimWaitForUpdate, TEXT("After a resimulation, wait for replicated states that correspond to post-resim state before processing replication again."));
 		
 		static bool bVelocityBased = true;
 		static FAutoConsoleVariableRef CVarVelocityBased(TEXT("np2.PredictiveInterpolation.VelocityBased"), bVelocityBased, TEXT("When true, predictive interpolation replication mode will only apply linear velocity and angular velocity"));
@@ -178,6 +181,12 @@ namespace PhysicsReplicationCVars
 
 		static bool bDontClearTarget = false;
 		static FAutoConsoleVariableRef CVarDontClearTarget(TEXT("np2.PredictiveInterpolation.DontClearTarget"), bDontClearTarget, TEXT("When true, predictive interpolation will not lose track of the last replicated state after coming to rest."));
+		
+		static bool bDrawDebugTargets = false;
+		static FAutoConsoleVariableRef CVarDrawDebugTargets(TEXT("np2.PredictiveInterpolation.DrawDebugTargets"), bDrawDebugTargets, TEXT("Draw target states, color coded by which ServerFrame they originate from, replicated targets are large and extrapolated targets are small. There is a Z offset to the draw calls."));
+		
+		static bool bDrawDebugVectors = false;
+		static FAutoConsoleVariableRef CVarDrawDebugVectors(TEXT("np2.PredictiveInterpolation.DrawDebugVectors"), bDrawDebugVectors, TEXT("Draw replication vectors, target velocity, replicated velocity, velocity change between replication calls etc."));
 	}
 
 }
@@ -742,7 +751,7 @@ void FPhysicsReplicationAsync::OnPreSimulate_Internal()
 		if (bRewindDataExist && RewindData->IsResim())
 		{
 			// TODO, Handle the transition from post-resim to interpolation better (disabled by default, resim vs replication interaction is handled via FPhysicsReplicationAsync::CacheResimInteractions)
-			if (PhysicsReplicationCVars::PredictiveInterpolationCVars::PostResimWaitForUpdate && RewindData->IsFinalResim())
+			if (PhysicsReplicationCVars::PredictiveInterpolationCVars::bPostResimWaitForUpdate && RewindData->IsFinalResim())
 			{
 				for (auto Itr = ObjectToTarget.CreateIterator(); Itr; ++Itr)
 				{
@@ -831,7 +840,8 @@ void FPhysicsReplicationAsync::UpdateAsyncTarget(const FPhysicsRepAsyncInputData
 	}
 
 	FReplicatedPhysicsTargetAsync* Target = ObjectToTarget.Find(Input.PhysicsObject);
-	if (Target == nullptr)
+	bool bFirstTarget = Target == nullptr;
+	if (bFirstTarget)
 	{
 		// First time we add a target, set it's previous and correction
 		// positions to the target position to avoid math with uninitialized
@@ -847,6 +857,8 @@ void FPhysicsReplicationAsync::UpdateAsyncTarget(const FPhysicsRepAsyncInputData
 	{
 		const int32 PrevTickCount = Target->TickCount;
 		const int32 PrevReceiveInterval = Target->ReceiveInterval;
+		const int32 SendInterval = Input.ServerFrame - Target->ServerFrame;
+		const int32 AverageReceiveInterval = FMath::CeilToInt(Target->AverageReceiveInterval);
 
 		Target->PrevServerFrame = Target->bWaiting ? Input.ServerFrame : Target->ServerFrame;
 		Target->ServerFrame = Target->bWaiting ? Target->ServerFrame : Input.ServerFrame;
@@ -856,19 +868,26 @@ void FPhysicsReplicationAsync::UpdateAsyncTarget(const FPhysicsRepAsyncInputData
 		Target->TargetState = Input.TargetState;
 		Target->RepMode = Input.RepMode;
 		Target->FrameOffset = Input.FrameOffset;
-		Target->TickCount = 0;
+		Target->TickCount = bFirstTarget ? 0 : FMath::Clamp((PrevTickCount - (SendInterval > 0 ? SendInterval : PrevReceiveInterval)), -AverageReceiveInterval, AverageReceiveInterval);
 
 		if (Input.RepMode == EPhysicsReplicationMode::PredictiveInterpolation)
 		{
+#if !(UE_BUILD_SHIPPING || UE_BUILD_TEST)
+			if (PhysicsReplicationCVars::PredictiveInterpolationCVars::bDrawDebugTargets)
+			{
+				const FVector Offset = FVector(0.0f, 0.0f, 50.0f);
+				Chaos::FDebugDrawQueue::GetInstance().DrawDebugBox(Input.TargetState.Position + Offset, FVector(15.0f, 15.0f, 15.0f), Input.TargetState.Quaternion, FColor::MakeRandomSeededColor(Input.ServerFrame), false, CharacterMovementCVars::NetCorrectionLifetime, 0, 1.0f);
+			}
+#endif
+
 			// Cache the position we received this target at, Predictive Interpolation will alter the target state but use this as the source position for reconciliation.
 			Target->PrevPosTarget = Input.TargetState.Position;
 			Target->PrevRotTarget = Input.TargetState.Quaternion;
 
 			// If we extrapolated the previous target past the receive interval, extrapolate this target by the overshoot
-			const int32 OvershootingExtrapolation = FMath::Clamp((PrevTickCount - PrevReceiveInterval), -1, FMath::CeilToInt(Target->AverageReceiveInterval));
-			if (!Target->bWaiting && OvershootingExtrapolation > 0)
+			if (!Target->bWaiting)
 			{
-				ExtrapolateTarget(*Target, OvershootingExtrapolation, GetDeltaTime_Internal());
+				ExtrapolateTarget(*Target, Target->TickCount, GetDeltaTime_Internal());
 			}
 		}
 	}
@@ -1305,9 +1324,9 @@ bool FPhysicsReplicationAsync::PredictiveInterpolation(Chaos::FPBDRigidParticleH
 		return bClearTarget;
 	};
 
-	// Get the distance from the current position to the source position of our target state
-	const float SourceDistanceSqr = (Target.PrevPosTarget - Handle->X()).SizeSquared();
-	const bool bXCanEarlyOut = SourceDistanceSqr < PhysicsReplicationCVars::PredictiveInterpolationCVars::EarlyOutDistanceSqr;
+	// If target velocity is low enough, check the distance from the current position to the source position of our target to see if it's low enough to early out of replication
+	const bool bXCanEarlyOut = Target.TargetState.LinVel.SizeSquared() < UE_KINDA_SMALL_NUMBER &&
+		(Target.PrevPosTarget - Handle->X()).SizeSquared() < PhysicsReplicationCVars::PredictiveInterpolationCVars::EarlyOutDistanceSqr;
 
 	// Early out if we are within range of target, also apply target sleep state
 	if (bXCanEarlyOut)
@@ -1349,7 +1368,7 @@ bool FPhysicsReplicationAsync::PredictiveInterpolation(Chaos::FPBDRigidParticleH
 	* SoftSnap is performed each tick while there is a registered error, if enough time pass HardSnap forces the object into the correct state. */
 	bool bSoftSnap = !PhysicsReplicationCVars::PredictiveInterpolationCVars::bVelocityBased;
 
-	if ( PhysicsReplicationCVars::PredictiveInterpolationCVars::DisableErrorVelocityLimits ||
+	if ( PhysicsReplicationCVars::PredictiveInterpolationCVars::bDisableErrorVelocityLimits ||
 		(TargetLinVel.Size() < PhysicsReplicationCVars::PredictiveInterpolationCVars::ErrorAccLinVelMaxLimit && TargetAngVel.Size() < PhysicsReplicationCVars::PredictiveInterpolationCVars::ErrorAccAngVelMaxLimit))
 	{
 		const FVector PrevDiff = CurrentState.Position - Target.PrevPos;
@@ -1408,10 +1427,8 @@ bool FPhysicsReplicationAsync::PredictiveInterpolation(Chaos::FPBDRigidParticleH
 		const float PosCorrectionTime = FMath::Max(PhysicsReplicationCVars::PredictiveInterpolationCVars::PosCorrectionTimeBase + RTT * PhysicsReplicationCVars::PredictiveInterpolationCVars::PosCorrectionTimeMultiplier, DeltaSeconds + PhysicsReplicationCVars::PredictiveInterpolationCVars::PosCorrectionTimeMin);
 		const float RotCorrectionTime = FMath::Max(PhysicsReplicationCVars::PredictiveInterpolationCVars::RotCorrectionTimeBase + RTT * PhysicsReplicationCVars::PredictiveInterpolationCVars::RotCorrectionTimeMultiplier, DeltaSeconds + PhysicsReplicationCVars::PredictiveInterpolationCVars::RotCorrectionTimeMin);
 
-		/* Temp disabled until angular replication has implemented InterpolationTime
 		// Calculate interpolation time based on current average receive rate
-		const float InterpolationTime = Target.AverageReceiveInterval * DeltaSeconds * PhysicsReplicationCVars::PredictiveInterpolationCVars::InterpolationTimeMultiplier;
-		*/
+		const float InterpolationTime = Target.AverageReceiveInterval * DeltaSeconds * PhysicsReplicationCVars::PredictiveInterpolationCVars::PosInterpolationTimeMultiplier;
 
 		if (!bXCanEarlyOut)
 		{	// --- Velocity Replication ---
@@ -1429,11 +1446,28 @@ bool FPhysicsReplicationAsync::PredictiveInterpolation(Chaos::FPBDRigidParticleH
 			const FVector BlendedTargetVelocity = LinVelDiff + PosDiffVelocity;
 
 			// Add BlendedTargetVelocity onto current velocity
-			const float BlendStepAmount = 1.0f; /* FMath::Clamp(DeltaSeconds / InterpolationTime, 0.0f, 1.0f); Temp disabled until angular replication has implemented InterpolationTime */
+			const float BlendStepAmount = FMath::Clamp(DeltaSeconds / InterpolationTime, 0.0f, 1.0f);
 			const FVector RepLinVel = CurrentState.LinVel + (BlendedTargetVelocity * BlendStepAmount);
 			
 			Handle->SetV(RepLinVel);
 
+#if !(UE_BUILD_SHIPPING || UE_BUILD_TEST)
+			if (PhysicsReplicationCVars::PredictiveInterpolationCVars::bDrawDebugVectors)
+			{
+				const FVector Offset = FVector(0.0f, 0.0f, 50.0f);
+				const FVector OffsetAdd = FVector(0.0f, 0.0f, 10.0f);
+				const FVector StartPos = TargetPos + Offset;
+				FVector Direction = TargetLinVel;
+				Direction.Normalize();
+				Chaos::FDebugDrawQueue::GetInstance().DrawDebugDirectionalArrow((StartPos + OffsetAdd * 0), (StartPos + OffsetAdd * 0) + TargetLinVel * 0.5f, 5.0f, FColor::Green, false, -1.0f, 0, 2.0f);
+				Chaos::FDebugDrawQueue::GetInstance().DrawDebugDirectionalArrow((StartPos + OffsetAdd * 1), (StartPos + OffsetAdd * 1) + CurrentState.LinVel * 0.5f, 5.0f, FColor::Blue, false, -1.0f, 0, 2.0f);
+				Chaos::FDebugDrawQueue::GetInstance().DrawDebugDirectionalArrow((StartPos + OffsetAdd * 2), (StartPos + OffsetAdd * 2) + (Target.PrevLinVel - CurrentState.LinVel) * 0.5f, 5.0f, FColor::Red, false, -1.0f, 0, 2.0f);
+				Chaos::FDebugDrawQueue::GetInstance().DrawDebugDirectionalArrow((StartPos + OffsetAdd * 3), (StartPos + OffsetAdd * 3) + RepLinVel * 0.5f, 5.0f, FColor::Magenta, false, -1.0f, 0, 2.0f);
+				Chaos::FDebugDrawQueue::GetInstance().DrawDebugDirectionalArrow((StartPos + OffsetAdd * 4), (StartPos + OffsetAdd * 4) + (Target.PrevLinVel - RepLinVel) * 0.5f, 5.0f, FColor::Orange, false, -1.0f, 0, 2.0f);
+				Chaos::FDebugDrawQueue::GetInstance().DrawDebugDirectionalArrow((StartPos + OffsetAdd * 5), (StartPos + OffsetAdd * 5) + Direction * RTT, 5.0f, FColor::White, false, -1.0f, 0, 2.0f);
+				Chaos::FDebugDrawQueue::GetInstance().DrawDebugDirectionalArrow((StartPos + OffsetAdd * 6), (StartPos + OffsetAdd * 6) + Direction * InterpolationTime, 5.0f, FColor::Yellow, false, -1.0f, 0, 2.0f);
+			}
+#endif
 			// Cache data for next replication
 			Target.PrevLinVel = FVector(RepLinVel);
 		}
@@ -1460,7 +1494,7 @@ bool FPhysicsReplicationAsync::PredictiveInterpolation(Chaos::FPBDRigidParticleH
 			float WAngle;
 			FVector WAxis;
 			TargetRotDelta.ToAxisAndAngle(WAxis, WAngle);
-			const FVector TargetRotDeltaBlend = FVector(WAxis * (WAngle / (DeltaSeconds * PhysicsReplicationCVars::PredictiveInterpolationCVars::InterpolationTimeMultiplier)));
+			const FVector TargetRotDeltaBlend = FVector(WAxis * (WAngle / (DeltaSeconds * PhysicsReplicationCVars::PredictiveInterpolationCVars::RotInterpolationTimeMultiplier)));
 			const FVector RepAngVel = FMath::DegreesToRadians(TargetAngVel) + TargetRotDeltaBlend;
 
 			Handle->SetW(RepAngVel);
@@ -1469,6 +1503,15 @@ bool FPhysicsReplicationAsync::PredictiveInterpolation(Chaos::FPBDRigidParticleH
 		// Cache data for next replication
 		Target.PrevPos = FVector(CurrentState.Position);
 
+#if !(UE_BUILD_SHIPPING || UE_BUILD_TEST)
+		if (PhysicsReplicationCVars::PredictiveInterpolationCVars::bDrawDebugTargets)
+		{
+			const FVector Offset = FVector(0.0f, 0.0f, 50.0f);
+			const FVector StartPos = TargetPos + Offset;
+			Chaos::FDebugDrawQueue::GetInstance().DrawDebugBox(StartPos, FVector(5.0f + Target.TickCount * 0.75f, 5.0f + Target.TickCount * 0.75f, 5.0f + Target.TickCount * 0.75f), Target.TargetState.Quaternion, FColor::MakeRandomSeededColor(Target.ServerFrame), false, CharacterMovementCVars::NetCorrectionLifetime, 0, 1.0f);
+		}
+#endif
+
 		// --- Target Extrapolation ---
 		if (Target.TickCount <= FMath::CeilToInt(Target.ReceiveInterval * PhysicsReplicationCVars::PredictiveInterpolationCVars::ExtrapolationTimeMultiplier))
 		{
