Support not stripping quality levels when cooking

diff --git a/Engine/Source/Runtime/Engine/Private/PerQualityLevelProperties.cpp b/Engine/Source/Runtime/Engine/Private/PerQualityLevelProperties.cpp
--- a/Engine/Source/Runtime/Engine/Private/PerQualityLevelProperties.cpp
+++ b/Engine/Source/Runtime/Engine/Private/PerQualityLevelProperties.cpp
@@ -292,7 +292,7 @@ template<typename StructType, typename ValueType, EName _BasePropertyName>
 void FPerQualityLevelProperty<StructType, ValueType, _BasePropertyName>::StripQualtiyLevelForCooking(const TCHAR* InPlatformName)
 {
 	StructType* This = StaticCast<StructType*>(this);
-	if (This->PerQuality.Num() > 0)
+	if (This->PerQuality.Num() > 0 && !CVarName.IsEmpty())
 	{
 		FSupportedQualityLevelArray CookQualityLevelInfo = This->GetSupportedQualityLevels(InPlatformName);
 		CookQualityLevelInfo.Sort([&](const int32& A, const int32& B) { return (A > B); });
diff --git a/Engine/Source/Runtime/Engine/Private/SkeletalMesh.cpp b/Engine/Source/Runtime/Engine/Private/SkeletalMesh.cpp
--- a/Engine/Source/Runtime/Engine/Private/SkeletalMesh.cpp
+++ b/Engine/Source/Runtime/Engine/Private/SkeletalMesh.cpp
@@ -382,7 +382,7 @@ USkeletalMesh::USkeletalMesh(const FObjectInitializer& ObjectInitializer)
 #endif
 	SetMinLod(FPerPlatformInt(0));
 	SetQualityLevelMinLod(0);
-	MinQualityLevelLOD.Init(GSkeletalMeshMinLodQualityLevelCVarName, GSkeletalMeshMinLodQualityLevelScalabilitySection);
+	MinQualityLevelLOD.SetQualityLevelCVarForCooking(GSkeletalMeshMinLodQualityLevelCVarName, GSkeletalMeshMinLodQualityLevelScalabilitySection);
 	SetDisableBelowMinLodStripping(FPerPlatformBool(false));
 PRAGMA_DISABLE_DEPRECATION_WARNINGS
 	bSupportRayTracing = true;
diff --git a/Engine/Source/Runtime/Engine/Private/SkeletalMeshLODSettings.cpp b/Engine/Source/Runtime/Engine/Private/SkeletalMeshLODSettings.cpp
--- a/Engine/Source/Runtime/Engine/Private/SkeletalMeshLODSettings.cpp
+++ b/Engine/Source/Runtime/Engine/Private/SkeletalMeshLODSettings.cpp
@@ -24,7 +24,7 @@ USkeletalMeshLODSettings::USkeletalMeshLODSettings(const FObjectInitializer& Obj
 	// TODO: support saving some but not all optional LODs
 	MaxNumOptionalLODs.Default = 0;
 
-	MinQualityLevelLod.Init(GSkeletalMeshMinLodQualityLevelCVarName, GSkeletalMeshMinLodQualityLevelScalabilitySection);
+	MinQualityLevelLod.SetQualityLevelCVarForCooking(GSkeletalMeshMinLodQualityLevelCVarName, GSkeletalMeshMinLodQualityLevelScalabilitySection);
 }
 
 const FSkeletalMeshLODGroupSettings& USkeletalMeshLODSettings::GetSettingsForLODLevel(const int32 LODIndex) const
diff --git a/Engine/Source/Runtime/Engine/Private/StaticMesh.cpp b/Engine/Source/Runtime/Engine/Private/StaticMesh.cpp
--- a/Engine/Source/Runtime/Engine/Private/StaticMesh.cpp
+++ b/Engine/Source/Runtime/Engine/Private/StaticMesh.cpp
@@ -3372,7 +3372,7 @@ UStaticMesh::UStaticMesh(const FObjectInitializer& ObjectInitializer)
 	BuildCacheAutomationTestGuid.Invalidate();
 #endif
 	SetQualityLevelMinLOD(0);
-	MinQualityLevelLOD.Init(GMinLodQualityLevelCVarName, GMinLodQualityLevelScalabilitySection);
+	MinQualityLevelLOD.SetQualityLevelCVarForCooking(GMinLodQualityLevelCVarName, GMinLodQualityLevelScalabilitySection);
 }
 
 // We don't care if the default implementation of the destructor is cleaning up 
diff --git a/Engine/Source/Runtime/Engine/Private/StreamableRenderAsset.cpp b/Engine/Source/Runtime/Engine/Private/StreamableRenderAsset.cpp
--- a/Engine/Source/Runtime/Engine/Private/StreamableRenderAsset.cpp
+++ b/Engine/Source/Runtime/Engine/Private/StreamableRenderAsset.cpp
@@ -35,7 +35,7 @@ UStreamableRenderAsset::UStreamableRenderAsset(const FObjectInitializer& ObjectI
 	check(sizeof(FStreamableRenderResourceState) == sizeof(uint64));
 
 	SetNoRefStreamingLODBias(-1);
-	NoRefStreamingLODBias.Init(GNoRefBiasQualityLevelCVarName, GNoRefBiasQualityLevelScalabilitySection);
+	NoRefStreamingLODBias.SetQualityLevelCVarForCooking(GNoRefBiasQualityLevelCVarName, GNoRefBiasQualityLevelScalabilitySection);
 }
 
 UStreamableRenderAsset::~UStreamableRenderAsset() = default;
diff --git a/Engine/Source/Runtime/Engine/Public/PerQualityLevelProperties.h b/Engine/Source/Runtime/Engine/Public/PerQualityLevelProperties.h
--- a/Engine/Source/Runtime/Engine/Public/PerQualityLevelProperties.h
+++ b/Engine/Source/Runtime/Engine/Public/PerQualityLevelProperties.h
@@ -99,7 +99,9 @@ struct FPerQualityLevelProperty
 	void ConvertQualtiyLevelData(TMap<FName, _ValueType>& PlaformData, TMultiMap<FName, FName>& PerPlatformToQualityLevel, _ValueType Default);
 #endif
 
-	void Init(const TCHAR* InCVarName, const TCHAR* InSection)
+	// Set Cvar to be able to scan ini files at cook-time and only have the supported ranges of quality levels relevant to the platform.
+	// Unsupported quality levels will be stripped.
+	void SetQualityLevelCVarForCooking(const TCHAR* InCVarName, const TCHAR* InSection)
 	{
 #if WITH_EDITOR
 		ScalabilitySection = FString(InSection);
@@ -107,6 +109,12 @@ struct FPerQualityLevelProperty
 		CVarName = FString(InCVarName);
 	}
 
+	UE_DEPRECATED(5.4, "If no cvar is associated with the property, all quality levels will be keept when cooking. Call SetQualtiyLevelCVarForCooking to strip unsupported quality levels when cooking")
+	void Init(const TCHAR* InCVarName, const TCHAR* InSection)
+	{
+		SetQualityLevelCVarForCooking(InCVarName, InSection);
+	}
+
 	_ValueType GetDefault() const
 	{
 		const _StructType* This = StaticCast<const _StructType*>(this);
diff --git a/Engine/Source/Runtime/Landscape/Private/LandscapeGrass.cpp b/Engine/Source/Runtime/Landscape/Private/LandscapeGrass.cpp
--- a/Engine/Source/Runtime/Landscape/Private/LandscapeGrass.cpp
+++ b/Engine/Source/Runtime/Landscape/Private/LandscapeGrass.cpp
@@ -1290,9 +1290,9 @@ FGrassVariety::FGrassVariety()
 	, InstanceWorldPositionOffsetDisableDistance(0)
 	, ShadowCacheInvalidationBehavior(EShadowCacheInvalidationBehavior::Auto)
 {
-	GrassDensityQuality.Init(GGrassQualityLevelCVarName, GGrassQualityLevelScalabilitySection);
-	StartCullDistanceQuality.Init(GGrassQualityLevelCVarName, GGrassQualityLevelScalabilitySection);
-	EndCullDistanceQuality.Init(GGrassQualityLevelCVarName, GGrassQualityLevelScalabilitySection);
+	GrassDensityQuality.SetQualityLevelCVarForCooking(GGrassQualityLevelCVarName, GGrassQualityLevelScalabilitySection);
+	StartCullDistanceQuality.SetQualityLevelCVarForCooking(GGrassQualityLevelCVarName, GGrassQualityLevelScalabilitySection);
+	EndCullDistanceQuality.SetQualityLevelCVarForCooking(GGrassQualityLevelCVarName, GGrassQualityLevelScalabilitySection);
 }
 
 bool FGrassVariety::IsGrassQualityLevelEnable() const
