Fix for rewind debugger scrubbing crash

diff --git a/Engine/Source/Runtime/Engine/Private/Components/SkeletalMeshComponent.cpp b/Engine/Source/Runtime/Engine/Private/Components/SkeletalMeshComponent.cpp
--- a/Engine/Source/Runtime/Engine/Private/Components/SkeletalMeshComponent.cpp
+++ b/Engine/Source/Runtime/Engine/Private/Components/SkeletalMeshComponent.cpp
@@ -1962,12 +1962,11 @@ void USkeletalMeshComponent::EvaluateAnimation(const USkeletalMesh* InSkeletalMe
 
 void USkeletalMeshComponent::UpdateFollowerComponent()
 {
-	check (LeaderPoseComponent.IsValid());
-
 	ResetMorphTargetCurves();
 
-	if (USkeletalMeshComponent* LeaderSMC = Cast<USkeletalMeshComponent>(LeaderPoseComponent.Get()))
+	if(ensure(LeaderPoseComponent.IsValid()))
 	{
+		USkeletalMeshComponent* LeaderSMC = Cast<USkeletalMeshComponent>(LeaderPoseComponent.Get());
 		// first set any animation-driven curves from the leader SMC
 		if (LeaderSMC->AnimScriptInstance)
 		{
@@ -1991,7 +1990,6 @@ void USkeletalMeshComponent::UpdateFollowerComponent()
 				FAnimationRuntime::AppendActiveMorphTargets(GetSkeletalMeshAsset(), MorphTargetCurves, ActiveMorphTargets, MorphTargetWeights);
 			}
 		}
-
 	}
  
 	Super::UpdateFollowerComponent();
