Actually record destruction infos into a replay if the replay recording time limit is unlimited (a negative number).

diff --git a/Engine/Source/Runtime/Engine/Private/DemoNetDriver.cpp b/Engine/Source/Runtime/Engine/Private/DemoNetDriver.cpp
--- a/Engine/Source/Runtime/Engine/Private/DemoNetDriver.cpp
+++ b/Engine/Source/Runtime/Engine/Private/DemoNetDriver.cpp
@@ -2242,7 +2242,8 @@ bool UDemoNetDriver::ReplicatePrioritizedActor(const FActorPriority& ActorPriori
 	if (bDestructionInfo)
 	{
 		// only process destruction infos if we're below the time limit
-		if (Params.TotalDestructionInfoRecordTime < Params.DestructionInfoTimeLimitSeconds)
+		const bool bHasDestructionInfoTimeBeenExhausted = (Params.DestructionInfoTimeLimitSeconds > 0.0) && (Params.TotalDestructionInfoRecordTime < Params.DestructionInfoTimeLimitSeconds);
+		if (!bHasDestructionInfoTimeBeenExhausted)
 		{
 			++Params.NumDestructionInfosReplicated;
 
