Landscape: PerQuality scalability setting Landscape LOD parameters

diff --git a/Engine/Config/BaseDeviceProfiles.ini b/Engine/Config/BaseDeviceProfiles.ini
--- a/Engine/Config/BaseDeviceProfiles.ini
+++ b/Engine/Config/BaseDeviceProfiles.ini
@@ -300,6 +300,7 @@ bIsVisibleForAssets=True
 +CVars=sg.TextureQuality=3
 +CVars=sg.EffectsQuality=3
 +CVars=sg.FoliageQuality=3
++CVars=sg.LandscapeQuality=2
 ; LightGrid Settings
 +CVars=r.Forward.LightLinkedListCulling=0
 
@@ -807,6 +808,7 @@ BaseProfileName=Android
 +CVars=sg.TextureQuality=0
 +CVars=sg.EffectsQuality=0
 +CVars=sg.FoliageQuality=0
++CVars=sg.LandscapeQuality=0
 
 [Android_Mid DeviceProfile]
 DeviceType=Android
@@ -821,6 +823,7 @@ BaseProfileName=Android
 +CVars=sg.TextureQuality=1
 +CVars=sg.EffectsQuality=1
 +CVars=sg.FoliageQuality=1
++CVars=sg.LandscapeQuality=1
 
 [Android_High DeviceProfile]
 DeviceType=Android
@@ -834,6 +837,7 @@ BaseProfileName=Android
 +CVars=sg.TextureQuality=2
 +CVars=sg.EffectsQuality=2
 +CVars=sg.FoliageQuality=2
++CVars=sg.LandscapeQuality=2
 +CVars=r.MobileContentScaleFactor=1.0
 
 ; Example base settings for Vulkan SM5 on Android
@@ -850,6 +854,7 @@ BaseProfileName=Android
 +CVars=sg.EffectsQuality=2
 +CVars=sg.FoliageQuality=2
 +CVars=sg.ShadingQuality=2
++CVars=sg.LandscapeQuality=2
 +CVars=r.BloomQuality=2
 +CVars=r.LightShaftQuality=1
 ; Shadows
diff --git a/Engine/Source/Editor/UnrealEd/Private/SScalabilitySettings.cpp b/Engine/Source/Editor/UnrealEd/Private/SScalabilitySettings.cpp
--- a/Engine/Source/Editor/UnrealEd/Private/SScalabilitySettings.cpp
+++ b/Engine/Source/Editor/UnrealEd/Private/SScalabilitySettings.cpp
@@ -49,6 +49,7 @@ ECheckBoxState SScalabilitySettings::IsGroupQualityLevelSelected(const TCHAR* In
 	else if (FCString::Strcmp(InGroupName, TEXT("EffectsQuality")) == 0) QualityLevel = CachedQualityLevels.EffectsQuality;
 	else if (FCString::Strcmp(InGroupName, TEXT("FoliageQuality")) == 0) QualityLevel = CachedQualityLevels.FoliageQuality;
 	else if (FCString::Strcmp(InGroupName, TEXT("ShadingQuality")) == 0) QualityLevel = CachedQualityLevels.ShadingQuality;
+ 	else if (FCString::Strcmp(InGroupName, TEXT("LandscapeQuality")) == 0) QualityLevel = CachedQualityLevels.LandscapeQuality;
 
 	return (QualityLevel == InQualityLevel) ? ECheckBoxState::Checked : ECheckBoxState::Unchecked;
 }
@@ -66,6 +67,7 @@ void SScalabilitySettings::OnGroupQualityLevelChanged(ECheckBoxState NewState, c
 	else if (FCString::Strcmp(InGroupName, TEXT("EffectsQuality")) == 0) CachedQualityLevels.EffectsQuality = InQualityLevel;
 	else if (FCString::Strcmp(InGroupName, TEXT("FoliageQuality")) == 0) CachedQualityLevels.FoliageQuality = InQualityLevel;
 	else if (FCString::Strcmp(InGroupName, TEXT("ShadingQuality")) == 0) CachedQualityLevels.ShadingQuality = InQualityLevel;
+	else if (FCString::Strcmp(InGroupName, TEXT("LandscapeQuality")) == 0) CachedQualityLevels.LandscapeQuality = InQualityLevel;
 
 	Scalability::SetQualityLevels(CachedQualityLevels);
 	Scalability::SaveState(GEditorSettingsIni);
@@ -383,8 +385,9 @@ void SScalabilitySettings::Construct( const FArguments& InArgs )
 			FMath::Max(LevelCounts.EffectsQuality,
 			FMath::Max(LevelCounts.FoliageQuality,
 			FMath::Max(LevelCounts.ShadingQuality,
-			FMath::Max(LevelCounts.PostProcessQuality, LevelCounts.AntiAliasingQuality)
-			))))))));
+			FMath::Max(LevelCounts.PostProcessQuality,
+			FMath::Max(LevelCounts.LandscapeQuality,LevelCounts.AntiAliasingQuality)
+			)))))))));
 
 		const int32 TotalWidth = MaxLevelCount + 1;
 		const int32 ScalabilityFirstRowId = 1;
@@ -414,6 +417,7 @@ void SScalabilitySettings::Construct( const FArguments& InArgs )
 			+MakeGridSlot(0,ScalabilityFirstRowId+7,TotalWidth,1) [ SNew (SBorder).BorderImage(FAppStyle::GetBrush("Scalability.RowBackground")) ]
 			+MakeGridSlot(0,ScalabilityFirstRowId+8,TotalWidth,1) [ SNew (SBorder).BorderImage(FAppStyle::GetBrush("Scalability.RowBackground")) ]
 			+MakeGridSlot(0,ScalabilityFirstRowId+9,TotalWidth,1) [ SNew (SBorder).BorderImage(FAppStyle::GetBrush("Scalability.RowBackground")) ]
+			+MakeGridSlot(0,ScalabilityFirstRowId+10,TotalWidth,1) [ SNew (SBorder).BorderImage(FAppStyle::GetBrush("Scalability.RowBackground")) ]
 
 			+MakeGridSlot(0,ScalabilityFirstRowId + 0) [ SNew(STextBlock).Text(LOCTEXT("ViewDistanceLabel1", "View Distance")).Font(GroupFont) ]
 			+MakeGridSlot(0,ScalabilityFirstRowId + 1) [ SNew(STextBlock).Text(FText::Format(LOCTEXT("AntiAliasingQualityLabel1", "Anti-Aliasing ({0})"), AntiAliasingMethodShortName)).Font(GroupFont) ]
@@ -425,6 +429,7 @@ void SScalabilitySettings::Construct( const FArguments& InArgs )
 			+MakeGridSlot(0,ScalabilityFirstRowId + 7) [ SNew(STextBlock).Text(LOCTEXT("EffectsQualityLabel1", "Effects")).Font(GroupFont) ]
 			+MakeGridSlot(0,ScalabilityFirstRowId + 8) [ SNew(STextBlock).Text(LOCTEXT("FoliageQualityLabel1", "Foliage")).Font(GroupFont) ]
 			+MakeGridSlot(0,ScalabilityFirstRowId + 9) [ SNew(STextBlock).Text(LOCTEXT("ShadingQualityLabel1", "Shading")).Font(GroupFont) ]
+			+MakeGridSlot(0,ScalabilityFirstRowId + 10) [ SNew(STextBlock).Text(LOCTEXT("LandscapeQualityLabel1", "Landscape")).Font(GroupFont) ]
 			;
 
 		AddButtonsToGrid(1, ScalabilityFirstRowId + 0, ButtonMatrix, FiveDistanceNames, LevelCounts.ViewDistanceQuality, TEXT("ViewDistanceQuality"), LOCTEXT("ViewDistanceQualityTooltip", "Set view distance to {0}"));
@@ -437,6 +442,7 @@ void SScalabilitySettings::Construct( const FArguments& InArgs )
 		AddButtonsToGrid(1, ScalabilityFirstRowId + 7, ButtonMatrix, FiveNames, LevelCounts.EffectsQuality, TEXT("EffectsQuality"), LOCTEXT("EffectsQualityTooltip", "Set effects quality to {0}"));
 		AddButtonsToGrid(1, ScalabilityFirstRowId + 8, ButtonMatrix, FiveNames, LevelCounts.FoliageQuality, TEXT("FoliageQuality"), LOCTEXT("FoliageQualityTooltip", "Set foliage quality to {0}"));
 		AddButtonsToGrid(1, ScalabilityFirstRowId + 9, ButtonMatrix, FiveNames, LevelCounts.ShadingQuality, TEXT("ShadingQuality"), LOCTEXT("ShadingQualityTooltip", "Set shading quality to {0}"));
+		AddButtonsToGrid(1, ScalabilityFirstRowId + 10, ButtonMatrix, FiveNames, LevelCounts.LandscapeQuality, TEXT("LandscapeQuality"), LOCTEXT("LandscapeQualityTooltip", "Set landscape quality to {0}"));
 	
 		ScalabilityGroupsWidget = ButtonMatrix;
 	}
diff --git a/Engine/Source/Editor/UnrealEd/Private/Settings/EditorSettings.cpp b/Engine/Source/Editor/UnrealEd/Private/Settings/EditorSettings.cpp
--- a/Engine/Source/Editor/UnrealEd/Private/Settings/EditorSettings.cpp
+++ b/Engine/Source/Editor/UnrealEd/Private/Settings/EditorSettings.cpp
@@ -136,6 +136,7 @@ void UEditorSettings::LoadScalabilityBenchmark()
 		GConfig->GetInt(Section, TEXT("EffectsQuality"), Temporary.EffectsQuality, GEditorSettingsIni);
 		GConfig->GetInt(Section, TEXT("FoliageQuality"), Temporary.FoliageQuality, GEditorSettingsIni);
 		GConfig->GetInt(Section, TEXT("ShadingQuality"), Temporary.ShadingQuality, GEditorSettingsIni);
+		GConfig->GetInt(Section, TEXT("LandscapeQuality"), Temporary.LandscapeQuality, GEditorSettingsIni);
 		EngineBenchmarkResult = Temporary;
 	}
 }
@@ -162,6 +163,7 @@ void UEditorSettings::AutoApplyScalabilityBenchmark()
 	GConfig->SetInt(Section, TEXT("EffectsQuality"), Temporary.EffectsQuality, GEditorSettingsIni);
 	GConfig->SetInt(Section, TEXT("FoliageQuality"), Temporary.FoliageQuality, GEditorSettingsIni);
 	GConfig->SetInt(Section, TEXT("ShadingQuality"), Temporary.ShadingQuality, GEditorSettingsIni);
+	GConfig->SetInt(Section, TEXT("LandscapeQuality"), Temporary.LandscapeQuality, GEditorSettingsIni);
 
 	Scalability::SetQualityLevels(Temporary);
 	Scalability::SaveState(GEditorSettingsIni);
diff --git a/Engine/Source/Runtime/Engine/Private/Scalability.cpp b/Engine/Source/Runtime/Engine/Private/Scalability.cpp
--- a/Engine/Source/Runtime/Engine/Private/Scalability.cpp
+++ b/Engine/Source/Runtime/Engine/Private/Scalability.cpp
@@ -107,6 +107,13 @@ static TAutoConsoleVariable<int32> CVarShadingQuality(
 	TEXT(" 0:low, 1:med, 2:high, 3:epic, 4:cinematic, default: 3"),
 	ECVF_ScalabilityGroup | ECVF_Preview);
 
+static TAutoConsoleVariable<int32> CVarLandscapeQuality(
+	TEXT("sg.LandscapeQuality"),
+	Scalability::DefaultQualityLevel,
+	TEXT("Scalability quality state (internally used by scalability system, ini load/save or using SCALABILITY console command)\n")
+	TEXT(" 0:low, 1:med, 2:high, 3:epic, 4:cinematic, default: 3"),
+	ECVF_ScalabilityGroup | ECVF_Preview);
+
 static TAutoConsoleVariable<int32> CVarViewDistanceQuality_NumLevels(
 	TEXT("sg.ViewDistanceQuality.NumLevels"),
 	5,
@@ -177,6 +184,13 @@ static TAutoConsoleVariable<int32> CVarShadingQuality_NumLevels(
 	TEXT(" default: 5 (0..4)"),
 	ECVF_ReadOnly);
 
+static TAutoConsoleVariable<int32> CVarLandscapeQuality_NumLevels(
+	TEXT("sg.LandscapeQuality.NumLevels"),
+	5,
+	TEXT("Number of settings quality levels in sg.LandscapeQuality\n")
+	TEXT(" default: 5 (0..4)"),
+	ECVF_ReadOnly);
+
 namespace Scalability
 {
 static FQualityLevels GScalabilityBackupQualityLevels(false);
@@ -551,6 +565,11 @@ void OnChangeShadingQuality(IConsoleVariable* Var)
 	SetGroupQualityLevel(TEXT("ShadingQuality"), Var->GetInt(), CVarShadingQuality_NumLevels->GetInt());
 }
 
+void OnChangeLandscapeQuality(IConsoleVariable* Var)
+{
+	SetGroupQualityLevel(TEXT("LandscapeQuality"), Var->GetInt(), CVarLandscapeQuality_NumLevels->GetInt());
+}
+	
 void InitScalabilitySystem()
 {
 	// needed only once
@@ -576,6 +595,7 @@ void InitScalabilitySystem()
 	CVarEffectsQuality.AsVariable()->SetOnChangedCallback(FConsoleVariableDelegate::CreateStatic(&OnChangeEffectsQuality));
 	CVarFoliageQuality.AsVariable()->SetOnChangedCallback(FConsoleVariableDelegate::CreateStatic(&OnChangeFoliageQuality));
 	CVarShadingQuality.AsVariable()->SetOnChangedCallback(FConsoleVariableDelegate::CreateStatic(&OnChangeShadingQuality));
+	CVarLandscapeQuality.AsVariable()->SetOnChangedCallback(FConsoleVariableDelegate::CreateStatic(&OnChangeLandscapeQuality));
 
 #if WITH_EDITOR
 	ScalabilityShaderPlatform = GMaxRHIShaderPlatform;
@@ -691,6 +711,7 @@ FQualityLevels BenchmarkQualityLevels(uint32 WorkScale, float CPUMultiplier, flo
 	Results.EffectsQuality = ComputeOptionFromPerfIndex(TEXT("EffectsQuality"), CPUPerfIndex, GPUPerfIndex);
 	Results.FoliageQuality = ComputeOptionFromPerfIndex(TEXT("FoliageQuality"), CPUPerfIndex, GPUPerfIndex);
 	Results.ShadingQuality = ComputeOptionFromPerfIndex(TEXT("ShadingQuality"), CPUPerfIndex, GPUPerfIndex);
+	Results.LandscapeQuality = ComputeOptionFromPerfIndex( TEXT("LandscapeQuality"), CPUPerfIndex, GPUPerfIndex);
 	Results.CPUBenchmarkResults = CPUPerfIndex;
 	Results.GPUBenchmarkResults = GPUPerfIndex;
 
@@ -849,6 +870,7 @@ void SetQualityLevels(const FQualityLevels& QualityLevels, bool bForce/* = false
 	ClampedLevels.SetEffectsQuality(QualityLevels.EffectsQuality);
 	ClampedLevels.SetFoliageQuality(QualityLevels.FoliageQuality);
 	ClampedLevels.SetShadingQuality(QualityLevels.ShadingQuality);
+	ClampedLevels.SetLandscapeQuality(QualityLevels.LandscapeQuality);
 
 #if WITH_EDITOR
 	bForce = bForce || bScalabilityShaderPlatformHasBeenChanged;
@@ -875,6 +897,7 @@ void SetQualityLevels(const FQualityLevels& QualityLevels, bool bForce/* = false
 		SetQualityLevelCVar(CVarEffectsQuality, ClampedLevels.EffectsQuality, GScalabilityQualityLevelsOverride.EffectsQuality, bForce);
 		SetQualityLevelCVar(CVarFoliageQuality, ClampedLevels.FoliageQuality, GScalabilityQualityLevelsOverride.FoliageQuality, bForce);
 		SetQualityLevelCVar(CVarShadingQuality, ClampedLevels.ShadingQuality, GScalabilityQualityLevelsOverride.ShadingQuality, bForce);
+		SetQualityLevelCVar(CVarLandscapeQuality, ClampedLevels.LandscapeQuality, GScalabilityQualityLevelsOverride.LandscapeQuality, bForce);
 
 		OnScalabilitySettingsChanged.Broadcast(ClampedLevels);
 	}
@@ -906,6 +929,7 @@ FQualityLevels GetQualityLevels()
 		Ret.EffectsQuality = CVarEffectsQuality.GetValueOnGameThread();
 		Ret.FoliageQuality = CVarFoliageQuality.GetValueOnGameThread();
 		Ret.ShadingQuality = CVarShadingQuality.GetValueOnGameThread();
+		Ret.LandscapeQuality = CVarLandscapeQuality.GetValueOnGameThread();
 	}
 	else
 	{
@@ -976,6 +1000,7 @@ void FQualityLevels::SetFromSingleQualityLevel(int32 Value)
 	EffectsQuality = FMath::Clamp(Value, 0, CVarEffectsQuality_NumLevels->GetInt() - 1);
 	FoliageQuality = FMath::Clamp(Value, 0, CVarFoliageQuality_NumLevels->GetInt() - 1);
 	ShadingQuality = FMath::Clamp(Value, 0, CVarShadingQuality_NumLevels->GetInt() - 1);
+	LandscapeQuality = FMath::Clamp(Value, 0, CVarLandscapeQuality_NumLevels->GetInt() - 1);
 }
 
 void FQualityLevels::SetFromSingleQualityLevelRelativeToMax(int32 Value)
@@ -995,6 +1020,7 @@ void FQualityLevels::SetFromSingleQualityLevelRelativeToMax(int32 Value)
 	EffectsQuality = FMath::Max(CVarEffectsQuality_NumLevels->GetInt() - Value, 0);
 	FoliageQuality = FMath::Max(CVarFoliageQuality_NumLevels->GetInt() - Value, 0);
 	ShadingQuality = FMath::Max(CVarShadingQuality_NumLevels->GetInt() - Value, 0);
+	LandscapeQuality =  FMath::Max(CVarLandscapeQuality_NumLevels->GetInt() - Value, 0);
 }
 
 // Returns the overall value if all settings are set to the same thing
@@ -1028,6 +1054,7 @@ int32 FQualityLevels::GetMinQualityLevel() const
 	Level = FMath::Min(Level, EffectsQuality);
 	Level = FMath::Min(Level, FoliageQuality);
 	Level = FMath::Min(Level, ShadingQuality);
+	Level = FMath::Min(Level, LandscapeQuality);
 
 	return Level;
 }
@@ -1082,6 +1109,11 @@ void FQualityLevels::SetShadingQuality(int32 Value)
 	ShadingQuality = FMath::Clamp(Value, 0, CVarShadingQuality_NumLevels->GetInt() - 1);
 }
 
+void FQualityLevels::SetLandscapeQuality(int32 Value)
+{
+	LandscapeQuality = FMath::Clamp(Value, 0, CVarLandscapeQuality_NumLevels->GetInt() - 1);
+}
+	
 void LoadState(const FString& IniName)
 {
 	check(!IniName.IsEmpty());
@@ -1106,6 +1138,7 @@ void LoadState(const FString& IniName)
 	GConfig->GetInt(Section, TEXT("sg.EffectsQuality"), State.EffectsQuality, IniName);
 	GConfig->GetInt(Section, TEXT("sg.FoliageQuality"), State.FoliageQuality, IniName);
 	GConfig->GetInt(Section, TEXT("sg.ShadingQuality"), State.ShadingQuality, IniName);
+	GConfig->GetInt(Section, TEXT("sg.LandscapeQuality"), State.LandscapeQuality, IniName);
 
 	// If possible apply immediately, else store in backup so we can re-apply later
 	if (!GScalabilityUsingTemporaryQualityLevels)
@@ -1139,6 +1172,7 @@ void SaveState(const FString& IniName)
 	GConfig->SetInt(Section, TEXT("sg.EffectsQuality"), State.EffectsQuality, IniName);
 	GConfig->SetInt(Section, TEXT("sg.FoliageQuality"), State.FoliageQuality, IniName);
 	GConfig->SetInt(Section, TEXT("sg.ShadingQuality"), State.ShadingQuality, IniName);
+	GConfig->SetInt(Section, TEXT("sg.LandscapeQuality"), State.LandscapeQuality, IniName);
 }
 
 void RecordQualityLevelsAnalytics(bool bAutoApplied)
@@ -1160,6 +1194,7 @@ void RecordQualityLevelsAnalytics(bool bAutoApplied)
 		Attributes.Add(FAnalyticsEventAttribute(TEXT("EffectsQuality"), State.EffectsQuality));
 		Attributes.Add(FAnalyticsEventAttribute(TEXT("FoliageQuality"), State.FoliageQuality));
 		Attributes.Add(FAnalyticsEventAttribute(TEXT("ShadingQuality"), State.ShadingQuality));
+		Attributes.Add( FAnalyticsEventAttribute(TEXT("LandscapeQuality"), State.LandscapeQuality));
 		Attributes.Add(FAnalyticsEventAttribute(TEXT("AutoAppliedSettings"), bAutoApplied));
 		Attributes.Add(FAnalyticsEventAttribute(TEXT("Enterprise"), IProjectManager::Get().IsEnterpriseProject()));
 
@@ -1181,6 +1216,7 @@ FQualityLevels GetQualityLevelCounts()
 	Result.EffectsQuality = CVarEffectsQuality_NumLevels->GetInt();
 	Result.FoliageQuality = CVarFoliageQuality_NumLevels->GetInt();
 	Result.ShadingQuality = CVarShadingQuality_NumLevels->GetInt();
+	Result.LandscapeQuality = CVarLandscapeQuality_NumLevels->GetInt();
 	return Result;
 }
 
diff --git a/Engine/Source/Runtime/Engine/Public/Scalability.h b/Engine/Source/Runtime/Engine/Public/Scalability.h
--- a/Engine/Source/Runtime/Engine/Public/Scalability.h
+++ b/Engine/Source/Runtime/Engine/Public/Scalability.h
@@ -37,6 +37,7 @@ namespace Scalability
 		int32 EffectsQuality;
 		int32 FoliageQuality;
 		int32 ShadingQuality;
+		int32 LandscapeQuality;
 
 		float CPUBenchmarkResults;
 		float GPUBenchmarkResults;
@@ -66,7 +67,8 @@ namespace Scalability
 				TextureQuality == Other.TextureQuality &&
 				EffectsQuality == Other.EffectsQuality &&
 				FoliageQuality == Other.FoliageQuality &&
-				ShadingQuality == Other.ShadingQuality;
+				ShadingQuality == Other.ShadingQuality &&
+				LandscapeQuality == Other.LandscapeQuality;
 		}
 
 		bool operator!=(const FQualityLevels& Other ) const
@@ -87,7 +89,8 @@ namespace Scalability
 				FCrc::TypeCrc32<int32>(TextureQuality) ^
 				FCrc::TypeCrc32<int32>(EffectsQuality) ^
 				FCrc::TypeCrc32<int32>(FoliageQuality) ^
-				FCrc::TypeCrc32<int32>(ShadingQuality);
+				FCrc::TypeCrc32<int32>(ShadingQuality) ^
+				FCrc::TypeCrc32<int32>(LandscapeQuality);
 		}
 
 		// Sets all other settings based on an overall value
@@ -142,10 +145,14 @@ namespace Scalability
 		// @param Value 0:low, 1:medium, 2:high, 3:epic, 4:cinematic (gets clamped if needed)
 		ENGINE_API void SetFoliageQuality(int32 Value);
 
-		// Sets the sharing quality
+		// Sets the shading quality
 		// @param Value 0:low, 1:medium, 2:high, 3:epic, 4:cinematic (gets clamped if needed)
 		ENGINE_API void SetShadingQuality(int32 Value);
 
+		// Sets the landscape quality
+		// @param Value 0:low, 1:medium, 2:high, 3:epic, 4:cinematic (gets clamped if needed)
+		ENGINE_API void SetLandscapeQuality(int32 Value);
+
 		ENGINE_API void SetBenchmarkFallback();
 
 		ENGINE_API void SetDefaults();
diff --git a/Engine/Source/Runtime/Landscape/Classes/LandscapeProxy.h b/Engine/Source/Runtime/Landscape/Classes/LandscapeProxy.h
--- a/Engine/Source/Runtime/Landscape/Classes/LandscapeProxy.h
+++ b/Engine/Source/Runtime/Landscape/Classes/LandscapeProxy.h
@@ -466,19 +466,31 @@ public:
 	float ComponentScreenSizeToUseSubSections;
 
 	/** This is the starting screen size used to calculate the distribution. You can increase the value if you want less LOD0 component, and you use very large landscape component. */
-	UPROPERTY(EditAnywhere, Category = "LOD Distribution", meta = (DisplayName = "LOD 0 Screen Size", ClampMin = "0.1", ClampMax = "10.0", UIMin = "1.0", UIMax = "10.0", LandscapeInherited))
-	float LOD0ScreenSize;
-
+	UPROPERTY(EditAnywhere, Category = "LOD Distribution", meta = (EditCondition = "!bUseScalableLODSettings", DisplayName = "LOD 0 Screen Size", ClampMin = "0.1", ClampMax = "10.0", UIMin = "1.0", UIMax = "10.0", LandscapeInherited))
+	float LOD0ScreenSize = 0.5f;
+	
 	UPROPERTY(EditAnywhere, Category = LOD, AdvancedDisplay, meta = (ToolTip = "Specifies the LOD Group (Zero is No Group). All landscapes in the same group calculate their LOD together, allowing matching border LODs to fix geometry seams.", LandscapeInherited))
 	uint32 LODGroupKey = 0;
 
 	/** The distribution setting used to change the LOD 0 generation, 1.75 is the normal distribution, numbers influence directly the LOD0 proportion on screen. */
-	UPROPERTY(EditAnywhere, Category = "LOD Distribution", meta = (DisplayName = "LOD 0", ClampMin = "1.0", ClampMax = "10.0", UIMin = "1.0", UIMax = "10.0", LandscapeInherited))
-	float LOD0DistributionSetting;
+	UPROPERTY(EditAnywhere, Category = "LOD Distribution", meta = (EditCondition = "!bUseScalableLODSettings", DisplayName = "LOD 0", ClampMin = "1.0", ClampMax = "10.0", UIMin = "1.0", UIMax = "10.0", LandscapeInherited))
+	float LOD0DistributionSetting = 1.25f;
 
 	/** The distribution setting used to change the LOD generation, 2 is the normal distribution, small number mean you want your last LODs to take more screen space and big number mean you want your first LODs to take more screen space. */
-	UPROPERTY(EditAnywhere, Category = "LOD Distribution", meta = (DisplayName = "Other LODs", ClampMin = "1.0", ClampMax = "10.0", UIMin = "1.0", UIMax = "10.0", LandscapeInherited))
-	float LODDistributionSetting;
+	UPROPERTY(EditAnywhere, Category = "LOD Distribution", meta = (EditCondition = "!bUseScalableLODSettings", DisplayName = "Other LODs", ClampMin = "1.0", ClampMax = "10.0", UIMin = "1.0", UIMax = "10.0", LandscapeInherited))
+	float LODDistributionSetting = 3.0f;
+
+	UPROPERTY(EditAnywhere, Category = "LOD Distribution", meta = (EditCondition = "bUseScalableLODSettings", DisplayName = "Scalable LOD 0 Screen Size", ClampMin = "0.1", ClampMax = "10.0", UIMin = "1.0", UIMax = "10.0", LandscapeInherited))
+	FPerQualityLevelFloat ScalableLOD0ScreenSize = 0.5f;
+	
+	UPROPERTY(EditAnywhere, Category = "LOD Distribution", meta = (EditCondition = "bUseScalableLODSettings", DisplayName = "Scalable LOD 0", ClampMin = "1.0", ClampMax = "10.0", UIMin = "1.0", UIMax = "10.0", LandscapeInherited))
+	FPerQualityLevelFloat ScalableLOD0DistributionSetting = 1.25f;
+	
+	UPROPERTY(EditAnywhere, Category = "LOD Distribution", meta = (EditCondition = "bUseScalableLODSettings", DisplayName = "Scalable Other LODs", ClampMin = "1.0", ClampMax = "10.0", UIMin = "1.0", UIMax = "10.0", LandscapeInherited))
+	FPerQualityLevelFloat ScalableLODDistributionSetting = 3.0f;
+
+	UPROPERTY(EditAnywhere, Category = "LOD Distribution", meta = (LandscapeInherited, tooltip="Use ScalableLODDistributionSetting, ScalableLOD0DistributionSetting & ScalableLODDistributionSetting (ignoring r.LandscapeLOD0DistributionScale) instead of LODDistributionSetting, LOD0DistributionSetting & LODDistributionSetting respectively"))
+	bool bUseScalableLODSettings = false;
 
 	/** This controls the area that blends LOD between neighboring sections. At 1.0 it blends across the entire section, and lower numbers reduce the blend region to be closer to the boundary. */
 	UPROPERTY(EditAnywhere, Category = "LOD Distribution", meta = (DisplayName = "Blend Range", ClampMin = "0.01", ClampMax = "1.0", UIMin = "0.01", UIMax = "1.0", LandscapeInherited))
diff --git a/Engine/Source/Runtime/Landscape/Private/Landscape.cpp b/Engine/Source/Runtime/Landscape/Private/Landscape.cpp
--- a/Engine/Source/Runtime/Landscape/Private/Landscape.cpp
+++ b/Engine/Source/Runtime/Landscape/Private/Landscape.cpp
@@ -1454,9 +1454,6 @@ ALandscapeProxy::ALandscapeProxy(const FObjectInitializer& ObjectInitializer)
 	bLockLocation = true;
 #endif // WITH_EDITORONLY_DATA
 	ComponentScreenSizeToUseSubSections = 0.65f;
-	LOD0ScreenSize = 0.5f;
-	LOD0DistributionSetting = 1.25f;
-	LODDistributionSetting = 3.0f;
 	bCastStaticShadow = true;
 	ShadowCacheInvalidationBehavior = EShadowCacheInvalidationBehavior::Auto;
 	bUsedForNavigation = true;
@@ -1642,10 +1639,27 @@ void ALandscapeProxy::UpdateSharedProperties(ULandscapeInfo* InLandscapeInfo)
 
 static TArray<float> GetLODScreenSizeArray(const ALandscapeProxy* InLandscapeProxy, const int32 InNumLODLevels)
 {
+	float LOD0ScreenSize;
+	float LOD0Distribution;
+	if (InLandscapeProxy->bUseScalableLODSettings)
+	{
+		const int32 LandscapeQuality = Scalability::GetQualityLevels().LandscapeQuality;
+		
+		LOD0ScreenSize = InLandscapeProxy->ScalableLOD0ScreenSize.GetValue(LandscapeQuality);
+		LOD0Distribution = InLandscapeProxy->ScalableLOD0DistributionSetting.GetValue(LandscapeQuality);	
+	}
+	else
+	{
+		static IConsoleVariable* CVarLSLOD0DistributionScale = IConsoleManager::Get().FindConsoleVariable(TEXT("r.LandscapeLOD0DistributionScale"));
+		
+		LOD0ScreenSize = InLandscapeProxy->LOD0ScreenSize;
+		LOD0Distribution = InLandscapeProxy->LOD0DistributionSetting * CVarLSLOD0DistributionScale->GetFloat();
+	}
+	
 	static TConsoleVariableData<float>* CVarSMLODDistanceScale = IConsoleManager::Get().FindTConsoleVariableDataFloat(TEXT("r.StaticMeshLODDistanceScale"));
-	static IConsoleVariable* CVarLSLOD0DistributionScale = IConsoleManager::Get().FindConsoleVariable(TEXT("r.LandscapeLOD0DistributionScale"));
-	float CurrentScreenSize = InLandscapeProxy->LOD0ScreenSize / CVarSMLODDistanceScale->GetValueOnGameThread();
-	const float ScreenSizeMult = 1.f / FMath::Max(InLandscapeProxy->LOD0DistributionSetting * CVarLSLOD0DistributionScale->GetFloat(), 1.01f);
+	
+	float CurrentScreenSize = LOD0ScreenSize / CVarSMLODDistanceScale->GetValueOnGameThread();
+	const float ScreenSizeMult = 1.f / FMath::Max(LOD0Distribution , 1.01f);
 
 	TArray<float> Result;
 	Result.Empty(InNumLODLevels);
@@ -3612,6 +3626,7 @@ void ALandscapeProxy::Serialize(FArchive& Ar)
 			}
 		}
 	}
+
 #endif
 }
 
@@ -4005,6 +4020,7 @@ void ALandscapeProxy::PostLoad()
 		}
 		LandscapeMaterialsOverride_DEPRECATED.Reset();
 	}
+	
 	PRAGMA_ENABLE_DEPRECATION_WARNINGS;
 
 	if (GIsEditor)
@@ -4218,6 +4234,9 @@ void ALandscapeProxy::GetSharedProperties(ALandscapeProxy* Landscape)
 		SubsectionSizeQuads = Landscape->SubsectionSizeQuads;
 		MaxLODLevel = Landscape->MaxLODLevel;
 		ComponentScreenSizeToUseSubSections = Landscape->ComponentScreenSizeToUseSubSections;
+		ScalableLODDistributionSetting = Landscape->ScalableLODDistributionSetting;
+		ScalableLOD0DistributionSetting = Landscape->ScalableLOD0DistributionSetting;
+		ScalableLOD0ScreenSize = Landscape->ScalableLOD0ScreenSize;
 		LODDistributionSetting = Landscape->LODDistributionSetting;
 		LOD0DistributionSetting = Landscape->LOD0DistributionSetting;
 		LOD0ScreenSize = Landscape->LOD0ScreenSize;
diff --git a/Engine/Source/Runtime/Landscape/Private/LandscapeEdit.cpp b/Engine/Source/Runtime/Landscape/Private/LandscapeEdit.cpp
--- a/Engine/Source/Runtime/Landscape/Private/LandscapeEdit.cpp
+++ b/Engine/Source/Runtime/Landscape/Private/LandscapeEdit.cpp
@@ -6049,6 +6049,9 @@ bool ALandscapeProxy::CanEditChange(const FProperty* InProperty) const
 			|| (PropertyName == GET_MEMBER_NAME_CHECKED(ALandscapeProxy, LODDistributionSetting))
 			|| (PropertyName == GET_MEMBER_NAME_CHECKED(ALandscapeProxy, LOD0DistributionSetting))
 			|| (PropertyName == GET_MEMBER_NAME_CHECKED(ALandscapeProxy, LOD0ScreenSize))
+			|| (PropertyName == GET_MEMBER_NAME_CHECKED(ALandscapeProxy, ScalableLODDistributionSetting))
+			|| (PropertyName == GET_MEMBER_NAME_CHECKED(ALandscapeProxy, ScalableLOD0DistributionSetting))
+			|| (PropertyName == GET_MEMBER_NAME_CHECKED(ALandscapeProxy, ScalableLOD0ScreenSize))
 			|| (PropertyName == GET_MEMBER_NAME_CHECKED(ALandscapeProxy, LODGroupKey))
 			|| (PropertyName == GET_MEMBER_NAME_CHECKED(ALandscapeProxy, bUseCompressedHeightmapStorage))
 			|| (PropertyName == GET_MEMBER_NAME_CHECKED(ALandscapeProxy, TargetDisplayOrder))
@@ -6167,7 +6170,10 @@ void ALandscapeProxy::PostEditChangeProperty(FPropertyChangedEvent& PropertyChan
 	{
 		ChangeComponentScreenSizeToUseSubSections(ComponentScreenSizeToUseSubSections);
 	}
-	else if (PropertyName == GET_MEMBER_NAME_CHECKED(ALandscapeProxy, LODDistributionSetting)
+	else if (PropertyName == GET_MEMBER_NAME_CHECKED(ALandscapeProxy, ScalableLODDistributionSetting)
+		|| PropertyName == GET_MEMBER_NAME_CHECKED(ALandscapeProxy, ScalableLOD0DistributionSetting)
+		|| PropertyName == GET_MEMBER_NAME_CHECKED(ALandscapeProxy, ScalableLOD0ScreenSize)
+		|| PropertyName == GET_MEMBER_NAME_CHECKED(ALandscapeProxy, LODDistributionSetting)
 		|| PropertyName == GET_MEMBER_NAME_CHECKED(ALandscapeProxy, LOD0DistributionSetting)
 		|| PropertyName == GET_MEMBER_NAME_CHECKED(ALandscapeProxy, LOD0ScreenSize))
 	{		
@@ -6621,6 +6627,21 @@ void ALandscape::PostEditChangeProperty(FPropertyChangedEvent& PropertyChangedEv
 {
 	TRACE_CPUPROFILER_EVENT_SCOPE(ALandscape::PostEditChangeProperty);
 
+	// Clamp all values of a FPerQualityLevelFloat
+	auto Clamp = [](const FPerQualityLevelFloat& QualityLevelFloat, float Min, float Max) -> FPerQualityLevelFloat
+	{
+		FPerQualityLevelFloat Clamped = QualityLevelFloat;
+		Clamped.Default = FMath::Clamp(Clamped.GetDefault(), Min, Max);
+
+		for (auto It = Clamped.PerQuality.CreateIterator(); It; ++It)
+		{
+			float& Value = It.Value();
+			Value = FMath::Clamp(Value, Min, Max);
+		}
+		
+		return Clamped;
+	};
+
 	const FName PropertyName = PropertyChangedEvent.Property ? PropertyChangedEvent.Property->GetFName() : NAME_None;
 	const FName MemberPropertyName = PropertyChangedEvent.MemberProperty ? PropertyChangedEvent.MemberProperty->GetFName() : NAME_None;
 
@@ -6710,17 +6731,29 @@ void ALandscape::PostEditChangeProperty(FPropertyChangedEvent& PropertyChangedEv
 	{
 		ComponentScreenSizeToUseSubSections = FMath::Clamp<float>(ComponentScreenSizeToUseSubSections, 0.01f, 1.0f);
 	}
-	else if (PropertyName == GET_MEMBER_NAME_CHECKED(ALandscapeProxy, LODDistributionSetting))
+	else if (MemberPropertyName == GET_MEMBER_NAME_CHECKED(ALandscapeProxy, LODDistributionSetting))
+	{
+		LODDistributionSetting = FMath::Clamp(LODDistributionSetting, 1.0f, 10.0f);
+	}
+	else if (MemberPropertyName == GET_MEMBER_NAME_CHECKED(ALandscapeProxy, LOD0DistributionSetting))
+	{
+		LOD0DistributionSetting = FMath::Clamp(LOD0DistributionSetting, 1.0f, 10.0f);
+	}
+	else if (MemberPropertyName == GET_MEMBER_NAME_CHECKED(ALandscapeProxy, LOD0ScreenSize))
+	{
+		LOD0ScreenSize = FMath::Clamp(LOD0ScreenSize, 0.1f, 10.0f);
+	}
+	else if (MemberPropertyName == GET_MEMBER_NAME_CHECKED(ALandscapeProxy, ScalableLODDistributionSetting))
 	{
-		LODDistributionSetting = FMath::Clamp<float>(LODDistributionSetting, 1.0f, 10.0f);
+		ScalableLODDistributionSetting = Clamp(ScalableLODDistributionSetting, 1.0f, 10.0f);
 	}
-	else if (PropertyName == GET_MEMBER_NAME_CHECKED(ALandscapeProxy, LOD0DistributionSetting))
+	else if (MemberPropertyName == GET_MEMBER_NAME_CHECKED(ALandscapeProxy, ScalableLOD0DistributionSetting))
 	{
-		LOD0DistributionSetting = FMath::Clamp<float>(LOD0DistributionSetting, 1.0f, 10.0f);
+		ScalableLOD0DistributionSetting = Clamp(ScalableLOD0DistributionSetting, 1.0f, 10.0f);
 	}
-	else if (PropertyName == GET_MEMBER_NAME_CHECKED(ALandscapeProxy, LOD0ScreenSize))
+	else if (MemberPropertyName == GET_MEMBER_NAME_CHECKED(ALandscapeProxy, ScalableLOD0ScreenSize))
 	{
-		LOD0ScreenSize = FMath::Clamp<float>(LOD0ScreenSize, 0.1f, 10.0f);
+		ScalableLOD0ScreenSize = Clamp(ScalableLOD0ScreenSize, 0.1f, 10.0f);
 	}
 	else if (PropertyName == GET_MEMBER_NAME_CHECKED(ALandscapeProxy, CollisionMipLevel))
 	{
diff --git a/Engine/Source/Runtime/Landscape/Private/LandscapeRender.cpp b/Engine/Source/Runtime/Landscape/Private/LandscapeRender.cpp
--- a/Engine/Source/Runtime/Landscape/Private/LandscapeRender.cpp
+++ b/Engine/Source/Runtime/Landscape/Private/LandscapeRender.cpp
@@ -54,6 +54,7 @@ LandscapeRender.cpp: New terrain rendering
 #include "Algo/Transform.h"
 #include "LandscapeCulling.h"
 #include "RenderGraphBuilder.h"
+#include "Scalability.h"
 
 using namespace UE::Landscape;
 
@@ -1219,9 +1220,27 @@ FLandscapeComponentSceneProxy::FLandscapeComponentSceneProxy(ULandscapeComponent
 		HeightmapSubsectionOffsetV = ((float)(InComponent->SubsectionSizeQuads + 1) / (float)FMath::Max<int32>(1, HeightmapTexture->GetSizeY()));
 	}
 
-	float LOD0ScreenSize = InComponent->GetLandscapeProxy()->LOD0ScreenSize;
-	float LOD0Distribution = InComponent->GetLandscapeProxy()->LOD0DistributionSetting * GLandscapeLOD0DistributionScale;
-	float LODDistribution = InComponent->GetLandscapeProxy()->LODDistributionSetting * GLandscapeLODDistributionScale;
+
+	const ALandscapeProxy* Proxy = InComponent->GetLandscapeProxy();
+
+	float LOD0ScreenSize;
+	float LOD0Distribution;
+	float LODDistribution;
+	
+	if (Proxy->bUseScalableLODSettings)
+	{
+		const int32 LandscapeQuality = Scalability::GetQualityLevels().LandscapeQuality;
+		
+		LOD0ScreenSize = Proxy->ScalableLOD0ScreenSize.GetValue(LandscapeQuality);
+		LOD0Distribution = Proxy->ScalableLOD0DistributionSetting.GetValue(LandscapeQuality);
+		LODDistribution = Proxy->ScalableLODDistributionSetting.GetValue(LandscapeQuality);		
+	}
+	else
+	{
+		LOD0ScreenSize = Proxy->LOD0ScreenSize;
+		LOD0Distribution = Proxy->LOD0DistributionSetting * GLandscapeLOD0DistributionScale;
+		LODDistribution = Proxy->LODDistributionSetting * GLandscapeLODDistributionScale;
+	}
 
 #if !UE_BUILD_SHIPPING
 	if (GLandscapeLOD0ScreenSizeOverride > 0.0)
@@ -1379,7 +1398,6 @@ FLandscapeComponentSceneProxy::FLandscapeComponentSceneProxy(ULandscapeComponent
 #if !(UE_BUILD_SHIPPING || UE_BUILD_TEST) || (UE_BUILD_SHIPPING && WITH_EDITOR)
 	if (GIsEditor)
 	{
-		ALandscapeProxy* Proxy = InComponent->GetLandscapeProxy();
 		// Try to find a color for level coloration.
 		if (Proxy)
 		{
diff --git a/Engine/Source/Runtime/Landscape/Private/LandscapeSubsystem.cpp b/Engine/Source/Runtime/Landscape/Private/LandscapeSubsystem.cpp
--- a/Engine/Source/Runtime/Landscape/Private/LandscapeSubsystem.cpp
+++ b/Engine/Source/Runtime/Landscape/Private/LandscapeSubsystem.cpp
@@ -115,7 +115,17 @@ void ULandscapeSubsystem::Initialize(FSubsystemCollectionBase& Collection)
 	{
 		NotificationManager = new FLandscapeNotificationManager();
 	}
+
 #endif
+	
+	OnScalabilityChangedHandle = Scalability::OnScalabilitySettingsChanged.AddLambda([](const Scalability::FQualityLevels& QualityLevels)
+    {
+    	for (auto* LandscapeComponent : TObjectRange<ULandscapeComponent>(RF_ClassDefaultObject | RF_ArchetypeObject, true, EInternalObjectFlags::Garbage))
+    	{
+    		LandscapeComponent->MarkRenderStateDirty();
+    	}
+    });
+    	
 }
 
 void ULandscapeSubsystem::Deinitialize()
@@ -143,9 +153,11 @@ void ULandscapeSubsystem::Deinitialize()
 	{
 		LandscapeNaniteEnabledCVar->OnChangedDelegate().RemoveAll(this);
 	}
-	
+
+	Scalability::OnScalabilitySettingsChanged.Remove(OnScalabilityChangedHandle);
 	
 #if WITH_EDITOR
+	
 	while (NaniteBuildsInFlight != 0)
 	{
 		ENamedThreads::Type CurrentThread = FTaskGraphInterface::Get().GetCurrentThreadIfKnown();
diff --git a/Engine/Source/Runtime/Landscape/Public/LandscapeSubsystem.h b/Engine/Source/Runtime/Landscape/Public/LandscapeSubsystem.h
--- a/Engine/Source/Runtime/Landscape/Public/LandscapeSubsystem.h
+++ b/Engine/Source/Runtime/Landscape/Public/LandscapeSubsystem.h
@@ -149,5 +149,8 @@ private:
 	float NumNaniteMeshUpdatesAvailable = 0.0f;
 
 	std::atomic<int32> NaniteBuildsInFlight;
+
 #endif // WITH_EDITOR
+	
+	FDelegateHandle OnScalabilityChangedHandle;
 };
