Fix bug where NetPing control messages could prevent other control messages from being parsed if sent in the same bunch.

#rnx

diff --git a/Engine/Source/Runtime/Engine/Private/DataChannel.cpp b/Engine/Source/Runtime/Engine/Private/DataChannel.cpp
--- a/Engine/Source/Runtime/Engine/Private/DataChannel.cpp
+++ b/Engine/Source/Runtime/Engine/Private/DataChannel.cpp
@@ -1712,8 +1712,9 @@ void UControlChannel::ReceivedBunch( FInBunch& Bunch )
 		return;
 	}
 
+	bool bStopReadingBunch = false;
 	// Process the packet
-	while (!Bunch.AtEnd() && Connection != NULL && Connection->GetConnectionState() != USOCK_Closed) // if the connection got closed, we don't care about the rest
+	while (!Bunch.AtEnd() && bStopReadingBunch == false && Connection != nullptr && Connection->GetConnectionState() != USOCK_Closed) // if the connection got closed, we don't care about the rest
 	{
 		uint8 MessageType = 0;
 		Bunch << MessageType;
@@ -1807,7 +1808,7 @@ void UControlChannel::ReceivedBunch( FInBunch& Bunch )
 			if (FNetControlMessage<NMT_SecurityViolation>::Receive(Bunch, DebugMessage))
 			{
 				UE_LOG(LogSecurity, Warning, TEXT("%s: Closed: %s"), *Connection->RemoteAddressToString(), *DebugMessage);
-				break;
+				bStopReadingBunch = true;
 			}
 		}
 		else if (MessageType == NMT_DestructionInfo)
@@ -1831,7 +1832,7 @@ void UControlChannel::ReceivedBunch( FInBunch& Bunch )
 				Connection->HandleReceiveCloseReason(CloseReasonList);
 			}
 
-			break;
+			bStopReadingBunch = true;
 		}
 		else if (MessageType == NMT_NetPing)
 		{
@@ -1843,8 +1844,6 @@ void UControlChannel::ReceivedBunch( FInBunch& Bunch )
 			{
 				UE::Net::FNetPing::HandleNetPingControlMessage(Connection, NetPingMessageType, MessageStr);
 			}
-
-			break;
 		}
 		else if (Connection->Driver->Notify != nullptr)
 		{
@@ -1940,7 +1939,7 @@ void UControlChannel::ReceivedBunch( FInBunch& Bunch )
 
 			AddToChainResultPtr(Bunch.ExtendedError, ENetCloseResult::ControlChannelMessagePayloadFail);
 
-			break;
+			bStopReadingBunch = true;
 		}
 	}
 
