World Partition Actor Descriptors: detect multiple actors in the same package and treat them as invalid.

diff --git a/Engine/Source/Runtime/Engine/Private/WorldPartition/ActorDescContainer.cpp b/Engine/Source/Runtime/Engine/Private/WorldPartition/ActorDescContainer.cpp
--- a/Engine/Source/Runtime/Engine/Private/WorldPartition/ActorDescContainer.cpp
+++ b/Engine/Source/Runtime/Engine/Private/WorldPartition/ActorDescContainer.cpp
@@ -77,36 +77,27 @@ void UActorDescContainer::Initialize(const FInitializeParams& InitParams)
 		ClassDescRegistry.PrefetchClassDescs(ClassPaths.Array());
 	}
 
+	TMap<FGuid, TUniquePtr<FWorldPartitionActorDesc>> ValidActorDescs;
 	{
-		TRACE_CPUPROFILER_EVENT_SCOPE(RegisterDescriptors);
+		TRACE_CPUPROFILER_EVENT_SCOPE(CreateDescriptors);
+
+		TMap<FName, FWorldPartitionActorDesc*> ActorDescsByPackage;
 		for (const FAssetData& Asset : Assets)
 		{
 			TUniquePtr<FWorldPartitionActorDesc> ActorDesc = FWorldPartitionActorDescUtils::GetActorDescriptorFromAssetData(Asset);
 								
-			bool bValid = true;
 			if (!ActorDesc.IsValid())
 			{
 				UE_LOG(LogWorldPartition, Warning, TEXT("Invalid actor descriptor for actor '%s' from package '%s'"), *Asset.GetObjectPathString(), *Asset.PackageName.ToString());
-				bValid = false;
+				InvalidActors.Emplace(Asset);
 			} 
-			else if (FWorldPartitionActorDesc* ExistingDesc = FActorDescList::GetActorDesc(ActorDesc->GetGuid()))
-			{
-				check(ExistingDesc->GetGuid() == ActorDesc->GetGuid());
-				UE_LOG(LogWorldPartition, Warning, TEXT("Duplicate actor descriptor guid `%s`: Actor: '%s' from package '%s' -> Existing actor '%s' from package '%s'"), 
-					*ActorDesc->GetGuid().ToString(), 
-					*ActorDesc->GetActorName().ToString(), 
-					*ActorDesc->GetActorPackage().ToString(),
-					*ExistingDesc->GetActorName().ToString(),
-					*ExistingDesc->GetActorPackage().ToString());
-				bValid = false;
-			}
 			else if (!ActorDesc->GetNativeClass().IsValid())
 			{
 				UE_LOG(LogWorldPartition, Warning, TEXT("Invalid actor native class: Actor: '%s' (guid '%s') from package '%s'"),
 					*ActorDesc->GetActorName().ToString(),
 					*ActorDesc->GetGuid().ToString(),
 					*ActorDesc->GetActorPackage().ToString());
-				bValid = false;
+				InvalidActors.Emplace(Asset);
 			}
 			else if (ActorDesc->GetBaseClass().IsValid() && !ClassDescRegistry.IsRegisteredClass(ActorDesc->GetBaseClass()))
 			{
@@ -115,19 +106,51 @@ void UActorDescContainer::Initialize(const FInitializeParams& InitParams)
 					*ActorDesc->GetActorName().ToString(),
 					*ActorDesc->GetGuid().ToString(),
 					*ActorDesc->GetActorPackage().ToString());
-				bValid = false;
+				InvalidActors.Emplace(Asset);
 			}
 			else if (InitParams.FilterActorDesc && !InitParams.FilterActorDesc(ActorDesc.Get()))
 			{
-				bValid = false;
+				InvalidActors.Emplace(Asset);
 			}
+			// At this point, the actor descriptor is well formed and valid on its own. We now make validations based on the already registered
+			// actor descriptors, such as duplicated actor GUIDs or multiple actors in the same package, etc.
+			else if (FWorldPartitionActorDesc* ExistingDescPackage = ActorDescsByPackage.FindRef(ActorDesc->GetActorPackage()))
+			{
+				UE_LOG(LogWorldPartition, Warning, TEXT("Duplicate actor descriptor in package `%s`: Actor: '%s' -> Existing actor '%s'"), 
+					*ActorDesc->GetActorPackage().ToString(), 
+					*ActorDesc->GetActorName().ToString(), 
+					*ExistingDescPackage->GetActorName().ToString());
 
-			if (!bValid)
+				// No need to add all actors in the same package several times as we only want to open the package for delete when repairing
+				if (ValidActorDescs.Contains(ActorDesc->GetGuid()))
+				{
+					InvalidActors.Emplace(Asset);
+					ValidActorDescs.Remove(ActorDesc->GetGuid());
+				}
+			}
+			else if (FWorldPartitionActorDesc* ExistingDescGuid = FActorDescList::GetActorDesc(ActorDesc->GetGuid()))
 			{
+				check(ExistingDescGuid->GetGuid() == ActorDesc->GetGuid());
+				UE_LOG(LogWorldPartition, Warning, TEXT("Duplicate actor descriptor guid `%s`: Actor: '%s' from package '%s' -> Existing actor '%s' from package '%s'"), 
+					*ActorDesc->GetGuid().ToString(), 
+					*ActorDesc->GetActorName().ToString(), 
+					*ActorDesc->GetActorPackage().ToString(),
+					*ExistingDescGuid->GetActorName().ToString(),
+					*ExistingDescGuid->GetActorPackage().ToString());
 				InvalidActors.Emplace(Asset);
-				continue;
 			}
+			else
+			{
+				ActorDescsByPackage.Add(ActorDesc->GetActorPackage(), ActorDesc.Get());
+				ValidActorDescs.Add(ActorDesc->GetGuid(), MoveTemp(ActorDesc));
+			}
+		}
+	}
 
+	{
+		TRACE_CPUPROFILER_EVENT_SCOPE(RegisterDescriptors);
+		for (auto& [ActorGuid, ActorDesc] : ValidActorDescs)
+		{
 			RegisterActorDescriptor(ActorDesc.Release(), OwningWorld);
 		}
 	}
