Fixed crash due to global read buffer not being committed after gathering ray tracing Niagara proxies.

diff --git a/Engine/Source/Runtime/Engine/Private/SceneManagement.cpp b/Engine/Source/Runtime/Engine/Private/SceneManagement.cpp
--- a/Engine/Source/Runtime/Engine/Private/SceneManagement.cpp
+++ b/Engine/Source/Runtime/Engine/Private/SceneManagement.cpp
@@ -348,7 +348,7 @@ FRayTracingMaterialGatheringContext::FRayTracingMaterialGatheringContext(
 	const FSceneViewFamily& InReferenceViewFamily,
 	FRDGBuilder& InGraphBuilder,
 	FRayTracingMeshResourceCollector& InRayTracingMeshResourceCollector,
-	FGlobalDynamicReadBuffer& DynamicReadBuffer)
+	FGlobalDynamicReadBuffer& InDynamicReadBuffer)
 	: Scene(InScene)
 	, ReferenceView(InReferenceView)
 	, ReferenceViewFamily(InReferenceViewFamily)
@@ -357,6 +357,7 @@ FRayTracingMaterialGatheringContext::FRayTracingMaterialGatheringContext(
 	, RayTracingMeshResourceCollector(InRayTracingMeshResourceCollector)
 	, DynamicVertexBuffer(GraphBuilder.RHICmdList)
 	, DynamicIndexBuffer(GraphBuilder.RHICmdList)
+	, DynamicReadBuffer(InDynamicReadBuffer)
 {
 	RayTracingMeshResourceCollector.Start(RHICmdList, DynamicVertexBuffer, DynamicIndexBuffer, DynamicReadBuffer);
 }
@@ -364,6 +365,7 @@ FRayTracingMaterialGatheringContext::FRayTracingMaterialGatheringContext(
 FRayTracingMaterialGatheringContext::~FRayTracingMaterialGatheringContext()
 {
 	RayTracingMeshResourceCollector.Finish();
+	DynamicReadBuffer.Commit(GraphBuilder.RHICmdList);
 }
 
 #endif
diff --git a/Engine/Source/Runtime/Engine/Public/SceneManagement.h b/Engine/Source/Runtime/Engine/Public/SceneManagement.h
--- a/Engine/Source/Runtime/Engine/Public/SceneManagement.h
+++ b/Engine/Source/Runtime/Engine/Public/SceneManagement.h
@@ -2148,6 +2148,7 @@ struct FRayTracingMaterialGatheringContext
 	TArray<FRayTracingDynamicGeometryUpdateParams> DynamicRayTracingGeometriesToUpdate;
 	FGlobalDynamicVertexBuffer DynamicVertexBuffer;
 	FGlobalDynamicIndexBuffer DynamicIndexBuffer;
+	FGlobalDynamicReadBuffer& DynamicReadBuffer;
 
 	ENGINE_API FRayTracingMaterialGatheringContext(
 		const FScene* InScene,
