-Force instantiation of Nanite::ShouldCreateNaniteProxy<T> for UStaticMeshComponent/FStaticMeshSceneProxyDesc to avoid linker issues with some PGO builds.

diff --git a/Engine/Source/Runtime/Engine/Private/StaticMeshRender.cpp b/Engine/Source/Runtime/Engine/Private/StaticMeshRender.cpp
--- a/Engine/Source/Runtime/Engine/Private/StaticMeshRender.cpp
+++ b/Engine/Source/Runtime/Engine/Private/StaticMeshRender.cpp
@@ -2588,6 +2588,9 @@ namespace Nanite
 
 		return true;
 	}
+	
+	template bool ShouldCreateNaniteProxy(const UStaticMeshComponent& Component, FMaterialAudit* OutNaniteMaterials);
+	template bool ShouldCreateNaniteProxy(const FStaticMeshSceneProxyDesc& Component, FMaterialAudit* OutNaniteMaterials);
 }
 
 bool UStaticMeshComponent::ShouldCreateNaniteProxy(Nanite::FMaterialAudit* OutNaniteMaterials) const
