Add note to explain that setting ShadowResolutionScale to zero disables shadow maps but not contact shadows.

diff --git a/Engine/Source/Runtime/Engine/Classes/Components/LightComponent.h b/Engine/Source/Runtime/Engine/Classes/Components/LightComponent.h
--- a/Engine/Source/Runtime/Engine/Classes/Components/LightComponent.h
+++ b/Engine/Source/Runtime/Engine/Classes/Components/LightComponent.h
@@ -81,6 +81,7 @@ class ULightComponent : public ULightComponentBase
 
 	/** 
 	 * Scales the resolution of shadowmaps used to shadow this light.  By default shadowmap resolution is chosen based on screen size of the caster. 
+	 * Setting the scale to zero disables shadow maps, but does not disable, e.g., contact shadows.
 	 * Note: shadowmap resolution is still clamped by 'r.Shadow.MaxResolution'
 	 */
 	UPROPERTY(EditAnywhere, BlueprintReadOnly, Category=Light, AdvancedDisplay, meta=(UIMin = ".125", UIMax = "8"))
