Update material ddc key when files included by (custom) material expressions change.

diff --git a/Engine/Source/Runtime/Core/Private/UObject/FortniteShaderworkObjectVersion.cpp b/Engine/Source/Runtime/Core/Private/UObject/FortniteShaderworkObjectVersion.cpp
--- a/Engine/Source/Runtime/Core/Private/UObject/FortniteShaderworkObjectVersion.cpp
+++ b/Engine/Source/Runtime/Core/Private/UObject/FortniteShaderworkObjectVersion.cpp
@@ -8,6 +8,6 @@ TMap<FGuid, FGuid> FFortniteShaderworkObjectVersion::GetSystemGuids()
 
 	SystemGuids.Add(DevGuids.NANITE_DERIVEDDATA_VER, FGuid("6ED76D6068904EBB9B3ADAC259F0F53E"));
 	SystemGuids.Add(DevGuids.STATICMESH_DERIVEDDATA_VER, FGuid("EB4D2761094040DD889CC1BE3D24E1B3"));
-	SystemGuids.Add(DevGuids.MATERIALSHADERMAP_DERIVEDDATA_VER, FGuid("0820372747734716A3BAA8C0F6A40ABA"));
+	SystemGuids.Add(DevGuids.MATERIALSHADERMAP_DERIVEDDATA_VER, FGuid("21EC97751B084536AE887AA52996DA1C"));
 	return SystemGuids;
 }
diff --git a/Engine/Source/Runtime/Core/Public/UObject/RenderingObjectVersion.h b/Engine/Source/Runtime/Core/Public/UObject/RenderingObjectVersion.h
--- a/Engine/Source/Runtime/Core/Public/UObject/RenderingObjectVersion.h
+++ b/Engine/Source/Runtime/Core/Public/UObject/RenderingObjectVersion.h
@@ -133,6 +133,9 @@ struct FRenderingObjectVersion
 		// Added support to force Nanite material usage if Nanite mesh builders are also forced.
 		NaniteForceMaterialUsage,
 
+		// Custom material expression includes are now serialized & used for DDC key generation
+		AddedMaterialExpressionIncludesHash,
+
 		// -----<new versions can be added above this line>-------------------------------------------------
 		VersionPlusOne,
 		LatestVersion = VersionPlusOne - 1
diff --git a/Engine/Source/Runtime/Engine/Classes/Materials/MaterialExpression.h b/Engine/Source/Runtime/Engine/Classes/Materials/MaterialExpression.h
--- a/Engine/Source/Runtime/Engine/Classes/Materials/MaterialExpression.h
+++ b/Engine/Source/Runtime/Engine/Classes/Materials/MaterialExpression.h
@@ -565,6 +565,7 @@ class UMaterialExpression : public UObject
 	virtual bool SetParameterValue(const FName& Name, const FMaterialParameterMetadata& Meta, EMaterialExpressionSetParameterValueFlags Flags = EMaterialExpressionSetParameterValueFlags::None) { return false; }
 
 	virtual void GetLandscapeLayerNames(TArray<FName>& OutLayers) const {}
+	virtual void GetIncludeFilePaths(TSet<FString>& OutIncludeFilePaths) const {}
 
 	/**
 	 * Called after a node copy, once the Material and Function properties are set correctly and that all new expressions are added to Material->Expressions
diff --git a/Engine/Source/Runtime/Engine/Classes/Materials/MaterialExpressionCustom.h b/Engine/Source/Runtime/Engine/Classes/Materials/MaterialExpressionCustom.h
--- a/Engine/Source/Runtime/Engine/Classes/Materials/MaterialExpressionCustom.h
+++ b/Engine/Source/Runtime/Engine/Classes/Materials/MaterialExpressionCustom.h
@@ -108,6 +108,7 @@ class UMaterialExpressionCustom : public UMaterialExpression
 	virtual uint32 GetOutputType(int32 OutputIndex) override;
 	virtual bool IsResultMaterialAttributes(int32 OutputIndex) override;
 	virtual bool GenerateHLSLExpression(FMaterialHLSLGenerator& Generator, UE::HLSLTree::FScope& Scope, int32 OutputIndex, UE::HLSLTree::FExpression const*& OutExpression) const override;
+	virtual void GetIncludeFilePaths(TSet<FString>& OutIncludeFilePaths) const override;
 #endif // WITH_EDITOR
 	//~ End UMaterialExpression Interface
 };
diff --git a/Engine/Source/Runtime/Engine/Private/HLSLTree/HLSLTreeCommon.cpp b/Engine/Source/Runtime/Engine/Private/HLSLTree/HLSLTreeCommon.cpp
--- a/Engine/Source/Runtime/Engine/Private/HLSLTree/HLSLTreeCommon.cpp
+++ b/Engine/Source/Runtime/Engine/Private/HLSLTree/HLSLTreeCommon.cpp
@@ -1300,6 +1300,15 @@ bool FExpressionCustomHLSL::PrepareValue(FEmitContext& Context, FEmitScope& Scop
 		Context.MarkInputType(Input.Expression, InputType.Type);
 	}
 
+	FMaterialCachedExpressionData* CachedExpressionData = Context.FindData<Material::FEmitData>().CachedExpressionData;
+	if (CachedExpressionData && CachedExpressionData->EditorOnlyData.IsValid())
+	{
+		for (const FString& IncludeFilePath : IncludeFilePaths)
+		{
+			CachedExpressionData->EditorOnlyData->ExpressionIncludeFilePaths.Add(IncludeFilePath);
+		}
+	}
+
 	return OutResult.SetType(Context, RequestedType, EExpressionEvaluation::Shader, Shader::FType(OutputStructType));
 }
 
diff --git a/Engine/Source/Runtime/Engine/Private/Materials/MaterialCachedData.cpp b/Engine/Source/Runtime/Engine/Private/Materials/MaterialCachedData.cpp
--- a/Engine/Source/Runtime/Engine/Private/Materials/MaterialCachedData.cpp
+++ b/Engine/Source/Runtime/Engine/Private/Materials/MaterialCachedData.cpp
@@ -370,6 +370,8 @@ void FMaterialCachedExpressionData::UpdateForExpressions(const FMaterialCachedEx
 
 		Expression->GetLandscapeLayerNames(EditorOnlyData->LandscapeLayerNames);
 
+		Expression->GetIncludeFilePaths(EditorOnlyData->ExpressionIncludeFilePaths);
+
 		if (UMaterialExpressionCollectionParameter* ExpressionCollectionParameter = Cast<UMaterialExpressionCollectionParameter>(Expression))
 		{
 			UMaterialParameterCollection* Collection = ExpressionCollectionParameter->Collection;
@@ -1004,6 +1006,36 @@ void FMaterialCachedExpressionData::Validate()
 			check(EditorEntry.EditorInfo.Num() == Entry.ParameterInfoSet.Num());
 		}
 		FMaterialLayersFunctions::Validate(MaterialLayers, EditorOnlyData->MaterialLayers);
+
+
+		// TODO Validate the shader include paths in the editor before storing them on the expression instead
+		for (auto PathIt = EditorOnlyData->ExpressionIncludeFilePaths.CreateIterator(); PathIt; ++PathIt)
+		{
+			const FString& IncludeFilePath = *PathIt;
+			bool bValidExpressionIncludePath = false;
+
+			if (!IncludeFilePath.IsEmpty())
+			{
+				FString ValidatedPath = GetShaderSourceFilePath(IncludeFilePath);
+				if (!ValidatedPath.IsEmpty())
+				{
+					ValidatedPath = IFileManager::Get().ConvertToAbsolutePathForExternalAppForWrite(*ValidatedPath);
+					if (FPaths::FileExists(ValidatedPath))
+					{
+						bValidExpressionIncludePath = true;
+					}
+				}
+			}
+
+			if (!bValidExpressionIncludePath)
+			{
+				UE_LOG(LogMaterial, Warning, TEXT("Expression include file path [%s] is invalid, removing from cached material data."), *IncludeFilePath);
+				PathIt.RemoveCurrent();
+			}
+		}
+
+		// Sort to make hashing less dependent on the order of expression visiting
+		EditorOnlyData->ExpressionIncludeFilePaths.Sort(TLess<>());
 	}
 }
 
diff --git a/Engine/Source/Runtime/Engine/Private/Materials/MaterialExpressionHLSL.cpp b/Engine/Source/Runtime/Engine/Private/Materials/MaterialExpressionHLSL.cpp
--- a/Engine/Source/Runtime/Engine/Private/Materials/MaterialExpressionHLSL.cpp
+++ b/Engine/Source/Runtime/Engine/Private/Materials/MaterialExpressionHLSL.cpp
@@ -4256,6 +4256,7 @@ bool UMaterialExpressionCustom::GenerateHLSLExpression(FMaterialHLSLGenerator& G
 	const FExpression* ExpressionCustom = Generator.GetTree().NewExpression<FExpressionCustomHLSL>(
 		UE::MemStack::AllocateStringView(Allocator, DeclarationCode.ToView()),
 		FunctionCode,
+		IncludeFilePaths, // Can just reference field directly, the UMaterialExpressionCustom lifetime will be longer than the resulting HLSLTree
 		LocalInputs,
 		OutputStructType);
 
@@ -4263,6 +4264,11 @@ bool UMaterialExpressionCustom::GenerateHLSLExpression(FMaterialHLSLGenerator& G
 	return true;
 }
 
+void UMaterialExpressionCustom::GetIncludeFilePaths(TSet<FString>& OutIncludeFilePaths) const
+{
+	OutIncludeFilePaths.Append(IncludeFilePaths);
+}
+
 bool UMaterialExpressionClearCoatNormalCustomOutput::GenerateHLSLExpression(FMaterialHLSLGenerator& Generator, UE::HLSLTree::FScope& Scope, int32 OutputIndex, UE::HLSLTree::FExpression const*& OutExpression) const
 {
 	OutExpression = Input.AcquireHLSLExpression(Generator, Scope);
diff --git a/Engine/Source/Runtime/Engine/Private/Materials/MaterialShader.cpp b/Engine/Source/Runtime/Engine/Private/Materials/MaterialShader.cpp
--- a/Engine/Source/Runtime/Engine/Private/Materials/MaterialShader.cpp
+++ b/Engine/Source/Runtime/Engine/Private/Materials/MaterialShader.cpp
@@ -601,6 +601,7 @@ void FMaterialShaderMapId::Serialize(FArchive& Ar, bool bLoadedByCookedMaterial)
 	Ar.UsingCustomVersion(FEditorObjectVersion::GUID);
 	Ar.UsingCustomVersion(FReleaseObjectVersion::GUID);
 	Ar.UsingCustomVersion(FUE5ReleaseStreamObjectVersion::GUID);
+	Ar.UsingCustomVersion(FRenderingObjectVersion::GUID);
 
 	const bool bIsLegacyPackage = Ar.UEVer() < VER_UE4_PURGED_FMATERIAL_COMPILE_OUTPUTS;
 
@@ -641,8 +642,6 @@ void FMaterialShaderMapId::Serialize(FArchive& Ar, bool bLoadedByCookedMaterial)
 #if WITH_EDITOR
 	if (!bIsSavingCooked && !bLoadedByCookedMaterial)
 	{
-		Ar.UsingCustomVersion(FRenderingObjectVersion::GUID);
-
 		if (Ar.CustomVer(FRenderingObjectVersion::GUID) < FRenderingObjectVersion::MaterialShaderMapIdSerialization)
 		{
 			// Serialize using old path
@@ -703,6 +702,11 @@ void FMaterialShaderMapId::Serialize(FArchive& Ar, bool bLoadedByCookedMaterial)
 			Ar << LegacyHash;
 		}
 
+		if (Ar.CustomVer(FRenderingObjectVersion::GUID) >= FRenderingObjectVersion::AddedMaterialExpressionIncludesHash)
+		{
+			Ar << ExpressionIncludesHash;
+		}
+
 		if (Ar.UEVer() >= VER_UE4_MATERIAL_INSTANCE_BASE_PROPERTY_OVERRIDES)
 		{
 			Ar << BasePropertyOverridesHash;
@@ -812,6 +816,8 @@ void FMaterialShaderMapId::GetMaterialHash(FSHAHash& OutHash, bool bWithStaticPa
 
 	HashState.Update((const uint8*)&TextureReferencesHash, sizeof(TextureReferencesHash));
 
+	HashState.Update((const uint8*)&ExpressionIncludesHash, sizeof(ExpressionIncludesHash));
+
 	HashState.Update((const uint8*)&BasePropertyOverridesHash, sizeof(BasePropertyOverridesHash));
 
 	HashState.Update((const uint8*)&bUsingNewHLSLGenerator, sizeof(bUsingNewHLSLGenerator));
@@ -963,6 +969,11 @@ bool FMaterialShaderMapId::Equals(const FMaterialShaderMapId& ReferenceSet, bool
 			return false;
 		}
 
+		if (ExpressionIncludesHash != ReferenceSet.ExpressionIncludesHash)
+		{
+			return false;
+		}
+
 		if( BasePropertyOverridesHash != ReferenceSet.BasePropertyOverridesHash )
 		{
 			return false;
@@ -1139,6 +1150,8 @@ void FMaterialShaderMapId::AppendKeyString(FString& KeyString, bool bIncludeSour
 
 	BytesToHex(&TextureReferencesHash.Hash[0], sizeof(TextureReferencesHash.Hash), KeyString);
 
+	BytesToHex(&ExpressionIncludesHash.Hash[0], sizeof(ExpressionIncludesHash.Hash), KeyString);
+
 	BytesToHex(&BasePropertyOverridesHash.Hash[0], sizeof(BasePropertyOverridesHash.Hash), KeyString);
 
 	if (bUsingNewHLSLGenerator)
diff --git a/Engine/Source/Runtime/Engine/Private/Materials/MaterialShared.cpp b/Engine/Source/Runtime/Engine/Private/Materials/MaterialShared.cpp
--- a/Engine/Source/Runtime/Engine/Private/Materials/MaterialShared.cpp
+++ b/Engine/Source/Runtime/Engine/Private/Materials/MaterialShared.cpp
@@ -856,6 +856,7 @@ void FMaterial::GetShaderMapId(EShaderPlatform Platform, const ITargetPlatform*
 		OutId.FeatureLevel = GetFeatureLevel();
 		OutId.SetShaderDependencies(ShaderTypes, ShaderPipelineTypes, VFTypes, Platform);
 		GetReferencedTexturesHash(Platform, OutId.TextureReferencesHash);
+		GetExpressionIncludesHash(Platform, OutId.ExpressionIncludesHash);
 
 		OutId.SubstrateCompilationConfig = GetSubstrateCompilationConfig();
 
@@ -4195,6 +4196,7 @@ void FMaterial::GetDependentShaderAndVFTypes(EShaderPlatform Platform, const FPl
 	}
 }
 
+#if WITH_EDITOR
 void FMaterial::GetReferencedTexturesHash(EShaderPlatform Platform, FSHAHash& OutHash) const
 {
 	FSHA1 HashState;
@@ -4223,6 +4225,22 @@ void FMaterial::GetReferencedTexturesHash(EShaderPlatform Platform, FSHAHash& Ou
 	HashState.GetHash(&OutHash.Hash[0]);
 }
 
+void FMaterial::GetExpressionIncludesHash(EShaderPlatform Platform, FSHAHash& OutHash) const
+{
+	FSHA1 HashState;
+
+	for (const FString& ExpressionIncludeFilePath : GetCachedExpressionData().EditorOnlyData->ExpressionIncludeFilePaths)
+	{
+		checkf(!ExpressionIncludeFilePath.IsEmpty(), TEXT("Expression include path is empty but it should have been previously validated."));
+
+		const FSHAHash& FileHash = GetShaderFileHash(*ExpressionIncludeFilePath, Platform);
+		HashState.Update(FileHash.Hash, UE_ARRAY_COUNT(FileHash.Hash));
+	}
+
+	OutHash = HashState.Finalize();
+}
+#endif // WITH_EDITOR
+
 /**
  * Get user source code for the material, with a list of code snippets to highlight representing the code for each MaterialExpression
  * @param OutSource - generated source code
diff --git a/Engine/Source/Runtime/Engine/Public/HLSLTree/HLSLTreeCommon.h b/Engine/Source/Runtime/Engine/Public/HLSLTree/HLSLTreeCommon.h
--- a/Engine/Source/Runtime/Engine/Public/HLSLTree/HLSLTreeCommon.h
+++ b/Engine/Source/Runtime/Engine/Public/HLSLTree/HLSLTreeCommon.h
@@ -342,18 +342,25 @@ public:
 class FExpressionCustomHLSL : public FExpression
 {
 public:
-	FExpressionCustomHLSL(FStringView InDeclarationCode, FStringView InFunctionCode, TArrayView<FCustomHLSLInput> InInputs, const Shader::FStructType* InOutputStructType)
+	FExpressionCustomHLSL(FStringView InDeclarationCode, FStringView InFunctionCode, TConstArrayView<FString> InIncludeFilePaths, TArrayView<FCustomHLSLInput> InInputs, const Shader::FStructType* InOutputStructType)
 		: DeclarationCode(InDeclarationCode)
 		, FunctionCode(InFunctionCode)
+		, IncludeFilePaths(InIncludeFilePaths)
 		, Inputs(InInputs)
 		, OutputStructType(InOutputStructType)
 	{}
 
+	UE_DEPRECATED(5.4, "This constructor is deprecated.")
+	FExpressionCustomHLSL(FStringView InDeclarationCode, FStringView InFunctionCode, TArrayView<FCustomHLSLInput> InInputs, const Shader::FStructType* InOutputStructType)
+		: FExpressionCustomHLSL(InDeclarationCode, InFunctionCode, {}, InInputs,  InOutputStructType)
+	{}
+
 	virtual bool PrepareValue(FEmitContext& Context, FEmitScope& Scope, const FRequestedType& RequestedType, FPrepareValueResult& OutResult) const override;
 	virtual void EmitValueShader(FEmitContext& Context, FEmitScope& Scope, const FRequestedType& RequestedType, FEmitValueShaderResult& OutResult) const override;
 
 	FStringView DeclarationCode;
 	FStringView FunctionCode;
+	TConstArrayView<FString> IncludeFilePaths;
 	TArray<FCustomHLSLInput, TInlineAllocator<8>> Inputs;
 	const Shader::FStructType* OutputStructType = nullptr;
 };
diff --git a/Engine/Source/Runtime/Engine/Public/MaterialCachedData.h b/Engine/Source/Runtime/Engine/Public/MaterialCachedData.h
--- a/Engine/Source/Runtime/Engine/Public/MaterialCachedData.h
+++ b/Engine/Source/Runtime/Engine/Public/MaterialCachedData.h
@@ -173,6 +173,9 @@ struct FMaterialCachedExpressionEditorOnlyData
 
 	UPROPERTY()
 	TArray<FName> LandscapeLayerNames;
+
+	UPROPERTY()
+	TSet<FString> ExpressionIncludeFilePaths;
 };
 
 USTRUCT()
diff --git a/Engine/Source/Runtime/Engine/Public/MaterialShared.h b/Engine/Source/Runtime/Engine/Public/MaterialShared.h
--- a/Engine/Source/Runtime/Engine/Public/MaterialShared.h
+++ b/Engine/Source/Runtime/Engine/Public/MaterialShared.h
@@ -964,6 +964,9 @@ public:
 	 * that generates the array of textures that the uniform expressions use to link up after being loaded from the DDC.
 	 */
 	FSHAHash TextureReferencesHash;
+
+	/** A hash of the content of files (dynamically) included by material expressions. */
+	FSHAHash ExpressionIncludesHash;
 	
 	/** A hash of the base property overrides for this material instance. */
 	FSHAHash BasePropertyOverridesHash;
@@ -2558,8 +2561,12 @@ private:
 	 */
 	ENGINE_API TShaderRef<FShader> GetShader(class FMeshMaterialShaderType* ShaderType, FVertexFactoryType* VertexFactoryType, int32 PermutationId, bool bFatalIfMissing = true) const;
 
+#if WITH_EDITOR
 	void GetReferencedTexturesHash(EShaderPlatform Platform, FSHAHash& OutHash) const;
 
+	void GetExpressionIncludesHash(EShaderPlatform Platform, FSHAHash& OutHash) const;
+#endif // WITH_EDITOR
+
 	friend class FMaterialShaderMap;
 	friend class FShaderCompilingManager;
 	friend class FHLSLMaterialTranslator;
