Temp fix of a crash when editing a level instance that was previously pinned.

diff --git a/Engine/Source/Runtime/Engine/Private/WorldPartition/WorldPartitionActorLoaderInterface.cpp b/Engine/Source/Runtime/Engine/Private/WorldPartition/WorldPartitionActorLoaderInterface.cpp
--- a/Engine/Source/Runtime/Engine/Private/WorldPartition/WorldPartitionActorLoaderInterface.cpp
+++ b/Engine/Source/Runtime/Engine/Private/WorldPartition/WorldPartitionActorLoaderInterface.cpp
@@ -244,11 +244,14 @@ void IWorldPartitionActorLoaderInterface::ILoaderAdapter::ForEachReferencedActor
 {
 	for (const auto& [Container, ActorReferenceMap] : ContainerActorReferences)
 	{
-		for (const auto& [ReferenceGuid, ReferenceMap] : ActorReferenceMap)
+		if (Container.IsValid())
 		{
-			for (const auto& [Guid, ActorReference] : ReferenceMap)
+			for (const auto& [ReferenceGuid, ReferenceMap] : ActorReferenceMap)
 			{
-				InFunc(ActorReference);
+				for (const auto& [Guid, ActorReference] : ReferenceMap)
+				{
+					InFunc(ActorReference);
+				}
 			}
 		}
 	}
