World Partition: removed 5.1 deprecated functions.

diff --git a/Engine/Source/Runtime/Engine/Private/WorldPartition/WorldPartitionEditorHash.cpp b/Engine/Source/Runtime/Engine/Private/WorldPartition/WorldPartitionEditorHash.cpp
--- a/Engine/Source/Runtime/Engine/Private/WorldPartition/WorldPartitionEditorHash.cpp
+++ b/Engine/Source/Runtime/Engine/Private/WorldPartition/WorldPartitionEditorHash.cpp
@@ -17,13 +17,4 @@ UWorldPartitionEditorHash::FForEachIntersectingActorParams::FForEachIntersecting
 	: bIncludeSpatiallyLoadedActors(true)
 	, bIncludeNonSpatiallyLoadedActors(true)
 {}
-
-int32 UWorldPartitionEditorHash::ForEachIntersectingActor(const FBox& Box, TFunctionRef<void(FWorldPartitionActorDesc*)> InOperation, bool bIncludeSpatiallyLoadedActors, bool bIncludeNonSpatiallyLoadedActors)
-{
-	UWorldPartitionEditorHash::FForEachIntersectingActorParams ForEachIntersectingActorParams = UWorldPartitionEditorHash::FForEachIntersectingActorParams()
-		.SetIncludeSpatiallyLoadedActors(bIncludeSpatiallyLoadedActors)
-		.SetIncludeNonSpatiallyLoadedActors(bIncludeNonSpatiallyLoadedActors);
-
-	return ForEachIntersectingActor(Box, InOperation, ForEachIntersectingActorParams);
-}
 #endif
diff --git a/Engine/Source/Runtime/Engine/Private/WorldPartition/WorldPartitionHelpers.cpp b/Engine/Source/Runtime/Engine/Private/WorldPartition/WorldPartitionHelpers.cpp
--- a/Engine/Source/Runtime/Engine/Private/WorldPartition/WorldPartitionHelpers.cpp
+++ b/Engine/Source/Runtime/Engine/Private/WorldPartition/WorldPartitionHelpers.cpp
@@ -154,23 +154,6 @@ FWorldPartitionHelpers::FForEachActorWithLoadingParams::FForEachActorWithLoading
 	, ActorClasses({ AActor::StaticClass() })
 {}
 
-void FWorldPartitionHelpers::ForEachActorWithLoading(UWorldPartition* WorldPartition, TSubclassOf<AActor> ActorClass, TFunctionRef<bool(const FWorldPartitionActorDesc*)> Func, TFunctionRef<void()> OnReleasingActorReferences, bool bGCPerActor)
-{
-	FForEachActorWithLoadingParams Params;
-	Params.ActorClasses = { ActorClass };
-	Params.OnPreGarbageCollect = [&OnReleasingActorReferences]() { OnReleasingActorReferences(); };
-	ForEachActorWithLoading(WorldPartition, Func, Params);
-}
-
-void FWorldPartitionHelpers::ForEachActorWithLoading(UWorldPartition* WorldPartition, const TArray<FGuid>& ActorGuids, TFunctionRef<bool(const FWorldPartitionActorDesc*)> Func, TFunctionRef<void()> OnReleasingActorReferences, bool bGCPerActor)
-{
-	TSet<FGuid> ActorGuidsSet(ActorGuids);
-	FForEachActorWithLoadingParams Params;
-	Params.FilterActorDesc = [&ActorGuidsSet](const FWorldPartitionActorDesc* ActorDesc) -> bool { return ActorGuidsSet.Contains(ActorDesc->GetGuid());	};
-	Params.OnPreGarbageCollect = [&OnReleasingActorReferences]() { OnReleasingActorReferences(); };	
-	ForEachActorWithLoading(WorldPartition, Func, Params);
-}
-
 void FWorldPartitionHelpers::ForEachActorWithLoading(UWorldPartition* WorldPartition, TFunctionRef<bool(const FWorldPartitionActorDesc*)> Func, const FForEachActorWithLoadingParams& Params)
 {
 	FForEachActorWithLoadingResult Result;
diff --git a/Engine/Source/Runtime/Engine/Private/WorldPartition/WorldPartitionRuntimeHash.cpp b/Engine/Source/Runtime/Engine/Private/WorldPartition/WorldPartitionRuntimeHash.cpp
--- a/Engine/Source/Runtime/Engine/Private/WorldPartition/WorldPartitionRuntimeHash.cpp
+++ b/Engine/Source/Runtime/Engine/Private/WorldPartition/WorldPartitionRuntimeHash.cpp
@@ -385,54 +385,6 @@ void UWorldPartitionRuntimeHash::ForceExternalActorLevelReference(bool bForceExt
 }
 #endif
 
-int32 UWorldPartitionRuntimeHash::GetAllStreamingCells(TSet<const UWorldPartitionRuntimeCell*>& Cells, bool bAllDataLayers, bool bDataLayersOnly, const TSet<FName>& InDataLayers) const
-{
-	ForEachStreamingCells([&Cells, bAllDataLayers, bDataLayersOnly, InDataLayers](const UWorldPartitionRuntimeCell* Cell)
-	{
-		if (!bDataLayersOnly && !Cell->HasDataLayers())
-		{
-			Cells.Add(Cell);
-		}
-		else if (Cell->HasDataLayers() && (bAllDataLayers || Cell->HasAnyDataLayer(InDataLayers)))
-		{
-			Cells.Add(Cell);
-		}
-		return true;
-	});
-
-	return Cells.Num();
-}
-
-bool UWorldPartitionRuntimeHash::GetStreamingCells(const FWorldPartitionStreamingQuerySource& QuerySource, TSet<const UWorldPartitionRuntimeCell*>& OutCells) const
-{
-	ForEachStreamingCellsQuery(QuerySource, [QuerySource, &OutCells](const UWorldPartitionRuntimeCell* Cell)
-	{
-		OutCells.Add(Cell);
-		return true;
-	});
-
-	return !!OutCells.Num();
-}
-
-bool UWorldPartitionRuntimeHash::GetStreamingCells(const TArray<FWorldPartitionStreamingSource>& Sources, FStreamingSourceCells& OutActivateCells, FStreamingSourceCells& OutLoadCells) const
-{
-	ForEachStreamingCellsSources(Sources, [&OutActivateCells, &OutLoadCells](const UWorldPartitionRuntimeCell* Cell, EStreamingSourceTargetState TargetState)
-	{
-		switch (TargetState)
-		{
-		case EStreamingSourceTargetState::Loaded:
-			OutLoadCells.GetCells().Add(Cell);
-			break;
-		case EStreamingSourceTargetState::Activated:
-			OutActivateCells.GetCells().Add(Cell);
-			break;
-		}
-		return true;
-	});
-
-	return !!(OutActivateCells.Num() + OutLoadCells.Num());
-}
-
 bool UWorldPartitionRuntimeHash::IsCellRelevantFor(bool bClientOnlyVisible) const
 {
 	if (bClientOnlyVisible)
diff --git a/Engine/Source/Runtime/Engine/Public/WorldPartition/WorldPartitionEditorHash.h b/Engine/Source/Runtime/Engine/Public/WorldPartition/WorldPartitionEditorHash.h
--- a/Engine/Source/Runtime/Engine/Public/WorldPartition/WorldPartitionEditorHash.h
+++ b/Engine/Source/Runtime/Engine/Public/WorldPartition/WorldPartitionEditorHash.h
@@ -26,9 +26,6 @@ public:
 	ENGINE_API virtual void HashActor(FWorldPartitionHandle& InActorHandle) PURE_VIRTUAL(UWorldPartitionEditorHash::HashActor, ;);
 	ENGINE_API virtual void UnhashActor(FWorldPartitionHandle& InActorHandle) PURE_VIRTUAL(UWorldPartitionEditorHash::UnhashActor, ;);
 
-	UE_DEPRECATED(5.1, "Use version that takes FForEachIntersectingActorParams instead.")	
-	ENGINE_API int32 ForEachIntersectingActor(const FBox& Box, TFunctionRef<void(FWorldPartitionActorDesc*)> InOperation, bool bIncludeSpatiallyLoadedActors, bool bIncludeNonSpatiallyLoadedActors);
-
 	/* Struct of optional parameters passed to ForEachIntersectingActor. */
 	struct FForEachIntersectingActorParams
 	{
diff --git a/Engine/Source/Runtime/Engine/Public/WorldPartition/WorldPartitionHelpers.h b/Engine/Source/Runtime/Engine/Public/WorldPartition/WorldPartitionHelpers.h
--- a/Engine/Source/Runtime/Engine/Public/WorldPartition/WorldPartitionHelpers.h
+++ b/Engine/Source/Runtime/Engine/Public/WorldPartition/WorldPartitionHelpers.h
@@ -123,12 +123,6 @@ public:
 		TMap<FGuid, FWorldPartitionReference> ActorReferences;
 	};
 
-	UE_DEPRECATED(5.1, "ForEachActorWithLoading is deprecated, ForEachActorWithLoading with FForEachActorWithLoadingParams should be used instead.")
-	static ENGINE_API void ForEachActorWithLoading(UWorldPartition* WorldPartition, TSubclassOf<AActor> ActorClass, TFunctionRef<bool(const FWorldPartitionActorDesc*)> Func, TFunctionRef<void()> OnReleasingActorReferences = [](){}, bool bGCPerActor = false);
-
-	UE_DEPRECATED(5.1, "ForEachActorWithLoading is deprecated, ForEachActorWithLoading with FForEachActorWithLoadingParams should be used instead.")
-	static ENGINE_API void ForEachActorWithLoading(UWorldPartition* WorldPartition, const TArray<FGuid>& ActorGuids, TFunctionRef<bool(const FWorldPartitionActorDesc*)> Func, TFunctionRef<void()> OnReleasingActorReferences = [](){}, bool bGCPerActor = false);
-
 	static ENGINE_API void ForEachActorWithLoading(UWorldPartition* WorldPartition, TFunctionRef<bool(const FWorldPartitionActorDesc*)> Func, const FForEachActorWithLoadingParams& Params = FForEachActorWithLoadingParams());
 	static ENGINE_API void ForEachActorWithLoading(UWorldPartition* WorldPartition, TFunctionRef<bool(const FWorldPartitionActorDesc*)> Func, const FForEachActorWithLoadingParams& Params, FForEachActorWithLoadingResult& Result);
 	
diff --git a/Engine/Source/Runtime/Engine/Public/WorldPartition/WorldPartitionRuntimeHash.h b/Engine/Source/Runtime/Engine/Public/WorldPartition/WorldPartitionRuntimeHash.h
--- a/Engine/Source/Runtime/Engine/Public/WorldPartition/WorldPartitionRuntimeHash.h
+++ b/Engine/Source/Runtime/Engine/Public/WorldPartition/WorldPartitionRuntimeHash.h
@@ -139,16 +139,6 @@ public:
 		TSet<const UWorldPartitionRuntimeCell*> Cells;
 	};
 
-	// Deprecated streaming interface
-	UE_DEPRECATED(5.1, "GetAllStreamingCells is deprecated, use ForEachStreamingCells instead.")
-	ENGINE_API int32 GetAllStreamingCells(TSet<const UWorldPartitionRuntimeCell*>& Cells, bool bAllDataLayers = false, bool bDataLayersOnly = false, const TSet<FName>& InDataLayers = TSet<FName>()) const;
-
-	UE_DEPRECATED(5.1, "GetStreamingCells is deprecated, use ForEachStreamingCells instead.")
-	ENGINE_API bool GetStreamingCells(const FWorldPartitionStreamingQuerySource& QuerySource, TSet<const UWorldPartitionRuntimeCell*>& OutCells) const;
-
-	UE_DEPRECATED(5.1, "GetStreamingCells is deprecated, use ForEachStreamingCells instead.")
-	ENGINE_API bool GetStreamingCells(const TArray<FWorldPartitionStreamingSource>& Sources, UWorldPartitionRuntimeHash::FStreamingSourceCells& OutActivateCells, UWorldPartitionRuntimeHash::FStreamingSourceCells& OutLoadCells) const;
-
 	// Streaming interface
 	virtual void ForEachStreamingCells(TFunctionRef<bool(const UWorldPartitionRuntimeCell*)> Func) const {}
 	virtual void ForEachStreamingCellsQuery(const FWorldPartitionStreamingQuerySource& QuerySource, TFunctionRef<bool(const UWorldPartitionRuntimeCell*)> Func, FWorldPartitionQueryCache* QueryCache = nullptr) const {}
