Use inclusive range for bEnabledZCulling.

diff --git a/Engine/Source/Runtime/Engine/Private/WorldPartition/WorldPartitionRuntimeSpatialHash.cpp b/Engine/Source/Runtime/Engine/Private/WorldPartition/WorldPartitionRuntimeSpatialHash.cpp
--- a/Engine/Source/Runtime/Engine/Private/WorldPartition/WorldPartitionRuntimeSpatialHash.cpp
+++ b/Engine/Source/Runtime/Engine/Private/WorldPartition/WorldPartitionRuntimeSpatialHash.cpp
@@ -235,7 +235,7 @@ void FSpatialHashStreamingGrid::GetCells(const FWorldPartitionStreamingQuerySour
 					if (bEnableZCulling)
 					{
 						const FVector2D CellMinMaxZ(Cell->GetContentBounds().Min.Z, Cell->GetContentBounds().Max.Z);
-						bIncludeCell = TRange<double>(CellMinMaxZ.X, CellMinMaxZ.Y).Overlaps(TRange<double>(Shape.GetCenter().Z - Shape.GetRadius(), Shape.GetCenter().Z + Shape.GetRadius()));
+						bIncludeCell = TRange<double>::Inclusive(CellMinMaxZ.X, CellMinMaxZ.Y).Overlaps(TRange<double>::Inclusive(Shape.GetCenter().Z - Shape.GetRadius(), Shape.GetCenter().Z + Shape.GetRadius()));
 					}
 						
 					if (bIncludeCell)
@@ -317,7 +317,7 @@ void FSpatialHashStreamingGrid::GetCells(const TArray<FWorldPartitionStreamingSo
 						if (bEnableZCulling)
 						{
 							const FVector2D CellMinMaxZ(Cell->GetContentBounds().Min.Z, Cell->GetContentBounds().Max.Z);
-							bIncludeCell = TRange<double>(CellMinMaxZ.X, CellMinMaxZ.Y).Overlaps(TRange<double>(Shape.GetCenter().Z - Shape.GetRadius(), Shape.GetCenter().Z + Shape.GetRadius()));
+							bIncludeCell = TRange<double>::Inclusive(CellMinMaxZ.X, CellMinMaxZ.Y).Overlaps(TRange<double>::Inclusive(Shape.GetCenter().Z - Shape.GetRadius(), Shape.GetCenter().Z + Shape.GetRadius()));
 						}
 						
 						if (bIncludeCell)
