ISM creates a scene proxy even when no instances have been added.
With recent ISM changes, instances can be added without the need to recreate the scene proxy.

diff --git a/Engine/Source/Runtime/Engine/Private/InstancedStaticMesh.cpp b/Engine/Source/Runtime/Engine/Private/InstancedStaticMesh.cpp
--- a/Engine/Source/Runtime/Engine/Private/InstancedStaticMesh.cpp
+++ b/Engine/Source/Runtime/Engine/Private/InstancedStaticMesh.cpp
@@ -2204,8 +2204,6 @@ FPrimitiveSceneProxy* UInstancedStaticMeshComponent::CreateSceneProxy()
 #if WITH_EDITOR
 		bIsInstanceDataApplyCompleted && 
 #endif
-		// make sure we have instances
-		PerInstanceSMData.Num() > 0 &&
 		// make sure we have an actual static mesh
 		GetStaticMesh() &&
 		GetStaticMesh()->IsCompiling() == false &&
diff --git a/Engine/Source/Runtime/Engine/Public/MeshBatch.h b/Engine/Source/Runtime/Engine/Public/MeshBatch.h
--- a/Engine/Source/Runtime/Engine/Public/MeshBatch.h
+++ b/Engine/Source/Runtime/Engine/Public/MeshBatch.h
@@ -453,7 +453,7 @@ struct FMeshBatch
 	{
 		for (int32 ElementIdx = 0; ElementIdx < Elements.Num(); ++ElementIdx)
 		{
-			if (Elements[ElementIdx].GetNumPrimitives() > 0 || Elements[ElementIdx].IndirectArgsBuffer)
+			if (Elements[ElementIdx].GetNumPrimitives() > 0 || Elements[ElementIdx].IndirectArgsBuffer || Elements[ElementIdx].bFetchInstanceCountFromScene)
 			{
 				return true;
 			}
