Enable editor mode for level streaming instance when either in editor or if a user has set cvar to turn it on.

diff --git a/Engine/Source/Runtime/Engine/Private/LevelInstance/LevelInstanceLevelStreaming.cpp b/Engine/Source/Runtime/Engine/Private/LevelInstance/LevelInstanceLevelStreaming.cpp
--- a/Engine/Source/Runtime/Engine/Private/LevelInstance/LevelInstanceLevelStreaming.cpp
+++ b/Engine/Source/Runtime/Engine/Private/LevelInstance/LevelInstanceLevelStreaming.cpp
@@ -26,6 +26,8 @@
 #include "Misc/LazySingleton.h"
 #include "Misc/ScopeExit.h"
 #include "UObject/LinkerLoad.h"
+#include "Engine/GameEngine.h"
+#include "Streaming/LevelStreamingDelegates.h"
 
 static bool GDisableLevelInstanceEditorPartialLoading = false;
 FAutoConsoleVariableRef CVarDisableLevelInstanceEditorPartialLoading(
@@ -34,6 +36,14 @@ FAutoConsoleVariableRef CVarDisableLevelInstanceEditorPartialLoading(
 	TEXT("Allow disabling partial loading of level instances in the editor."),
 	ECVF_Default);
 
+static bool GForceEditorWorldMode = false;
+FAutoConsoleVariableRef CVarForceEditorWorldMode(
+	TEXT("LevelInstance.ForceEditorWorldMode"),
+	GForceEditorWorldMode,
+	TEXT("Allow -game instances to behave like an editor with temporary root object attached to instance. This will prevent HLOD from working in -game. This feature is only supported on non WP worlds."),
+	ECVF_Default);
+
+
 namespace FLevelInstanceLevelStreamingUtils
 {
 	static void MarkObjectsInPackageAsTransientAndNonTransactional(UPackage* InPackage)
@@ -74,6 +84,23 @@ ILevelInstanceInterface* ULevelStreamingLevelInstance::GetLevelInstance() const
 }
 
 #if WITH_EDITOR
+
+bool ULevelStreamingLevelInstance::IsEditorWorldMode() const
+{
+	bool bCanSupportForceEditorWorldMode = !GIsEditor;
+	if (GForceEditorWorldMode && !GIsEditor)
+	{
+		UWorld* OwningWorld = GetWorld();
+		check(OwningWorld);
+		// We do not support GForceEditorWorldMode in WP worlds.
+		if (OwningWorld->GetWorldPartition() != nullptr)
+		{
+			bCanSupportForceEditorWorldMode = false;
+		}
+	}
+	return (GForceEditorWorldMode && bCanSupportForceEditorWorldMode) || !GetWorld()->IsGameWorld();
+}
+
 TOptional<FFolder::FRootObject> ULevelStreamingLevelInstance::GetFolderRootObject() const
 {
 	if (ILevelInstanceInterface* LevelInstance = GetLevelInstance())
@@ -109,7 +136,7 @@ FBox ULevelStreamingLevelInstance::GetBounds() const
 
 void ULevelStreamingLevelInstance::OnLoadedActorsAddedToLevelPreEvent(const TArray<AActor*>& InActors)
 {
-	if (!GetWorld()->IsGameWorld())
+	if (IsEditorWorldMode())
 	{
 		if (ILevelInstanceInterface* LevelInstance = GetLevelInstance())
 		{
@@ -147,7 +174,7 @@ void ULevelStreamingLevelInstance::OnLoadedActorsAddedToLevelPreEvent(const TArr
 
 void ULevelStreamingLevelInstance::OnLoadedActorsAddedToLevelPostEvent(const TArray<AActor*>& InActors)
 {
-	if (!GetWorld()->IsGameWorld())
+	if (IsEditorWorldMode())
 	{
 		if (GetLoadedLevel()->bAreComponentsCurrentlyRegistered)
 		{
@@ -279,27 +306,11 @@ ULevelStreamingLevelInstance* ULevelStreamingLevelInstance::LoadInstance(ILevelI
 	if (bOutSuccess)
 	{
 		LevelStreaming->LevelInstanceID = LevelInstance->GetLevelInstanceID();
-		
+
 #if WITH_EDITOR
 		if (!World->IsGameWorld())
 		{
 			GEngine->BlockTillLevelStreamingCompleted(LevelInstanceActor->GetWorld());
-
-			if (ULevel* Level = LevelStreaming->GetLoadedLevel())
-			{
-				check(LevelStreaming->GetLevelStreamingState() == ELevelStreamingState::LoadedVisible);
-
-				// Create special actor that will handle selection and transform
-				LevelStreaming->LevelInstanceEditorInstanceActor = ALevelInstanceEditorInstanceActor::Create(LevelInstance, Level);
-
-				// Push editing state to child actors
-				LevelInstanceActor->PushLevelInstanceEditingStateToProxies(LevelInstanceActor->IsInEditLevelInstanceHierarchy());
-			}
-			else
-			{
-				// Failed to load package
-				return nullptr;
-			}
 		}
 #endif
 		return LevelStreaming;
@@ -308,17 +319,36 @@ ULevelStreamingLevelInstance* ULevelStreamingLevelInstance::LoadInstance(ILevelI
 	return nullptr;
 }
 
-void ULevelStreamingLevelInstance::UnloadInstance(ULevelStreamingLevelInstance* LevelStreaming)
+#if WITH_EDITOR
+void ULevelStreamingLevelInstance::OnLevelStreamingStateChanged(UWorld* InWorld, const ULevelStreaming* InLevelStreaming, ULevel* InLevelIfLoaded, ELevelStreamingState InPrevState, ELevelStreamingState InNewState)
 {
-	if (LevelStreaming->GetWorld()->IsGameWorld())
+	if (InNewState == ELevelStreamingState::LoadedVisible && InLevelStreaming == this)
 	{
-		LevelStreaming->SetShouldBeLoaded(false);
-		LevelStreaming->SetShouldBeVisible(false);
-		LevelStreaming->SetIsRequestingUnloadAndRemoval(true);
+		ILevelInstanceInterface* LevelInstance = GetLevelInstance();
+
+		ULevel* Level = GetLoadedLevel();
+		check(GetLevelStreamingState() == ELevelStreamingState::LoadedVisible);
+
+		// Create special actor that will handle selection and transform
+		LevelInstanceEditorInstanceActor = ALevelInstanceEditorInstanceActor::Create(LevelInstance, Level);
+
+		// Push editing state to child actors
+		AActor* LevelInstanceActor = CastChecked<AActor>(LevelInstance);
+		LevelInstanceActor->PushLevelInstanceEditingStateToProxies(LevelInstanceActor->IsInEditLevelInstanceHierarchy());
+
+		// Unregister
+		FLevelStreamingDelegates::OnLevelStreamingStateChanged.RemoveAll(this);
 	}
+}
+#endif
+
+void ULevelStreamingLevelInstance::UnloadInstance(ULevelStreamingLevelInstance* LevelStreaming)
+{
 #if WITH_EDITOR
-	else
+	if (LevelStreaming->IsEditorWorldMode())
 	{
+		FLevelStreamingDelegates::OnLevelStreamingStateChanged.RemoveAll(LevelStreaming);
+
 		ULevel* LoadedLevel = LevelStreaming->GetLoadedLevel();
 		LoadedLevel->OnLoadedActorAddedToLevelPreEvent.RemoveAll(LevelStreaming);
 		LoadedLevel->OnLoadedActorAddedToLevelPostEvent.RemoveAll(LevelStreaming);
@@ -341,7 +371,15 @@ void ULevelStreamingLevelInstance::UnloadInstance(ULevelStreamingLevelInstance*
 
 		LevelStreaming->GetWorld()->GetSubsystem<ULevelInstanceSubsystem>()->RemoveLevelsFromWorld({ LevelStreaming->GetLoadedLevel() }, bResetTrans);
 	}
-#endif 
+	else
+#endif
+	if (LevelStreaming->GetWorld()->IsGameWorld())
+	{
+		LevelStreaming->SetShouldBeLoaded(false);
+		LevelStreaming->SetShouldBeVisible(false);
+		LevelStreaming->SetIsRequestingUnloadAndRemoval(true);
+	}
+
 }
 
 void ULevelStreamingLevelInstance::OnLevelLoadedChanged(ULevel* InLevel)
@@ -351,8 +389,10 @@ void ULevelStreamingLevelInstance::OnLevelLoadedChanged(ULevel* InLevel)
 	if (ULevel* NewLoadedLevel = GetLoadedLevel())
 	{
 #if WITH_EDITOR
-		if (!GetWorld()->IsGameWorld())
+		if (IsEditorWorldMode())
 		{
+			FLevelStreamingDelegates::OnLevelStreamingStateChanged.AddUObject(this, &ULevelStreamingLevelInstance::OnLevelStreamingStateChanged);
+
 			// Most of the code here is meant to allow partial support for undo/redo of LevelInstance Instance Loading:
 			// by setting the objects RF_Transient and !RF_Transactional we can check when unloading if those flags
 			// have been changed and figure out if we need to clear the transaction buffer or not.
@@ -413,4 +453,4 @@ void ULevelStreamingLevelInstance::OnLevelLoadedChanged(ULevel* InLevel)
 #endif
 		}
 	}
-}
\ No newline at end of file
+}
diff --git a/Engine/Source/Runtime/Engine/Public/LevelInstance/LevelInstanceLevelStreaming.h b/Engine/Source/Runtime/Engine/Public/LevelInstance/LevelInstanceLevelStreaming.h
--- a/Engine/Source/Runtime/Engine/Public/LevelInstance/LevelInstanceLevelStreaming.h
+++ b/Engine/Source/Runtime/Engine/Public/LevelInstance/LevelInstanceLevelStreaming.h
@@ -43,6 +43,9 @@ private:
 	ENGINE_API virtual void OnLoadedActorsAddedToLevelPreEvent(const TArray<AActor*>& InActors);
 	ENGINE_API virtual void OnLoadedActorsAddedToLevelPostEvent(const TArray<AActor*>& InActors);
 	ENGINE_API virtual void OnLoadedActorsRemovedFromLevelPostEvent(const TArray<AActor*>& InActors);
+	ENGINE_API void OnLevelStreamingStateChanged(UWorld* InWorld, const ULevelStreaming* InLevelStreaming, ULevel* InLevelIfLoaded, ELevelStreamingState InPrevState, ELevelStreamingState InNewState);
+
+	bool IsEditorWorldMode() const;
 
 	TWeakObjectPtr<ALevelInstanceEditorInstanceActor> LevelInstanceEditorInstanceActor;
 
