ColorVertexBuffer: Fix clobbering of NumVertices when cooking client and server platforms in the same process: we should only modify data on this during Serialize if IsLoading.
#rnx

diff --git a/Engine/Source/Runtime/Engine/Private/Rendering/ColorVertexBuffer.cpp b/Engine/Source/Runtime/Engine/Private/Rendering/ColorVertexBuffer.cpp
--- a/Engine/Source/Runtime/Engine/Private/Rendering/ColorVertexBuffer.cpp
+++ b/Engine/Source/Runtime/Engine/Private/Rendering/ColorVertexBuffer.cpp
@@ -211,22 +211,28 @@ void FColorVertexBuffer::Serialize( FArchive& Ar, bool bNeedsCPUAccess )
 
 		if (!StripFlags.IsDataStrippedForServer() || Ar.IsCountingMemory())
 		{
-			if (VertexData != NULL)
+			if (VertexData != nullptr)
 			{
 				// Serialize the vertex data.
 				VertexData->Serialize(Ar);
+			}
+		}
 
-				if (VertexData->Num() > 0)
+		if (Ar.IsLoading())
+		{
+			if (!StripFlags.IsDataStrippedForServer())
+			{
+				if (VertexData != nullptr && VertexData->Num() > 0)
 				{
 					// Make a copy of the vertex data pointer.
 					Data = VertexData->GetDataPointer();
 				}
 			}
-		}
-		if (StripFlags.IsDataStrippedForServer())
-		{
-			// if we stripped all the other stuff and decided not to serialize it in probably need to strip the NumVertices Too
-			NumVertices = Stride = 0;
+			else
+			{
+				// if we stripped all the other stuff and decided not to serialize it in probably need to strip the NumVertices Too
+				NumVertices = Stride = 0;
+			}
 		}
 	}
 }
