Always duplicate inner objects when using GetInstancedSubObjectsRecursive, this matches the behavior of StaticDuplicateObject

diff --git a/Engine/Source/Runtime/Engine/Private/InstancedReferenceSubobjectHelper.cpp b/Engine/Source/Runtime/Engine/Private/InstancedReferenceSubobjectHelper.cpp
--- a/Engine/Source/Runtime/Engine/Private/InstancedReferenceSubobjectHelper.cpp
+++ b/Engine/Source/Runtime/Engine/Private/InstancedReferenceSubobjectHelper.cpp
@@ -261,6 +261,103 @@ void FFindInstancedReferenceSubobjectHelper::ForEachInstancedSubObject(FInstance
 	}
 }
 
+void FFindInstancedReferenceSubobjectHelper::ForEachSubObject(FInstancedPropertyPath& PropertyPath, const UObject* Outer, const void* ContainerAddress, TFunctionRef<void(const FInstancedSubObjRef& Ref)> ObjRefFunc)
+{
+	check(ContainerAddress && Outer);
+	const FProperty* TargetProp = PropertyPath.Head();
+
+	if (const FArrayProperty* ArrayProperty = CastField<const FArrayProperty>(TargetProp))
+	{
+		FScriptArrayHelper ArrayHelper(ArrayProperty, ContainerAddress);
+		for (int32 ElementIndex = 0; ElementIndex < ArrayHelper.Num(); ++ElementIndex)
+		{
+			const void* ValueAddress = ArrayHelper.GetRawPtr(ElementIndex);
+
+			PropertyPath.Push(ArrayProperty->Inner, ElementIndex);
+			ForEachSubObject(PropertyPath, Outer, ValueAddress, ObjRefFunc);
+			PropertyPath.Pop();
+		}
+	}
+	else if (const FMapProperty* MapProperty = CastField<const FMapProperty>(TargetProp))
+	{
+		// Exit now if the map doesn't contain any instanced references.
+		int32 LogicalIndex = 0;
+		FScriptMapHelper MapHelper(MapProperty, ContainerAddress);
+		for (int32 ElementIndex = 0; ElementIndex < MapHelper.GetMaxIndex(); ++ElementIndex)
+		{
+			if (MapHelper.IsValidIndex(ElementIndex))
+			{
+				const void* KeyAddress = MapHelper.GetKeyPtr(ElementIndex);
+				const void* ValueAddress = MapHelper.GetValuePtr(ElementIndex);
+
+				// Note: Keep these as the logical (Nth) index in case the map changes internally after we construct the path or in case we resolve using a different object.
+				PropertyPath.Push(MapProperty->KeyProp, LogicalIndex);
+				ForEachSubObject(PropertyPath, Outer, KeyAddress, ObjRefFunc);
+				PropertyPath.Pop();
+
+				PropertyPath.Push(MapProperty->ValueProp, LogicalIndex, true);
+				ForEachSubObject(PropertyPath, Outer, ValueAddress, ObjRefFunc);
+				PropertyPath.Pop();
+
+				++LogicalIndex;
+			}
+		}
+	}
+	else if (const FSetProperty* SetProperty = CastField<const FSetProperty>(TargetProp))
+	{
+		int32 LogicalIndex = 0;
+		FScriptSetHelper SetHelper(SetProperty, ContainerAddress);
+		for (int32 ElementIndex = 0; ElementIndex < SetHelper.GetMaxIndex(); ++ElementIndex)
+		{
+			if (SetHelper.IsValidIndex(ElementIndex))
+			{
+				const void* ValueAddress = SetHelper.GetElementPtr(ElementIndex);
+
+				// Note: Keep this as the logical (Nth) index in case the set changes internally after we construct the path or in case we resolve using a different object.
+				PropertyPath.Push(SetProperty->ElementProp, LogicalIndex);
+				ForEachSubObject(PropertyPath, Outer, ValueAddress, ObjRefFunc);
+				PropertyPath.Pop();
+
+				++LogicalIndex;
+			}
+		}
+	}
+	else if (const FOptionalProperty* OptionalProperty = CastField<FOptionalProperty>(TargetProp))
+	{
+		if (const void* ValueAddress = static_cast<const void*>(OptionalProperty->GetValuePointerForReadOrReplaceIfSet(ContainerAddress)))
+		{
+			PropertyPath.Push(OptionalProperty->GetValueProperty());
+			ForEachSubObject(PropertyPath, Outer, ValueAddress, ObjRefFunc);
+			PropertyPath.Pop();
+		}
+	}
+	else if (const FStructProperty* StructProperty = CastField<const FStructProperty>(TargetProp))
+	{
+		for (FProperty* StructProp = StructProperty->Struct->RefLink; StructProp; StructProp = StructProp->NextRef)
+		{
+			for (int32 ArrayIdx = 0; ArrayIdx < StructProp->ArrayDim; ++ArrayIdx)
+			{
+				const void* ValueAddress = StructProp->ContainerPtrToValuePtr<uint8>(ContainerAddress, ArrayIdx);
+
+				PropertyPath.Push(StructProp, ArrayIdx);
+				ForEachSubObject(PropertyPath, Outer, ValueAddress, ObjRefFunc);
+				PropertyPath.Pop();
+			}
+		}
+	}
+	else if (const FObjectProperty* ObjectProperty = CastField<const FObjectProperty>(TargetProp))
+	{
+		if (UObject* ObjectValue = ObjectProperty->GetObjectPropertyValue(ContainerAddress))
+		{
+			if (ObjectValue->GetOuter() == Outer)
+			{
+				// don't need to push to PropertyPath, since this property is already at its head
+				ObjRefFunc(FInstancedSubObjRef(ObjectValue, PropertyPath));
+			}
+		}
+	}
+}
+
 template ENGINE_API void FFindInstancedReferenceSubobjectHelper::ForEachInstancedSubObject<void*>(FInstancedPropertyPath& PropertyPath, void* ContainerAddress, TFunctionRef<void(const FInstancedSubObjRef&, void*)> ObjRefFunc);
 template ENGINE_API void FFindInstancedReferenceSubobjectHelper::ForEachInstancedSubObject<const void*>(FInstancedPropertyPath& PropertyPath, const void* ContainerAddress, TFunctionRef<void(const FInstancedSubObjRef&, const void*)> ObjRefFunc);
 
diff --git a/Engine/Source/Runtime/Engine/Public/InstancedReferenceSubobjectHelper.h b/Engine/Source/Runtime/Engine/Public/InstancedReferenceSubobjectHelper.h
--- a/Engine/Source/Runtime/Engine/Public/InstancedReferenceSubobjectHelper.h
+++ b/Engine/Source/Runtime/Engine/Public/InstancedReferenceSubobjectHelper.h
@@ -124,7 +124,7 @@ public:
 			{
 				FInstancedPropertyPath RootPropertyPath(Prop, ArrayIdx);
 				const uint8* ValuePtr = Prop->ContainerPtrToValuePtr<uint8>(Container, ArrayIdx);
-				ForEachInstancedSubObject<const void*>(RootPropertyPath, ValuePtr, [&OutObjects](const FInstancedSubObjRef& Ref, const void*)
+				ForEachSubObject(RootPropertyPath, Container, ValuePtr, [&OutObjects](const FInstancedSubObjRef& Ref)
 				{
 					if (!OutObjects.Contains(Ref))
 					{
@@ -146,4 +146,6 @@ public:
 
 	template<typename T>
 	static void ForEachInstancedSubObject(FInstancedPropertyPath& PropertyPath, T ContainerAddress, TFunctionRef<void(const FInstancedSubObjRef& Ref, T PropertyValueAddress)> ObjRefFunc);
+
+	static ENGINE_API void ForEachSubObject(FInstancedPropertyPath& PropertyPath, const UObject* Outer, const void* ContainerAddress, TFunctionRef<void(const FInstancedSubObjRef& Ref)> ObjRefFunc);
 };
