#ue5 - Add option to ProjectileMovementComponent to use scoped movement updates during simulation for cheaper movement/transform updates when processing multiple movements (substepping or deflections).

New flag: bSimulationUseScopedMovement (default false, as the delayed impact events could change behavior, use with care).

#tests yes

diff --git a/Engine/Source/Runtime/Engine/Classes/GameFramework/ProjectileMovementComponent.h b/Engine/Source/Runtime/Engine/Classes/GameFramework/ProjectileMovementComponent.h
--- a/Engine/Source/Runtime/Engine/Classes/GameFramework/ProjectileMovementComponent.h
+++ b/Engine/Source/Runtime/Engine/Classes/GameFramework/ProjectileMovementComponent.h
@@ -135,6 +135,13 @@ protected:
 
 public:
 
+	/**
+	 * If true, uses FScopedMovementUpdate to avoid moving the UpdatedComponent more than once during a tick during simulation.
+	 * This also defers overlap updates and some impact events until after the simulation update completes, so it may delay important events and continue deflection, so use with caution.
+	 */
+	UPROPERTY(EditAnywhere, BlueprintReadWrite, AdvancedDisplay, Category=ProjectileInterpolation)
+	uint8 bSimulationUseScopedMovement:1;
+
 	/**
 	 * If true, uses FScopedMovementUpdate to avoid moving the attached interpolated object's children more than once during a tick when it would both interpolate and move during projectile simulation.
 	 * This also defers overlap updates for the interpolated object until after the simulation update completes.
diff --git a/Engine/Source/Runtime/Engine/Private/Components/ProjectileMovementComponent.cpp b/Engine/Source/Runtime/Engine/Private/Components/ProjectileMovementComponent.cpp
--- a/Engine/Source/Runtime/Engine/Private/Components/ProjectileMovementComponent.cpp
+++ b/Engine/Source/Runtime/Engine/Private/Components/ProjectileMovementComponent.cpp
@@ -26,6 +26,7 @@ UProjectileMovementComponent::UProjectileMovementComponent(const FObjectInitiali
 	bInterpMovement = false;
 	bInterpRotation = false;
 	bInterpolationComplete = true;
+	bSimulationUseScopedMovement = false;
 	bInterpolationUseScopedMovement = true;
 	InterpLocationTime = 0.100f;
 	InterpRotationTime = 0.050f;
@@ -190,6 +191,9 @@ void UProjectileMovementComponent::TickComponent(float DeltaTime, enum ELevelTic
 	int32 Iterations = 0;
 	FHitResult Hit(1.f);
 	
+	QUICK_SCOPE_CYCLE_COUNTER( STAT_ProjectileMovementComponent_PerformMovement );
+	const FScopedMovementUpdate ScopedProjectileUpdate(bSimulationUseScopedMovement ? UpdatedComponent : nullptr, EScopedUpdate::DeferredUpdates);
+
 	while (bSimulationEnabled && RemainingTime >= MIN_TICK_TIME && (Iterations < MaxSimulationIterations) && IsValid(ActorOwner) && !HasStoppedSimulation())
 	{
 		LoopCount++;
