LevelInstance:
Allow Selection of Sub-Level Instance
Implement CanEditChangeComponent / IsActorLabelEditable / IsUserManaged / ShouldExport for selected non editable Level Instances

#rnx

diff --git a/Engine/Source/Runtime/Engine/Private/LevelInstance/LevelInstanceActor.cpp b/Engine/Source/Runtime/Engine/Private/LevelInstance/LevelInstanceActor.cpp
--- a/Engine/Source/Runtime/Engine/Private/LevelInstance/LevelInstanceActor.cpp
+++ b/Engine/Source/Runtime/Engine/Private/LevelInstance/LevelInstanceActor.cpp
@@ -224,6 +224,11 @@ bool ALevelInstance::CanEditChange(const FProperty* Property) const
 	return Super::CanEditChange(Property) && LevelInstanceActorImpl.CanEditChange(Property);
 }
 
+bool ALevelInstance::CanEditChangeComponent(const UActorComponent* InComponent, const FProperty* InProperty) const
+{
+	return Super::CanEditChangeComponent(InComponent, InProperty) && LevelInstanceActorImpl.CanEditChangeComponent(InComponent, InProperty);
+}
+
 bool ALevelInstance::ResolveSubobject(const TCHAR* SubObjectPath, UObject*& OutObject, bool bLoadIfExists)
 {
 	return LevelInstanceActorImpl.ResolveSubobject(SubObjectPath, OutObject, bLoadIfExists);
@@ -266,11 +271,26 @@ FBox ALevelInstance::GetStreamingBounds() const
 	return Super::GetStreamingBounds();
 }
 
+bool ALevelInstance::IsUserManaged() const
+{
+	return Super::IsUserManaged() || LevelInstanceActorImpl.IsUserManaged();
+}
+
+bool ALevelInstance::ShouldExport()
+{
+	return Super::ShouldExport() && LevelInstanceActorImpl.ShouldExport();
+}
+
 bool ALevelInstance::IsLockLocation() const
 {
 	return Super::IsLockLocation() || LevelInstanceActorImpl.IsLockLocation();
 }
 
+bool ALevelInstance::IsActorLabelEditable() const
+{
+	return Super::IsActorLabelEditable() && LevelInstanceActorImpl.IsActorLabelEditable();
+}
+
 bool ALevelInstance::GetReferencedContentObjects(TArray<UObject*>& Objects) const
 {
 	Objects.Add(const_cast<ALevelInstance*>(this));
diff --git a/Engine/Source/Runtime/Engine/Private/LevelInstance/LevelInstanceActorImpl.cpp b/Engine/Source/Runtime/Engine/Private/LevelInstance/LevelInstanceActorImpl.cpp
--- a/Engine/Source/Runtime/Engine/Private/LevelInstance/LevelInstanceActorImpl.cpp
+++ b/Engine/Source/Runtime/Engine/Private/LevelInstance/LevelInstanceActorImpl.cpp
@@ -200,21 +200,6 @@ void FLevelInstanceActorImpl::PostEditImport(TFunctionRef<void()> SuperCall)
 	LevelInstance->UpdateLevelInstanceFromWorldAsset();
 }
 
-bool FLevelInstanceActorImpl::CanEditChange(const FProperty* Property) const
-{
-	if (LevelInstance->IsEditing())
-	{
-		return false;
-	}
-
-	if (LevelInstance->HasDirtyChildren())
-	{
-		return false;
-	}
-
-	return true;
-}
-
 void FLevelInstanceActorImpl::PreEditChange(FProperty* Property, bool bWorldAssetChange, TFunctionRef<void(FProperty*)> SuperCall)
 {
 	{
@@ -317,9 +302,55 @@ void FLevelInstanceActorImpl::EditorGetUnderlyingActors(TSet<AActor*>& OutUnderl
 	}
 }
 
+bool FLevelInstanceActorImpl::IsLockedActor() const
+{
+	AActor* LevelInstanceActor = CastChecked<AActor>(LevelInstance);
+	if (LevelInstanceActor->IsInLevelInstance() && !LevelInstanceActor->IsInEditLevelInstance())
+	{
+		return true;
+	}
+
+	if (LevelInstance->IsEditing())
+	{
+		return true;
+	}
+
+	if (LevelInstance->HasChildEdit())
+	{
+		return true;
+	}
+
+	return false;
+}
+
+bool FLevelInstanceActorImpl::ShouldExport() const
+{
+	return !IsLockedActor();
+}
+
+bool FLevelInstanceActorImpl::IsUserManaged() const
+{
+	return !IsLockedActor();
+}
+
 bool FLevelInstanceActorImpl::IsLockLocation() const
 {
-	return LevelInstance->IsEditing() || LevelInstance->HasChildEdit();
+	return IsLockedActor();
+}
+
+bool FLevelInstanceActorImpl::IsActorLabelEditable() const
+{
+	return !IsLockedActor();
+}
+
+bool FLevelInstanceActorImpl::CanEditChange(const FProperty* Property) const
+{
+	return !IsLockedActor();
+}
+
+bool FLevelInstanceActorImpl::CanEditChangeComponent(const UActorComponent* Component, const FProperty* InProperty) const
+{
+	return !IsLockedActor();
 }
 
 bool FLevelInstanceActorImpl::GetBounds(FBox& OutBounds) const
diff --git a/Engine/Source/Runtime/Engine/Public/LevelInstance/LevelInstanceActor.h b/Engine/Source/Runtime/Engine/Public/LevelInstance/LevelInstanceActor.h
--- a/Engine/Source/Runtime/Engine/Public/LevelInstance/LevelInstanceActor.h
+++ b/Engine/Source/Runtime/Engine/Public/LevelInstance/LevelInstanceActor.h
@@ -85,6 +85,7 @@ public:
 	ENGINE_API virtual bool ResolveSubobject(const TCHAR* SubObjectPath, UObject*& OutObject, bool bLoadIfExists) override;
 
 	// AActor overrides
+	ENGINE_API virtual bool CanEditChangeComponent(const UActorComponent* InComponent, const FProperty* InProperty) const override;
 	ENGINE_API virtual void CheckForErrors() override;
 	ENGINE_API virtual TUniquePtr<class FWorldPartitionActorDesc> CreateClassActorDesc() const override;
 	ENGINE_API virtual bool CanDeleteSelectedActor(FText& OutReason) const override;
@@ -96,10 +97,14 @@ public:
 	ENGINE_API virtual FBox GetComponentsBoundingBox(bool bNonColliding = false, bool bIncludeFromChildActors = false) const override;
 	ENGINE_API virtual FBox GetStreamingBounds() const override;
 	ENGINE_API virtual bool IsLockLocation() const override;
+	ENGINE_API virtual bool IsActorLabelEditable() const override;
 	ENGINE_API virtual bool GetReferencedContentObjects(TArray<UObject*>& Objects) const override;
 	ENGINE_API virtual bool GetSoftReferencedContentObjects(TArray<FSoftObjectPath>& SoftObjects) const override;
 	ENGINE_API virtual bool OpenAssetEditor() override;
-	ENGINE_API virtual bool EditorCanAttachFrom(const AActor* InChild, FText& OutReason) const override;
+	ENGINE_API virtual bool EditorCanAttachFrom(const AActor* InChild, FText& OutReason) const override;	
+	ENGINE_API virtual bool IsUserManaged() const override;
+	ENGINE_API virtual bool ShouldExport() override;
+	ENGINE_API virtual bool SupportsSubRootSelection() const override { return true; }
 	// End of AActor interface
 
 	DECLARE_MULTICAST_DELEGATE_OneParam(FOnLevelInstanceActorPostLoad, ALevelInstance*);
diff --git a/Engine/Source/Runtime/Engine/Public/LevelInstance/LevelInstanceActorImpl.h b/Engine/Source/Runtime/Engine/Public/LevelInstance/LevelInstanceActorImpl.h
--- a/Engine/Source/Runtime/Engine/Public/LevelInstance/LevelInstanceActorImpl.h
+++ b/Engine/Source/Runtime/Engine/Public/LevelInstance/LevelInstanceActorImpl.h
@@ -81,6 +81,7 @@ public:
 	ENGINE_API virtual void PostEditUndo(TSharedPtr<ITransactionObjectAnnotation> TransactionAnnotation, TFunctionRef<void(TSharedPtr<ITransactionObjectAnnotation> TransactionAnnotation)> SuperCall);
 	ENGINE_API virtual void PostEditImport(TFunctionRef<void()> SuperCall);
 	ENGINE_API virtual bool CanEditChange(const FProperty* Property) const;
+	ENGINE_API virtual bool CanEditChangeComponent(const UActorComponent* Component, const FProperty* InProperty) const;
 	ENGINE_API virtual void PreEditChange(FProperty* Property, bool bWorldAssetChange, TFunctionRef<void(FProperty*)> SuperCall);
 	ENGINE_API virtual void PostEditChangeProperty(FPropertyChangedEvent& PropertyChangedEvent, bool bWorldAssetChange, TFunction<void(FPropertyChangedEvent&)> SuperCall);
 	ENGINE_API virtual bool CanDeleteSelectedActor(FText& OutReason) const;
@@ -88,12 +89,16 @@ public:
 	ENGINE_API virtual bool SetIsHiddenEdLayer(bool bIsHiddenEdLayer, TFunctionRef<bool(bool)> SuperCall);
 	ENGINE_API virtual void EditorGetUnderlyingActors(TSet<AActor*>& OutUnderlyingActors) const;
 	ENGINE_API virtual bool IsLockLocation() const;
+	ENGINE_API virtual bool IsActorLabelEditable() const;
+	ENGINE_API virtual bool IsUserManaged() const;
+	ENGINE_API virtual bool ShouldExport() const;
 	ENGINE_API virtual bool GetBounds(FBox& OutBounds) const;
 	ENGINE_API virtual void PushSelectionToProxies();
 	ENGINE_API virtual void PushLevelInstanceEditingStateToProxies(bool bInEditingState);
 	ENGINE_API virtual void CheckForErrors();
 	// End AActor Implementation
 private:
+	ENGINE_API virtual bool IsLockedActor() const;
 	ENGINE_API virtual void PostEditUndoInternal();
 #endif
 };
