World Partition Runtime HashSet:
- Fix ForEachStreamingCells incorrectly applying cells filtering.
- Fix FlushStreaming invalid asserts when running maps with content bundles.

diff --git a/Engine/Source/Runtime/Engine/Private/WorldPartition/RuntimeHashSet/WorldPartitionRuntimeHashSet.cpp b/Engine/Source/Runtime/Engine/Private/WorldPartition/RuntimeHashSet/WorldPartitionRuntimeHashSet.cpp
--- a/Engine/Source/Runtime/Engine/Private/WorldPartition/RuntimeHashSet/WorldPartitionRuntimeHashSet.cpp
+++ b/Engine/Source/Runtime/Engine/Private/WorldPartition/RuntimeHashSet/WorldPartitionRuntimeHashSet.cpp
@@ -93,13 +93,7 @@ void UWorldPartitionRuntimeHashSet::SetDefaultValues()
 void UWorldPartitionRuntimeHashSet::FlushStreaming()
 {
 	Super::FlushStreaming();
-	
-	check(PersistentPartitionDesc.Class);
-	PersistentPartitionDesc.Class = nullptr;
-	PersistentPartitionDesc.Name = NAME_None;
-	PersistentPartitionDesc.MainLayer = nullptr;
-
-	check(!RuntimeStreamingData.IsEmpty());
+	check(!PersistentPartitionDesc.Class);
 	RuntimeStreamingData.Empty();
 }
 
@@ -293,10 +287,7 @@ void UWorldPartitionRuntimeHashSet::ForEachStreamingCells(TFunctionRef<bool(cons
 	{
 		for (UWorldPartitionRuntimeCell* Cell : InCells)
 		{
-			if (IsCellRelevantFor(Cell->GetClientOnlyVisible()))
-			{
-				Func(Cell);
-			}
+			Func(Cell);
 		}
 	};
 
diff --git a/Engine/Source/Runtime/Engine/Private/WorldPartition/RuntimeHashSet/WorldPartitionRuntimeHashSetStreamingGeneration.cpp b/Engine/Source/Runtime/Engine/Private/WorldPartition/RuntimeHashSet/WorldPartitionRuntimeHashSetStreamingGeneration.cpp
--- a/Engine/Source/Runtime/Engine/Private/WorldPartition/RuntimeHashSet/WorldPartitionRuntimeHashSetStreamingGeneration.cpp
+++ b/Engine/Source/Runtime/Engine/Private/WorldPartition/RuntimeHashSet/WorldPartitionRuntimeHashSetStreamingGeneration.cpp
@@ -156,6 +156,14 @@ bool UWorldPartitionRuntimeHashSet::GenerateStreaming(UWorldPartitionStreamingPo
 	PersistentPartitionDesc.MainLayer = NewObject<URuntimePartition>(this, URuntimePartitionPersistent::StaticClass(), NAME_None);
 	PersistentPartitionDesc.MainLayer->Name = NAME_PersistentLevel;
 
+	ON_SCOPE_EXIT
+	{
+		check(PersistentPartitionDesc.Class);
+		PersistentPartitionDesc.Class = nullptr;
+		PersistentPartitionDesc.Name = NAME_None;
+		PersistentPartitionDesc.MainLayer = nullptr;
+	};
+
 	//
 	// Generate runtime partitions streaming cell desccriptors
 	//
