Fix performance issues in linked anim layers that use animation attributes

Anim graph debug data was only being reset on root graph execution, so layers leaked debug data each frame

diff --git a/Engine/Source/Runtime/Engine/Private/Animation/AnimInstanceProxy.cpp b/Engine/Source/Runtime/Engine/Private/Animation/AnimInstanceProxy.cpp
--- a/Engine/Source/Runtime/Engine/Private/Animation/AnimInstanceProxy.cpp
+++ b/Engine/Source/Runtime/Engine/Private/Animation/AnimInstanceProxy.cpp
@@ -477,6 +477,11 @@ void FAnimInstanceProxy::PreUpdate(UAnimInstance* InAnimInstance, float DeltaSec
 #if WITH_EDITORONLY_DATA
 	bIsGameWorld = World ? World->IsGameWorld() : false;
 
+	UpdatedNodesThisFrame.Reset();
+	NodeInputAttributesThisFrame.Reset();
+	NodeOutputAttributesThisFrame.Reset();
+	NodeSyncsThisFrame.Reset();
+
 	if (FAnimBlueprintDebugData* DebugData = GetAnimBlueprintDebugData())
 	{
 		DebugData->ResetNodeVisitSites();
@@ -1210,12 +1215,6 @@ void FAnimInstanceProxy::UpdateAnimation_WithRoot(const FAnimationUpdateContext&
 
 	if(InRootNode == RootNode)
 	{
-#if WITH_EDITORONLY_DATA
-	    UpdatedNodesThisFrame.Reset();
-	    NodeInputAttributesThisFrame.Reset();
-	    NodeOutputAttributesThisFrame.Reset();
-	    NodeSyncsThisFrame.Reset();
-#endif
 		if(bInitializeSubsystems && AnimClassInterface)
 		{
 			AnimClassInterface->ForEachSubsystem(GetAnimInstanceObject(), [this](const FAnimSubsystemInstanceContext& InContext)
