WorldPartition: Add GetSortedAssets function for the cases where we will use the ordering in cooked output data.
#rnx

diff --git a/Engine/Source/Runtime/Engine/Private/ExternalPackageHelper.cpp b/Engine/Source/Runtime/Engine/Private/ExternalPackageHelper.cpp
--- a/Engine/Source/Runtime/Engine/Private/ExternalPackageHelper.cpp
+++ b/Engine/Source/Runtime/Engine/Private/ExternalPackageHelper.cpp
@@ -4,6 +4,7 @@
 
 #if WITH_EDITOR
 
+#include "AssetRegistry/IAssetRegistry.h"
 #include "Misc/ArchiveMD5.h"
 #include "Misc/Paths.h"
 #include "HAL/PlatformApplicationMisc.h"
@@ -178,4 +179,12 @@ void FExternalPackageHelper::CopyObjectsExternalPackageFilePathToClipboard(const
 	}
 }
 
+void FExternalPackageHelper::GetSortedAssets(const FARFilter& Filter, TArray<FAssetData>& OutAssets)
+{
+	OutAssets.Reset();
+	IAssetRegistry::GetChecked().GetAssets(Filter, OutAssets);
+	OutAssets.Sort();
+}
+
+
 #endif
diff --git a/Engine/Source/Runtime/Engine/Private/Level.cpp b/Engine/Source/Runtime/Engine/Private/Level.cpp
--- a/Engine/Source/Runtime/Engine/Private/Level.cpp
+++ b/Engine/Source/Runtime/Engine/Private/Level.cpp
@@ -3965,10 +3965,7 @@ TArray<FString> ULevel::GetOnDiskExternalActorPackages(const FString& ExternalAc
 		TArray<FAssetData> ActorAssets;
 		IAssetRegistry& AssetRegistry = FModuleManager::LoadModuleChecked<FAssetRegistryModule>(TEXT("AssetRegistry")).Get();
 		AssetRegistry.ScanSynchronous({ ExternalActorsPath }, TArray<FString>());
-		AssetRegistry.GetAssets(Filter, ActorAssets);
-		// Sort the list as the order affects the cooking which needs to be deterministic.
-		// @todo_ow: Remove this once UE-198035 is fixed.
-		ActorAssets.Sort();
+		FExternalPackageHelper::GetSortedAssets(Filter, ActorAssets);
 
 		ActorPackageNames.Reserve(ActorAssets.Num());
 		Algo::Transform(ActorAssets, ActorPackageNames, [](const FAssetData& ActorAssetData) { return ActorAssetData.PackageName.ToString(); });
diff --git a/Engine/Source/Runtime/Engine/Private/WorldPartition/ActorDescContainer.cpp b/Engine/Source/Runtime/Engine/Private/WorldPartition/ActorDescContainer.cpp
--- a/Engine/Source/Runtime/Engine/Private/WorldPartition/ActorDescContainer.cpp
+++ b/Engine/Source/Runtime/Engine/Private/WorldPartition/ActorDescContainer.cpp
@@ -8,6 +8,7 @@
 #include "Editor.h"
 #include "Algo/Transform.h"
 #include "Engine/Level.h"
+#include "ExternalPackageHelper.h"
 #include "UObject/ObjectSaveContext.h"
 #include "AssetRegistry/ARFilter.h"
 #include "AssetRegistry/AssetRegistryModule.h"
@@ -58,10 +59,7 @@ void UActorDescContainer::Initialize(const FInitializeParams& InitParams)
 		Filter.PackagePaths.Add(*ContainerExternalActorsPath);
 
 		TRACE_CPUPROFILER_EVENT_SCOPE(GetAssets);
-		AssetRegistry.GetAssets(Filter, Assets);
-		// Sort the list as the order affects the cooking which needs to be deterministic.
-		// @todo_ow: Remove this once UE-198035 is fixed.
-		Assets.Sort();
+		FExternalPackageHelper::GetSortedAssets(Filter, Assets);
 	}
 
 	FWorldPartitionClassDescRegistry& ClassDescRegistry = FWorldPartitionClassDescRegistry::Get();
diff --git a/Engine/Source/Runtime/Engine/Private/WorldPartition/DataLayer/WorldDataLayersActorDesc.cpp b/Engine/Source/Runtime/Engine/Private/WorldPartition/DataLayer/WorldDataLayersActorDesc.cpp
--- a/Engine/Source/Runtime/Engine/Private/WorldPartition/DataLayer/WorldDataLayersActorDesc.cpp
+++ b/Engine/Source/Runtime/Engine/Private/WorldPartition/DataLayer/WorldDataLayersActorDesc.cpp
@@ -279,10 +279,7 @@ const TArray<FDataLayerInstanceDesc>& FWorldDataLayersActorDesc::GetExternalPack
 			Filter.bRecursiveClasses = true;
 			Filter.PackagePaths.Add(*ExternalObjectsPath);
 			TArray<FAssetData> Assets;
-			AssetRegistry.GetAssets(Filter, Assets);
-			// Sort the list as the order affects the cooking which needs to be deterministic.
-			// @todo_ow: Remove this once UE-198035 is fixed.
-			Assets.Sort();
+			FExternalPackageHelper::GetSortedAssets(Filter, Assets);
 
 			for (const FAssetData& Asset : Assets)
 			{
diff --git a/Engine/Source/Runtime/Engine/Public/ExternalPackageHelper.h b/Engine/Source/Runtime/Engine/Public/ExternalPackageHelper.h
--- a/Engine/Source/Runtime/Engine/Public/ExternalPackageHelper.h
+++ b/Engine/Source/Runtime/Engine/Public/ExternalPackageHelper.h
@@ -92,6 +92,10 @@ public:
 	 */
 	static ENGINE_API void CopyObjectsExternalPackageFilePathToClipboard(const TArray<const UObject*>& InObjects);
 	
+	/**
+	 * Call AssetRegistry.GetAssets and sort the results for deterministic use in cooked data.
+	 */
+	static ENGINE_API void GetSortedAssets(const FARFilter& Filter, TArray<FAssetData>& OutAssets);
 private:
 	/** Get the external object package instance name. */
 	static ENGINE_API FString GetExternalObjectPackageInstanceName(const FString& OuterPackageName, const FString& ObjectPackageName);
@@ -116,10 +120,7 @@ void FExternalPackageHelper::LoadObjectsFromExternalPackages(UObject* InOuter, T
 	Filter.bRecursiveClasses = true;
 	Filter.PackagePaths.Add(*ExternalObjectsPath);
 	TArray<FAssetData> Assets;
-	AssetRegistry.GetAssets(Filter, Assets);
-	// Sort the list as the order affects the cooking which needs to be deterministic.
-	// @todo_ow: Remove this once UE-198035 is fixed.
-	Assets.Sort();
+	GetSortedAssets(Filter, Assets);
 
 	ObjectPackageNames.Reserve(Assets.Num());
 	for (const FAssetData& Asset : Assets)
