ISM has option to use a conservative bounds.
This will not iterate over all instances to recalculate the bounds on every instance add/remove/transform update. But instead will maintain a conservative bounds that can only grow.
This is an optimization for the case where we might have very high instance counts causing the bounds calculation to become very expensive.
It's off by default, and controlled by SetUseConservativeBounds().
Even when the option is turned on for an ISM, we keep the existing behavior whenever the number of instances is below r.InstancedStaticMeshes.ConservativeBounds.Threshold

diff --git a/Engine/Source/Runtime/Engine/Classes/Components/InstancedStaticMeshComponent.h b/Engine/Source/Runtime/Engine/Classes/Components/InstancedStaticMeshComponent.h
--- a/Engine/Source/Runtime/Engine/Classes/Components/InstancedStaticMeshComponent.h
+++ b/Engine/Source/Runtime/Engine/Classes/Components/InstancedStaticMeshComponent.h
@@ -369,6 +369,13 @@ class UInstancedStaticMeshComponent : public UStaticMeshComponent, public ISMIns
 	/** Returns true if RemoveAtSwap is enabled. Derived classes may always return true regardless of whether SetRemoveSwapEnabled() was called. */
 	ENGINE_API virtual bool SupportsRemoveSwap() const;
 
+	/** 
+	 * Sets whether to use conservative bounds. 
+	 * This doesn't fully recalculate bounds on instance addition/removal/transform change.
+	 * Instead maintains a conservative bound that only grows.
+	 */
+	void SetUseConservativeBounds(bool bValue) { bUseConservativeBounds = bValue; CachedConservativeInstanceBounds.Init(); }
+
 #if WITH_EDITOR
 	ENGINE_API virtual bool ComponentIsTouchingSelectionBox(const FBox& InSelBBox, const bool bConsiderOnlyBSP, const bool bMustEncompassEntireComponent) const override;
 	ENGINE_API virtual bool ComponentIsTouchingSelectionFrustum(const FConvexVolume& InFrustum, const bool bConsiderOnlyBSP, const bool bMustEncompassEntireComponent) const override;
@@ -382,7 +389,13 @@ class UInstancedStaticMeshComponent : public UStaticMeshComponent, public ISMIns
 
 private:
 	bool bHasPreviousTransforms = false;
-	/** 
+
+	/** Flag for whether we are using conservative bounds. */
+	bool bUseConservativeBounds = false;
+	/** Current cached conservativ bounds. */
+	FBox CachedConservativeInstanceBounds;
+
+	/**
 	 *  Buffers with per-instance data laid out for rendering. 
 	 *  Serialized for cooked content. Used to create PerInstanceRenderData. 
 	 *  Alive between Serialize and PostLoad calls 
@@ -431,6 +444,8 @@ public:
 	ENGINE_API virtual bool CanEditSimulatePhysics() override;
 
 	ENGINE_API virtual FBoxSphereBounds CalcBounds(const FTransform& BoundTransform) const override;
+	ENGINE_API virtual void UpdateBounds() override;
+
 	virtual bool SupportsStaticLighting() const override { return true; }
 #if WITH_EDITOR
 	ENGINE_API virtual void GetStaticLightingInfo(FStaticLightingPrimitiveInfo& OutPrimitiveInfo,const TArray<ULightComponent*>& InRelevantLights,const FLightingBuildOptions& Options) override;
diff --git a/Engine/Source/Runtime/Engine/Private/InstancedStaticMesh.cpp b/Engine/Source/Runtime/Engine/Private/InstancedStaticMesh.cpp
--- a/Engine/Source/Runtime/Engine/Private/InstancedStaticMesh.cpp
+++ b/Engine/Source/Runtime/Engine/Private/InstancedStaticMesh.cpp
@@ -178,6 +178,12 @@ static TAutoConsoleVariable<int32> CVarISMFetchInstanceCountFromScene(
 	1,
 	TEXT("Enables the data path that allows instance count to be fetched from the Scene rather than the Mesh Draw Commands (MDCs), which removes the need to re-cache MDCs when instance count changes."));
 
+static int32 GConservativeBoundsThreshold = 30;
+FAutoConsoleVariableRef CVarISMConservativeBoundsThreshold(
+	TEXT("r.InstancedStaticMeshes.ConservativeBounds.Threshold"),
+	GConservativeBoundsThreshold,
+	TEXT("Number of instances in an ISM before we start using conservative bounds. Set to -1 to disable conservative bounds."));
+
 class FISMExecHelper : public FSelfRegisteringExec
 {
 	virtual bool Exec_Runtime(class UWorld* InWorld, const TCHAR* Cmd, FOutputDevice& Ar)
@@ -2409,7 +2415,31 @@ void UInstancedStaticMeshComponent::BuildComponentInstanceData(FInstanceUpdateCo
 	OutData.BuildChangeSet = [&](FISMInstanceUpdateChangeSet &ChangeSet)
 	{
 		BuildInstanceDataDeltaChangeSetCommon(ChangeSet);
-		ChangeSet.SetInstanceTransforms(MakeStridedView(PerInstanceSMData, &FInstancedStaticMeshInstanceData::Transform));
+
+		const bool bUpdateConservativeBounds = bUseConservativeBounds && GConservativeBoundsThreshold >= 0 && PerInstanceSMData.Num() > GConservativeBoundsThreshold;
+		if (bUpdateConservativeBounds)
+		{
+			if (ChangeSet.IsFullUpdate())
+			{
+				CachedConservativeInstanceBounds.Init();
+			}
+			else if (!CachedConservativeInstanceBounds.IsValid && !Bounds.GetBox().GetSize().IsZero())
+			{
+				// Initialize to current bounds in local space.
+				CachedConservativeInstanceBounds = Bounds.GetBox().TransformBy(GetComponentTransform().Inverse());
+			}
+
+			// Use the variation of SetInstanceTransforms() that updates the conservative bounds.
+			// Do this here, rather than in a separate loop, to take advantage of the fact that we are already iterating over the transform data.
+			ChangeSet.SetInstanceTransforms(MakeStridedView(PerInstanceSMData, &FInstancedStaticMeshInstanceData::Transform), GetStaticMesh()->GetBounds().GetBox(), CachedConservativeInstanceBounds);
+		}
+		else
+		{
+			CachedConservativeInstanceBounds.Init();
+
+			ChangeSet.SetInstanceTransforms(MakeStridedView(PerInstanceSMData, &FInstancedStaticMeshInstanceData::Transform));
+		}
+
 		ChangeSet.SetInstancePrevTransforms(MakeArrayView(PerInstancePrevTransform));
 	};
 }
@@ -2622,6 +2652,18 @@ FBoxSphereBounds UInstancedStaticMeshComponent::CalcBoundsImpl(const FTransform&
 	return FBoxSphereBounds(BoundTransform.GetLocation(), FVector::ZeroVector, 0.f);
 }
 
+void UInstancedStaticMeshComponent::UpdateBounds()
+{
+	if (bUseConservativeBounds && CachedConservativeInstanceBounds.IsValid)
+	{
+		Bounds = CachedConservativeInstanceBounds.TransformBy(GetComponentTransform());
+	}
+	else
+	{
+		Super::UpdateBounds();
+	}
+}
+
 #if WITH_EDITOR
 void UInstancedStaticMeshComponent::GetStaticLightingInfo(FStaticLightingPrimitiveInfo& OutPrimitiveInfo, const TArray<ULightComponent*>& InRelevantLights, const FLightingBuildOptions& Options)
 {
diff --git a/Engine/Source/Runtime/Engine/Public/InstancedStaticMesh/ISMInstanceUpdateChangeSet.h b/Engine/Source/Runtime/Engine/Public/InstancedStaticMesh/ISMInstanceUpdateChangeSet.h
--- a/Engine/Source/Runtime/Engine/Public/InstancedStaticMesh/ISMInstanceUpdateChangeSet.h
+++ b/Engine/Source/Runtime/Engine/Public/InstancedStaticMesh/ISMInstanceUpdateChangeSet.h
@@ -197,6 +197,16 @@ public:
 		TransformsDelta.GatherTransform(Transforms, InInstanceTransforms, [](const FMatrix &M) -> FRenderTransform { return FRenderTransform(M); });
 	}
 
+	inline void SetInstanceTransforms(TStridedView<FMatrix> InInstanceTransforms, FBox const& InInstanceBounds, FBox& OutGatheredBounds)
+	{
+		TransformsDelta.GatherTransform(Transforms, InInstanceTransforms, [&InInstanceBounds, &OutGatheredBounds](const FMatrix& M) -> FRenderTransform
+		{
+			FRenderTransform Transform(M);
+			OutGatheredBounds += InInstanceBounds.TransformBy(Transform.ToMatrix());
+			return Transform;
+		});
+	}
+
 	ENGINE_API void SetInstancePrevTransforms(TArrayView<FMatrix> InPrevInstanceTransforms, const FVector &Offset);
 	ENGINE_API void SetInstancePrevTransforms(TArrayView<FMatrix> InPrevInstanceTransforms);
 
diff --git a/Engine/Source/Runtime/Experimental/GeometryCollectionEngine/Private/GeometryCollection/GeometryCollectionISMPoolComponent.cpp b/Engine/Source/Runtime/Experimental/GeometryCollectionEngine/Private/GeometryCollection/GeometryCollectionISMPoolComponent.cpp
--- a/Engine/Source/Runtime/Experimental/GeometryCollectionEngine/Private/GeometryCollection/GeometryCollectionISMPoolComponent.cpp
+++ b/Engine/Source/Runtime/Experimental/GeometryCollectionEngine/Private/GeometryCollection/GeometryCollectionISMPoolComponent.cpp
@@ -137,6 +137,8 @@ void FGeometryCollectionISM::InitISM(const FGeometryCollectionStaticMeshInstance
 	ISMComponent->bOverrideMinLOD = MeshInstance.Desc.MinLod > 0;
 	ISMComponent->MinLOD = MeshInstance.Desc.MinLod;
 	ISMComponent->SetLODDistanceScale(MeshInstance.Desc.LodScale);
+	ISMComponent->SetUseConservativeBounds(true);
+	ISMComponent->bComputeFastLocalBounds = true;
 	ISMComponent->SetMeshDrawCommandStatsCategory(MeshInstance.Desc.StatsCategory);
 	ISMComponent->ComponentTags = MeshInstance.Desc.Tags;
 }
