Added cvar ('r.InstancedStaticMeshes.AllowCreateEmpty', off by default) to make it possible to create ISMs with zero instances. Disabling this behavior is needed until the ISM setup logic can be fixed.

diff --git a/Engine/Source/Runtime/Engine/Private/InstancedStaticMesh.cpp b/Engine/Source/Runtime/Engine/Private/InstancedStaticMesh.cpp
--- a/Engine/Source/Runtime/Engine/Private/InstancedStaticMesh.cpp
+++ b/Engine/Source/Runtime/Engine/Private/InstancedStaticMesh.cpp
@@ -87,6 +87,11 @@ TAutoConsoleVariable<int32> CVarGpuLodSelection(
 	FConsoleVariableDelegate::CreateLambda([](IConsoleVariable* InVariable) { FGlobalComponentRecreateRenderStateContext Context; }),
 	ECVF_RenderThreadSafe);
 
+static TAutoConsoleVariable<int32> CVarAllowCreateEmptyISMs(
+	TEXT("r.InstancedStaticMeshes.AllowCreateEmpty"),
+	0,
+	TEXT("Whether to allow creation of empty ISMS."));
+
 TAutoConsoleVariable<int32> CVarMinLOD(
 	TEXT("foliage.MinLOD"),
 	-1,
@@ -2210,6 +2215,8 @@ FPrimitiveSceneProxy* UInstancedStaticMeshComponent::CreateSceneProxy()
 #if WITH_EDITOR
 		bIsInstanceDataApplyCompleted && 
 #endif
+		// make sure we have instances - or have explicitly permitted creating emtpty ISMs
+		(PerInstanceSMData.Num() > 0 || CVarAllowCreateEmptyISMs.GetValueOnGameThread()) &&
 		// make sure we have an actual static mesh
 		GetStaticMesh() &&
 		GetStaticMesh()->IsCompiling() == false &&
