Allow access to last server base in CharacterMovementComponent child class

#rnx

diff --git a/Engine/Source/Runtime/Engine/Classes/GameFramework/CharacterMovementComponent.h b/Engine/Source/Runtime/Engine/Classes/GameFramework/CharacterMovementComponent.h
--- a/Engine/Source/Runtime/Engine/Classes/GameFramework/CharacterMovementComponent.h
+++ b/Engine/Source/Runtime/Engine/Classes/GameFramework/CharacterMovementComponent.h
@@ -884,6 +884,9 @@ public:
 	UPROPERTY(Category="Character Movement (General Settings)", EditAnywhere, BlueprintReadWrite)
 	TEnumAsByte<enum EMovementMode> DefaultWaterMovementMode;
 
+	/** Accessor to the Last Server Movement Base, only relevant during server update*/
+	UPrimitiveComponent* GetLastServerMovementBase() const { return LastServerMovementBase.Get(); };
+
 private:
 	/**
 	 * Ground movement mode to switch to after falling and resuming ground movement.
