HISM - GDME improvement
 - No need to create unfirom buffer if we have empty LOD array.
 - No need to allocate mesh batch if we have empty section.

diff --git a/Engine/Source/Runtime/Engine/Private/HierarchicalInstancedStaticMesh.cpp b/Engine/Source/Runtime/Engine/Private/HierarchicalInstancedStaticMesh.cpp
--- a/Engine/Source/Runtime/Engine/Private/HierarchicalInstancedStaticMesh.cpp
+++ b/Engine/Source/Runtime/Engine/Private/HierarchicalInstancedStaticMesh.cpp
@@ -1223,23 +1223,30 @@ void FHierarchicalStaticMeshSceneProxy::FillDynamicMeshElements(const FSceneView
 
 	for (int32 LODIndex = FirstLOD; LODIndex < LastLODPlusOne; LODIndex++)
 	{
+		const bool bDitherLODEnabled = ElementParams.bBlendLODs;
+		const uint32 InstancedLODRange = bDitherLODEnabled ? 1 : 0;
+
+		TArray<uint32, SceneRenderingAllocator>& RunArray = bDitherLODEnabled ? Params.MultipleLODRuns[LODIndex] : Params.SingleLODRuns[LODIndex];
+
+		// No need to create uniform buffer if array is empty for given LOD
+		if (!RunArray.Num())
+		{
+			continue;
+		}
+
 		const FStaticMeshLODResources& LODModel = RenderData->LODResources[LODIndex];
 		const int32 TotalNumSections = LODModel.Sections.Num();
 
 		for (int32 SelectionGroupIndex = 0; SelectionGroupIndex < ElementParams.NumSelectionGroups; SelectionGroupIndex++)
 		{
-			const bool bDitherLODEnabled = ElementParams.bBlendLODs;
-			const uint32 InstancedLODRange = bDitherLODEnabled ? 1 : 0;
 			FInstancedStaticMeshVFLooseUniformShaderParametersRef LooseUniformBuffer = CreateLooseUniformBuffer(View, ElementParams.PassUserData[SelectionGroupIndex], InstancedLODRange, LODIndex, EUniformBufferUsage::UniformBuffer_SingleFrame);
 
 			for (int32 SectionIndex = 0; SectionIndex < TotalNumSections; SectionIndex++)
 			{
-				const FLODInfo& ProxyLODInfo = LODs[LODIndex];
-				UMaterialInterface* Material = ProxyLODInfo.Sections[SectionIndex].Material;
-
-				TArray<uint32, SceneRenderingAllocator>& RunArray = bDitherLODEnabled ? Params.MultipleLODRuns[LODIndex] : Params.SingleLODRuns[LODIndex];
+				const FStaticMeshSection& Section = LODModel.Sections[SectionIndex];
 
-				if (!RunArray.Num())
+				// No need to allocate mesh batch if section is empty
+				if (Section.NumTriangles == 0)
 				{
 					continue;
 				}
