Skeleton: Check if asset userdata is nullptr before dereferncing it.

#rnx

diff --git a/Engine/Source/Runtime/Engine/Private/Animation/Skeleton.cpp b/Engine/Source/Runtime/Engine/Private/Animation/Skeleton.cpp
--- a/Engine/Source/Runtime/Engine/Private/Animation/Skeleton.cpp
+++ b/Engine/Source/Runtime/Engine/Private/Animation/Skeleton.cpp
@@ -2011,9 +2011,12 @@ void USkeleton::GetAssetRegistryTags(TArray<FAssetRegistryTag>& OutTags) const
 	OutTags.Add(FAssetRegistryTag(USkeleton::AnimSyncMarkerTag, SyncMarkersBuilder.ToString(), FAssetRegistryTag::TT_Hidden));
 	
 	// Allow asset user data to output tags
-	for(UAssetUserData* AssetUserDataItem : *GetAssetUserDataArray())
+	for(const UAssetUserData* AssetUserDataItem : *GetAssetUserDataArray())
 	{
-		AssetUserDataItem->GetAssetRegistryTags(OutTags);
+		if (AssetUserDataItem)
+		{
+			AssetUserDataItem->GetAssetRegistryTags(OutTags);
+		}
 	}
 }
 
