Made SceneColor node output scene alpha on top of scene color RGB.

diff --git a/Engine/Shaders/Private/MaterialTemplate.ush b/Engine/Shaders/Private/MaterialTemplate.ush
--- a/Engine/Shaders/Private/MaterialTemplate.ush
+++ b/Engine/Shaders/Private/MaterialTemplate.ush
@@ -2629,11 +2629,11 @@ float4 SceneTextureLookup(float2 UV, int SceneTextureIndex, bool bFiltered)
 #if SHADING_PATH_DEFERRED
 
 /** Applies an offset to the scene texture lookup and decodes the HDR linear space color. */
-float3 DecodeSceneColorForMaterialNode(float2 ScreenUV)
+float4 DecodeSceneColorForMaterialNode(float2 ScreenUV)
 {
 #if !defined(SceneColorCopyTexture)
 	// Hit proxies rendering pass doesn't have access to valid render buffers
-	return float3(0.0f, 0.0f, 0.0f);
+	return float4(0.0f, 0.0f, 0.0f, 0.0f);
 #else
 	float4 EncodedSceneColor = Texture2DSample(SceneColorCopyTexture, SceneColorCopySampler, ScreenUV);
 
@@ -2642,7 +2642,7 @@ float3 DecodeSceneColorForMaterialNode(float2 ScreenUV)
 
 	SampledColor *= View.OneOverPreExposure.xxx;
 
-	return SampledColor;
+	return float4(SampledColor, EncodedSceneColor.a);
 #endif
 }
 
diff --git a/Engine/Source/Runtime/Engine/Private/Materials/HLSLMaterialTranslator.cpp b/Engine/Source/Runtime/Engine/Private/Materials/HLSLMaterialTranslator.cpp
--- a/Engine/Source/Runtime/Engine/Private/Materials/HLSLMaterialTranslator.cpp
+++ b/Engine/Source/Runtime/Engine/Private/Materials/HLSLMaterialTranslator.cpp
@@ -7526,7 +7526,7 @@ int32 FHLSLMaterialTranslator::SceneColor(int32 Offset, int32 ViewportUV, bool b
 
 	int32 ScreenUVCode = GetScreenAlignedUV(Offset, ViewportUV, bUseOffset);
 	return AddCodeChunk(
-		MCT_Float3,
+		MCT_Float4,
 		TEXT("DecodeSceneColorForMaterialNode(%s)"),
 		*GetParameterCode(ScreenUVCode)
 		);
diff --git a/Engine/Source/Runtime/Engine/Private/Materials/MaterialExpressions.cpp b/Engine/Source/Runtime/Engine/Private/Materials/MaterialExpressions.cpp
--- a/Engine/Source/Runtime/Engine/Private/Materials/MaterialExpressions.cpp
+++ b/Engine/Source/Runtime/Engine/Private/Materials/MaterialExpressions.cpp
@@ -11255,8 +11255,13 @@ UMaterialExpressionSceneColor::UMaterialExpressionSceneColor(const FObjectInitia
 	MenuCategories.Add(ConstructorStatics.NAME_Texture);
 
 	bShaderInputData = true;
+	bShowOutputNameOnPin = true;
 #endif
 	ConstInput = FVector2D(0.f, 0.f);
+
+	Outputs.Reset();
+	Outputs.Add(FExpressionOutput(TEXT("RGB"), 1, 1, 1, 1, 0));
+	Outputs.Add(FExpressionOutput(TEXT("A"), 1, 0, 0, 0, 1));
 }
 
 void UMaterialExpressionSceneColor::PostLoad()
diff --git a/Engine/Source/Runtime/Engine/Private/Materials/MaterialHLSLTree.cpp b/Engine/Source/Runtime/Engine/Private/Materials/MaterialHLSLTree.cpp
--- a/Engine/Source/Runtime/Engine/Private/Materials/MaterialHLSLTree.cpp
+++ b/Engine/Source/Runtime/Engine/Private/Materials/MaterialHLSLTree.cpp
@@ -2369,13 +2369,13 @@ bool FExpressionSceneColor::PrepareValue(FEmitContext& Context, FEmitScope& Scop
 		Context.MaterialCompilationOutput->SetIsSceneTextureUsed(PPI_SceneColor);
 	}
 
-	return OutResult.SetType(Context, RequestedType, EExpressionEvaluation::Shader, Shader::EValueType::Float3);
+	return OutResult.SetType(Context, RequestedType, EExpressionEvaluation::Shader, Shader::EValueType::Float4);
 }
 
 void FExpressionSceneColor::EmitValueShader(FEmitContext& Context, FEmitScope& Scope, const FRequestedType& RequestedType, FEmitValueShaderResult& OutResult) const
 {
 	FEmitShaderExpression* EmitScreenUV = ScreenUVExpression->GetValueShader(Context, Scope, Shader::EValueType::Float2);
-	OutResult.Code = Context.EmitExpression(Scope, Shader::EValueType::Float3, TEXT("DecodeSceneColorForMaterialNode(%)"), EmitScreenUV);
+	OutResult.Code = Context.EmitExpression(Scope, Shader::EValueType::Float4, TEXT("DecodeSceneColorForMaterialNode(%)"), EmitScreenUV);
 }
 
 bool FExpressionNoise::PrepareValue(FEmitContext& Context, FEmitScope& Scope, const FRequestedType& RequestedType, FPrepareValueResult& OutResult) const
