Skip validating shader file paths in a non-shader-compiling context.

diff --git a/Engine/Source/Runtime/Engine/Private/Materials/MaterialCachedData.cpp b/Engine/Source/Runtime/Engine/Private/Materials/MaterialCachedData.cpp
--- a/Engine/Source/Runtime/Engine/Private/Materials/MaterialCachedData.cpp
+++ b/Engine/Source/Runtime/Engine/Private/Materials/MaterialCachedData.cpp
@@ -1007,30 +1007,32 @@ void FMaterialCachedExpressionData::Validate()
 		}
 		FMaterialLayersFunctions::Validate(MaterialLayers, EditorOnlyData->MaterialLayers);
 
-
-		// TODO Validate the shader include paths in the editor before storing them on the expression instead
-		for (auto PathIt = EditorOnlyData->ExpressionIncludeFilePaths.CreateIterator(); PathIt; ++PathIt)
+		if (!FPlatformProperties::RequiresCookedData() && AllowShaderCompiling())
 		{
-			const FString& IncludeFilePath = *PathIt;
-			bool bValidExpressionIncludePath = false;
-
-			if (!IncludeFilePath.IsEmpty())
+			// TODO Validate the shader include paths in the editor before storing them on the expression instead
+			for (auto PathIt = EditorOnlyData->ExpressionIncludeFilePaths.CreateIterator(); PathIt; ++PathIt)
 			{
-				FString ValidatedPath = GetShaderSourceFilePath(IncludeFilePath);
-				if (!ValidatedPath.IsEmpty())
+				const FString& IncludeFilePath = *PathIt;
+				bool bValidExpressionIncludePath = false;
+
+				if (!IncludeFilePath.IsEmpty())
 				{
-					ValidatedPath = IFileManager::Get().ConvertToAbsolutePathForExternalAppForWrite(*ValidatedPath);
-					if (FPaths::FileExists(ValidatedPath))
+					FString ValidatedPath = GetShaderSourceFilePath(IncludeFilePath);
+					if (!ValidatedPath.IsEmpty())
 					{
-						bValidExpressionIncludePath = true;
+						ValidatedPath = IFileManager::Get().ConvertToAbsolutePathForExternalAppForWrite(*ValidatedPath);
+						if (FPaths::FileExists(ValidatedPath))
+						{
+							bValidExpressionIncludePath = true;
+						}
 					}
 				}
-			}
 
-			if (!bValidExpressionIncludePath)
-			{
-				UE_LOG(LogMaterial, Warning, TEXT("Expression include file path [%s] is invalid, removing from cached material data."), *IncludeFilePath);
-				PathIt.RemoveCurrent();
+				if (!bValidExpressionIncludePath)
+				{
+					UE_LOG(LogMaterial, Warning, TEXT("Expression include file path [%s] is invalid, removing from cached material data."), *IncludeFilePath);
+					PathIt.RemoveCurrent();
+				}
 			}
 		}
 
