Fix CTF_SkipPostEdit flag use in TextureRenderTarget2D

#10968

diff --git a/Engine/Source/Runtime/Engine/Private/TextureRenderTarget.cpp b/Engine/Source/Runtime/Engine/Private/TextureRenderTarget.cpp
--- a/Engine/Source/Runtime/Engine/Private/TextureRenderTarget.cpp
+++ b/Engine/Source/Runtime/Engine/Private/TextureRenderTarget.cpp
@@ -385,7 +385,7 @@ UTexture* UTextureRenderTarget::ConstructTexture(UObject* InOuter, const FString
 		Result->DeferCompression = false;
 	}
 
-	if ((InFlags & CTF_SkipPostEdit) != 0)
+	if ((InFlags & CTF_SkipPostEdit) == 0)
 	{
 		Result->PostEditChange();
 	}
