Fix for Audio Hang on Shutdown

diff --git a/Engine/Source/Runtime/AudioMixer/Private/AudioBusSubsystem.cpp b/Engine/Source/Runtime/AudioMixer/Private/AudioBusSubsystem.cpp
--- a/Engine/Source/Runtime/AudioMixer/Private/AudioBusSubsystem.cpp
+++ b/Engine/Source/Runtime/AudioMixer/Private/AudioBusSubsystem.cpp
@@ -25,7 +25,7 @@ void UAudioBusSubsystem::Initialize(FSubsystemCollectionBase& Collection)
 
 void UAudioBusSubsystem::Deinitialize()
 {
-	UE_LOG(LogAudioMixer, Log, TEXT("Deinitializing Audio Bus Subsystem for audio device with ID %d"), GetMixerDevice()->DeviceID);
+	UE_LOG(LogAudioMixer, Log, TEXT("Deinitializing Audio Bus Subsystem for audio device with ID %d"), GetMixerDevice() ? GetMixerDevice()->DeviceID : -1);
 	ShutdownDefaultAudioBuses();
 }
 
diff --git a/Engine/Source/Runtime/Engine/Private/AudioDeviceManager.cpp b/Engine/Source/Runtime/Engine/Private/AudioDeviceManager.cpp
--- a/Engine/Source/Runtime/Engine/Private/AudioDeviceManager.cpp
+++ b/Engine/Source/Runtime/Engine/Private/AudioDeviceManager.cpp
@@ -214,12 +214,20 @@ FAudioDeviceManager::~FAudioDeviceManager()
 
 	FAudioThread::StopAudioThread();
 
+	TMap<Audio::FDeviceId, FAudioDeviceContainer> DevicesToShutdown;
 	{
-		FScopeLock ScopeLock(&DeviceMapCriticalSection);
-		MainAudioDeviceHandle.Reset();
-		Devices.Reset();
+		FScopeLock ScopeLock(&DeviceMapCriticalSection);		
+		DevicesToShutdown = MoveTemp(Devices);
 	}
 
+	// Can only be destroyed outside of critical section to avoid a deadlock,
+	// but need to remove the device from the manager's list in case of calls
+	// being executed from individual device render thread commands attempting
+	// to access their given device. This is a means to communicate to pending
+	// commands the device is no longer available without destroying it mid-flight.
+	DevicesToShutdown.Reset();
+	MainAudioDeviceHandle.Reset();
+
 	FCoreDelegates::ApplicationWillEnterBackgroundDelegate.RemoveAll(this);
 
 	// Release any loaded buffers - this calls stop on any sources that need it
@@ -511,42 +519,43 @@ void FAudioDeviceManager::DecrementDevice(Audio::FDeviceId DeviceID, UWorld* InW
 		FScopeLock ScopeLock(&DeviceMapCriticalSection);
 
 		// If there is an FAudioDeviceHandle out in the world
-		check(Devices.Contains(DeviceID));
+		if (Devices.Contains(DeviceID))
+		{
+			FAudioDeviceContainer& Container = Devices[DeviceID];
+			check(Container.NumberOfHandlesToThisDevice > 0);
 
-		FAudioDeviceContainer& Container = Devices[DeviceID];
-		check(Container.NumberOfHandlesToThisDevice > 0);
+			// Report device being destroyed before actual destruction
+			// to allow listeners to access and respond where applicable.
+			bool bDestroyingDevice = false;
+			if (Container.NumberOfHandlesToThisDevice == 1)
+			{
+				bDestroyingDevice = true;
+				FAudioDeviceManagerDelegates::OnAudioDeviceDestroyed.Broadcast(DeviceID);
 
-		// Report device being destroyed before actual destruction
-		// to allow listeners to access and respond where applicable.
-		bool bDestroyingDevice = false;
-		if (Container.NumberOfHandlesToThisDevice == 1)
-		{
-			bDestroyingDevice = true;
-			FAudioDeviceManagerDelegates::OnAudioDeviceDestroyed.Broadcast(DeviceID);
+				// Subsystems deinitialization
+				Container.Device->Deinitialize();
 
-			// Subsystems deinitialization
-			Container.Device->Deinitialize();
+				// If this is the active device and being destroyed, set the main device as the active device.
+				if (DeviceID == ActiveAudioDeviceID)
+				{
+					SetActiveDevice(MainAudioDeviceHandle.GetDeviceID());
+				}
 
-			// If this is the active device and being destroyed, set the main device as the active device.
-			if (DeviceID == ActiveAudioDeviceID)
-			{
-				SetActiveDevice(MainAudioDeviceHandle.GetDeviceID());
+				UnregisterWorld(InWorld, DeviceID);
 			}
 
-			UnregisterWorld(InWorld, DeviceID);
-		}
-
-		Container.NumberOfHandlesToThisDevice--;
+			Container.NumberOfHandlesToThisDevice--;
 
-		// If there is no longer any users of this device, destroy it.
-		if (Container.NumberOfHandlesToThisDevice)
-		{
-			ensureMsgf(!bDestroyingDevice, TEXT("AudioDevice Destruction Failure: 'OnAudioDeviceDestroyed' listener generated new persistent handle(s) to AudioDevice."));
-		}
-		else
-		{
-			Swap(DeviceToTearDown, Container.Device);
-			Devices.Remove(DeviceID);
+			// If there is no longer any users of this device, destroy it.
+			if (Container.NumberOfHandlesToThisDevice)
+			{
+				ensureMsgf(!bDestroyingDevice, TEXT("AudioDevice Destruction Failure: 'OnAudioDeviceDestroyed' listener generated new persistent handle(s) to AudioDevice."));
+			}
+			else
+			{
+				Swap(DeviceToTearDown, Container.Device);
+				Devices.Remove(DeviceID);
+			}
 		}
 	}
 
