An InstancedStaticMeshComponent may defer scene proxy creation if instances are added after render state creation. When that happens, the component is not registered to streaming managers and we need to register it later when its instance data is available.

#tests ran through a replay on a platform and verified that the issue was resolved

diff --git a/Engine/Source/Runtime/Engine/Classes/Components/PrimitiveComponent.h b/Engine/Source/Runtime/Engine/Classes/Components/PrimitiveComponent.h
--- a/Engine/Source/Runtime/Engine/Classes/Components/PrimitiveComponent.h
+++ b/Engine/Source/Runtime/Engine/Classes/Components/PrimitiveComponent.h
@@ -2390,6 +2390,9 @@ protected:
 
 	/**  Go through attached primitive components and call MarkRenderStateDirty */
 	ENGINE_API void MarkChildPrimitiveComponentRenderStateDirty();
+
+	/** Conditionally notify streamers that this primitive has updated its render state */
+	ENGINE_API void ConditionalNotifyStreamingPrimitiveUpdated_Concurrent() const;
 public:
 
 	//~ Begin UObject Interface.
diff --git a/Engine/Source/Runtime/Engine/Private/Components/PrimitiveComponent.cpp b/Engine/Source/Runtime/Engine/Private/Components/PrimitiveComponent.cpp
--- a/Engine/Source/Runtime/Engine/Private/Components/PrimitiveComponent.cpp
+++ b/Engine/Source/Runtime/Engine/Private/Components/PrimitiveComponent.cpp
@@ -598,22 +598,7 @@ void UPrimitiveComponent::CreateRenderState_Concurrent(FRegisterComponentContext
 		}
 	}
 
-	// Components are either registered as static or dynamic in the streaming manager.
-	// Static components are registered in batches the first frame the level becomes visible (or incrementally each frame when loaded but not yet visible). 
-	// The level static streaming data is never updated after this, and gets reused whenever the level becomes visible again (after being hidden).
-	// Dynamic components, on the other hand, are updated whenever their render states change.
-	// The following logic handles all cases where static components should fallback on the dynamic path.
-	// It is based on a design where each component must either have bHandledByStreamingManagerAsDynamic or bAttachedToStreamingManagerAsStatic set.
-	// If this is not the case, then the component has never been handled before.
-	// The bIgnoreStreamingManagerUpdate flag is used to prevent handling component that are already in the update list or that don't have streaming data.
-	if (!bIgnoreStreamingManagerUpdate && (Mobility != EComponentMobility::Static || bHandledByStreamingManagerAsDynamic || (!bAttachedToStreamingManagerAsStatic && OwnerLevelHasRegisteredStaticComponentsInStreamingManager(GetOwner()))))
-	{
-		FStreamingManagerCollection* Collection = IStreamingManager::Get_Concurrent();
-		if (Collection)
-		{
-			Collection->NotifyPrimitiveUpdated_Concurrent(this);
-		}
-	}
+	ConditionalNotifyStreamingPrimitiveUpdated_Concurrent();
 }
 
 void UPrimitiveComponent::SendRenderTransform_Concurrent()
@@ -904,6 +889,26 @@ void UPrimitiveComponent::MarkChildPrimitiveComponentRenderStateDirty()
 }
 
 
+void UPrimitiveComponent::ConditionalNotifyStreamingPrimitiveUpdated_Concurrent() const
+{
+	// Components are either registered as static or dynamic in the streaming manager.
+	// Static components are registered in batches the first frame the level becomes visible (or incrementally each frame when loaded but not yet visible). 
+	// The level static streaming data is never updated after this, and gets reused whenever the level becomes visible again (after being hidden).
+	// Dynamic components, on the other hand, are updated whenever their render states change.
+	// The following logic handles all cases where static components should fallback on the dynamic path.
+	// It is based on a design where each component must either have bHandledByStreamingManagerAsDynamic or bAttachedToStreamingManagerAsStatic set.
+	// If this is not the case, then the component has never been handled before.
+	// The bIgnoreStreamingManagerUpdate flag is used to prevent handling component that are already in the update list or that don't have streaming data.
+	if (!bIgnoreStreamingManagerUpdate && (Mobility != EComponentMobility::Static || bHandledByStreamingManagerAsDynamic || (!bAttachedToStreamingManagerAsStatic && OwnerLevelHasRegisteredStaticComponentsInStreamingManager(GetOwner()))))
+	{
+		FStreamingManagerCollection* Collection = IStreamingManager::Get_Concurrent();
+		if (Collection)
+		{
+			Collection->NotifyPrimitiveUpdated_Concurrent(this);
+		}
+	}
+}
+
 bool UPrimitiveComponent::IsWelded() const
 {
 	return BodyInstance.WeldParent != nullptr;
diff --git a/Engine/Source/Runtime/Engine/Private/InstancedStaticMesh.cpp b/Engine/Source/Runtime/Engine/Private/InstancedStaticMesh.cpp
--- a/Engine/Source/Runtime/Engine/Private/InstancedStaticMesh.cpp
+++ b/Engine/Source/Runtime/Engine/Private/InstancedStaticMesh.cpp
@@ -2171,6 +2171,10 @@ void UInstancedStaticMeshComponent::SendRenderInstanceData_Concurrent()
 		{
 			UpdateBounds();
 			GetWorld()->Scene->AddPrimitive(this);
+
+			// A scene proxy was not created during render state creation possibly due to lacking instance data.
+			// In that case, this component hasn't been added to the streamers so do it now
+			ConditionalNotifyStreamingPrimitiveUpdated_Concurrent();
 		}
 	}
 }
