Chaos : Persistence activation :
- Let the ClusterUnionComponent broadcast when its synced with the physics solver and connected to the constrained collections, this allows associated objects to perform post-persistence operations.

diff --git a/Engine/Source/Runtime/Engine/Classes/PhysicsEngine/ClusterUnionComponent.h b/Engine/Source/Runtime/Engine/Classes/PhysicsEngine/ClusterUnionComponent.h
--- a/Engine/Source/Runtime/Engine/Classes/PhysicsEngine/ClusterUnionComponent.h
+++ b/Engine/Source/Runtime/Engine/Classes/PhysicsEngine/ClusterUnionComponent.h
@@ -93,6 +93,18 @@ struct FClusterUnionPendingAddData
 	TArray<FExternalSpatialAccelerationPayload> AccelerationPayloads;
 };
 
+USTRUCT()
+struct FClusterUnionInitializationData
+{
+	GENERATED_BODY()
+
+	UPROPERTY()
+	TObjectPtr<UClusterUnionComponent> ClusterUnionComponent;
+
+	UPROPERTY()
+	TArray<UPrimitiveComponent*> ProcessedComponents;
+};
+
 DECLARE_DYNAMIC_MULTICAST_DELEGATE_ThreeParams(FOnClusterUnionAddedComponent, UPrimitiveComponent*, Component, const TSet<int32>&, BoneIds, bool, bIsNew);
 DECLARE_MULTICAST_DELEGATE_ThreeParams(FOnClusterUnionAddedComponentNative, UPrimitiveComponent*, const TSet<int32>& /*BoneIds*/, bool /*bIsNew*/);
 
@@ -101,6 +113,8 @@ DECLARE_MULTICAST_DELEGATE_OneParam(FOnClusterUnionRemovedComponentNative, UPrim
 
 DECLARE_DYNAMIC_MULTICAST_DELEGATE_TwoParams(FOnClusterUnionBoundsChanged, UClusterUnionComponent*, Component, const FBoxSphereBounds&, Bounds);
 
+DECLARE_MULTICAST_DELEGATE_OneParam(FOnClusterUnionPostSyncBodies, const FClusterUnionInitializationData&);
+
 /**
  * This does the bulk of the work exposing a physics cluster union to the game thread.
  * This component needs to be a primitive component primarily because of how physics
@@ -179,6 +193,7 @@ public:
 	// native (fast, low overhead) versions 
 	FOnClusterUnionAddedComponentNative OnComponentAddedNativeEvent;
 	FOnClusterUnionRemovedComponentNative OnComponentRemovedNativeEvent;
+	FOnClusterUnionPostSyncBodies OnClusterUnionPostSyncBodiesEvent;
 
 	// Lambda returns whether or not iteration should continue;
 	ENGINE_API void VisitAllCurrentChildComponents(const TFunction<bool(UPrimitiveComponent*)>& Lambda) const;
@@ -221,6 +236,11 @@ protected:
 	// side benefit here is being able to track which components are clustered.
 	TMap<TObjectKey<UPrimitiveComponent>, FClusteredComponentData> PerComponentData;
 
+
+	// Whether or not this code is running on the server.
+	UFUNCTION()
+	ENGINE_API bool IsAuthority() const;
+
 private:
 	// These are the statically clustered components. These should
 	// be specified in the editor and never change.
@@ -292,10 +312,6 @@ private:
 	ENGINE_API TArray<AActor*> GetAllCurrentActors() const;
 	ENGINE_API void VisitAllCurrentChildComponentsForCollision(ECollisionChannel TraceChannel, const struct FCollisionQueryParams& Params, const struct FCollisionResponseParams& ResponseParams, const struct FCollisionObjectQueryParams& ObjectParams, const TFunction<bool(UPrimitiveComponent*)>& Lambda) const;
 
-	// Whether or not this code is running on the server.
-	UFUNCTION()
-	ENGINE_API bool IsAuthority() const;
-
 	// Merge all the physics objects geometries into the cluster union
 	ENGINE_API void AddGTParticleGeometry(const TArray<Chaos::FPhysicsObjectHandle>& PhysicsObjects);
 
diff --git a/Engine/Source/Runtime/Engine/Private/PhysicsEngine/ClusterUnionComponent.cpp b/Engine/Source/Runtime/Engine/Private/PhysicsEngine/ClusterUnionComponent.cpp
--- a/Engine/Source/Runtime/Engine/Private/PhysicsEngine/ClusterUnionComponent.cpp
+++ b/Engine/Source/Runtime/Engine/Private/PhysicsEngine/ClusterUnionComponent.cpp
@@ -912,6 +912,11 @@ void UClusterUnionComponent::SyncClusterUnionFromProxy()
 	{
 		OnComponentBoundsChangedEvent.Broadcast(this, CachedLocalBounds);
 	}
+
+	if (OnClusterUnionPostSyncBodiesEvent.IsBound())
+	{
+		OnClusterUnionPostSyncBodiesEvent.Broadcast({ this,GetPrimitiveComponents() });
+	}
 }
 
 void UClusterUnionComponent::HandleAddOrModifiedClusteredComponent(UPrimitiveComponent* ChangedComponent, const TMap<int32, FTransform>& PerBoneChildToParent)
