Require Mobile HDR when using Mobile Deferred Shading in Project Settings.

diff --git a/Engine/Source/Runtime/Core/Private/HAL/ConsoleManager.cpp b/Engine/Source/Runtime/Core/Private/HAL/ConsoleManager.cpp
--- a/Engine/Source/Runtime/Core/Private/HAL/ConsoleManager.cpp
+++ b/Engine/Source/Runtime/Core/Private/HAL/ConsoleManager.cpp
@@ -3171,7 +3171,7 @@ static TAutoConsoleVariable<int32> CVarMobileShadingPath(
 	TEXT("r.Mobile.ShadingPath"),
 	0,
 	TEXT("0: Forward shading (default)\n"
-		 "1: Deferred shading"),
+		 "1: Deferred shading (Mobile HDR is required for Deferred)"),
 	ECVF_RenderThreadSafe | ECVF_ReadOnly);
 
 static TAutoConsoleVariable<int32> CVarMobileAllowDeferredShadingOpenGL(
diff --git a/Engine/Source/Runtime/Engine/Classes/Engine/RendererSettings.h b/Engine/Source/Runtime/Engine/Classes/Engine/RendererSettings.h
--- a/Engine/Source/Runtime/Engine/Classes/Engine/RendererSettings.h
+++ b/Engine/Source/Runtime/Engine/Classes/Engine/RendererSettings.h
@@ -291,7 +291,7 @@ class URendererSettings : public UDeveloperSettings
 
 	UPROPERTY(config, EditAnywhere, Category = Mobile, meta = (
 		ConsoleVariable = "r.Mobile.ShadingPath", DisplayName = "Mobile Shading",
-		ToolTip = "The shading path to use on mobile platforms. Changing this setting requires restarting the editor.",
+		ToolTip = "The shading path to use on mobile platforms. Changing this setting requires restarting the editor. Mobile HDR is required for Deferred Shading.",
 		ConfigRestartRequired = true))
 	TEnumAsByte<EMobileShadingPath::Type> MobileShadingPath;
 
@@ -896,6 +896,7 @@ class URendererSettings : public UDeveloperSettings
 	uint32 bMultiView : 1;
 
 	UPROPERTY(config, EditAnywhere, Category = VR, meta=(
+		EditCondition = "MobileShadingPath == 0",
 		ConsoleVariable="r.MobileHDR", DisplayName="Mobile HDR",
 		ToolTip="If true, mobile pipelines include a full post-processing pass with tonemapping. Disable this setting for a performance boost and to enable stereoscopic rendering optimizations. Changing this setting requires restarting the editor.",
 		ConfigRestartRequired = true))
diff --git a/Engine/Source/Runtime/Engine/Private/RendererSettings.cpp b/Engine/Source/Runtime/Engine/Private/RendererSettings.cpp
--- a/Engine/Source/Runtime/Engine/Private/RendererSettings.cpp
+++ b/Engine/Source/Runtime/Engine/Private/RendererSettings.cpp
@@ -267,6 +267,11 @@ void URendererSettings::PostEditChangeProperty(FPropertyChangedEvent& PropertyCh
 			{
 				bMobileSupportDeferredOnOpenGL = 0;
 			}
+			if(MobileShadingPath.GetValue() == 1)
+			{
+				bMobilePostProcessing = 1;
+				UpdateDependentPropertyInConfigFile(this, GET_MEMBER_NAME_CHECKED(URendererSettings, bMobilePostProcessing));
+			}
 		}
 
 		if ((PropertyChangedEvent.Property->GetFName() == GET_MEMBER_NAME_CHECKED(URendererSettings, bEnableSubstrate)))
@@ -370,6 +375,11 @@ bool URendererSettings::CanEditChange(const FProperty* InProperty) const
 		return MobileShadingPath.GetValue() > 0;
 	}
 
+	if ((InProperty->GetFName() == GET_MEMBER_NAME_CHECKED(URendererSettings, bMobilePostProcessing)))
+	{
+		return MobileShadingPath.GetValue() == 0;
+	}
+
 	PRAGMA_DISABLE_DEPRECATION_WARNINGS
 	// only allow changing ExtendDefaultLuminanceRange if it was disabled.
 	// we don't want new projects disabling this setting but still allow existing projects to enable it.
diff --git a/Engine/Source/Runtime/RenderCore/Private/RenderUtils.cpp b/Engine/Source/Runtime/RenderCore/Private/RenderUtils.cpp
--- a/Engine/Source/Runtime/RenderCore/Private/RenderUtils.cpp
+++ b/Engine/Source/Runtime/RenderCore/Private/RenderUtils.cpp
@@ -489,12 +489,13 @@ RENDERCORE_API bool IsMobileDeferredShadingEnabled(const FStaticShaderPlatform P
 {
 	static FShaderPlatformCachedIniValue<bool> MobileShadingPathIniValue(TEXT("r.Mobile.ShadingPath"));
 	static TConsoleVariableData<int32>* MobileAllowDeferredShadingOpenGL = IConsoleManager::Get().FindTConsoleVariableDataInt(TEXT("r.Mobile.AllowDeferredShadingOpenGL"));
+	static TConsoleVariableData<int32>* MobileHDR = IConsoleManager::Get().FindTConsoleVariableDataInt(TEXT("r.MobileHDR"));
 
 	// OpenGL requires DXC for deferred shading
 	bool bSupportedPlatform = !IsOpenGLPlatform(Platform) ||
 		(MobileAllowDeferredShadingOpenGL && MobileAllowDeferredShadingOpenGL->GetValueOnAnyThread() != 0 && IsDxcEnabledForPlatform(Platform));
 
-	return MobileShadingPathIniValue.Get(Platform) == 1 && bSupportedPlatform;
+	return MobileShadingPathIniValue.Get(Platform) == 1 && bSupportedPlatform && MobileHDR->GetValueOnAnyThread() == 1;
 }
 
 RENDERCORE_API bool MobileRequiresSceneDepthAux(const FStaticShaderPlatform Platform)
