Improved the instance delta tracking.
 - Consolidated the bit arrays into a single array with compacted bit flags
 - Changed storage order to be index based rather than ID to improve memory coherency
 - Removed intermediate stage of transforming the tracked bits into delta index arrays
 - Simplified the logic for constructing the delta ranges.
- made id mapping updates also use delta (rather than bulk copy)
- Removed deltas for baked lighting and editor data since they don't actually support this.

diff --git a/Engine/Source/Runtime/Engine/Private/InstanceDataSceneProxy.cpp b/Engine/Source/Runtime/Engine/Private/InstanceDataSceneProxy.cpp
--- a/Engine/Source/Runtime/Engine/Private/InstanceDataSceneProxy.cpp
+++ b/Engine/Source/Runtime/Engine/Private/InstanceDataSceneProxy.cpp
@@ -20,6 +20,45 @@ DECLARE_DWORD_ACCUMULATOR_STAT(TEXT("Hierarchy Offset Instances"), STAT_Instance
 
 const FInstanceDataBufferHeader FInstanceDataBufferHeader::SinglePrimitiveHeader = { 1, 0, FInstanceDataFlags() };
 
+void FInstanceIdIndexMap::Reset(int32 InNumInstances)
+{
+	IndexToIdMap.Reset();
+	IdToIndexMap.Reset();
+	NumInstances = InNumInstances;
+}
+
+void FInstanceIdIndexMap::ResizeExplicit(int32 InNumInstances, int32 MaxInstanceId)
+{
+	if (IsIdentity())
+	{
+		// Create explicit mapping & before applying changes
+		CreateExplicitIdentityMapping();
+	}
+	IndexToIdMap.SetNumUninitialized(InNumInstances);
+	NumInstances = InNumInstances;
+	if (MaxInstanceId != IdToIndexMap.Num())
+	{
+		int32 OldCount = IdToIndexMap.Num();
+		IdToIndexMap.SetNumUninitialized(MaxInstanceId);
+		for (int32 Index = OldCount; Index < MaxInstanceId; ++Index)
+		{
+			IdToIndexMap[Index] = INDEX_NONE;
+		}
+	}
+}
+
+void FInstanceIdIndexMap::CreateExplicitIdentityMapping()
+{
+	check(IsIdentity());
+	IndexToIdMap.SetNumUninitialized(NumInstances);
+	IdToIndexMap.SetNumUninitialized(NumInstances);
+	for (int32 Index = 0; Index < NumInstances; ++Index)
+	{
+		IndexToIdMap[Index] = FPrimitiveInstanceId{Index};
+		IdToIndexMap[Index] = Index;
+	}
+}
+
 uint32 FInstanceSceneDataBuffers::CalcPayloadDataStride(FInstanceDataFlags Flags, int32 InNumCustomDataFloats, int32 InNumPayloadExtensionFloat4s)
 {
 	static_assert(sizeof(FRenderTransform) == sizeof(float) * 3 * 4); // Sanity check
diff --git a/Engine/Source/Runtime/Engine/Private/InstancedStaticMesh.cpp b/Engine/Source/Runtime/Engine/Private/InstancedStaticMesh.cpp
--- a/Engine/Source/Runtime/Engine/Private/InstancedStaticMesh.cpp
+++ b/Engine/Source/Runtime/Engine/Private/InstancedStaticMesh.cpp
@@ -46,6 +46,7 @@
 #include "Rendering/RenderCommandPipes.h"
 #include "NaniteVertexFactory.h"
 #include "InstancedStaticMeshSceneProxyDesc.h"
+#include "InstancedStaticMesh/ISMInstanceDataManager.h"
 
 #if RHI_RAYTRACING
 #endif
@@ -2390,9 +2391,9 @@ void UInstancedStaticMeshComponent::BuildInstanceDataDeltaChangeSetCommon(FISMIn
 		if (MeshMapBuildData == nullptr && LODData.Num() > 0)
 		{
 			MeshMapBuildData = GetMeshMapBuildData(LODData[0], false);
-	}
+		}
 
-		for (int32 Index : ChangeSet.InstanceLightShadowUVBiasDelta)
+		for (int32 Index = 0; Index < ChangeSet.PostUpdateNumInstances; ++Index)
 		{
 			FVector2D LightmapUVBias = FVector2D(-1.0f, -1.0f);
 			FVector2D ShadowmapUVBias = FVector2D(-1.0f, -1.0f);
@@ -2401,10 +2402,9 @@ void UInstancedStaticMeshComponent::BuildInstanceDataDeltaChangeSetCommon(FISMIn
 			{
 				LightmapUVBias = FVector2D(MeshMapBuildData->PerInstanceLightmapData[Index].LightmapUVBias);
 				ShadowmapUVBias = FVector2D(MeshMapBuildData->PerInstanceLightmapData[Index].ShadowmapUVBias);
-}
+			}
 			ChangeSet.AddInstanceLightShadowUVBias(FVector4f(LightmapUVBias.X, LightmapUVBias.Y, ShadowmapUVBias.X, ShadowmapUVBias.Y));
 		}
-
 	}
 	ChangeSet.SetCustomData(MakeArrayView(PerInstanceSMCustomData), NumCustomDataFloats);
 }
diff --git a/Engine/Source/Runtime/Engine/Private/InstancedStaticMesh/ISMInstanceDataManager.cpp b/Engine/Source/Runtime/Engine/Private/InstancedStaticMesh/ISMInstanceDataManager.cpp
--- a/Engine/Source/Runtime/Engine/Private/InstancedStaticMesh/ISMInstanceDataManager.cpp
+++ b/Engine/Source/Runtime/Engine/Private/InstancedStaticMesh/ISMInstanceDataManager.cpp
@@ -2,6 +2,8 @@
 
 #include "InstancedStaticMesh/ISMInstanceDataManager.h"
 #include "InstancedStaticMesh/ISMInstanceUpdateChangeSet.h"
+#include "InstancedStaticMesh/ISMScatterGatherUtil.h"
+
 #include "Engine/InstancedStaticMesh.h"
 #include "DataDrivenShaderPlatformInfo.h"
 #include "Rendering/MotionVectorSimulation.h"
@@ -84,7 +86,7 @@ void FPrimitiveInstanceDataManager::Add(int32 InInstanceAddAtIndex, bool bInsert
 	if (HasIdentityMapping())
 	{
 		int32 InstanceId = NumInstances++;
-		MarkChangeHelper(AddedInstances, InstanceId);
+		MarkChangeHelper<EChangeFlag::Added>(InstanceId);
 
 		LOG_INST_DATA(TEXT("Add(IDX: %d, bInsert: %d) -> Id: %d"), InInstanceAddAtIndex, bInsert ? 1 : 0, InstanceId);
 		return;
@@ -113,7 +115,7 @@ void FPrimitiveInstanceDataManager::Add(int32 InInstanceAddAtIndex, bool bInsert
 			FPrimitiveInstanceId MovedId = IndexToIdMap[Index];
 			IdToIndexMap[MovedId.Id] = Index;
 			LOG_INST_DATA(TEXT("IdToIndexMap[%d] = %d"), MovedId.GetAsIndex(), Index);
-			MarkIndexChanged(MovedId);
+			InstanceUpdateTracker.MarkIndex<EChangeFlag::IndexChanged>(Index, GetMaxInstanceIndex());
 		}
 	}
 	else
@@ -126,7 +128,7 @@ void FPrimitiveInstanceDataManager::Add(int32 InInstanceAddAtIndex, bool bInsert
 	}
 	NumInstances = IndexToIdMap.Num();
 	check(ValidInstanceIdMask.Num() >= NumInstances);
-	MarkChangeHelper(AddedInstances, InstanceId);
+	MarkChangeHelper<EChangeFlag::Added>(InstanceId);
 	LOG_INST_DATA(TEXT("Add(IDX: %d, bInsert: %d) -> Id: %d"), InInstanceAddAtIndex, bInsert, InstanceId.Id);
 
 	ValidateMapping();
@@ -143,6 +145,10 @@ void FPrimitiveInstanceDataManager::RemoveAtSwap(int32 InstanceIndex)
 
 	check(Mode != EMode::ExternalLegacyData);
 
+	FPrimitiveInstanceId InstanceId = IndexToId(InstanceIndex);
+	// resize to the max at once so we don't have to grow piecemeal
+	InstanceUpdateTracker.RemoveAtSwap(InstanceId, InstanceIndex, GetMaxInstanceIndex());
+
 	// If the remove would cause reordering, we create the explicit mapping
 	const bool bCausesReordering = InstanceIndex != NumInstances - 1;
 	if (bCausesReordering && HasIdentityMapping())
@@ -152,20 +158,17 @@ void FPrimitiveInstanceDataManager::RemoveAtSwap(int32 InstanceIndex)
 
 	MarkComponentRenderInstancesDirty();
 
+	FreeInstanceId(InstanceId);
+
 	// If we still have the identity mapping, we must be removing the last item
 	if (HasIdentityMapping())
 	{
 		check(!bCausesReordering);
-		FPrimitiveInstanceId InstanceId {InstanceIndex};
-		MarkRemoved(InstanceId);
 		--NumInstances;
 		LOG_INST_DATA(TEXT("RemoveAtSwap(IDX: %d) -> Id: %d"), InstanceIndex, InstanceId.Id);
 		return;
 	}
 
-	FPrimitiveInstanceId InstanceId = IndexToIdMap[InstanceIndex];
-	MarkRemoved(InstanceId);
-
 	FPrimitiveInstanceId LastInstanceId = IndexToIdMap.Pop();
 	NumInstances = IndexToIdMap.Num();
 	check(ValidInstanceIdMask.Num() >= NumInstances);
@@ -175,7 +178,6 @@ void FPrimitiveInstanceDataManager::RemoveAtSwap(int32 InstanceIndex)
 		IdToIndexMap[LastInstanceId.Id] = InstanceIndex;
 		LOG_INST_DATA(TEXT("IdToIndexMap[%d] = %d"), LastInstanceId.Id, InstanceIndex);
 		IndexToIdMap[InstanceIndex] = LastInstanceId;
-		MarkIndexChanged(LastInstanceId);
 	}
 	ValidateMapping();
 	LOG_INST_DATA(TEXT("RemoveAtSwap(IDX: %d) -> Id: %d"), InstanceIndex, InstanceId.Id);
@@ -191,6 +193,9 @@ void FPrimitiveInstanceDataManager::RemoveAt(int32 InstanceIndex)
 	ValidateMapping();
 
 	check(Mode != EMode::ExternalLegacyData);
+	FPrimitiveInstanceId InstanceId = IndexToId(InstanceIndex);
+
+	InstanceUpdateTracker.RemoveAt(InstanceId, InstanceIndex, GetMaxInstanceIndex());
 		
 	const bool bCausesReordering = InstanceIndex != NumInstances - 1;
 	if (bCausesReordering && HasIdentityMapping())
@@ -199,21 +204,17 @@ void FPrimitiveInstanceDataManager::RemoveAt(int32 InstanceIndex)
 	}
 
 	MarkComponentRenderInstancesDirty();
+	FreeInstanceId(InstanceId);
 
 	// If we still have the identity mapping, do the simplified tracking update
 	if (HasIdentityMapping())
 	{
 		check(!bCausesReordering);
-		FPrimitiveInstanceId InstanceId {InstanceIndex};
-		MarkRemoved(InstanceId);
 		--NumInstances;
 		LOG_INST_DATA(TEXT("RemoveAt(IDX: %d) -> Id: %d"), InstanceIndex, InstanceId.Id);
 		return;
 	}
 
-	FPrimitiveInstanceId InstanceId = IndexToIdMap[InstanceIndex];
-	MarkRemoved(InstanceId);
-
 	if (InstanceIndex == IndexToIdMap.Num() - 1)
 	{
 		IndexToIdMap.SetNum(InstanceIndex);
@@ -226,7 +227,6 @@ void FPrimitiveInstanceDataManager::RemoveAt(int32 InstanceIndex)
 			FPrimitiveInstanceId MovedId = IndexToIdMap[Index];
 			IdToIndexMap[MovedId.Id] = Index;
 			LOG_INST_DATA(TEXT("IdToIndexMap[%d] = %d"), MovedId.GetAsIndex(), Index);
-			MarkIndexChanged(MovedId);
 		}
 	}
 	NumInstances = IndexToIdMap.Num();
@@ -239,14 +239,15 @@ void FPrimitiveInstanceDataManager::RemoveAt(int32 InstanceIndex)
 void FPrimitiveInstanceDataManager::TransformChanged(int32 InstanceIndex)
 {
 	LOG_INST_DATA(TEXT("TransformChanged(IDX: %d)"), InstanceIndex);
-	MarkChangeHelper(TransformChangedInstances, InstanceIndex);
+	MarkChangeHelper<EChangeFlag::TransformChanged>( InstanceIndex);
 }
 
 void FPrimitiveInstanceDataManager::TransformChanged(FPrimitiveInstanceId InstanceId)
 {
 	LOG_INST_DATA(TEXT("TransformChanged(ID: %d)"), InstanceId.Id);
-	MarkChangeHelper(TransformChangedInstances, InstanceId);
+	MarkChangeHelper<EChangeFlag::TransformChanged>(InstanceId);
 }
+
 void FPrimitiveInstanceDataManager::TransformsChangedAll()
 {
 	if (GetState() == ETrackingState::Disabled)
@@ -262,13 +263,14 @@ void FPrimitiveInstanceDataManager::TransformsChangedAll()
 void FPrimitiveInstanceDataManager::CustomDataChanged(int32 InstanceIndex)
 {
 	LOG_INST_DATA(TEXT("CustomDataChanged(IDX: %d)"), InstanceIndex);
-	MarkChangeHelper(CustomDataChangedInstances, InstanceIndex);
+	MarkChangeHelper<EChangeFlag::CustomDataChanged>(InstanceIndex);
 }
 
 void FPrimitiveInstanceDataManager::BakedLightingDataChanged(int32 InstanceIndex)
 {
 	LOG_INST_DATA(TEXT("BakedLightingDataChanged(IDX: %d)"), InstanceIndex);
-	MarkChangeHelper(BakedLightingDataChangedInstances, InstanceIndex);
+	bBakedLightingDataChanged = true;
+	MarkComponentRenderInstancesDirty();
 }
 
 void FPrimitiveInstanceDataManager::BakedLightingDataChangedAll()
@@ -325,7 +327,7 @@ void FPrimitiveInstanceDataManager::PrimitiveTransformChanged()
 
 bool FPrimitiveInstanceDataManager::HasAnyInstanceChanges() const
 {
-	return !AddedInstances.IsEmpty() 
+	return InstanceUpdateTracker.HasAnyChanges() 
 #if WITH_EDITOR
 		|| bAnyEditorDataChanged 
 #endif
@@ -487,130 +489,21 @@ struct FLegacyRebuildChangeSet
 
 void FPrimitiveInstanceDataManager::InitChangeSet(const FChangeDesc &ChangeDesc, const FInstanceUpdateComponentDesc &ComponentData, FISMInstanceUpdateChangeSet &ChangeSet)
 {
-	if (ChangeSet.IsFullUpdate())
+	// Collect the delta data to be able to update the index mapping.
+	ChangeSet.MaxInstanceId = GetMaxInstanceId();
+	ChangeSet.bIdentityIdMap = IsIdentity();
+	if (!IsIdentity())
 	{
-		// Initialize to full update collection
-		// TODO: Initialize to empty conditionally using the Flags
-		ChangeSet.CustomDataDelta 
-#if WITH_EDITOR
-			= ChangeSet.InstanceEditorDataDelta 
-#endif
-			= ChangeSet.InstanceLightShadowUVBiasDelta 
-			= ChangeSet.TransformsDelta 
-			= FArrayIndexDelta(ComponentData.NumSourceInstances);
+		Gather(ChangeSet.GetIndexChangedDelta(), ChangeSet.IndexToIdMapDeltaData, IndexToIdMap);
 	}
-	else
-	{
-		// TODO: It is moderately efficient to iterate all the bits multiple times... may need to use a different scheme anyway,
-		//       maybe collect set of "anything changed" indexes.
-
-		// If the primitive transform changed, or all instances changed, we want to send all the transforms since there is no proxy side storage of the pre-transformed ones.
-		bool bSendAllTransforms = ChangeDesc.bPrimitiveTransformChanged || bTransformChangedAllInstances;
-		// skip the delta collection if none is used
-		bool bAnyAddDeltaUsed = !(bSendAllTransforms  && bNumCustomDataChanged  && bBakedLightingDataChanged);
-		if (bAnyAddDeltaUsed)
-		{
-			for (TConstSetBitIterator<> BitIt(AddedInstances); BitIt; ++BitIt)
-			{
-				int32 Index = IdToIndex(FPrimitiveInstanceId{BitIt.GetIndex()});
-				check(Index != INDEX_NONE);
-				if (!bSendAllTransforms)
-				{
-					ChangeSet.TransformsDelta.AddIndex(Index);
-				}
-
-				if (!bNumCustomDataChanged && Flags.bHasPerInstanceCustomData)
-				{
-					ChangeSet.CustomDataDelta.AddIndex(Index);
-				}
-
-				if (!bBakedLightingDataChanged)
-				{
-					ChangeSet.InstanceLightShadowUVBiasDelta.AddIndex(Index);
-				}
-				//ChangeSet.InstanceEditorDataDelta.AddIndex(Index);
-			}
-		}
-#if WITH_EDITOR
-		// Bulk update editor data - the hit-proxies are recreated all the time.
-		ChangeSet.InstanceEditorDataDelta = FArrayIndexDelta(ComponentData.NumSourceInstances);
-#endif
-		if (bSendAllTransforms)
-		{
-			ChangeSet.TransformsDelta = FArrayIndexDelta(ComponentData.NumSourceInstances);
-		}
-		else
-		{
-			for (TConstSetBitIterator<> BitIt(TransformChangedInstances); BitIt; ++BitIt)
-			{
-				FPrimitiveInstanceId Id{BitIt.GetIndex()};
-				if (!AddedInstances[Id] && !RemovedInstances[Id])
-				{
-					int32 Index = IdToIndex(Id);
-					check(Index != INDEX_NONE);
-					ChangeSet.TransformsDelta.AddIndex(Index);
-				}
-			}
-		}
 
-		// Update all the custom data if there was a change in number of floats (this resets the whole thing).
-		if (Flags.bHasPerInstanceCustomData)
-		{
-			if (bNumCustomDataChanged)
-			{
-				ChangeSet.CustomDataDelta = FArrayIndexDelta(ComponentData.NumSourceInstances);
-			}
-			else
-			{
-				for (TConstSetBitIterator<> BitIt(CustomDataChangedInstances); BitIt; ++BitIt)
-				{
-					FPrimitiveInstanceId Id{BitIt.GetIndex()};
-					if (!AddedInstances[Id] && !RemovedInstances[Id])
-					{
-						int32 Index = IdToIndex(Id);
-						check(Index != INDEX_NONE);
-						ChangeSet.CustomDataDelta.AddIndex(Index);
-					}
-				}
-			}
-		}
-		else
-		{
-			ChangeSet.CustomDataDelta = FArrayIndexDelta();
-		}
-
-		if (bBakedLightingDataChanged)
-		{
-			ChangeSet.InstanceLightShadowUVBiasDelta = FArrayIndexDelta(ComponentData.NumSourceInstances);
-		}
-		else
-		{
-			for (TConstSetBitIterator<> BitIt(BakedLightingDataChangedInstances); BitIt; ++BitIt)
-			{
-				FPrimitiveInstanceId Id{BitIt.GetIndex()};
-				if (!AddedInstances[Id] && !RemovedInstances[Id])
-				{
-					int32 Index = IdToIndex(Id);
-					check(Index != INDEX_NONE);
-					ChangeSet.InstanceLightShadowUVBiasDelta.AddIndex(Index);
-				}
-			}
-		}
-	}
-
-	// Move & reset tracked state
-	ChangeSet.ChangeMask.AddedInstances = MoveTemp(AddedInstances);
-	ChangeSet.ChangeMask.RemovedInstances = MoveTemp(RemovedInstances);
-	ChangeSet.ChangeMask.TransformChangedInstances = MoveTemp(TransformChangedInstances);
-	ChangeSet.ChangeMask.CustomDataChangedInstances = MoveTemp(CustomDataChangedInstances);
-	ChangeSet.ChangeMask.bNumCustomDataChanged = bNumCustomDataChanged;
-	ChangeSet.ChangeMask.bBakedLightingDataChanged = bBakedLightingDataChanged;
+	ChangeSet.bNumCustomDataChanged = bNumCustomDataChanged;
+	ChangeSet.bBakedLightingDataChanged = bBakedLightingDataChanged;
 
 #if WITH_EDITOR
-	ChangeSet.ChangeMask.bAnyEditorDataChanged = bAnyEditorDataChanged;
+	ChangeSet.bAnyEditorDataChanged = bAnyEditorDataChanged;
 	bAnyEditorDataChanged = false;
 #endif	
-	ChangeSet.ChangeMask.IndexChangeInstances = MoveTemp(IndexChangeInstances);
 
 	bTransformChangedAllInstances = false;
 	bNumCustomDataChanged = false;
@@ -621,11 +514,6 @@ void FPrimitiveInstanceDataManager::InitChangeSet(const FChangeDesc &ChangeDesc,
 	ChangeSet.Flags = Flags;
 	ChangeSet.AbsMaxDisplacement = AbsMaxDisplacement;
 	ChangeSet.NumCustomDataFloats = NumCustomDataFloats;
-	if (!ChangeSet.Flags.bHasPerInstanceCustomData)
-	{
-		ChangeSet.CustomDataDelta = FArrayIndexDelta();
-	}
-	ChangeSet.InstanceIdIndexMap = FInstanceIdIndexMap(*this);
 
 	// This is the odd one out
 	ChangeSet.SetInstanceLocalBounds(ComponentData.StaticMeshBounds);
@@ -671,7 +559,6 @@ bool FPrimitiveInstanceDataManager::FlushChanges(FInstanceUpdateComponentDesc &&
 	}
 	bFirstFlush = false;
 
-
 	bool bWasUpdateQueued = false;
 	// Marked for externally managed update, this may be combined with tracked changes (which results in double updates)
 	if (LegacyBuildData)
@@ -722,7 +609,6 @@ bool FPrimitiveInstanceDataManager::FlushChanges(FInstanceUpdateComponentDesc &&
 			FISMCInstanceDataSceneProxy &ProxyRef = *Proxy;
 			
 			// Forcibly destroy any tracking state
-			ProxyRef.ChangeMask.Reset();
 			ProxyRef.InstanceIdIndexMap = MoveTemp(ExternalChangeSet.InstanceIdIndexMap);
 #if WITH_EDITOR
 			ProxyRef.HitProxyContainer = MoveTemp(ExternalChangeSet.HitProxyContainer);
@@ -809,17 +695,29 @@ bool FPrimitiveInstanceDataManager::FlushChanges(FInstanceUpdateComponentDesc &&
 		// TODO: The states other than bCreateNewProxy _can_ be handled through delta updates (e.g., add instance + offset all the rest).
 		bool bNeedFullUpdate = ChangeDesc.bUntrackedState || ChangeDesc.bMaterialUsageFlagsChanged;
 
-		FISMInstanceUpdateChangeSet ChangeSet(bNeedFullUpdate);
+		// NOTE: Moving the update tracker to the change set implicitly resets it.
+		FISMInstanceUpdateChangeSet ChangeSet(bNeedFullUpdate, MoveTemp(InstanceUpdateTracker));
+		ChangeSet.bUpdateAllInstanceTransforms = ChangeDesc.bPrimitiveTransformChanged || bTransformChangedAllInstances;
+		ChangeSet.PostUpdateNumInstances = ComponentData.NumProxyInstances;
+		ChangeSet.NumSourceInstances = ComponentData.NumSourceInstances;
 
 		// Initialize the change set before collecting instance change data.
 		InitChangeSet(ChangeDesc, ComponentData, ChangeSet);
 
+		// Too many numbers, make sure they line up.
+		if (Mode == EMode::Default)
+		{
+			check(ChangeSet.LegacyInstanceReorderTable.IsEmpty());
+			check(ComponentData.NumSourceInstances == GetMaxInstanceIndex());
+			check(ComponentData.NumProxyInstances == GetMaxInstanceIndex());
+		}
+
 		// Callback to the owner to fill in change data.
 		ComponentData.BuildChangeSet(ChangeSet);
 		
 		// make sure the custom data change is correctly tracked
 		check(!ChangeSet.Flags.bHasPerInstanceCustomData || NumCustomDataFloats == ChangeSet.NumCustomDataFloats);
-		check(ChangeSet.Flags.bHasPerInstanceCustomData || ChangeSet.CustomDataDelta.IsEmpty() && ChangeSet.PerInstanceCustomData.IsEmpty());
+		checkSlow(ChangeSet.Flags.bHasPerInstanceCustomData || ChangeSet.GetCustomDataDelta().IsEmpty() && ChangeSet.PerInstanceCustomData.IsEmpty());
 
 		// If we have per-instance previous local to world, they are expected to be in the local space of the _previous_ local to world. If they are in fact not (e.g., if someone sets them explicitly from world space) then, well, this won't be correct
 		if (ChangeSet.Flags.bHasPerInstanceDynamicData)
@@ -838,19 +736,11 @@ bool FPrimitiveInstanceDataManager::FlushChanges(FInstanceUpdateComponentDesc &&
 		// Assemble header info to enable nonblocking primitive update.
 		FInstanceDataBufferHeader InstanceDataBufferHeader;
 
-		// Too many numbers, make sure they line up.
-		if (Mode == EMode::Default)
-		{
-			check(ChangeSet.LegacyInstanceReorderTable.IsEmpty());
-			check(ComponentData.NumSourceInstances == GetMaxInstanceIndex());
-			check(ComponentData.NumProxyInstances == GetMaxInstanceIndex());
-		}
-
-		ChangeSet.PostUpdateNumInstances = ComponentData.NumProxyInstances;
-		InstanceDataBufferHeader.NumInstances = ComponentData.NumProxyInstances;
+		InstanceDataBufferHeader.NumInstances = ChangeSet.PostUpdateNumInstances;
 		InstanceDataBufferHeader.PayloadDataStride = FInstanceSceneDataBuffers::CalcPayloadDataStride(ChangeSet.Flags, ChangeSet.NumCustomDataFloats, 0);
 		InstanceDataBufferHeader.Flags = ChangeSet.Flags;
 
+		CSV_CUSTOM_STAT_GLOBAL(NumInstanceTransformUpdates, ChangeSet.Transforms.Num(), ECsvCustomStatOp::Accumulate);
 
 		// Special handling for the case of external data & transform change, may want to extend this to others (but need to send data in that case)
 		if (Mode == EMode::ExternalLegacyData && ChangeDesc.bPrimitiveTransformChanged)
@@ -950,21 +840,15 @@ void FPrimitiveInstanceDataManager::ClearChangeTracking()
 	TrackingState = ETrackingState::Initial;
 
 	LegacyBuildData.Reset();
-
-	AddedInstances.Empty();
-	RemovedInstances.Empty();
-	TransformChangedInstances.Empty();
-	CustomDataChangedInstances.Empty();
+	InstanceUpdateTracker.Reset();
 	bNumCustomDataChanged = false;
 	bBakedLightingDataChanged = false;
 	bTransformChangedAllInstances = false;
 
-	BakedLightingDataChangedInstances.Empty();
 #if WITH_EDITOR
 	bAnyEditorDataChanged = false;
 #endif	
 	bPrimitiveTransformChanged = false;
-	IndexChangeInstances.Empty();
 }
 
 int32 FPrimitiveInstanceDataManager::GetMaxInstanceId() const
@@ -977,30 +861,6 @@ int32 FPrimitiveInstanceDataManager::GetMaxInstanceIndex() const
 	return HasIdentityMapping() ? NumInstances : IndexToIdMap.Num();
 }
 
-bool FPrimitiveInstanceDataManager::ResizeChangeTrackingState()
-{
-	if (GetState() == ETrackingState::Tracked)
-	{
-		int32 MaxInstanceId = GetMaxInstanceId();
-
-		// Grow to fit the new max
-		if (AddedInstances.Num() < MaxInstanceId)
-		{
-			AddedInstances.SetNum(MaxInstanceId, false);
-			RemovedInstances.SetNum(MaxInstanceId, false);
-			IndexChangeInstances.SetNum(MaxInstanceId, false);
-			TransformChangedInstances.SetNum(MaxInstanceId, false);
-			CustomDataChangedInstances.SetNum(MaxInstanceId, false);
-			BakedLightingDataChangedInstances.SetNum(MaxInstanceId, false);
-#if WITH_EDITOR
-			//bAnyEditorDataChanged = false;
-#endif	
-		}
-		return true;
-	}
-	return false;
-}
-
 void FPrimitiveInstanceDataManager::CreateExplicitIdentityMapping()
 {
 	check(HasIdentityMapping());
@@ -1014,11 +874,10 @@ void FPrimitiveInstanceDataManager::CreateExplicitIdentityMapping()
 	ValidInstanceIdMask.Reset();
 	ValidInstanceIdMask.SetNum(NumInstances, true);
 	IdSearchStartIndex = NumInstances;
-
-	ResizeChangeTrackingState();
 }
 
-void FPrimitiveInstanceDataManager::MarkChangeHelper(FChangeBitArray& TrackingArray, int32 InstanceIndex)
+template<FPrimitiveInstanceDataManager::EChangeFlag Flag>
+void FPrimitiveInstanceDataManager::MarkChangeHelper(int32 InstanceIndex)
 {
 	if (GetState() == ETrackingState::Disabled)
 	{
@@ -1032,10 +891,12 @@ void FPrimitiveInstanceDataManager::MarkChangeHelper(FChangeBitArray& TrackingAr
 		MarkComponentRenderInstancesDirty();
 		return;
 	}
-	MarkChangeHelper(TrackingArray, IndexToId(InstanceIndex));
+	InstanceUpdateTracker.MarkIndex<Flag>(InstanceIndex, GetMaxInstanceIndex());
+	MarkComponentRenderInstancesDirty();
 }
 
-void FPrimitiveInstanceDataManager::MarkChangeHelper(FChangeBitArray& TrackingArray, FPrimitiveInstanceId InstanceId)
+template<FPrimitiveInstanceDataManager::EChangeFlag Flag>
+void FPrimitiveInstanceDataManager::MarkChangeHelper(FPrimitiveInstanceId InstanceId)
 {
 	check(Mode != EMode::ExternalLegacyData);
 
@@ -1044,14 +905,7 @@ void FPrimitiveInstanceDataManager::MarkChangeHelper(FChangeBitArray& TrackingAr
 		MarkComponentRenderInstancesDirty();
 		return;
 	}
-
-	// Make sure we have (lazy) allocated the change tracking arrays
-	ResizeChangeTrackingState();
-
-	// Just mark the info, actual data will be pulled out when dispatching an update
-	TrackingArray[InstanceId] = true;
-
-	MarkComponentRenderInstancesDirty();
+	MarkChangeHelper<Flag>(IdToIndex(InstanceId));
 }
 
 void FPrimitiveInstanceDataManager::MarkComponentRenderInstancesDirty()
@@ -1076,16 +930,9 @@ bool FPrimitiveInstanceDataManager::ShouldTrackIds() const
 	return Proxy != nullptr;
 }
 
-void FPrimitiveInstanceDataManager::MarkRemoved(FPrimitiveInstanceId InstanceId)
+void FPrimitiveInstanceDataManager::FreeInstanceId(FPrimitiveInstanceId InstanceId)
 {
-	LOG_INST_DATA(TEXT("MarkRemoved(Id: %d)"), InstanceId.Id);
-
-	if (ResizeChangeTrackingState())
-	{
-		RemovedInstances[InstanceId] = true;
-		AddedInstances[InstanceId] = false;
-		IndexChangeInstances[InstanceId] = false;
-	}
+	LOG_INST_DATA(TEXT("FreeInstanceId(Id: %d)"), InstanceId.Id);
 
 	if (!HasIdentityMapping())
 	{
@@ -1098,12 +945,6 @@ void FPrimitiveInstanceDataManager::MarkRemoved(FPrimitiveInstanceId InstanceId)
 	LOG_INST_DATA(TEXT("IdToIndexMap[%d] = %d"), InstanceId.Id, INDEX_NONE);
 }
 
-void FPrimitiveInstanceDataManager::MarkIndexChanged(FPrimitiveInstanceId InstanceId)
-{
-	LOG_INST_DATA(TEXT("MarkIndexChanged(Id: %d)"), InstanceId.Id);
-	MarkChangeHelper(IndexChangeInstances, InstanceId);	
-}
-
 TSharedPtr<FISMCInstanceDataSceneProxy, ESPMode::ThreadSafe> FPrimitiveInstanceDataManager::GetOrCreateProxy(FStaticShaderPlatform InShaderPlatform, ERHIFeatureLevel::Type InFeatureLevel)
 {
 	LOG_INST_DATA(TEXT("GetOrCreateProxy"));
@@ -1147,9 +988,9 @@ void FPrimitiveInstanceDataManager::Invalidate(int32 InNumInstances)
 	}
 }
 
+#if DO_GUARD_SLOW
 void FPrimitiveInstanceDataManager::ValidateMapping() const
 {
-#if DO_GUARD_SLOW
 	check(HasIdentityMapping() || IndexToIdMap.Num() == NumInstances);
 	for (int32 Index = 0; Index < IndexToIdMap.Num(); ++Index)
 	{
@@ -1172,8 +1013,8 @@ void FPrimitiveInstanceDataManager::ValidateMapping() const
 	}
 	int32 FirstFalse = ValidInstanceIdMask.Find(false);
 	check(FirstFalse < 0 || FirstFalse >= IdSearchStartIndex);
-#endif
 }
+#endif
 
 void FPrimitiveInstanceDataManager::MarkForRebuildFromLegacy(TUniquePtr<FStaticMeshInstanceData> &&InLegacyInstanceData, const TArray<int32> &InstanceReorderTable, const TArray<TRefCountPtr<HHitProxy>> &HitProxies)
 {
@@ -1217,12 +1058,7 @@ void FPrimitiveInstanceDataManager::MarkForRebuildFromLegacy(TUniquePtr<FStaticM
 SIZE_T FPrimitiveInstanceDataManager::GetAllocatedSize() const
 {
 	return ValidInstanceIdMask.GetAllocatedSize() +
-		AddedInstances.GetAllocatedSize() +
-		RemovedInstances.GetAllocatedSize() +
-		TransformChangedInstances.GetAllocatedSize() +
-		CustomDataChangedInstances.GetAllocatedSize() +
-		BakedLightingDataChangedInstances.GetAllocatedSize() +
-		IndexChangeInstances.GetAllocatedSize();
+		InstanceUpdateTracker.GetAllocatedSize();
 }
 
 void FPrimitiveInstanceDataManager::OnRegister(int32 InNumInstances)
diff --git a/Engine/Source/Runtime/Engine/Private/InstancedStaticMesh/ISMInstanceDataSceneProxy.cpp b/Engine/Source/Runtime/Engine/Private/InstancedStaticMesh/ISMInstanceDataSceneProxy.cpp
--- a/Engine/Source/Runtime/Engine/Private/InstancedStaticMesh/ISMInstanceDataSceneProxy.cpp
+++ b/Engine/Source/Runtime/Engine/Private/InstancedStaticMesh/ISMInstanceDataSceneProxy.cpp
@@ -2,6 +2,7 @@
 
 #include "InstancedStaticMesh/ISMInstanceDataSceneProxy.h"
 #include "InstancedStaticMesh/ISMInstanceUpdateChangeSet.h"
+#include "InstancedStaticMesh/ISMScatterGatherUtil.h"
 #include "Engine/InstancedStaticMesh.h"
 #include "DataDrivenShaderPlatformInfo.h"
 #include "Rendering/RenderingSpatialHash.h"
@@ -27,12 +28,12 @@ struct FIdentityIndexRemap
 {
 	inline int32 operator[](int32 InIndex) const { return InIndex; }
 
-	template <typename ValueType>
-	inline void Scatter(bool bHasData, const FArrayIndexDelta &Delta, TArray<ValueType> &DestData, int32 NumOutElements, TArray<ValueType> &&InData, int32 ElementStride = 1) const
+	template <typename DeltaType, typename ValueType>
+	inline void Scatter(bool bHasData, const DeltaType &Delta, TArray<ValueType> &DestData, int32 NumOutElements, TArray<ValueType> &&InData, int32 ElementStride = 1) const
 	{
 		if (bHasData)
 		{
-			Delta.Scatter(DestData, NumOutElements, MoveTemp(InData), ElementStride);
+			::Scatter(Delta, DestData, NumOutElements, MoveTemp(InData), ElementStride);
 		}
 		else
 		{
@@ -71,12 +72,12 @@ struct FReorderTableIndexRemap
 		return Index != INDEX_NONE; 
 	}
 
-	template <typename ValueType>
-	inline void Scatter(bool bHasData, const FArrayIndexDelta &Delta, TArray<ValueType> &DestData, int32 NumOutElements, TArray<ValueType> &&InData, int32 ElementStride = 1) const
+	template <typename DeltaType, typename ValueType>
+	inline void Scatter(bool bHasData, const DeltaType &Delta, TArray<ValueType> &DestData, int32 NumOutElements, TArray<ValueType> &&InData, int32 ElementStride = 1) const
 	{
 		if (bHasData)
 		{
-			Delta.Scatter(DestData, NumOutElements, MoveTemp(InData), *this, ElementStride);
+			::Scatter(Delta, DestData, NumOutElements, MoveTemp(InData), *this, ElementStride);
 		}
 		else
 		{
@@ -122,12 +123,15 @@ void FISMCInstanceDataSceneProxy::ApplyDataChanges(FISMInstanceUpdateChangeSet &
 
 	// unpack transform deltas
 	ProxyData.InstanceToPrimitiveRelative.SetNumUninitialized(PostUpdateNumInstances);
-	for (int32 PackedIndex = 0; PackedIndex < ChangeSet.TransformsDelta.Num(); ++PackedIndex)
+	auto TransformDelta = ChangeSet.GetTransformDelta();
+	for (auto It = TransformDelta.GetIterator(); It; ++It)
 	{
-		int32 InstanceIndex = ChangeSet.TransformsDelta[PackedIndex];
+		int32 ItemIndex = It.GetItemIndex();
+		int32 InstanceIndex = It.GetIndex();
+
 		if (IndexRemap.RemapIndex(InstanceIndex))
 		{
-			FRenderTransform LocalToPrimitiveRelativeWorld = ChangeSet.Transforms[PackedIndex] * ChangeSet.PrimitiveToRelativeWorld;
+			FRenderTransform LocalToPrimitiveRelativeWorld = ChangeSet.Transforms[ItemIndex] * ChangeSet.PrimitiveToRelativeWorld;
 			// Remove shear
 			LocalToPrimitiveRelativeWorld.Orthogonalize();
 			ProxyData.InstanceToPrimitiveRelative[InstanceIndex] = LocalToPrimitiveRelativeWorld;
@@ -138,12 +142,13 @@ void FISMCInstanceDataSceneProxy::ApplyDataChanges(FISMInstanceUpdateChangeSet &
 	{
 		FRenderTransform PrevPrimitiveToRelativeWorld = ChangeSet.PreviousPrimitiveToRelativeWorld.Get(ChangeSet.PrimitiveToRelativeWorld);
 		ProxyData.PrevInstanceToPrimitiveRelative.SetNumUninitialized(PostUpdateNumInstances);
-		for (int32 PackedIndex = 0; PackedIndex < ChangeSet.TransformsDelta.Num(); ++PackedIndex)
+		for (auto It = TransformDelta.GetIterator(); It; ++It)
 		{
-			int32 InstanceIndex = ChangeSet.TransformsDelta[PackedIndex];
+			int32 ItemIndex = It.GetItemIndex();
+			int32 InstanceIndex = It.GetIndex();
 			if (IndexRemap.RemapIndex(InstanceIndex))
 			{
-				FRenderTransform PrevLocalToPrimitiveRelativeWorld = ChangeSet.PrevTransforms[PackedIndex] * PrevPrimitiveToRelativeWorld;
+				FRenderTransform PrevLocalToPrimitiveRelativeWorld = ChangeSet.PrevTransforms[ItemIndex] * PrevPrimitiveToRelativeWorld;
 				PrevLocalToPrimitiveRelativeWorld.Orthogonalize();
 				ProxyData.PrevInstanceToPrimitiveRelative[InstanceIndex] = PrevLocalToPrimitiveRelativeWorld;
 			}
@@ -157,7 +162,7 @@ void FISMCInstanceDataSceneProxy::ApplyDataChanges(FISMInstanceUpdateChangeSet &
 	if (ChangeSet.Flags.bHasPerInstanceCustomData)
 	{
 		ProxyData.NumCustomDataFloats = ChangeSet.NumCustomDataFloats;
-		IndexRemap.Scatter(ChangeSet.Flags.bHasPerInstanceCustomData, ChangeSet.CustomDataDelta, ProxyData.InstanceCustomData, PostUpdateNumInstances, MoveTemp(ChangeSet.PerInstanceCustomData), ProxyData.NumCustomDataFloats);
+		IndexRemap.Scatter(ChangeSet.Flags.bHasPerInstanceCustomData, ChangeSet.GetCustomDataDelta(), ProxyData.InstanceCustomData, PostUpdateNumInstances, MoveTemp(ChangeSet.PerInstanceCustomData), ProxyData.NumCustomDataFloats);
 	}
 	else
 	{
@@ -165,9 +170,9 @@ void FISMCInstanceDataSceneProxy::ApplyDataChanges(FISMInstanceUpdateChangeSet &
 		ProxyData.InstanceCustomData.Reset();
 	}
 
-	IndexRemap.Scatter(ChangeSet.Flags.bHasPerInstanceLMSMUVBias, ChangeSet.InstanceLightShadowUVBiasDelta, ProxyData.InstanceLightShadowUVBias, PostUpdateNumInstances, MoveTemp(ChangeSet.InstanceLightShadowUVBias));
+	IndexRemap.Scatter(ChangeSet.Flags.bHasPerInstanceLMSMUVBias, ChangeSet.GetInstanceLightShadowUVBiasDelta(), ProxyData.InstanceLightShadowUVBias, PostUpdateNumInstances, MoveTemp(ChangeSet.InstanceLightShadowUVBias));
 #if WITH_EDITOR
-	IndexRemap.Scatter(ChangeSet.Flags.bHasPerInstanceEditorData, ChangeSet.InstanceEditorDataDelta, ProxyData.InstanceEditorData, PostUpdateNumInstances, MoveTemp(ChangeSet.InstanceEditorData));
+	IndexRemap.Scatter(ChangeSet.Flags.bHasPerInstanceEditorData, ChangeSet.GetInstanceEditorDataDelta(), ProxyData.InstanceEditorData, PostUpdateNumInstances, MoveTemp(ChangeSet.InstanceEditorData));
 
 	// replace the HP container.
 	if (ChangeSet.HitProxyContainer)
@@ -182,7 +187,6 @@ void FISMCInstanceDataSceneProxy::ApplyDataChanges(FISMInstanceUpdateChangeSet &
 	{
 		// TODO: only need to process added instances? No help for ISM since the move path would be taken.
 		// TODO: OTOH for HISM there is no meaningful data, so just skipping and letting the SetNumZeroed fill in the blanks is fine.
-		FArrayIndexDelta PerInstanceRandomDelta(PostUpdateNumInstances);
 
 		ProxyData.InstanceRandomIDs.SetNumZeroed(PostUpdateNumInstances);
 		if (ChangeSet.GeneratePerInstanceRandomIds)
@@ -191,6 +195,7 @@ void FISMCInstanceDataSceneProxy::ApplyDataChanges(FISMInstanceUpdateChangeSet &
 			TArray<float> TmpInstanceRandomIDs;
 			TmpInstanceRandomIDs.SetNumZeroed(PostUpdateNumInstances);
 			ChangeSet.GeneratePerInstanceRandomIds(TmpInstanceRandomIDs);
+			FIdentityDeltaRange PerInstanceRandomDelta(TmpInstanceRandomIDs.Num());
 			IndexRemap.Scatter(true, PerInstanceRandomDelta, ProxyData.InstanceRandomIDs, PostUpdateNumInstances, MoveTemp(TmpInstanceRandomIDs));
 		}
 		//else 
@@ -204,26 +209,6 @@ void FISMCInstanceDataSceneProxy::ApplyDataChanges(FISMInstanceUpdateChangeSet &
 	}
 }
 
-struct FDataRemap
-{
-	// Keep bit mask to efficiently iterate, but not save any data.
-	// TODO: Maybe make sparse in the future?
-	TBitArray<> Used;
-	TArray<int32> FromInices;
-
-	FDataRemap(int32 InNum)
-	{
-		Used.SetNum(InNum, false);
-		FromInices.AddUninitialized(InNum);
-	}
-
-	inline void Set(int32 To, int32 From)
-	{
-		Used[To] = true;
-		FromInices[To] = From;
-	}
-};
-
 template<typename ValueType>
 void CondMove(bool bCondition, TArray<ValueType> &Data, int32 FromIndex, int32 ToIndex, int32 NumElements = 1)
 {
@@ -237,31 +222,66 @@ void FISMCInstanceDataSceneProxy::Build(FISMInstanceUpdateChangeSet&& ChangeSet)
 {
 	DecStatCounters();
 	check(ChangeSet.IsFullUpdate());
-	check(!ChangeSet.TransformsDelta.IsDelta());
-	check(!ChangeSet.CustomDataDelta.IsDelta() || (!ChangeSet.Flags.bHasPerInstanceCustomData && ChangeSet.CustomDataDelta.IsEmpty()));
-	check(!ChangeSet.InstanceLightShadowUVBiasDelta.IsDelta() || ChangeSet.InstanceLightShadowUVBiasDelta.IsEmpty());
+	checkSlow(!ChangeSet.GetTransformDelta().IsDelta());
+	checkSlow(!ChangeSet.GetCustomDataDelta().IsDelta() || (!ChangeSet.Flags.bHasPerInstanceCustomData && ChangeSet.GetCustomDataDelta().IsEmpty()));
+	checkSlow(!ChangeSet.GetInstanceLightShadowUVBiasDelta().IsDelta() || ChangeSet.GetInstanceLightShadowUVBiasDelta().IsEmpty());
 #if WITH_EDITOR
-	check(!ChangeSet.InstanceEditorDataDelta.IsDelta() || ChangeSet.InstanceEditorDataDelta.IsEmpty());
+	checkSlow(!ChangeSet.GetInstanceEditorDataDelta().IsDelta() || ChangeSet.GetInstanceEditorDataDelta().IsEmpty());
 #endif
-	check(ChangeSet.PostUpdateNumInstances == ChangeSet.InstanceIdIndexMap.GetMaxInstanceIndex());
+
 
 	FInstanceSceneDataBuffers::FAccessTag AccessTag(PointerHash(this));
 	FInstanceSceneDataBuffers::FWriteView WriteView = InstanceSceneDataBuffers.BeginWriteAccess(AccessTag);
 
 	WriteView.Flags = ChangeSet.Flags;
 
+	UpdateIdMapping(ChangeSet);
+	check(ChangeSet.PostUpdateNumInstances == InstanceIdIndexMap.GetMaxInstanceIndex());
+
 	FIdentityIndexRemap IndexRemap;
-	ApplyDataChanges(ChangeSet, IndexRemap, ChangeSet.InstanceIdIndexMap.GetMaxInstanceIndex(), WriteView);
+	ApplyDataChanges(ChangeSet, IndexRemap, InstanceIdIndexMap.GetMaxInstanceIndex(), WriteView);
 	InstanceSceneDataBuffers.EndWriteAccess(AccessTag);
 
-	ChangeMask = MoveTemp(ChangeSet.ChangeMask);
-	InstanceIdIndexMap = MoveTemp(ChangeSet.InstanceIdIndexMap);
-
 	InstanceSceneDataBuffers.ValidateData();
 
 	IncStatCounters();
 }
 
+void FISMCInstanceDataSceneProxy::UpdateIdMapping(FISMInstanceUpdateChangeSet& ChangeSet)
+{
+	// update mapping, create explicit mapping if needed
+	if (ChangeSet.bIdentityIdMap)
+	{
+		// Reset to identity mapping with the new number of instances
+		InstanceIdIndexMap.Reset(ChangeSet.PostUpdateNumInstances);
+	}
+	else
+	{
+		InstanceIdIndexMap.ResizeExplicit(ChangeSet.PostUpdateNumInstances, ChangeSet.MaxInstanceId);
+
+		// If any were removed, we need to clear the associated IDs, before updating (since they may have been added again)
+		for (TConstSetBitIterator<> It(ChangeSet.InstanceAttributeTracker.GetRemovedIterator()); It; ++It)
+		{
+			// There may be more bits set as things that are marked as removed may no longer be in the map
+			if (It.GetIndex() >= InstanceIdIndexMap.GetMaxInstanceId())
+			{
+				break;
+			}
+			InstanceIdIndexMap.SetInvalid(FPrimitiveInstanceId{It.GetIndex()});
+		}
+
+		// Update index mappings (if not identity)
+		auto IndexDelta = ChangeSet.GetIndexChangedDelta();
+		for (auto It = IndexDelta.GetIterator(); It; ++It)
+		{
+			int32 NewInstanceIndex = It.GetIndex();
+			int32 ItemIndex = It.GetItemIndex();
+			FPrimitiveInstanceId InstanceId = ChangeSet.IndexToIdMapDeltaData[ItemIndex];
+			InstanceIdIndexMap.Update(InstanceId, NewInstanceIndex);
+		}
+	}
+}
+
 void FISMCInstanceDataSceneProxy::Update(FISMInstanceUpdateChangeSet&& ChangeSet)
 {
 	check(!ChangeSet.IsFullUpdate());
@@ -273,55 +293,43 @@ void FISMCInstanceDataSceneProxy::Update(FISMInstanceUpdateChangeSet&& ChangeSet
 
 	ProxyData.Flags = ChangeSet.Flags;
 
-	int32 PostUpdateNumInstances = ChangeSet.InstanceIdIndexMap.GetMaxInstanceIndex();
-	check(ChangeSet.PostUpdateNumInstances == ChangeSet.InstanceIdIndexMap.GetMaxInstanceIndex());
+	int32 PostUpdateNumInstances = ChangeSet.PostUpdateNumInstances;
 
-	// 1. Perform any changes to instance ordering (that might have resulted from deleting or otherwise rearranging instances)
-	FDataRemap DataRemap(PostUpdateNumInstances);
-
-	// Build a remap from old instance order to new instance order, carrying out any movement requested to fill holes from removes
-		
-//FIXME:
-//	1. Resize dest arrays first, since we may move data to the upper reaches
-//	2. Filter out anything that is added or fully updated anyway
-//	3. Add validation to check that we don't validate reordering assumptions.
-
-	// Note: in general, it is more efficient to forego the source order (since it supports inefficient Remove() operations, not RemoveSwap) and instead figure out a 
-	//       better mapping for the renderer. However, this is the simplest implementation! Or supposed to be anyway.
-	for (TConstSetBitIterator<> BitIt(ChangeSet.ChangeMask.IndexChangeInstances); BitIt; ++BitIt)
-	{
-		FPrimitiveInstanceId InstanceId{BitIt.GetIndex()};
-		// Filter out added instances, since they don't exist in the previous state.
-		if (!ChangeSet.ChangeMask.AddedInstances[InstanceId])
+	// Handle data movement, needs old & new ID maps
+	// These can only be caused by removes, which means an item can only ever move towards lower index in the array.
+	// Thus, we can always safely overwrite the data in the new slot, since we do them in increasing order.
+	// NOTE: If we start allowing some other kind of permutation of the ISM data, this assumption will break.
+	// TODO: Add validation code somewhere in the pipeline.
+	const auto IndexDelta = ChangeSet.GetIndexChangedDelta();
+	for (auto It = IndexDelta.GetIterator(); It; ++It)
+	{
+		// Index in the source (e.g., component)
+		const int32 ToIndex = It.GetIndex();
+		if (!ChangeSet.InstanceAttributeTracker.TestFlag<FInstanceAttributeTracker::EFlag::Added>(ToIndex))
 		{
-			int32 ToIndex = ChangeSet.InstanceIdIndexMap.IdToIndex(InstanceId);
-			if (ToIndex != INDEX_NONE && ToIndex < PostUpdateNumInstances)
+			int32 ItemIndex = It.GetItemIndex();
+			FPrimitiveInstanceId InstanceId = ChangeSet.bIdentityIdMap ? FPrimitiveInstanceId{ToIndex} : ChangeSet.IndexToIdMapDeltaData[ItemIndex];
+			if (InstanceIdIndexMap.IsValidId(InstanceId))
 			{
-				DataRemap.Set(ToIndex, InstanceIdIndexMap.IdToIndex(InstanceId));
-			}
-		}
-	}
+				const int32 FromIndex =  InstanceIdIndexMap.IdToIndex(InstanceId);
 
-	// Iterate in order from low to high index since remove always moves data up (does not work for more general shuffling!)
-	for (TConstSetBitIterator<> BitIt(DataRemap.Used); BitIt; ++BitIt)
-	{
-		int32 ToIndex = BitIt.GetIndex();
-		int32 FromIndex = DataRemap.FromInices[ToIndex];
-		ProxyData.InstanceToPrimitiveRelative[ToIndex] = ProxyData.InstanceToPrimitiveRelative[FromIndex];
-		CondMove(ChangeSet.Flags.bHasPerInstanceCustomData, ProxyData.InstanceCustomData, FromIndex, ToIndex, ChangeSet.NumCustomDataFloats);
-		CondMove(ChangeSet.Flags.bHasPerInstanceRandom, ProxyData.InstanceRandomIDs, FromIndex, ToIndex);
-		CondMove(ChangeSet.Flags.bHasPerInstanceLMSMUVBias, ProxyData.InstanceLightShadowUVBias, FromIndex, ToIndex);
+				ProxyData.InstanceToPrimitiveRelative[ToIndex] = ProxyData.InstanceToPrimitiveRelative[FromIndex];
+				CondMove(ChangeSet.Flags.bHasPerInstanceCustomData, ProxyData.InstanceCustomData, FromIndex, ToIndex, ChangeSet.NumCustomDataFloats);
+				CondMove(ChangeSet.Flags.bHasPerInstanceRandom, ProxyData.InstanceRandomIDs, FromIndex, ToIndex);
+				CondMove(ChangeSet.Flags.bHasPerInstanceLMSMUVBias, ProxyData.InstanceLightShadowUVBias, FromIndex, ToIndex);
 #if WITH_EDITOR
-		CondMove(ChangeSet.Flags.bHasPerInstanceEditorData, ProxyData.InstanceEditorData, FromIndex, ToIndex);
+				CondMove(ChangeSet.Flags.bHasPerInstanceEditorData, ProxyData.InstanceEditorData, FromIndex, ToIndex);
 #endif
+			}
+		}
 	}
 
+	UpdateIdMapping(ChangeSet);
+	check(ChangeSet.PostUpdateNumInstances == InstanceIdIndexMap.GetMaxInstanceIndex());
+
 	FIdentityIndexRemap IndexRemap;
 	ApplyDataChanges(ChangeSet, IndexRemap, PostUpdateNumInstances, ProxyData);
 
-	ChangeMask = MoveTemp(ChangeSet.ChangeMask);
-	InstanceIdIndexMap = MoveTemp(ChangeSet.InstanceIdIndexMap);
-
 	InstanceSceneDataBuffers.EndWriteAccess(AccessTag);
 
 	InstanceSceneDataBuffers.ValidateData();
@@ -332,7 +340,8 @@ void FISMCInstanceDataSceneProxy::Update(FISMInstanceUpdateChangeSet&& ChangeSet
 void FISMCInstanceDataSceneProxy::DebugDrawInstanceChanges(FPrimitiveDrawInterface* DebugPDI, ESceneDepthPriorityGroup SceneDepthPriorityGroup)
 {
 	InstanceDataUpdateTaskInfo.WaitForUpdateCompletion();
-
+#if 0
+	// TODO: The tracked changes are not available in the proxy. Need a new  mechanism to propagate this info, probably.
 	for (int Index = 0; Index < InstanceSceneDataBuffers.GetNumInstances(); ++Index)
 	{
 		FPrimitiveInstanceId InstanceId = InstanceIdIndexMap.IndexToId(Index);
@@ -345,6 +354,7 @@ void FISMCInstanceDataSceneProxy::DebugDrawInstanceChanges(FPrimitiveDrawInterfa
 			DrawCircle(DebugPDI, InstanceToWorld.GetOrigin(), FVector(1, 0, 0), FVector(0, 1, 0), FColor::Orange, 40.0f, 32, SceneDepthPriorityGroup);
 		}
 	}
+#endif
 }
 
 FISMCInstanceDataSceneProxyLegacyReordered::FISMCInstanceDataSceneProxyLegacyReordered(FStaticShaderPlatform InShaderPlatform, ERHIFeatureLevel::Type InFeatureLevel, bool bInLegacyReordered) 
@@ -369,10 +379,10 @@ void FISMCInstanceDataSceneProxyLegacyReordered::Update(FISMInstanceUpdateChange
 
 		ProxyData.VisibleInstances.SetNum(ChangeSet.PostUpdateNumInstances, true);
 		ProxyData.Flags.bHasPerInstanceVisible = true;
-		for (TConstSetBitIterator<> BitIt(ChangeSet.ChangeMask.RemovedInstances); BitIt; ++BitIt)
+		for (auto It = ChangeSet.InstanceAttributeTracker.GetRemovedIterator(); It; ++It)
 		{
 			// This is somewhat nonintuitive, but the current instance->index map is where we retain knowledge of where the instance used to be placed (in the component address space at last update)
-			int32 InstanceIndex = InstanceIdIndexMap.IdToIndex(FPrimitiveInstanceId{BitIt.GetIndex()});
+			int32 InstanceIndex = InstanceIdIndexMap.IdToIndex(FPrimitiveInstanceId{It.GetIndex()});
 			if (IndexRemapOld.RemapIndex(InstanceIndex))
 			{
 				LOG_INST_DATA(TEXT("Update/HideInstance, ID: %d, IDX: %d"), BitIt.GetIndex(), InstanceIndex);
@@ -380,6 +390,8 @@ void FISMCInstanceDataSceneProxyLegacyReordered::Update(FISMInstanceUpdateChange
 			}
 		}
 	}
+	UpdateIdMapping(ChangeSet);
+
 	LegacyInstanceReorderTable = MoveTemp(ChangeSet.LegacyInstanceReorderTable);
 	FReorderTableIndexRemap IndexRemap(LegacyInstanceReorderTable, ChangeSet.PostUpdateNumInstances);
 
@@ -388,9 +400,6 @@ void FISMCInstanceDataSceneProxyLegacyReordered::Update(FISMInstanceUpdateChange
 	
 	InstanceSceneDataBuffers.EndWriteAccess(AccessTag);
 
-	ChangeMask = MoveTemp(ChangeSet.ChangeMask);
-	InstanceIdIndexMap = MoveTemp(ChangeSet.InstanceIdIndexMap);
-
 	InstanceSceneDataBuffers.ValidateData();
 	IncStatCounters();
 }
@@ -400,12 +409,12 @@ void FISMCInstanceDataSceneProxyLegacyReordered::Build(FISMInstanceUpdateChangeS
 {
 	DecStatCounters();
 	check(ChangeSet.IsFullUpdate());
-	check(!ChangeSet.TransformsDelta.IsDelta());
-	check(!ChangeSet.CustomDataDelta.IsDelta() || (!ChangeSet.Flags.bHasPerInstanceCustomData && ChangeSet.CustomDataDelta.IsEmpty()));
-	check(!ChangeSet.InstanceLightShadowUVBiasDelta.IsDelta() || ChangeSet.InstanceLightShadowUVBiasDelta.IsEmpty());
+	checkSlow(!ChangeSet.GetTransformDelta().IsDelta());
+	checkSlow(!ChangeSet.GetCustomDataDelta().IsDelta() || (!ChangeSet.Flags.bHasPerInstanceCustomData && ChangeSet.GetCustomDataDelta().IsEmpty()));
+	checkSlow(!ChangeSet.GetInstanceLightShadowUVBiasDelta().IsDelta() || ChangeSet.GetInstanceLightShadowUVBiasDelta().IsEmpty());
 	check(bLegacyReordered || ChangeSet.LegacyInstanceReorderTable.IsEmpty());
 #if WITH_EDITOR
-	check(!ChangeSet.InstanceEditorDataDelta.IsDelta() || ChangeSet.InstanceEditorDataDelta.IsEmpty());
+	checkSlow(!ChangeSet.GetInstanceEditorDataDelta().IsDelta() || ChangeSet.GetInstanceEditorDataDelta().IsEmpty());
 #endif
 
 	FInstanceSceneDataBuffers::FAccessTag AccessTag(PointerHash(this));
@@ -414,6 +423,8 @@ void FISMCInstanceDataSceneProxyLegacyReordered::Build(FISMInstanceUpdateChangeS
 	LegacyInstanceReorderTable = MoveTemp(ChangeSet.LegacyInstanceReorderTable);
 	ProxyData.Flags = ChangeSet.Flags;
 
+	UpdateIdMapping(ChangeSet);
+
 	FReorderTableIndexRemap IndexRemap(LegacyInstanceReorderTable, ChangeSet.PostUpdateNumInstances);
 	ApplyDataChanges(ChangeSet, IndexRemap, ChangeSet.PostUpdateNumInstances, ProxyData);
 
@@ -439,9 +450,6 @@ void FISMCInstanceDataSceneProxyLegacyReordered::Build(FISMInstanceUpdateChangeS
 	}	
 	InstanceSceneDataBuffers.EndWriteAccess(AccessTag);
 
-	ChangeMask = MoveTemp(ChangeSet.ChangeMask);
-	InstanceIdIndexMap = MoveTemp(ChangeSet.InstanceIdIndexMap);
-	
 	InstanceSceneDataBuffers.ValidateData();
 	IncStatCounters();
 }
@@ -599,8 +607,6 @@ void FISMCInstanceDataSceneProxyNoGPUScene::ApplyDataChanges(FISMInstanceUpdateC
 	{
 		// TODO: only need to process added instances? No help for ISM since the move path would be taken.
 		// TODO: OTOH for HISM there is no meaningful data, so just skipping and letting the SetNumZeroed fill in the blanks is fine.
-		FArrayIndexDelta PerInstanceRandomDelta(PostUpdateNumInstances);
-
 		InstanceRandomIDs.SetNumZeroed(PostUpdateNumInstances);
 		if (ChangeSet.GeneratePerInstanceRandomIds)
 		{
@@ -611,9 +617,12 @@ void FISMCInstanceDataSceneProxyNoGPUScene::ApplyDataChanges(FISMInstanceUpdateC
 	// unpack transform deltas
 	// TODO: Only do if requested / needed
 	ProxyData.InstanceToPrimitiveRelative.SetNumUninitialized(PostUpdateNumInstances);
-	for (int32 PackedIndex = 0; PackedIndex < ChangeSet.TransformsDelta.Num(); ++PackedIndex)
+	auto TransformDelta = ChangeSet.GetTransformDelta();
+	for (auto It = TransformDelta.GetIterator(); It; ++It)
 	{
-		int32 InstanceIndex = ChangeSet.TransformsDelta[PackedIndex];
+		int32 PackedIndex = It.GetItemIndex();
+		int32 InstanceIndex = It.GetIndex();
+
 		if (IndexRemap.RemapIndex(InstanceIndex))
 		{
 			LegacyInstanceData.SetInstance(InstanceIndex, ChangeSet.Transforms[PackedIndex].ToMatrix44f(), ChangeSet.Flags.bHasPerInstanceRandom ? InstanceRandomIDs[InstanceIndex] : 0.0f);
@@ -628,9 +637,11 @@ void FISMCInstanceDataSceneProxyNoGPUScene::ApplyDataChanges(FISMInstanceUpdateC
 
 	if (ChangeSet.Flags.bHasPerInstanceCustomData)
 	{
-		for (int32 PackedIndex = 0; PackedIndex < ChangeSet.CustomDataDelta.Num(); ++PackedIndex)
+		auto CustomDataDelta = ChangeSet.GetCustomDataDelta();
+		for (auto It = CustomDataDelta.GetIterator(); It; ++It)
 		{
-			int32 InstanceIndex = ChangeSet.CustomDataDelta[PackedIndex];
+			int32 PackedIndex = It.GetItemIndex();
+			int32 InstanceIndex = It.GetIndex();
 			if (IndexRemap.RemapIndex(InstanceIndex))
 			{
 				for (int32 j = 0; j < ProxyData.NumCustomDataFloats; ++j)
@@ -642,9 +653,11 @@ void FISMCInstanceDataSceneProxyNoGPUScene::ApplyDataChanges(FISMInstanceUpdateC
 	}
 	if (ChangeSet.Flags.bHasPerInstanceLMSMUVBias)
 	{
-		for (int32 PackedIndex = 0; PackedIndex < ChangeSet.InstanceLightShadowUVBiasDelta.Num(); ++PackedIndex)
+		auto InstanceLightShadowUVBiasDelta = ChangeSet.GetInstanceLightShadowUVBiasDelta();
+		for (auto It = InstanceLightShadowUVBiasDelta.GetIterator(); It; ++It)
 		{
-			int32 InstanceIndex = ChangeSet.InstanceLightShadowUVBiasDelta[PackedIndex];
+			int32 PackedIndex = It.GetItemIndex();
+			int32 InstanceIndex = It.GetIndex();
 			if (IndexRemap.RemapIndex(InstanceIndex))
 			{
 				FVector4f Packed = ChangeSet.InstanceLightShadowUVBias[PackedIndex];
@@ -659,9 +672,11 @@ void FISMCInstanceDataSceneProxyNoGPUScene::ApplyDataChanges(FISMInstanceUpdateC
 #if WITH_EDITOR
 	if (ChangeSet.Flags.bHasPerInstanceEditorData)
 	{
-		for (int32 PackedIndex = 0; PackedIndex < ChangeSet.InstanceEditorDataDelta.Num(); ++PackedIndex)
+		auto InstanceEditorDataDelta = ChangeSet.GetInstanceEditorDataDelta();
+		for (auto It = InstanceEditorDataDelta.GetIterator(); It; ++It)
 		{
-			int32 InstanceIndex = ChangeSet.InstanceEditorDataDelta[PackedIndex];
+			int32 PackedIndex = It.GetItemIndex();
+			int32 InstanceIndex = It.GetIndex();
 			if (IndexRemap.RemapIndex(InstanceIndex))
 			{
 				FColor HitProxyColor;
@@ -696,7 +711,6 @@ void FISMCInstanceDataSceneProxyNoGPUScene::Update(FISMInstanceUpdateChangeSet&&
 {
 	check(!ChangeSet.IsFullUpdate());
 
-	check(bLegacyReordered || ChangeSet.PostUpdateNumInstances == ChangeSet.InstanceIdIndexMap.GetMaxInstanceIndex());
 	check(bLegacyReordered || ChangeSet.LegacyInstanceReorderTable.IsEmpty());
 
 	DecStatCounters();
@@ -712,10 +726,10 @@ void FISMCInstanceDataSceneProxyNoGPUScene::Update(FISMInstanceUpdateChangeSet&&
 		// HISMTODO: move to own implementation
 		ProxyData.VisibleInstances.SetNum(ChangeSet.PostUpdateNumInstances, true);
 		ProxyData.Flags.bHasPerInstanceVisible = true;
-		for (TConstSetBitIterator<> BitIt(ChangeSet.ChangeMask.RemovedInstances); BitIt; ++BitIt)
+		for (auto It = ChangeSet.InstanceAttributeTracker.GetRemovedIterator(); It; ++It)
 		{
 			// This is somewhat nonintuitive, but the current instance->index map is where we retain knowledge of where the instance used to be placed (in the component address space at last update)
-			int32 InstanceIndex = InstanceIdIndexMap.IdToIndex(FPrimitiveInstanceId{BitIt.GetIndex()});
+			int32 InstanceIndex = InstanceIdIndexMap.IdToIndex(FPrimitiveInstanceId{It.GetIndex()});
 			if (IndexRemapOld.RemapIndex(InstanceIndex))
 			{
 				LOG_INST_DATA(TEXT("Update/HideInstance, ID: %d, IDX: %d"), BitIt.GetIndex(), InstanceIndex);
@@ -723,6 +737,10 @@ void FISMCInstanceDataSceneProxyNoGPUScene::Update(FISMInstanceUpdateChangeSet&&
 			}
 		}
 	}
+
+	UpdateIdMapping(ChangeSet);
+	check(bLegacyReordered || ChangeSet.PostUpdateNumInstances == InstanceIdIndexMap.GetMaxInstanceIndex());
+
 	LegacyInstanceReorderTable = MoveTemp(ChangeSet.LegacyInstanceReorderTable);
 	FReorderTableIndexRemap IndexRemap(LegacyInstanceReorderTable, ChangeSet.PostUpdateNumInstances);
 
@@ -732,9 +750,6 @@ void FISMCInstanceDataSceneProxyNoGPUScene::Update(FISMInstanceUpdateChangeSet&&
 
 	LegacyInstanceBuffer->SetFlushToGPUPending();
 		
-	ChangeMask = MoveTemp(ChangeSet.ChangeMask);
-	InstanceIdIndexMap = MoveTemp(ChangeSet.InstanceIdIndexMap);
-
 	IncStatCounters();
 }
 
@@ -742,16 +757,18 @@ void FISMCInstanceDataSceneProxyNoGPUScene::Build(FISMInstanceUpdateChangeSet&&
 {
 	DecStatCounters();
 	check(ChangeSet.IsFullUpdate());
-	check(!ChangeSet.TransformsDelta.IsDelta());
-	check(!ChangeSet.CustomDataDelta.IsDelta() || (!ChangeSet.Flags.bHasPerInstanceCustomData && ChangeSet.CustomDataDelta.IsEmpty()));
-	check(!ChangeSet.InstanceLightShadowUVBiasDelta.IsDelta() || ChangeSet.InstanceLightShadowUVBiasDelta.IsEmpty());
+	checkSlow(!ChangeSet.GetTransformDelta().IsDelta());
+	checkSlow(!ChangeSet.GetCustomDataDelta().IsDelta() || (!ChangeSet.Flags.bHasPerInstanceCustomData && ChangeSet.GetCustomDataDelta().IsEmpty()));
+	checkSlow(!ChangeSet.GetInstanceLightShadowUVBiasDelta().IsDelta() || ChangeSet.GetInstanceLightShadowUVBiasDelta().IsEmpty());
 #if WITH_EDITOR
-	check(!ChangeSet.InstanceEditorDataDelta.IsDelta() || ChangeSet.InstanceEditorDataDelta.IsEmpty());
+	checkSlow(!ChangeSet.GetInstanceEditorDataDelta().IsDelta() || ChangeSet.GetInstanceEditorDataDelta().IsEmpty());
 #endif
 
-	check(bLegacyReordered || ChangeSet.PostUpdateNumInstances == ChangeSet.InstanceIdIndexMap.GetMaxInstanceIndex());
 	check(bLegacyReordered || ChangeSet.LegacyInstanceReorderTable.IsEmpty());
 
+	UpdateIdMapping(ChangeSet);
+	check(bLegacyReordered || ChangeSet.PostUpdateNumInstances == InstanceIdIndexMap.GetMaxInstanceIndex());
+
 	LegacyInstanceReorderTable = MoveTemp(ChangeSet.LegacyInstanceReorderTable);
 
 	FInstanceSceneDataBuffers::FAccessTag AccessTag(PointerHash(this));
@@ -790,9 +807,6 @@ void FISMCInstanceDataSceneProxyNoGPUScene::Build(FISMInstanceUpdateChangeSet&&
 	LegacyInstanceBuffer->InitFromPreallocatedData(LegacyInstanceData);
 	LegacyInstanceBuffer->SetFlushToGPUPending();
 
-	ChangeMask = MoveTemp(ChangeSet.ChangeMask);
-	InstanceIdIndexMap = MoveTemp(ChangeSet.InstanceIdIndexMap);
-
 	IncStatCounters();
 }
 
diff --git a/Engine/Source/Runtime/Engine/Private/InstancedStaticMesh/ISMInstanceUpdateChangeSet.cpp b/Engine/Source/Runtime/Engine/Private/InstancedStaticMesh/ISMInstanceUpdateChangeSet.cpp
--- a/Engine/Source/Runtime/Engine/Private/InstancedStaticMesh/ISMInstanceUpdateChangeSet.cpp
+++ b/Engine/Source/Runtime/Engine/Private/InstancedStaticMesh/ISMInstanceUpdateChangeSet.cpp
@@ -1,6 +1,7 @@
 // Copyright Epic Games, Inc. All Rights Reserved.
 
 #include "InstancedStaticMesh/ISMInstanceUpdateChangeSet.h"
+#include "InstancedStaticMesh/ISMScatterGatherUtil.h"
 #include "Engine/InstancedStaticMesh.h"
 #include "DataDrivenShaderPlatformInfo.h"
 #include "Rendering/RenderingSpatialHash.h"
@@ -42,7 +43,7 @@ FOpaqueHitProxyContainer::~FOpaqueHitProxyContainer()
 void FISMInstanceUpdateChangeSet::SetEditorData(const TArray<TRefCountPtr<HHitProxy>>& HitProxies, const TBitArray<> &InSelectedInstances)//, bool bWasHitProxiesReallocated)
 {
 	HitProxyContainer = MakeOpaqueHitProxyContainer(HitProxies);
-	for (int32 Index : InstanceEditorDataDelta)
+	for (int32 Index = 0; Index < HitProxies.Num(); ++Index)
 	{
 		// Record if the instance is selected
 		FColor HitProxyColor(ForceInit);
@@ -58,6 +59,27 @@ void FISMInstanceUpdateChangeSet::SetEditorData(const TArray<TRefCountPtr<HHitPr
 
 #endif
 
+void FISMInstanceUpdateChangeSet::SetInstanceTransforms(TStridedView<FMatrix> InInstanceTransforms, const FVector Offset)
+{
+
+	GatherTransform(GetTransformDelta(), Transforms, InInstanceTransforms, [Offset](const FMatrix &M) -> FRenderTransform { return FRenderTransform(M.ConcatTranslation(Offset)); });
+}
+
+void FISMInstanceUpdateChangeSet::SetInstanceTransforms(TStridedView<FMatrix> InInstanceTransforms)
+{
+	GatherTransform(GetTransformDelta(), Transforms, InInstanceTransforms, [](const FMatrix &M) -> FRenderTransform { return FRenderTransform(M); });
+}
+
+void FISMInstanceUpdateChangeSet::SetInstanceTransforms(TStridedView<FMatrix> InInstanceTransforms, FBox const& InInstanceBounds, FBox& OutGatheredBounds)
+{
+	GatherTransform(GetTransformDelta(),Transforms, InInstanceTransforms, [&InInstanceBounds, &OutGatheredBounds](const FMatrix& M) -> FRenderTransform
+	{
+		FRenderTransform Transform(M);
+		OutGatheredBounds += InInstanceBounds.TransformBy(Transform.ToMatrix());
+		return Transform;
+	});
+}
+
 void FISMInstanceUpdateChangeSet::SetInstancePrevTransforms(TArrayView<FMatrix> InPrevInstanceTransforms, const FVector &Offset)
 {
 	if (Flags.bHasPerInstanceDynamicData)
@@ -68,8 +90,7 @@ void FISMInstanceUpdateChangeSet::SetInstancePrevTransforms(TArrayView<FMatrix>
 		}
 		else
 		{
-
-			TransformsDelta.GatherTransform(PrevTransforms, InPrevInstanceTransforms, [Offset](const FMatrix &M) -> FRenderTransform { return FRenderTransform(M.ConcatTranslation(Offset)); });
+			GatherTransform(GetTransformDelta(), PrevTransforms, InPrevInstanceTransforms, [Offset](const FMatrix &M) -> FRenderTransform { return FRenderTransform(M.ConcatTranslation(Offset)); });
 		}
 	}
 }
@@ -84,8 +105,7 @@ void FISMInstanceUpdateChangeSet::SetInstancePrevTransforms(TArrayView<FMatrix>
 		}
 		else
 		{
-
-			TransformsDelta.GatherTransform(PrevTransforms, InPrevInstanceTransforms, [](const FMatrix &M) -> FRenderTransform { return FRenderTransform(M); });
+			GatherTransform(GetTransformDelta(), PrevTransforms, InPrevInstanceTransforms, [](const FMatrix &M) -> FRenderTransform { return FRenderTransform(M); });
 		}
 	}
 }
@@ -96,13 +116,13 @@ void FISMInstanceUpdateChangeSet::SetCustomData(const TArrayView<const float> &I
 	{
 		check(InNumCustomDataFloats == NumCustomDataFloats);
 		check(NumCustomDataFloats > 0);
-		check(InPerInstanceCustomData.Num() == InstanceIdIndexMap.GetMaxInstanceIndex() * NumCustomDataFloats);
+		//check(InPerInstanceCustomData.Num() == InstanceIdIndexMap.GetMaxInstanceIndex() * NumCustomDataFloats);
 
-		CustomDataDelta.Gather(PerInstanceCustomData, InPerInstanceCustomData, InNumCustomDataFloats);
+		Gather(GetCustomDataDelta(), PerInstanceCustomData, InPerInstanceCustomData, InNumCustomDataFloats);
 	}
 	else
 	{
-		check(CustomDataDelta.IsEmpty());
+		checkSlow(GetCustomDataDelta().IsEmpty());
 	}
 }
 
diff --git a/Engine/Source/Runtime/Engine/Private/InstancedStaticMesh/ISMScatterGatherUtil.h b/Engine/Source/Runtime/Engine/Private/InstancedStaticMesh/ISMScatterGatherUtil.h
new file mode 100644
--- /dev/null
+++ b/Engine/Source/Runtime/Engine/Private/InstancedStaticMesh/ISMScatterGatherUtil.h
@@ -0,0 +1,105 @@
+// Copyright Epic Games, Inc. All Rights Reserved.
+
+#pragma once
+
+#include "CoreTypes.h"
+#include "HAL/UnrealMemory.h"
+#include "Math/UnrealMathUtility.h"
+#include "Containers/Array.h"
+
+/**
+ * Write InSource that was previously gathered using the above methods to the final destination array OutDest
+ * using the same delta information. 
+ * If there is no delta, it performs a move of the source data to the final array, saving a copy.
+ */
+template <typename DeltaType, typename ValueType>
+void Scatter(const DeltaType &Delta, TArray<ValueType> &OutDest, int32 DestNumElements, TArray<ValueType> &&InSource, int32 ElementStride = 1)
+{
+	check(InSource.Num() == Delta.GetNumItems() * ElementStride);
+	if (Delta.IsDelta())
+	{
+		OutDest.SetNumUninitialized(DestNumElements * ElementStride);
+		for (auto It = Delta.GetIterator(); It; ++It)
+		{
+			int32 ItemIndex = It.GetItemIndex();
+			int32 DestIndex = It.GetIndex();
+			FMemory::Memcpy(&OutDest[DestIndex * ElementStride], &InSource[ItemIndex * ElementStride], ElementStride * sizeof(ValueType));
+		}
+	}
+	else
+	{
+		check(InSource.Num() == DestNumElements * ElementStride);
+		OutDest = MoveTemp(InSource);
+	}
+}
+/**
+ * Write InSource that was previously gathered using the above methods to the final destination array OutDest
+ * using the same delta information. 
+ * Never takes move shortcut as there is an index remap.
+ */
+template <typename DeltaType, typename ValueType, typename DestIndexRemapType>
+void Scatter(const DeltaType &Delta, TArray<ValueType> &OutDest, int32 DestNumElements, const TArray<ValueType> &InSource, const DestIndexRemapType &DestIndexRemap, int32 ElementStride = 1)
+{
+	check(InSource.Num() == Delta.GetNumItems() * ElementStride);
+	OutDest.SetNumUninitialized(DestNumElements * ElementStride);
+	for (auto It = Delta.GetIterator(); It; ++It)
+	{
+		int32 ItemIndex = It.GetItemIndex();
+		int32 DestIndex = DestIndexRemap[It.GetIndex()];
+		if (DestIndex != INDEX_NONE)
+		{
+			FMemory::Memcpy(&OutDest[DestIndex * ElementStride], &InSource[ItemIndex * ElementStride], ElementStride * sizeof(ValueType));
+		}
+	}
+}
+
+/**
+ * Gather the needed values from InSource to OutDest, according to the delta.
+ * If there is no delta, it will perform a bulk copy.
+ */
+template <typename DeltaType, typename ValueType, typename InValueArrayType>
+void Gather(const DeltaType &Delta, TArray<ValueType> &OutDest, const InValueArrayType &InSource, int32 ElementStride = 1)
+{
+	// strides & element count matches - just copy the data
+	if (InSource.Num() == Delta.GetNumItems() * ElementStride)
+	{
+		OutDest = InSource;
+	}
+	else if (Delta.IsEmpty())
+	{
+		OutDest.Reset();
+	}
+	else if (Delta.IsDelta() || ElementStride != 1)
+	{
+		OutDest.Reset(Delta.GetNumItems() * ElementStride);
+		for (auto It = Delta.GetIterator(); It; ++It)
+		{
+			check(OutDest.Num() < Delta.GetNumItems() * ElementStride);
+			OutDest.Append(&InSource[It.GetIndex() * ElementStride], ElementStride);
+		}
+	}
+}
+
+
+/**
+ * Gather the needed values from InSource to OutDest, according to the delta.
+ * If there is no delta, it will perform a bulk copy.
+ */
+template <typename DeltaType, typename OutValueType, typename InValueArrayType, typename LambdaType>
+void GatherTransform(const DeltaType &Delta, TArray<OutValueType> &OutDest, const InValueArrayType &InSource, LambdaType TransformLambda)
+{
+	if (Delta.IsEmpty())
+	{
+		OutDest.Reset();
+	}
+	else
+	{
+		OutDest.Reset(Delta.GetNumItems());
+		for (auto It = Delta.GetIterator(); It; ++It)
+		{
+			check(OutDest.Num() < Delta.GetNumItems());
+			check(OutDest.Num() == It.GetItemIndex());
+			OutDest.Add(TransformLambda(InSource[It.GetIndex()]));
+		}
+	}
+}
diff --git a/Engine/Source/Runtime/Engine/Private/InstancedStaticMesh/InstanceAttributeTracker.cpp b/Engine/Source/Runtime/Engine/Private/InstancedStaticMesh/InstanceAttributeTracker.cpp
new file mode 100644
--- /dev/null
+++ b/Engine/Source/Runtime/Engine/Private/InstancedStaticMesh/InstanceAttributeTracker.cpp
@@ -0,0 +1,95 @@
+// Copyright Epic Games, Inc. All Rights Reserved.
+
+#include "InstancedStaticMesh/InstanceAttributeTracker.h"
+
+FInstanceAttributeTracker::FInstanceAttributeTracker()
+{
+	Reset();
+}
+
+FInstanceAttributeTracker::FInstanceAttributeTracker(FInstanceAttributeTracker &&Other)
+{
+	Move(*this, Other);
+}
+
+void FInstanceAttributeTracker::operator=(FInstanceAttributeTracker &&Other)
+{
+	Move(*this, Other);
+}
+
+void FInstanceAttributeTracker::Reset()
+{
+	Removed.Reset();
+	Data.Reset();
+	MaxIndex = 0;
+	FirstChangedIndex = MAX_int32;
+	LastChangedIndex = 0;		
+	FirstRemovedIdIndex = INT32_MAX;
+	FMemory::Memset(NumChanged, 0, sizeof(NumChanged));
+}
+	
+void FInstanceAttributeTracker::Move(FInstanceAttributeTracker &Dest, FInstanceAttributeTracker &Source)
+{
+	Dest.Removed = MoveTemp(Source.Removed);
+	Dest.Data = MoveTemp(Source.Data);
+	Dest.FirstRemovedIdIndex = Source.FirstRemovedIdIndex;
+	Dest.MaxIndex = Source.MaxIndex;
+	Dest.FirstChangedIndex = Source.FirstChangedIndex;
+	Dest.LastChangedIndex = Source.LastChangedIndex;
+	FMemory::Memcpy(Dest.NumChanged, Source.NumChanged, sizeof(Dest.NumChanged));
+		
+	Source.Reset();
+}
+
+TConstSetBitIterator<> FInstanceAttributeTracker::GetRemovedIterator() const
+{
+	return TConstSetBitIterator<>(Removed, FMath::Min(Removed.Num(), FirstRemovedIdIndex));
+}
+
+#if DO_GUARD_SLOW
+
+void FInstanceAttributeTracker::Validate() const
+{
+	int32 NumChangedCounts[uint32(EFlag::Num)];
+	FMemory::Memset(NumChangedCounts, 0, sizeof(NumChangedCounts));
+
+	for (int32 Index = 0; Index < MaxIndex; ++Index)
+	{
+		uint32 Flags = GetFlags(Index);
+		if (Flags & FToBit<EFlag::Added>::Bit)
+		{
+			NumChangedCounts[uint32(EFlag::Added)] += 1;
+		}
+		else
+		{
+			for (uint32 Flag = uint32(EFlag::TransformChanged); Flag < uint32(EFlag::Num); ++Flag)
+			{
+				// If it is set & not added
+				if (Flags & (1u << Flag))
+				{
+					NumChangedCounts[Flag] += 1;
+				}
+			}
+		}
+	}
+	for (uint32 Flag = 0u; Flag < uint32(EFlag::Num); ++Flag)
+	{
+		check(NumChangedCounts[Flag] == NumChanged[Flag]);
+	}
+
+	if (MaxIndex > 0)
+	{
+		int32 LastIndex = MaxIndex - 1;
+		const int32 ElementIndex = LastIndex / MasksPerElement;
+		const uint32 SubIndex = (LastIndex % MasksPerElement) + 1u;
+		if (SubIndex < MasksPerElement)
+		{
+			const uint32 ElementSubShift = SubIndex * uint32(EFlag::Num);
+			// Load and clear the mask for this index
+			ElementType LastElement = Data[ElementIndex];
+			check((LastElement >> ElementSubShift) == 0u);
+		}
+	}
+}
+
+#endif
\ No newline at end of file
diff --git a/Engine/Source/Runtime/Engine/Public/InstanceDataSceneProxy.h b/Engine/Source/Runtime/Engine/Public/InstanceDataSceneProxy.h
--- a/Engine/Source/Runtime/Engine/Public/InstanceDataSceneProxy.h
+++ b/Engine/Source/Runtime/Engine/Public/InstanceDataSceneProxy.h
@@ -9,85 +9,78 @@
 
 class FPrimitiveDrawInterface;
 
-/**
- * Helper to make it easy to query an empty array & index by FPrimitiveInstanceId.
- */
-class FChangeBitArray : public TBitArray<>
+class FInstanceIdIndexMap
 {
 public:
 	/**
-	 * Returns true if the corresponding bit is set, false if it is not, OR if it does not exist.
+	 * Returns true if the mapping is an identity mapping. I.e., each instance ID corresponds to the same index.
 	 */
-	inline bool WasChanged(FPrimitiveInstanceId InstanceId) const 
-	{ 
-		return IsValidIndex(InstanceId.GetAsIndex()) && TBitArray<>::operator[](InstanceId.GetAsIndex()); 
-	}
+	FORCEINLINE bool IsIdentity() const { return IndexToIdMap.IsEmpty(); }
 
-	inline bool operator[](FPrimitiveInstanceId InstanceId) const
-	{
-		return WasChanged(InstanceId);
-	}
-
-	inline FBitReference operator[](FPrimitiveInstanceId InstanceId)
-	{
-		return TBitArray<>::operator[](InstanceId.GetAsIndex());
-	}
-};
-
-struct FChangeMask
-{
-
-	FChangeBitArray AddedInstances;
-	FChangeBitArray RemovedInstances;
-	FChangeBitArray TransformChangedInstances;
-	FChangeBitArray CustomDataChangedInstances;
-#if WITH_EDITOR
-	bool bAnyEditorDataChanged = false;
-#endif	
-	// Do we need to track this?
-	FChangeBitArray IndexChangeInstances;
-	bool bNumCustomDataChanged = false;
-	bool bBakedLightingDataChanged = false;
-
-	void Reset()
-	{
-		AddedInstances.Reset();
-		RemovedInstances.Reset();
-		TransformChangedInstances.Reset();
-		CustomDataChangedInstances.Reset();
-#if WITH_EDITOR
-		bAnyEditorDataChanged = false;
-#endif	
-		IndexChangeInstances.Reset();
-		bNumCustomDataChanged = false;
-		bBakedLightingDataChanged = false;
-	}
-};
+	/**
+	 * Returms the upper bound on the instance ID represented as an integer, may be larger than the number of instances but is never larger than the max that has ever been allocated.
+	 */
+	FORCEINLINE  int32 GetMaxInstanceId() const { return IsIdentity() ? NumInstances : IdToIndexMap.Num();}
 
-class FInstanceIdIndexMap
-{
-public:
-	bool IsIdentity() const { return IndexToIdMap.IsEmpty(); }
+	/**
+	 * Returns the maximum instance index (or rather the maximum valid number plus one), AKA the number of instances represented.
+	 */
+	FORCEINLINE  int32 GetMaxInstanceIndex() const { return IsIdentity() ? NumInstances : IndexToIdMap.Num();}
 
-	inline int32 GetMaxInstanceId() const { return IsIdentity() ? NumInstances : IdToIndexMap.Num();}
-	inline int32 GetMaxInstanceIndex() const { return IsIdentity() ? NumInstances : IndexToIdMap.Num();}
-	inline bool IsValidId(FPrimitiveInstanceId InstanceId) const { return InstanceId.Id >= 0 && InstanceId.Id < GetMaxInstanceId() && (IsIdentity() || IdToIndexMap[InstanceId.Id] != INDEX_NONE); }
-	inline int32 IdToIndex(FPrimitiveInstanceId InstanceId) const
+	/**
+	 * Returns true if InstanceId is in the mapped range AND .
+	 */
+	FORCEINLINE  bool IsValidId(FPrimitiveInstanceId InstanceId) const { return InstanceId.Id >= 0 && InstanceId.Id < GetMaxInstanceId() && (IsIdentity() || IdToIndexMap[InstanceId.Id] != INDEX_NONE); }
+	/**
+	 * Translate from FPrimitiveInstanceId to Index
+	 */
+	FORCEINLINE  int32 IdToIndex(FPrimitiveInstanceId InstanceId) const
 	{ 
 		return IsIdentity() ? InstanceId.Id : IdToIndexMap[InstanceId.Id]; 
 	}
 
-	inline FPrimitiveInstanceId IndexToId(int32 InstanceIndex) const
+	/**
+	 * Translate from Index to FPrimitiveInstanceId
+	 */
+	FORCEINLINE  FPrimitiveInstanceId IndexToId(int32 InstanceIndex) const
 	{ 
 		check(InstanceIndex < GetMaxInstanceIndex());
 		return IsIdentity() ? FPrimitiveInstanceId{ InstanceIndex } : IndexToIdMap[InstanceIndex]; 
 	}
 
-	void Reset()
+	/**
+	 * Set the index mapping for the given InstanceId to INDEX_NONE.
+	 */
+	FORCEINLINE void SetInvalid(FPrimitiveInstanceId InstanceId)
+	{
+		IdToIndexMap[InstanceId.Id] = INDEX_NONE;
+	}
+
+	/**
+	 * Update the mapping from ID to index (and vice versa).
+	 */
+	FORCEINLINE void Update(FPrimitiveInstanceId InstanceId, int32 InstanceIndex)
 	{
-		IndexToIdMap.Reset();
-		IdToIndexMap.Reset();		
+		IndexToIdMap[InstanceIndex] = InstanceId;
+		IdToIndexMap[InstanceId.Id] = InstanceIndex;
 	}
+
+	/**
+	 * Reset the mapping to an identity map of size InNumInstances.
+	 */
+	ENGINE_API void Reset(int32 InNumInstances);
+
+	/**
+	 * Make sure the mapping is explicit and add enough space to accommodate the InNumInstances & MaxInstanceId.
+	 */
+	ENGINE_API void ResizeExplicit(int32 InNumInstances, int32 MaxInstanceId);
+
+	/**
+	 * Convert an implicit identity mapping to an explicit one, by filling in the mapping arrays. 
+	 * Not allowed to be called if the mapping is already explicit.
+	 */
+	ENGINE_API void CreateExplicitIdentityMapping();
+
 	FInstanceIdIndexMap() = default;
 protected:
 	// Bidirectional mapping to / from ID.
diff --git a/Engine/Source/Runtime/Engine/Public/InstancedStaticMesh/ISMInstanceDataManager.h b/Engine/Source/Runtime/Engine/Public/InstancedStaticMesh/ISMInstanceDataManager.h
--- a/Engine/Source/Runtime/Engine/Public/InstancedStaticMesh/ISMInstanceDataManager.h
+++ b/Engine/Source/Runtime/Engine/Public/InstancedStaticMesh/ISMInstanceDataManager.h
@@ -3,6 +3,7 @@
 #pragma once
 
 #include "ISMInstanceDataSceneProxy.h"
+#include "InstanceAttributeTracker.h"
 
 class FPrimitiveSceneProxy;
 class HHitProxy;
@@ -120,22 +121,19 @@ public:
 
 	int32 GetMaxInstanceIndex() const;
 
-	/**
-	 * Resize (grow) the tracking state for data changes.
-	 * Returns true if the state is actually tracked, false otherwise.
-	 */
-	bool ResizeChangeTrackingState();
-
 	void CreateExplicitIdentityMapping();
 
-
 	ETrackingState GetState() const { return TrackingState; }
 
 	ENGINE_API TSharedPtr<FISMCInstanceDataSceneProxy, ESPMode::ThreadSafe> GetOrCreateProxy(FStaticShaderPlatform InShaderPlatform, ERHIFeatureLevel::Type InFeatureLevel);
 
 	void Invalidate(int32 InNumInstances);
 	
+#if DO_GUARD_SLOW
 	void ValidateMapping() const;
+#else
+	FORCEINLINE void ValidateMapping() const {};
+#endif
 
 	/**
 	 * Call to mark the manager as needing a full rebuild & having an external driver for this.
@@ -163,8 +161,12 @@ private:
 	template <typename TaskLambdaType>
 	void DispatchUpdateTask(bool bUnattached, const FInstanceDataBufferHeader &InstanceDataBufferHeader, TaskLambdaType &&TaskLambda);
 
-	inline void MarkChangeHelper(FChangeBitArray& TrackingArray, int32 InstanceIndex);
-	inline void MarkChangeHelper(FChangeBitArray& TrackingArray, FPrimitiveInstanceId InstanceId);
+	using EChangeFlag = FInstanceAttributeTracker::EFlag;
+
+	template<EChangeFlag Flag>
+	inline void MarkChangeHelper(int32 InstanceIndex);
+	template<EChangeFlag Flag>
+	inline void MarkChangeHelper(FPrimitiveInstanceId InstanceId);
 
 	void MarkComponentRenderInstancesDirty();
 
@@ -172,9 +174,7 @@ private:
 
 	bool ShouldTrackIds() const;
 
-	void MarkRemoved(FPrimitiveInstanceId InstanceId);
-
-	void MarkIndexChanged(FPrimitiveInstanceId InstanceId);
+	void FreeInstanceId(FPrimitiveInstanceId InstanceId);
 
 	void InitChangeSet(const union FChangeDesc &ChangeDesc, const FInstanceUpdateComponentDesc &ComponentData, FISMInstanceUpdateChangeSet &ChangeSet);
 
@@ -186,11 +186,8 @@ private:
 	int32 IdSearchStartIndex = 0;
 
 	// Change set.
-	FChangeBitArray AddedInstances;
-	FChangeBitArray RemovedInstances;
-	FChangeBitArray TransformChangedInstances;
-	FChangeBitArray CustomDataChangedInstances;
-	FChangeBitArray BakedLightingDataChangedInstances;
+	FInstanceAttributeTracker InstanceUpdateTracker;
+
 	bool bNumCustomDataChanged = false;
 	bool bBakedLightingDataChanged = false;
 	bool bTransformChangedAllInstances = false;
@@ -198,7 +195,6 @@ private:
 	bool bAnyEditorDataChanged = false;
 #endif	
 	bool bPrimitiveTransformChanged = false;
-	FChangeBitArray IndexChangeInstances;
 
 	TSharedPtr<FISMCInstanceDataSceneProxy, ESPMode::ThreadSafe> Proxy;
 	TWeakObjectPtr<UPrimitiveComponent> PrimitiveComponent = nullptr;
diff --git a/Engine/Source/Runtime/Engine/Public/InstancedStaticMesh/ISMInstanceDataSceneProxy.h b/Engine/Source/Runtime/Engine/Public/InstancedStaticMesh/ISMInstanceDataSceneProxy.h
--- a/Engine/Source/Runtime/Engine/Public/InstancedStaticMesh/ISMInstanceDataSceneProxy.h
+++ b/Engine/Source/Runtime/Engine/Public/InstancedStaticMesh/ISMInstanceDataSceneProxy.h
@@ -65,6 +65,9 @@ public:
 	ENGINE_API static FVector3f GetLocalBoundsPadExtent(const FRenderTransform& LocalToWorld, float PadAmount);
 
 protected:
+	// Update the InstanceIdIndexMap given the change set.
+	void UpdateIdMapping(FISMInstanceUpdateChangeSet& ChangeSet);
+
 	template <typename IndexRemapType>
 	void ApplyDataChanges(FISMInstanceUpdateChangeSet &ChangeSet, const IndexRemapType &IndexRemap, int32 PostUpdateNumInstances, FInstanceSceneDataBuffers::FWriteView &ProxyData);
 	
@@ -79,8 +82,6 @@ protected:
 	TBitArray<> ValidInstanceIdMask; // redundant, unsure if we need a copy of that, except insofar as to be able to scan the valid ones quicky? IdToIndexMap has the same info.
 	FInstanceIdIndexMap InstanceIdIndexMap;
 
-	FChangeMask ChangeMask;
-
 	FInstanceDataUpdateTaskInfo InstanceDataUpdateTaskInfo;
 	// True when it has never been updated before.
 	bool bIsNew = true;
diff --git a/Engine/Source/Runtime/Engine/Public/InstancedStaticMesh/ISMInstanceUpdateChangeSet.h b/Engine/Source/Runtime/Engine/Public/InstancedStaticMesh/ISMInstanceUpdateChangeSet.h
--- a/Engine/Source/Runtime/Engine/Public/InstancedStaticMesh/ISMInstanceUpdateChangeSet.h
+++ b/Engine/Source/Runtime/Engine/Public/InstancedStaticMesh/ISMInstanceUpdateChangeSet.h
@@ -4,214 +4,142 @@
 
 #include "InstanceDataSceneProxy.h"
 #include "Containers/StridedView.h"
+#include "InstanceAttributeTracker.h"
 
 class HHitProxy;
 class FOpaqueHitProxyContainer;
 
 /**
- * Utility to gather and scatter a set of delta-updating values.
- * If there is no delta, and in other words there is a bulk update, the "scatter" to destination can use a more efficient move operation.
+ * Helper to make it possible to use the same paths for scattering through the remap...
  */
-class FArrayIndexDelta
+class FIdentityDeltaRange
 {
 public:
-	FArrayIndexDelta(int32 InNumElements = -1) : NumElements(InNumElements) {}
+	FIdentityDeltaRange(int32 InNum) : NumItems(InNum) {}
 
 	/**
-	 * Set this delta to represent a full (dense) delta.
-	 */
-	void SetToFull(int32 InNumElements)
-	{
-		NumElements = InNumElements;
-		DeltaIndices.Empty();
-	}
+		*/
+	inline bool IsEmpty() const { return NumItems == 0; }
 
-	bool IsDelta() const { return NumElements < 0; }
-	int32 Num() const { return IsDelta() ? DeltaIndices.Num() : NumElements; }
-	bool IsEmpty() const { return Num() == 0; }
-
-	void AddIndex(int32 Index)
-	{
-		NumElements = -1;
-		DeltaIndices.Add(Index);
-	}
-
-	int32 operator[](int32 Index) const
-	{
-		if (IsDelta())
-		{
-			return DeltaIndices[Index];
-		}
-		return Index;
-	}
+	/**
+		*/
+	inline bool IsDelta() const { return false; }
 
+	/**
+		* Returns the number of items in this range - i.e., the number of items that need to be copied to collect an update.
+		*/
+	inline int32 GetNumItems() const { return NumItems; }
 
+	/**
+		*/
 	struct FConstIterator
 	{
-		FConstIterator(const FArrayIndexDelta *InArrayIndexDelta, int32 InIndex) : ArrayIndexDelta(InArrayIndexDelta), Index(InIndex) {}
-
-		void operator++() 
-		{ 
-			++Index;
-		}
+		int32 ItemIndex = 0;
+		int32 MaxNum = 0;
 
-		int32 operator*() const
+		FConstIterator(int32 InIndex, int32 InMaxNum)
+		:	ItemIndex(InIndex),
+			MaxNum(InMaxNum)
 		{
-			return (*ArrayIndexDelta)[Index];
 		}
 
-		int32 GetIndex() const { return Index; }
+		void operator++() {  ++ItemIndex; }
 
-		bool operator != (const FConstIterator& It ) const { return Index != It.Index || ArrayIndexDelta != It.ArrayIndexDelta; }
-	private:
-		const FArrayIndexDelta *ArrayIndexDelta;
-		int32 Index;
-	};
+		/**
+			* Get the index of the data in the source / destination arrays. 
+			*/
+		int32 GetIndex() const { return ItemIndex; }
 
-	FConstIterator begin() const { return FConstIterator(this, 0);}
-	FConstIterator end() const { return FConstIterator(this, Num());}
+		/**
+			* Get the continuous index of the data item in the collected item array.
+			*/
+		int32 GetItemIndex() const { return ItemIndex; }
 
-	/**
-	 * Gather the needed values from InSource to OutDest, according to the delta.
-	 * If there is no delta, it will perform a bulk copy.
-	 */
-	template <typename ValueType>
-	void Gather(TArray<ValueType> &OutDest, const TArrayView<const ValueType> &InSource, int32 ElementStride) const
-	{
-		if (IsDelta())
-		{
-			OutDest.Reset(Num() * ElementStride);
-			for (int32 Index : *this)
-			{
-				OutDest.Append(&InSource[Index * ElementStride], ElementStride);
-			}
-		}
-		else
-		{
-			OutDest = InSource;
-		}
-	}
+		explicit operator bool() const {  return ItemIndex < MaxNum; }
+	};
 
-	/**
-	 * Gather the needed values from InSource to OutDest, according to the delta, calling TransformLambda on each element.
-	 * Never performs bulk copy.
-	 */
-	template <typename OutValueType, typename InValueArray, typename LambdaType>
-	void GatherTransform(TArray<OutValueType> &OutDest, const InValueArray &InSource, LambdaType TransformLambda) const
+	FConstIterator GetIterator() const
 	{
-		for (int32 Index : *this)
-		{
-			OutDest.Add(TransformLambda(InSource[Index]));
-		}
+		return FConstIterator(0, NumItems);
 	}
 
-	/**
-	 * Write InSource that was previously gathered using the above methods to the final destination array OutDest
-	 * using the same delta information. 
-	 * If there is no delta, it performs a move of the source data to the final array, saving a copy.
-	 */
-	template <typename ValueType>
-	void Scatter(TArray<ValueType> &OutDest, int32 DestNumElements, TArray<ValueType> &&InSource, int32 ElementStride = 1) const
-	{
-		check(InSource.Num() == Num() * ElementStride);
-		if (IsDelta())
-		{
-			OutDest.SetNumUninitialized(DestNumElements * ElementStride);
-			for (int32 ItemIndex = 0; ItemIndex < DeltaIndices.Num(); ++ItemIndex)
-			{
-				int32 DestIndex = DeltaIndices[ItemIndex];
-				FMemory::Memcpy(&OutDest[DestIndex * ElementStride], &InSource[ItemIndex * ElementStride], ElementStride * sizeof(ValueType));
-			}
-		}
-		else
-		{
-			check(InSource.Num() == DestNumElements * ElementStride);
-			OutDest = MoveTemp(InSource);
-		}
-	}
-	/**
-	 * Write InSource that was previously gathered using the above methods to the final destination array OutDest
-	 * using the same delta information. 
-	 * Never takes move shortcut as there is an index remap.
-	 */
-	template <typename ValueType, typename DestIndexRemapType>
-	void Scatter(TArray<ValueType> &OutDest, int32 DestNumElements, const TArray<ValueType> &InSource, const DestIndexRemapType &DestIndexRemap, int32 ElementStride = 1) const
-	{
-		check(InSource.Num() == Num() * ElementStride);
-		OutDest.SetNumUninitialized(DestNumElements * ElementStride);
-		if (IsDelta())
-		{
-			for (int32 ItemIndex = 0; ItemIndex < DeltaIndices.Num(); ++ItemIndex)
-			{
-				int32 DestIndex = DestIndexRemap[DeltaIndices[ItemIndex]];
-				if (DestIndex != INDEX_NONE)
-				{
-					FMemory::Memcpy(&OutDest[DestIndex * ElementStride], &InSource[ItemIndex * ElementStride], ElementStride * sizeof(ValueType));
-				}
-			}
-		}
-		else
-		{
-			for (int32 ItemIndex = 0; ItemIndex < NumElements; ++ItemIndex)
-			{
-				int32 DestIndex = DestIndexRemap[ItemIndex];
-				if (DestIndex != INDEX_NONE)
-				{
-					FMemory::Memcpy(&OutDest[DestIndex * ElementStride], &InSource[ItemIndex * ElementStride], ElementStride * sizeof(ValueType));
-				}
-			}
-		}
-	}
 private:
-	int32 NumElements = -1; // Number of elements if there is no delta (AKA identity mapping) otherwise -1 and the number is the size of DeltaIndices
-	TArray<int32> DeltaIndices;
+	int32 NumItems = 0;
 };
 
-
 class FISMInstanceUpdateChangeSet
 {
 public:
-	FISMInstanceUpdateChangeSet(bool bInNeedFullUpdate) : bNeedFullUpdate(bInNeedFullUpdate) {}
+	FISMInstanceUpdateChangeSet(bool bInNeedFullUpdate, FInstanceAttributeTracker &&InInstanceAttributeTracker) 
+		: InstanceAttributeTracker(MoveTemp(InInstanceAttributeTracker))
+		, bNeedFullUpdate(bInNeedFullUpdate) 
+	{
+	}
 #if WITH_EDITOR
 	// Set editor data 
 	void SetEditorData(const TArray<TRefCountPtr<HHitProxy>>& HitProxies, const TBitArray<> &SelectedInstances);//, bool bWasHitProxiesReallocated);
 #endif
 
-	TFunction<void(TArray<float> &InstanceRandomIDs)> GeneratePerInstanceRandomIds;
+	template <FInstanceAttributeTracker::EFlag Flag>
+	FInstanceAttributeTracker::FDeltaRange<Flag> GetDelta(bool bForceEmpty, bool bForceFull = false) const 
+	{ 
+		if (bForceEmpty)
+		{
+			return FInstanceAttributeTracker::FDeltaRange<Flag>();
+		}
+		return InstanceAttributeTracker.GetDeltaRange<Flag>(bNeedFullUpdate || bForceFull, NumSourceInstances); 
+	}
 
-	/**
-	 * Add a value, must be done in the order represented in the InstanceLightShadowUVBiasDelta.
-	 */
-	inline void AddInstanceLightShadowUVBias(const FVector4f &Value)
-	{
-		InstanceLightShadowUVBias.Emplace(Value); 
+
+	FInstanceAttributeTracker::FDeltaRange<FInstanceAttributeTracker::EFlag::TransformChanged> GetTransformDelta() const 
+	{ 
+		return GetDelta<FInstanceAttributeTracker::EFlag::TransformChanged>(false, bUpdateAllInstanceTransforms);
+	}
+
+	FInstanceAttributeTracker::FDeltaRange<FInstanceAttributeTracker::EFlag::CustomDataChanged> GetCustomDataDelta() const 
+	{ 
+		// Force empty range if no custom data
+		return GetDelta<FInstanceAttributeTracker::EFlag::CustomDataChanged>(NumCustomDataFloats == 0 || !Flags.bHasPerInstanceCustomData);
+	}
+
+	FInstanceAttributeTracker::FDeltaRange<FInstanceAttributeTracker::EFlag::IndexChanged> GetIndexChangedDelta() const 
+	{ 
+		return GetDelta<FInstanceAttributeTracker::EFlag::IndexChanged>(false);
 	}
 
-	inline void SetInstanceTransforms(TStridedView<FMatrix> InInstanceTransforms, const FVector Offset)
+	FIdentityDeltaRange GetInstanceLightShadowUVBiasDelta() const
 	{
-		TransformsDelta.GatherTransform(Transforms, InInstanceTransforms, [Offset](const FMatrix &M) -> FRenderTransform { return FRenderTransform(M.ConcatTranslation(Offset)); });
+		return FIdentityDeltaRange(InstanceLightShadowUVBias.Num());
 	}
 
-	inline void SetInstanceTransforms(TStridedView<FMatrix> InInstanceTransforms)
+#if WITH_EDITOR
+	FIdentityDeltaRange GetInstanceEditorDataDelta() const
 	{
-		TransformsDelta.GatherTransform(Transforms, InInstanceTransforms, [](const FMatrix &M) -> FRenderTransform { return FRenderTransform(M); });
+		return FIdentityDeltaRange(InstanceEditorData.Num());
 	}
+#endif
 
-	inline void SetInstanceTransforms(TStridedView<FMatrix> InInstanceTransforms, FBox const& InInstanceBounds, FBox& OutGatheredBounds)
+	TFunction<void(TArray<float> &InstanceRandomIDs)> GeneratePerInstanceRandomIds;
+
+	/**
+	 * Add a value, must be done in the order represented in the InstanceLightShadowUVBiasDelta.
+	 */
+	inline void AddInstanceLightShadowUVBias(const FVector4f &Value)
 	{
-		TransformsDelta.GatherTransform(Transforms, InInstanceTransforms, [&InInstanceBounds, &OutGatheredBounds](const FMatrix& M) -> FRenderTransform
-		{
-			FRenderTransform Transform(M);
-			OutGatheredBounds += InInstanceBounds.TransformBy(Transform.ToMatrix());
-			return Transform;
-		});
+		InstanceLightShadowUVBias.Emplace(Value); 
 	}
 
+	ENGINE_API void SetInstanceTransforms(TStridedView<FMatrix> InInstanceTransforms, const FVector Offset);
+	ENGINE_API void SetInstanceTransforms(TStridedView<FMatrix> InInstanceTransforms);
+
+	/**
+	 * This version produces the bounds of the gathered transforms as a side-effect.
+	 */
+	ENGINE_API void SetInstanceTransforms(TStridedView<FMatrix> InInstanceTransforms, FBox const& InInstanceBounds, FBox& OutGatheredBounds);
+
 	ENGINE_API void SetInstancePrevTransforms(TArrayView<FMatrix> InPrevInstanceTransforms, const FVector &Offset);
 	ENGINE_API void SetInstancePrevTransforms(TArrayView<FMatrix> InPrevInstanceTransforms);
-
 	ENGINE_API void SetCustomData(const TArrayView<const float> &InPerInstanceCustomData, int32 InNumCustomDataFloats);
-	
 	ENGINE_API void SetInstanceLocalBounds(const FRenderBounds &Bounds);
 
 	bool IsFullUpdate() const
@@ -219,26 +147,26 @@ public:
 		return bNeedFullUpdate;
 	}
 
+	FInstanceAttributeTracker InstanceAttributeTracker;
 	bool bNeedFullUpdate = false;
-	FChangeMask ChangeMask;
+	bool bUpdateAllInstanceTransforms = false;
+	
+	bool bNumCustomDataChanged = false;
+	bool bBakedLightingDataChanged = false;
+	bool bAnyEditorDataChanged = false;
 
-	FInstanceIdIndexMap InstanceIdIndexMap;
+	bool bIdentityIdMap = false;
+	TArray<FPrimitiveInstanceId> IndexToIdMapDeltaData;
 	int32 NumCustomDataFloats = 0;
 
-	// TODO: Make it possible to share the deltas such that we don't have to produce an unique index array for each attribute
-	//       e.g., if the only thing that happened was somethis was added and we therefore need all attributes for one instance.
-	FArrayIndexDelta TransformsDelta;
 	TArray<FRenderTransform> Transforms;
 	TArray<FRenderTransform> PrevTransforms;
 
-	FArrayIndexDelta CustomDataDelta;
 	TArray<float> PerInstanceCustomData;
 
-	FArrayIndexDelta InstanceLightShadowUVBiasDelta;
 	TArray<FVector4f> InstanceLightShadowUVBias;
 	FInstanceDataFlags Flags;
 #if WITH_EDITOR
-	FArrayIndexDelta InstanceEditorDataDelta;
 	TArray<uint32> InstanceEditorData;
 	TBitArray<> SelectedInstances;
 	TPimplPtr<FOpaqueHitProxyContainer> HitProxyContainer;
@@ -249,7 +177,9 @@ public:
 	FVector PrimitiveWorldSpaceOffset;
 	TOptional<FRenderTransform> PreviousPrimitiveToRelativeWorld;
 	float AbsMaxDisplacement = 0.0f;
+	int32 NumSourceInstances = 0;
 	int32 PostUpdateNumInstances = 0;
+	int32 MaxInstanceId = 0;
 };
 
 
diff --git a/Engine/Source/Runtime/Engine/Public/InstancedStaticMesh/InstanceAttributeTracker.h b/Engine/Source/Runtime/Engine/Public/InstancedStaticMesh/InstanceAttributeTracker.h
new file mode 100644
--- /dev/null
+++ b/Engine/Source/Runtime/Engine/Public/InstancedStaticMesh/InstanceAttributeTracker.h
@@ -0,0 +1,675 @@
+// Copyright Epic Games, Inc. All Rights Reserved.
+
+#pragma once
+
+#include "CoreTypes.h"
+#include "HAL/UnrealMemory.h"
+#include "Math/UnrealMathUtility.h"
+#include "Containers/BitArray.h"
+#include "Containers/Array.h"
+#include "InstanceDataTypes.h"
+
+// Implementation details
+namespace UE::InstanceAttributeTracker::Implementation
+{
+template <typename ElementType>
+struct FBitManip
+{
+	// Utility to, at compile time, replicate the given bit mask 'Mask' as many times as MaskSize inside sizeof(ElementType)*8.
+	// e.g. turning Mask=0x1 into 0x11111111 if MaskSize=4 & ElementType is uint32
+	template <ElementType Mask, uint32 MaskSize>
+	struct FRepMask
+	{
+		template <uint32 Step>
+		struct FIterator;
+		template <>
+		struct FIterator<0u>
+		{
+			static constexpr ElementType Result = Mask;
+		};
+		template <uint32 Step>
+		struct FIterator
+		{
+			static constexpr ElementType Result = (Mask << (Step * MaskSize)) | FIterator<Step - 1u>::Result;
+			
+		};
+
+		static constexpr uint32 NumSteps = (sizeof(ElementType) * 8u / MaskSize) - 1u;
+		static constexpr ElementType Result = FIterator<NumSteps>::Result;
+	};
+	/**
+	 * Make it possible to select the function based on the overloaded type for 32/64-bit use.
+	 */
+	static FORCEINLINE uint32 CountTrailingZeros(ElementType Value);
+};
+
+
+template <>
+FORCEINLINE uint32 FBitManip<uint32>::CountTrailingZeros(uint32 Value)  
+{ 
+	return FMath::CountTrailingZeros(Value); 
+}
+template <>
+FORCEINLINE uint32 FBitManip<uint64>::CountTrailingZeros(uint64 Value)  
+{ 
+	return uint32(FMath::CountTrailingZeros64(Value)); 
+}
+
+
+} // namespace UE::InstanceAttributeTracker::Implementation
+
+/**
+ * Tracks changes for instances, but that assumes storing stuff per index except the removed status which must be kept ID-based
+ * Removed items are tracked by ID and kept separate as we need to remember if an item was ever removed in order to handle these correctly.
+ * Designed to use move semantics to clear the state and transfer to the updating worker task.
+ */
+class FInstanceAttributeTracker
+{
+public:
+	// Underlying type to store per-instance masks packed in. For example, with EFlag::Num == 4, and uint32 we can store 8 masks in each element.
+	using ElementType = uint32;
+
+	using FBitManip = UE::InstanceAttributeTracker::Implementation::FBitManip<ElementType>;
+
+	/**
+	 * Flag describing the changed attribute. 
+	 * Note that not all are created equal, see comments.
+	 */
+	enum class EFlag : ElementType
+	{
+		Added, /** Set when */
+		TransformChanged,
+		CustomDataChanged, 
+		IndexChanged, /** Implicitly set when calling Remove or RemoveAt as this causes movement to fill holes. */
+		Num,
+	};
+
+	static constexpr uint32 BitsPerElement = sizeof(ElementType) * 8u;
+	static constexpr int32 MasksPerElement = BitsPerElement / uint32(EFlag::Num);
+	static constexpr ElementType AnyFlagMask = (ElementType(1U) << uint32(EFlag::Num)) - ElementType(1U);
+
+	// Helper struct to convert the flag to its bitmask at comile time.
+	template <EFlag Flag>
+	struct FToBit
+	{
+		static constexpr ElementType Bit = ElementType(1U) << uint32(Flag);
+	};
+
+	/**
+	 */
+	FInstanceAttributeTracker();
+
+	/**
+	 */
+	FInstanceAttributeTracker(FInstanceAttributeTracker &&Other);
+
+	/**
+	 */
+	void operator=(FInstanceAttributeTracker &&Other);
+
+	/**
+	 */
+	void Reset();
+	
+	/**
+	 */
+	static void Move(FInstanceAttributeTracker &Dest, FInstanceAttributeTracker &Source);
+
+	/**
+	 * Get an iterator that iterates over the IDs (note, not index as they don't exist any more) of removed instances. The remove flag is not cleared when a new instances is added.
+	 */
+	TConstSetBitIterator<> GetRemovedIterator() const;
+
+	/**
+	 * Lazy allocate bits in the bit vector tracking removed instances and set the flag for the instance ID.
+	 * Also updates the FirstRemovedIdIndex, such that we can start iteration from that bit.
+	 */
+	FORCEINLINE void MarkRemoved(FPrimitiveInstanceId Id)
+	{
+		FirstRemovedIdIndex = FMath::Min(FirstRemovedIdIndex, Id.GetAsIndex());
+		// Lazy allocate the memory so there is nothing to move around in case nothing is removed
+		if (Id.GetAsIndex() >= Removed.Num())
+		{
+			// TBitArray does a lot of work for each added bit, first clearing the slack redundantly and then also clearing the bit even though it was already cleared because the slack was cleared...
+			// so we pad the allocation request with at least a dword at a time.
+			int32 RemovedPadded = FMath::DivideAndRoundUp(Id.GetAsIndex() + 1, 32) * 32;
+			Removed.SetNum(RemovedPadded, false);
+		}
+		Removed[Id.GetAsIndex()] = true;
+	}
+
+
+	/**
+	 * Mark removed & update the tracked bits (remove & swap operation)
+	 */
+	FORCEINLINE void RemoveAtSwap(FPrimitiveInstanceId Id, int32 Index, int32 MaxInstanceIndex)
+	{
+		Validate();
+		
+		SetNum(FMath::Max(Index, MaxInstanceIndex));
+
+		MarkRemoved(Id);
+
+		// book-keeping
+		DecChangedCounts(Index);
+
+		// Remove the tracked state bits
+		int32 LastIndex = MaxIndex - 1;
+		if (Index != LastIndex)
+		{
+			// Move tracked & tag as moved.
+			ElementType PrevFlags = GetFlags(LastIndex);
+			if ((((FToBit<EFlag::IndexChanged>::Bit | FToBit<EFlag::Added>::Bit)) & PrevFlags) == 0u)
+			{
+				check(NumChanged[uint32(EFlag::IndexChanged)] >= 0);
+				NumChanged[uint32(EFlag::IndexChanged)] += 1;
+			}
+			SetMaskTo(Index, FToBit<EFlag::IndexChanged>::Bit | PrevFlags);
+			FirstChangedIndex = FMath::Min(FirstChangedIndex, Index);
+			LastChangedIndex = FMath::Max(LastChangedIndex, Index);
+		}
+		// clear the last one so that it is all zero after the last mask (important for when we do implicit add, as we do at the moment).
+		ClearLastElement();
+		MaxIndex -= 1;
+
+		Validate();
+	}
+
+	/**
+	 * Mark removed & update the tracked bits (remove & move operation)
+	 * Note: this is not something that should be used really as it forces the tracker to touch all the masks after the removed item...
+	 */
+	FORCEINLINE void RemoveAt(FPrimitiveInstanceId Id, int32 Index, int32 MaxInstanceIndex)
+	{
+		Validate();
+
+		SetNum(FMath::Max(Index, MaxInstanceIndex));
+
+		MarkRemoved(Id);
+
+		DecChangedCounts(Index);
+
+		// Remove the tracked state bits
+		int32 LastIndex = MaxIndex - 1;
+		if (Index != LastIndex)
+		{			
+			int32 FirstElementIndex = Index / MasksPerElement;
+			int32 LastElementIndex = LastIndex / MasksPerElement;
+
+			ElementType CurrentElement = Data[FirstElementIndex];
+			ElementType NextElement = LoadCountAndMask(FirstElementIndex + 1, FirstElementIndex != LastElementIndex);
+			// 1. Move up to "Num" bits to cover the removed item, shifting in zero.
+			{
+				uint32 ElementSubIndex = (Index % MasksPerElement);
+				// Count mask changes only in the moved part (not including the removed or any previous).
+				uint32 NumMasksInElement = FMath::Min(MasksPerElement, MaxIndex - FirstElementIndex * MasksPerElement);
+				CountMaskChanges(CurrentElement, ElementSubIndex + 1, NumMasksInElement);
+				ElementType CurrentElementResult = (CurrentElement >> uint32(EFlag::Num)) | IndexChangedElementMask;
+				if (ElementSubIndex != 0u)
+				{
+					uint32 ElementSubShift = ElementSubIndex * uint32(EFlag::Num);
+					ElementType Mask = ~ElementType(0) << ElementSubShift;
+					// we shifted some as we shouldn't, put them back
+					CurrentElementResult = (CurrentElementResult & Mask) | (CurrentElement & (~Mask));
+				}
+				// Move the first item from the next step
+				CurrentElementResult |= NextElement << (MasksPerElement - 1u) * uint32(EFlag::Num);
+				Data[FirstElementIndex] = CurrentElementResult;
+			}
+			// 2. Do the middle, whole, elements...
+			for (int32 ElementIndex = FirstElementIndex + 1; ElementIndex < LastElementIndex; ++ElementIndex)
+			{
+				CurrentElement = NextElement >> uint32(EFlag::Num);
+				NextElement = LoadCountAndMask(ElementIndex + 1, true);
+				ElementType CurrentElementResult = CurrentElement | NextElement << (MasksPerElement - 1u) * uint32(EFlag::Num);
+				Data[ElementIndex] = CurrentElementResult;
+			}
+			// 2. Do the last element...
+			if (LastElementIndex != FirstElementIndex)
+			{
+				Data[LastElementIndex] = NextElement >> uint32(EFlag::Num);
+			}
+
+			FirstChangedIndex = FMath::Min(FirstChangedIndex, Index);
+			LastChangedIndex = FMath::Max(LastChangedIndex, LastIndex);
+		}	
+		ClearLastElement();
+		MaxIndex -= 1;
+
+		Validate();
+	}
+
+	/**
+	 */
+	template <EFlag Flag>
+	FORCEINLINE void MarkIndex(int32 Index, int32 MaxInstanceIndex)
+	{
+		Validate();
+		
+		// Lazy allocate
+		SetNum(FMath::Max(Index, MaxInstanceIndex));
+
+		check(Index < MaxIndex);
+
+		FirstChangedIndex = FMath::Min(FirstChangedIndex, Index);
+		LastChangedIndex = FMath::Max(LastChangedIndex, Index);
+
+		uint32 ElementSubShift = (Index % MasksPerElement) * uint32(EFlag::Num);
+
+		ElementType SetMask = FToBit<Flag>::Bit << ElementSubShift;
+
+		int32 ElementIndex = Index / MasksPerElement;
+
+		ElementType PrevElement = Data[ElementIndex];
+
+		if ((((FToBit<Flag>::Bit | FToBit<EFlag::Added>::Bit) << ElementSubShift) & PrevElement) == 0u)
+		{
+			check(NumChanged[uint32(Flag)] >= 0);
+			NumChanged[uint32(Flag)] += 1;
+		}
+
+		Data[ElementIndex] = PrevElement | SetMask;
+
+		Validate();
+	}
+
+	/**
+	 * Get the masked flags for a given index.
+	 */
+	FORCEINLINE ElementType GetFlags(int32 Index) const
+	{
+		check(Index < MaxIndex);
+		ElementType Element = Data[Index / MasksPerElement];
+		uint32 ElementSubShift = (Index % MasksPerElement) * uint32(EFlag::Num);
+		return (Element >> ElementSubShift) & AnyFlagMask;
+	}
+
+	/**
+	 * Test a specific flag at a given index.
+	 */
+	template <EFlag Flag>
+	FORCEINLINE bool TestFlag(int32 Index) const
+	{
+		return (GetFlags(Index) & FToBit<Flag>::Bit) != 0u;
+	}
+
+	/**
+	 */
+	FORCEINLINE void SetNum(int32 InstanceIndexMax)
+	{
+		Validate();
+
+		int32 ElementIndex = FMath::DivideAndRoundUp(InstanceIndexMax, int32(MasksPerElement));
+
+		if (ElementIndex > Data.Num())
+		{
+			Data.SetNumZeroed(ElementIndex);
+		}
+		MaxIndex = InstanceIndexMax;
+		Validate();
+	}
+
+	/**
+	 * Iterator for iterating tracked state with the given flags. Uses bit logic to efficiently skip empty (under the given mask) sections.
+	 */
+	template <uint32 Mask = AnyFlagMask>
+	class FAnyValidIterator
+	{
+	public:
+		FORCEINLINE FAnyValidIterator() = default;
+		FORCEINLINE FAnyValidIterator(const FInstanceAttributeTracker *InTracker, int32 StartIndex) : Tracker(InTracker)
+		{
+			ElementIndex = StartIndex / MasksPerElement;
+			ElementOffset = StartIndex % MasksPerElement;
+			CurrentElement = Tracker->Data.IsValidIndex(ElementIndex) ? (Tracker->Data[ElementIndex] & ElementMask) : 0;
+			CurrentElement >>= uint32(EFlag::Num) * ElementOffset;
+			AdvanceUntilAny();
+		}
+
+		FORCEINLINE void AdvanceUntilAny()
+		{
+			while(ElementIndex < Tracker->Data.Num())
+			{
+				// Skip the whole word if nothing is set.
+				if (CurrentElement)
+				{
+					uint32 ZeroCount = FBitManip::CountTrailingZeros(CurrentElement);
+					uint32 Steps = ZeroCount / uint32(EFlag::Num);
+					CurrentElement >>= uint32(EFlag::Num) * Steps;
+					ElementOffset += Steps;
+					return;
+				}
+				++ElementIndex;
+				ElementOffset = 0;
+				if (ElementIndex >= Tracker->Data.Num())
+				{
+					CurrentElement = 0u;
+					return;
+				}
+				CurrentElement = Tracker->Data[ElementIndex] & ElementMask;
+			}
+		}
+
+		FORCEINLINE void AdvanceToNext()
+		{
+			// Clear the current value (which will cause AdvanceUntilAny to move forwards)
+			CurrentElement &= ~ElementType(Mask);
+			AdvanceUntilAny();
+		}
+
+		FORCEINLINE FAnyValidIterator &operator++()
+		{
+			AdvanceToNext();
+			return *this;
+		}
+
+		FORCEINLINE explicit operator bool() const
+		{
+			return GetIndex() < FMath::Min(Tracker->MaxIndex, Tracker->LastChangedIndex + 1);
+		}
+
+		FORCEINLINE int32 GetIndex() const
+		{
+			return ElementIndex * MasksPerElement + ElementOffset;
+		}
+
+		/* 
+		 * Returns the mask for the current item
+		 * NOTE: does not clear any bits outside the search mask.
+		 */
+		FORCEINLINE uint32 GetMask() const
+		{
+			return CurrentElement;
+		}
+
+		template <EFlag Flag>
+		FORCEINLINE bool TestFlag() const
+		{
+			return (CurrentElement & FToBit<Flag>::Bit) != 0u;
+		}
+	private:
+		const FInstanceAttributeTracker *Tracker = nullptr;
+		int32 ElementIndex = 0;
+		ElementType CurrentElement;
+		int32 ElementOffset = 0;
+		static constexpr ElementType ElementMask = FBitManip::FRepMask<Mask, uint32(EFlag::Num)>::Result;
+	};
+
+	template <uint32 Mask = AnyFlagMask>
+	FORCEINLINE FAnyValidIterator<Mask> GetChangedIterator() const
+	{
+		return FAnyValidIterator<Mask>(this, FMath::Min(MaxIndex, FirstChangedIndex));
+	}
+
+#if DO_GUARD_SLOW
+	ENGINE_API void Validate() const;
+#else
+	FORCEINLINE void Validate() const { }
+#endif
+
+	/**
+	 * Represents the delta range for a given attribute, can either refer to the tracked data or just a span, or be empty.
+	 * "Added" is always implied & added on top. It is a helper to make it easier to iterate a change set.
+	 */
+	template <EFlag Flag>
+	class FDeltaRange
+	{
+	public:
+		static constexpr uint32 IteratorMask = FToBit<Flag>::Bit | FToBit<EFlag::Added>::Bit;
+
+		FORCEINLINE FDeltaRange() = default;
+
+		FORCEINLINE FDeltaRange(bool bInFullUpdate, int32 InNumItems, const FInstanceAttributeTracker *InInstanceUpdateTracker)
+			: bFullUpdate(bInFullUpdate)
+			, NumItems(InNumItems)
+			, Tracker(InInstanceUpdateTracker)
+		{
+			int32 NumAdded = Tracker->NumChanged[uint32(EFlag::Added)];
+			bFullUpdate = bFullUpdate || (NumAdded + Tracker->NumChanged[uint32(Flag)] == NumItems);
+			if (!bFullUpdate)
+			{
+				NumItems = NumAdded + Tracker->NumChanged[uint32(Flag)];
+			}
+		}
+
+		/**
+		 */
+		FORCEINLINE bool IsEmpty() const { return NumItems == 0; }
+
+		/**
+		 */
+		FORCEINLINE bool IsDelta() const { return !bFullUpdate; }
+
+		/**
+		 * Returns the number of items in this range - i.e., the number of items that need to be copied to collect an update.
+		 */
+		FORCEINLINE int32 GetNumItems() const { return NumItems; }
+
+		/**
+		 * Iterator for traversing the range of Items that need update. Provides two indexes:
+		 * 1. The sparse Index of the source & destination data
+		 * 2. The continuous Item Index that represents the linear index into the copied data.
+		 */
+		struct FConstIterator
+		{
+			bool bUseIterator = true;
+			int32 ItemIndex = 0;
+			int32 MaxNum = 0;
+			FAnyValidIterator<IteratorMask> It;
+
+			FORCEINLINE FConstIterator(int32 InIndex, int32 InMaxNum)
+			:	bUseIterator(false),
+				ItemIndex(InIndex),
+				MaxNum(InMaxNum)
+			{
+			}
+
+			FORCEINLINE FConstIterator(FAnyValidIterator<IteratorMask> &&InIt)
+			:	bUseIterator(true),
+				ItemIndex(0),
+				MaxNum(),
+				It(MoveTemp(InIt))
+			{
+			}
+
+			FORCEINLINE void operator++() 
+			{ 
+				++ItemIndex;
+				if (bUseIterator)
+				{
+					++It;
+				}
+			}
+
+			/**
+			 * Get the index of the data in the source / destination arrays. 
+			 */
+			FORCEINLINE int32 GetIndex() const 
+			{ 
+				if (bUseIterator)
+				{
+					return It.GetIndex();
+				}
+				return ItemIndex;
+			}
+
+			/**
+			 * Get the continuous index of the data item in the collected item array.
+			 */
+			FORCEINLINE int32 GetItemIndex() const
+			{
+				return ItemIndex;
+			}
+
+			FORCEINLINE explicit operator bool() const 
+			{ 
+				if (bUseIterator)
+				{
+					return bool(It);
+				}
+				return ItemIndex < MaxNum;
+			}
+		};
+
+		FORCEINLINE FConstIterator GetIterator() const
+		{
+			if (bFullUpdate)
+			{
+				return FConstIterator(0, NumItems);
+			}
+			return FConstIterator(Tracker->GetChangedIterator<IteratorMask>());
+		}
+
+	private:
+		bool bFullUpdate = true;
+		int32 NumItems = 0;
+		const FInstanceAttributeTracker *Tracker = nullptr;
+	};
+	
+	template <EFlag Flag>
+	FORCEINLINE FDeltaRange<Flag> GetDeltaRange(bool bForceFullUpdate, int32 InNumItems) const
+	{
+		return FDeltaRange<Flag>(bForceFullUpdate, InNumItems, this);
+
+	}
+
+	FORCEINLINE bool HasAnyChanges() const
+	{
+		return MaxIndex > 0 || !Removed.IsEmpty();
+	}
+
+	FORCEINLINE SIZE_T GetAllocatedSize() const
+	{
+		return Data.GetAllocatedSize() + Removed.GetAllocatedSize();
+	}
+
+private:
+
+	/**
+	 */
+	FORCEINLINE void SetMaskTo(int32 Index, ElementType Mask)
+	{
+		check(Index < MaxIndex);
+		check(Mask <= AnyFlagMask);
+
+		const int32 ElementIndex = Index / MasksPerElement;
+		const uint32 ElementSubShift = (Index % MasksPerElement) * uint32(EFlag::Num);
+		// Load and clear the mask for this index
+		ElementType CurrentElement = Data[ElementIndex] & ~(AnyFlagMask << ElementSubShift);
+		// add the new mask and store
+		Data[ElementIndex] =  CurrentElement | (Mask << ElementSubShift);
+	}
+
+	/**
+	 * Clear the last item & all remaining in element to ensure it is always all zero
+	 */
+	FORCEINLINE void ClearLastElement()
+	{
+		int32 LastIndex = MaxIndex - 1;
+		const int32 ElementIndex = LastIndex / MasksPerElement;
+		const int32 SubIndex = LastIndex % MasksPerElement;
+		if (SubIndex > 0)
+		{
+			const uint32 ElementSubShift = (MasksPerElement - (LastIndex % MasksPerElement)) * uint32(EFlag::Num);
+			// Load and clear the mask for this index
+			ElementType CurrentElement = Data[ElementIndex] & ((~ElementType(0)) >> ElementSubShift);
+			Data[ElementIndex] =  CurrentElement;
+		}
+		else
+		{
+			Data[ElementIndex] =  ElementType(0);
+		}
+	}
+
+	/**
+	 */
+	template <EFlag Flag>
+	FORCEINLINE void DecChangedCount(uint32 Flags)
+	{
+		if (Flags & FToBit<Flag>::Bit)
+		{
+			NumChanged[uint32(Flag)] -= 1;
+			check(NumChanged[uint32(Flag)] >= 0);
+		}
+	}
+
+	/**
+	 */
+	FORCEINLINE void DecChangedCounts(int32 Index)
+	{
+		const uint32 PrevFlags = GetFlags(Index);
+		if ((PrevFlags & FToBit<EFlag::Added>::Bit) == 0)
+		{
+			DecChangedCount<EFlag::CustomDataChanged>(PrevFlags);
+			DecChangedCount<EFlag::TransformChanged>(PrevFlags);
+			DecChangedCount<EFlag::IndexChanged>(PrevFlags);
+		}
+		else
+		{
+			NumChanged[uint32(EFlag::Added)] -= 1;
+			check(NumChanged[uint32(EFlag::Added)] >= 0);
+		}
+	}
+
+	/**
+	 */
+	FORCEINLINE void CountMaskChanges(ElementType Element, uint32 Offset, uint32 NumMasksInElement)
+	{
+		// If the offset is equal to the number of masks then there can't be any changes
+		if (Offset < MasksPerElement)
+		{
+			// 1. Element & IndexChangedElementMask - 1 if it already had index changed
+			ElementType HasIndexChanged = (Element & IndexChangedElementMask);
+			// 2. Element & AddedElementMask - 1 if already had added << to align
+			ElementType HasAdded = (Element & AddedElementMask) << (uint32(EFlag::IndexChanged) - uint32(EFlag::Added));
+			ElementType CombinedMask = HasIndexChanged | HasAdded;
+
+			int32 Delta = NumMasksInElement - Offset - FMath::CountBits(CombinedMask >> (Offset * uint32(EFlag::Num)));
+			check(Delta >= 0);
+			NumChanged[uint32(EFlag::IndexChanged)] += Delta;
+		}
+	};
+	/**
+	 * Helper to load and count mask changes, returns the element with the IndexChanged applied to all items. Used while moving elements in the compacting RemoveAt 
+	 */
+	FORCEINLINE ElementType LoadCountAndMask(int32 ElementIndex, bool bHasNextElement)
+	{
+		if (bHasNextElement)
+		{
+			uint32 NumMasksInElement = FMath::Min(MasksPerElement, MaxIndex - ElementIndex * MasksPerElement);
+			ElementType Element = Data[ElementIndex];
+			CountMaskChanges(Element, 0, NumMasksInElement);
+			return Element | IndexChangedElementMask;
+		}
+		return ElementType(0);
+	};
+
+
+private:
+	/** Data for packed per-instance change attribute mask. */
+	TArray<ElementType> Data;
+	
+	/** Removed items are tracked by ID and kept separate as we need to remember if an item was ever removed in order to handle these correctly. */
+	TBitArray<> Removed;
+
+	/** Maximum valid index. */
+	int32 MaxIndex = 0;
+
+	/** Tracks the index of the first non-zero bit in the Removed bitvector. */
+	int32 FirstRemovedIdIndex = MAX_int32;
+
+	/** Tracks the first & last mask with any change flag set */
+	int32 FirstChangedIndex = MAX_int32;
+	int32 LastChangedIndex = 0;
+
+	/** 
+	 * Tracks the number of changes for each attribute in such a way that 
+	 * NumChanged[Added] + NumChanged[OtherAttribute] is the total number we need to copy over. This works by counting other attributes ONLY if the
+	 * Added flag is not also set.
+	 */
+	int32 NumChanged[uint32(EFlag::Num)];
+
+	/** Masks replicated to all masks that fit inside ElementType */
+	static constexpr ElementType IndexChangedElementMask = FBitManip::FRepMask<FToBit<EFlag::IndexChanged>::Bit, uint32(EFlag::Num)>::Result;
+	static constexpr ElementType AddedElementMask = FBitManip::FRepMask<FToBit<EFlag::Added>::Bit, uint32(EFlag::Num)>::Result;
+};
diff --git a/Engine/Source/Runtime/Foliage/Private/InstancedGrass.cpp b/Engine/Source/Runtime/Foliage/Private/InstancedGrass.cpp
--- a/Engine/Source/Runtime/Foliage/Private/InstancedGrass.cpp
+++ b/Engine/Source/Runtime/Foliage/Private/InstancedGrass.cpp
@@ -113,13 +113,9 @@ void UGrassInstancedStaticMeshComponent::BuildComponentInstanceData(FInstanceUpd
 	OutData.BuildChangeSet = [&](FISMInstanceUpdateChangeSet &ChangeSet)
 	{
 		// Cancel update as there is no source data
-		check(ChangeSet.TransformsDelta.IsEmpty());
-		ChangeSet.TransformsDelta = FArrayIndexDelta();
-		ChangeSet.InstanceLightShadowUVBiasDelta = FArrayIndexDelta();
-		ChangeSet.CustomDataDelta = FArrayIndexDelta();
-#if WITH_EDITOR
-		ChangeSet.InstanceEditorDataDelta = FArrayIndexDelta();
-#endif
+		check(ChangeSet.GetTransformDelta().IsEmpty());
+		check(ChangeSet.GetCustomDataDelta().IsEmpty());
+
 		BuildInstanceDataDeltaChangeSetCommon(ChangeSet);
 		check(GetTranslatedInstanceSpaceOrigin().IsNearlyZero());
 		check(PerInstanceSMData.IsEmpty());
