Fixed for ISM instance bodies causing a crash when removed by being null when zero scaled.

diff --git a/Engine/Source/Runtime/Engine/Private/InstancedStaticMesh.cpp b/Engine/Source/Runtime/Engine/Private/InstancedStaticMesh.cpp
--- a/Engine/Source/Runtime/Engine/Private/InstancedStaticMesh.cpp
+++ b/Engine/Source/Runtime/Engine/Private/InstancedStaticMesh.cpp
@@ -3461,16 +3461,22 @@ bool UInstancedStaticMeshComponent::RemoveInstanceInternal(int32 InstanceIndex,
 				InstanceBodies.RemoveAtSwap(InstanceIndex);
 
 				// Update the Instance body index to the new swapped location
-				InstanceBodies[InstanceIndex]->InstanceBodyIndex = InstanceIndex;
+				if (InstanceBodies[InstanceIndex])
+				{
+					InstanceBodies[InstanceIndex]->InstanceBodyIndex = InstanceIndex;
+				}
 			}
 			else
 			{
 				InstanceBodies.RemoveAt(InstanceIndex);
 
 				// Re-target instance indices for shifting of array.
-				for (int32 i = InstanceIndex; i < InstanceBodies.Num(); ++i)
+				for (int32 Index = InstanceIndex; Index < InstanceBodies.Num(); ++Index)
 				{
-					InstanceBodies[i]->InstanceBodyIndex = i;
+					if (InstanceBodies[Index])
+					{
+						InstanceBodies[Index]->InstanceBodyIndex = Index;
+					}
 				}
 			}
 		}
