Fixed UInstancedStaticMeshComponent::GetInstancesOverlapping(Box|Sphere) variously not taking parts of the instance transform into account.

diff --git a/Engine/Source/Runtime/Engine/Private/InstancedStaticMesh.cpp b/Engine/Source/Runtime/Engine/Private/InstancedStaticMesh.cpp
--- a/Engine/Source/Runtime/Engine/Private/InstancedStaticMesh.cpp
+++ b/Engine/Source/Runtime/Engine/Private/InstancedStaticMesh.cpp
@@ -4088,14 +4088,11 @@ TArray<int32> UInstancedStaticMeshComponent::GetInstancesOverlappingSphere(const
 			Sphere = Sphere.TransformBy(GetComponentTransform().Inverse());
 		}
 
-		const float StaticMeshBoundsRadius = Mesh->GetBounds().SphereRadius;
+		FSphere StaticMeshSphere = Mesh->GetBounds().GetSphere();
 
 		for (int32 Index = 0; Index < PerInstanceSMData.Num(); Index++)
 		{
-			const FMatrix& Matrix = PerInstanceSMData[Index].Transform;
-			const FSphere InstanceSphere(Matrix.GetOrigin(), StaticMeshBoundsRadius * Matrix.GetScaleVector().GetMax());
-
-			if (Sphere.Intersects(InstanceSphere))
+			if (Sphere.Intersects(StaticMeshSphere.TransformBy(PerInstanceSMData[Index].Transform)))
 			{
 				Result.Add(Index);
 			}
@@ -4117,20 +4114,15 @@ TArray<int32> UInstancedStaticMeshComponent::GetInstancesOverlappingBox(const FB
 			Box = Box.TransformBy(GetComponentTransform().Inverse());
 		}
 
-		const FVector StaticMeshBoundsExtent = Mesh->GetBounds().BoxExtent;
-
+		FBox StaticMeshBox = Mesh->GetBounds().GetBox();
 		for (int32 Index = 0; Index < PerInstanceSMData.Num(); Index++)
 		{
-			const FMatrix& Matrix = PerInstanceSMData[Index].Transform;
-			FBox InstanceBox(Matrix.GetOrigin() - StaticMeshBoundsExtent, Matrix.GetOrigin() + StaticMeshBoundsExtent);
-
-			if (Box.Intersect(InstanceBox))
+			if (Box.Intersect(StaticMeshBox.TransformBy(PerInstanceSMData[Index].Transform)))
 			{
 				Result.Add(Index);
 			}
 		}
 	}
-
 	return Result;
 }
 
