Fix log warning typo.

diff --git a/Engine/Source/Runtime/Engine/Private/WorldPartition/HLOD/HLODLayer.cpp b/Engine/Source/Runtime/Engine/Private/WorldPartition/HLOD/HLODLayer.cpp
--- a/Engine/Source/Runtime/Engine/Private/WorldPartition/HLOD/HLODLayer.cpp
+++ b/Engine/Source/Runtime/Engine/Private/WorldPartition/HLOD/HLODLayer.cpp
@@ -158,7 +158,7 @@ void UHLODLayer::PostEditChangeProperty(FPropertyChangedEvent& PropertyChangedEv
 			VisitedHLODLayers.Add(CurHLODLayer, &bHLODLayerWasAlreadyInSet);
 			if (bHLODLayerWasAlreadyInSet)
 			{
-				UE_LOG(LogHLODLayer, Warning, TEXT("Circular HLOD parent chain detedted: HLODLayer=%s ParentLayer=%s"), *GetName(), *ParentLayer->GetName());
+				UE_LOG(LogHLODLayer, Warning, TEXT("Circular HLOD parent chain detected: HLODLayer=%s ParentLayer=%s"), *GetName(), *ParentLayer->GetName());
 				ParentLayer = nullptr;
 				break;
 			}
