Removed TAnd on variadic function argument checks.

diff --git a/Engine/Source/Programs/UnrealLightmass/Private/ImportExport/LightmassSwarm.h b/Engine/Source/Programs/UnrealLightmass/Private/ImportExport/LightmassSwarm.h
--- a/Engine/Source/Programs/UnrealLightmass/Private/ImportExport/LightmassSwarm.h
+++ b/Engine/Source/Programs/UnrealLightmass/Private/ImportExport/LightmassSwarm.h
@@ -12,7 +12,6 @@
 #include "Editor/SwarmInterface/Public/SwarmDefines.h"
 #include "Editor/SwarmInterface/Public/SwarmInterface.h"
 #include "Templates/IsValidVariadicFunctionArg.h"
-#include "Templates/AndOrNot.h"
 #include "Templates/IsArrayOrRefOfTypeByPredicate.h"
 #include "Traits/IsCharEncodingCompatibleWith.h"
 #include "ImportExport.h"
@@ -236,7 +235,7 @@ public:
 	void SendTextMessage(const FmtType& Fmt, Types... Args)
 	{
 		static_assert(TIsArrayOrRefOfTypeByPredicate<FmtType, TIsCharEncodingCompatibleWithTCHAR>::Value, "Formatting string must be a TCHAR array.");
-		static_assert(TAnd<TIsValidVariadicFunctionArg<Types>...>::Value, "Invalid argument(s) passed to FString::Printf");
+		static_assert((TIsValidVariadicFunctionArg<Types>::Value && ...), "Invalid argument(s) passed to FString::Printf");
 
 		SendTextMessageImpl((const TCHAR*)Fmt, Args...);
 	}
diff --git a/Engine/Source/Runtime/Core/Public/Containers/UnrealString.h.inl b/Engine/Source/Runtime/Core/Public/Containers/UnrealString.h.inl
--- a/Engine/Source/Runtime/Core/Public/Containers/UnrealString.h.inl
+++ b/Engine/Source/Runtime/Core/Public/Containers/UnrealString.h.inl
@@ -1372,7 +1372,7 @@ public:
 	UE_NODISCARD static UE_STRING_CLASS Printf(const FmtType& Fmt, Types... Args)
 	{
 		static_assert(TIsArrayOrRefOfTypeByPredicate<FmtType, TIsCharEncodingCompatibleWithElementType>::Value, "Formatting string must be a character array.");
-		static_assert(TAnd<TIsValidVariadicFunctionArg<Types>...>::Value, "Invalid argument(s) passed to Printf");
+		static_assert((TIsValidVariadicFunctionArg<Types>::Value && ...), "Invalid argument(s) passed to Printf");
 
 		return PrintfImpl((const ElementType*)Fmt, Args...);
 	}
@@ -1385,7 +1385,7 @@ public:
 	UE_STRING_CLASS& Appendf(const FmtType& Fmt, Types... Args)
 	{
 		static_assert(TIsArrayOrRefOfTypeByPredicate<FmtType, TIsCharEncodingCompatibleWithElementType>::Value, "Formatting string must be a character array.");
-		static_assert(TAnd<TIsValidVariadicFunctionArg<Types>...>::Value, "Invalid argument(s) passed to TString::Appendf");
+		static_assert((TIsValidVariadicFunctionArg<Types>::Value && ...), "Invalid argument(s) passed to TString::Appendf");
 
 		AppendfImpl(*this, (const ElementType*)Fmt, Args...);
 		return *this;
diff --git a/Engine/Source/Runtime/Core/Public/HAL/IPlatformFileOpenLogWrapper.h b/Engine/Source/Runtime/Core/Public/HAL/IPlatformFileOpenLogWrapper.h
--- a/Engine/Source/Runtime/Core/Public/HAL/IPlatformFileOpenLogWrapper.h
+++ b/Engine/Source/Runtime/Core/Public/HAL/IPlatformFileOpenLogWrapper.h
@@ -272,7 +272,7 @@ public:
 	static void AddLabel(const FmtType& Fmt, Types... Args)
 	{
 		static_assert(TIsArrayOrRefOfTypeByPredicate<FmtType, TIsCharEncodingCompatibleWithTCHAR>::Value, "Formatting string must be a TCHAR array.");
-		static_assert(TAnd<TIsValidVariadicFunctionArg<Types>...>::Value, "Invalid argument(s) passed to FCsvProfiler::AddLabel");
+		static_assert((TIsValidVariadicFunctionArg<Types>::Value && ...), "Invalid argument(s) passed to FCsvProfiler::AddLabel");
 		AddLabelInternal((const TCHAR*)Fmt, Args...);
 	}
 
diff --git a/Engine/Source/Runtime/Core/Public/Internationalization/Text.h b/Engine/Source/Runtime/Core/Public/Internationalization/Text.h
--- a/Engine/Source/Runtime/Core/Public/Internationalization/Text.h
+++ b/Engine/Source/Runtime/Core/Public/Internationalization/Text.h
@@ -21,7 +21,6 @@
 #include "Misc/Optional.h"
 #include "Templates/UniquePtr.h"
 #include "Templates/IsConstructible.h"
-#include "Templates/AndOrNot.h"
 
 class FText;
 class FTextHistory;
@@ -611,7 +610,7 @@ public:
 	template <typename... ArgTypes>
 	static FORCEINLINE FText Format(FTextFormat Fmt, ArgTypes... Args)
 	{
-		static_assert(TAnd<TIsConstructible<FFormatArgumentValue, ArgTypes>...>::Value, "Invalid argument type passed to FText::Format");
+		static_assert((TIsConstructible<FFormatArgumentValue, ArgTypes>::Value && ...), "Invalid argument type passed to FText::Format");
 		static_assert(sizeof...(Args) > 0, "FText::Format expects at least one non-format argument"); // we do this to ensure that people don't call Format for no good reason
 
 		// We do this to force-select the correct overload, because overload resolution will cause compile
@@ -673,7 +672,7 @@ public:
 	template <typename... ArgTypes>
 	static FORCEINLINE FText Join(const FText& Delimiter, ArgTypes... Args)
 	{
-		static_assert(TAnd<TIsConstructible<FFormatArgumentValue, ArgTypes>...>::Value, "Invalid argument type passed to FText::Join");
+		static_assert((TIsConstructible<FFormatArgumentValue, ArgTypes>::Value && ...), "Invalid argument type passed to FText::Join");
 		static_assert(sizeof...(Args) > 0, "FText::Join expects at least one non-format argument"); // we do this to ensure that people don't call Join for no good reason
 
 		return Join(Delimiter, FFormatOrderedArguments{ MoveTemp(Args)... });
@@ -1272,7 +1271,7 @@ public:
 	template <typename... ArgTypes>
 	FORCEINLINE void AppendLineFormat(FTextFormat Pattern, ArgTypes... Args)
 	{
-		static_assert(TAnd<TIsConstructible<FFormatArgumentValue, ArgTypes>...>::Value, "Invalid argument type passed to FTextBuilder::AppendLineFormat");
+		static_assert((TIsConstructible<FFormatArgumentValue, ArgTypes>::Value && ...), "Invalid argument type passed to FTextBuilder::AppendLineFormat");
 		static_assert(sizeof...(Args) > 0, "FTextBuilder::AppendLineFormat expects at least one non-format argument"); // we do this to ensure that people don't call AppendLineFormat for no good reason
 
 		BuildAndAppendLine(FText::Format(MoveTemp(Pattern), FFormatOrderedArguments{ MoveTemp(Args)... }));
diff --git a/Engine/Source/Runtime/Core/Public/Logging/LogMacros.h b/Engine/Source/Runtime/Core/Public/Logging/LogMacros.h
--- a/Engine/Source/Runtime/Core/Public/Logging/LogMacros.h
+++ b/Engine/Source/Runtime/Core/Public/Logging/LogMacros.h
@@ -13,7 +13,6 @@
 #include "Misc/Build.h"
 #include "Misc/VarArgs.h"
 #include "String/FormatStringSan.h"
-#include "Templates/AndOrNot.h"
 #include "Templates/EnableIf.h"
 #include "Templates/IsArrayOrRefOfTypeByPredicate.h"
 #include "Templates/IsValidVariadicFunctionArg.h"
@@ -42,7 +41,7 @@ struct FMsg
 	static void SendNotificationStringf(const FmtType& Fmt, Types... Args)
 	{
 		static_assert(TIsArrayOrRefOfTypeByPredicate<FmtType, TIsCharEncodingCompatibleWithTCHAR>::Value, "Formatting string must be a const TCHAR array.");
-		static_assert(TAnd<TIsValidVariadicFunctionArg<Types>...>::Value, "Invalid argument(s) passed to FMsg::SendNotificationStringf");
+		static_assert((TIsValidVariadicFunctionArg<Types>::Value && ...), "Invalid argument(s) passed to FMsg::SendNotificationStringf");
 
 		SendNotificationStringfImpl((const TCHAR*)Fmt, Args...);
 	}
@@ -52,7 +51,7 @@ struct FMsg
 	static void Logf(const ANSICHAR* File, int32 Line, const FLogCategoryName& Category, ELogVerbosity::Type Verbosity, const FmtType& Fmt, Types... Args)
 	{
 		static_assert(TIsArrayOrRefOfTypeByPredicate<FmtType, TIsCharEncodingCompatibleWithTCHAR>::Value, "Formatting string must be a const TCHAR array.");
-		static_assert(TAnd<TIsValidVariadicFunctionArg<Types>...>::Value, "Invalid argument(s) passed to FMsg::Logf");
+		static_assert((TIsValidVariadicFunctionArg<Types>::Value && ...), "Invalid argument(s) passed to FMsg::Logf");
 
 		LogfImpl(File, Line, Category, Verbosity, (const TCHAR*)Fmt, Args...);
 	}
@@ -62,7 +61,7 @@ struct FMsg
 	static void Logf_Internal(const ANSICHAR* File, int32 Line, const FLogCategoryName& Category, ELogVerbosity::Type Verbosity, const FmtType& Fmt, Types... Args)
 	{
 		static_assert(TIsArrayOrRefOfTypeByPredicate<FmtType, TIsCharEncodingCompatibleWithTCHAR>::Value, "Formatting string must be a const TCHAR array.");
-		static_assert(TAnd<TIsValidVariadicFunctionArg<Types>...>::Value, "Invalid argument(s) passed to FMsg::Logf_Internal");
+		static_assert((TIsValidVariadicFunctionArg<Types>::Value && ...), "Invalid argument(s) passed to FMsg::Logf_Internal");
 
 		Logf_InternalImpl(File, Line, Category, Verbosity, (const TCHAR*)Fmt, Args...);
 	}
diff --git a/Engine/Source/Runtime/Core/Public/Misc/AssertionMacros.h b/Engine/Source/Runtime/Core/Public/Misc/AssertionMacros.h
--- a/Engine/Source/Runtime/Core/Public/Misc/AssertionMacros.h
+++ b/Engine/Source/Runtime/Core/Public/Misc/AssertionMacros.h
@@ -6,7 +6,6 @@
 #include "HAL/Platform.h"
 #include "HAL/PlatformMisc.h"
 #include "HAL/PreprocessorHelpers.h"
-#include "Templates/AndOrNot.h"
 #include "Templates/EnableIf.h"
 #include "Templates/IsArrayOrRefOfTypeByPredicate.h"
 #include "Templates/IsValidVariadicFunctionArg.h"
@@ -137,7 +136,7 @@ public:
 	static FORCEINLINE bool OptionallyLogFormattedEnsureMessageReturningFalse(bool bLog, const ANSICHAR* Expr, const ANSICHAR* File, int32 Line, void* ProgramCounter, const FmtType& FormattedMsg, Types... Args)
 	{
 		static_assert(TIsArrayOrRefOfTypeByPredicate<FmtType, TIsCharEncodingCompatibleWithTCHAR>::Value, "Formatting string must be a TCHAR array.");
-		static_assert(TAnd<TIsValidVariadicFunctionArg<Types>...>::Value, "Invalid argument(s) passed to ensureMsgf");
+		static_assert((TIsValidVariadicFunctionArg<Types>::Value && ...), "Invalid argument(s) passed to ensureMsgf");
 
 		return OptionallyLogFormattedEnsureMessageReturningFalseImpl(bLog, Expr, File, Line, ProgramCounter, (const TCHAR*)FormattedMsg, Args...);
 	}
@@ -183,7 +182,7 @@ public:
 //		Types... Args)
 //	{
 //		static_assert(TIsArrayOrRefOfTypeByPredicate<FmtType, TIsCharEncodingCompatibleWithTCHAR>::Value, "Formatting string must be a TCHAR array.");
-//		static_assert(TAnd<TIsValidVariadicFunctionArg<Types>...>::Value, "Invalid argument(s) passed to CheckVerifyFailed()");
+//		static_assert(TIsValidVariadicFunctionArg<Types>::Value && ...), "Invalid argument(s) passed to CheckVerifyFailed()");
 //		return CheckVerifyFailedImpl(Expr, File, Line, ProgramCounter, (const TCHAR*)Format, Args...);
 //	}
 //#endif
@@ -361,7 +360,7 @@ RetType FORCENOINLINE UE_DEBUG_SECTION DispatchCheckVerify(InnerType&& Inner, Ar
 		template <typename... Types>
 		FValidateArgsInternal(Types... Args)
 		{
-			static_assert(TAnd<TIsValidVariadicFunctionArg<Types>...>::Value, "Invalid argument(s) passed to ensureMsgf");
+			static_assert((TIsValidVariadicFunctionArg<Types>::Value && ...), "Invalid argument(s) passed to ensureMsgf");
 		}
 	};
 
diff --git a/Engine/Source/Runtime/Core/Public/Misc/CString.h b/Engine/Source/Runtime/Core/Public/Misc/CString.h
--- a/Engine/Source/Runtime/Core/Public/Misc/CString.h
+++ b/Engine/Source/Runtime/Core/Public/Misc/CString.h
@@ -8,7 +8,6 @@
 #include "Misc/AssertionMacros.h"
 #include "Misc/Char.h"
 #include "Misc/VarArgs.h"
-#include "Templates/AndOrNot.h"
 #include "Templates/IsArrayOrRefOfTypeByPredicate.h"
 #include "Templates/IsValidVariadicFunctionArg.h"
 #include "Traits/IsCharEncodingCompatibleWith.h"
@@ -513,7 +512,7 @@ public:
 	static int32 Sprintf(CharType* Dest, const FmtType& Fmt, Types... Args)
 	{
 		static_assert(TIsArrayOrRefOfTypeByPredicate<FmtType, TIsCharEncodingCompatibleWithCharType>::Value, "Formatting string must be a literal string of a char type compatible with the TCString type.");
-		static_assert(TAnd<TIsValidVariadicFunctionArg<Types>...>::Value, "Invalid argument(s) passed to TCString::Sprintf");
+		static_assert((TIsValidVariadicFunctionArg<Types>::Value && ...), "Invalid argument(s) passed to TCString::Sprintf");
 
 		return SprintfImpl(Dest, (const CharType*)Fmt, Args...);
 	}
@@ -525,7 +524,7 @@ public:
 	static int32 Snprintf(CharType* Dest, int32 DestSize, const FmtType& Fmt, Types... Args)
 	{
 		static_assert(TIsArrayOrRefOfTypeByPredicate<FmtType, TIsCharEncodingCompatibleWithCharType>::Value, "Formatting string must be a literal string of a char type compatible with the TCString type.");
-		static_assert(TAnd<TIsValidVariadicFunctionArg<Types>...>::Value, "Invalid argument(s) passed to TCString::Snprintf");
+		static_assert((TIsValidVariadicFunctionArg<Types>::Value && ...), "Invalid argument(s) passed to TCString::Snprintf");
 
 		return SnprintfImpl(Dest, DestSize, (const CharType*)Fmt, Args...);
 	}
diff --git a/Engine/Source/Runtime/Core/Public/Misc/OutputDevice.h b/Engine/Source/Runtime/Core/Public/Misc/OutputDevice.h
--- a/Engine/Source/Runtime/Core/Public/Misc/OutputDevice.h
+++ b/Engine/Source/Runtime/Core/Public/Misc/OutputDevice.h
@@ -6,7 +6,6 @@
 #include "CoreTypes.h"
 #include "Logging/LogVerbosity.h"
 #include "Misc/VarArgs.h"
-#include "Templates/AndOrNot.h"
 #include "Templates/IsArrayOrRefOfTypeByPredicate.h"
 #include "Templates/IsValidVariadicFunctionArg.h"
 #include "Traits/IsCharEncodingCompatibleWith.h"
@@ -241,7 +240,7 @@ public:
 	FORCEINLINE void Logf(const FmtType& Fmt, Types... Args)
 	{
 		static_assert(TIsArrayOrRefOfTypeByPredicate<FmtType, TIsCharEncodingCompatibleWithTCHAR>::Value, "Formatting string must be a TCHAR array.");
-		static_assert(TAnd<TIsValidVariadicFunctionArg<Types>...>::Value, "Invalid argument(s) passed to FOutputDevice::Logf");
+		static_assert((TIsValidVariadicFunctionArg<Types>::Value && ...), "Invalid argument(s) passed to FOutputDevice::Logf");
 
 		LogfImpl((const TCHAR*)Fmt, Args...);
 	}
@@ -250,7 +249,7 @@ public:
 	FORCEINLINE void Logf(ELogVerbosity::Type Verbosity, const FmtType& Fmt, Types... Args)
 	{
 		static_assert(TIsArrayOrRefOfTypeByPredicate<FmtType, TIsCharEncodingCompatibleWithTCHAR>::Value, "Formatting string must be a TCHAR array.");
-		static_assert(TAnd<TIsValidVariadicFunctionArg<Types>...>::Value, "Invalid argument(s) passed to FOutputDevice::Logf");
+		static_assert((TIsValidVariadicFunctionArg<Types>::Value && ...), "Invalid argument(s) passed to FOutputDevice::Logf");
 
 		LogfImpl(Verbosity, (const TCHAR*)Fmt, Args...);
 	}
@@ -259,7 +258,7 @@ public:
 	FORCEINLINE void CategorizedLogf(const FName& Category, ELogVerbosity::Type Verbosity, const FmtType& Fmt, Types... Args)
 	{
 		static_assert(TIsArrayOrRefOfTypeByPredicate<FmtType, TIsCharEncodingCompatibleWithTCHAR>::Value, "Formatting string must be a TCHAR array.");
-		static_assert(TAnd<TIsValidVariadicFunctionArg<Types>...>::Value, "Invalid argument(s) passed to FOutputDevice::CategorizedLogf");
+		static_assert((TIsValidVariadicFunctionArg<Types>::Value && ...), "Invalid argument(s) passed to FOutputDevice::CategorizedLogf");
 
 		CategorizedLogfImpl(Category, Verbosity, (const TCHAR*)Fmt, Args...);
 	}
diff --git a/Engine/Source/Runtime/Core/Public/Misc/StringBuilder.h b/Engine/Source/Runtime/Core/Public/Misc/StringBuilder.h
--- a/Engine/Source/Runtime/Core/Public/Misc/StringBuilder.h
+++ b/Engine/Source/Runtime/Core/Public/Misc/StringBuilder.h
@@ -9,7 +9,6 @@
 #include "HAL/UnrealMemory.h"
 #include "Misc/AssertionMacros.h"
 #include "Misc/CString.h"
-#include "Templates/AndOrNot.h"
 #include "Templates/EnableIf.h"
 #include "Templates/IsArrayOrRefOfTypeByPredicate.h"
 #include "Templates/IsValidVariadicFunctionArg.h"
@@ -348,7 +347,7 @@ public:
 		UE_REQUIRES(TIsArrayOrRefOfTypeByPredicate<FmtType, TIsCharEncodingCompatibleWithCharType>::Value)>
 	BuilderType& Appendf(const FmtType& Fmt, Types... Args)
 	{
-		static_assert(TAnd<TIsValidVariadicFunctionArg<Types>...>::Value, "Invalid argument(s) passed to Appendf.");
+		static_assert((TIsValidVariadicFunctionArg<Types>::Value && ...), "Invalid argument(s) passed to Appendf.");
 		return AppendfImpl(*this, (const CharType*)Fmt, Forward<Types>(Args)...);
 	}
 
diff --git a/Engine/Source/Runtime/Core/Public/ProfilingDebugging/CsvProfiler.h b/Engine/Source/Runtime/Core/Public/ProfilingDebugging/CsvProfiler.h
--- a/Engine/Source/Runtime/Core/Public/ProfilingDebugging/CsvProfiler.h
+++ b/Engine/Source/Runtime/Core/Public/ProfilingDebugging/CsvProfiler.h
@@ -24,7 +24,6 @@
 #include "ProfilingDebugging/CsvProfilerConfig.h"
 #include "ProfilingDebugging/CsvProfilerTrace.h"
 #include "ProfilingDebugging/MiscTrace.h"
-#include "Templates/AndOrNot.h"
 #include "Templates/IsArrayOrRefOfTypeByPredicate.h"
 #include "Templates/IsValidVariadicFunctionArg.h"
 #include "Templates/RefCounting.h"
@@ -374,7 +373,7 @@ public:
 	FORCEINLINE static void RecordEventf(int32 CategoryIndex, const FmtType& Fmt, Types... Args)
 	{
 		static_assert(TIsArrayOrRefOfTypeByPredicate<FmtType, TIsCharEncodingCompatibleWithTCHAR>::Value, "Formatting string must be a TCHAR array.");
-		static_assert(TAnd<TIsValidVariadicFunctionArg<Types>...>::Value, "Invalid argument(s) passed to FCsvProfiler::RecordEventf");
+		static_assert((TIsValidVariadicFunctionArg<Types>::Value && ...), "Invalid argument(s) passed to FCsvProfiler::RecordEventf");
 		RecordEventfInternal(CategoryIndex, (const TCHAR*)Fmt, Args...);
 	}
 
diff --git a/Engine/Source/Runtime/Core/Public/Serialization/Archive.h b/Engine/Source/Runtime/Core/Public/Serialization/Archive.h
--- a/Engine/Source/Runtime/Core/Public/Serialization/Archive.h
+++ b/Engine/Source/Runtime/Core/Public/Serialization/Archive.h
@@ -13,7 +13,6 @@
 #include "Misc/EngineVersionBase.h"
 #include "Misc/VarArgs.h"
 #include "Serialization/ArchiveCookData.h"
-#include "Templates/AndOrNot.h"
 #include "Templates/EnableIf.h"
 #include "Templates/IsArrayOrRefOfTypeByPredicate.h"
 #include "Templates/IsEnumClass.h"
@@ -1837,7 +1836,7 @@ public:
 	void Logf(const FmtType& Fmt, Types... Args)
 	{
 		static_assert(TIsArrayOrRefOfTypeByPredicate<FmtType, TIsCharEncodingCompatibleWithTCHAR>::Value, "Formatting string must be a TCHAR array.");
-		static_assert(TAnd<TIsValidVariadicFunctionArg<Types>...>::Value, "Invalid argument(s) passed to FArchive::Logf");
+		static_assert((TIsValidVariadicFunctionArg<Types>::Value && ...), "Invalid argument(s) passed to FArchive::Logf");
 
 		LogfImpl(Fmt, Args...);
 	}
diff --git a/Engine/Source/Runtime/Datasmith/CADKernel/Public/CADKernel/UI/Message.h b/Engine/Source/Runtime/Datasmith/CADKernel/Public/CADKernel/UI/Message.h
--- a/Engine/Source/Runtime/Datasmith/CADKernel/Public/CADKernel/UI/Message.h
+++ b/Engine/Source/Runtime/Datasmith/CADKernel/Public/CADKernel/UI/Message.h
@@ -28,14 +28,14 @@ public:
 	template <typename FmtType, typename... Types>
 	static void Printf(EVerboseLevel Level, const FmtType& Text, Types... Args)
 	{
-		static_assert(TAnd<TIsValidVariadicFunctionArg<Types>...>::Value, "Invalid argument(s) passed to FString::Printf");
+		static_assert((TIsValidVariadicFunctionArg<Types>::Value && ...), "Invalid argument(s) passed to FString::Printf");
 		VPrintf(Level, Text, Args...);
 	}
 
 	template <typename FmtType, typename... Types>
 	static void Error(const FmtType& Text, Types... Args)
 	{
-		static_assert(TAnd<TIsValidVariadicFunctionArg<Types>...>::Value, "Invalid argument(s) passed to FString::Printf");
+		static_assert((TIsValidVariadicFunctionArg<Types>::Value && ...), "Invalid argument(s) passed to FString::Printf");
 		FString CompleteText = TEXT("ERROR: ");
 		CompleteText += Text;
 		VPrintf(Log, *CompleteText, Args...);
@@ -44,7 +44,7 @@ public:
 	template <typename FmtType, typename... Types>
 	static void Warning(const FmtType& Text, Types... Args)
 	{
-		static_assert(TAnd<TIsValidVariadicFunctionArg<Types>...>::Value, "Invalid argument(s) passed to FString::Printf");
+		static_assert((TIsValidVariadicFunctionArg<Types>::Value && ...), "Invalid argument(s) passed to FString::Printf");
 		FString CompleteText = TEXT("WARNING: ");
 		CompleteText += Text;
 		VPrintf(Log, *CompleteText, Args...);
@@ -53,7 +53,7 @@ public:
 	template <typename FmtType, typename... Types>
 	static void FillReportFile(const FmtType Header, const FmtType Text, Types... Args)
 	{
-		static_assert(TAnd<TIsValidVariadicFunctionArg<Types>...>::Value, "Invalid argument(s) passed to FString::Printf");
+		static_assert((TIsValidVariadicFunctionArg<Types>::Value && ...), "Invalid argument(s) passed to FString::Printf");
 		VReportPrintF(Header, Text, Args...);
 	}
 
diff --git a/Engine/Source/Runtime/Engine/Private/Materials/HLSLMaterialTranslator.h b/Engine/Source/Runtime/Engine/Private/Materials/HLSLMaterialTranslator.h
--- a/Engine/Source/Runtime/Engine/Private/Materials/HLSLMaterialTranslator.h
+++ b/Engine/Source/Runtime/Engine/Private/Materials/HLSLMaterialTranslator.h
@@ -698,21 +698,21 @@ protected:
 	template <typename... Types>
 	int32 AddCodeChunk(EMaterialValueType Type, const TCHAR* Format, Types... Args)
 	{
-		static_assert(TAnd<TIsValidVariadicFunctionArg<Types>...>::Value, "Invalid argument(s) passed to AddCodeChunk");
+		static_assert((TIsValidVariadicFunctionArg<Types>::Value && ...), "Invalid argument(s) passed to AddCodeChunk");
 		return AddCodeChunkInner(Type, EDerivativeStatus::NotAware, false, Format, Args...);
 	}
 
 	template <typename... Types>
 	int32 AddCodeChunkZeroDeriv(EMaterialValueType Type, const TCHAR* Format, Types... Args)
 	{
-		static_assert(TAnd<TIsValidVariadicFunctionArg<Types>...>::Value, "Invalid argument(s) passed to AddCodeChunkZeroDeriv");
+		static_assert((TIsValidVariadicFunctionArg<Types>::Value && ...), "Invalid argument(s) passed to AddCodeChunkZeroDeriv");
 		return AddCodeChunkInner(Type, EDerivativeStatus::Zero, false, Format, Args...);
 	}
 
 	template <typename... Types>
 	int32 AddCodeChunkFiniteDeriv(EMaterialValueType Type, const TCHAR* Format, Types... Args)
 	{
-		static_assert(TAnd<TIsValidVariadicFunctionArg<Types>...>::Value, "Invalid argument(s) passed to AddCodeChunkFiniteDeriv");
+		static_assert((TIsValidVariadicFunctionArg<Types>::Value && ...), "Invalid argument(s) passed to AddCodeChunkFiniteDeriv");
 		return AddCodeChunkInner(Type, EDerivativeStatus::NotValid, false, Format, Args...);
 	}
 	
@@ -724,21 +724,21 @@ protected:
 	template <typename... Types>
 	int32 AddInlinedCodeChunk(EMaterialValueType Type, const TCHAR* Format, Types... Args)
 	{
-		static_assert(TAnd<TIsValidVariadicFunctionArg<Types>...>::Value, "Invalid argument(s) passed to AddInlinedCodeChunk");
+		static_assert((TIsValidVariadicFunctionArg<Types>::Value && ...), "Invalid argument(s) passed to AddInlinedCodeChunk");
 		return AddCodeChunkInner(Type, EDerivativeStatus::NotAware, true, Format, Args...);
 	}
 
 	template <typename... Types>
 	int32 AddInlinedCodeChunkZeroDeriv(EMaterialValueType Type, const TCHAR* Format, Types... Args)
 	{
-		static_assert(TAnd<TIsValidVariadicFunctionArg<Types>...>::Value, "Invalid argument(s) passed to AddInlinedCodeChunkZeroDeriv");
+		static_assert((TIsValidVariadicFunctionArg<Types>::Value && ...), "Invalid argument(s) passed to AddInlinedCodeChunkZeroDeriv");
 		return AddCodeChunkInner(Type, EDerivativeStatus::Zero, true, Format, Args...);
 	}
 
 	template <typename... Types>
 	int32 AddInlinedCodeChunkFiniteDeriv(EMaterialValueType Type, const TCHAR* Format, Types... Args)
 	{
-		static_assert(TAnd<TIsValidVariadicFunctionArg<Types>...>::Value, "Invalid argument(s) passed to AddInlinedCodeChunkFiniteDeriv");
+		static_assert((TIsValidVariadicFunctionArg<Types>::Value && ...), "Invalid argument(s) passed to AddInlinedCodeChunkFiniteDeriv");
 		return AddCodeChunkInner(Type, EDerivativeStatus::NotValid, true, Format, Args...);
 	}
 
diff --git a/Engine/Source/Runtime/Engine/Public/VisualLogger/VisualLogger.h b/Engine/Source/Runtime/Engine/Public/VisualLogger/VisualLogger.h
--- a/Engine/Source/Runtime/Engine/Public/VisualLogger/VisualLogger.h
+++ b/Engine/Source/Runtime/Engine/Public/VisualLogger/VisualLogger.h
@@ -11,7 +11,6 @@
 #include "EngineStats.h"
 #endif
 #include "Templates/IsValidVariadicFunctionArg.h"
-#include "Templates/AndOrNot.h"
 #include "Templates/IsArrayOrRefOfTypeByPredicate.h"
 #include "Traits/IsCharEncodingCompatibleWith.h"
 #include "UObject/ObjectKey.h"
@@ -205,7 +204,7 @@ public:
 	static void CategorizedLogf(const UObject* LogOwner, const FLogCategoryBase& Category, ELogVerbosity::Type Verbosity, const FmtType& Fmt, Types... Args)
 	{
 		static_assert(TIsArrayOrRefOfTypeByPredicate<FmtType, TIsCharEncodingCompatibleWithTCHAR>::Value, "Formatting string must be a TCHAR array.");
-		static_assert(TAnd<TIsValidVariadicFunctionArg<Types>...>::Value, "Invalid argument(s) passed to FVisualLogger::CategorizedLogf");
+		static_assert((TIsValidVariadicFunctionArg<Types>::Value && ...), "Invalid argument(s) passed to FVisualLogger::CategorizedLogf");
 
 		CategorizedLogfImpl(LogOwner, Category, Verbosity, (const TCHAR*)Fmt, Args...);
 	}
@@ -213,7 +212,7 @@ public:
 	static void CategorizedLogf(const UObject* LogOwner, const FName& CategoryName, ELogVerbosity::Type Verbosity, const FmtType& Fmt, Types... Args)
 	{
 		static_assert(TIsArrayOrRefOfTypeByPredicate<FmtType, TIsCharEncodingCompatibleWithTCHAR>::Value, "Formatting string must be a TCHAR array.");
-		static_assert(TAnd<TIsValidVariadicFunctionArg<Types>...>::Value, "Invalid argument(s) passed to FVisualLogger::CategorizedLogf");
+		static_assert((TIsValidVariadicFunctionArg<Types>::Value && ...), "Invalid argument(s) passed to FVisualLogger::CategorizedLogf");
 
 		CategorizedLogfImpl(LogOwner, CategoryName, Verbosity, (const TCHAR*)Fmt, Args...);
 	}
@@ -223,7 +222,7 @@ public:
 	static void GeometryShapeLogf(const UObject* LogOwner, const FLogCategoryBase& Category, ELogVerbosity::Type Verbosity, const FVector& Start, const FVector& End, const FColor& Color, const uint16 Thickness, const FmtType& Fmt, Types... Args)
 	{
 		static_assert(TIsArrayOrRefOfTypeByPredicate<FmtType, TIsCharEncodingCompatibleWithTCHAR>::Value, "Formatting string must be a TCHAR array.");
-		static_assert(TAnd<TIsValidVariadicFunctionArg<Types>...>::Value, "Invalid argument(s) passed to FVisualLogger::GeometryShapeLogf");
+		static_assert((TIsValidVariadicFunctionArg<Types>::Value && ...), "Invalid argument(s) passed to FVisualLogger::GeometryShapeLogf");
 
 		GeometryShapeLogfImpl(LogOwner, Category, Verbosity, Start, End, Color, Thickness, (const TCHAR*)Fmt, Args...);
 	}
@@ -231,7 +230,7 @@ public:
 	static void GeometryShapeLogf(const UObject* LogOwner, const FName& CategoryName, ELogVerbosity::Type Verbosity, const FVector& Start, const FVector& End, const FColor& Color, const uint16 Thickness, const FmtType& Fmt, Types... Args)
 	{
 		static_assert(TIsArrayOrRefOfTypeByPredicate<FmtType, TIsCharEncodingCompatibleWithTCHAR>::Value, "Formatting string must be a TCHAR array.");
-		static_assert(TAnd<TIsValidVariadicFunctionArg<Types>...>::Value, "Invalid argument(s) passed to FVisualLogger::GeometryShapeLogf");
+		static_assert((TIsValidVariadicFunctionArg<Types>::Value && ...), "Invalid argument(s) passed to FVisualLogger::GeometryShapeLogf");
 
 		GeometryShapeLogfImpl(LogOwner, CategoryName, Verbosity, Start, End, Color, Thickness, (const TCHAR*)Fmt, Args...);
 	}
@@ -241,14 +240,14 @@ public:
 	static void ArrowLogf(const UObject* LogOwner, const FLogCategoryBase& Category, ELogVerbosity::Type Verbosity, const FVector& Start, const FVector& End, const FColor& Color, const FmtType& Fmt, Types... Args)
 	{
 		static_assert(TIsArrayOrRefOfTypeByPredicate<FmtType, TIsCharEncodingCompatibleWithTCHAR>::Value, "Formatting string must be a TCHAR array.");
-		static_assert(TAnd<TIsValidVariadicFunctionArg<Types>...>::Value, "Invalid argument(s) passed to FVisualLogger::ArrowLogf");
+		static_assert((TIsValidVariadicFunctionArg<Types>::Value && ...), "Invalid argument(s) passed to FVisualLogger::ArrowLogf");
 		ArrowLogfImpl(LogOwner, Category, Verbosity, Start, End, Color, (const TCHAR*)Fmt, Args...);
 	}
 	template <typename FmtType, typename... Types>
 	static void ArrowLogf(const UObject* LogOwner, const FName& CategoryName, ELogVerbosity::Type Verbosity, const FVector& Start, const FVector& End, const FColor& Color, const FmtType& Fmt, Types... Args)
 	{
 		static_assert(TIsArrayOrRefOfTypeByPredicate<FmtType, TIsCharEncodingCompatibleWithTCHAR>::Value, "Formatting string must be a TCHAR array.");
-		static_assert(TAnd<TIsValidVariadicFunctionArg<Types>...>::Value, "Invalid argument(s) passed to FVisualLogger::ArrowLogf");
+		static_assert((TIsValidVariadicFunctionArg<Types>::Value && ...), "Invalid argument(s) passed to FVisualLogger::ArrowLogf");
 
 		ArrowLogfImpl(LogOwner, CategoryName, Verbosity, Start, End, Color, (const TCHAR*)Fmt, Args...);
 	}
@@ -259,7 +258,7 @@ public:
 	static void CircleLogf(const UObject* LogOwner, const FLogCategoryBase& Category, ELogVerbosity::Type Verbosity, const FVector& Center, const FVector& UpAxis, const float Radius, const FColor& Color, const uint16 Thickness, const FmtType& Fmt, Types... Args)
 	{
 		static_assert(TIsArrayOrRefOfTypeByPredicate<FmtType, TIsCharEncodingCompatibleWithTCHAR>::Value, "Formatting string must be a TCHAR array.");
-		static_assert(TAnd<TIsValidVariadicFunctionArg<Types>...>::Value, "Invalid argument(s) passed to FVisualLogger::CircleLogf");
+		static_assert((TIsValidVariadicFunctionArg<Types>::Value && ...), "Invalid argument(s) passed to FVisualLogger::CircleLogf");
 
 		CircleLogfImpl(LogOwner, Category, Verbosity, Center, UpAxis, Radius, Color, Thickness, (const TCHAR*)Fmt, Args...);
 	}
@@ -267,7 +266,7 @@ public:
 	static void CircleLogf(const UObject* LogOwner, const FName& CategoryName, ELogVerbosity::Type Verbosity, const FVector& Center, const FVector& UpAxis, const float Radius, const FColor& Color, const uint16 Thickness, const FmtType& Fmt, Types... Args)
 	{
 		static_assert(TIsArrayOrRefOfTypeByPredicate<FmtType, TIsCharEncodingCompatibleWithTCHAR>::Value, "Formatting string must be a TCHAR array.");
-		static_assert(TAnd<TIsValidVariadicFunctionArg<Types>...>::Value, "Invalid argument(s) passed to FVisualLogger::CircleLogf");
+		static_assert((TIsValidVariadicFunctionArg<Types>::Value && ...), "Invalid argument(s) passed to FVisualLogger::CircleLogf");
 
 		CircleLogfImpl(LogOwner, CategoryName, Verbosity, Center, UpAxis, Radius, Color, Thickness, (const TCHAR*)Fmt, Args...);
 	}
@@ -277,7 +276,7 @@ public:
 	static void GeometryShapeLogf(const UObject* LogOwner, const FLogCategoryBase& Category, ELogVerbosity::Type Verbosity, const FVector& Location, float Radius, const FColor& Color, const FmtType& Fmt, Types... Args)
 	{
 		static_assert(TIsArrayOrRefOfTypeByPredicate<FmtType, TIsCharEncodingCompatibleWithTCHAR>::Value, "Formatting string must be a TCHAR array.");
-		static_assert(TAnd<TIsValidVariadicFunctionArg<Types>...>::Value, "Invalid argument(s) passed to FVisualLogger::GeometryShapeLogf");
+		static_assert((TIsValidVariadicFunctionArg<Types>::Value && ...), "Invalid argument(s) passed to FVisualLogger::GeometryShapeLogf");
 
 		GeometryShapeLogfImpl(LogOwner, Category, Verbosity, Location, Radius, Color, (const TCHAR*)Fmt, Args...);
 	}
@@ -285,7 +284,7 @@ public:
 	static void GeometryShapeLogf(const UObject* LogOwner, const FName& CategoryName, ELogVerbosity::Type Verbosity, const FVector& Location, float Radius, const FColor& Color, const FmtType& Fmt, Types... Args)
 	{
 		static_assert(TIsArrayOrRefOfTypeByPredicate<FmtType, TIsCharEncodingCompatibleWithTCHAR>::Value, "Formatting string must be a TCHAR array.");
-		static_assert(TAnd<TIsValidVariadicFunctionArg<Types>...>::Value, "Invalid argument(s) passed to FVisualLogger::GeometryShapeLogf");
+		static_assert((TIsValidVariadicFunctionArg<Types>::Value && ...), "Invalid argument(s) passed to FVisualLogger::GeometryShapeLogf");
 
 		GeometryShapeLogfImpl(LogOwner, CategoryName, Verbosity, Location, Radius, Color, (const TCHAR*)Fmt, Args...);
 	}
@@ -295,7 +294,7 @@ public:
 	static void GeometryBoxLogf(const UObject* LogOwner, const FLogCategoryBase& Category, ELogVerbosity::Type Verbosity, const FBox& Box, const FMatrix& Matrix, const FColor& Color, const FmtType& Fmt, Types... Args)
 	{
 		static_assert(TIsArrayOrRefOfTypeByPredicate<FmtType, TIsCharEncodingCompatibleWithTCHAR>::Value, "Formatting string must be a TCHAR array.");
-		static_assert(TAnd<TIsValidVariadicFunctionArg<Types>...>::Value, "Invalid argument(s) passed to FVisualLogger::GeometryBoxLogf");
+		static_assert((TIsValidVariadicFunctionArg<Types>::Value && ...), "Invalid argument(s) passed to FVisualLogger::GeometryBoxLogf");
 
 		GeometryBoxLogfImpl(LogOwner, Category, Verbosity, Box, Matrix, Color, (const TCHAR*)Fmt, Args...);
 	}
@@ -303,7 +302,7 @@ public:
 	static void GeometryBoxLogf(const UObject* LogOwner, const FName& CategoryName, ELogVerbosity::Type Verbosity, const FBox& Box, const FMatrix& Matrix, const FColor& Color, const FmtType& Fmt, Types... Args)
 	{
 		static_assert(TIsArrayOrRefOfTypeByPredicate<FmtType, TIsCharEncodingCompatibleWithTCHAR>::Value, "Formatting string must be a TCHAR array.");
-		static_assert(TAnd<TIsValidVariadicFunctionArg<Types>...>::Value, "Invalid argument(s) passed to FVisualLogger::GeometryBoxLogf");
+		static_assert((TIsValidVariadicFunctionArg<Types>::Value && ...), "Invalid argument(s) passed to FVisualLogger::GeometryBoxLogf");
 
 		GeometryBoxLogfImpl(LogOwner, CategoryName, Verbosity, Box, Matrix, Color, (const TCHAR*)Fmt, Args...);
 	}
@@ -313,7 +312,7 @@ public:
 	static void GeometryShapeLogf(const UObject* LogOwner, const FLogCategoryBase& Category, ELogVerbosity::Type Verbosity, const FVector& Origin, const FVector& Direction, const float Length, const float Angle, const FColor& Color, const FmtType& Fmt, Types... Args)
 	{
 		static_assert(TIsArrayOrRefOfTypeByPredicate<FmtType, TIsCharEncodingCompatibleWithTCHAR>::Value, "Formatting string must be a TCHAR array.");
-		static_assert(TAnd<TIsValidVariadicFunctionArg<Types>...>::Value, "Invalid argument(s) passed to FVisualLogger::GeometryShapeLogf");
+		static_assert((TIsValidVariadicFunctionArg<Types>::Value && ...), "Invalid argument(s) passed to FVisualLogger::GeometryShapeLogf");
 
 		GeometryShapeLogfImpl(LogOwner, Category, Verbosity, Origin, Direction, Length, Angle, Color, (const TCHAR*)Fmt, Args...);
 	}
@@ -321,7 +320,7 @@ public:
 	static void GeometryShapeLogf(const UObject* LogOwner, const FName& CategoryName, ELogVerbosity::Type Verbosity, const FVector& Origin, const FVector& Direction, const float Length, const float Angle, const FColor& Color, const FmtType& Fmt, Types... Args)
 	{
 		static_assert(TIsArrayOrRefOfTypeByPredicate<FmtType, TIsCharEncodingCompatibleWithTCHAR>::Value, "Formatting string must be a TCHAR array.");
-		static_assert(TAnd<TIsValidVariadicFunctionArg<Types>...>::Value, "Invalid argument(s) passed to FVisualLogger::GeometryShapeLogf");
+		static_assert((TIsValidVariadicFunctionArg<Types>::Value && ...), "Invalid argument(s) passed to FVisualLogger::GeometryShapeLogf");
 
 		GeometryShapeLogfImpl(LogOwner, CategoryName, Verbosity, Origin, Direction, Length, Angle, Color, (const TCHAR*)Fmt, Args...);
 	}
@@ -331,7 +330,7 @@ public:
 	static void GeometryShapeLogf(const UObject* LogOwner, const FLogCategoryBase& Category, ELogVerbosity::Type Verbosity, const FVector& Start, const FVector& End, const float Radius, const FColor& Color, const FmtType& Fmt, Types... Args)
 	{
 		static_assert(TIsArrayOrRefOfTypeByPredicate<FmtType, TIsCharEncodingCompatibleWithTCHAR>::Value, "Formatting string must be a TCHAR array.");
-		static_assert(TAnd<TIsValidVariadicFunctionArg<Types>...>::Value, "Invalid argument(s) passed to FVisualLogger::GeometryShapeLogf");
+		static_assert((TIsValidVariadicFunctionArg<Types>::Value && ...), "Invalid argument(s) passed to FVisualLogger::GeometryShapeLogf");
 
 		GeometryShapeLogfImpl(LogOwner, Category, Verbosity, Start, End, Radius, Color, (const TCHAR*)Fmt, Args...);
 	}
@@ -339,7 +338,7 @@ public:
 	static void GeometryShapeLogf(const UObject* LogOwner, const FName& CategoryName, ELogVerbosity::Type Verbosity, const FVector& Start, const FVector& End, const float Radius, const FColor& Color, const FmtType& Fmt, Types... Args)
 	{
 		static_assert(TIsArrayOrRefOfTypeByPredicate<FmtType, TIsCharEncodingCompatibleWithTCHAR>::Value, "Formatting string must be a TCHAR array.");
-		static_assert(TAnd<TIsValidVariadicFunctionArg<Types>...>::Value, "Invalid argument(s) passed to FVisualLogger::GeometryShapeLogf");
+		static_assert((TIsValidVariadicFunctionArg<Types>::Value && ...), "Invalid argument(s) passed to FVisualLogger::GeometryShapeLogf");
 
 		GeometryShapeLogfImpl(LogOwner, CategoryName, Verbosity, Start, End, Radius, Color, (const TCHAR*)Fmt, Args...);
 	}
@@ -349,7 +348,7 @@ public:
 	static void GeometryShapeLogf(const UObject* LogOwner, const FLogCategoryBase& Category, ELogVerbosity::Type Verbosity, const FVector& Base, float HalfHeight, float Radius, const FQuat & Rotation, const FColor& Color, const FmtType& Fmt, Types... Args)
 	{
 		static_assert(TIsArrayOrRefOfTypeByPredicate<FmtType, TIsCharEncodingCompatibleWithTCHAR>::Value, "Formatting string must be a TCHAR array.");
-		static_assert(TAnd<TIsValidVariadicFunctionArg<Types>...>::Value, "Invalid argument(s) passed to FVisualLogger::GeometryShapeLogf");
+		static_assert((TIsValidVariadicFunctionArg<Types>::Value && ...), "Invalid argument(s) passed to FVisualLogger::GeometryShapeLogf");
 
 		GeometryShapeLogfImpl(LogOwner, Category, Verbosity, Base, HalfHeight, Radius, Rotation, Color, (const TCHAR*)Fmt, Args...);
 	}
@@ -357,7 +356,7 @@ public:
 	static void GeometryShapeLogf(const UObject* LogOwner, const FName& CategoryName, ELogVerbosity::Type Verbosity, const FVector& Base, float HalfHeight, float Radius, const FQuat & Rotation, const FColor& Color, const FmtType& Fmt, Types... Args)
 	{
 		static_assert(TIsArrayOrRefOfTypeByPredicate<FmtType, TIsCharEncodingCompatibleWithTCHAR>::Value, "Formatting string must be a TCHAR array.");
-		static_assert(TAnd<TIsValidVariadicFunctionArg<Types>...>::Value, "Invalid argument(s) passed to FVisualLogger::GeometryShapeLogf");
+		static_assert((TIsValidVariadicFunctionArg<Types>::Value && ...), "Invalid argument(s) passed to FVisualLogger::GeometryShapeLogf");
 
 		GeometryShapeLogfImpl(LogOwner, CategoryName, Verbosity, Base, HalfHeight, Radius, Rotation, Color, (const TCHAR*)Fmt, Args...);
 	}
@@ -367,7 +366,7 @@ public:
 	static void NavAreaShapeLogf(const UObject* LogOwner, const FLogCategoryBase& Category, ELogVerbosity::Type Verbosity, const TArray<FVector>& ConvexPoints, float MinZ, float MaxZ, const FColor& Color, const FmtType& Fmt, Types... Args)
 	{
 		static_assert(TIsArrayOrRefOfTypeByPredicate<FmtType, TIsCharEncodingCompatibleWithTCHAR>::Value, "Formatting string must be a TCHAR array.");
-		static_assert(TAnd<TIsValidVariadicFunctionArg<Types>...>::Value, "Invalid argument(s) passed to FVisualLogger::NavAreaShapeLogf");
+		static_assert((TIsValidVariadicFunctionArg<Types>::Value && ...), "Invalid argument(s) passed to FVisualLogger::NavAreaShapeLogf");
 
 		NavAreaShapeLogfImpl(LogOwner, Category, Verbosity, ConvexPoints, MinZ, MaxZ, Color, (const TCHAR*)Fmt, Args...);
 	}
@@ -375,7 +374,7 @@ public:
 	static void NavAreaShapeLogf(const UObject* LogOwner, const FName& CategoryName, ELogVerbosity::Type Verbosity, const TArray<FVector>& ConvexPoints, float MinZ, float MaxZ, const FColor& Color, const FmtType& Fmt, Types... Args)
 	{
 		static_assert(TIsArrayOrRefOfTypeByPredicate<FmtType, TIsCharEncodingCompatibleWithTCHAR>::Value, "Formatting string must be a TCHAR array.");
-		static_assert(TAnd<TIsValidVariadicFunctionArg<Types>...>::Value, "Invalid argument(s) passed to FVisualLogger::NavAreaShapeLogf");
+		static_assert((TIsValidVariadicFunctionArg<Types>::Value && ...), "Invalid argument(s) passed to FVisualLogger::NavAreaShapeLogf");
 
 		NavAreaShapeLogfImpl(LogOwner, CategoryName, Verbosity, ConvexPoints, MinZ, MaxZ, Color, (const TCHAR*)Fmt, Args...);
 	}
@@ -385,7 +384,7 @@ public:
 	static void GeometryShapeLogf(const UObject* LogOwner, const FLogCategoryBase& Category, ELogVerbosity::Type Verbosity, const TArray<FVector>& Vertices, const TArray<int32>& Indices, const FColor& Color, const FmtType& Fmt, Types... Args)
 	{
 		static_assert(TIsArrayOrRefOfTypeByPredicate<FmtType, TIsCharEncodingCompatibleWithTCHAR>::Value, "Formatting string must be a TCHAR array.");
-		static_assert(TAnd<TIsValidVariadicFunctionArg<Types>...>::Value, "Invalid argument(s) passed to FVisualLogger::GeometryShapeLogf");
+		static_assert((TIsValidVariadicFunctionArg<Types>::Value && ...), "Invalid argument(s) passed to FVisualLogger::GeometryShapeLogf");
 
 		GeometryShapeLogfImpl(LogOwner, Category, Verbosity, Vertices, Indices, Color, (const TCHAR*)Fmt, Args...);
 	}
@@ -393,7 +392,7 @@ public:
 	static void GeometryShapeLogf(const UObject* LogOwner, const FName& CategoryName, ELogVerbosity::Type Verbosity, const TArray<FVector>& Vertices, const TArray<int32>& Indices, const FColor& Color, const FmtType& Fmt, Types... Args)
 	{
 		static_assert(TIsArrayOrRefOfTypeByPredicate<FmtType, TIsCharEncodingCompatibleWithTCHAR>::Value, "Formatting string must be a TCHAR array.");
-		static_assert(TAnd<TIsValidVariadicFunctionArg<Types>...>::Value, "Invalid argument(s) passed to FVisualLogger::GeometryShapeLogf");
+		static_assert((TIsValidVariadicFunctionArg<Types>::Value && ...), "Invalid argument(s) passed to FVisualLogger::GeometryShapeLogf");
 
 		GeometryShapeLogfImpl(LogOwner, CategoryName, Verbosity, Vertices, Indices, Color, (const TCHAR*)Fmt, Args...);
 	}
@@ -403,7 +402,7 @@ public:
 	static void GeometryConvexLogf(const UObject* LogOwner, const FLogCategoryBase& Category, ELogVerbosity::Type Verbosity, const TArray<FVector>& Points, const FColor& Color, const FmtType& Fmt, Types... Args)
 	{
 		static_assert(TIsArrayOrRefOfTypeByPredicate<FmtType, TIsCharEncodingCompatibleWithTCHAR>::Value, "Formatting string must be a TCHAR array.");
-		static_assert(TAnd<TIsValidVariadicFunctionArg<Types>...>::Value, "Invalid argument(s) passed to FVisualLogger::GeometryConvexLogf");
+		static_assert((TIsValidVariadicFunctionArg<Types>::Value && ...), "Invalid argument(s) passed to FVisualLogger::GeometryConvexLogf");
 
 		GeometryConvexLogfImpl(LogOwner, Category, Verbosity, Points, Color, (const TCHAR*)Fmt, Args...);
 	}
@@ -411,7 +410,7 @@ public:
 	static void GeometryConvexLogf(const UObject* LogOwner, const FName& CategoryName, ELogVerbosity::Type Verbosity, const TArray<FVector>& Points, const FColor& Color, const FmtType& Fmt, Types... Args)
 	{
 		static_assert(TIsArrayOrRefOfTypeByPredicate<FmtType, TIsCharEncodingCompatibleWithTCHAR>::Value, "Formatting string must be a TCHAR array.");
-		static_assert(TAnd<TIsValidVariadicFunctionArg<Types>...>::Value, "Invalid argument(s) passed to FVisualLogger::GeometryConvexLogf");
+		static_assert((TIsValidVariadicFunctionArg<Types>::Value && ...), "Invalid argument(s) passed to FVisualLogger::GeometryConvexLogf");
 
 		GeometryConvexLogfImpl(LogOwner, CategoryName, Verbosity, Points, Color, (const TCHAR*)Fmt, Args...);
 	}
@@ -421,7 +420,7 @@ public:
 	static void HistogramDataLogf(const UObject* LogOwner, const FLogCategoryBase& Category, ELogVerbosity::Type Verbosity, FName GraphName, FName DataName, const FVector2D& Data, const FColor& Color, const FmtType& Fmt, Types... Args)
 	{
 		static_assert(TIsArrayOrRefOfTypeByPredicate<FmtType, TIsCharEncodingCompatibleWithTCHAR>::Value, "Formatting string must be a TCHAR array.");
-		static_assert(TAnd<TIsValidVariadicFunctionArg<Types>...>::Value, "Invalid argument(s) passed to FVisualLogger::HistogramDataLogf");
+		static_assert((TIsValidVariadicFunctionArg<Types>::Value && ...), "Invalid argument(s) passed to FVisualLogger::HistogramDataLogf");
 
 		HistogramDataLogfImpl(LogOwner, Category, Verbosity, GraphName, DataName, Data, Color, (const TCHAR*)Fmt, Args...);
 	}
@@ -429,7 +428,7 @@ public:
 	static void HistogramDataLogf(const UObject* LogOwner, const FName& CategoryName, ELogVerbosity::Type Verbosity, FName GraphName, FName DataName, const FVector2D& Data, const FColor& Color, const FmtType& Fmt, Types... Args)
 	{
 		static_assert(TIsArrayOrRefOfTypeByPredicate<FmtType, TIsCharEncodingCompatibleWithTCHAR>::Value, "Formatting string must be a TCHAR array.");
-		static_assert(TAnd<TIsValidVariadicFunctionArg<Types>...>::Value, "Invalid argument(s) passed to FVisualLogger::HistogramDataLogf");
+		static_assert((TIsValidVariadicFunctionArg<Types>::Value && ...), "Invalid argument(s) passed to FVisualLogger::HistogramDataLogf");
 
 		HistogramDataLogfImpl(LogOwner, CategoryName, Verbosity, GraphName, DataName, Data, Color, (const TCHAR*)Fmt, Args...);
 	}
diff --git a/Engine/Source/Runtime/GameplayDebugger/Public/GameplayDebuggerTypes.h b/Engine/Source/Runtime/GameplayDebugger/Public/GameplayDebuggerTypes.h
--- a/Engine/Source/Runtime/GameplayDebugger/Public/GameplayDebuggerTypes.h
+++ b/Engine/Source/Runtime/GameplayDebugger/Public/GameplayDebuggerTypes.h
@@ -66,7 +66,7 @@ public:
 	void Printf(const FmtType& Fmt, Types... Args)
 	{
 		static_assert(TIsArrayOrRefOfTypeByPredicate<FmtType, TIsCharEncodingCompatibleWithTCHAR>::Value, "Formatting string must be a TCHAR array.");
-		static_assert(TAnd<TIsValidVariadicFunctionArg<Types>...>::Value, "Invalid argument(s) passed to FGameplayDebuggerCanvasContext::Printf");
+		static_assert((TIsValidVariadicFunctionArg<Types>::Value && ...), "Invalid argument(s) passed to FGameplayDebuggerCanvasContext::Printf");
 
 		PrintfImpl(FColor::White, 1.0f, (const TCHAR*)Fmt, Args...);
 	}
@@ -75,7 +75,7 @@ public:
 	void Printf(const FColor& Color, const FmtType& Fmt, Types... Args)
 	{
 		static_assert(TIsArrayOrRefOfTypeByPredicate<FmtType, TIsCharEncodingCompatibleWithTCHAR>::Value, "Formatting string must be a TCHAR array.");
-		static_assert(TAnd<TIsValidVariadicFunctionArg<Types>...>::Value, "Invalid argument(s) passed to FGameplayDebuggerCanvasContext::Printf");
+		static_assert((TIsValidVariadicFunctionArg<Types>::Value && ...), "Invalid argument(s) passed to FGameplayDebuggerCanvasContext::Printf");
 
 		PrintfImpl(Color, 1.0f, (const TCHAR*)Fmt, Args...);
 	}
@@ -84,7 +84,7 @@ public:
 	void Printf(const FColor& Color, const float Alpha, const FmtType& Fmt, Types... Args)
 	{
 		static_assert(TIsArrayOrRefOfTypeByPredicate<FmtType, TIsCharEncodingCompatibleWithTCHAR>::Value, "Formatting string must be a TCHAR array.");
-		static_assert(TAnd<TIsValidVariadicFunctionArg<Types>...>::Value, "Invalid argument(s) passed to FGameplayDebuggerCanvasContext::Printf");
+		static_assert((TIsValidVariadicFunctionArg<Types>::Value && ...), "Invalid argument(s) passed to FGameplayDebuggerCanvasContext::Printf");
 
 		PrintfImpl(Color, Alpha, (const TCHAR*)Fmt, Args...);
 	}
@@ -93,7 +93,7 @@ public:
 	void PrintfAt(float PosX, float PosY, const FmtType& Fmt, Types... Args)
 	{
 		static_assert(TIsArrayOrRefOfTypeByPredicate<FmtType, TIsCharEncodingCompatibleWithTCHAR>::Value, "Formatting string must be a TCHAR array.");
-		static_assert(TAnd<TIsValidVariadicFunctionArg<Types>...>::Value, "Invalid argument(s) passed to FGameplayDebuggerCanvasContext::PrintfAt");
+		static_assert((TIsValidVariadicFunctionArg<Types>::Value && ...), "Invalid argument(s) passed to FGameplayDebuggerCanvasContext::PrintfAt");
 
 		PrintfAtImpl(PosX, PosY, FColor::White, 1.0f, (const TCHAR*)Fmt, Args...);
 	}
@@ -102,7 +102,7 @@ public:
 	void PrintfAt(float PosX, float PosY, const FColor& Color, const FmtType& Fmt, Types... Args)
 	{
 		static_assert(TIsArrayOrRefOfTypeByPredicate<FmtType, TIsCharEncodingCompatibleWithTCHAR>::Value, "Formatting string must be a TCHAR array.");
-		static_assert(TAnd<TIsValidVariadicFunctionArg<Types>...>::Value, "Invalid argument(s) passed to FGameplayDebuggerCanvasContext::PrintfAt");
+		static_assert((TIsValidVariadicFunctionArg<Types>::Value && ...), "Invalid argument(s) passed to FGameplayDebuggerCanvasContext::PrintfAt");
 
 		PrintfAtImpl(PosX, PosY, Color, 1.0f, (const TCHAR*)Fmt, Args...);
 	}
@@ -111,7 +111,7 @@ public:
 	void PrintfAt(float PosX, float PosY, const FColor& Color, const float Alpha, const FmtType& Fmt, Types... Args)
 	{
 		static_assert(TIsArrayOrRefOfTypeByPredicate<FmtType, TIsCharEncodingCompatibleWithTCHAR>::Value, "Formatting string must be a TCHAR array.");
-		static_assert(TAnd<TIsValidVariadicFunctionArg<Types>...>::Value, "Invalid argument(s) passed to FGameplayDebuggerCanvasContext::PrintfAt");
+		static_assert((TIsValidVariadicFunctionArg<Types>::Value && ...), "Invalid argument(s) passed to FGameplayDebuggerCanvasContext::PrintfAt");
 
 		PrintfAtImpl(PosX, PosY, Color, Alpha, (const TCHAR*)Fmt, Args...);
 	}
