Augment invalid material cached data warning log with material path to facilitate its cleanup.

We also mute shader log errors during path validation. In editor, these are redundant with compilation failure error(s). In cooks, the warning(s) should suffice if paths are invalid.

#rnx

diff --git a/Engine/Source/Runtime/Engine/Private/Materials/Material.cpp b/Engine/Source/Runtime/Engine/Private/Materials/Material.cpp
--- a/Engine/Source/Runtime/Engine/Private/Materials/Material.cpp
+++ b/Engine/Source/Runtime/Engine/Private/Materials/Material.cpp
@@ -2070,7 +2070,7 @@ void UMaterial::UpdateCachedExpressionData()
 		LocalCachedExpressionData->EditorOnlyData->MaterialLayers.LinkAllLayersToParent();
 	}
 
-	LocalCachedExpressionData->Validate();
+	LocalCachedExpressionData->Validate(*this);
 
 	CachedExpressionData.Reset(LocalCachedExpressionData);
 	CachedHLSLTree.Reset(LocalCachedTree);
diff --git a/Engine/Source/Runtime/Engine/Private/Materials/MaterialCachedData.cpp b/Engine/Source/Runtime/Engine/Private/Materials/MaterialCachedData.cpp
--- a/Engine/Source/Runtime/Engine/Private/Materials/MaterialCachedData.cpp
+++ b/Engine/Source/Runtime/Engine/Private/Materials/MaterialCachedData.cpp
@@ -35,6 +35,7 @@
 #include "SparseVolumeTexture/SparseVolumeTexture.h"
 #include "Engine/Font.h"
 #include "LandscapeGrassType.h"
+#include "Logging/LogScopedVerbosityOverride.h"
 #include "Curves/CurveLinearColor.h"
 #include "Curves/CurveLinearColorAtlas.h"
 #include "UObject/UE5MainStreamObjectVersion.h"
@@ -995,7 +996,7 @@ void FMaterialCachedExpressionData::UpdateForCachedHLSLTree(const FMaterialCache
 	}
 }
 
-void FMaterialCachedExpressionData::Validate()
+void FMaterialCachedExpressionData::Validate(const UMaterialInterface& Material)
 {
 	if (EditorOnlyData)
 	{
@@ -1009,7 +1010,9 @@ void FMaterialCachedExpressionData::Validate()
 
 		if (!FPlatformProperties::RequiresCookedData() && AllowShaderCompiling())
 		{
-			// TODO Validate the shader include paths in the editor before storing them on the expression instead
+			// Mute log errors created by GetShaderSourceFilePath during include path validation
+			LOG_SCOPE_VERBOSITY_OVERRIDE(LogShaders, ELogVerbosity::Fatal);
+
 			for (auto PathIt = EditorOnlyData->ExpressionIncludeFilePaths.CreateIterator(); PathIt; ++PathIt)
 			{
 				const FString& IncludeFilePath = *PathIt;
@@ -1030,7 +1033,7 @@ void FMaterialCachedExpressionData::Validate()
 
 				if (!bValidExpressionIncludePath)
 				{
-					UE_LOG(LogMaterial, Warning, TEXT("Expression include file path [%s] is invalid, removing from cached material data."), *IncludeFilePath);
+					UE_LOG(LogMaterial, Warning, TEXT("Expression include file path '%s' is invalid, removing from cached data for material '%s'."), *IncludeFilePath, *Material.GetPathName());
 					PathIt.RemoveCurrent();
 				}
 			}
diff --git a/Engine/Source/Runtime/Engine/Private/Materials/MaterialInterface.cpp b/Engine/Source/Runtime/Engine/Private/Materials/MaterialInterface.cpp
--- a/Engine/Source/Runtime/Engine/Private/Materials/MaterialInterface.cpp
+++ b/Engine/Source/Runtime/Engine/Private/Materials/MaterialInterface.cpp
@@ -215,7 +215,7 @@ void UMaterialInterface::Serialize(FArchive& Ar)
 		else
 		{
 #if WITH_EDITOR
-			CachedExpressionData->Validate();
+			CachedExpressionData->Validate(*this);
 #endif
 		}
 
diff --git a/Engine/Source/Runtime/Engine/Public/MaterialCachedData.h b/Engine/Source/Runtime/Engine/Public/MaterialCachedData.h
--- a/Engine/Source/Runtime/Engine/Public/MaterialCachedData.h
+++ b/Engine/Source/Runtime/Engine/Public/MaterialCachedData.h
@@ -194,7 +194,7 @@ struct FMaterialCachedExpressionData
 	void AnalyzeMaterial(UMaterial& Material);
 
 	ENGINE_API void UpdateForCachedHLSLTree(const FMaterialCachedHLSLTree& CachedTree, const FStaticParameterSet* StaticParameters);
-	void Validate();
+	void Validate(const UMaterialInterface& Material);
 
 	/** Adds a parameter. If this returns false, a parameter with identical name has already been added but it was set to a different value. */
 	bool AddParameter(const FMaterialParameterInfo& ParameterInfo, const FMaterialParameterMetadata& ParameterMeta, UObject*& OutReferencedTexture);
