World Partition RuntimeHashSet: don't render viewport 3d cells unless 3d debug draw is activated.

diff --git a/Engine/Source/Runtime/Engine/Private/WorldPartition/RuntimeHashSet/WorldPartitionRuntimeHashSetDebugDraw.cpp b/Engine/Source/Runtime/Engine/Private/WorldPartition/RuntimeHashSet/WorldPartitionRuntimeHashSetDebugDraw.cpp
--- a/Engine/Source/Runtime/Engine/Private/WorldPartition/RuntimeHashSet/WorldPartitionRuntimeHashSetDebugDraw.cpp
+++ b/Engine/Source/Runtime/Engine/Private/WorldPartition/RuntimeHashSet/WorldPartitionRuntimeHashSetDebugDraw.cpp
@@ -185,14 +185,6 @@ bool UWorldPartitionRuntimeHashSet::Draw2D(FWorldPartitionDraw2DContext& DrawCon
 				}
 				
 				DrawContext.LocalDrawBox(GridScreenBounds, CellBoundsMin, CellBoundsSize, FLinearColor::Black, 1, WorldToScreen);
-
-				// Draw Cell using its debug color
-				const FBox Box(Cell->GetCellBounds());
-				const FVector BoxCenter(Box.GetCenter());
-				const FColor BoxColor(CellColors[0].CopyWithNewOpacity(CellOpacity).ToFColor(true));
-				const FVector CellPos = WorldPartitionTransform.TransformPosition(BoxCenter);
-				DrawDebugSolidBox(OwningWorld, Box, BoxColor, WorldPartitionTransform, false, -1.f, 255);				
-				DrawDebugBox(OwningWorld, CellPos, Box.GetExtent(), WorldPartitionTransform.GetRotation(), BoxColor.WithAlpha(255), false, -1.f, 255, 10.f);
 			}
 		}
 
@@ -316,5 +308,76 @@ bool UWorldPartitionRuntimeHashSet::Draw2D(FWorldPartitionDraw2DContext& DrawCon
 
 void UWorldPartitionRuntimeHashSet::Draw3D(const TArray<FWorldPartitionStreamingSource>& Sources) const
 {
-	return;
+	TRACE_CPUPROFILER_EVENT_SCOPE(UWorldPartitionRuntimeHashSet::Draw3D);
+
+	const UWorldPartition* WorldPartition = GetOuterUWorldPartition();
+	const FTransform WorldPartitionTransform = WorldPartition->GetInstanceTransform();
+	const UWorld* OwningWorld = WorldPartition->GetWorld();
+	const UContentBundleManager* ContentBundleManager = OwningWorld->ContentBundleManager;
+	const EWorldPartitionRuntimeCellVisualizeMode VisualizeMode = GetStreamingCellVisualizeMode();
+	const TMap<FName, FColor> DataLayerDebugColors = GetDataLayerDebugColors(WorldPartition);
+
+	ForEachStreamingData([this, &Sources, VisualizeMode, &DataLayerDebugColors, ContentBundleManager, OwningWorld, &WorldPartitionTransform](const FRuntimePartitionStreamingData& StreamingData)
+	{
+		for (const FWorldPartitionStreamingSource& Source : Sources)
+		{
+			// @todo_jfd
+			const FSoftObjectPath HLODLayer;
+			Source.ForEachShape(StreamingData.LoadingRange, StreamingData.Name, HLODLayer, false, [this, &StreamingData, VisualizeMode, &DataLayerDebugColors, ContentBundleManager, OwningWorld, &WorldPartitionTransform](const FSphericalSector& Shape)
+			{
+				const FSphere ShapeSphere(Shape.GetCenter(), Shape.GetRadius());
+
+				StreamingData.SpatialIndex.Get()->ForEachIntersectingElement(ShapeSphere, [this, VisualizeMode, &DataLayerDebugColors, ContentBundleManager, OwningWorld, &WorldPartitionTransform](UWorldPartitionRuntimeCell* Cell)
+				{
+					const FVector2D CellBoundsSize = FVector2D(Cell->GetCellBounds().GetSize());
+					const FVector2D CellBoundsMin = FVector2D(Cell->GetCellBounds().Min);
+
+					float CellOpacity = 0.0f;;
+					TArray<FLinearColor> CellColors;
+
+					switch (GShowRuntimeHashSetDebugDisplayMode)
+					{
+					case 0:
+						CellColors.Add(Cell->GetDebugColor(VisualizeMode));
+						CellOpacity = 0.25f / FMath::Max<float>(GShowRuntimeHashSetDebugDisplayLevelCount, 1);
+						break;
+					case 1:
+						if (DataLayerDebugColors.Num() && Cell->GetDataLayers().Num())
+						{
+							for (const FName& DataLayer : Cell->GetDataLayers())
+							{
+								CellColors.Add(DataLayerDebugColors[DataLayer]);
+							}
+							CellOpacity = 0.67f;
+						}
+						break;
+					case 2:
+						if (ContentBundleManager && Cell->GetContentBundleID().IsValid())
+						{
+							if (const FContentBundleBase* ContentBundle = ContentBundleManager->GetContentBundle(OwningWorld, Cell->GetContentBundleID()))
+							{
+								check(ContentBundle->GetDescriptor());
+								CellColors.Add(ContentBundle->GetDescriptor()->GetDebugColor());
+								CellOpacity = 0.67f;
+							}
+						}
+						break;
+					}
+
+					if (CellColors.IsEmpty())
+					{
+						CellColors.Add(FLinearColor::White);
+						CellOpacity = 0.1f;
+					}
+
+					// Draw Cell using its debug color
+					const FBox Box(Cell->GetCellBounds());
+					const FVector BoxCenter(Box.GetCenter());
+					const FColor BoxColor(CellColors[0].CopyWithNewOpacity(CellOpacity).ToFColor(true));
+					const FVector CellPos = WorldPartitionTransform.TransformPosition(BoxCenter);
+					DrawDebugBox(OwningWorld, CellPos, Box.GetExtent(), WorldPartitionTransform.GetRotation(), BoxColor.WithAlpha(255), false, -1.f, 255, 100.f);
+				});
+			});
+		}
+	});
 }
\ No newline at end of file
