World Partition Class Descriptor Registry: detect redirector loops.

diff --git a/Engine/Source/Runtime/Engine/Private/WorldPartition/WorldPartitionClassDescRegistry.cpp b/Engine/Source/Runtime/Engine/Private/WorldPartition/WorldPartitionClassDescRegistry.cpp
--- a/Engine/Source/Runtime/Engine/Private/WorldPartition/WorldPartitionClassDescRegistry.cpp
+++ b/Engine/Source/Runtime/Engine/Private/WorldPartition/WorldPartitionClassDescRegistry.cpp
@@ -253,7 +253,7 @@ void FWorldPartitionClassDescRegistry::PrefetchClassDescs(const TArray<FTopLevel
 	{
 		check(AssetData.IsValid());
 
-		TArray<FTopLevelAssetPath> ClassRedirects;
+		TSet<FTopLevelAssetPath> ClassRedirects;
 		while (AssetData.IsRedirector())
 		{
 			// Folow the redirector to the destination object
@@ -277,7 +277,17 @@ void FWorldPartitionClassDescRegistry::PrefetchClassDescs(const TArray<FTopLevel
 				break;
 			}
 
-			ClassRedirects.Add(FTopLevelAssetPath(AssetData.ToSoftObjectPath().ToString()));
+			bool bRedirectAlreadyRegistered;
+			ClassRedirects.Add(FTopLevelAssetPath(AssetData.ToSoftObjectPath().ToString()), &bRedirectAlreadyRegistered);
+
+			if (bRedirectAlreadyRegistered)
+			{
+				const FString ClassRedirectsLoop = FString::JoinBy(ClassRedirects, TEXT("\n"), [](const FTopLevelAssetPath& RedirectPath) { return FString::Printf(TEXT("  -> %s"), *RedirectPath.ToString()); })
+					+ FString::Printf(TEXT("\n  -> %s"), *FTopLevelAssetPath(AssetData.ToSoftObjectPath().ToString()).ToString());
+				UE_LOG(LogWorldPartition, Warning, TEXT("Redirector loop detected for '%s' from '%s':\n%s"), *AssetData.ToSoftObjectPath().ToString(), *AssetClassPath.ToString(), *ClassRedirectsLoop);
+				AssetData = FAssetData();
+				break;
+			}
 
 			// Find the blueprint asset or the proper blueprint redirector
 			RedirectAssets.Sort(SortBlueprintAssetsByNameLength);
