Make USplineComponent's public-exposed recursive spline sampling functions sample the spline accurately for arbitrary curve ranges (previously were only accurate for half-segment regions).

Move the original limit-range recursive helper function logic to private, not-blueprint-exposed methods.

diff --git a/Engine/Source/Runtime/Engine/Classes/Components/SplineComponent.h b/Engine/Source/Runtime/Engine/Classes/Components/SplineComponent.h
--- a/Engine/Source/Runtime/Engine/Classes/Components/SplineComponent.h
+++ b/Engine/Source/Runtime/Engine/Classes/Components/SplineComponent.h
@@ -789,11 +789,11 @@ public:
 	UFUNCTION(BlueprintCallable, Category = Spline)
 	ENGINE_API FTransform FindTransformClosestToWorldLocation(const FVector& WorldLocation, ESplineCoordinateSpace::Type CoordinateSpace, bool bUseScale = false) const;
 
-	/** Given a threshold, recursively sub-divides the spline section until the list of segments (polyline) matches the spline shape. */
+	/** Given a threshold, recursively sub-divides the spline section until the list of segments (polyline) matches the spline shape. Note: Prefer ConvertSplineToPolyline_InDistanceRange */
 	UFUNCTION(BlueprintCallable, Category = Spline)
 	ENGINE_API bool DivideSplineIntoPolylineRecursive(float StartDistanceAlongSpline, float EndDistanceAlongSpline, ESplineCoordinateSpace::Type CoordinateSpace, const float MaxSquareDistanceFromSpline, TArray<FVector>& OutPoints) const;
 
-	/** Given a threshold, recursively sub-divides the spline section until the list of segments (polyline) matches the spline shape. */
+	/** Given a threshold, recursively sub-divides the spline section until the list of segments (polyline) matches the spline shape. Note: Prefer ConvertSplineToPolyline_InDistanceRange */
 	UFUNCTION(BlueprintCallable, Category = Spline)
 	ENGINE_API bool DivideSplineIntoPolylineRecursiveWithDistances(float StartDistanceAlongSpline, float EndDistanceAlongSpline, ESplineCoordinateSpace::Type CoordinateSpace, const float MaxSquareDistanceFromSpline, TArray<FVector>& OutPoints, TArray<double>& OutDistancesAlongSpline) const;
 
@@ -835,6 +835,11 @@ private:
 	/** Returns the parametric value t which would result in a spline segment of the given length between S(0)...S(t) */
 	ENGINE_API float GetSegmentParamFromLength(const int32 Index, const float Length, const float SegmentLength) const;
 
+	// Internal helper function called by ConvertSplineSegmentToPolyLine -- assumes the input is within a half-segment, so testing the distance to midpoint will be an accurate guide to subdivision
+	bool DivideSplineIntoPolylineRecursiveWithDistancesHelper(float StartDistanceAlongSpline, float EndDistanceAlongSpline, ESplineCoordinateSpace::Type CoordinateSpace, const float MaxSquareDistanceFromSpline, TArray<FVector>& OutPoints, TArray<double>& OutDistancesAlongSpline) const;
+	// Internal helper function called by ConvertSplineSegmentToPolyLine -- assumes the input is within a half-segment, so testing the distance to midpoint will be an accurate guide to subdivision
+	bool DivideSplineIntoPolylineRecursiveHelper(float StartDistanceAlongSpline, float EndDistanceAlongSpline, ESplineCoordinateSpace::Type CoordinateSpace, const float MaxSquareDistanceFromSpline, TArray<FVector>& OutPoints) const;
+
 	/** Returns a const reference to the specified position point, but gives back a dummy point if there are no points */
 	inline const FInterpCurvePointVector& GetPositionPointSafe(int32 PointIndex) const
 	{
diff --git a/Engine/Source/Runtime/Engine/Private/Components/SplineComponent.cpp b/Engine/Source/Runtime/Engine/Private/Components/SplineComponent.cpp
--- a/Engine/Source/Runtime/Engine/Private/Components/SplineComponent.cpp
+++ b/Engine/Source/Runtime/Engine/Private/Components/SplineComponent.cpp
@@ -1617,6 +1617,11 @@ FTransform USplineComponent::FindTransformClosestToWorldLocation(const FVector&
 }
 
 bool USplineComponent::DivideSplineIntoPolylineRecursiveWithDistances(float StartDistanceAlongSpline, float EndDistanceAlongSpline, ESplineCoordinateSpace::Type CoordinateSpace, const float MaxSquareDistanceFromSpline, TArray<FVector>& OutPoints, TArray<double>& OutDistancesAlongSpline) const
+{
+	return ConvertSplineToPolyline_InDistanceRange(CoordinateSpace, MaxSquareDistanceFromSpline, StartDistanceAlongSpline, EndDistanceAlongSpline, OutPoints, OutDistancesAlongSpline, false);
+}
+
+bool USplineComponent::DivideSplineIntoPolylineRecursiveWithDistancesHelper(float StartDistanceAlongSpline, float EndDistanceAlongSpline, ESplineCoordinateSpace::Type CoordinateSpace, const float MaxSquareDistanceFromSpline, TArray<FVector>& OutPoints, TArray<double>& OutDistancesAlongSpline) const
 {
 	double Dist = EndDistanceAlongSpline - StartDistanceAlongSpline;
 	if (Dist <= 0.0f)
@@ -1633,8 +1638,8 @@ bool USplineComponent::DivideSplineIntoPolylineRecursiveWithDistances(float Star
 	{
 		TArray<FVector> NewPoints[2];
 		TArray<double> NewDistancesAlongSpline[2];
-		DivideSplineIntoPolylineRecursiveWithDistances(StartDistanceAlongSpline, MiddlePointDistancAlongSpline, CoordinateSpace, MaxSquareDistanceFromSpline, NewPoints[0], NewDistancesAlongSpline[0]);
-		DivideSplineIntoPolylineRecursiveWithDistances(MiddlePointDistancAlongSpline, EndDistanceAlongSpline, CoordinateSpace, MaxSquareDistanceFromSpline, NewPoints[1], NewDistancesAlongSpline[1]);
+		DivideSplineIntoPolylineRecursiveWithDistancesHelper(StartDistanceAlongSpline, MiddlePointDistancAlongSpline, CoordinateSpace, MaxSquareDistanceFromSpline, NewPoints[0], NewDistancesAlongSpline[0]);
+		DivideSplineIntoPolylineRecursiveWithDistancesHelper(MiddlePointDistancAlongSpline, EndDistanceAlongSpline, CoordinateSpace, MaxSquareDistanceFromSpline, NewPoints[1], NewDistancesAlongSpline[1]);
 		if ((NewPoints[0].Num() > 0) && (NewPoints[1].Num() > 0))
 		{
 			check(NewPoints[0].Last() == NewPoints[1][0]);
@@ -1660,10 +1665,16 @@ bool USplineComponent::DivideSplineIntoPolylineRecursiveWithDistances(float Star
 	return (OutPoints.Num() > 0);
 }
 
+bool USplineComponent::DivideSplineIntoPolylineRecursiveHelper(float StartDistanceAlongSpline, float EndDistanceAlongSpline, ESplineCoordinateSpace::Type CoordinateSpace, const float MaxSquareDistanceFromSpline, TArray<FVector>& OutPoints) const
+{
+	TArray<double> DummyDistancesAlongSpline;
+	return DivideSplineIntoPolylineRecursiveWithDistancesHelper(StartDistanceAlongSpline, EndDistanceAlongSpline, CoordinateSpace, MaxSquareDistanceFromSpline, OutPoints, DummyDistancesAlongSpline);
+}
+
 bool USplineComponent::DivideSplineIntoPolylineRecursive(float StartDistanceAlongSpline, float EndDistanceAlongSpline, ESplineCoordinateSpace::Type CoordinateSpace, const float MaxSquareDistanceFromSpline, TArray<FVector>& OutPoints) const
 {
 	TArray<double> DummyDistancesAlongSpline;
-	return DivideSplineIntoPolylineRecursiveWithDistances(StartDistanceAlongSpline, EndDistanceAlongSpline, CoordinateSpace, MaxSquareDistanceFromSpline, OutPoints, DummyDistancesAlongSpline);
+	return ConvertSplineToPolyline_InDistanceRange(CoordinateSpace, MaxSquareDistanceFromSpline, StartDistanceAlongSpline, EndDistanceAlongSpline, OutPoints, DummyDistancesAlongSpline, false);
 }
 
 bool USplineComponent::ConvertSplineSegmentToPolyLine(int32 SplinePointStartIndex, ESplineCoordinateSpace::Type CoordinateSpace, const float MaxSquareDistanceFromSpline, TArray<FVector>& OutPoints) const
@@ -1689,7 +1700,7 @@ bool USplineComponent::ConvertSplineSegmentToPolyLine(int32 SplinePointStartInde
 		double SubstepEndDist = SubstepStartDist + SubstepSize;
 		TArray<FVector> NewPoints;
 		// Recursively sub-divide each segment until the requested precision is reached :
-		if (DivideSplineIntoPolylineRecursive(SubstepStartDist, SubstepEndDist, CoordinateSpace, MaxSquareDistanceFromSpline, NewPoints))
+		if (DivideSplineIntoPolylineRecursiveHelper(SubstepStartDist, SubstepEndDist, CoordinateSpace, MaxSquareDistanceFromSpline, NewPoints))
 		{
 			if (OutPoints.Num() > 0)
 			{
@@ -1856,7 +1867,7 @@ bool USplineComponent::ConvertSplineToPolyline_InDistanceRange(ESplineCoordinate
 			NewPoints.Reset();
 			NewDistances.Reset();
 			// Recursively sub-divide each segment until the requested precision is reached :
-			if (DivideSplineIntoPolylineRecursiveWithDistances(SubstepStartDist, SubstepEndDist, CoordinateSpace, MaxSquareDistanceFromSpline, NewPoints, NewDistances))
+			if (DivideSplineIntoPolylineRecursiveWithDistancesHelper(SubstepStartDist, SubstepEndDist, CoordinateSpace, MaxSquareDistanceFromSpline, NewPoints, NewDistances))
 			{
 				if (OutPoints.Num() > 0)
 				{
