Clarifying comments around sound wave raw data and oodle.

diff --git a/Engine/Source/Runtime/Engine/Classes/Sound/SoundWave.h b/Engine/Source/Runtime/Engine/Classes/Sound/SoundWave.h
--- a/Engine/Source/Runtime/Engine/Classes/Sound/SoundWave.h
+++ b/Engine/Source/Runtime/Engine/Classes/Sound/SoundWave.h
@@ -878,7 +878,17 @@ public:
 	ENGINE_API const uint8* GetResourceData() const;
 
 #if WITH_EDITORONLY_DATA
-	/** Uncompressed wav data 16 bit in mono or stereo - stereo not allowed for multichannel data */
+	/** 
+	* Holds the uncompressed wav data that was imported. This is guaranteed to be 16 bit, and is
+	* mono or stereo - stereo not allowed for multichannel data. For multichannel data, there are
+	* distinct RIFF files concatenated in the RawData - one for each channel. These can be accessed with
+	* ChannelOffsets and ChannelSizes (see GetImportedSoundWaveData for example).
+	* 
+	* This structure is a pass-through for editor bulk data. It does an in-place conversion on the audio
+	* bits to allow the audio to compress significantly better when the underlying bulk data compression hits it.
+	* 
+	* If you need access to the underlying audio data, use GetImportedSoundWaveData and avoid touching this.
+	*/
 	struct FEditorAudioBulkData
 	{
 		UE::Serialization::FEditorBulkData RawData;
@@ -891,8 +901,8 @@ public:
 
 		//
 		// Deprecated unused API forwarding for potential backwards compatability issues.
+		// As the raw data needs to be converted before use or storage, always access it via the above functions.
 		//
-
 		UE_DEPRECATED(5.4, "CreateLegacyUniqueIdentifier is provided just for API backwards compatibility.")
 		void CreateLegacyUniqueIdentifier(UObject* Owner)
 		{
