fix sound class overrides not respecting the new value of bApplyAmbientVolumes

diff --git a/Engine/Source/Runtime/Engine/Private/ActiveSound.cpp b/Engine/Source/Runtime/Engine/Private/ActiveSound.cpp
--- a/Engine/Source/Runtime/Engine/Private/ActiveSound.cpp
+++ b/Engine/Source/Runtime/Engine/Private/ActiveSound.cpp
@@ -280,8 +280,15 @@ void FActiveSound::SetWorld(UWorld* InWorld)
 void FActiveSound::SetSound(USoundBase* InSound)
 {
 	Sound = InSound;
-	bApplyInteriorVolumes = (SoundClassOverride && SoundClassOverride->Properties.bApplyAmbientVolumes)
-		|| (Sound && Sound->ShouldApplyInteriorVolumes());
+	
+	if (SoundClassOverride)
+	{
+		bApplyInteriorVolumes = SoundClassOverride->Properties.bApplyAmbientVolumes;
+	}
+	else
+	{
+		bApplyInteriorVolumes = Sound && Sound->ShouldApplyInteriorVolumes();
+	}
 }
 
 void FActiveSound::SetSourceEffectChain(USoundEffectSourcePresetChain* InSourceEffectChain)
@@ -292,8 +299,15 @@ void FActiveSound::SetSourceEffectChain(USoundEffectSourcePresetChain* InSourceE
 void FActiveSound::SetSoundClass(USoundClass* SoundClass)
 {
 	SoundClassOverride = SoundClass;
-	bApplyInteriorVolumes = (SoundClassOverride && SoundClassOverride->Properties.bApplyAmbientVolumes)
-		|| (Sound && Sound->ShouldApplyInteriorVolumes());
+	
+	if (SoundClassOverride)
+	{
+		bApplyInteriorVolumes = SoundClassOverride->Properties.bApplyAmbientVolumes;
+	}
+	else
+	{
+		bApplyInteriorVolumes = Sound && Sound->ShouldApplyInteriorVolumes();
+	}
 }
 
 void FActiveSound::SetAttenuationSettingsAsset(TObjectPtr<USoundAttenuation> InSoundAttenuation)
