- Stop PSO precache event from preventing GC
- This fixes a GC timeout on some platforms, as we are saturing the request queue so it doesn't complete fast enough to prevent GC failure

diff --git a/Engine/Source/Runtime/Engine/Classes/Particles/ParticleSystem.h b/Engine/Source/Runtime/Engine/Classes/Particles/ParticleSystem.h
--- a/Engine/Source/Runtime/Engine/Classes/Particles/ParticleSystem.h
+++ b/Engine/Source/Runtime/Engine/Classes/Particles/ParticleSystem.h
@@ -122,8 +122,6 @@ public:
 
 	ENGINE_API virtual void PostInitProperties() override;
 
-	ENGINE_API virtual bool IsReadyForFinishDestroy() override;
-
 	/** Max number of components of this system to keep resident in the world component pool. */
 	UPROPERTY(EditAnywhere, Category = Performance, AdvancedDisplay)
 	uint32 MaxPoolSize;
diff --git a/Engine/Source/Runtime/Engine/Private/Particles/ParticleComponents.cpp b/Engine/Source/Runtime/Engine/Private/Particles/ParticleComponents.cpp
--- a/Engine/Source/Runtime/Engine/Private/Particles/ParticleComponents.cpp
+++ b/Engine/Source/Runtime/Engine/Private/Particles/ParticleComponents.cpp
@@ -145,17 +145,6 @@ void UFXSystemAsset::PostInitProperties()
 #endif
 }
 
-bool UFXSystemAsset::IsReadyForFinishDestroy()
-{
-	// Don't touch PrecachePSOsEvent directly because the cleanup task could set to a nullptr
-	if (PrecachePSOsEvent)
-	{
-		return false;
-	}
-
-	return Super::IsReadyForFinishDestroy();
-}
-
 void UFXSystemAsset::LaunchPSOPrecaching(TArrayView<VFsPerMaterialData> VFsPerMaterials)
 {
 	FPSOPrecacheParams PreCachePSOParams;
@@ -177,25 +166,28 @@ void UFXSystemAsset::LaunchPSOPrecaching(TArrayView<VFsPerMaterialData> VFsPerMa
 	{
 		struct FReleasePrecachePSOsEventTask
 		{
-			explicit FReleasePrecachePSOsEventTask(FGraphEventRef& InPrecachePSOsEvent)
-				: PrecachePSOsEvent(&InPrecachePSOsEvent)
+			explicit FReleasePrecachePSOsEventTask(UFXSystemAsset* OwnerAsset)
+				: WeakOwnerAsset(OwnerAsset)
 			{
 			}
 
 			static TStatId GetStatId() { return TStatId(); }
-			static ENamedThreads::Type GetDesiredThread() { return ENamedThreads::AnyThread; }
+			static ENamedThreads::Type GetDesiredThread() { return ENamedThreads::GameThread; }
 			static ESubsequentsMode::Type GetSubsequentsMode() { return ESubsequentsMode::TrackSubsequents; }
 
 			void DoTask(ENamedThreads::Type CurrentThread, const FGraphEventRef& MyCompletionGraphEvent)
 			{
-				*PrecachePSOsEvent = nullptr;
+				if (UFXSystemAsset* Asset = WeakOwnerAsset.Get())
+				{
+					Asset->PrecachePSOsEvent = nullptr;
+				}
 			}
 
-			FGraphEventRef* PrecachePSOsEvent;
+			TWeakObjectPtr<UFXSystemAsset> WeakOwnerAsset;
 		};
 
 		// need to set `PrecachePSOsEvent` before the task is launched to not race with its execution
-		TGraphTask<FReleasePrecachePSOsEventTask>* ReleasePrecachePSOsEventTask = TGraphTask<FReleasePrecachePSOsEventTask>::CreateTask(&PrecachePSOsEvents).ConstructAndHold(PrecachePSOsEvent);
+		TGraphTask<FReleasePrecachePSOsEventTask>* ReleasePrecachePSOsEventTask = TGraphTask<FReleasePrecachePSOsEventTask>::CreateTask(&PrecachePSOsEvents).ConstructAndHold(this);
 		PrecachePSOsEvent = ReleasePrecachePSOsEventTask->GetCompletionEvent();
 		ReleasePrecachePSOsEventTask->Unlock();
 	}
