UClusterUnionComponent::SyncClusterUnionFromProxy optimizations

- Removed the need of maintain two sets to cache FExternalSpatialAccelerationPayload structures as components are added/removed
- Added a new constructor to FExternalSpatialAccelerationPayload so we can directly pass a TObjectKey and the chaos UniqueIdx, so we don't have to build them from scratch (which requires lock the scene to get the Particle idx and access the Global objects array to build the WeakObjtPtr).
- Wrapped TObjectKey in a wrapper we can use to store other elements we will use later, so we don't have to go back and forth between converting between ObjectKey and Raw ptr which requires access to the global objects array and is not cheap.
- Removed TSet the need to do TSet differences
- Added a new native event that when broadcasted in addition to the new bone ids that were added, it also passes an array of bone IDs removed. (So listeners don't have to calculate it on their own again).
- Splitted UGeometryCollectionComponent::SetPerParticleCollisionProfileName so its components can be called independently. There are cases where we need to call it with different profile name witching the same scope, so part of the method doesn't need to be executed again. To ensure all parts are executed as desired, a new transaction object was created that need to be passed to the new UGeometryCollectionComponent::SetParticleCollisionProfileName

diff --git a/Engine/Source/Runtime/Engine/Classes/PhysicsEngine/ClusterUnionComponent.h b/Engine/Source/Runtime/Engine/Classes/PhysicsEngine/ClusterUnionComponent.h
--- a/Engine/Source/Runtime/Engine/Classes/PhysicsEngine/ClusterUnionComponent.h
+++ b/Engine/Source/Runtime/Engine/Classes/PhysicsEngine/ClusterUnionComponent.h
@@ -32,16 +32,46 @@ namespace Chaos
 	class ISpatialAccelerationCollection;
 }
 
+USTRUCT()
+struct FClusterUnionBoneData
+{
+	GENERATED_BODY()
+
+	FClusterUnionBoneData() = default;
+	FClusterUnionBoneData(int32 InBoneId) : ID(InBoneId)
+	{	
+	}
+	
+	FClusterUnionBoneData(int32 InBoneId, Chaos::FUniqueIdx InParticleID) : ID(InBoneId), ParticleID(InParticleID)
+	{	
+	}
+
+	int32 ID = INDEX_NONE;
+	Chaos::FUniqueIdx ParticleID;
+
+	bool operator==(const FClusterUnionBoneData& Other) const
+	{
+		return ID == Other.ID;
+	}
+
+	bool operator!=(const FClusterUnionBoneData& Other) const
+	{
+		return !(*this == Other);
+	}
+
+	friend uint32 GetTypeHash(const FClusterUnionBoneData& Data)
+	{
+		return GetTypeHash(Data.ID);
+	}
+};
+
 USTRUCT()
 struct FClusteredComponentData
 {
 	GENERATED_BODY()
 
 	// Set of bone Ids that we actually added into the cluster union.
-	TSet<int32> BoneIds;
-
-	// Cached acceleration structure handles - needed to properly cleanup the component from the accel structure.
-	TSet<FExternalSpatialAccelerationPayload> CachedAccelerationPayloads;
+	TSet<FClusterUnionBoneData> BonesData;
 
 	// Using a TWeakObjectPtr here because the UClusterUnionReplicatedProxyComponent will have a pointer back
 	// and we don't want to get into a situation where a circular reference occurs.
@@ -87,10 +117,7 @@ struct FClusterUnionPendingAddData
 	GENERATED_BODY()
 
 	UPROPERTY()
-	TArray<int32> BoneIds;
-	
-	UPROPERTY()
-	TArray<FExternalSpatialAccelerationPayload> AccelerationPayloads;
+	TSet<FClusterUnionBoneData> BonesData;
 };
 
 USTRUCT()
@@ -106,7 +133,7 @@ struct FClusterUnionInitializationData
 };
 
 DECLARE_DYNAMIC_MULTICAST_DELEGATE_ThreeParams(FOnClusterUnionAddedComponent, UPrimitiveComponent*, Component, const TSet<int32>&, BoneIds, bool, bIsNew);
-DECLARE_MULTICAST_DELEGATE_ThreeParams(FOnClusterUnionAddedComponentNative, UPrimitiveComponent*, const TSet<int32>& /*BoneIds*/, bool /*bIsNew*/);
+DECLARE_MULTICAST_DELEGATE_FourParams(FOnClusterUnionAddedComponentNative, UPrimitiveComponent*, const TSet<FClusterUnionBoneData>& /*BoneIds*/, const TArray<int32>& /*RemovedBoneIds*/, bool /*bIsNew*/);
 
 DECLARE_DYNAMIC_MULTICAST_DELEGATE_OneParam(FOnClusterUnionRemovedComponent, UPrimitiveComponent*, Component);
 DECLARE_MULTICAST_DELEGATE_OneParam(FOnClusterUnionRemovedComponentNative, UPrimitiveComponent*);
@@ -168,6 +195,59 @@ public:
 
 	ENGINE_API bool IsAnchored() const;
 
+	/** Structure that stores an object key and its raw ptr, so it can be accessed without going trough
+	* the global uobjects array (so it could be garbage, therefore it has to be used withing the same frame)
+	*/
+	struct FMappedComponentKey
+	{
+		FMappedComponentKey() : ComponentPtr(nullptr)
+		{
+		}
+
+		FMappedComponentKey(TObjectKey<UPrimitiveComponent> InComponentKey, UPrimitiveComponent* InComponent)
+		{
+			ComponentKey = InComponentKey;
+			ComponentPtr = InComponent;
+		}
+		
+		FMappedComponentKey(UPrimitiveComponent* InComponent)
+		{
+			ComponentKey = InComponent;
+			ComponentPtr = InComponent;
+		}
+
+		bool operator==(const FMappedComponentKey& Other) const
+		{
+			return ComponentKey == Other.ComponentKey;
+		}
+
+		friend uint32 GetTypeHash(const FMappedComponentKey& InMappedComponentKey)
+		{
+			return GetTypeHash(InMappedComponentKey.ComponentKey);
+		}
+
+		TObjectKey<UPrimitiveComponent> ComponentKey;
+		UPrimitiveComponent* ComponentPtr;
+	};
+
+	struct FMappedBoneData
+	{
+		FMappedBoneData() = default;
+		
+		FMappedBoneData(Chaos::FPhysicsObjectHandle PhysicsObjectHande, Chaos::FPBDRigidParticle* RigidParticle, const Chaos::FUniqueIdx& ParticleID, const FTransform& ChildToParentTransform)
+			: PhysicsObjectHandle(PhysicsObjectHande)
+			, RigidParticle(RigidParticle)
+			, ParticleID(ParticleID)
+			, ChildToParentTransform(ChildToParentTransform)
+		{
+		}
+
+		Chaos::FPhysicsObjectHandle PhysicsObjectHandle = nullptr;
+		Chaos::FPBDRigidParticle* RigidParticle = nullptr;
+		Chaos::FUniqueIdx ParticleID;
+		FTransform ChildToParentTransform;
+	};
+
 	// SyncClusterUnionFromProxy will examine the make up of the cluster union (particles, child to parent, etc.) and do whatever is needed on the GT in terms of bookkeeping.
 	ENGINE_API void SyncClusterUnionFromProxy();
 
@@ -224,7 +304,7 @@ protected:
 	UFUNCTION()
 	ENGINE_API void ForceSetChildToParent(UPrimitiveComponent* InComponent, const TArray<int32>& BoneIds, const TArray<FTransform>& ChildToParent);
 
-	ENGINE_API void BroadcastComponentAddedEvents(UPrimitiveComponent* ChangedComponent, const TSet<int32>& BoneIds, bool bIsNew);
+	ENGINE_API void BroadcastComponentAddedEvents(UPrimitiveComponent* ChangedComponent, const TSet<FClusterUnionBoneData>& BoneIds, bool bIsNew, TArray<int32>&& RemovedBoneIDs);
 	ENGINE_API void BroadcastComponentRemovedEvents(UPrimitiveComponent* ChangedComponent);
 
 	Chaos::FClusterUnionPhysicsProxy* GetPhysicsProxy() const { return PhysicsProxy; }
@@ -293,6 +373,18 @@ private:
 	// An acceleration structure of all children components managed by the cluster union itself.
 	TUniquePtr<FSpatialAcceleration> AccelerationStructure;
 
+	/** Takes a source bone data type and extracts Bone IDs from it*/
+	template<typename TContainer, typename TSourceData>
+	void GetBoneIDsFromComponentData(TContainer& OutBoneIDs, const TSourceData& InSourceData) const
+	{
+		static_assert(std::is_same_v<TSourceData, FClusterUnionPendingAddData> || std::is_same_v<TSourceData, FClusteredComponentData>, "Only FClusterUnionPendingAddData and FClusteredComponentData are supported");
+
+		OutBoneIDs.Reset();
+		OutBoneIDs.Reserve(InSourceData.BonesData.Num());
+
+		Algo::Transform(InSourceData.BonesData, OutBoneIDs, &FClusterUnionBoneData::ID);
+	}
+
 	// Need to handle the fact that this component may or may not be initialized prior to the components referenced in
 	// ClusteredComponentsReferences. This function lets us listen to OnComponentPhysicsStateChanged on the incoming
 	// primitive component so that once the physics state is properly created we can begin the process of adding it.
@@ -305,7 +397,7 @@ private:
 
 	// These functions only get called when the physics thread syncs to the game thread thereby enforcing a physics thread authoritative view of
 	// what particles are currently contained within the cluster union.
-	ENGINE_API void HandleAddOrModifiedClusteredComponent(UPrimitiveComponent* ChangedComponent, const TMap<int32, FTransform>& PerBoneChildToParent);
+	ENGINE_API void HandleAddOrModifiedClusteredComponent(const FMappedComponentKey& ChangedComponentData, const TMap<int32, FMappedBoneData>& PerBoneChildToParent);
 	ENGINE_API void HandleRemovedClusteredComponent(TObjectKey<UPrimitiveComponent> ChangedComponent);
 
 	ENGINE_API TArray<UPrimitiveComponent*> GetAllCurrentChildComponents() const;
diff --git a/Engine/Source/Runtime/Engine/Classes/PhysicsEngine/ExternalSpatialAccelerationPayload.h b/Engine/Source/Runtime/Engine/Classes/PhysicsEngine/ExternalSpatialAccelerationPayload.h
--- a/Engine/Source/Runtime/Engine/Classes/PhysicsEngine/ExternalSpatialAccelerationPayload.h
+++ b/Engine/Source/Runtime/Engine/Classes/PhysicsEngine/ExternalSpatialAccelerationPayload.h
@@ -8,7 +8,7 @@
 #include "Chaos/ParticleHandleFwd.h"
 #include "Chaos/SpatialAccelerationFwd.h"
 #include "UObject/Object.h"
-#include "UObject/WeakObjectPtrTemplates.h"
+#include "UObject/ObjectKey.h"
 
 #include "ExternalSpatialAccelerationPayload.generated.h"
 
@@ -28,7 +28,8 @@ public:
 
 	ENGINE_API FExternalSpatialAccelerationPayload();
 
-	ENGINE_API void Initialize(UPrimitiveComponent* InComponent, int32 InBoneId);
+	ENGINE_API void Initialize(TObjectKey<UPrimitiveComponent> InComponent, int32 InBoneId);
+	ENGINE_API void Initialize(TObjectKey<UPrimitiveComponent> InComponent, int32 InBoneId, const Chaos::FUniqueIdx& UniqueIdx);
 
 	ENGINE_API Chaos::FGeometryParticle* GetExternalGeometryParticle_ExternalThread() const;
 	ENGINE_API bool operator==(const FExternalSpatialAccelerationPayload& Other) const;
@@ -44,7 +45,7 @@ public:
 
 	friend ENGINE_API uint32 GetTypeHash(const FExternalSpatialAccelerationPayload& Payload);
 private:
-	TWeakObjectPtr<UPrimitiveComponent> Component = nullptr;
+	TObjectKey<UPrimitiveComponent> Component = nullptr;
 	int32 BoneId = INDEX_NONE;
 	Chaos::FUniqueIdx CachedUniqueIdx;
 
diff --git a/Engine/Source/Runtime/Engine/Private/PhysicsEngine/ClusterUnionComponent.cpp b/Engine/Source/Runtime/Engine/Private/PhysicsEngine/ClusterUnionComponent.cpp
--- a/Engine/Source/Runtime/Engine/Private/PhysicsEngine/ClusterUnionComponent.cpp
+++ b/Engine/Source/Runtime/Engine/Private/PhysicsEngine/ClusterUnionComponent.cpp
@@ -95,21 +95,24 @@ void UClusterUnionComponent::AddComponentToCluster(UPrimitiveComponent* InCompon
 		return;
 	}
 
+	TObjectKey<UPrimitiveComponent> ComponentKey(InComponent);
+
 	FClusterUnionPendingAddData PendingData;
-	PendingData.BoneIds = BoneIds;
 
 	if (!InComponent->HasValidPhysicsState())
 	{
-		if (!PendingComponentsToAdd.Contains(InComponent))
+		if (!PendingComponentsToAdd.Contains(ComponentKey))
 		{
+			Algo::Transform(BoneIds, PendingData.BonesData, [](int32 BoneID){ return FClusterUnionBoneData(BoneID); });
+
 			// Early out - defer adding the component to the cluster until the component has a valid physics state.
-			PendingComponentsToAdd.Add(InComponent, PendingData);
+			PendingComponentsToAdd.Add(ComponentKey, PendingData);
 			InComponent->OnComponentPhysicsStateChanged.AddDynamic(this, &UClusterUnionComponent::HandleComponentPhysicsStateChange);
 		}
 		return;
 	}
 
-	PendingComponentsToAdd.Remove(InComponent);
+	PendingComponentsToAdd.Remove(ComponentKey);
 
 	TArray<Chaos::FPhysicsObjectHandle> AllObjects = InComponent->GetAllPhysicsObjects();
 	FLockedWritePhysicsObjectExternalInterface Interface = FPhysicsObjectExternalInterface::LockWrite(AllObjects);
@@ -144,8 +147,7 @@ void UClusterUnionComponent::AddComponentToCluster(UPrimitiveComponent* InCompon
 		return;
 	}
 
-	PendingData.BoneIds.Empty(Objects.Num());
-	PendingData.AccelerationPayloads.Empty(Objects.Num());
+	PendingData.BonesData.Empty(Objects.Num());
 
 	// We want to fix the ChildToParent relationship between the particles we're adding and the cluster union
 	// based on the current game-thread state, not the future physics-thread state as the two may diverge
@@ -159,20 +161,24 @@ void UClusterUnionComponent::AddComponentToCluster(UPrimitiveComponent* InCompon
 	const bool bHasKinematicObject = !Interface->AreAllDynamicOrSleeping(Objects);
 	for (Chaos::FPhysicsObjectHandle Object : Objects)
 	{
-		if (Object)
+		if (Chaos::FPBDRigidParticle* Particle = Interface->GetRigidParticle(Object))
 		{
 			const int32 BoneId = Chaos::FPhysicsObjectInterface::GetId(Object);
-			PendingData.BoneIds.Add(BoneId);
+
+			FClusterUnionBoneData BoneData;
+			BoneData.ID = BoneId;
+			BoneData.ParticleID = Particle ? Particle->UniqueIdx() : Chaos::FUniqueIdx();
+			
+			PendingData.BonesData.Add(BoneData);
 
 			if (AccelerationStructure)
 			{
 				const FBox HandleBounds = Interface->GetWorldBounds({ &Object, 1 });
 				FExternalSpatialAccelerationPayload Handle;
-				Handle.Initialize(InComponent, BoneId);
-				if (Handle.IsValid())
+				Handle.Initialize(ComponentKey, BoneData.ID, BoneData.ParticleID);
+				if (ensure(Handle.IsValid()))
 				{
 					AccelerationStructure->UpdateElement(Handle, Chaos::TAABB<Chaos::FReal, 3>{HandleBounds.Min, HandleBounds.Max}, HandleBounds.IsValid != 0);
-					PendingData.AccelerationPayloads.Add(Handle);
 				}
 			}
 
@@ -197,10 +203,13 @@ void UClusterUnionComponent::AddComponentToCluster(UPrimitiveComponent* InCompon
 		AddGTParticleGeometry(Objects);
 	}
 
-	PendingComponentSync.Add(InComponent, PendingData);
+	PendingComponentSync.Add(ComponentKey, PendingData);
 
 	PhysicsProxy->AddPhysicsObjects_External(Objects);
-	ForceSetChildToParent(InComponent, PendingData.BoneIds, ChildToParents);
+
+	TArray<int32> BoneIDs;
+	GetBoneIDsFromComponentData(BoneIDs, PendingData);
+	ForceSetChildToParent(InComponent, BoneIDs, ChildToParents);
 }
 
 void UClusterUnionComponent::RemoveComponentFromCluster(UPrimitiveComponent* InComponent)
@@ -210,7 +219,8 @@ void UClusterUnionComponent::RemoveComponentFromCluster(UPrimitiveComponent* InC
 		return;
 	}
 
-	const int32 NumRemoved = PendingComponentsToAdd.Remove(InComponent);
+	const TObjectKey<UPrimitiveComponent> RemovedComponentKey(InComponent);
+	const int32 NumRemoved = PendingComponentsToAdd.Remove(RemovedComponentKey);
 	if (NumRemoved > 0)
 	{
 		// We haven't actually added yet so we can early out.
@@ -224,9 +234,11 @@ void UClusterUnionComponent::RemoveComponentFromCluster(UPrimitiveComponent* InC
 	{
 		if (AccelerationStructure)
 		{
-			for (const FExternalSpatialAccelerationPayload& Payload : PendingData->AccelerationPayloads)
+			for (const FClusterUnionBoneData BoneData : PendingData->BonesData)
 			{
-				if (Payload.IsValid())
+				FExternalSpatialAccelerationPayload Payload;
+				Payload.Initialize(RemovedComponentKey, BoneData.ID, BoneData.ParticleID);
+				if (ensure(Payload.IsValid()))
 				{
 					AccelerationStructure->RemoveElement(Payload);
 				}
@@ -249,7 +261,9 @@ void UClusterUnionComponent::RemoveComponentFromCluster(UPrimitiveComponent* InC
 
 		if (InComponent->HasValidPhysicsState())
 		{
-			PhysicsObjectsToRemove = TSet<Chaos::FPhysicsObjectHandle>{ GetAllPhysicsObjectsById(InComponent, ComponentData->BoneIds.Array()) };
+			TArray<int32> BoneIDs;
+			GetBoneIDsFromComponentData(BoneIDs, *ComponentData);
+			PhysicsObjectsToRemove = TSet<Chaos::FPhysicsObjectHandle>{ GetAllPhysicsObjectsById(InComponent, BoneIDs) };
 		}
 
 		if (IsAuthority())
@@ -262,14 +276,15 @@ void UClusterUnionComponent::RemoveComponentFromCluster(UPrimitiveComponent* InC
 
 		if (AccelerationStructure)
 		{
-			for (const FExternalSpatialAccelerationPayload& Payload : ComponentData->CachedAccelerationPayloads)
+			for (const FClusterUnionBoneData& BoneData : ComponentData->BonesData)
 			{
-				if (Payload.IsValid())
+				FExternalSpatialAccelerationPayload Payload;
+				Payload.Initialize(RemovedComponentKey, BoneData.ID, BoneData.ParticleID);
+				if (ensure(Payload.IsValid()))
 				{
 					AccelerationStructure->RemoveElement(Payload);
 				}
 			}
-			ComponentData->CachedAccelerationPayloads.Reset();
 		}
 
 		// If PhysicsObjectsToRemove is empty, it either means that BoneIds is empty OR the component's physics state is already destroyed.
@@ -289,19 +304,21 @@ void UClusterUnionComponent::RemoveComponentBonesFromCluster(UPrimitiveComponent
 	{
 		return;
 	}
-	TSet<int32> RemoveBoneIdsSet{ BoneIds };
+
+	TObjectKey<UPrimitiveComponent> ComponentKey(InComponent);
+
 	// If we're removing bones from a component and it hasn't been added into the component yet, then we want to make sure 
 	// the set we actually want to add doesn't have the remove bones in it.
-	if (FClusterUnionPendingAddData* PendingAdd = PendingComponentsToAdd.Find(InComponent))
+	if (FClusterUnionPendingAddData* PendingAdd = PendingComponentsToAdd.Find(ComponentKey))
 	{
-		for (int32 BoneId : RemoveBoneIdsSet)
+		for (int32 BoneId : BoneIds)
 		{
-			PendingAdd->BoneIds.Remove(BoneId);
+			PendingAdd->BonesData.Remove(FClusterUnionBoneData(BoneId));
 		}
 
-		if (PendingAdd->BoneIds.IsEmpty())
+		if (PendingAdd->BonesData.IsEmpty())
 		{
-			PendingComponentsToAdd.Remove(InComponent);
+			PendingComponentsToAdd.Remove(ComponentKey);
 		}
 
 		// Component is pending add so we can be confident we don't need to do any more work.
@@ -311,54 +328,53 @@ void UClusterUnionComponent::RemoveComponentBonesFromCluster(UPrimitiveComponent
 	// If we're still waiting on the physics sync for the added component, we still need to
 	// remove its acceleration handles from the acceleration structure since the payloads have
 	// already been added and will not yet be stored in PerComponentData.
-	if (FClusterUnionPendingAddData* PendingData = PendingComponentSync.Find(InComponent))
+	if (FClusterUnionPendingAddData* PendingData = PendingComponentSync.Find(ComponentKey))
 	{
 		if (AccelerationStructure)
 		{
-			for (int32 BoneId : RemoveBoneIdsSet)
+			for (int32 BoneId : BoneIds)
 			{
 				FExternalSpatialAccelerationPayload Payload;
-				Payload.Initialize(InComponent, BoneId);
-				PendingData->AccelerationPayloads.Remove(Payload);
-				if (Payload.IsValid())
+				Payload.Initialize(ComponentKey, BoneId);
+				if (ensure(Payload.IsValid()))
 				{
 					AccelerationStructure->RemoveElement(Payload);
 				}
 			}
 		}
 
-		for (int32 BoneId : RemoveBoneIdsSet)
+		for (int32 BoneId : BoneIds)
 		{
-			PendingData->BoneIds.Remove(BoneId);
+			PendingData->BonesData.Remove(FClusterUnionBoneData(BoneId));
 		}
 
-		if (PendingData->BoneIds.IsEmpty())
+		if (PendingData->BonesData.IsEmpty())
 		{
-			PendingComponentSync.Remove(InComponent);
+			PendingComponentSync.Remove(ComponentKey);
 		}
 	}
 
 	TSet<Chaos::FPhysicsObjectHandle> PhysicsObjectsToRemove;
-	if (FClusteredComponentData* ComponentData = PerComponentData.Find(InComponent))
+	if (FClusteredComponentData* ComponentData = PerComponentData.Find(ComponentKey))
 	{
-		for (int32 BoneId : RemoveBoneIdsSet)
+		for (TSet<FClusterUnionBoneData>::TIterator RemoveBoneDataIterator = ComponentData->BonesData.CreateIterator(); RemoveBoneDataIterator; ++RemoveBoneDataIterator)
 		{
 			if (AccelerationStructure)
 			{
 				FExternalSpatialAccelerationPayload Payload;
-				Payload.Initialize(InComponent, BoneId);
-				ComponentData->CachedAccelerationPayloads.Remove(Payload);
-				if (Payload.IsValid())
+				Payload.Initialize(InComponent, RemoveBoneDataIterator->ID, RemoveBoneDataIterator->ParticleID);
+				if (ensure(Payload.IsValid()))
 				{
 					AccelerationStructure->RemoveElement(Payload);
 				}
 			}
-			ComponentData->BoneIds.Remove(BoneId);
+
+			RemoveBoneDataIterator.RemoveCurrent();
 		}
 
 		if (InComponent->HasValidPhysicsState())
 		{
-			PhysicsObjectsToRemove = TSet<Chaos::FPhysicsObjectHandle>{ GetAllPhysicsObjectsById(InComponent, RemoveBoneIdsSet.Array()) };
+			PhysicsObjectsToRemove = TSet<Chaos::FPhysicsObjectHandle>{ GetAllPhysicsObjectsById(InComponent, BoneIds) };
 		}
 
 		// If PhysicsObjectsToRemove is empty, it either means that BoneIds is empty OR the component's physics state is already destroyed.
@@ -369,7 +385,7 @@ void UClusterUnionComponent::RemoveComponentBonesFromCluster(UPrimitiveComponent
 			RemoveGTParticleGeometry(PhysicsObjectsToRemove);
 		}
 
-		if (ComponentData->BoneIds.IsEmpty())
+		if (ComponentData->BonesData.IsEmpty())
 		{
 			// We need to mark the replicated proxy as pending deletion.
 			// This way anyone who tries to use the replicated proxy component knows that it
@@ -384,7 +400,7 @@ void UClusterUnionComponent::RemoveComponentBonesFromCluster(UPrimitiveComponent
 				}
 			}
 
-			PerComponentData.Remove(InComponent);
+			PerComponentData.Remove(ComponentKey);
 		}
 
 	}
@@ -800,7 +816,9 @@ void UClusterUnionComponent::HandleComponentPhysicsStateChange(UPrimitiveCompone
 
 	if (FClusterUnionPendingAddData* PendingData = PendingComponentsToAdd.Find(ChangedComponent))
 	{
-		AddComponentToCluster(ChangedComponent, PendingData->BoneIds);
+		TArray<int32> AsBoneIDs;
+		GetBoneIDsFromComponentData(AsBoneIDs, *PendingData);
+		AddComponentToCluster(ChangedComponent, AsBoneIDs);
 	}
 }
 
@@ -848,7 +866,7 @@ void UClusterUnionComponent::SyncClusterUnionFromProxy()
 	// Note that at the UClusterUnionComponent level we really only want to be dealing with components.
 	// Hence why we need to modify each of the particles that we synced from the game thread into a
 	// component + bone id combination for identification. 
-	TMap<TObjectKey<UPrimitiveComponent>, TMap<int32, FTransform>> MappedData;
+	TMap<FMappedComponentKey, TMap<int32, FMappedBoneData>> MappedData;
 
 	{
 		FLockedReadPhysicsObjectExternalInterface Interface = FPhysicsObjectExternalInterface::LockRead(GetWorld()->GetPhysicsScene());
@@ -857,10 +875,17 @@ void UClusterUnionComponent::SyncClusterUnionFromProxy()
 			// Using the scene's proxy to component mapping let's us detect a component physics state was destroyed.
 			if (UPrimitiveComponent* Component = Scene->GetOwningComponent<UPrimitiveComponent>(ChildData.Proxy))
 			{
-				MappedData.FindOrAdd(Component).Add(ChildData.BoneId, ChildData.ChildToParent);
+				FMappedComponentKey ComponentKey(Component);
+				TMap<int32, FMappedBoneData>& BoneIDToBoneData = MappedData.FindOrAdd(ComponentKey);
 	
 				Chaos::FPhysicsObjectHandle Handle = Component->GetPhysicsObjectById(ChildData.BoneId);
-				ChildParticles.Add(Interface->GetRigidParticle(Handle));
+				
+				Chaos::FPBDRigidParticle* Particle = Interface->GetRigidParticle(Handle);
+				ChildParticles.Add(Particle);
+
+				FMappedBoneData MappedBoneData(Handle, Particle, Particle ? Particle->UniqueIdx() : Chaos::FUniqueIdx(), ChildData.ChildToParent);
+				BoneIDToBoneData.Add(ChildData.BoneId, MappedBoneData);
+
 				PerShapeComponentBone.Add({Component, ChildData.BoneId});
 			}
 			else
@@ -881,16 +906,18 @@ void UClusterUnionComponent::SyncClusterUnionFromProxy()
 	// We need to handle any additions, deletions, and modifications to any child in the cluster union here.
 	// If a component lives in MappedData but not in PerComponentData, new component!
 	// If a component lives in both, then it's a modified component.
-	for (const TPair<TObjectKey<UPrimitiveComponent>, TMap<int32, FTransform>>& Kvp : MappedData)
+	for (const TPair<FMappedComponentKey, TMap<int32, FMappedBoneData>>& Kvp : MappedData)
 	{
-		HandleAddOrModifiedClusteredComponent(Kvp.Key.ResolveObjectPtr(), Kvp.Value);
+		HandleAddOrModifiedClusteredComponent(Kvp.Key, Kvp.Value);
 	}
 
 	// If a component lives in PerComponentData but not in MappedData, deleted component!
 	TArray<TObjectKey<UPrimitiveComponent>> ComponentsToRemove;
 	for (const TPair<TObjectKey<UPrimitiveComponent>, FClusteredComponentData>& Kvp : PerComponentData)
 	{
-		if (!MappedData.Contains(Kvp.Key))
+		// Build a mapped key with a null component ptr as we don't have it at this point and we don't need either
+		FMappedComponentKey MappedComponentKey(Kvp.Key, nullptr);
+		if (!MappedData.Contains(MappedComponentKey))
 		{
 			ComponentsToRemove.Add(Kvp.Key);
 		}
@@ -919,41 +946,36 @@ void UClusterUnionComponent::SyncClusterUnionFromProxy()
 	}
 }
 
-void UClusterUnionComponent::HandleAddOrModifiedClusteredComponent(UPrimitiveComponent* ChangedComponent, const TMap<int32, FTransform>& PerBoneChildToParent)
+void UClusterUnionComponent::HandleAddOrModifiedClusteredComponent(const FMappedComponentKey& ChangedComponentData, const TMap<int32, FMappedBoneData>& PerBoneChildToParent)
 {
-	if (!ChangedComponent || !ChangedComponent->HasValidPhysicsState() || !ChangedComponent->GetWorld())
+	if (!ChangedComponentData.ComponentPtr || !ChangedComponentData.ComponentPtr->HasValidPhysicsState() || !ChangedComponentData.ComponentPtr->GetWorld())
 	{
 		return;
 	}
 
-	const bool bIsNew = !PerComponentData.Contains(ChangedComponent);
-	FClusteredComponentData& ComponentData = PerComponentData.FindOrAdd(ChangedComponent);
-
-	// A component shouldn't be in the PendingComponentSync map unless it's new.
-	FClusterUnionPendingAddData PendingData;
-	if (bIsNew && PendingComponentSync.RemoveAndCopyValue(ChangedComponent, PendingData))
-	{
-		ComponentData.CachedAccelerationPayloads.Append(PendingData.AccelerationPayloads);
-	}
+	const bool bIsNew = !PerComponentData.Contains(ChangedComponentData.ComponentKey);
+	FClusteredComponentData& ComponentData = PerComponentData.FindOrAdd(ChangedComponentData.ComponentKey);
 
 	// If this is a *new* component that we're keeping track of then there's additional book-keeping
 	// we need to do to make sure we don't forget what exactly we're tracking. Additionally, we need to
 	// modify the component and its parent actor to ensure their replication stops.
 	if (bIsNew)
 	{
+		PendingComponentSync.Remove(ChangedComponentData.ComponentKey);
+
 		// Force the component and its parent actor to stop replicating movement.
 		// Setting the component to not replicate should be sufficient since a simulating
 		// component shouldn't be doing much more than replicating its position anyway.
-		if (AActor* Owner = ChangedComponent->GetOwner())
+		if (AActor* Owner = ChangedComponentData.ComponentPtr->GetOwner())
 		{
 			if (FClusteredActorData* Data = ActorToComponents.Find(Owner))
 			{
-				Data->Components.Add(ChangedComponent);
+				Data->Components.Add(ChangedComponentData.ComponentKey);
 			}
 			else
 			{
 				FClusteredActorData NewData;
-				NewData.Components.Add(ChangedComponent);
+				NewData.Components.Add(ChangedComponentData.ComponentKey);
 				NewData.bWasReplicatingMovement = Owner->IsReplicatingMovement();
 				ActorToComponents.Add(Owner, NewData);
 
@@ -966,10 +988,10 @@ void UClusterUnionComponent::HandleAddOrModifiedClusteredComponent(UPrimitiveCom
 			ComponentData.Owner = Owner;
 		}
 
-		ComponentData.bWasReplicating = ChangedComponent->GetIsReplicated();
+		ComponentData.bWasReplicating = ChangedComponentData.ComponentPtr->GetIsReplicated();
 		if (IsAuthority())
 		{
-			if (AActor* Owner = ChangedComponent->GetOwner())
+			if (AActor* Owner = ChangedComponentData.ComponentPtr->GetOwner())
 			{
 				// Create a replicated proxy component and add it to the actor being added to the cluster.
 				// This component will take care of replicating this addition into the cluster.
@@ -978,7 +1000,7 @@ void UClusterUnionComponent::HandleAddOrModifiedClusteredComponent(UPrimitiveCom
 				{
 					ReplicatedProxy->RegisterComponent();
 					ReplicatedProxy->SetParentClusterUnion(this);
-					ReplicatedProxy->SetChildClusteredComponent(ChangedComponent);
+					ReplicatedProxy->SetChildClusteredComponent(ChangedComponentData.ComponentPtr);
 					ReplicatedProxy->SetIsReplicated(true);
 				}
 
@@ -988,18 +1010,53 @@ void UClusterUnionComponent::HandleAddOrModifiedClusteredComponent(UPrimitiveCom
 
 	}
 
-	const TSet<int32> OldBoneIds{ ComponentData.BoneIds };
-	PerBoneChildToParent.GetKeys(ComponentData.BoneIds);
-	BroadcastComponentAddedEvents(ChangedComponent, ComponentData.BoneIds, bIsNew);
+	TArray<int32> RemovedBoneIDs;
+	RemovedBoneIDs.Reserve(ComponentData.BonesData.Num());
+
+	// In the case where we need to keep the acceleration structure up to date, we need to make sure old bone ids are properly
+	// removed from the acceleration structure.
+	if (AccelerationStructure)
+	{
+		for (const FClusterUnionBoneData& BoneData : ComponentData.BonesData)
+		{
+			if (!PerBoneChildToParent.Contains(BoneData.ID))
+			{
+				if (Chaos::FPhysicsObjectHandle PhysicsObject = ChangedComponentData.ComponentPtr->GetPhysicsObjectById(BoneData.ID))
+				{
+					FExternalSpatialAccelerationPayload Handle;
+					Handle.Initialize(ChangedComponentData.ComponentKey, BoneData.ID, BoneData.ParticleID);
+					if (ensure(Handle.IsValid()))
+					{
+						AccelerationStructure->RemoveElement(Handle);
+					}
+				}
+				RemovedBoneIDs.Emplace(BoneData.ID);
+			}
+		}
+	}
+
+	ComponentData.BonesData.Reset();
+	ComponentData.BonesData.Reserve(PerBoneChildToParent.Num());
+
+	Algo::Transform(PerBoneChildToParent, ComponentData.BonesData, [](const TPair<int32, FMappedBoneData>& MappedData)
+	{
+		return FClusterUnionBoneData{MappedData.Key, MappedData.Value.ParticleID };
+	});
+
+	BroadcastComponentAddedEvents(ChangedComponentData.ComponentPtr, ComponentData.BonesData, bIsNew, MoveTemp(RemovedBoneIDs));
 
 	if (IsAuthority() && ComponentData.ReplicatedProxyComponent.IsValid())
 	{
 		// We really only need to do modifications on the server since that's where we're changing the replicated proxy to broadcast this data change.
 		TObjectPtr<UClusterUnionReplicatedProxyComponent> ReplicatedProxy = ComponentData.ReplicatedProxyComponent.Get();
-		ReplicatedProxy->SetParticleBoneIds(ComponentData.BoneIds.Array());
-		for (const TPair<int32, FTransform>& Kvp : PerBoneChildToParent)
+		
+		TArray<int32> AsBoneIdsArray;
+		GetBoneIDsFromComponentData(AsBoneIdsArray, ComponentData);
+		ReplicatedProxy->SetParticleBoneIds(AsBoneIdsArray);
+
+		for (const TPair<int32, FMappedBoneData>& Kvp : PerBoneChildToParent)
 		{
-			ReplicatedProxy->SetParticleChildToParent(Kvp.Key, Kvp.Value);
+			ReplicatedProxy->SetParticleChildToParent(Kvp.Key, Kvp.Value.ChildToParentTransform);
 		}
 	}
 
@@ -1010,45 +1067,23 @@ void UClusterUnionComponent::HandleAddOrModifiedClusteredComponent(UPrimitiveCom
 
 		// One more loop to ensure that our sets of physics objects are valid and up to date.
 		// This needs to happen on both the client and the server.
-		for (const TPair<int32, FTransform>& Kvp : PerBoneChildToParent)
+		for (const TPair<int32, FMappedBoneData>& Kvp : PerBoneChildToParent)
 		{
-			Chaos::FPhysicsObjectHandle PhysicsObject = ChangedComponent->GetPhysicsObjectById(Kvp.Key);
-
 			if (AccelerationStructure)
 			{
-				if (PhysicsObject)
+				if (Chaos::FPhysicsObjectHandle PhysicsObject = Kvp.Value.PhysicsObjectHandle)
 				{
 					const FBox HandleBounds = Interface->GetWorldBounds({ &PhysicsObject, 1 });
 					FExternalSpatialAccelerationPayload Handle;
-					Handle.Initialize(ChangedComponent, Kvp.Key);
-					if (Handle.IsValid())
+					Handle.Initialize(ChangedComponentData.ComponentKey, Kvp.Key, Kvp.Value.ParticleID);
+					if (ensure(Handle.IsValid()))
 					{
 						AccelerationStructure->UpdateElement(Handle, Chaos::TAABB<Chaos::FReal, 3>{HandleBounds.Min, HandleBounds.Max}, HandleBounds.IsValid != 0);
-						ComponentData.CachedAccelerationPayloads.Add(Handle);
 					}
 				}
 			}
 		}
 	}
-
-	// In the case where we need to keep the acceleration structure up to date, we need to make sure old bone ids are properly
-	// removed from the acceleration structure.
-	if (AccelerationStructure)
-	{
-		for (int32 BoneId : OldBoneIds.Difference(ComponentData.BoneIds))
-		{
-			if (Chaos::FPhysicsObjectHandle PhysicsObject = ChangedComponent->GetPhysicsObjectById(BoneId))
-			{
-				FExternalSpatialAccelerationPayload Handle;
-				Handle.Initialize(ChangedComponent, BoneId);
-				if (Handle.IsValid())
-				{
-					AccelerationStructure->RemoveElement(Handle);
-					ComponentData.CachedAccelerationPayloads.Remove(Handle);
-				}
-			}
-		}
-	}
 }
 
 void UClusterUnionComponent::HandleRemovedClusteredComponent(TObjectKey<UPrimitiveComponent> RemovedComponent)
@@ -1057,9 +1092,12 @@ void UClusterUnionComponent::HandleRemovedClusteredComponent(TObjectKey<UPrimiti
 	{
 		if (AccelerationStructure)
 		{
-			for (const FExternalSpatialAccelerationPayload& Handle : Data->CachedAccelerationPayloads)
+			for (const FClusterUnionBoneData& BoneData : Data->BonesData)
 			{
-				if (Handle.IsValid())
+				FExternalSpatialAccelerationPayload Handle;
+				Handle.Initialize(RemovedComponent, BoneData.ID, BoneData.ParticleID);
+
+				if (ensure(Handle.IsValid()))
 				{
 					AccelerationStructure->RemoveElement(Handle);
 				}
@@ -1101,15 +1139,18 @@ void UClusterUnionComponent::HandleRemovedClusteredComponent(TObjectKey<UPrimiti
 	PendingComponentSync.Remove(RemovedComponent);
 }
 
-void UClusterUnionComponent::BroadcastComponentAddedEvents(UPrimitiveComponent* ChangedComponent, const TSet<int32>& BoneIds, bool bIsNew)
+void UClusterUnionComponent::BroadcastComponentAddedEvents(UPrimitiveComponent* ChangedComponent, const TSet<FClusterUnionBoneData>& BoneIds, bool bIsNew, TArray<int32>&& RemovedBoneIDs)
 {
+	// TODO: Add a new delegate that takes a TArray and deprecate this one
 	if (OnComponentAddedEvent.IsBound())
 	{
-		OnComponentAddedEvent.Broadcast(ChangedComponent, BoneIds, bIsNew);
+		TSet<int32> BoneIDsSet;
+		Algo::Transform(BoneIds, BoneIDsSet, &FClusterUnionBoneData::ID);
+		OnComponentAddedEvent.Broadcast(ChangedComponent, BoneIDsSet, bIsNew);
 	}
 	if (OnComponentAddedNativeEvent.IsBound())
 	{
-		OnComponentAddedNativeEvent.Broadcast(ChangedComponent, BoneIds, bIsNew);
+		OnComponentAddedNativeEvent.Broadcast(ChangedComponent, BoneIds, MoveTemp(RemovedBoneIDs), bIsNew);
 	}
 }
 
@@ -1543,12 +1584,18 @@ TArray<int32> UClusterUnionComponent::GetAddedBoneIdsForComponent(UPrimitiveComp
 {
 	if (const FClusteredComponentData* Data = PerComponentData.Find(Component))
 	{
-		return Data->BoneIds.Array();
+		TArray<int32> BoneIDsArray;
+		GetBoneIDsFromComponentData(BoneIDsArray, *Data);
+
+		return MoveTemp(BoneIDsArray);
 	}
 
 	if (const FClusterUnionPendingAddData* Data = PendingComponentSync.Find(Component))
 	{
-		return Data->BoneIds;
+		TArray<int32> BoneIDsArray;
+		GetBoneIDsFromComponentData(BoneIDsArray, *Data);
+
+		return MoveTemp(BoneIDsArray);
 	}
 
 	return {};
diff --git a/Engine/Source/Runtime/Engine/Private/PhysicsEngine/ExternalSpatialAccelerationPayload.cpp b/Engine/Source/Runtime/Engine/Private/PhysicsEngine/ExternalSpatialAccelerationPayload.cpp
--- a/Engine/Source/Runtime/Engine/Private/PhysicsEngine/ExternalSpatialAccelerationPayload.cpp
+++ b/Engine/Source/Runtime/Engine/Private/PhysicsEngine/ExternalSpatialAccelerationPayload.cpp
@@ -13,7 +13,7 @@ FExternalSpatialAccelerationPayload::FExternalSpatialAccelerationPayload()
 {
 }
 
-void FExternalSpatialAccelerationPayload::Initialize(UPrimitiveComponent* InComponent, int32 InBoneId)
+void FExternalSpatialAccelerationPayload::Initialize(TObjectKey<UPrimitiveComponent> InComponent, int32 InBoneId)
 {
 	Component = InComponent;
 	BoneId = InBoneId;
@@ -24,9 +24,18 @@ void FExternalSpatialAccelerationPayload::Initialize(UPrimitiveComponent* InComp
 	}
 }
 
+
+void FExternalSpatialAccelerationPayload::Initialize(TObjectKey<UPrimitiveComponent> InComponent, int32 InBoneId, const Chaos::FUniqueIdx& UniqueIdx)
+{
+	Component = InComponent;
+	BoneId = InBoneId;
+	CachedUniqueIdx = UniqueIdx;
+}
+
+
 Chaos::FGeometryParticle* FExternalSpatialAccelerationPayload::GetExternalGeometryParticle_ExternalThread() const
 {
-	if (UPrimitiveComponent* ValidComponent = Component.Get())
+	if (UPrimitiveComponent* ValidComponent = Component.ResolveObjectPtr())
 	{
 		if (!ValidComponent->HasValidPhysicsState())
 		{
@@ -66,7 +75,7 @@ void FExternalSpatialAccelerationPayload::DebugDraw(const bool bExternal, const
 
 uint32 GetTypeHash(const FExternalSpatialAccelerationPayload& Payload)
 {
-	const uint32 ComponentHash = Payload.Component.GetWeakPtrTypeHash();
+	const uint32 ComponentHash = GetTypeHash(Payload.Component);
 	const uint32 BoneIdHash = GetTypeHash(Payload.BoneId);
 	return HashCombine(ComponentHash, BoneIdHash);
 }
\ No newline at end of file
diff --git a/Engine/Source/Runtime/Experimental/GeometryCollectionEngine/Private/GeometryCollection/GeometryCollectionComponent.cpp b/Engine/Source/Runtime/Experimental/GeometryCollectionEngine/Private/GeometryCollection/GeometryCollectionComponent.cpp
--- a/Engine/Source/Runtime/Experimental/GeometryCollectionEngine/Private/GeometryCollection/GeometryCollectionComponent.cpp
+++ b/Engine/Source/Runtime/Experimental/GeometryCollectionEngine/Private/GeometryCollection/GeometryCollectionComponent.cpp
@@ -4385,6 +4385,46 @@ void UGeometryCollectionComponent::SetPerParticleCollisionProfileName(const TSet
 	}
 }
 
+bool UGeometryCollectionComponent::UpdatePerParticleCollisionProfilesNum()
+{
+	if (!RestCollection)
+	{
+		return false;
+	}
+
+	const int32 NumTransforms = RestCollection->GetGeometryCollection()->Transform.Num();
+	const bool bHasChanged = CollisionProfilePerParticle.Num() != NumTransforms;
+
+	if (bHasChanged)
+	{
+		CollisionProfilePerParticle.SetNumZeroed(NumTransforms);
+	}
+
+	return bHasChanged;
+}
+
+void UGeometryCollectionComponent::SetParticleCollisionProfileName(int32 BoneId, FName ProfileName, FGCCollisionProfileScopedTransaction& InProfileNameUpdateTransaction)
+{
+	if (!InProfileNameUpdateTransaction.IsValid())
+	{
+		return;
+	}
+	
+	if (!CollisionProfilePerParticle.IsValidIndex(BoneId))
+	{
+		return;
+	}
+
+	FName& CurrentProfileName = CollisionProfilePerParticle[BoneId];
+
+	if (CurrentProfileName !=  ProfileName)
+	{
+		CurrentProfileName = ProfileName;
+
+		InProfileNameUpdateTransaction.MarkDirty();
+	}
+}
+
 void UGeometryCollectionComponent::SetPerParticleCollisionProfileName(const TArray<int32>& BoneIds, FName ProfileName)
 {
 	const bool bHasChanged = SetPerParticleCollisionProfileNameFromIterable(RestCollection, BoneIds, ProfileName, CollisionProfilePerParticle);
diff --git a/Engine/Source/Runtime/Experimental/GeometryCollectionEngine/Public/GeometryCollection/GeometryCollectionComponent.h b/Engine/Source/Runtime/Experimental/GeometryCollectionEngine/Public/GeometryCollection/GeometryCollectionComponent.h
--- a/Engine/Source/Runtime/Experimental/GeometryCollectionEngine/Public/GeometryCollection/GeometryCollectionComponent.h
+++ b/Engine/Source/Runtime/Experimental/GeometryCollectionEngine/Public/GeometryCollection/GeometryCollectionComponent.h
@@ -18,7 +18,9 @@
 #include "Physics/Experimental/PhysScene_Chaos.h"
 #endif
 #include "GeometryCollectionEditorSelection.h"
+#include "GeometryCollection/GeometryCollection.h"
 #include "GeometryCollection/GeometryCollectionDamagePropagationData.h"
+#include "GeometryCollectionObject.h"
 #include "GeometryCollection/RecordedTransformTrack.h"
 #include "GeometryCollection/GeometryCollectionSimulationTypes.h"
 #include "Templates/UniquePtr.h"
@@ -560,6 +562,8 @@ struct TStructOpsTypeTraits<FGeometryCollectionRepDynamicData> : public TStructO
 	};
 };
 
+struct FGCCollisionProfileScopedTransaction;
+
 /**
 *	GeometryCollectionComponent
 */
@@ -853,6 +857,14 @@ public:
 
 	GEOMETRYCOLLECTIONENGINE_API void SetPerParticleCollisionProfileName(const TSet<int32>& BoneIds, FName ProfileName);
 
+	GEOMETRYCOLLECTIONENGINE_API void SetParticleCollisionProfileName(int32 BoneId, FName ProfileName, FGCCollisionProfileScopedTransaction& InProfileNameUpdateTransaction);
+
+private:
+
+	bool UpdatePerParticleCollisionProfilesNum();
+
+public:
+
 	/** API for getting at geometry collection data */
 	GEOMETRYCOLLECTIONENGINE_API int32 GetNumElements(FName Group) const;
 
@@ -1798,4 +1810,52 @@ public:
 	GEOMETRYCOLLECTIONENGINE_API virtual bool IsHLODRelevant() const override;
 	//~ End UActorComponent interface.
 #endif
+
+	friend struct FGCCollisionProfileScopedTransaction;
+};
+
+/** Struct to be used as Transaction object used to make updates on particle per particle basis within a scope.
+ * It makes sure the collision profile names containers is up to date and the Collision profiles are loaded if needed when it goes out of scope
+ */
+struct GEOMETRYCOLLECTIONENGINE_API FGCCollisionProfileScopedTransaction
+{
+	explicit FGCCollisionProfileScopedTransaction(UGeometryCollectionComponent* InGCComponentInstance) : GCComponentInstance(InGCComponentInstance)
+	{
+		if (!ensure(GCComponentInstance))
+		{
+			return;
+		}
+
+		bHasChanged = InGCComponentInstance->UpdatePerParticleCollisionProfilesNum();	
+	}
+
+	~FGCCollisionProfileScopedTransaction()
+	{
+		if (!GCComponentInstance)
+		{
+			return;
+		}
+
+		if (bHasChanged)
+		{
+			GCComponentInstance->LoadCollisionProfiles();
+		}
+	}
+
+	FGCCollisionProfileScopedTransaction& operator=(const FGCCollisionProfileScopedTransaction& Other) = delete;
+	FGCCollisionProfileScopedTransaction& operator=(FGCCollisionProfileScopedTransaction&& Other) = delete;
+	FGCCollisionProfileScopedTransaction(FGCCollisionProfileScopedTransaction&& Other) = delete;
+	FGCCollisionProfileScopedTransaction(FGCCollisionProfileScopedTransaction& Other) = delete;
+
+	/** Marks this transaction dirty. It will load the collision profiles if needed when this transaction goes out of scope */
+	void MarkDirty()
+	{
+		bHasChanged = true;
+	}
+
+	bool IsValid() const { return GCComponentInstance != nullptr; }
+
+private:
+	UGeometryCollectionComponent* GCComponentInstance = nullptr;
+	bool bHasChanged = false;
 };
