World Partition: deprecated streaming sources target HLOD layers, we should be using target grids instead.
Since we support merging HLOD layers that shares the same settings, this filtering couldn't be honored properly.

diff --git a/Engine/Source/Runtime/Engine/Classes/Components/WorldPartitionStreamingSourceComponent.h b/Engine/Source/Runtime/Engine/Classes/Components/WorldPartitionStreamingSourceComponent.h
--- a/Engine/Source/Runtime/Engine/Classes/Components/WorldPartitionStreamingSourceComponent.h
+++ b/Engine/Source/Runtime/Engine/Classes/Components/WorldPartitionStreamingSourceComponent.h
@@ -68,8 +68,8 @@ class UWorldPartitionStreamingSourceComponent : public UActorComponent, public I
 	FColor DebugColor;
 
 	/** Optional target HLODLayers affected by the streaming source. */
-	UPROPERTY(BlueprintReadWrite, EditAnywhere, Category = "Streaming")
-	TArray<TObjectPtr<const UHLODLayer>> TargetHLODLayers;
+	UPROPERTY(meta = (DeprecatedProperty, DeprecationMessage = "Use TargetGrids instead."))
+	TArray<TObjectPtr<const UHLODLayer>> TargetHLODLayers_DEPRECATED;
 
 	UPROPERTY(meta = (DeprecatedProperty, DeprecationMessage = "Use TargetHLODLayers instead."))
 	TObjectPtr<const UHLODLayer> TargetHLODLayer_DEPRECATED;
diff --git a/Engine/Source/Runtime/Engine/Private/Components/WorldPartitionStreamingSourceComponent.cpp b/Engine/Source/Runtime/Engine/Private/Components/WorldPartitionStreamingSourceComponent.cpp
--- a/Engine/Source/Runtime/Engine/Private/Components/WorldPartitionStreamingSourceComponent.cpp
+++ b/Engine/Source/Runtime/Engine/Private/Components/WorldPartitionStreamingSourceComponent.cpp
@@ -75,11 +75,6 @@ void UWorldPartitionStreamingSourceComponent::PostLoad()
 		{
 			TargetGrids.Add(TargetGrid_DEPRECATED);
 		}
-
-		if (TargetHLODLayer_DEPRECATED)
-		{
-			TargetHLODLayers.Add(TargetHLODLayer_DEPRECATED);
-		}
 	}
 }
 
@@ -106,15 +101,6 @@ bool UWorldPartitionStreamingSourceComponent::GetStreamingSource(FWorldPartition
 		{
 			OutStreamingSource.TargetGrids.Append(TargetGrids);
 		}
-
-		if (!TargetHLODLayers.IsEmpty())
-		{
-			OutStreamingSource.TargetHLODLayers.Reserve(TargetHLODLayers.Num());
-			for (const TObjectPtr<const UHLODLayer>& HLODLayer : TargetHLODLayers)
-			{
-				OutStreamingSource.TargetHLODLayers.Add(HLODLayer);
-			}
-		}
 		return true;
 	}
 	return false;
diff --git a/Engine/Source/Runtime/Engine/Private/WorldPartition/RuntimeHashSet/WorldPartitionRuntimeHashSet.cpp b/Engine/Source/Runtime/Engine/Private/WorldPartition/RuntimeHashSet/WorldPartitionRuntimeHashSet.cpp
--- a/Engine/Source/Runtime/Engine/Private/WorldPartition/RuntimeHashSet/WorldPartitionRuntimeHashSet.cpp
+++ b/Engine/Source/Runtime/Engine/Private/WorldPartition/RuntimeHashSet/WorldPartitionRuntimeHashSet.cpp
@@ -363,9 +363,7 @@ void UWorldPartitionRuntimeHashSet::ForEachStreamingCellsSources(const TArray<FW
 		{
 			for (const FWorldPartitionStreamingSource& Source : Sources)
 			{
-				// @todo_jfd
-				const FSoftObjectPath HLODLayer;
-				Source.ForEachShape(InLoadingRange, InGridName, HLODLayer, false, [this, &Source, InSpatialIndex, &Func, &ActivateStreamingSourceCells, &LoadStreamingSourceCells](const FSphericalSector& Shape)
+				Source.ForEachShape(InLoadingRange, InGridName, false, [this, &Source, InSpatialIndex, &Func, &ActivateStreamingSourceCells, &LoadStreamingSourceCells](const FSphericalSector& Shape)
 				{
 					const FSphere ShapeSphere(Shape.GetCenter(), Shape.GetRadius());
 
diff --git a/Engine/Source/Runtime/Engine/Private/WorldPartition/RuntimeHashSet/WorldPartitionRuntimeHashSetDebugDraw.cpp b/Engine/Source/Runtime/Engine/Private/WorldPartition/RuntimeHashSet/WorldPartitionRuntimeHashSetDebugDraw.cpp
--- a/Engine/Source/Runtime/Engine/Private/WorldPartition/RuntimeHashSet/WorldPartitionRuntimeHashSetDebugDraw.cpp
+++ b/Engine/Source/Runtime/Engine/Private/WorldPartition/RuntimeHashSet/WorldPartitionRuntimeHashSetDebugDraw.cpp
@@ -207,9 +207,7 @@ bool UWorldPartitionRuntimeHashSet::Draw2D(FWorldPartitionDraw2DContext& DrawCon
 		for (const FWorldPartitionStreamingSource& Source : Sources)
 		{
 			const FColor Color = Source.GetDebugColor();
-			// @todo_jfd
-			const FSoftObjectPath HLODLayer;
-			Source.ForEachShape(StreamingData->LoadingRange, StreamingData->Name, HLODLayer, true, [&Color, &WorldToScreen, &GridScreenBounds, &DrawContext, this](const FSphericalSector& Shape)
+			Source.ForEachShape(StreamingData->LoadingRange, StreamingData->Name, true, [&Color, &WorldToScreen, &GridScreenBounds, &DrawContext, this](const FSphericalSector& Shape)
 			{
 				check(!Shape.IsNearlyZero())
 
@@ -250,9 +248,7 @@ bool UWorldPartitionRuntimeHashSet::Draw2D(FWorldPartitionDraw2DContext& DrawCon
 	{
 		for (const FWorldPartitionStreamingSource& Source : Sources)
 		{
-			// @todo_jfd
-			const FSoftObjectPath HLODLayer;
-			Source.ForEachShape(StreamingDataList[0]->LoadingRange, Name, HLODLayer, true, [&GridsShapeBounds](const FSphericalSector& Shape) { GridsShapeBounds += Shape.CalcBounds(); });
+			Source.ForEachShape(StreamingDataList[0]->LoadingRange, Name, true, [&GridsShapeBounds](const FSphericalSector& Shape) { GridsShapeBounds += Shape.CalcBounds(); });
 		}
 
 		const FVector2D GridReferenceWorldPos = FVector2D(WorldRegion.GetCenter());
@@ -321,9 +317,7 @@ void UWorldPartitionRuntimeHashSet::Draw3D(const TArray<FWorldPartitionStreaming
 	{
 		for (const FWorldPartitionStreamingSource& Source : Sources)
 		{
-			// @todo_jfd
-			const FSoftObjectPath HLODLayer;
-			Source.ForEachShape(StreamingData.LoadingRange, StreamingData.Name, HLODLayer, false, [this, &StreamingData, VisualizeMode, &DataLayerDebugColors, ContentBundleManager, OwningWorld, &WorldPartitionTransform](const FSphericalSector& Shape)
+			Source.ForEachShape(StreamingData.LoadingRange, StreamingData.Name, false, [this, &StreamingData, VisualizeMode, &DataLayerDebugColors, ContentBundleManager, OwningWorld, &WorldPartitionTransform](const FSphericalSector& Shape)
 			{
 				const FSphere ShapeSphere(Shape.GetCenter(), Shape.GetRadius());
 
@@ -375,7 +369,7 @@ void UWorldPartitionRuntimeHashSet::Draw3D(const TArray<FWorldPartitionStreaming
 					const FVector BoxCenter(Box.GetCenter());
 					const FColor BoxColor(CellColors[0].CopyWithNewOpacity(CellOpacity).ToFColor(true));
 					const FVector CellPos = WorldPartitionTransform.TransformPosition(BoxCenter);
-					DrawDebugBox(OwningWorld, CellPos, Box.GetExtent(), WorldPartitionTransform.GetRotation(), BoxColor.WithAlpha(255), false, -1.f, 255, 100.f);
+					DrawDebugBox(OwningWorld, CellPos, Box.GetExtent(), WorldPartitionTransform.GetRotation(), BoxColor.WithAlpha(255), false, -1.f, 255, 20.f);
 				});
 			});
 		}
diff --git a/Engine/Source/Runtime/Engine/Private/WorldPartition/WorldPartitionRuntimeSpatialHash.cpp b/Engine/Source/Runtime/Engine/Private/WorldPartition/WorldPartitionRuntimeSpatialHash.cpp
--- a/Engine/Source/Runtime/Engine/Private/WorldPartition/WorldPartitionRuntimeSpatialHash.cpp
+++ b/Engine/Source/Runtime/Engine/Private/WorldPartition/WorldPartitionRuntimeSpatialHash.cpp
@@ -224,7 +224,7 @@ void FSpatialHashStreamingGrid::GetCells(const FWorldPartitionStreamingQuerySour
 	// Spatial Query
 	if (QuerySource.bSpatialQuery)
 	{
-		QuerySource.ForEachShape(GetLoadingRange(), GridName, HLODLayer, /*bProjectIn2D*/ true, [&](const FSphericalSector& Shape)
+		QuerySource.ForEachShape(GetLoadingRange(), GridName, /*bProjectIn2D*/ true, [&](const FSphericalSector& Shape)
 		{
 			Helper.ForEachIntersectingCells(Shape, [&](const FGridCellCoord& Coords)
 			{
@@ -298,7 +298,7 @@ void FSpatialHashStreamingGrid::GetCells(const TArray<FWorldPartitionStreamingSo
 	const FSquare2DGridHelper& Helper = GetGridHelper();
 	for (const FWorldPartitionStreamingSource& Source : Sources)
 	{
-		Source.ForEachShape(GridLoadingRange, GridName, HLODLayer, /*bProjectIn2D*/ true, [&](const FSphericalSector& Shape)
+		Source.ForEachShape(GridLoadingRange, GridName, /*bProjectIn2D*/ true, [&](const FSphericalSector& Shape)
 		{
 			FStreamingSourceInfo Info(Source, Shape);
 
@@ -752,7 +752,7 @@ void FSpatialHashStreamingGrid::Draw3D(const UWorldPartitionRuntimeSpatialHash*
 
 	for (const FWorldPartitionStreamingSource& Source : Sources)
 	{
-		FBox Region = Source.CalcBounds(GridLoadingRange, GridName, HLODLayer);
+		FBox Region = Source.CalcBounds(GridLoadingRange, GridName);
 		Region += FBox(Region.GetCenter() - MinExtent, Region.GetCenter() + MinExtent);
 
 		for (int32 GridLevel = MinGridLevel; GridLevel <= MaxGridLevel; ++GridLevel)
@@ -831,7 +831,7 @@ void FSpatialHashStreamingGrid::Draw3D(const UWorldPartitionRuntimeSpatialHash*
 		}
 
 		const FColor Color = Source.GetDebugColor();
-		Source.ForEachShape(GetLoadingRange(), GridName, HLODLayer, /*bProjectIn2D*/ true, [&Color, &SourceLocationZ, &Transform, &OwningWorld, this](const FSphericalSector& Shape)
+		Source.ForEachShape(GetLoadingRange(), GridName, /*bProjectIn2D*/ true, [&Color, &SourceLocationZ, &Transform, &OwningWorld, this](const FSphericalSector& Shape)
 		{
 			FSphericalSector ZOffsettedShape = Shape;
 			ZOffsettedShape.SetCenter(FVector(FVector2D(ZOffsettedShape.GetCenter()), SourceLocationZ));
@@ -984,7 +984,7 @@ void FSpatialHashStreamingGrid::Draw2D(const UWorldPartitionRuntimeSpatialHash*
 	for (const FWorldPartitionStreamingSource& Source : Sources)
 	{
 		const FColor Color = Source.GetDebugColor();
-		Source.ForEachShape(GridLoadingRange, GridName, HLODLayer, /*bProjectIn2D*/ true, [&Color, &WorldToScreen, &GridScreenBounds, &DrawContext, this](const FSphericalSector& Shape)
+		Source.ForEachShape(GridLoadingRange, GridName, /*bProjectIn2D*/ true, [&Color, &WorldToScreen, &GridScreenBounds, &DrawContext, this](const FSphericalSector& Shape)
 		{
 			DrawStreamingSource2D(GridScreenBounds, Shape, WorldToScreen, Color, DrawContext);
 		});
@@ -1875,7 +1875,7 @@ bool UWorldPartitionRuntimeSpatialHash::Draw2D(FWorldPartitionDraw2DContext& Dra
 	{
 		for (const FWorldPartitionStreamingSource& Source : Sources)
 		{
-			Source.ForEachShape(StreamingGrid->GetLoadingRange(), StreamingGrid->GridName, StreamingGrid->HLODLayer, /*bProjectIn2D*/ true, [&GridsShapeBounds](const FSphericalSector& Shape) { GridsShapeBounds += Shape.CalcBounds(); });
+			Source.ForEachShape(StreamingGrid->GetLoadingRange(), StreamingGrid->GridName, /*bProjectIn2D*/ true, [&GridsShapeBounds](const FSphericalSector& Shape) { GridsShapeBounds += Shape.CalcBounds(); });
 		}
 
 		FVector2D GridReferenceWorldPos;
diff --git a/Engine/Source/Runtime/Engine/Private/WorldPartition/WorldPartitionStreamingPolicy.cpp b/Engine/Source/Runtime/Engine/Private/WorldPartition/WorldPartitionStreamingPolicy.cpp
--- a/Engine/Source/Runtime/Engine/Private/WorldPartition/WorldPartitionStreamingPolicy.cpp
+++ b/Engine/Source/Runtime/Engine/Private/WorldPartition/WorldPartitionStreamingPolicy.cpp
@@ -879,7 +879,6 @@ bool UWorldPartitionStreamingPolicy::IsStreamingCompleted(const TArray<FWorldPar
 		QuerySource.Rotation = StreamingSource.Rotation;
 		QuerySource.TargetBehavior = StreamingSource.TargetBehavior;
 		QuerySource.TargetGrids = StreamingSource.TargetGrids;
-		QuerySource.TargetHLODLayers = StreamingSource.TargetHLODLayers;
 		QuerySource.Shapes = StreamingSource.Shapes;
 		QuerySource.bUseGridLoadingRange = true;
 		QuerySource.Radius = 0.f;
diff --git a/Engine/Source/Runtime/Engine/Private/WorldPartition/WorldPartitionStreamingSource.cpp b/Engine/Source/Runtime/Engine/Private/WorldPartition/WorldPartitionStreamingSource.cpp
--- a/Engine/Source/Runtime/Engine/Private/WorldPartition/WorldPartitionStreamingSource.cpp
+++ b/Engine/Source/Runtime/Engine/Private/WorldPartition/WorldPartitionStreamingSource.cpp
@@ -109,7 +109,7 @@ void FWorldPartitionStreamingSource::UpdateHash()
 	}
 
 	FHashBuilder HashBuilder;
-	HashBuilder	<< Name << TargetState << bBlockOnSlowLoading << bReplay << bRemote << Priority << TargetBehavior << TargetGrids << TargetHLODLayers << Shapes << ExtraRadius  << ExtraAngle;
+	HashBuilder	<< Name << TargetState << bBlockOnSlowLoading << bReplay << bRemote << Priority << TargetBehavior << TargetGrids << Shapes << ExtraRadius  << ExtraAngle;
 
 	if (LocationQuantization)
 	{
diff --git a/Engine/Source/Runtime/Engine/Public/WorldPartition/WorldPartitionStreamingSource.h b/Engine/Source/Runtime/Engine/Public/WorldPartition/WorldPartitionStreamingSource.h
--- a/Engine/Source/Runtime/Engine/Public/WorldPartition/WorldPartitionStreamingSource.h
+++ b/Engine/Source/Runtime/Engine/Public/WorldPartition/WorldPartitionStreamingSource.h
@@ -153,15 +153,14 @@ enum class EStreamingSourceTargetBehavior : uint8
 class FStreamingSourceShapeHelper
 {
 public:
-
-	FORCEINLINE static bool IsSourceAffectingGrid(const TSet<FName>& InSourceTargetGrids, const TSet<FSoftObjectPath>& InSourceTargetHLODLayers, EStreamingSourceTargetBehavior InSourceTargetBehavior, FName InGridName, const FSoftObjectPath& InGridHLODLayer)
+	FORCEINLINE static bool IsSourceAffectingGrid(const TSet<FName>& InSourceTargetGrids, EStreamingSourceTargetBehavior InSourceTargetBehavior, FName InGridName)
 	{
 		switch (InSourceTargetBehavior)
 		{
 		case EStreamingSourceTargetBehavior::Include:
-			return (InSourceTargetGrids.IsEmpty() && InSourceTargetHLODLayers.IsEmpty()) || (InSourceTargetGrids.Contains(InGridName) || InSourceTargetHLODLayers.Contains(InGridHLODLayer));
+			return InSourceTargetGrids.IsEmpty() || InSourceTargetGrids.Contains(InGridName);
 		case EStreamingSourceTargetBehavior::Exclude:
-			return !(InSourceTargetGrids.Contains(InGridName) || InSourceTargetHLODLayers.Contains(InGridHLODLayer));
+			return !InSourceTargetGrids.Contains(InGridName);
 		default:
 			checkNoEntry();
 		}
@@ -264,7 +263,6 @@ struct FWorldPartitionStreamingQuerySource
 		Rotation = Other.Rotation;
 		TargetBehavior = Other.TargetBehavior;
 		TargetGrids = Other.TargetGrids;
-		TargetHLODLayers = Other.TargetHLODLayers;
 		Shapes = Other.Shapes;
 		return *this;
 	}
@@ -296,24 +294,20 @@ struct FWorldPartitionStreamingQuerySource
 	/* Reserved settings used by UWorldPartitionStreamingSourceComponent::IsStreamingCompleted. */
 	FRotator Rotation;
 	
-	/** Defines how TargetGrids/TargetHLODLayers will be applied to this streaming source. */
+	/** Defines how TargetGrids will be applied to this streaming source. */
 	EStreamingSourceTargetBehavior TargetBehavior;
 
 	TSet<FName> TargetGrids;
 
+	UE_DEPRECATED(5.4, "TargetHLODLayers is depredcated, use TargetGrids instead.")
 	TSet<FSoftObjectPath> TargetHLODLayers;
 	
 	TArray<FStreamingSourceShape> Shapes;
 
 	/** Helper method that iterates over all shapes. If none is provided, it will still pass a sphere shape using Radius or grid's loading range (see bUseGridLoadingRange). */
-	FORCEINLINE void ForEachShape(float InGridLoadingRange, FName InGridName, const FSoftObjectPath& InGridHLODLayer, bool bInProjectIn2D, TFunctionRef<void(const FSphericalSector&)> InOperation) const
+	FORCEINLINE void ForEachShape(float InGridLoadingRange, FName InGridName, bool bInProjectIn2D, TFunctionRef<void(const FSphericalSector&)> InOperation) const
 	{
-		if (!bSpatialQuery)
-		{
-			return;
-		}
-
-		if (FStreamingSourceShapeHelper::IsSourceAffectingGrid(TargetGrids, TargetHLODLayers, TargetBehavior, InGridName, InGridHLODLayer))
+		if (bSpatialQuery && FStreamingSourceShapeHelper::IsSourceAffectingGrid(TargetGrids, TargetBehavior, InGridName))
 		{
 			FStreamingSourceShapeHelper::ForEachShape(InGridLoadingRange, bUseGridLoadingRange ? InGridLoadingRange : Radius, bInProjectIn2D, Location, Rotation, Shapes, InOperation);
 		}
@@ -397,7 +391,6 @@ struct FWorldPartitionStreamingSource
 		DebugColor = Other.DebugColor;
 		TargetBehavior = Other.TargetBehavior;
 		TargetGrids = Other.TargetGrids;
-		TargetHLODLayers = Other.TargetHLODLayers;
 		Shapes = Other.Shapes;
 		bReplay = Other.bReplay;
 		bRemote = Other.bRemote;
@@ -454,6 +447,7 @@ struct FWorldPartitionStreamingSource
 	TSet<FName> TargetGrids;
 	
 	/** When set, this will change how this streaming source is applied to the provided HLODLayers based on the TargetBehavior. */
+	UE_DEPRECATED(5.4, "TargetHLODLayers is depredcated, use TargetGrids instead.")
 	TSet<FSoftObjectPath> TargetHLODLayers;
 
 	/** Source internal shapes. When none are provided, a sphere is automatically used. It's radius is equal to grid's loading range and center equals source's location. */
@@ -466,10 +460,10 @@ struct FWorldPartitionStreamingSource
 	bool bRemote;
 
 	/** Returns a box encapsulating all shapes. */
-	FBox CalcBounds(float InGridLoadingRange, FName InGridName, const FSoftObjectPath& InGridHLODLayer, bool bCalcIn2D = false) const
+	FBox CalcBounds(float InGridLoadingRange, FName InGridName, bool bCalcIn2D = false) const
 	{
 		FBox OutBounds(ForceInit);
-		ForEachShape(InGridLoadingRange, InGridName, InGridHLODLayer, bCalcIn2D, [&OutBounds](const FSphericalSector& Sector)
+		ForEachShape(InGridLoadingRange, InGridName, bCalcIn2D, [&OutBounds](const FSphericalSector& Sector)
 		{
 			OutBounds += Sector.CalcBounds();
 		});
@@ -477,9 +471,9 @@ struct FWorldPartitionStreamingSource
 	}
 
 	/** Helper method that iterates over all shapes. If none is provided, it will still pass a sphere shape using grid's loading range. */
-	void ForEachShape(float InGridLoadingRange, FName InGridName, const FSoftObjectPath& InGridHLODLayer, bool bInProjectIn2D, TFunctionRef<void(const FSphericalSector&)> InOperation) const
+	void ForEachShape(float InGridLoadingRange, FName InGridName, bool bInProjectIn2D, TFunctionRef<void(const FSphericalSector&)> InOperation) const
 	{
-		if (FStreamingSourceShapeHelper::IsSourceAffectingGrid(TargetGrids, TargetHLODLayers, TargetBehavior, InGridName, InGridHLODLayer))
+		if (FStreamingSourceShapeHelper::IsSourceAffectingGrid(TargetGrids, TargetBehavior, InGridName))
 		{
 			FStreamingSourceShapeHelper::ForEachShape(InGridLoadingRange, InGridLoadingRange, bInProjectIn2D, Location, Rotation, Shapes, InOperation, ExtraRadius, ExtraAngle);
 		}
