Exposed UDataLayerLoadingPolicy and FWorldPartitionActorViewProxy to plugins via ENGINE_API

diff --git a/Engine/Source/Runtime/Engine/Public/WorldPartition/DataLayer/DataLayerLoadingPolicy.h b/Engine/Source/Runtime/Engine/Public/WorldPartition/DataLayer/DataLayerLoadingPolicy.h
--- a/Engine/Source/Runtime/Engine/Public/WorldPartition/DataLayer/DataLayerLoadingPolicy.h
+++ b/Engine/Source/Runtime/Engine/Public/WorldPartition/DataLayer/DataLayerLoadingPolicy.h
@@ -9,7 +9,7 @@
 class UDataLayerInstance;
 
 UCLASS()
-class UDataLayerLoadingPolicy : public UObject
+class ENGINE_API UDataLayerLoadingPolicy : public UObject
 {
 	GENERATED_BODY()
 public:
diff --git a/Engine/Source/Runtime/Engine/Public/WorldPartition/WorldPartitionActorDescViewProxy.h b/Engine/Source/Runtime/Engine/Public/WorldPartition/WorldPartitionActorDescViewProxy.h
--- a/Engine/Source/Runtime/Engine/Public/WorldPartition/WorldPartitionActorDescViewProxy.h
+++ b/Engine/Source/Runtime/Engine/Public/WorldPartition/WorldPartitionActorDescViewProxy.h
@@ -7,10 +7,10 @@
 /**
  * A view on top of an actor desc, which reverts to the actor if its loaded and dirty.
  */
-class FWorldPartitionActorViewProxy : public FWorldPartitionActorDescView
+class ENGINE_API FWorldPartitionActorViewProxy : public FWorldPartitionActorDescView
 {
 public:
-	ENGINE_API FWorldPartitionActorViewProxy(const FWorldPartitionActorDesc* InActorDesc);
+	FWorldPartitionActorViewProxy(const FWorldPartitionActorDesc* InActorDesc);
 
 private:
 	TUniquePtr<FWorldPartitionActorDesc> CachedActorDesc;
