World Partition - HLOD: Make HLOD warmup progress on the game thread rather than waiting on OnBeginRenderViews() for the associated world
* Prevent loading screens that causes the game not to be rendered from holding back the warmup (and possibly delaying the level visibility forever)
* Also fixed an issue with the warmup bounds being incorrect when dealing with instanced world partitions that have a transform.

diff --git a/Engine/Source/Runtime/Engine/Private/WorldPartition/HLOD/HLODRuntimeSubsystem.cpp b/Engine/Source/Runtime/Engine/Private/WorldPartition/HLOD/HLODRuntimeSubsystem.cpp
--- a/Engine/Source/Runtime/Engine/Private/WorldPartition/HLOD/HLODRuntimeSubsystem.cpp
+++ b/Engine/Source/Runtime/Engine/Private/WorldPartition/HLOD/HLODRuntimeSubsystem.cpp
@@ -572,6 +572,18 @@ bool UWorldPartitionHLODRuntimeSubsystem::PrepareToWarmup(const UWorldPartitionR
 		{
 			WarmupState.WarmupCallsUntilReady = CVarHLODWarmupNumFrames.GetValueOnGameThread();
 			WarmupState.WarmupBounds = InCell->GetContentBounds();
+
+			// If we're dealing with an instanced world partition, take the instance transform into account
+			const UWorldPartition* WorldPartition = InCell->GetOuterWorld()->GetWorldPartition();
+			if (ensure(WorldPartition) && !WorldPartition->IsMainWorldPartition())
+			{
+				WarmupState.WarmupBounds = WarmupState.WarmupBounds.TransformBy(WorldPartition->GetInstanceTransform());
+			}			
+		}
+		else if (WarmupState.WarmupCallsUntilReady != 0)
+		{	
+			// Progress toward warmup readiness
+			WarmupState.WarmupCallsUntilReady--;
 		}
 		
 		bHLODActorNeedsWarmUp = WarmupState.WarmupCallsUntilReady != 0;
@@ -818,12 +830,6 @@ void UWorldPartitionHLODRuntimeSubsystem::OnBeginRenderViews(const FSceneViewFam
 			}
 		}
 
-		// Progress toward warmup readiness
-		if (HLODWarmupState.WarmupCallsUntilReady != 0)
-		{
-			HLODWarmupState.WarmupCallsUntilReady--;
-		}
-
 		if (!HLODActor)
 		{
 			// HLOD actor has been unloaded, stop tracking it's warmup state
