Remove check in FConstraintInstance ComputeAverageMass_AssumesLocked. The average mass does not get used for anything, but regardless a value of zero is handled.

#tests Tested in PIE

diff --git a/Engine/Source/Runtime/Engine/Private/PhysicsEngine/ConstraintInstance.cpp b/Engine/Source/Runtime/Engine/Private/PhysicsEngine/ConstraintInstance.cpp
--- a/Engine/Source/Runtime/Engine/Private/PhysicsEngine/ConstraintInstance.cpp
+++ b/Engine/Source/Runtime/Engine/Private/PhysicsEngine/ConstraintInstance.cpp
@@ -444,9 +444,7 @@ float ComputeAverageMass_AssumesLocked(Chaos::FPhysicsObject* Body1, Chaos::FPhy
 		++NumDynamic;
 	}
 
-	check(NumDynamic);
-
-	if(NumDynamic > 0) // Some builds not taking the assumption from the check above and warn of zero divide
+	if(NumDynamic > 1)
 	{
 		AverageMass = TotalMass / NumDynamic; //-V609
 	}
@@ -561,6 +559,7 @@ void FConstraintProfileProperties::UpdateConstraintFlags_AssumesLocked(const FPh
 
 void FConstraintInstance::UpdateAverageMass_AssumesLocked(Chaos::FPhysicsObject* Body1, Chaos::FPhysicsObject* Body2)
 {
+	// @todo(chaos): Average mass isn't required by anything any more. We should probably remove this
 	AverageMass = ComputeAverageMass_AssumesLocked(Body1, Body2);
 }
 
