Fixed lost actor's binding to ActorFolders after running UWorldPartitionRenameDuplicateBuilder

diff --git a/Engine/Source/Runtime/Engine/Private/Level.cpp b/Engine/Source/Runtime/Engine/Private/Level.cpp
--- a/Engine/Source/Runtime/Engine/Private/Level.cpp
+++ b/Engine/Source/Runtime/Engine/Private/Level.cpp
@@ -2823,6 +2823,16 @@ void ULevel::OnFolderLabelChanged(UActorFolder* InActorFolder, const FString& In
 
 void ULevel::FixupActorFolders()
 {
+	auto BuildFolderLabelToActorFolders = [this]()
+	{
+		ForEachActorFolder([this](UActorFolder* ActorFolder)
+		{
+			FActorFolderSet& Folders = FolderLabelToActorFolders.FindOrAdd(ActorFolder->GetLabel());
+			Folders.Add(ActorFolder);
+			return true;
+		}, /*bSkipDeleted*/ true);
+	};
+
 	if (!IsActorFolderObjectsFeatureAvailable())
 	{
 		return;
@@ -2841,16 +2851,17 @@ void ULevel::FixupActorFolders()
 				AddActorFolder(LoadedActorFolder);
 			}
 			LoadedExternalActorFolders.Empty();
+
+			if (bWasDuplicated)
+			{
+				// Duplicated folders won't call AddActorFolder, build FolderLabelToActorFolders
+				BuildFolderLabelToActorFolders();
+			}
 		}
 		else
 		{
-			// Update FolderLabelToActorFolders for non-externalized ActorFolders
-			ForEachActorFolder([this](UActorFolder* ActorFolder)
-			{
-				FActorFolderSet& Folders = FolderLabelToActorFolders.FindOrAdd(ActorFolder->GetLabel());
-				Folders.Add(ActorFolder);
-				return true;
-			}, /*bSkipDeleted*/ true);
+			// Build FolderLabelToActorFolders for non-externalized ActorFolders
+			BuildFolderLabelToActorFolders();
 		}
 
 		ForEachActorFolder([](UActorFolder* ActorFolder)
