Fixed invalid message reported in PIE when an actor has a FolderGuid on an ActorFolder that doesn't exist.

diff --git a/Engine/Source/Runtime/Engine/Private/ActorEditor.cpp b/Engine/Source/Runtime/Engine/Private/ActorEditor.cpp
--- a/Engine/Source/Runtime/Engine/Private/ActorEditor.cpp
+++ b/Engine/Source/Runtime/Engine/Private/ActorEditor.cpp
@@ -1450,7 +1450,7 @@ void AActor::SetFolderGuidInternal(const FGuid& InFolderGuid, bool bInBroadcastC
 void AActor::SetFolderPathInternal(const FName& InNewFolderPath, bool bInBroadcastChange)
 {
 	FName OldPath = FolderPath;
-	if (InNewFolderPath.IsEqual(OldPath, ENameCase::CaseSensitive))
+	if (InNewFolderPath.IsEqual(OldPath, ENameCase::CaseSensitive) && !FolderGuid.IsValid())
 	{
 		return;
 	}
diff --git a/Engine/Source/Runtime/Engine/Private/WorldPartition/WorldPartitionLevelHelper.cpp b/Engine/Source/Runtime/Engine/Private/WorldPartition/WorldPartitionLevelHelper.cpp
--- a/Engine/Source/Runtime/Engine/Private/WorldPartition/WorldPartitionLevelHelper.cpp
+++ b/Engine/Source/Runtime/Engine/Private/WorldPartition/WorldPartitionLevelHelper.cpp
@@ -519,10 +519,8 @@ bool FWorldPartitionLevelHelper::LoadActors(const FLoadActorsParams& InParams)
 					const bool bDirectAccess = true;
 					const FGuid ActorFolderGuid = Actor->GetFolderGuid(bDirectAccess);
 					// Resolve folder guid from source container level and resolve/backup the folder path
-					if (UActorFolder* SrcFolder = ContainerWorld->PersistentLevel->GetActorFolder(ActorFolderGuid))
-					{
-						SrcActorFolderPath = SrcFolder->GetPath();
-					}
+					UActorFolder* SrcFolder = ContainerWorld->PersistentLevel->GetActorFolder(ActorFolderGuid);
+					SrcActorFolderPath = SrcFolder ? SrcFolder->GetPath() : NAME_None;
 				}
 
 				if (!PackageObjectMapping->ContainerID.IsMainContainer())
