Fix for encode speed not showing up on newly built textures.

diff --git a/Engine/Source/Runtime/Engine/Private/TextureDerivedDataTask.cpp b/Engine/Source/Runtime/Engine/Private/TextureDerivedDataTask.cpp
--- a/Engine/Source/Runtime/Engine/Private/TextureDerivedDataTask.cpp
+++ b/Engine/Source/Runtime/Engine/Private/TextureDerivedDataTask.cpp
@@ -659,6 +659,8 @@ static void DDC1_BuildTexture(
 	FTextureSourceData& TextureData,
 	FTextureSourceData& CompositeTextureData,
 	const TArrayView<FTextureBuildSettings>& InBuildSettingsPerLayer,
+	const FTexturePlatformData::FTextureEncodeResultMetadata& InBuildResultMetadata,
+
 	const FString& KeySuffix,
 	bool bReplaceExistingDDC,
 	int64 RequiredMemoryEstimate,
@@ -813,6 +815,7 @@ static void DDC1_BuildTexture(
 				NumMipsInTail, ExtData, bReplaceExistingDDC, TexturePathName, KeySuffix, BytesCached);
 
 			bSucceeded = true;
+			DerivedData->ResultMetadata = InBuildResultMetadata;
 
 			const bool bInlineMips = EnumHasAnyFlags(CacheFlags, ETextureCacheFlags::InlineMips);
 			if (bInlineMips) // Note that mips are inlined when cooking.
@@ -1800,6 +1803,7 @@ bool DDC1_BuildTiledClassicTexture(
 				TextureData,
 				CompositeTextureData,
 				LinearSettingsPerLayerFetchOrBuild,
+				FetchOrBuildMetadata,
 				LinearKeySuffix,
 				bLinearDDCCorrupted,
 				RequiredMemoryEstimate,
@@ -2004,8 +2008,8 @@ void FTextureCacheDerivedDataWorker::DoWork()
 			else
 			{
 				DDC1_BuildTexture(
-					Compressor, ImageWrapper, Texture, TexturePathName, CacheFlags, TextureData, CompositeTextureData, BuildSettingsPerLayerFetchOrBuild, KeySuffix, bReplaceExistingDDC,
-					RequiredMemoryEstimate, DerivedData, BytesCached, bSucceeded);
+					Compressor, ImageWrapper, Texture, TexturePathName, CacheFlags, TextureData, CompositeTextureData, BuildSettingsPerLayerFetchOrBuild, FetchOrBuildMetadata, 
+					KeySuffix, bReplaceExistingDDC, RequiredMemoryEstimate, DerivedData, BytesCached, bSucceeded);
 			}
 
 			if (bInvalidVirtualTextureCompression && DerivedData->VTData)
@@ -2123,7 +2127,7 @@ void FTextureCacheDerivedDataWorker::Finalize()
 			else
 			{
 				DDC1_BuildTexture(Compressor, ImageWrapper, Texture, TexturePathName, CacheFlags,
-					TextureData, CompositeTextureData, BuildSettingsPerLayerFetchOrBuild, KeySuffix, false /* currently corrupt vt data is not routed out of DoWork() */,
+					TextureData, CompositeTextureData, BuildSettingsPerLayerFetchOrBuild, FetchOrBuildMetadata, KeySuffix, false /* currently corrupt vt data is not routed out of DoWork() */,
 					RequiredMemoryEstimate, DerivedData, BytesCached, bSucceeded);
 			}
 
