Chaos - fix ComputeAverageMass_AssumesLocked knock-on from CL 29135827

#tests PIE Tests

diff --git a/Engine/Source/Runtime/Engine/Private/PhysicsEngine/ConstraintInstance.cpp b/Engine/Source/Runtime/Engine/Private/PhysicsEngine/ConstraintInstance.cpp
--- a/Engine/Source/Runtime/Engine/Private/PhysicsEngine/ConstraintInstance.cpp
+++ b/Engine/Source/Runtime/Engine/Private/PhysicsEngine/ConstraintInstance.cpp
@@ -426,27 +426,25 @@ void FConstraintInstance::SetDisableCollision(bool InDisableCollision)
 float ComputeAverageMass_AssumesLocked(Chaos::FPhysicsObject* Body1, Chaos::FPhysicsObject* Body2)
 {
 	float AverageMass = 0;
-
-	float TotalMass = 0;
 	int NumDynamic = 0;
 
 	Chaos::FReadPhysicsObjectInterface_External Interface = FPhysicsObjectExternalInterface::GetRead_AssumesLocked();
 
 	if (Interface.AreAllRigidBody({ &Body1, 1 }))
 	{
-		TotalMass += Interface.GetMass({ &Body1, 1 });
+		AverageMass += Interface.GetMass({ &Body1, 1 });
 		++NumDynamic;
 	}
 
 	if (Interface.AreAllRigidBody({ &Body2, 1 }))
 	{
-		TotalMass += Interface.GetMass({ &Body2, 1 });
+		AverageMass += Interface.GetMass({ &Body2, 1 });
 		++NumDynamic;
 	}
 
 	if(NumDynamic > 1)
 	{
-		AverageMass = TotalMass / NumDynamic; //-V609
+		AverageMass = AverageMass / NumDynamic; //-V609
 	}
 
 	return AverageMass;
