LFV - adds support for sequencer animation on component parameters.

diff --git a/Engine/Source/Runtime/Engine/Classes/Components/LocalFogVolumeComponent.h b/Engine/Source/Runtime/Engine/Classes/Components/LocalFogVolumeComponent.h
--- a/Engine/Source/Runtime/Engine/Classes/Components/LocalFogVolumeComponent.h
+++ b/Engine/Source/Runtime/Engine/Classes/Components/LocalFogVolumeComponent.h
@@ -47,6 +47,21 @@ class ULocalFogVolumeComponent : public USceneComponent
 	/** The priority can be used as a way to override the sorting by distance. A lower value means the volume will be considered further away, i.e. it will draw behind the one with a higher priority value. */
 	UPROPERTY(EditAnywhere, BlueprintReadOnly, Category = "Sorting", meta = (UIMin = "-127", UIMax = "127", ClampMin = -127, ClampMax = 127))
 	int32 FogSortPriority = 0;
+	
+	UFUNCTION(BlueprintCallable, Category = "Rendering", meta = (DisplayName = "Set Radial Fog Density"))
+	ENGINE_API void SetRadialFogExtinction(float NewValue);
+	UFUNCTION(BlueprintCallable, Category = "Rendering", meta = (DisplayName = "Set Height Fog Density"))
+	ENGINE_API void SetHeightFogExtinction(float NewValue);
+	UFUNCTION(BlueprintCallable, Category = "Rendering")
+	ENGINE_API void SetHeightFogFalloff(float NewValue);
+	UFUNCTION(BlueprintCallable, Category = "Rendering")
+	ENGINE_API void SetHeightFogOffset(float NewValue);
+	UFUNCTION(BlueprintCallable, Category = "Rendering", meta = (DisplayName = "Set Scattering Distribution"))
+	ENGINE_API void SetFogPhaseG(float NewValue);
+	UFUNCTION(BlueprintCallable, Category = "Rendering")
+	ENGINE_API void SetFogAlbedo(FLinearColor NewValue);
+	UFUNCTION(BlueprintCallable, Category = "Rendering")
+	ENGINE_API void SetFogEmissive(FLinearColor NewValue);
 
 public:
 
diff --git a/Engine/Source/Runtime/Engine/Private/Components/LocalFogVolumeComponent.cpp b/Engine/Source/Runtime/Engine/Private/Components/LocalFogVolumeComponent.cpp
--- a/Engine/Source/Runtime/Engine/Private/Components/LocalFogVolumeComponent.cpp
+++ b/Engine/Source/Runtime/Engine/Private/Components/LocalFogVolumeComponent.cpp
@@ -86,6 +86,42 @@ void ULocalFogVolumeComponent::DestroyRenderState_Concurrent()
 	}
 }
 
+#define LFV_BLUEPRINT_SETFUNCTION(MemberType, MemberName) void ULocalFogVolumeComponent::Set##MemberName(MemberType NewValue)\
+{\
+	if (AreDynamicDataChangesAllowed() && MemberName != NewValue)\
+	{\
+		MemberName = FMath::Max(MemberType(0), NewValue);\
+		MarkRenderStateDirty();\
+	}\
+}\
+
+#define LFV_BLUEPRINT_SETFUNCTION_LINEARCOLOR(MemberName) void ULocalFogVolumeComponent::Set##MemberName(FLinearColor NewValue)\
+{\
+	if (AreDynamicDataChangesAllowed() && MemberName != NewValue)\
+	{\
+		MemberName = NewValue.GetClamped(0.0f, 1e38f); \
+		MarkRenderStateDirty();\
+	}\
+}\
+
+#define LFV_BLUEPRINT_SETFUNCTION_LINEARCOLOR01(MemberName) void ULocalFogVolumeComponent::Set##MemberName(FLinearColor NewValue)\
+{\
+	if (AreDynamicDataChangesAllowed() && MemberName != NewValue)\
+	{\
+		MemberName = NewValue.GetClamped(0.0f, 1.0f); \
+		MarkRenderStateDirty();\
+	}\
+}\
+
+LFV_BLUEPRINT_SETFUNCTION(float, RadialFogExtinction);
+LFV_BLUEPRINT_SETFUNCTION(float, HeightFogExtinction);
+LFV_BLUEPRINT_SETFUNCTION(float, HeightFogFalloff);
+LFV_BLUEPRINT_SETFUNCTION(float, HeightFogOffset);
+LFV_BLUEPRINT_SETFUNCTION(float, FogPhaseG);
+LFV_BLUEPRINT_SETFUNCTION_LINEARCOLOR01(FogAlbedo);
+LFV_BLUEPRINT_SETFUNCTION_LINEARCOLOR(FogEmissive);
+
+
 #if WITH_EDITOR
 
 bool ULocalFogVolumeComponent::CanEditChange(const FProperty* InProperty) const
