[Chaos] Added a way to get the world bounds of a Particle In ClusterunionComponent that does a bit of less work, and changed FReadPhysicsObjectInterface<Id>::GetWorldBounds to make it not create an FTransform for then have to convert it to an TRigidTransform

diff --git a/Engine/Source/Runtime/Engine/Private/PhysicsEngine/ClusterUnionComponent.cpp b/Engine/Source/Runtime/Engine/Private/PhysicsEngine/ClusterUnionComponent.cpp
--- a/Engine/Source/Runtime/Engine/Private/PhysicsEngine/ClusterUnionComponent.cpp
+++ b/Engine/Source/Runtime/Engine/Private/PhysicsEngine/ClusterUnionComponent.cpp
@@ -62,6 +62,22 @@ namespace
 		check(Structure != nullptr);
 		return TUniquePtr<Chaos::ISpatialAcceleration<FExternalSpatialAccelerationPayload, Chaos::FReal, 3>>(Structure);
 	}
+
+	FBox GetWorldBoundsForParticle(const Chaos::FPBDRigidParticle* Particle)
+	{
+		if (!Particle)
+		{
+			return FBox(ForceInit);
+		}
+
+		if (const Chaos::FImplicitObjectRef Geometry = Particle->GetGeometry(); Geometry && Geometry->HasBoundingBox())
+		{
+			const Chaos::FAABB3 WorldBox = Geometry->CalculateTransformedBounds(Chaos::TRigidTransform<Chaos::FReal, 3>(Particle->X(), Particle->R()));
+			return FBox{ WorldBox.Min(), WorldBox.Max() };
+		}
+
+		return FBox(ForceInit);
+	}
 }
 
 UClusterUnionComponent::UClusterUnionComponent(const FObjectInitializer& ObjectInitializer)
@@ -173,7 +189,7 @@ void UClusterUnionComponent::AddComponentToCluster(UPrimitiveComponent* InCompon
 
 			if (AccelerationStructure)
 			{
-				const FBox HandleBounds = Interface->GetWorldBounds({ &Object, 1 });
+				const FBox HandleBounds = GetWorldBoundsForParticle(Particle);
 				FExternalSpatialAccelerationPayload Handle;
 				Handle.Initialize(ComponentKey, BoneData.ID, BoneData.ParticleID);
 				if (ensure(Handle.IsValid()))
@@ -1073,7 +1089,7 @@ void UClusterUnionComponent::HandleAddOrModifiedClusteredComponent(const FMapped
 			{
 				if (Chaos::FPhysicsObjectHandle PhysicsObject = Kvp.Value.PhysicsObjectHandle)
 				{
-					const FBox HandleBounds = Interface->GetWorldBounds({ &PhysicsObject, 1 });
+					const FBox HandleBounds = GetWorldBoundsForParticle(Kvp.Value.RigidParticle);
 					FExternalSpatialAccelerationPayload Handle;
 					Handle.Initialize(ChangedComponentData.ComponentKey, Kvp.Key, Kvp.Value.ParticleID);
 					if (ensure(Handle.IsValid()))
diff --git a/Engine/Source/Runtime/Experimental/Chaos/Private/Chaos/PhysicsObjectInterface.cpp b/Engine/Source/Runtime/Experimental/Chaos/Private/Chaos/PhysicsObjectInterface.cpp
--- a/Engine/Source/Runtime/Experimental/Chaos/Private/Chaos/PhysicsObjectInterface.cpp
+++ b/Engine/Source/Runtime/Experimental/Chaos/Private/Chaos/PhysicsObjectInterface.cpp
@@ -739,12 +739,10 @@ namespace Chaos
 				continue;
 			}
 
-			const FTransform WorldTransform = GetTransform(Object);
-
 			FBox ParticleBox(ForceInit);
 			if (const FImplicitObjectRef Geometry = Particle->GetGeometry(); Geometry && Geometry->HasBoundingBox())
 			{
-				const Chaos::FAABB3 WorldBox = Geometry->CalculateTransformedBounds(WorldTransform);
+				const Chaos::FAABB3 WorldBox = Geometry->CalculateTransformedBounds(TRigidTransform<FReal, 3>(Particle->X(), Particle->R()));
 				ParticleBox = FBox{ WorldBox.Min(), WorldBox.Max() };
 			}
 
