Move EMobileLocalLightSetting from MobileBasePassRendering to EngineTypes so it can be set and cached in default rendererSettings

diff --git a/Engine/Source/Runtime/Engine/Classes/Engine/EngineTypes.h b/Engine/Source/Runtime/Engine/Classes/Engine/EngineTypes.h
--- a/Engine/Source/Runtime/Engine/Classes/Engine/EngineTypes.h
+++ b/Engine/Source/Runtime/Engine/Classes/Engine/EngineTypes.h
@@ -580,6 +580,19 @@ namespace EGBufferFormat
 	};
 }
 
+
+/** 
+ * Enumerates available MobileLocalLightSetting. 
+ * @warning When this enum is updated please update CVarMobileForwardEnableLocalLights comments 
+ */
+UENUM()
+enum EMobileLocalLightSetting : int
+{
+	LOCAL_LIGHTS_DISABLED UMETA(DisplayName = "Local Lights Disabled"),
+	LOCAL_LIGHTS_ENABLED UMETA(DisplayName = "Local Lights Enabled"),
+	LOCAL_LIGHTS_BUFFER UMETA(DisplayName = "Local Lights Buffer Enabled")
+};
+
 /** Controls the way that the width scale property affects animation trails. */
 UENUM()
 enum ETrailWidthMode : int
diff --git a/Engine/Source/Runtime/Engine/Classes/Engine/RendererSettings.h b/Engine/Source/Runtime/Engine/Classes/Engine/RendererSettings.h
--- a/Engine/Source/Runtime/Engine/Classes/Engine/RendererSettings.h
+++ b/Engine/Source/Runtime/Engine/Classes/Engine/RendererSettings.h
@@ -1076,10 +1076,10 @@ class URendererSettings : public UDeveloperSettings
 
 	UPROPERTY(config, EditAnywhere, Category = Mobile, meta = (
 		ConsoleVariable = "r.Mobile.Forward.EnableLocalLights",
-		DisplayName = "Enable local lights support on mobile forward",
-		ToolTip = "Enable local lights support for mobile forward shading (including translucency in deferred). 0 is disabled, 1 is enabled (default). Changing this setting requires restarting the editor.",
+		DisplayName = "Mobile Local Light Setting",
+		ToolTip = "Select which Local Light Setting to use for Mobile. Changing this setting requires restarting the editor.",
 		ConfigRestartRequired = true))
-		uint32 bMobileForwardEnableLocalLights : 1;
+		TEnumAsByte<EMobileLocalLightSetting> MobileLocalLightSetting;
 
 	UPROPERTY(config, EditAnywhere, Category = Mobile, meta = (
 		ConsoleVariable = "r.Mobile.Forward.EnableClusteredReflections",
diff --git a/Engine/Source/Runtime/Renderer/Private/MobileBasePassRendering.h b/Engine/Source/Runtime/Renderer/Private/MobileBasePassRendering.h
--- a/Engine/Source/Runtime/Renderer/Private/MobileBasePassRendering.h
+++ b/Engine/Source/Runtime/Renderer/Private/MobileBasePassRendering.h
@@ -105,12 +105,6 @@ enum EOutputFormat
 	HDR_LINEAR_64,
 };
 
-enum EMobileLocalLightSetting
-{
-	LOCAL_LIGHTS_DISABLED,
-	LOCAL_LIGHTS_ENABLED,
-	LOCAL_LIGHTS_BUFFER
-};
 EMobileLocalLightSetting GetMobileForwardLocalLightSetting(EShaderPlatform ShaderPlatform, bool bIsTranslucent);
 
 bool ShouldCacheShaderByPlatformAndOutputFormat(EShaderPlatform Platform, EOutputFormat OutputFormat);
