WorldPartition: wp.Runtime.UpdateStreamingStateTimeLimit is now ignored if world is HighPriorityLoading or is in SeamlessTravel or if is inside a BlockTillLevelStreamingCompleted or if match hasn't started.

diff --git a/Engine/Source/Runtime/Engine/Private/WorldPartition/WorldPartitionSubsystem.cpp b/Engine/Source/Runtime/Engine/Private/WorldPartition/WorldPartitionSubsystem.cpp
--- a/Engine/Source/Runtime/Engine/Private/WorldPartition/WorldPartitionSubsystem.cpp
+++ b/Engine/Source/Runtime/Engine/Private/WorldPartition/WorldPartitionSubsystem.cpp
@@ -27,6 +27,7 @@
 #include "Engine/LevelBounds.h"
 #include "Debug/DebugDrawService.h"
 #include "GameFramework/PlayerController.h"
+#include "GameFramework/WorldSettings.h"
 #include "Async/ParallelFor.h"
 #include "Algo/ForEach.h"
 #include "Misc/HashBuilder.h"
@@ -1141,6 +1142,13 @@ bool UWorldPartitionSubsystem::IncrementalUpdateStreamingState()
 	return IncrementalUpdateWorldPartitions.IsEmpty();
 }
 
+static bool IsHighPriorityLoading(const UWorld* InWorld)
+{
+	const AWorldSettings* WorldSettings = InWorld ? InWorld->GetWorldSettings(false, false) : nullptr;
+	const bool bHighPriorityLoading = ensure(WorldSettings) ? (WorldSettings->bHighPriorityLoadingLocal || WorldSettings->bHighPriorityLoading) : false;
+	return bHighPriorityLoading;
+}
+
 void UWorldPartitionSubsystem::UpdateStreamingStateInternal(const UWorld* InWorld, UWorldPartition* InWorldPartition)
 {
 	TRACE_CPUPROFILER_EVENT_SCOPE(UWorldPartitionSubsystem::UpdateStreamingState);
@@ -1181,11 +1189,13 @@ void UWorldPartitionSubsystem::UpdateStreamingStateInternal(const UWorld* InWorl
 	const bool bIsServer = IsServer(InWorld);
 	const bool bServerStreamingEnabled = bIsServer && WorldPartitionSubsystem && WorldPartitionSubsystem->HasAnyWorldPartitionServerStreamingEnabled();
 	const int32 WorldPartitionUpdateCount = InWorldPartition ? 1 : WorldPartitionSubsystem->RegisteredWorldPartitions.Num();
+
+	const bool bForceDisableIncrementalUpdate = IsHighPriorityLoading(InWorld) || !InWorld->bMatchStarted || InWorld->IsInSeamlessTravel() || InWorld->GetIsInBlockTillLevelStreamingCompleted();
 	const bool bIncrementalUpdate = (GUpdateStreamingStateTimeLimit > 0.f) &&
 									(WorldPartitionUpdateCount > 1) &&
-									(!bIsServer || bServerStreamingEnabled) &&
-									!InWorld->GetIsInBlockTillLevelStreamingCompleted();
-	
+									!bForceDisableIncrementalUpdate &&
+									(!bIsServer || bServerStreamingEnabled); // No increment on server except if server streaming is enabled
+
 	// Update streaming state of all registered world partitions
 	if (bIncrementalUpdate)
 	{
@@ -1429,6 +1439,7 @@ void UWorldPartitionSubsystem::Draw(UCanvas* Canvas, class APlayerController* PC
 			if (IsIncrementalUnhashPending()) { StatusText += TEXT("(Unhashing) "); }
 			if (IsAsyncLoading()) { StatusText += TEXT("(AsyncLoading) "); }
 			if (StatusText.IsEmpty()) { StatusText = TEXT("(Idle) "); }
+			if (IsHighPriorityLoading(GetWorld())) { StatusText += TEXT("(HighPriorityLoading) "); }
 
 			FString DebugWorldText = FString::Printf(TEXT("(%s)"), *GetDebugStringForWorld(GetWorld()));
 			if (SingleWorldPartition && SingleWorldPartition->IsServer())
