World Partition: change cells streaming priority debug display to cubic grayscale for clearer visualization.

diff --git a/Engine/Source/Runtime/Engine/Private/WorldPartition/WorldPartitionDebugHelper.cpp b/Engine/Source/Runtime/Engine/Private/WorldPartition/WorldPartitionDebugHelper.cpp
--- a/Engine/Source/Runtime/Engine/Private/WorldPartition/WorldPartitionDebugHelper.cpp
+++ b/Engine/Source/Runtime/Engine/Private/WorldPartition/WorldPartitionDebugHelper.cpp
@@ -201,38 +201,4 @@ bool FWorldPartitionDebugHelper::bShowRuntimeSpatialHashCellStreamingPriority =
 FAutoConsoleVariableRef FWorldPartitionDebugHelper::ShowRuntimeSpatialHashCellStreamingPriorityCommand(
 	TEXT("wp.Runtime.ShowRuntimeSpatialHashCellStreamingPriority"),
 	FWorldPartitionDebugHelper::bShowRuntimeSpatialHashCellStreamingPriority,
-	TEXT("Enable to show a heatmap of the runtime spatial hash grid cells based on their priority."));
-
-FLinearColor FWorldPartitionDebugHelper::GetHeatMapColor(float ValueNormalized)
-{
-	const int32 NUM_COLORS = 4;
-	static FLinearColor Colors[NUM_COLORS] =
-	{
-		FLinearColor::Blue,
-		FLinearColor::Green,
-		FLinearColor::Yellow,
-		FLinearColor::Red
-	};
-
-	int32 idx1 = 0;
-	int32 idx2 = 0;
-	float Fraction = 0.f;
-
-	if (ValueNormalized >= 1.f)
-	{
-		idx1 = idx2 = NUM_COLORS - 1;
-	}
-	else if (ValueNormalized > 0)
-	{
-		// Find the 2 color indexes in which value lies.
-		ValueNormalized = ValueNormalized * (NUM_COLORS - 1);
-		idx1 = FMath::Floor(ValueNormalized);
-		idx2 = idx1 + 1;
-		// Calculate the fraction between the two indexes [0, 1].
-		Fraction = ValueNormalized - float(idx1);
-	}
-
-	// Interpolate between the 2 colors
-	FLinearColor Result = Colors[idx1] + Fraction * (Colors[idx2] - Colors[idx1]);
-	return Result;
-}
\ No newline at end of file
+	TEXT("Enable to show a heatmap of the runtime spatial hash grid cells based on their priority."));
\ No newline at end of file
diff --git a/Engine/Source/Runtime/Engine/Private/WorldPartition/WorldPartitionRuntimeCell.cpp b/Engine/Source/Runtime/Engine/Private/WorldPartition/WorldPartitionRuntimeCell.cpp
--- a/Engine/Source/Runtime/Engine/Private/WorldPartition/WorldPartitionRuntimeCell.cpp
+++ b/Engine/Source/Runtime/Engine/Private/WorldPartition/WorldPartitionRuntimeCell.cpp
@@ -112,17 +112,6 @@ bool UWorldPartitionRuntimeCell::IsDebugShown() const
 			RuntimeCellData->IsDebugShown();
 }
 
-FLinearColor UWorldPartitionRuntimeCell::GetDebugStreamingPriorityColor() const
-{
-#if !UE_BUILD_SHIPPING
-	if (DebugStreamingPriority >= 0.f && DebugStreamingPriority <= 1.f)
-	{
-		return FWorldPartitionDebugHelper::GetHeatMapColor(1.f - DebugStreamingPriority);
-	}
-#endif
-	return FLinearColor::Transparent;
-}
-
 UDataLayerManager* UWorldPartitionRuntimeCell::GetDataLayerManager() const
 {
 	return GetOuterWorld()->GetWorldPartition()->GetDataLayerManager();
diff --git a/Engine/Source/Runtime/Engine/Private/WorldPartition/WorldPartitionRuntimeLevelStreamingCell.cpp b/Engine/Source/Runtime/Engine/Private/WorldPartition/WorldPartitionRuntimeLevelStreamingCell.cpp
--- a/Engine/Source/Runtime/Engine/Private/WorldPartition/WorldPartitionRuntimeLevelStreamingCell.cpp
+++ b/Engine/Source/Runtime/Engine/Private/WorldPartition/WorldPartitionRuntimeLevelStreamingCell.cpp
@@ -201,22 +201,27 @@ EStreamingStatus UWorldPartitionRuntimeLevelStreamingCell::GetStreamingStatus()
 
 FLinearColor UWorldPartitionRuntimeLevelStreamingCell::GetDebugColor(EWorldPartitionRuntimeCellVisualizeMode VisualizeMode) const
 {
+#if !UE_BUILD_SHIPPING
 	switch (VisualizeMode)
 	{
 		case EWorldPartitionRuntimeCellVisualizeMode::StreamingPriority:
 		{
-			return GetDebugStreamingPriorityColor();
+			if (DebugStreamingPriority >= 0.0f && DebugStreamingPriority <= 1.0f)
+			{
+				const float PriorityGradient = FMath::Cube(1.0f - DebugStreamingPriority);
+				return FLinearColor(PriorityGradient, PriorityGradient, PriorityGradient, 1.0f);
+			}
+			return FLinearColor::Transparent;
 		}
 		case EWorldPartitionRuntimeCellVisualizeMode::StreamingStatus:
 		{
 			// Return streaming status color
 			return LevelStreaming ? ULevelStreaming::GetLevelStreamingStatusColor(LevelStreaming->GetLevelStreamingStatus()) : FLinearColor::Black;
 		}
-		default:
-		{
-			return Super::GetDebugColor(VisualizeMode);
-		}
 	}
+#endif
+
+	return Super::GetDebugColor(VisualizeMode);
 }
 
 void UWorldPartitionRuntimeLevelStreamingCell::SetIsAlwaysLoaded(bool bInIsAlwaysLoaded)
diff --git a/Engine/Source/Runtime/Engine/Private/WorldPartition/WorldPartitionRuntimeSpatialHash.cpp b/Engine/Source/Runtime/Engine/Private/WorldPartition/WorldPartitionRuntimeSpatialHash.cpp
--- a/Engine/Source/Runtime/Engine/Private/WorldPartition/WorldPartitionRuntimeSpatialHash.cpp
+++ b/Engine/Source/Runtime/Engine/Private/WorldPartition/WorldPartitionRuntimeSpatialHash.cpp
@@ -916,7 +916,7 @@ void FSpatialHashStreamingGrid::Draw2D(const UWorldPartitionRuntimeSpatialHash*
 					// Draw Cell using its debug color
 					FVector2D StartPos = CellWorldBounds.Min + CellOffset;
 					{
-						DrawContext.LocalDrawTile(GridScreenBounds, StartPos, CellBoundsSize, Cell->GetDebugColor(VisualizeMode).CopyWithNewOpacity(0.25f), WorldToScreen);
+						DrawContext.LocalDrawTile(GridScreenBounds, StartPos, CellBoundsSize, Cell->GetDebugColor(VisualizeMode).CopyWithNewOpacity((VisualizeMode == EWorldPartitionRuntimeCellVisualizeMode::StreamingPriority) ? 0.75f : 0.25f), WorldToScreen);
 					}
 
 					CellOffset.Y += CellBoundsSize.Y;
diff --git a/Engine/Source/Runtime/Engine/Public/WorldPartition/WorldPartitionDebugHelper.h b/Engine/Source/Runtime/Engine/Public/WorldPartition/WorldPartitionDebugHelper.h
--- a/Engine/Source/Runtime/Engine/Public/WorldPartition/WorldPartitionDebugHelper.h
+++ b/Engine/Source/Runtime/Engine/Public/WorldPartition/WorldPartitionDebugHelper.h
@@ -22,7 +22,6 @@ public:
 	static bool IsDebugCellNameShown(const FString& Name);
 	static void DrawText(UCanvas* Canvas, const FString& Text, const UFont* Font, const FColor& Color, FVector2D& Pos, float* MaxTextWidth = nullptr);
 	static void DrawLegendItem(UCanvas* Canvas, const FString& Text, const UFont* Font, const FColor& Color, const FColor& TextColor, FVector2D& Pos, float* MaxItemWidth = nullptr);
-	static FLinearColor GetHeatMapColor(float Value);
 	FORCEINLINE static bool IsRuntimeSpatialHashCellStreamingPriorityShown() { return bShowRuntimeSpatialHashCellStreamingPriority; }
 	FORCEINLINE static bool CanDrawContentBundles() { return bCanDrawContentBundles; }
 
diff --git a/Engine/Source/Runtime/Engine/Public/WorldPartition/WorldPartitionRuntimeCell.h b/Engine/Source/Runtime/Engine/Public/WorldPartition/WorldPartitionRuntimeCell.h
--- a/Engine/Source/Runtime/Engine/Public/WorldPartition/WorldPartitionRuntimeCell.h
+++ b/Engine/Source/Runtime/Engine/Public/WorldPartition/WorldPartitionRuntimeCell.h
@@ -273,8 +273,6 @@ class UWorldPartitionRuntimeCell : public UObject, public IWorldPartitionCell
 #endif
 
 protected:
-	ENGINE_API FLinearColor GetDebugStreamingPriorityColor() const;
-
 	//@todo_ow: Implement ServerOnlyVisible and refactor ClientOnlyVisible.
 	//          Instead of this function, server would not not wait for client level visibility 
 	//          for server-only visible cells.
