Fix FMessageStack::HasMessage to properly account for popped messages

diff --git a/Engine/Source/Runtime/Engine/Private/Animation/AnimNodeMessages.cpp b/Engine/Source/Runtime/Engine/Private/Animation/AnimNodeMessages.cpp
--- a/Engine/Source/Runtime/Engine/Private/Animation/AnimNodeMessages.cpp
+++ b/Engine/Source/Runtime/Engine/Private/Animation/AnimNodeMessages.cpp
@@ -150,7 +150,12 @@ void FMessageStack::BottomMessage(FName InMessageType, TFunctionRef<void(IGraphM
 
 bool FMessageStack::HasMessage(FName InMessageType) const
 {
-	return MessageStacks.Contains(InMessageType);
+	if (const FMessageStackEntry* StackEntryPtr = MessageStacks.Find(InMessageType))
+	{
+		return !StackEntryPtr->IsEmpty();
+	}
+
+	return false;
 }
 
 void FMessageStack::PushTag(FAnimNode_Base* InNode, const UScriptStruct* InStruct, FName InTag)
