Add null pointer check to FinInputdByName().

diff --git a/Engine/Source/Runtime/Engine/Private/Materials/MaterialExpressions.cpp b/Engine/Source/Runtime/Engine/Private/Materials/MaterialExpressions.cpp
--- a/Engine/Source/Runtime/Engine/Private/Materials/MaterialExpressions.cpp
+++ b/Engine/Source/Runtime/Engine/Private/Materials/MaterialExpressions.cpp
@@ -15010,7 +15010,7 @@ static const FFunctionExpressionInput* FindInputByName(const FName& Name, const
 	for (int32 InputIndex = 0; InputIndex < Inputs.Num(); InputIndex++)
 	{
 		const FFunctionExpressionInput& CurrentInput = Inputs[InputIndex];
-		if (CurrentInput.ExpressionInput->InputName == Name)
+		if (CurrentInput.ExpressionInput && CurrentInput.ExpressionInput->InputName == Name)
 		{
 			return &CurrentInput;
 		}
