[Navigation]
- Added flag in FNavigationRelevantData to indicates elements that want to skip dirty areas when added/removed from the NavigationOctree.
- Method ShouldSkipDirtyAreaOnAddOrRemove has been added to INavRelevantInterface to control the flag.
- UpdateNavOctreeElementBounds now takes a list of dirty areas instead of a single area to allow batched modifications.

UInstancedStaticMeshComponent updates:
- Added a cache for navigation bounds
- Implements ShouldSkipDirtyAreaOnAddOrRemove so we don't dirty the whole area covered by the instances on registration but only a list of explicit areas using the instances.
- Added PartialNavigationUpdates for batched modifications
- Fixed partial update not sending the right area for the original position

diff --git a/Engine/Source/Runtime/Engine/Classes/AI/Navigation/NavRelevantInterface.h b/Engine/Source/Runtime/Engine/Classes/AI/Navigation/NavRelevantInterface.h
--- a/Engine/Source/Runtime/Engine/Classes/AI/Navigation/NavRelevantInterface.h
+++ b/Engine/Source/Runtime/Engine/Classes/AI/Navigation/NavRelevantInterface.h
@@ -101,6 +101,9 @@ struct FNavigationRelevantData : public TSharedFromThis<FNavigationRelevantData,
 
 	uint32 bSupportsGatheringGeometrySlices : 1;
 
+	/** Indicates that this data will not dirty the navmesh when added or removed from the octree. */
+	uint32 bShouldSkipDirtyAreaOnAddOrRemove : 1;
+
 	/** From level loading (only valid in WP dynamic mode) */
 	uint32 bLoadedData : 1;
 
@@ -110,6 +113,7 @@ struct FNavigationRelevantData : public TSharedFromThis<FNavigationRelevantData,
 		, bPendingLazyModifiersGathering(false)
 		, bPendingChildLazyModifiersGathering(false)
 		, bSupportsGatheringGeometrySlices(false)
+		, bShouldSkipDirtyAreaOnAddOrRemove(false)
 		, bLoadedData(false)
 	{}
 
@@ -120,6 +124,13 @@ struct FNavigationRelevantData : public TSharedFromThis<FNavigationRelevantData,
 	FORCEINLINE bool IsPendingChildLazyModifiersGathering() const { return bPendingChildLazyModifiersGathering; }
 	FORCEINLINE bool NeedAnyPendingLazyModifiersGathering() const { return bPendingLazyModifiersGathering || bPendingChildLazyModifiersGathering; }
 	FORCEINLINE bool SupportsGatheringGeometrySlices() const { return bSupportsGatheringGeometrySlices; }
+
+	/**
+	 * Indicates that this object will not dirty the navmesh when added or removed from the octree.
+	 * In this case we expect it to manually dirty areas (e.g. using OnObjectBoundsChanged).
+	 */
+	FORCEINLINE bool ShouldSkipDirtyAreaOnAddOrRemove() const { return bShouldSkipDirtyAreaOnAddOrRemove; }
+
 	FORCEINLINE bool IsEmpty() const { return !HasGeometry() && !HasModifiers(); }
 	FORCEINLINE SIZE_T GetAllocatedSize() const { return CollisionData.GetAllocatedSize() + VoxelData.GetAllocatedSize() + Modifiers.GetAllocatedSize(); }
 	FORCEINLINE SIZE_T GetGeometryAllocatedSize() const { return CollisionData.GetAllocatedSize() + VoxelData.GetAllocatedSize(); }
@@ -177,6 +188,11 @@ class INavRelevantInterface
 
 	/** Indicates if this instance knows how to export sub-sections of self */
 	virtual bool SupportsGatheringGeometrySlices() const { return false; }
+	
+	/** Indicates if the area covered by the navigation bounds of the object should not be dirtied when inserting, or removing, the object in the navigation octree.
+	 * In this case we expect that object to manually dirty areas (e.g. using OnObjectBoundsChanged).
+	 */
+	virtual bool ShouldSkipDirtyAreaOnAddOrRemove() const { return false; }
 
 	/** This function is called "on demand", whenever that specified piece of geometry is needed for navigation generation */
 	virtual void GatherGeometrySlice(FNavigableGeometryExport& GeomExport, const FBox& SliceBox) const {}
diff --git a/Engine/Source/Runtime/Engine/Classes/AI/NavigationSystemBase.h b/Engine/Source/Runtime/Engine/Classes/AI/NavigationSystemBase.h
--- a/Engine/Source/Runtime/Engine/Classes/AI/NavigationSystemBase.h
+++ b/Engine/Source/Runtime/Engine/Classes/AI/NavigationSystemBase.h
@@ -137,6 +137,8 @@ namespace FNavigationSystem
 	//ENGINE_API bool HasComponentData(UActorComponent& Comp);
 	ENGINE_API void OnActorBoundsChanged(AActor& Actor);
 	ENGINE_API void OnPostEditActorMove(AActor& Actor);
+
+	UE_DEPRECATED(5.4, "Use OnObjectBoundsChanged taking UObject and list of dirty areas as parameters instead.")
 	ENGINE_API void OnComponentBoundsChanged(UActorComponent& Comp, const FBox& NewBounds, const FBox& DirtyArea);
 	ENGINE_API void OnComponentTransformChanged(USceneComponent& Comp);
 
@@ -152,6 +154,7 @@ namespace FNavigationSystem
 	ENGINE_API void RegisterNavRelevantObject(UObject& Object);
 	ENGINE_API void UpdateNavRelevantObject(UObject& Object);
 	ENGINE_API void UnregisterNavRelevantObject(UObject& Object);
+	ENGINE_API void OnObjectBoundsChanged(UObject& Object, const FBox& NewBounds, TConstArrayView<FBox> DirtyAreas);
 
 	ENGINE_API void RemoveActorData(AActor& Actor);
 
@@ -195,6 +198,7 @@ namespace FNavigationSystem
 	ENGINE_API IPathFollowingAgentInterface* FindPathFollowingAgentForActor(const AActor& Actor);
 
 	DECLARE_DELEGATE_OneParam(FObjectBasedSignature, UObject& /*Object*/);
+	DECLARE_DELEGATE_ThreeParams(FObjectBoundsChangedSignature, UObject& /*Object*/, const FBox& /*NewBounds*/, TConstArrayView<FBox> /*DirtyAreas*/)
 	DECLARE_DELEGATE_OneParam(FActorBasedSignature, AActor& /*Actor*/);
 	DECLARE_DELEGATE_OneParam(FActorComponentBasedSignature, UActorComponent& /*Comp*/);
 	DECLARE_DELEGATE_OneParam(FSceneComponentBasedSignature, USceneComponent& /*Comp*/);
@@ -210,13 +214,15 @@ namespace FNavigationSystem
 	DECLARE_DELEGATE_RetVal_OneParam(double, FDoubleWorldBasedSignature, const UWorld& /*World*/);
 	DECLARE_DELEGATE_TwoParams(FWorldByteBasedSignature, UWorld& /*World*/, uint8 /*Flags*/);
 	DECLARE_DELEGATE_TwoParams(FActorBooleBasedSignature, AActor& /*Actor*/, bool /*bUpdateAttachedActors*/);
-	DECLARE_DELEGATE_ThreeParams(FComponentBoundsChangeSignature, UActorComponent& /*Comp*/, const FBox& /*NewBounds*/, const FBox& /*DirtyArea*/)
 	DECLARE_DELEGATE_RetVal_OneParam(INavigationDataInterface*, FNavDataForPropsSignature, const FNavAgentProperties& /*AgentProperties*/);
 	DECLARE_DELEGATE_RetVal_OneParam(INavigationDataInterface*, FNavDataForActorSignature, const AActor& /*Actor*/);
 	DECLARE_DELEGATE_RetVal(TSubclassOf<AActor>, FNavDataClassFetchSignature);
 	DECLARE_DELEGATE_TwoParams(FWorldBoolBasedSignature, UWorld& /*World*/, const bool /*bShow*/);
 	DECLARE_MULTICAST_DELEGATE_OneParam(FOnNavigationInitSignature, const UNavigationSystemBase&);
 	DECLARE_MULTICAST_DELEGATE_TwoParams(FOnNavAreaGenericEvent, const UWorld&, const UClass*);
+
+	UE_DEPRECATED(5.4, "Use FObjectBoundsChangedSignature delegate taking UObject and list of dirty areas as parameters instead.")
+	DECLARE_DELEGATE_ThreeParams(FComponentBoundsChangeSignature, UActorComponent& /*Comp*/, const FBox& /*NewBounds*/, const FBox& /*DirtyArea*/)
 }
 
 
@@ -285,7 +291,8 @@ protected:
 	static ENGINE_API FNavigationSystem::FObjectBasedSignature& RegisterNavRelevantObjectDelegate();
 	static ENGINE_API FNavigationSystem::FObjectBasedSignature& UpdateNavRelevantObjectDelegate();
 	static ENGINE_API FNavigationSystem::FObjectBasedSignature& UnregisterNavRelevantObjectDelegate();
-	
+	static ENGINE_API FNavigationSystem::FObjectBoundsChangedSignature& OnObjectBoundsChangedDelegate();
+
 	static ENGINE_API FNavigationSystem::FActorBasedSignature& UpdateActorDataDelegate();
 	static ENGINE_API FNavigationSystem::FActorComponentBasedSignature& UpdateComponentDataDelegate();
 	static ENGINE_API FNavigationSystem::FSceneComponentBasedSignature& UpdateComponentDataAfterMoveDelegate();
@@ -303,11 +310,14 @@ protected:
 	static ENGINE_API FNavigationSystem::FNavDataConfigBasedSignature& GetDefaultSupportedAgentDelegate();
 	static ENGINE_API FNavigationSystem::FNavDataConfigAndWorldSignature& GetBiggestSupportedAgentDelegate();
 	static ENGINE_API FNavigationSystem::FActorBooleBasedSignature& UpdateActorAndComponentDataDelegate();
-	static ENGINE_API FNavigationSystem::FComponentBoundsChangeSignature& OnComponentBoundsChangedDelegate();
 	static ENGINE_API FNavigationSystem::FNavDataForActorSignature& GetNavDataForActorDelegate();
 	static ENGINE_API FNavigationSystem::FNavDataClassFetchSignature& GetDefaultNavDataClassDelegate();
 	static ENGINE_API FNavigationSystem::FWorldBoolBasedSignature& VerifyNavigationRenderingComponentsDelegate();
 	static ENGINE_API FNavigationSystem::FWorldBasedSignature& BuildDelegate();
+
+	PRAGMA_DISABLE_DEPRECATION_WARNINGS
+	static ENGINE_API FNavigationSystem::FComponentBoundsChangeSignature& OnComponentBoundsChangedDelegate();
+	PRAGMA_ENABLE_DEPRECATION_WARNINGS
 #if WITH_EDITOR
 	static ENGINE_API FNavigationSystem::FWorldBasedSignature& OnPIEStartDelegate();
 	static ENGINE_API FNavigationSystem::FWorldBasedSignature& OnPIEEndDelegate();
diff --git a/Engine/Source/Runtime/Engine/Classes/Components/HierarchicalInstancedStaticMeshComponent.h b/Engine/Source/Runtime/Engine/Classes/Components/HierarchicalInstancedStaticMeshComponent.h
--- a/Engine/Source/Runtime/Engine/Classes/Components/HierarchicalInstancedStaticMeshComponent.h
+++ b/Engine/Source/Runtime/Engine/Classes/Components/HierarchicalInstancedStaticMeshComponent.h
@@ -386,7 +386,10 @@ protected:
 	ENGINE_API virtual void PartialNavigationUpdate(int32 InstanceIdx) override;
 	virtual bool SupportsPartialNavigationUpdate() const override { return true; }
 	ENGINE_API void FlushAccumulatedNavigationUpdates();
-	mutable FBox AccumulatedNavigationDirtyArea;
+
+	UE_DEPRECATED(5.4, "Use AccumulatedNavigationDirtyAreas instead.")
+    mutable FBox AccumulatedNavigationDirtyArea;
+	mutable TArray<FBox> AccumulatedNavigationDirtyAreas;
 
 	FGraphEventArray BuildTreeAsyncTasks;
 	EHISMViewRelevanceType ViewRelevanceType = EHISMViewRelevanceType::HISM;
diff --git a/Engine/Source/Runtime/Engine/Classes/Components/InstancedStaticMeshComponent.h b/Engine/Source/Runtime/Engine/Classes/Components/InstancedStaticMeshComponent.h
--- a/Engine/Source/Runtime/Engine/Classes/Components/InstancedStaticMeshComponent.h
+++ b/Engine/Source/Runtime/Engine/Classes/Components/InstancedStaticMeshComponent.h
@@ -105,6 +105,14 @@ class UInstancedStaticMeshComponent : public UStaticMeshComponent, public ISMIns
 	UPROPERTY(Transient)
 	TArray<FMatrix> PerInstancePrevTransform;
 
+	/** Bounds are calculated and cached on component registration. */
+	UPROPERTY(Transient)
+	FBox NavigationBounds;
+	
+	/** Main transform stored to be able to send updates when component's transform changed. */
+	UPROPERTY(Transient)
+	FTransform PreviousComponentTransform;
+
 	/** Defines the number of floats that will be available per instance for custom data */
 	UPROPERTY(EditAnywhere, Category=Instances, AdvancedDisplay)
 	int32 NumCustomDataFloats;
@@ -479,6 +487,7 @@ public:
 	ENGINE_API virtual void GetNavigationData(FNavigationRelevantData& Data) const override;
 	ENGINE_API virtual FBox GetNavigationBounds() const override;
 	ENGINE_API virtual bool IsNavigationRelevant() const override;
+	ENGINE_API virtual bool ShouldSkipDirtyAreaOnAddOrRemove() const override;
 	//~ End UNavRelevantInterface Interface
 
 	//~ Begin UObject Interface
@@ -548,6 +557,9 @@ public:
 	/** Request to navigation system to update only part of navmesh occupied by specified instance. */
 	ENGINE_API virtual void PartialNavigationUpdate(int32 InstanceIdx);
 
+	/** Request to navigation system to update only part of navmesh occupied specified instances transforms. */
+	ENGINE_API virtual void PartialNavigationUpdates(TConstArrayView<FTransform> InstanceTransforms);
+
 	/** 
 	 * Flag for using RemoveAtSwap on instance removal. 
 	 * The implementation is free to ignore this flag, but should honor whatever behavior is being returned by SupportsRemoveSwap().
@@ -564,11 +576,12 @@ public:
 	/** 
 	 * Clears all the updated instance tracking data AND instance ID association, also forcing a full update of the instance data the next time it is flushed.	 
 	 * NOTE: Destroying the instance updated tracking means the renderer has to treat the instances as completely new, preventing e.g., velocity tracking and caching from working reliably.
-	 * This function is only intended to be used when some outside entity has modified the instance data, this is not recomended and access to the public data will be removed in a future release.
+	 * This function is only intended to be used when some outside entity has modified the instance data, this is not recommended and access to the public data will be removed in a future release.
 	 */
 	ENGINE_API void InvalidateInstanceDataTracking();
 
 private:
+	void CalcAndCacheNavigationBounds();
 
 	/** Sets up new instance data to sensible defaults, creates physics counterparts if possible. */
 	ENGINE_API void SetupNewInstanceData(FInstancedStaticMeshInstanceData& InOutNewInstanceData, int32 InInstanceIndex, const FTransform& InInstanceTransform);
diff --git a/Engine/Source/Runtime/Engine/Private/AI/NavigationSystemBase.cpp b/Engine/Source/Runtime/Engine/Private/AI/NavigationSystemBase.cpp
--- a/Engine/Source/Runtime/Engine/Private/AI/NavigationSystemBase.cpp
+++ b/Engine/Source/Runtime/Engine/Private/AI/NavigationSystemBase.cpp
@@ -130,7 +130,8 @@ namespace FNavigationSystem
 		FObjectBasedSignature RegisterNavRelevantObject;
 		FObjectBasedSignature UpdateNavRelevantObject;
 		FObjectBasedSignature UnregisterNavRelevantObject;
-		
+		FObjectBoundsChangedSignature OnObjectBoundsChanged;
+
 		FActorBasedSignature UpdateActorData;
 		FActorComponentBasedSignature UpdateComponentData;
 		FSceneComponentBasedSignature UpdateComponentDataAfterMove;
@@ -150,7 +151,6 @@ namespace FNavigationSystem
 		FNavDataConfigBasedSignature GetDefaultSupportedAgent;
 		FNavDataConfigAndWorldSignature GetBiggestSupportedAgent;
 		FActorBooleBasedSignature UpdateActorAndComponentData;
-		FComponentBoundsChangeSignature OnComponentBoundsChanged;
 		FNavDataForActorSignature GetNavDataForActor;
 		FNavDataClassFetchSignature GetDefaultNavDataClass;
 		FWorldBoolBasedSignature VerifyNavigationRenderingComponents;
@@ -159,6 +159,11 @@ namespace FNavigationSystem
 		FOnNavigationInitSignature OnNavigationInitDone;
 		FOnNavAreaGenericEvent OnNavAreaRegistered;
 		FOnNavAreaGenericEvent OnNavAreaUnregistered;
+
+PRAGMA_DISABLE_DEPRECATION_WARNINGS
+		FComponentBoundsChangeSignature OnComponentBoundsChanged;
+PRAGMA_ENABLE_DEPRECATION_WARNINGS
+		
 #if WITH_EDITOR
 		FWorldBasedSignature OnPIEStart;
 		FWorldBasedSignature OnPIEEnd;
@@ -174,6 +179,7 @@ namespace FNavigationSystem
 			RegisterNavRelevantObject.BindLambda([](UObject&) {});
 			UpdateNavRelevantObject.BindLambda([](UObject&) {});
 			UnregisterNavRelevantObject.BindLambda([](UObject&) {});
+			OnObjectBoundsChanged.BindLambda([](UObject&, const FBox&, TConstArrayView<FBox>) {});
 
 			UpdateActorData.BindLambda([](AActor&) {});
 			UpdateComponentData.BindLambda([](UActorComponent&) {});
@@ -218,6 +224,7 @@ namespace FNavigationSystem
 	void RegisterNavRelevantObject(UObject& Object) { Delegates.RegisterNavRelevantObject.Execute(Object); }
 	void UpdateNavRelevantObject(UObject& Object) { Delegates.UpdateNavRelevantObject.Execute(Object); }
 	void UnregisterNavRelevantObject(UObject& Object) { Delegates.UnregisterNavRelevantObject.Execute(Object); }
+	void OnObjectBoundsChanged(UObject& Object, const FBox& NewBounds, const TConstArrayView<FBox> DirtyAreas) { Delegates.OnObjectBoundsChanged.Execute(Object, NewBounds, DirtyAreas); }
 
 	void UpdateActorData(AActor& Actor) { Delegates.UpdateActorData.Execute(Actor); }
 	void UpdateComponentData(UActorComponent& Comp) { Delegates.UpdateComponentData.Execute(Comp); }
@@ -225,7 +232,7 @@ namespace FNavigationSystem
 	void UpdateComponentDataAfterMove(USceneComponent& Comp) { Delegates.UpdateComponentDataAfterMove.Execute(Comp); }
 	void OnActorBoundsChanged(AActor& Actor) { Delegates.OnActorBoundsChanged.Execute(Actor); }
 	void OnPostEditActorMove(AActor& Actor) { Delegates.OnPostEditActorMove.Execute(Actor); }
-	void OnComponentBoundsChanged(UActorComponent& Comp, const FBox& NewBounds, const FBox& DirtyArea) { Delegates.OnComponentBoundsChanged.Execute(Comp, NewBounds, DirtyArea); }
+	void OnComponentBoundsChanged(UActorComponent& Comp, const FBox& NewBounds, const FBox& DirtyArea) { OnObjectBoundsChanged(Comp, NewBounds, { DirtyArea }); }
 	void OnComponentTransformChanged(USceneComponent& Comp) { Delegates.OnComponentTransformChanged.Execute(Comp); }
 	void OnActorRegistered(AActor& Actor) { Delegates.OnActorRegistered.Execute(Actor); }
 	void OnActorUnregistered(AActor& Actor) { Delegates.OnActorUnregistered.Execute(Actor); }
@@ -403,6 +410,7 @@ void UNavigationSystemBase::ResetEventDelegates() { FNavigationSystem::ResetDele
 FNavigationSystem::FObjectBasedSignature& UNavigationSystemBase::RegisterNavRelevantObjectDelegate() { return FNavigationSystem::Delegates.RegisterNavRelevantObject; }
 FNavigationSystem::FObjectBasedSignature& UNavigationSystemBase::UpdateNavRelevantObjectDelegate() { return FNavigationSystem::Delegates.UpdateNavRelevantObject; }
 FNavigationSystem::FObjectBasedSignature& UNavigationSystemBase::UnregisterNavRelevantObjectDelegate() { return FNavigationSystem::Delegates.UnregisterNavRelevantObject; }
+FNavigationSystem::FObjectBoundsChangedSignature& UNavigationSystemBase::OnObjectBoundsChangedDelegate() { return FNavigationSystem::Delegates.OnObjectBoundsChanged; }
 
 FNavigationSystem::FActorBasedSignature& UNavigationSystemBase::UpdateActorDataDelegate() { return FNavigationSystem::Delegates.UpdateActorData; }
 FNavigationSystem::FActorComponentBasedSignature& UNavigationSystemBase::UpdateComponentDataDelegate() { return FNavigationSystem::Delegates.UpdateComponentData; }
@@ -421,7 +429,6 @@ FNavigationSystem::FBoolActorComponentBasedSignature& UNavigationSystemBase::Has
 FNavigationSystem::FNavDataConfigBasedSignature& UNavigationSystemBase::GetDefaultSupportedAgentDelegate() { return FNavigationSystem::Delegates.GetDefaultSupportedAgent; }
 FNavigationSystem::FNavDataConfigAndWorldSignature& UNavigationSystemBase::GetBiggestSupportedAgentDelegate() { return FNavigationSystem::Delegates.GetBiggestSupportedAgent; }
 FNavigationSystem::FActorBooleBasedSignature& UNavigationSystemBase::UpdateActorAndComponentDataDelegate() { return FNavigationSystem::Delegates.UpdateActorAndComponentData; }
-FNavigationSystem::FComponentBoundsChangeSignature& UNavigationSystemBase::OnComponentBoundsChangedDelegate() { return FNavigationSystem::Delegates.OnComponentBoundsChanged; }
 FNavigationSystem::FNavDataForActorSignature& UNavigationSystemBase::GetNavDataForActorDelegate() { return FNavigationSystem::Delegates.GetNavDataForActor; }
 FNavigationSystem::FNavDataClassFetchSignature& UNavigationSystemBase::GetDefaultNavDataClassDelegate() { return FNavigationSystem::Delegates.GetDefaultNavDataClass; }
 FNavigationSystem::FWorldBoolBasedSignature& UNavigationSystemBase::VerifyNavigationRenderingComponentsDelegate() { return FNavigationSystem::Delegates.VerifyNavigationRenderingComponents; }
@@ -430,6 +437,11 @@ FNavigationSystem::FOnNavigationInitSignature& UNavigationSystemBase::OnNavigati
 FNavigationSystem::FOnNavigationInitSignature& UNavigationSystemBase::OnNavigationInitDoneStaticDelegate() { return FNavigationSystem::Delegates.OnNavigationInitDone; }
 FNavigationSystem::FOnNavAreaGenericEvent& UNavigationSystemBase::OnNavAreaRegisteredDelegate() { return FNavigationSystem::Delegates.OnNavAreaRegistered; }
 FNavigationSystem::FOnNavAreaGenericEvent& UNavigationSystemBase::OnNavAreaUnregisteredDelegate() { return FNavigationSystem::Delegates.OnNavAreaUnregistered; }
+
+PRAGMA_DISABLE_DEPRECATION_WARNINGS
+FNavigationSystem::FComponentBoundsChangeSignature& UNavigationSystemBase::OnComponentBoundsChangedDelegate() { return FNavigationSystem::Delegates.OnComponentBoundsChanged; }
+PRAGMA_ENABLE_DEPRECATION_WARNINGS
+
 #if WITH_EDITOR
 FNavigationSystem::FWorldBasedSignature& UNavigationSystemBase::OnPIEStartDelegate() { return FNavigationSystem::Delegates.OnPIEStart; }
 FNavigationSystem::FWorldBasedSignature& UNavigationSystemBase::OnPIEEndDelegate() { return FNavigationSystem::Delegates.OnPIEEnd; }
diff --git a/Engine/Source/Runtime/Engine/Private/HierarchicalInstancedStaticMesh.cpp b/Engine/Source/Runtime/Engine/Private/HierarchicalInstancedStaticMesh.cpp
--- a/Engine/Source/Runtime/Engine/Private/HierarchicalInstancedStaticMesh.cpp
+++ b/Engine/Source/Runtime/Engine/Private/HierarchicalInstancedStaticMesh.cpp
@@ -2036,7 +2036,6 @@ UHierarchicalInstancedStaticMeshComponent::UHierarchicalInstancedStaticMeshCompo
 #if WITH_EDITOR
 	, bCanEnableDensityScaling(true)
 #endif
-	, AccumulatedNavigationDirtyArea(ForceInit)
 {
 	PrimitiveInstanceDataManager.SetMode(FPrimitiveInstanceDataManager::EMode::Legacy);
 	bCanEverAffectNavigation = true;
@@ -3223,8 +3222,7 @@ void UHierarchicalInstancedStaticMeshComponent::GetNavigationPerInstanceTransfor
 	{
 		// This area should be processed again by navigation system when cluster tree is available
 		// Store smaller tile box in accumulated dirty area, so we will not unintentionally mark as dirty neighbor tiles 
-		const FBox SmallTileBox = AreaBox.ExpandBy(-AreaBox.GetExtent()/2.f);
-		AccumulatedNavigationDirtyArea+= SmallTileBox;
+		AccumulatedNavigationDirtyAreas.Emplace(AreaBox.ExpandBy(-AreaBox.GetExtent()/2.f));
 	}
 }
 
@@ -3232,7 +3230,7 @@ void UHierarchicalInstancedStaticMeshComponent::PartialNavigationUpdate(const in
 {
 	if (InstanceIdx == INDEX_NONE)
 	{
-		AccumulatedNavigationDirtyArea.Init();
+		AccumulatedNavigationDirtyAreas.Reset();
 		FNavigationSystem::UpdateComponentData(*this);
 	}
 	else if (GetStaticMesh() && FNavigationSystem::HasComponentData(*this))
@@ -3244,7 +3242,7 @@ void UHierarchicalInstancedStaticMeshComponent::PartialNavigationUpdate(const in
 			const FBox InstanceBounds = GetInstanceNavigationBounds();
 			if (InstanceBounds.IsValid)
 			{
-				AccumulatedNavigationDirtyArea += InstanceBounds.TransformBy(InstanceTransformInWorldSpace);
+				AccumulatedNavigationDirtyAreas.Emplace(InstanceBounds.TransformBy(InstanceTransformInWorldSpace));
 			}
 		}
 	}
@@ -3252,7 +3250,7 @@ void UHierarchicalInstancedStaticMeshComponent::PartialNavigationUpdate(const in
 
 void UHierarchicalInstancedStaticMeshComponent::FlushAccumulatedNavigationUpdates()
 {
-	if (AccumulatedNavigationDirtyArea.IsValid)
+	if (AccumulatedNavigationDirtyAreas.Num() > 0)
 	{
 		QUICK_SCOPE_CYCLE_COUNTER(STAT_UHierarchicalInstancedStaticMeshComponent_FlushAccumulatedNavigationUpdates);
 
@@ -3260,10 +3258,10 @@ void UHierarchicalInstancedStaticMeshComponent::FlushAccumulatedNavigationUpdate
 		if (ClusterTree.Num())
 		{
 			const FBox NewBounds = GetClusterTreeBounds(ClusterTree, TranslatedInstanceSpaceOrigin).TransformBy(GetComponentTransform());
-			FNavigationSystem::OnComponentBoundsChanged(*this, NewBounds, AccumulatedNavigationDirtyArea);
+			FNavigationSystem::OnObjectBoundsChanged(*this, NewBounds, AccumulatedNavigationDirtyAreas);
 		}
-			
-		AccumulatedNavigationDirtyArea.Init();
+
+		AccumulatedNavigationDirtyAreas.Reset();
 	}
 }
 
diff --git a/Engine/Source/Runtime/Engine/Private/InstancedStaticMesh.cpp b/Engine/Source/Runtime/Engine/Private/InstancedStaticMesh.cpp
--- a/Engine/Source/Runtime/Engine/Private/InstancedStaticMesh.cpp
+++ b/Engine/Source/Runtime/Engine/Private/InstancedStaticMesh.cpp
@@ -2535,8 +2535,6 @@ void UInstancedStaticMeshComponent::CreateAllInstanceBodies()
 					Transforms.Add(InstanceTM);
 				}
 			}
-
-	    	PartialNavigationUpdate(i);
 	    }
 
 		if (InstanceBodiesSanitized.Num() > 0 && Mobility != EComponentMobility::Movable)
@@ -2558,14 +2556,12 @@ void UInstancedStaticMeshComponent::ClearAllInstanceBodies()
 	QUICK_SCOPE_CYCLE_COUNTER(STAT_UInstancedStaticMeshComponent_ClearAllInstanceBodies);
 	STAT(FScopeCycleCounter Context(StatId);)
 
-	for (int32 i = 0; i < InstanceBodies.Num(); i++)
+	for (FBodyInstance*& Instance : InstanceBodies)
 	{
-		if (InstanceBodies[i])
+		if (Instance)
 		{
-			PartialNavigationUpdate(i);
-			
-			InstanceBodies[i]->TermBody();
-			delete InstanceBodies[i];
+			Instance->TermBody();
+			delete Instance;
 		}
 	}
 
@@ -2638,6 +2634,11 @@ FBoxSphereBounds UInstancedStaticMeshComponent::CalcBounds(const FTransform& Bou
 	return CalcBoundsImpl(BoundTransform, /*bForNavigation*/false);
 }
 
+void UInstancedStaticMeshComponent::CalcAndCacheNavigationBounds()
+{
+	NavigationBounds = CalcBoundsImpl(GetComponentTransform(), /*bForNavigation*/true).GetBox();
+}
+
 FBoxSphereBounds UInstancedStaticMeshComponent::CalcBoundsImpl(const FTransform& BoundTransform, const bool bForNavigation) const
 {
 	if (GetStaticMesh() && PerInstanceSMData.Num() > 0)
@@ -2669,6 +2670,11 @@ void UInstancedStaticMeshComponent::UpdateBounds()
 	{
 		Super::UpdateBounds();
 	}
+
+	if (bCanEverAffectNavigation && IsRegistered())
+	{
+		CalcAndCacheNavigationBounds();
+	}
 }
 
 #if WITH_EDITOR
@@ -3161,6 +3167,16 @@ TArray<int32> UInstancedStaticMeshComponent::AddInstancesInternal(TConstArrayVie
 
 	PerInstanceSMData.Reserve(InstanceIndex + Count);
 
+	// Perform partial navigation update if supported and explicitly requested.
+	// Note that we also test navigation relevancy but on the base class since this class
+	// implementation also take the number of instance into account which can be 0 here.
+	const bool bDoPartialNavigationUpdate = bUpdateNavigation && Super::IsNavigationRelevant() && SupportsPartialNavigationUpdate();
+	TArray<FTransform> NavigationUpdateTransforms;
+	if (bDoPartialNavigationUpdate)
+	{
+		NavigationUpdateTransforms.Reserve(InstanceTransforms.Num());
+	}
+	
 	for (const FTransform& InstanceTransform : InstanceTransforms)
 	{
 		FInstancedStaticMeshInstanceData& NewInstanceData = PerInstanceSMData.AddDefaulted_GetRef();
@@ -3174,7 +3190,7 @@ TArray<int32> UInstancedStaticMeshComponent::AddInstancesInternal(TConstArrayVie
 			NewInstanceIndices.Add(InstanceIndex);
 		}
 
-		if (bUpdateNavigation && SupportsPartialNavigationUpdate())
+		if (bDoPartialNavigationUpdate)
 		{
 			// If it's the first instance, register the component. 
 			// If there was no instance on component register, component registration was skipped because of UInstancedStaticMeshComponent::IsNavigationRelevant().
@@ -3186,8 +3202,8 @@ TArray<int32> UInstancedStaticMeshComponent::AddInstancesInternal(TConstArrayVie
 					FNavigationSystem::RegisterComponent(*this);
 				}
 			}
-			
-			PartialNavigationUpdate(InstanceIndex);
+
+			NavigationUpdateTransforms.Emplace(InstanceTransform);
 		}
 
 		if (FInstancedStaticMeshDelegates::OnInstanceIndexUpdated.IsBound())
@@ -3199,9 +3215,16 @@ TArray<int32> UInstancedStaticMeshComponent::AddInstancesInternal(TConstArrayVie
 		++InstanceIndex;
 	}
 
-	if (bUpdateNavigation && !SupportsPartialNavigationUpdate())
+	if (bUpdateNavigation)
 	{
-		FullNavigationUpdate();
+		if (bDoPartialNavigationUpdate)
+		{
+			PartialNavigationUpdates(NavigationUpdateTransforms);
+		}
+		else
+		{
+			FullNavigationUpdate();
+		}
 	}
 
 	return NewInstanceIndices;
@@ -3605,19 +3628,48 @@ void UInstancedStaticMeshComponent::OnUpdateTransform(EUpdateTransformFlags Upda
 	Super::OnUpdateTransform(UpdateTransformFlags | EUpdateTransformFlags::SkipPhysicsUpdate, Teleport);
 
 	const bool bTeleport = TeleportEnumToFlag(Teleport);
+	const bool bDoPartialNavigationUpdate = IsNavigationRelevant() && SupportsPartialNavigationUpdate();
+	const bool bUpdateBodies = bPhysicsStateCreated && !(EUpdateTransformFlags::SkipPhysicsUpdate & UpdateTransformFlags);
 
-	// Always send new transform to physics
-	if (bPhysicsStateCreated && !(EUpdateTransformFlags::SkipPhysicsUpdate & UpdateTransformFlags))
+	if (bDoPartialNavigationUpdate || bUpdateBodies)
 	{
+		TArray<FTransform> NavigationUpdateTransforms;
+		if (bDoPartialNavigationUpdate)
+		{
+			NavigationUpdateTransforms.Reserve(PerInstanceSMData.Num());
+		}
+
 		for (int32 i = 0; i < PerInstanceSMData.Num(); i++)
 		{
-			const FTransform InstanceTransform(PerInstanceSMData[i].Transform);
-			UpdateInstanceBodyTransform(i, InstanceTransform * GetComponentTransform(), bTeleport);
+			const FTransform NewInstanceTransform(FTransform(PerInstanceSMData[i].Transform) * GetComponentTransform());
+
+			// Append instance's previous and new transforms to dirty both areas
+			if (bDoPartialNavigationUpdate)
+			{
+				NavigationUpdateTransforms.Append(
+					{
+						FTransform(PerInstanceSMData[i].Transform) * PreviousComponentTransform,
+						NewInstanceTransform
+					});
+			}
+
+			// Send new transforms to physics
+			if (bUpdateBodies)
+			{
+				UpdateInstanceBodyTransform(i, NewInstanceTransform, bTeleport);
+			}
+		}
+		
+		if (bDoPartialNavigationUpdate)
+		{
+			PartialNavigationUpdates(NavigationUpdateTransforms);
 		}
 	}
 
 	// TODO: bTeleport???
 	PrimitiveInstanceDataManager.PrimitiveTransformChanged();
+
+	PreviousComponentTransform = GetComponentTransform();
 }
 
 void UInstancedStaticMeshComponent::UpdateInstanceBodyTransform(int32 InstanceIndex, const FTransform& WorldSpaceInstanceTransform, bool bTeleport)
@@ -3634,9 +3686,6 @@ void UInstancedStaticMeshComponent::UpdateInstanceBodyTransform(int32 InstanceIn
 	{
 		if (InstanceBodyInstance)
 		{
-			// Update navigation at current position (before removal)
-			PartialNavigationUpdate(InstanceIndex);
-			
 			// delete BodyInstance
 			InstanceBodyInstance->TermBody();
 			delete InstanceBodyInstance;
@@ -3647,9 +3696,6 @@ void UInstancedStaticMeshComponent::UpdateInstanceBodyTransform(int32 InstanceIn
 	{
 		if (InstanceBodyInstance)
 		{
-			// Update navigation at current position (before applying the transform)
-			PartialNavigationUpdate(InstanceIndex);
-			
 			// Update existing BodyInstance
 			InstanceBodyInstance->SetBodyTransform(WorldSpaceInstanceTransform, TeleportFlagToEnum(bTeleport));
 			InstanceBodyInstance->UpdateBodyScale(WorldSpaceInstanceTransform.GetScale3D());
@@ -3660,13 +3706,10 @@ void UInstancedStaticMeshComponent::UpdateInstanceBodyTransform(int32 InstanceIn
 			InstanceBodyInstance = new FBodyInstance();
 			InitInstanceBody(InstanceIndex, InstanceBodyInstance);
 		}
-
-		// Update navigation at new instance location
-		PartialNavigationUpdate(InstanceIndex);
 	}
 }
 
-bool UInstancedStaticMeshComponent::UpdateInstanceTransform(int32 InstanceIndex, const FTransform& NewInstanceTransform, bool bWorldSpace, bool bMarkRenderStateDirty, bool bTeleport)
+bool UInstancedStaticMeshComponent::UpdateInstanceTransform(const int32 InstanceIndex, const FTransform& NewInstanceTransform, const bool bWorldSpace, const bool bMarkRenderStateDirty, const bool bTeleport)
 {
 	if (!PerInstanceSMData.IsValidIndex(InstanceIndex))
 	{
@@ -3677,18 +3720,28 @@ bool UInstancedStaticMeshComponent::UpdateInstanceTransform(int32 InstanceIndex,
 
 	FInstancedStaticMeshInstanceData& InstanceData = PerInstanceSMData[InstanceIndex];
 
-    // TODO: Computing LocalTransform is useless when we're updating the world location for the entire mesh.
+	// TODO: Computing LocalTransform is useless when we're updating the world location for the entire mesh.
 	// Should find some way around this for performance.
-    
+
+	if (IsNavigationRelevant() && SupportsPartialNavigationUpdate())
+	{
+		// Append instance's previous and new transforms to dirty both areas
+		PartialNavigationUpdates(
+			{
+				FTransform(InstanceData.Transform) * GetComponentTransform(),
+				(bWorldSpace ? NewInstanceTransform : NewInstanceTransform * GetComponentTransform())
+			});
+	}
+
 	// Render data uses local transform of the instance
-	FTransform LocalTransform = bWorldSpace ? NewInstanceTransform.GetRelativeTransform(GetComponentTransform()) : NewInstanceTransform;
+	const FTransform LocalTransform = bWorldSpace ? NewInstanceTransform.GetRelativeTransform(GetComponentTransform()) : NewInstanceTransform;
 	InstanceData.Transform = LocalTransform.ToMatrixWithScale();
 	PrimitiveInstanceDataManager.TransformChanged(InstanceIndex);
 
 	if (bPhysicsStateCreated)
 	{
 		// Physics uses world transform of the instance
-		FTransform WorldTransform = bWorldSpace ? NewInstanceTransform : (LocalTransform * GetComponentTransform());
+		const FTransform WorldTransform = bWorldSpace ? NewInstanceTransform : (LocalTransform * GetComponentTransform());
 		UpdateInstanceBodyTransform(InstanceIndex, WorldTransform, bTeleport);
 	}
 
@@ -3716,6 +3769,13 @@ bool UInstancedStaticMeshComponent::BatchUpdateInstancesTransforms(int32 StartIn
 
 	Modify();
 
+	const bool bDoPartialNavigationUpdate = IsNavigationRelevant() && SupportsPartialNavigationUpdate();
+	TArray<FTransform> NavigationUpdateTransforms;
+	if (bDoPartialNavigationUpdate)
+	{
+		NavigationUpdateTransforms.Reserve(NewInstancesTransforms.Num());
+	}
+
 	for (int32 Index = 0; Index < NewInstancesTransforms.Num(); Index++)
 	{
 		const int32 InstanceIndex = StartInstanceIndex + Index;
@@ -3726,6 +3786,16 @@ bool UInstancedStaticMeshComponent::BatchUpdateInstancesTransforms(int32 StartIn
 		FInstancedStaticMeshInstanceData& InstanceData = PerInstanceSMData[InstanceIndex];
 		FMatrix& PrevInstanceData = PerInstancePrevTransform[InstanceIndex];
 
+		// Append instance's previous and new transforms to dirty both areas
+		if (bDoPartialNavigationUpdate)
+		{
+			NavigationUpdateTransforms.Append(
+			{
+				FTransform(InstanceData.Transform) * GetComponentTransform(),
+				(bWorldSpace ? NewInstanceTransform : NewInstanceTransform * GetComponentTransform())
+			});
+		}
+		
 		// TODO: Computing LocalTransform is useless when we're updating the world location for the entire mesh.
 		// Should find some way around this for performance.
 
@@ -3745,6 +3815,11 @@ bool UInstancedStaticMeshComponent::BatchUpdateInstancesTransforms(int32 StartIn
 		}
 	}
 
+	if (bDoPartialNavigationUpdate)
+	{
+		PartialNavigationUpdates(NavigationUpdateTransforms);
+	}
+
 	if (bMarkRenderStateDirty)
 	{
 		MarkRenderStateDirty();
@@ -3970,11 +4045,28 @@ bool UInstancedStaticMeshComponent::BatchUpdateInstancesTransformsInternal(int32
 
 	Modify();
 
+	const bool bDoPartialNavigationUpdate = IsNavigationRelevant() && SupportsPartialNavigationUpdate();
+	TArray<FTransform> NavigationUpdateTransforms;
+	if (bDoPartialNavigationUpdate)
+	{
+		NavigationUpdateTransforms.Reserve(NewInstancesTransforms.Num());
+	}
+
 	int32 InstanceIndex = StartInstanceIndex;
 	for (const FTransform& NewInstanceTransform : NewInstancesTransforms)
 	{
 		FInstancedStaticMeshInstanceData& InstanceData = PerInstanceSMData[InstanceIndex];
 
+		// Append instance's previous and new transforms to dirty both areas
+		if (bDoPartialNavigationUpdate)
+		{
+			NavigationUpdateTransforms.Append(
+			{
+				FTransform(InstanceData.Transform) * GetComponentTransform(),
+				(bWorldSpace ? NewInstanceTransform : NewInstanceTransform * GetComponentTransform())
+			});
+		}
+		
 		// TODO: Computing LocalTransform is useless when we're updating the world location for the entire mesh.
 		// Should find some way around this for performance.
 
@@ -3993,6 +4085,11 @@ bool UInstancedStaticMeshComponent::BatchUpdateInstancesTransformsInternal(int32
 		InstanceIndex++;
 	}
 
+	if (bDoPartialNavigationUpdate)
+	{
+		PartialNavigationUpdates(NavigationUpdateTransforms);
+	}
+
 	if (bMarkRenderStateDirty)
 	{
 		MarkRenderStateDirty();
@@ -4010,11 +4107,28 @@ bool UInstancedStaticMeshComponent::BatchUpdateInstancesTransform(int32 StartIns
 
 	Modify();
 
+	const bool bDoPartialNavigationUpdate = IsNavigationRelevant() && SupportsPartialNavigationUpdate();
+	TArray<FTransform> NavigationUpdateTransforms;
+	if (bDoPartialNavigationUpdate)
+	{
+		NavigationUpdateTransforms.Reserve(NumInstances);
+	}
+	
 	int32 EndInstanceIndex = StartInstanceIndex + NumInstances;
 	for(int32 InstanceIndex = StartInstanceIndex; InstanceIndex < EndInstanceIndex; ++InstanceIndex)
 	{
 		FInstancedStaticMeshInstanceData& InstanceData = PerInstanceSMData[InstanceIndex];
 
+		// Append instance's previous and new transforms to dirty both areas
+		if (bDoPartialNavigationUpdate)
+		{
+			NavigationUpdateTransforms.Append(
+			{
+				FTransform(InstanceData.Transform) * GetComponentTransform(),
+				(bWorldSpace ? NewInstancesTransform : NewInstancesTransform * GetComponentTransform())
+			});
+		}
+		
 		// TODO: Computing LocalTransform is useless when we're updating the world location for the entire mesh.
 		// Should find some way around this for performance.
 
@@ -4031,6 +4145,11 @@ bool UInstancedStaticMeshComponent::BatchUpdateInstancesTransform(int32 StartIns
 		}
 	}
 
+	if (bDoPartialNavigationUpdate)
+	{
+		PartialNavigationUpdates(NavigationUpdateTransforms);
+	}
+
 	if(bMarkRenderStateDirty)
 	{
 		MarkRenderStateDirty();
@@ -4048,11 +4167,28 @@ bool UInstancedStaticMeshComponent::BatchUpdateInstancesData(int32 StartInstance
 
 	Modify();
 
+	const bool bDoPartialNavigationUpdate = IsNavigationRelevant() && SupportsPartialNavigationUpdate();
+	TArray<FTransform> NavigationUpdateTransforms;
+	if (bDoPartialNavigationUpdate)
+	{
+		NavigationUpdateTransforms.Reserve(NumInstances);
+	}
+	
 	for (int32 Index = 0; Index < NumInstances; ++Index)
 	{
 		int32 InstanceIndex = StartInstanceIndex + Index;
 		FInstancedStaticMeshInstanceData& InstanceData = PerInstanceSMData[InstanceIndex];
 
+		// Append instance's previous and new transforms to dirty both areas
+		if (bDoPartialNavigationUpdate)
+		{
+			NavigationUpdateTransforms.Append(
+			{
+				FTransform(InstanceData.Transform) * GetComponentTransform(),
+				FTransform(StartInstanceData[Index].Transform) * GetComponentTransform()
+			});
+		}
+		
 		InstanceData = StartInstanceData[Index];
 		PrimitiveInstanceDataManager.TransformChanged(InstanceIndex);
 
@@ -4064,6 +4200,11 @@ bool UInstancedStaticMeshComponent::BatchUpdateInstancesData(int32 StartInstance
 		}
 	}
 
+	if (bDoPartialNavigationUpdate)
+	{
+		PartialNavigationUpdates(NavigationUpdateTransforms);
+	}
+
 	if (bMarkRenderStateDirty)
 	{
 		MarkRenderStateDirty();
@@ -4363,6 +4504,8 @@ void UInstancedStaticMeshComponent::InitPerInstanceRenderData(bool InitializeFro
 void UInstancedStaticMeshComponent::OnRegister()
 {
 	Super::OnRegister();
+	
+	PreviousComponentTransform = GetComponentTransform();
 
 	if (FApp::CanEverRender() && !HasAnyFlags(RF_ClassDefaultObject | RF_ArchetypeObject))
 	{
@@ -4738,14 +4881,34 @@ void UInstancedStaticMeshComponent::PartialNavigationUpdate(int32 InstanceIdx)
 		return;
 	}
 
+	FTransform InstanceTransform;
+	if (GetInstanceTransform(InstanceIdx, InstanceTransform, /*bWorldSpace*/true))
+	{
+		PartialNavigationUpdates({InstanceTransform});
+	}
+}
+
+void UInstancedStaticMeshComponent::PartialNavigationUpdates(const TConstArrayView<FTransform> InstanceTransforms)
+{
+	if (!IsNavigationRelevant() || InstanceTransforms.IsEmpty())
+	{
+		return;
+	}
+
 	const FBox InstanceBounds = GetInstanceNavigationBounds();
 	if (InstanceBounds.IsValid)
 	{
-		FTransform InstanceTransform;
-		if (GetInstanceTransform(InstanceIdx, InstanceTransform, /*bWorldSpace*/true))
+		// Update cached navigation bounds
+		CalcAndCacheNavigationBounds();
+
+		// Dirty areas around the modified instances
+		TArray<FBox> DirtyAreas;
+		DirtyAreas.Reserve(InstanceTransforms.Num());
+		for (const FTransform& InstanceTransform : InstanceTransforms)
 		{
-			FNavigationSystem::OnComponentBoundsChanged(*this, GetNavigationBounds(), InstanceBounds.TransformBy(InstanceTransform));
+			DirtyAreas.Emplace(InstanceBounds.TransformBy(InstanceTransform));
 		}
+		FNavigationSystem::OnObjectBoundsChanged(*this, GetNavigationBounds(), DirtyAreas);
 	}
 }
 
@@ -4948,7 +5111,8 @@ void UInstancedStaticMeshComponent::GetNavigationData(FNavigationRelevantData& D
 
 FBox UInstancedStaticMeshComponent::GetNavigationBounds() const
 {
-	return CalcBoundsImpl(GetComponentTransform(), /*bForNavigation*/true).GetBox();
+	ensureMsgf(NavigationBounds.IsValid, TEXT("Navigation bounds should be calculated before being requested."));
+	return NavigationBounds;
 }
 
 bool UInstancedStaticMeshComponent::IsNavigationRelevant() const
@@ -4956,6 +5120,14 @@ bool UInstancedStaticMeshComponent::IsNavigationRelevant() const
 	return GetInstanceCount() > 0 && Super::IsNavigationRelevant();
 }
 
+bool UInstancedStaticMeshComponent::ShouldSkipDirtyAreaOnAddOrRemove() const
+{
+	// If partial navigation updates are supported then we don't want to dirty the
+	// whole area covered by the navigation bounds when added added to the navigation octree,
+	// instead we use the partial update to push the list of dirty areas.
+	return SupportsPartialNavigationUpdate();
+}
+
 FBox UInstancedStaticMeshComponent::GetInstanceNavigationBounds() const
 {
 	if (const UStaticMesh* Mesh = GetStaticMesh())
diff --git a/Engine/Source/Runtime/NavigationSystem/Private/NavigationDataHandler.cpp b/Engine/Source/Runtime/NavigationSystem/Private/NavigationDataHandler.cpp
--- a/Engine/Source/Runtime/NavigationSystem/Private/NavigationDataHandler.cpp
+++ b/Engine/Source/Runtime/NavigationSystem/Private/NavigationDataHandler.cpp
@@ -47,9 +47,13 @@ void FNavigationDataHandler::RemoveNavOctreeElementId(const FOctreeElementId2& E
 	if (ensure(OctreeController.IsValidElement(ElementId)))
 	{
 		const FNavigationOctreeElement& ElementData = OctreeController.NavOctree->GetElementById(ElementId);
-		const int32 DirtyFlag = UE::NavigationHelper::Private::GetDirtyFlag(UpdateFlags, ElementData.Data->GetDirtyFlag());
-		// mark area occupied by given actor as dirty
-		DirtyAreasController.AddArea(ElementData.Bounds.GetBox(), DirtyFlag, [&ElementData] { return ElementData.Data->SourceObject.Get(); }, nullptr, "Remove from navoctree");
+		// mark area occupied by given element as dirty except if explicitly set to skip this default behavior
+		if (!ElementData.Data->bShouldSkipDirtyAreaOnAddOrRemove)
+		{
+			const int32 DirtyFlag = UE::NavigationHelper::Private::GetDirtyFlag(UpdateFlags, ElementData.Data->GetDirtyFlag());
+			DirtyAreasController.AddArea(ElementData.Bounds.GetBox(), DirtyFlag, [&ElementData] { return ElementData.Data->SourceObject.Get(); }, nullptr, "Remove from navoctree");
+		}
+
 		OctreeController.NavOctree->RemoveNode(ElementId);
 	}
 }
@@ -173,7 +177,8 @@ void FNavigationDataHandler::AddElementToNavOctree(const FNavigationDirtyElement
 		OctreeController.NavOctree->AddNode(ElementOwner, DirtyElement.NavInterface, ElementBounds, GeneratedData);
 	}
 
-	if (!GeneratedData.IsEmpty())
+	// mark area occupied by given element as dirty except if explicitly set to skip this default behavior
+	if (!GeneratedData.IsEmpty() && !GeneratedData.Data->bShouldSkipDirtyAreaOnAddOrRemove)
 	{
 		const int32 DirtyFlag = DirtyElement.FlagsOverride ? DirtyElement.FlagsOverride : GeneratedData.Data->GetDirtyFlag();
 		DirtyAreasController.AddArea(GeneratedData.Bounds.GetBox(), DirtyFlag, [&ElementOwner] { return ElementOwner; }, &DirtyElement, "Addition to navoctree");
@@ -357,22 +362,22 @@ void FNavigationDataHandler::UpdateNavOctreeParentChain(UObject& ElementOwner, b
 	}
 }
 
-bool FNavigationDataHandler::UpdateNavOctreeElementBounds(UObject& Object, const FBox& NewBounds, const FBox& DirtyArea)
+bool FNavigationDataHandler::UpdateNavOctreeElementBounds(UObject& Object, const FBox& NewBounds, const TConstArrayView<FBox> DirtyAreas)
 {
 	const FOctreeElementId2* ElementId = OctreeController.GetObjectsNavOctreeId(Object);
 	if (ElementId != nullptr && ensure(OctreeController.IsValidElement(*ElementId)))
 	{
 		OctreeController.NavOctree->UpdateNode(*ElementId, NewBounds);
 
-		// Add dirty area
-		if (DirtyArea.IsValid)
+		// Dirty areas
+		if (DirtyAreas.Num() > 0)
 		{
 			// Refresh ElementId since object may be stored in a different node after updating bounds
 			ElementId = OctreeController.GetObjectsNavOctreeId(Object);
 			if (ElementId != nullptr && ensure(OctreeController.IsValidElement(*ElementId)))
 			{
 				const FNavigationOctreeElement& ElementData = OctreeController.NavOctree->GetElementById(*ElementId);
-				DirtyAreasController.AddArea(DirtyArea, ElementData.Data->GetDirtyFlag(), [&Object] { return &Object; }, nullptr, "Bounds change");
+				DirtyAreasController.AddAreas(DirtyAreas, ElementData.Data->GetDirtyFlag(), [&Object] { return &Object; }, nullptr, "Bounds change");
 			}
 		}
 
@@ -381,6 +386,11 @@ bool FNavigationDataHandler::UpdateNavOctreeElementBounds(UObject& Object, const
 
 	return false;
 }
+
+bool FNavigationDataHandler::UpdateNavOctreeElementBounds(UObject& Object, const FBox& NewBounds, const FBox& DirtyArea)
+{
+	return UpdateNavOctreeElementBounds(Object, NewBounds,  TConstArrayView<FBox>{DirtyArea});
+}
 	
 void FNavigationDataHandler::FindElementsInNavOctree(const FBox& QueryBox, const FNavigationOctreeFilter& Filter, TArray<FNavigationOctreeElement>& Elements)
 {
diff --git a/Engine/Source/Runtime/NavigationSystem/Private/NavigationDirtyAreasController.cpp b/Engine/Source/Runtime/NavigationSystem/Private/NavigationDirtyAreasController.cpp
--- a/Engine/Source/Runtime/NavigationSystem/Private/NavigationDirtyAreasController.cpp
+++ b/Engine/Source/Runtime/NavigationSystem/Private/NavigationDirtyAreasController.cpp
@@ -50,26 +50,20 @@ void FNavigationDirtyAreasController::Tick(const float DeltaSeconds, const TArra
 
 void FNavigationDirtyAreasController::AddArea(const FBox& NewArea, const int32 Flags, const TFunction<UObject*()>& ObjectProviderFunc /*= nullptr*/,
 	const FNavigationDirtyElement* DirtyElement /*= nullptr*/, const FName& DebugReason /*= NAME_None*/)
+{
+	AddAreas({NewArea}, Flags, ObjectProviderFunc, DirtyElement, DebugReason);
+}
+
+void FNavigationDirtyAreasController::AddAreas(const TConstArrayView<FBox> NewAreas, const int32 Flags, const TFunction<UObject*()>& ObjectProviderFunc, const FNavigationDirtyElement* DirtyElement, const FName& DebugReason)
 {
 #if !UE_BUILD_SHIPPING
 	// always keep track of reported areas even when filtered out by invalid area as long as flags are valid
 	bDirtyAreasReportedWhileAccumulationLocked = bDirtyAreasReportedWhileAccumulationLocked || (Flags > 0 && !bCanAccumulateDirtyAreas);
-#endif // !UE_BUILD_SHIPPING
 
-	if (!NewArea.IsValid)
-	{
-		UE_LOG(LogNavigationDirtyArea, Warning, TEXT("Skipping dirty area creation because of invalid bounds (object: %s, from: %s)"),
-			*GetFullNameSafe(ObjectProviderFunc ? ObjectProviderFunc() : nullptr), *DebugReason.ToString());
-		return;
-	}
+	checkf(NewAreas.Num() > 0, TEXT("All callers of this method are expected to provide at least one area."));
+#endif // !UE_BUILD_SHIPPING
 
-	const FVector2D BoundsSize(NewArea.GetSize());
-	if (BoundsSize.IsNearlyZero())
-	{
-		UE_LOG(LogNavigationDirtyArea, Warning, TEXT("Skipping dirty area creation because of empty bounds (object: %s, from: %s)"),
-			*GetFullNameSafe(ObjectProviderFunc ? ObjectProviderFunc() : nullptr), *DebugReason.ToString());
-		return;
-	}
+	UObject* SourceObject = ObjectProviderFunc ? ObjectProviderFunc() : nullptr;
 
 	if (bUseWorldPartitionedDynamicMode)
 	{
@@ -78,62 +72,78 @@ void FNavigationDirtyAreasController::AddArea(const FBox& NewArea, const int32 F
 		//  If there is no visibility change, the change is not from loading/unloading a cell (dirtiness must be applied)
 		
 		// ObjectProviderFunc() is not always providing a valid object.
-		if (const bool bIsFromVisibilityChange = (DirtyElement && DirtyElement->bIsFromVisibilityChange) || (ObjectProviderFunc && FNavigationSystem::IsLevelVisibilityChanging(ObjectProviderFunc())))
+		if (const bool bIsFromVisibilityChange = (DirtyElement && DirtyElement->bIsFromVisibilityChange) || (SourceObject && FNavigationSystem::IsLevelVisibilityChanging(SourceObject)))
 		{
 			// If the area is from the addition or removal of objects caused by level loading/unloading and it's already in the base navmesh ignore the dirtiness.
-			if (const bool bIsIncludedInBaseNavmesh = (DirtyElement && DirtyElement->bIsInBaseNavmesh) || (ObjectProviderFunc && FNavigationSystem::IsInBaseNavmesh(ObjectProviderFunc())))
+			if (const bool bIsIncludedInBaseNavmesh = (DirtyElement && DirtyElement->bIsInBaseNavmesh) || (SourceObject && FNavigationSystem::IsInBaseNavmesh(SourceObject)))
 			{
 				UE_LOG(LogNavigationDirtyArea, VeryVerbose, TEXT("Ignoring dirtyness (visibility changed and in base navmesh). (object: %s from: %s)"),
-					*GetFullNameSafe(ObjectProviderFunc ? ObjectProviderFunc() : nullptr), *DebugReason.ToString());
+					*GetFullNameSafe(SourceObject), *DebugReason.ToString());
 				return;
 			}
-		}		
+		}
 	}
 
-	if (ShouldSkipObjectPredicate.IsBound() && ObjectProviderFunc)
+	if (ShouldSkipObjectPredicate.IsBound() && SourceObject)
 	{
-		UObject* Object = ObjectProviderFunc();
-		if (Object && ShouldSkipObjectPredicate.Execute(*Object))
+		if (ShouldSkipObjectPredicate.Execute(*SourceObject))
 		{
 			return;
 		}
 	}
 
-#if !UE_BUILD_SHIPPING
-	auto DumpExtraInfo = [ObjectProviderFunc, DebugReason, BoundsSize, NewArea]() {
-		FString ObjectInfo;
-		const UObject* Object = nullptr;
-		if (ObjectProviderFunc)
+	int32 NumInvalidBounds = 0;
+	int32 NumEmptyBounds = 0;
+	for (const FBox& NewArea : NewAreas)
+	{
+		if (!NewArea.IsValid)
 		{
-			Object = ObjectProviderFunc();
-			if (const UObject* ObjectOwner = (Object != nullptr ? Object->GetOuter() : nullptr))
-			{
-				UE_VLOG_BOX(ObjectOwner, LogNavigationDirtyArea, Log, NewArea, FColor::Red, TEXT(""));
-				ObjectInfo = FString::Printf(TEXT(" | Element's owner: %s"), *GetFullNameSafe(ObjectOwner));
-			}
+			NumInvalidBounds++;
+			continue;
 		}
 
-		return FString::Printf(TEXT("From: %s | Object: %s %s | Bounds: %s"),
-			*DebugReason.ToString(),
-			*GetFullNameSafe(Object),
-			*ObjectInfo,
-			*BoundsSize.ToString());
-	};
+		const FVector2D BoundsSize(NewArea.GetSize());
+		if (BoundsSize.IsNearlyZero())
+		{
+			NumEmptyBounds++;
+			continue;
+		}
 
-	if (ShouldReportOversizedDirtyArea() && BoundsSize.GetMax() > DirtyAreaWarningSizeThreshold)
-	{
-		UE_LOG(LogNavigationDirtyArea, Warning, TEXT("Adding an oversized dirty area: %s | Threshold: %.2f"), *DumpExtraInfo(), DirtyAreaWarningSizeThreshold);
-	}
-	else
-	{
-		UE_LOG(LogNavigationDirtyArea, VeryVerbose, TEXT("Adding dirty area object: %s"), *DumpExtraInfo());
-	}
+#if !UE_BUILD_SHIPPING
+		auto DumpExtraInfo = [SourceObject, DebugReason, BoundsSize, NewArea]() {
+				FString ObjectInfo;
+				if (const UObject* ObjectOwner = (SourceObject != nullptr ? SourceObject->GetOuter() : nullptr))
+				{
+					UE_VLOG_BOX(ObjectOwner, LogNavigationDirtyArea, Log, NewArea, FColor::Red, TEXT(""));
+					ObjectInfo = FString::Printf(TEXT(" | Element's owner: %s"), *GetFullNameSafe(ObjectOwner));
+				}
+
+				return FString::Printf(TEXT("From: %s | Object: %s %s | Bounds: %s"),
+					*DebugReason.ToString(),
+					*GetFullNameSafe(SourceObject),
+					*ObjectInfo,
+					*BoundsSize.ToString());
+		};
+
+		if (ShouldReportOversizedDirtyArea() && BoundsSize.GetMax() > DirtyAreaWarningSizeThreshold)
+		{
+			UE_LOG(LogNavigationDirtyArea, Warning, TEXT("Adding an oversized dirty area: %s | Threshold: %.2f"), *DumpExtraInfo(), DirtyAreaWarningSizeThreshold);
+		}
+		else
+		{
+			UE_LOG(LogNavigationDirtyArea, VeryVerbose, TEXT("Adding dirty area object: %s"), *DumpExtraInfo());
+		}
 #endif // !UE_BUILD_SHIPPING
 
-	if (Flags > 0 && bCanAccumulateDirtyAreas)
-	{
-		DirtyAreas.Add(FNavigationDirtyArea(NewArea, Flags, ObjectProviderFunc ? ObjectProviderFunc() : nullptr));
+		if (Flags > 0 && bCanAccumulateDirtyAreas)
+		{
+			DirtyAreas.Add(FNavigationDirtyArea(NewArea, Flags, SourceObject));
+		}
 	}
+	
+	UE_CLOG(NumInvalidBounds > 0 || NumEmptyBounds > 0, LogNavigationDirtyArea, Warning,
+		TEXT("Skipped some dirty area creation due to: %d invalid bounds, %d empty bounds (object: %s, from: %s)"),
+		NumInvalidBounds, NumEmptyBounds, *GetFullNameSafe(SourceObject), *DebugReason.ToString());
 }
 
 void FNavigationDirtyAreasController::OnNavigationBuildLocked()
diff --git a/Engine/Source/Runtime/NavigationSystem/Private/NavigationOctree.cpp b/Engine/Source/Runtime/NavigationSystem/Private/NavigationOctree.cpp
--- a/Engine/Source/Runtime/NavigationSystem/Private/NavigationOctree.cpp
+++ b/Engine/Source/Runtime/NavigationSystem/Private/NavigationOctree.cpp
@@ -155,6 +155,8 @@ void FNavigationOctree::AddNode(UObject* ElementOb, INavRelevantInterface* NavEl
 	{
 		checkf(ElementOb, TEXT("We assume NavElement is ElementOb already cast"));
 
+		Element.Data->bShouldSkipDirtyAreaOnAddOrRemove = NavElement && NavElement->ShouldSkipDirtyAreaOnAddOrRemove();
+
 		const bool bDoInstantGathering = !IsLazyGathering(*NavElement);
 
 		if (bGatherGeometry)
diff --git a/Engine/Source/Runtime/NavigationSystem/Private/NavigationSystem.cpp b/Engine/Source/Runtime/NavigationSystem/Private/NavigationSystem.cpp
--- a/Engine/Source/Runtime/NavigationSystem/Private/NavigationSystem.cpp
+++ b/Engine/Source/Runtime/NavigationSystem/Private/NavigationSystem.cpp
@@ -577,6 +577,13 @@ UNavigationSystemV1::UNavigationSystemV1(const FObjectInitializer& ObjectInitial
 			UNavigationSystemBase::RegisterNavRelevantObjectDelegate().BindLambda([](UObject& Object) { UNavigationSystemV1::OnNavRelevantObjectRegistered(Object); });
 			UNavigationSystemBase::UpdateNavRelevantObjectDelegate().BindStatic(&UNavigationSystemV1::UpdateNavRelevantObjectInNavOctree);
 			UNavigationSystemBase::UnregisterNavRelevantObjectDelegate().BindLambda([](UObject& Object) { UNavigationSystemV1::OnNavRelevantObjectUnregistered(Object); });
+			UNavigationSystemBase::OnObjectBoundsChangedDelegate().BindLambda([](UObject& Object, const FBox& NewBounds, const TConstArrayView<FBox> DirtyAreas)
+				{
+					if (UNavigationSystemV1* NavSys = FNavigationSystem::GetCurrent<UNavigationSystemV1>(Object.GetWorld()))
+					{
+						NavSys->UpdateNavOctreeElementBounds(Object, NewBounds, DirtyAreas);
+					}
+				});
 
 			UNavigationSystemBase::UpdateActorDataDelegate().BindStatic(&UNavigationSystemV1::UpdateActorInNavOctree);
 			UNavigationSystemBase::UpdateComponentDataDelegate().BindStatic(&UNavigationSystemV1::UpdateComponentInNavOctree);
@@ -634,7 +641,7 @@ UNavigationSystemV1::UNavigationSystemV1(const FObjectInitializer& ObjectInitial
 				UNavigationSystemV1* NavSys = FNavigationSystem::GetCurrent<UNavigationSystemV1>(Comp.GetWorld());
 				if (NavSys)
 				{
-					NavSys->UpdateNavOctreeElementBounds(Comp, NewBounds, DirtyArea);
+					NavSys->UpdateNavOctreeElementBounds(Comp, NewBounds, {DirtyArea});
 				}
 			});
 			//UNavigationSystemBase::GetNavDataForPropsDelegate();
@@ -3467,13 +3474,13 @@ void UNavigationSystemV1::UpdateNavOctreeParentChain(UObject* ElementOwner, bool
 bool UNavigationSystemV1::UpdateNavOctreeElementBounds(UActorComponent* Comp, const FBox& NewBounds, const FBox& DirtyArea)
 {
 	return Comp
-		? FNavigationDataHandler(DefaultOctreeController, DefaultDirtyAreasController).UpdateNavOctreeElementBounds(*Comp, NewBounds, DirtyArea)
+		? UpdateNavOctreeElementBounds(*Comp, NewBounds, { DirtyArea })
 		: false;
 }
 
-bool UNavigationSystemV1::UpdateNavOctreeElementBounds(UObject& Object, const FBox& NewBounds, const FBox& DirtyArea)
+bool UNavigationSystemV1::UpdateNavOctreeElementBounds(UObject& Object, const FBox& NewBounds, TConstArrayView<FBox> DirtyAreas)
 {
-	return FNavigationDataHandler(DefaultOctreeController, DefaultDirtyAreasController).UpdateNavOctreeElementBounds(Object, NewBounds, DirtyArea);
+	return FNavigationDataHandler(DefaultOctreeController, DefaultDirtyAreasController).UpdateNavOctreeElementBounds(Object, NewBounds, DirtyAreas);
 }
 
 bool UNavigationSystemV1::ReplaceAreaInOctreeData(const UObject& Object, TSubclassOf<UNavArea> OldArea, TSubclassOf<UNavArea> NewArea, bool bReplaceChildClasses)
diff --git a/Engine/Source/Runtime/NavigationSystem/Public/NavigationDataHandler.h b/Engine/Source/Runtime/NavigationSystem/Public/NavigationDataHandler.h
--- a/Engine/Source/Runtime/NavigationSystem/Public/NavigationDataHandler.h
+++ b/Engine/Source/Runtime/NavigationSystem/Public/NavigationDataHandler.h
@@ -30,7 +30,7 @@ struct FNavigationDataHandler
 	NAVIGATIONSYSTEM_API bool UnregisterNavOctreeElement(UObject& ElementOwner, INavRelevantInterface& ElementInterface, int32 UpdateFlags);
 	NAVIGATIONSYSTEM_API void UpdateNavOctreeElement(UObject& ElementOwner, INavRelevantInterface& ElementInterface, int32 UpdateFlags);
 	NAVIGATIONSYSTEM_API void UpdateNavOctreeParentChain(UObject& ElementOwner, bool bSkipElementOwnerUpdate = false);
-	NAVIGATIONSYSTEM_API bool UpdateNavOctreeElementBounds(UObject& Object, const FBox& NewBounds, const FBox& DirtyArea);
+	NAVIGATIONSYSTEM_API bool UpdateNavOctreeElementBounds(UObject& Object, const FBox& NewBounds, const TConstArrayView<FBox> DirtyAreas);
 	NAVIGATIONSYSTEM_API void FindElementsInNavOctree(const FBox& QueryBox, const FNavigationOctreeFilter& Filter, TArray<FNavigationOctreeElement>& Elements);
 	NAVIGATIONSYSTEM_API bool ReplaceAreaInOctreeData(const UObject& Object, TSubclassOf<UNavArea> OldArea, TSubclassOf<UNavArea> NewArea, bool bReplaceChildClasses);
 	NAVIGATIONSYSTEM_API void AddLevelCollisionToOctree(ULevel& Level);
@@ -38,4 +38,7 @@ struct FNavigationDataHandler
 	NAVIGATIONSYSTEM_API void UpdateActorAndComponentsInNavOctree(AActor& Actor);
 	NAVIGATIONSYSTEM_API void ProcessPendingOctreeUpdates();
 	NAVIGATIONSYSTEM_API void DemandLazyDataGathering(FNavigationRelevantData& ElementData);
+	
+	UE_DEPRECATED(5.4, "Use the overloaded version taking a list of areas as parameter instead.")
+	NAVIGATIONSYSTEM_API bool UpdateNavOctreeElementBounds(UObject& Object, const FBox& NewBounds, const FBox& DirtyArea);
 };
diff --git a/Engine/Source/Runtime/NavigationSystem/Public/NavigationDirtyAreasController.h b/Engine/Source/Runtime/NavigationSystem/Public/NavigationDirtyAreasController.h
--- a/Engine/Source/Runtime/NavigationSystem/Public/NavigationDirtyAreasController.h
+++ b/Engine/Source/Runtime/NavigationSystem/Public/NavigationDirtyAreasController.h
@@ -58,14 +58,27 @@ public:
 	/** Add a dirty area to the queue based on the provided bounds and flags.
 	 * Bounds must be valid and non empty otherwise the request will be ignored and a warning reported.
 	 * Accumulation must be allowed and flags valid otherwise the add is ignored.
-	 *	@param NewArea Bounding box of the affected area
-	 *	@param Flags Indicates the type of modification applied to the area
-	 *	@param ObjectProviderFunc Optional function to retrieve source object that can be use for error reporting and navmesh exclusion
-	 *	@param DirtyElement Optional dirty element
-	 *	@param DebugReason Source of the new area
+	 * @param NewArea Bounding box of the affected area
+	 * @param Flags Indicates the type of modification applied to the area
+	 * @param ObjectProviderFunc Optional function to retrieve source object that can be use for error reporting and navmesh exclusion
+	 * @param DirtyElement Optional dirty element
+	 * @param DebugReason Source of the new area
 	 */
 	NAVIGATIONSYSTEM_API void AddArea(const FBox& NewArea, const int32 Flags, const TFunction<UObject*()>& ObjectProviderFunc = nullptr,
 		const FNavigationDirtyElement* DirtyElement = nullptr, const FName& DebugReason = NAME_None);
+
+	/** Add non empty list of dirty areas to the queue based on the provided bounds and flags.
+	 * Bounds must be valid and non empty otherwise the request will be ignored and a warning reported.
+	 * Accumulation must be allowed and flags valid otherwise the add is ignored.
+	 * A check will be triggered if an empty array is provided.
+	 * @param NewAreas Array of bounding boxes of the affected areas
+	 * @param Flags Indicates the type of modification applied to the area
+	 * @param ObjectProviderFunc Optional function to retrieve source object that can be use for error reporting and navmesh exclusion
+	 * @param DirtyElement Optional dirty element
+	 * @param DebugReason Source of the new area
+	 */
+	NAVIGATIONSYSTEM_API void AddAreas(const TConstArrayView<FBox> NewAreas, const int32 Flags, const TFunction<UObject*()>& ObjectProviderFunc = nullptr,
+		const FNavigationDirtyElement* DirtyElement = nullptr, const FName& DebugReason = NAME_None);
 	
 	bool IsDirty() const { return GetNumDirtyAreas() > 0; }
 	int32 GetNumDirtyAreas() const { return DirtyAreas.Num(); }
diff --git a/Engine/Source/Runtime/NavigationSystem/Public/NavigationSystem.h b/Engine/Source/Runtime/NavigationSystem/Public/NavigationSystem.h
--- a/Engine/Source/Runtime/NavigationSystem/Public/NavigationSystem.h
+++ b/Engine/Source/Runtime/NavigationSystem/Public/NavigationSystem.h
@@ -957,14 +957,14 @@ public:
 	/** force updating parent node and all its children */
 	NAVIGATIONSYSTEM_API void UpdateNavOctreeParentChain(UObject* ElementOwner, bool bSkipElementOwnerUpdate = false);
 
-	UE_DEPRECATED(5.4, "Use UpdateNavOctreeElementBounds with object reference.")
+	UE_DEPRECATED(5.4, "Use the overloaded version with object reference and list of dirty areas as parameters instead.")
 	/** update component bounds in navigation octree and mark only specified area as dirty, doesn't re-export component geometry */
 	NAVIGATIONSYSTEM_API bool UpdateNavOctreeElementBounds(UActorComponent* Comp, const FBox& NewBounds, const FBox& DirtyArea);
 
 	/** update object bounds in navigation octree and mark only specified area as dirty, doesn't re-export geometry */
-	NAVIGATIONSYSTEM_API bool UpdateNavOctreeElementBounds(UObject& Object, const FBox& NewBounds, const FBox& DirtyArea);
+	NAVIGATIONSYSTEM_API bool UpdateNavOctreeElementBounds(UObject& Object, const FBox& NewBounds, TConstArrayView<FBox> DirtyAreas);
 
-	/** fetched Object's data from the octree and replaces occurences of OldArea with NewArea */
+	/** fetched Object's data from the octree and replaces occurrences of OldArea with NewArea */
 	NAVIGATIONSYSTEM_API bool ReplaceAreaInOctreeData(const UObject& Object, TSubclassOf<UNavArea> OldArea, TSubclassOf<UNavArea> NewArea, bool bReplaceChildClasses = false);
 
 	//----------------------------------------------------------------------//
