Update CuePoint frame positions based on resampling of Sound Wave on cook.

SoundWave
- Implement `<<` on FSoundWaveCuePoint and use UScriptStruct serialization
- Make `TArray<FSoundWaveCuePoints> CuePoints` WITH_EDITORONLY_DATA
- On Cook of SoundWave, resample the CuePoints frame position and frame length using the resample ratio (TargetSampleRate / OriginalSampleRate).
- When cooking and loading from cook content, Serialize (read/write) a separate copy of the CuePoints into the Archive as CuePoints won't be available in package builds.

SoundSourceBus & SoundWaveProcedural
- No-Op SerializeCuePoints

#tests PIE, Local Cook

diff --git a/Engine/Source/Runtime/Engine/Classes/Sound/SoundSourceBus.h b/Engine/Source/Runtime/Engine/Classes/Sound/SoundSourceBus.h
--- a/Engine/Source/Runtime/Engine/Classes/Sound/SoundSourceBus.h
+++ b/Engine/Source/Runtime/Engine/Classes/Sound/SoundSourceBus.h
@@ -79,4 +79,7 @@ public:
 
 protected:
 	ENGINE_API void Init();
+	
+	// SourceBus doesn't represent a wav file, don't do anything when serializing cue points
+	virtual void SerializeCuePoints(FArchive& Ar, const bool bIsLoadingFromCookedArchive) {}
 };
diff --git a/Engine/Source/Runtime/Engine/Classes/Sound/SoundWave.h b/Engine/Source/Runtime/Engine/Classes/Sound/SoundWave.h
--- a/Engine/Source/Runtime/Engine/Classes/Sound/SoundWave.h
+++ b/Engine/Source/Runtime/Engine/Classes/Sound/SoundWave.h
@@ -396,6 +396,16 @@ struct FSoundWaveCuePoint
 	UPROPERTY(Category = Info, VisibleAnywhere, BlueprintReadOnly)
 	int32 FrameLength = 0;
 
+	bool IsLoopRegion() const { return bIsLoopRegion; }
+
+#if WITH_EDITORONLY_DATA
+	void ScaleFrameValues(float Factor)
+	{
+		FramePosition = FMath::FloorToInt((float)FramePosition * Factor);
+		FrameLength = FMath::FloorToInt((float)FrameLength * Factor);
+	}
+#endif // WITH_EDITORONLY_DATA
+
 	friend class USoundFactory;
 	friend class USoundWave;
 private:
@@ -489,7 +499,7 @@ public:
 	UFUNCTION(BlueprintCallable, Category = "Audio")
 	ENGINE_API void SetSoundAssetCompressionType(ESoundAssetCompressionType InSoundAssetCompressionType, bool bMarkDirty = true);
 
-	/** Filters for the cue points that are _not_ loop regions and returns those as a new array*/
+	/** Filters for the cue points that are _not_ loop regions and returns those as a new array */
 	UFUNCTION(BlueprintCallable, Category = "Audio")
 	ENGINE_API TArray<FSoundWaveCuePoint> GetCuePoints() const;
 
@@ -747,10 +757,6 @@ public:
 	UPROPERTY(Category = Info, AssetRegistrySearchable, VisibleAnywhere)
 	int32 NumChannels;
 
-	/** Cue point data parsed fro the .wav file. Contains "Loop Regions" as cue points as well! */
-	UPROPERTY(Category = Info, VisibleAnywhere, BlueprintReadOnly)
-	TArray<FSoundWaveCuePoint> CuePoints;
-
 #if WITH_EDITORONLY_DATA
 	/** Offsets into the bulk data for the source wav data */
 	UPROPERTY()
@@ -772,8 +778,14 @@ protected:
 	/** Sample rate of the imported sound wave. */
 	UPROPERTY(Category = Info, AssetRegistrySearchable, VisibleAnywhere)
 	int32 ImportedSampleRate;
+
+	/** Cue point data parsed fro the .wav file. Contains "Loop Regions" as cue points as well! */
+	UPROPERTY(Category = Info, VisibleAnywhere, BlueprintGetter = GetCuePoints)
+	TArray<FSoundWaveCuePoint> CuePoints;
 #endif
 
+	ENGINE_API virtual void SerializeCuePoints(FArchive& Ar, const bool bIsLoadingFromCookedArchive);
+
 public:
 
 	// Returns the compression quality of the sound asset.
@@ -1492,6 +1504,7 @@ public:
 	uint32 GetNumChannels() const { return NumChannels; }
 	const TArray<FSoundWaveCuePoint>& GetCuePoints() const { return CuePoints; }
 	const TArray<FSoundWaveCuePoint>& GetLoopRegions() const { return LoopRegions; }
+	void SetAllCuePoints(const TArray<FSoundWaveCuePoint>& InCuePoints);
 
 	ENGINE_API MaxChunkSizeResults GetMaxChunkSizeResults() const;
 
diff --git a/Engine/Source/Runtime/Engine/Classes/Sound/SoundWaveProcedural.h b/Engine/Source/Runtime/Engine/Classes/Sound/SoundWaveProcedural.h
--- a/Engine/Source/Runtime/Engine/Classes/Sound/SoundWaveProcedural.h
+++ b/Engine/Source/Runtime/Engine/Classes/Sound/SoundWaveProcedural.h
@@ -52,6 +52,9 @@ protected:
 	// The number of PCM samples we want to generate. This can't be larger than SamplesNeeded in GeneratePCMData callback, but can be less.
 	int32 NumSamplesToGeneratePerCallback;
 
+	// Procedural Sounds don't represent a wav file, don't do anything when serializing cue points
+	virtual void SerializeCuePoints(FArchive& Ar, const bool bIsLoadingFromCookedArchive) {}
+
 public:
 	ENGINE_API USoundWaveProcedural(const FObjectInitializer& ObjectInitializer);
 
diff --git a/Engine/Source/Runtime/Engine/Private/SoundWave.cpp b/Engine/Source/Runtime/Engine/Private/SoundWave.cpp
--- a/Engine/Source/Runtime/Engine/Private/SoundWave.cpp
+++ b/Engine/Source/Runtime/Engine/Private/SoundWave.cpp
@@ -85,6 +85,12 @@ FAutoConsoleVariableRef CvarAllowReverbForMultichannelSources(
 	TEXT("0: Disable, >0: Enable"),
 	ECVF_Default);
 
+inline FArchive& operator<<(FArchive& Ar, FSoundWaveCuePoint& CuePoint)
+{
+	FSoundWaveCuePoint::StaticStruct()->SerializeItem(Ar, &CuePoint, nullptr);
+	return Ar;
+}
+
 namespace SoundWave_Private
 {
 	static ESoundWaveLoadingBehavior GetDefaultLoadingBehaviorCVar()
@@ -178,8 +184,6 @@ void FSoundWaveData::InitializeDataFromSoundWave(USoundWave& InWave)
 	}
 	
 	SoundWaveKeyCached = FObjectKey(&InWave);
-	CuePoints = InWave.GetCuePoints();
-	LoopRegions = InWave.GetLoopRegions();
 	SampleRate = InWave.GetSampleRateForCurrentPlatform();
 	Duration = InWave.Duration;
 	NumChannels = InWave.NumChannels;
@@ -281,6 +285,23 @@ void FSoundWaveData::DiscardZerothChunkData()
 	ZerothChunkData.Empty();
 }
 
+void FSoundWaveData::SetAllCuePoints(const TArray<FSoundWaveCuePoint>& InCuePoints)
+{
+	CuePoints.Reset();
+	LoopRegions.Reset();
+	for (const FSoundWaveCuePoint& CuePoint : InCuePoints)
+	{
+		if (!CuePoint.IsLoopRegion())
+		{
+			CuePoints.Add(CuePoint);
+		}
+		else
+		{
+			LoopRegions.Add(CuePoint);
+		}
+	}
+}
+
 FSoundWaveData::MaxChunkSizeResults FSoundWaveData::GetMaxChunkSizeResults() const
 {
 	MaxChunkSizeResults Results;
@@ -1175,6 +1196,8 @@ void USoundWave::Serialize( FArchive& Ar )
 		bSupportsStreaming = true;
 	}
 
+	SerializeCuePoints(Ar, bCooked && Ar.IsLoading());
+
 	if (bCooked)
 	{
 #if WITH_EDITOR
@@ -1346,28 +1369,12 @@ void USoundWave::SetSoundAssetCompressionType(ESoundAssetCompressionType InSound
 
 TArray<FSoundWaveCuePoint> USoundWave::GetCuePoints() const
 {
-	TArray<FSoundWaveCuePoint> OutCuePoints;
-	for (const FSoundWaveCuePoint& CuePoint : CuePoints)
-	{
-		if (!CuePoint.bIsLoopRegion)
-		{
-			OutCuePoints.Add(CuePoint);
-		}
-	}
-	return OutCuePoints;
+	return SoundWaveDataPtr->GetCuePoints();
 }
 
 TArray<FSoundWaveCuePoint> USoundWave::GetLoopRegions() const
 {
-	TArray<FSoundWaveCuePoint> OutLoopRegions;
-	for (const FSoundWaveCuePoint& CuePoint : CuePoints)
-	{
-		if (CuePoint.bIsLoopRegion)
-		{
-			OutLoopRegions.Add(CuePoint);
-		}
-	}
-	return OutLoopRegions;
+	return SoundWaveDataPtr->GetLoopRegions();
 }
 
 FName USoundWave::GetRuntimeFormat() const
@@ -4395,6 +4402,63 @@ void USoundWave::CacheInheritedLoadingBehavior() const
 	}
 }
 
+void USoundWave::SerializeCuePoints(FArchive& Ar, const bool bIsLoadingFromCookedArchive)
+{
+	TArray<FSoundWaveCuePoint> CuePointsCopy;
+#if WITH_EDITORONLY_DATA
+	if (Ar.IsCooking())
+	{
+		CuePointsCopy = CuePoints;
+		float ResampleRatio = -1.0f;
+		const ITargetPlatform* CookingTarget = Ar.CookingTarget();
+		if (ensure(CookingTarget))
+		{
+			if (const FPlatformAudioCookOverrides* Overrides = FPlatformCompressionUtilities::GetCookOverrides(*CookingTarget->IniPlatformName()))
+			{
+				if (Overrides->bResampleForDevice)
+				{
+					const float TargetSampleRate = GetSampleRateForTargetPlatform(CookingTarget);
+
+					// annoyingly, read the entire imported sound wave data to successfully get the ImportedSampleRate at this point
+					TArray<uint8> RawWaveData;
+					uint32 OriginalSampleRate = 0;
+					uint16 OriginalNumChannels = 0;
+					ensureMsgf(GetImportedSoundWaveData(RawWaveData, OriginalSampleRate, OriginalNumChannels), TEXT("SerializeCuePoints: %s Failed to retrieve imported sound wave data: OriginalSampleRate = %d"), *GetName(), OriginalSampleRate);
+					
+					if (OriginalSampleRate > 0 && !FMath::IsNearlyEqual(TargetSampleRate, OriginalSampleRate))
+					{
+						ResampleRatio = TargetSampleRate / (float)OriginalSampleRate;
+					}
+
+					if (ResampleRatio > 0.0f)
+					{
+						for (FSoundWaveCuePoint& CuePoint : CuePointsCopy)
+						{
+							CuePoint.ScaleFrameValues(ResampleRatio);
+						}
+					}
+				}
+			}
+		}
+	}
+	else if (Ar.IsLoading() && !bIsLoadingFromCookedArchive)
+	{
+		SoundWaveDataPtr->SetAllCuePoints(CuePoints);
+		return;
+	}
+#endif
+
+	if (Ar.IsCooking() || bIsLoadingFromCookedArchive)
+	{
+		Ar << CuePointsCopy;
+	}
+
+	if (bIsLoadingFromCookedArchive)
+	{
+		SoundWaveDataPtr->SetAllCuePoints(CuePointsCopy);
+	}
+}
+
 #if WITH_EDITOR
 bool USoundWave::GetRedrawThumbnail() const
 {
