[shaders] enable preprocessed job cache by default; change prioritization of ShaderJobTasks in non-cook cases to background to not occupy fg worker threads (for improved responsiveness in editor)

diff --git a/Engine/Source/Runtime/Engine/Private/ShaderCompiler/ShaderCompiler.cpp b/Engine/Source/Runtime/Engine/Private/ShaderCompiler/ShaderCompiler.cpp
--- a/Engine/Source/Runtime/Engine/Private/ShaderCompiler/ShaderCompiler.cpp
+++ b/Engine/Source/Runtime/Engine/Private/ShaderCompiler/ShaderCompiler.cpp
@@ -122,7 +122,7 @@ static FAutoConsoleVariableRef CVarShaderCompilerJobCacheOverflowReducePercent(
 
 static TAutoConsoleVariable<bool> CVarPreprocessedJobCache(
 	TEXT("r.ShaderCompiler.PreprocessedJobCache"),
-	false,
+	true,
 	TEXT("If enabled will shader compile jobs will be preprocessed at submission time in the cook process (when the job is queued) and generate job input hashes based on preprocessed source."),
 	ECVF_Default
 );
@@ -1729,6 +1729,7 @@ void FShaderJobCache::SubmitJobs(const TArray<FShaderCommonCompileJobPtr>& InJob
 		{
 			for (FShaderCommonCompileJobPtr Job : InJobs)
 			{
+				UE::Tasks::ETaskPriority Prio = IsRunningCookCommandlet() ? UE::Tasks::ETaskPriority::Normal : UE::Tasks::ETaskPriority::BackgroundNormal;
 				UE::Tasks::Launch(UE_SOURCE_LOCATION, [Job, this]()
 				{
 					TRACE_CPUPROFILER_EVENT_SCOPE(ShaderJobTask);
@@ -1747,7 +1748,7 @@ void FShaderJobCache::SubmitJobs(const TArray<FShaderCommonCompileJobPtr>& InJob
 					SubmitJob(Job);
 
 					Job->TimeTaskSubmitJobs = FPlatformTime::Seconds() - TimeStart;
-				});
+				}, Prio);
 			}
 		}
 		else
