- Extra safety checking for when we have a project define WITH_EDITOR that doesn't include EditorData

diff --git a/Engine/Source/Runtime/Engine/Private/UserDefinedStruct.cpp b/Engine/Source/Runtime/Engine/Private/UserDefinedStruct.cpp
--- a/Engine/Source/Runtime/Engine/Private/UserDefinedStruct.cpp
+++ b/Engine/Source/Runtime/Engine/Private/UserDefinedStruct.cpp
@@ -224,13 +224,16 @@ void UUserDefinedStruct::OnChanged()
 FProperty* UUserDefinedStruct::CustomFindProperty(const FName Name) const
 {
 #if WITH_EDITOR
-	// If we have the editor data, check that first as it's more up to date
-	const FGuid PropertyGuid = FStructureEditorUtils::GetGuidFromPropertyName(Name);
-	FProperty* EditorProperty = PropertyGuid.IsValid() ? FStructureEditorUtils::GetPropertyByGuid(this, PropertyGuid) : FStructureEditorUtils::GetPropertyByFriendlyName(this, Name.ToString());
-	ensure(!EditorProperty || !PropertyGuid.IsValid() || PropertyGuid == FStructureEditorUtils::GetGuidForProperty(EditorProperty));
-	if (EditorProperty)
+	if (EditorData != nullptr)
 	{
-		return EditorProperty;
+		// If we have the editor data, check that first as it's more up to date
+		const FGuid PropertyGuid = FStructureEditorUtils::GetGuidFromPropertyName(Name);
+		FProperty* EditorProperty = PropertyGuid.IsValid() ? FStructureEditorUtils::GetPropertyByGuid(this, PropertyGuid) : FStructureEditorUtils::GetPropertyByFriendlyName(this, Name.ToString());
+		ensure(!EditorProperty || !PropertyGuid.IsValid() || PropertyGuid == FStructureEditorUtils::GetGuidForProperty(EditorProperty));
+		if (EditorProperty)
+		{
+			return EditorProperty;
+		}
 	}
 #endif // WITH_EDITOR
 
