[Chaos] Added the the list of removed bones to UClusterUnionComponent::OnComponentRemovedNativeEvent so listeners can use them without having to re-build it from scratch on each call.

Also updated the OnComponentAdded event to pass the removed bones as FClusterUnionBoneData to keep the API consistent.

diff --git a/Engine/Source/Runtime/Engine/Classes/PhysicsEngine/ClusterUnionComponent.h b/Engine/Source/Runtime/Engine/Classes/PhysicsEngine/ClusterUnionComponent.h
--- a/Engine/Source/Runtime/Engine/Classes/PhysicsEngine/ClusterUnionComponent.h
+++ b/Engine/Source/Runtime/Engine/Classes/PhysicsEngine/ClusterUnionComponent.h
@@ -133,10 +133,10 @@ struct FClusterUnionInitializationData
 };
 
 DECLARE_DYNAMIC_MULTICAST_DELEGATE_ThreeParams(FOnClusterUnionAddedComponent, UPrimitiveComponent*, Component, const TSet<int32>&, BoneIds, bool, bIsNew);
-DECLARE_MULTICAST_DELEGATE_FourParams(FOnClusterUnionAddedComponentNative, UPrimitiveComponent*, const TSet<FClusterUnionBoneData>& /*BoneIds*/, const TArray<int32>& /*RemovedBoneIds*/, bool /*bIsNew*/);
+DECLARE_MULTICAST_DELEGATE_FourParams(FOnClusterUnionAddedComponentNative, UPrimitiveComponent*, const TSet<FClusterUnionBoneData>& /*BoneIds*/, const TArray<FClusterUnionBoneData>& /*RemovedBoneIds*/, bool /*bIsNew*/);
 
 DECLARE_DYNAMIC_MULTICAST_DELEGATE_OneParam(FOnClusterUnionRemovedComponent, UPrimitiveComponent*, Component);
-DECLARE_MULTICAST_DELEGATE_OneParam(FOnClusterUnionRemovedComponentNative, UPrimitiveComponent*);
+DECLARE_MULTICAST_DELEGATE_TwoParams(FOnClusterUnionRemovedComponentNative, UPrimitiveComponent*, const TSet<FClusterUnionBoneData>& /*RemovedBonesData*/);
 
 DECLARE_DYNAMIC_MULTICAST_DELEGATE_TwoParams(FOnClusterUnionBoundsChanged, UClusterUnionComponent*, Component, const FBoxSphereBounds&, Bounds);
 
@@ -304,8 +304,8 @@ protected:
 	UFUNCTION()
 	ENGINE_API void ForceSetChildToParent(UPrimitiveComponent* InComponent, const TArray<int32>& BoneIds, const TArray<FTransform>& ChildToParent);
 
-	ENGINE_API void BroadcastComponentAddedEvents(UPrimitiveComponent* ChangedComponent, const TSet<FClusterUnionBoneData>& BoneIds, bool bIsNew, TArray<int32>&& RemovedBoneIDs);
-	ENGINE_API void BroadcastComponentRemovedEvents(UPrimitiveComponent* ChangedComponent);
+	ENGINE_API void BroadcastComponentAddedEvents(UPrimitiveComponent* ChangedComponent, const TSet<FClusterUnionBoneData>& BoneIds, bool bIsNew, const TArray<FClusterUnionBoneData>& RemovedBoneIDs);
+	ENGINE_API void BroadcastComponentRemovedEvents(UPrimitiveComponent* ChangedComponent, const TSet<FClusterUnionBoneData>& InRemovedBonesData);
 
 	Chaos::FClusterUnionPhysicsProxy* GetPhysicsProxy() const { return PhysicsProxy; }
 	Chaos::FClusterUnionPhysicsProxy* GetPhysicsProxy() { return PhysicsProxy; }
diff --git a/Engine/Source/Runtime/Engine/Private/PhysicsEngine/ClusterUnionComponent.cpp b/Engine/Source/Runtime/Engine/Private/PhysicsEngine/ClusterUnionComponent.cpp
--- a/Engine/Source/Runtime/Engine/Private/PhysicsEngine/ClusterUnionComponent.cpp
+++ b/Engine/Source/Runtime/Engine/Private/PhysicsEngine/ClusterUnionComponent.cpp
@@ -1026,7 +1026,7 @@ void UClusterUnionComponent::HandleAddOrModifiedClusteredComponent(const FMapped
 
 	}
 
-	TArray<int32> RemovedBoneIDs;
+	TArray<FClusterUnionBoneData> RemovedBoneIDs;
 	RemovedBoneIDs.Reserve(ComponentData.BonesData.Num());
 
 	// In the case where we need to keep the acceleration structure up to date, we need to make sure old bone ids are properly
@@ -1046,7 +1046,7 @@ void UClusterUnionComponent::HandleAddOrModifiedClusteredComponent(const FMapped
 						AccelerationStructure->RemoveElement(Handle);
 					}
 				}
-				RemovedBoneIDs.Emplace(BoneData.ID);
+				RemovedBoneIDs.Emplace(FClusterUnionBoneData(BoneData.ID, BoneData.ParticleID));
 			}
 		}
 	}
@@ -1059,7 +1059,7 @@ void UClusterUnionComponent::HandleAddOrModifiedClusteredComponent(const FMapped
 		return FClusterUnionBoneData{MappedData.Key, MappedData.Value.ParticleID };
 	});
 
-	BroadcastComponentAddedEvents(ChangedComponentData.ComponentPtr, ComponentData.BonesData, bIsNew, MoveTemp(RemovedBoneIDs));
+	BroadcastComponentAddedEvents(ChangedComponentData.ComponentPtr, ComponentData.BonesData, bIsNew, RemovedBoneIDs);
 
 	if (IsAuthority() && ComponentData.ReplicatedProxyComponent.IsValid())
 	{
@@ -1147,7 +1147,7 @@ void UClusterUnionComponent::HandleRemovedClusteredComponent(TObjectKey<UPrimiti
 
 		if (UPrimitiveComponent* ChangedComponent = RemovedComponent.ResolveObjectPtr())
 		{
-			BroadcastComponentRemovedEvents(ChangedComponent);
+			BroadcastComponentRemovedEvents(ChangedComponent, Data->BonesData);
 		}
 		PerComponentData.Remove(RemovedComponent);
 	}
@@ -1155,7 +1155,7 @@ void UClusterUnionComponent::HandleRemovedClusteredComponent(TObjectKey<UPrimiti
 	PendingComponentSync.Remove(RemovedComponent);
 }
 
-void UClusterUnionComponent::BroadcastComponentAddedEvents(UPrimitiveComponent* ChangedComponent, const TSet<FClusterUnionBoneData>& BoneIds, bool bIsNew, TArray<int32>&& RemovedBoneIDs)
+void UClusterUnionComponent::BroadcastComponentAddedEvents(UPrimitiveComponent* ChangedComponent, const TSet<FClusterUnionBoneData>& BoneIds, bool bIsNew, const TArray<FClusterUnionBoneData>& RemovedBoneIDs)
 {
 	// TODO: Add a new delegate that takes a TArray and deprecate this one
 	if (OnComponentAddedEvent.IsBound())
@@ -1166,11 +1166,11 @@ void UClusterUnionComponent::BroadcastComponentAddedEvents(UPrimitiveComponent*
 	}
 	if (OnComponentAddedNativeEvent.IsBound())
 	{
-		OnComponentAddedNativeEvent.Broadcast(ChangedComponent, BoneIds, MoveTemp(RemovedBoneIDs), bIsNew);
+		OnComponentAddedNativeEvent.Broadcast(ChangedComponent, BoneIds, RemovedBoneIDs, bIsNew);
 	}
 }
 
-void UClusterUnionComponent::BroadcastComponentRemovedEvents(UPrimitiveComponent* ChangedComponent)
+void UClusterUnionComponent::BroadcastComponentRemovedEvents(UPrimitiveComponent* ChangedComponent, const TSet<FClusterUnionBoneData>& InRemovedBonesData)
 {
 	if (OnComponentRemovedEvent.IsBound())
 	{
@@ -1178,7 +1178,7 @@ void UClusterUnionComponent::BroadcastComponentRemovedEvents(UPrimitiveComponent
 	}
 	if (OnComponentRemovedNativeEvent.IsBound())
 	{
-		OnComponentRemovedNativeEvent.Broadcast(ChangedComponent);
+		OnComponentRemovedNativeEvent.Broadcast(ChangedComponent, InRemovedBonesData);
 	}
 }
 
