Make sure that read-only cache for CVars is populated before accessing from worker threads

This fixes a race condition between the PSO preloading and ReadOnlyCVarCache initialization.

diff --git a/Engine/Source/Runtime/Engine/Private/SceneManagement.cpp b/Engine/Source/Runtime/Engine/Private/SceneManagement.cpp
--- a/Engine/Source/Runtime/Engine/Private/SceneManagement.cpp
+++ b/Engine/Source/Runtime/Engine/Private/SceneManagement.cpp
@@ -1515,6 +1515,12 @@ ELightInteractionType FLightCacheInterface::GetStaticInteraction(const FLightSce
 
 FReadOnlyCVARCache GReadOnlyCVARCache;
 
+void InitializeReadonlyCVarCache()
+{
+	GReadOnlyCVARCache.Init();
+}
+
+
 const FReadOnlyCVARCache& FReadOnlyCVARCache::Get()
 {
 	checkSlow(GReadOnlyCVARCache.bInitialized);
@@ -1554,7 +1560,7 @@ void FReadOnlyCVARCache::Init()
 	MobileSkyLightPermutation = CVarMobileSkyLightPermutation->GetValueOnAnyThread();
 	bMobileEnableNoPrecomputedLightingCSMShader = CVarMobileEnableNoPrecomputedLightingCSMShader->GetValueOnAnyThread() != 0;
 
-	const bool bShowMissmatchedLowQualityLightmapsWarning = (!bEnableLowQualityLightmaps) && (GEngine->bShouldGenerateLowQualityLightmaps_DEPRECATED);
+	const bool bShowMissmatchedLowQualityLightmapsWarning = (!bEnableLowQualityLightmaps) && (GEngine && GEngine->bShouldGenerateLowQualityLightmaps_DEPRECATED);
 	if ( bShowMissmatchedLowQualityLightmapsWarning )
 	{
 		UE_LOG(LogInit, Warning, TEXT("Mismatch between bShouldGenerateLowQualityLightmaps(%d) and r.SupportLowQualityLightmaps(%d), UEngine::bShouldGenerateLowQualityLightmaps has been deprecated please use r.SupportLowQualityLightmaps instead"), GEngine->bShouldGenerateLowQualityLightmaps_DEPRECATED, bEnableLowQualityLightmaps);
diff --git a/Engine/Source/Runtime/Engine/Private/UnrealEngine.cpp b/Engine/Source/Runtime/Engine/Private/UnrealEngine.cpp
--- a/Engine/Source/Runtime/Engine/Private/UnrealEngine.cpp
+++ b/Engine/Source/Runtime/Engine/Private/UnrealEngine.cpp
@@ -2044,10 +2044,6 @@ void UEngine::Init(IEngineLoop* InEngineLoop)
 
 	UTextRenderComponent::InitializeMIDCache();
 
-	// Initialize scene cached cvars
-	extern FReadOnlyCVARCache GReadOnlyCVARCache;
-	GReadOnlyCVARCache.Init();
-
 	if (GIsEditor)
 	{
 		// Create a WorldContext for the editor to use and create an initially empty world.
diff --git a/Engine/Source/Runtime/Engine/Public/SceneManagement.h b/Engine/Source/Runtime/Engine/Public/SceneManagement.h
--- a/Engine/Source/Runtime/Engine/Public/SceneManagement.h
+++ b/Engine/Source/Runtime/Engine/Public/SceneManagement.h
@@ -2997,6 +2997,9 @@ extern ENGINE_API FSharedSamplerState* Clamp_WorldGroupSettings;
 /** Initializes the shared sampler states. */
 extern ENGINE_API void InitializeSharedSamplerStates();
 
+/** Initialized the cache of read-only CVars */
+extern ENGINE_API void InitializeReadonlyCVarCache();
+
 /**
 * Cache of read-only console variables used by the scene renderer
 */
diff --git a/Engine/Source/Runtime/Launch/Private/LaunchEngineLoop.cpp b/Engine/Source/Runtime/Launch/Private/LaunchEngineLoop.cpp
--- a/Engine/Source/Runtime/Launch/Private/LaunchEngineLoop.cpp
+++ b/Engine/Source/Runtime/Launch/Private/LaunchEngineLoop.cpp
@@ -3725,6 +3725,9 @@ int32 FEngineLoop::PreInitPostStartupScreen(const TCHAR* CmdLine)
 		}
 #endif
 
+		// Initialize CVar cache after config reload
+		InitializeReadonlyCVarCache();
+
 		InitGameTextLocalization();
 
 		DECLARE_SCOPE_CYCLE_COUNTER(TEXT("Initial UObject load"), STAT_InitialUObjectLoad, STATGROUP_LoadTime);
