[UE] Added extra field for project-specific categorization of actors for reporting purposing.

#tests pie
#rnx

diff --git a/Engine/Source/Runtime/Engine/Classes/GameFramework/Actor.h b/Engine/Source/Runtime/Engine/Classes/GameFramework/Actor.h
--- a/Engine/Source/Runtime/Engine/Classes/GameFramework/Actor.h
+++ b/Engine/Source/Runtime/Engine/Classes/GameFramework/Actor.h
@@ -600,6 +600,10 @@ private:
 
 	/** Internal helper to update Overlaps during Actor initialization/BeginPlay correctly based on the UpdateOverlapsMethodDuringLevelStreaming and bGenerateOverlapEventsDuringLevelStreaming settings. */
 	ENGINE_API void UpdateInitialOverlaps(bool bFromLevelStreaming);
+	
+	/** Describes how much control the remote machine has over the actor. */
+	UPROPERTY(Replicated, Transient, VisibleInstanceOnly, Category=Networking)
+	TEnumAsByte<enum ENetRole> RemoteRole;
 
 public:
 	/**
@@ -640,6 +644,9 @@ public:
 	ENGINE_API void SetNetAddressable();
 
 public:
+	/** Project-specific field that help to categorize actors for reporting purposes*/
+	uint8 ActorCategory = 0;
+
 	/** How long this Actor lives before dying, 0=forever. Note this is the INITIAL value and should not be modified once play has begun. */
 	UPROPERTY(EditAnywhere, BlueprintReadOnly, Category=Actor)
 	float InitialLifeSpan;
@@ -649,10 +656,6 @@ public:
 	float CustomTimeDilation;
 	
 private:
-	/** Describes how much control the remote machine has over the actor. */
-	UPROPERTY(Replicated, Transient, VisibleInstanceOnly, Category=Networking)
-	TEnumAsByte<enum ENetRole> RemoteRole;
-
 	/** The RayTracingGroupId this actor and its components belong to. (For components that did not specify any) */
 	UPROPERTY()
 	int32 RayTracingGroupId;
