Fix nullptr crash in UpdateVertexFactoryCloth

diff --git a/Engine/Source/Runtime/Engine/Private/SkeletalRenderGPUSkin.cpp b/Engine/Source/Runtime/Engine/Private/SkeletalRenderGPUSkin.cpp
--- a/Engine/Source/Runtime/Engine/Private/SkeletalRenderGPUSkin.cpp
+++ b/Engine/Source/Runtime/Engine/Private/SkeletalRenderGPUSkin.cpp
@@ -1965,22 +1965,27 @@ static void UpdateVertexFactoryCloth(TArray<TUniquePtr<FGPUBaseSkinAPEXClothVert
 {
 	for (TUniquePtr<FGPUBaseSkinAPEXClothVertexFactory>& FactoryPtr : VertexFactories)
 	{
-		FGPUBaseSkinVertexFactory* VertexFactory = FactoryPtr->GetVertexFactory();
+		FGPUBaseSkinAPEXClothVertexFactory* BaseVertexFactory = FactoryPtr.Get();
 
-		if (VertexFactory != nullptr)
+		if (BaseVertexFactory != nullptr)
 		{
-			// Setup the update data for enqueue
-			FDynamicUpdateVertexFactoryData VertexUpdateData(VertexFactory, InVertexBuffers);
+			FGPUBaseSkinVertexFactory* VertexFactory = BaseVertexFactory->GetVertexFactory();
 
-			// update vertex factory components and sync it
-			ENQUEUE_RENDER_COMMAND(InitGPUSkinAPEXClothVertexFactory)(UE::RenderCommandPipe::SkeletalMesh,
-				[VertexUpdateData](FRHICommandList& RHICmdList)
+			if (VertexFactory != nullptr)
 			{
-				FGPUSkinAPEXClothDataType Data;
-				InitGPUSkinVertexFactoryComponents(&Data, VertexUpdateData.VertexBuffers, VertexUpdateData.VertexFactory);
-				InitAPEXClothVertexFactoryComponents(&Data, VertexUpdateData.VertexBuffers);
-				VertexUpdateData.VertexFactory->SetData(RHICmdList, &Data);
-			});
+				// Setup the update data for enqueue
+				FDynamicUpdateVertexFactoryData VertexUpdateData(VertexFactory, InVertexBuffers);
+
+				// update vertex factory components and sync it
+				ENQUEUE_RENDER_COMMAND(InitGPUSkinAPEXClothVertexFactory)(UE::RenderCommandPipe::SkeletalMesh,
+					[VertexUpdateData](FRHICommandList& RHICmdList)
+				{
+					FGPUSkinAPEXClothDataType Data;
+					InitGPUSkinVertexFactoryComponents(&Data, VertexUpdateData.VertexBuffers, VertexUpdateData.VertexFactory);
+					InitAPEXClothVertexFactoryComponents(&Data, VertexUpdateData.VertexBuffers);
+					VertexUpdateData.VertexFactory->SetData(RHICmdList, &Data);
+				});
+			}
 		}
 	}
 }
