[UE] Improve accuracy of GetResourceSizeEx for AnimSequences by capturing approx compressed size.

#tests pie

diff --git a/Engine/Source/Runtime/Engine/Classes/Animation/AnimSequence.h b/Engine/Source/Runtime/Engine/Classes/Animation/AnimSequence.h
--- a/Engine/Source/Runtime/Engine/Classes/Animation/AnimSequence.h
+++ b/Engine/Source/Runtime/Engine/Classes/Animation/AnimSequence.h
@@ -398,6 +398,7 @@ public:
 	ENGINE_API virtual bool IsReadyForFinishDestroy() override;
 	ENGINE_API virtual void GetAssetRegistryTags(TArray<FAssetRegistryTag>& OutTags) const override;
 	static ENGINE_API void AddReferencedObjects(UObject* This, FReferenceCollector& Collector);
+	ENGINE_API virtual void GetResourceSizeEx(FResourceSizeEx& CumulativeResourceSize) override;
 	//~ End UObject Interface
 
 	//~ Begin UAnimationAsset Interface
diff --git a/Engine/Source/Runtime/Engine/Private/Animation/AnimSequence.cpp b/Engine/Source/Runtime/Engine/Private/Animation/AnimSequence.cpp
--- a/Engine/Source/Runtime/Engine/Private/Animation/AnimSequence.cpp
+++ b/Engine/Source/Runtime/Engine/Private/Animation/AnimSequence.cpp
@@ -462,6 +462,13 @@ void UAnimSequence::AddReferencedObjects(UObject* This, FReferenceCollector& Col
 	Collector.AddReferencedObject(AnimSeq->CompressedData.CurveCompressionCodec);
 }
 
+void UAnimSequence::GetResourceSizeEx(FResourceSizeEx& CumulativeResourceSize)
+{
+	Super::GetResourceSizeEx(CumulativeResourceSize);
+
+	CumulativeResourceSize.AddDedicatedSystemMemoryBytes(TEXT("CompressedData"), GetApproxCompressedSize());
+}
+
 #if WITH_EDITOR
 void UAnimSequence::WillNeverCacheCookedPlatformDataAgain()
 {
