Transfer bFillCollisionUnderneathForNavmesh from StaticMeshComponent to UInstancedStaticMeshComponent

#test tested in PIE: validated that the navmesh in not present after applying the fix

diff --git a/Engine/Source/Runtime/Engine/Private/ISMPartition/ISMComponentDescriptor.cpp b/Engine/Source/Runtime/Engine/Private/ISMPartition/ISMComponentDescriptor.cpp
--- a/Engine/Source/Runtime/Engine/Private/ISMPartition/ISMComponentDescriptor.cpp
+++ b/Engine/Source/Runtime/Engine/Private/ISMPartition/ISMComponentDescriptor.cpp
@@ -117,6 +117,7 @@ void FISMComponentDescriptorBase::InitFrom(const UStaticMeshComponent* Template,
 	bGenerateOverlapEvents = Template->GetGenerateOverlapEvents();
 	bOverrideNavigationExport = Template->bOverrideNavigationExport;
 	bForceNavigationObstacle = Template->bForceNavigationObstacle;
+	bFillCollisionUnderneathForNavmesh = Template->bFillCollisionUnderneathForNavmesh;
 
 #if WITH_EDITORONLY_DATA
 	HLODBatchingPolicy = Template->HLODBatchingPolicy;
@@ -222,6 +223,7 @@ bool FISMComponentDescriptorBase::operator==(const FISMComponentDescriptorBase&
 	bGenerateOverlapEvents == Other.bGenerateOverlapEvents &&
 	bOverrideNavigationExport == Other.bOverrideNavigationExport &&
 	bForceNavigationObstacle == Other.bForceNavigationObstacle &&
+	bFillCollisionUnderneathForNavmesh == Other.bFillCollisionUnderneathForNavmesh &&
 	WorldPositionOffsetDisableDistance == Other.WorldPositionOffsetDisableDistance &&
 	ShadowCacheInvalidationBehavior == Other.ShadowCacheInvalidationBehavior &&
 	DetailMode == Other.DetailMode &&
@@ -338,6 +340,7 @@ void FISMComponentDescriptorBase::InitComponent(UInstancedStaticMeshComponent* I
 	ISMComponent->SetGenerateOverlapEvents(bGenerateOverlapEvents);
 	ISMComponent->bOverrideNavigationExport = bOverrideNavigationExport;
 	ISMComponent->bForceNavigationObstacle = bForceNavigationObstacle;
+	ISMComponent->bFillCollisionUnderneathForNavmesh = bFillCollisionUnderneathForNavmesh;
 	ISMComponent->WorldPositionOffsetDisableDistance = WorldPositionOffsetDisableDistance;
 	ISMComponent->ShadowCacheInvalidationBehavior = ShadowCacheInvalidationBehavior;
 	ISMComponent->DetailMode = DetailMode;
diff --git a/Engine/Source/Runtime/Engine/Public/ISMPartition/ISMComponentDescriptor.h b/Engine/Source/Runtime/Engine/Public/ISMPartition/ISMComponentDescriptor.h
--- a/Engine/Source/Runtime/Engine/Public/ISMPartition/ISMComponentDescriptor.h
+++ b/Engine/Source/Runtime/Engine/Public/ISMPartition/ISMComponentDescriptor.h
@@ -178,6 +178,9 @@ public:
 	UPROPERTY(EditAnywhere, Category = "Component Settings")
 	uint8 bForceNavigationObstacle : 1;
 
+	UPROPERTY(EditAnywhere, Category = "Component Settings")
+	uint8 bFillCollisionUnderneathForNavmesh : 1;
+
 	UPROPERTY(EditAnywhere, Category = "Component Settings")
 	int32 WorldPositionOffsetDisableDistance;
 
