Added a new parameter for looping TimerManager timers to allow opting out of "catch-up" behaviour where the timer callback will run as many times as would fit in the current frame delta. For example, a looping timer with interval 0.001s would run 9-10 times in a frame with delta 0.01s.

Many looping timers are configured in such a way that they want to try something periodically, but doing so multiple times in a frame is not productive (e.g. wait for replication, update audio parameters, check for input), so this way we can avoid re-registering the timer while still only running once per expiration.

diff --git a/Engine/Source/Runtime/Engine/Classes/Kismet/KismetSystemLibrary.h b/Engine/Source/Runtime/Engine/Classes/Kismet/KismetSystemLibrary.h
--- a/Engine/Source/Runtime/Engine/Classes/Kismet/KismetSystemLibrary.h
+++ b/Engine/Source/Runtime/Engine/Classes/Kismet/KismetSystemLibrary.h
@@ -637,12 +637,13 @@ class UKismetSystemLibrary : public UBlueprintFunctionLibrary
 	 * @param Event						Event. Can be a K2 function or a Custom Event.
 	 * @param Time						How long to wait before executing the delegate, in seconds. Setting a timer to <= 0 seconds will clear it if it is set.
 	 * @param bLooping					True to keep executing the delegate every Time seconds, false to execute delegate only once.
+	 * @param bMaxOncePerFrame			For looping timers, whether to execute only once when the timer would otherwise expires multiple times in the current frame.
 	 * @param InitialStartDelay			Initial delay passed to the timer manager, in seconds.
 	 * @param InitialStartDelayVariance	Use this to add some variance to when the timer starts in lieu of doing a random range on the InitialStartDelay input, in seconds. 
 	 * @return							The timer handle to pass to other timer functions to manipulate this timer.
 	 */
 	UFUNCTION(BlueprintCallable, meta=(DisplayName = "Set Timer by Event", ScriptName = "SetTimerDelegate", AdvancedDisplay="InitialStartDelay, InitialStartDelayVariance"), Category="Utilities|Time")
-	static ENGINE_API FTimerHandle K2_SetTimerDelegate(UPARAM(DisplayName="Event") FTimerDynamicDelegate Delegate, float Time, bool bLooping, float InitialStartDelay = 0.f, float InitialStartDelayVariance = 0.f);
+	static ENGINE_API FTimerHandle K2_SetTimerDelegate(UPARAM(DisplayName="Event") FTimerDynamicDelegate Delegate, float Time, bool bLooping, bool bMaxOncePerFrame = false, float InitialStartDelay = 0.f, float InitialStartDelayVariance = 0.f);
 
 	/**
 	 * Set a timer to execute a delegate next tick.
@@ -806,13 +807,14 @@ class UKismetSystemLibrary : public UBlueprintFunctionLibrary
 	 * @param Object					Object that implements the delegate function. Defaults to self (this blueprint)
 	 * @param FunctionName				Delegate function name. Can be a K2 function or a Custom Event.
 	 * @param Time						How long to wait before executing the delegate, in seconds. Setting a timer to <= 0 seconds will clear it if it is set.
-	 * @param bLooping					true to keep executing the delegate every Time seconds, false to execute delegate only once.
+	 * @param bLooping					True to keep executing the delegate every Time seconds, false to execute delegate only once.
+	 * @param bMaxOncePerFrame			For looping timers, whether to execute only once when the timer would otherwise expires multiple times in the current frame.
 	 * @param InitialStartDelay			Initial delay passed to the timer manager to allow some variance in when the timer starts, in seconds.
 	 * @param InitialStartDelayVariance	Use this to add some variance to when the timer starts in lieu of doing a random range on the InitialStartDelay input, in seconds.
 	 * @return							The timer handle to pass to other timer functions to manipulate this timer.
 	 */
 	UFUNCTION(BlueprintCallable, meta=(DisplayName = "Set Timer by Function Name", ScriptName = "SetTimer", DefaultToSelf = "Object", AdvancedDisplay="InitialStartDelay, InitialStartDelayVariance"), Category="Utilities|Time")
-	static ENGINE_API FTimerHandle K2_SetTimer(UObject* Object, FString FunctionName, float Time, bool bLooping, float InitialStartDelay = 0.f, float InitialStartDelayVariance = 0.f);
+	static ENGINE_API FTimerHandle K2_SetTimer(UObject* Object, FString FunctionName, float Time, bool bLooping, bool bMaxOncePerFrame = false, float InitialStartDelay = 0.f, float InitialStartDelayVariance = 0.f);
 
 	/**
 	 * Set a timer to execute a delegate on the next tick.
diff --git a/Engine/Source/Runtime/Engine/Private/KismetSystemLibrary.cpp b/Engine/Source/Runtime/Engine/Private/KismetSystemLibrary.cpp
--- a/Engine/Source/Runtime/Engine/Private/KismetSystemLibrary.cpp
+++ b/Engine/Source/Runtime/Engine/Private/KismetSystemLibrary.cpp
@@ -610,7 +610,7 @@ FTimerHandle UKismetSystemLibrary::K2_InvalidateTimerHandle(FTimerHandle& TimerH
 	return TimerHandle;
 }
 
-FTimerHandle UKismetSystemLibrary::K2_SetTimer(UObject* Object, FString FunctionName, float Time, bool bLooping, float InitialStartDelay, float InitialStartDelayVariance)
+FTimerHandle UKismetSystemLibrary::K2_SetTimer(UObject* Object, FString FunctionName, float Time, bool bLooping, bool bMaxOncePerFrame, float InitialStartDelay, float InitialStartDelayVariance)
 {
 	FName const FunctionFName(*FunctionName);
 
@@ -629,7 +629,7 @@ FTimerHandle UKismetSystemLibrary::K2_SetTimer(UObject* Object, FString Function
 
 	FTimerDynamicDelegate Delegate;
 	Delegate.BindUFunction(Object, FunctionFName);
-	return K2_SetTimerDelegate(Delegate, Time, bLooping, InitialStartDelay);
+	return K2_SetTimerDelegate(Delegate, Time, bLooping, bMaxOncePerFrame, InitialStartDelay);
 }
 
 FTimerHandle UKismetSystemLibrary::K2_SetTimerForNextTick(UObject* Object, FString FunctionName)
@@ -654,7 +654,7 @@ FTimerHandle UKismetSystemLibrary::K2_SetTimerForNextTick(UObject* Object, FStri
 	return K2_SetTimerForNextTickDelegate(Delegate);
 }
 
-FTimerHandle UKismetSystemLibrary::K2_SetTimerDelegate(FTimerDynamicDelegate Delegate, float Time, bool bLooping, float InitialStartDelay, float InitialStartDelayVariance)
+FTimerHandle UKismetSystemLibrary::K2_SetTimerDelegate(FTimerDynamicDelegate Delegate, float Time, bool bLooping, bool bMaxOncePerFrame, float InitialStartDelay, float InitialStartDelayVariance)
 {
 	FTimerHandle Handle;
 	if (Delegate.IsBound())
@@ -672,7 +672,7 @@ FTimerHandle UKismetSystemLibrary::K2_SetTimerDelegate(FTimerDynamicDelegate Del
 
 			FTimerManager& TimerManager = World->GetTimerManager();
 			Handle = TimerManager.K2_FindDynamicTimerHandle(Delegate);
-			TimerManager.SetTimer(Handle, Delegate, Time, bLooping, (Time + InitialStartDelay));
+			TimerManager.SetTimer(Handle, Delegate, Time, FTimerManagerTimerParameters { .bLoop = bLooping, .bMaxOncePerFrame = bMaxOncePerFrame, .FirstDelay = Time + InitialStartDelay });
 		}
 	}
 	else
diff --git a/Engine/Source/Runtime/Engine/Private/TimerManager.cpp b/Engine/Source/Runtime/Engine/Private/TimerManager.cpp
--- a/Engine/Source/Runtime/Engine/Private/TimerManager.cpp
+++ b/Engine/Source/Runtime/Engine/Private/TimerManager.cpp
@@ -461,7 +461,12 @@ FTimerHandle FTimerManager::K2_FindDynamicTimerHandle(FTimerDynamicDelegate InDy
 	return Result;
 }
 
-void FTimerManager::InternalSetTimer(FTimerHandle& InOutHandle, FTimerUnifiedDelegate&& InDelegate, float InRate, bool InbLoop, float InFirstDelay)
+void FTimerManager::InternalSetTimer(FTimerHandle& InOutHandle, FTimerUnifiedDelegate&& InDelegate, float InRate, bool bInLoop, float InFirstDelay)
+{
+	InternalSetTimer(InOutHandle, MoveTemp(InDelegate), InRate, FTimerManagerTimerParameters{ .bLoop = bInLoop, .FirstDelay = InFirstDelay });
+}
+
+void FTimerManager::InternalSetTimer(FTimerHandle& InOutHandle, FTimerUnifiedDelegate&& InDelegate, float InRate, const FTimerManagerTimerParameters& InTimerParameters)
 {
 	SCOPE_CYCLE_COUNTER(STAT_SetTimer);
 
@@ -482,7 +487,8 @@ void FTimerManager::InternalSetTimer(FTimerHandle& InOutHandle, FTimerUnifiedDel
 		NewTimerData.TimerDelegate = MoveTemp(InDelegate);
 
 		NewTimerData.Rate = InRate;
-		NewTimerData.bLoop = InbLoop;
+		NewTimerData.bLoop = InTimerParameters.bLoop;
+		NewTimerData.bMaxOncePerFrame = InTimerParameters.bMaxOncePerFrame;
 		NewTimerData.bRequiresDelegate = NewTimerData.TimerDelegate.IsBound();
 
 		// Set level collection
@@ -492,7 +498,7 @@ void FTimerManager::InternalSetTimer(FTimerHandle& InOutHandle, FTimerUnifiedDel
 			NewTimerData.LevelCollection = OwningWorld->GetActiveLevelCollection()->GetType();
 		}
 
-		const float FirstDelay = (InFirstDelay >= 0.f) ? InFirstDelay : InRate;
+		const float FirstDelay = (InTimerParameters.FirstDelay >= 0.f) ? InTimerParameters.FirstDelay : InRate;
 
 		FTimerHandle NewTimerHandle;
 		if (HasBeenTickedThisFrame())
@@ -762,6 +768,7 @@ void FTimerManager::UnPauseTimer(FTimerHandle InHandle)
 
 FTimerData::FTimerData()
 	: bLoop(false)
+	, bMaxOncePerFrame(false)
 	, bRequiresDelegate(false)
 	, Status(ETimerStatus::Active)
 	, Rate(0)
@@ -932,7 +939,7 @@ void FTimerManager::Tick(float DeltaTime)
 				// Update Top pointer, in case it has been invalidated by the Execute call
 				Top = FindTimer(CurrentlyExecutingTimer);
 				checkf(!Top || !WillRemoveTimerAssert(CurrentlyExecutingTimer), TEXT("RemoveTimer(CurrentlyExecutingTimer) - due to fail after Execute()"));
-				if (!Top || Top->Status != ETimerStatus::Executing)
+				if (!Top || Top->Status != ETimerStatus::Executing || Top->bMaxOncePerFrame)
 				{
 					break;
 				}
diff --git a/Engine/Source/Runtime/Engine/Public/TimerManager.h b/Engine/Source/Runtime/Engine/Public/TimerManager.h
--- a/Engine/Source/Runtime/Engine/Public/TimerManager.h
+++ b/Engine/Source/Runtime/Engine/Public/TimerManager.h
@@ -111,6 +111,11 @@ struct FTimerData
 	/** If true, this timer will loop indefinitely.  Otherwise, it will be destroyed when it expires. */
 	uint8 bLoop : 1;
 
+	/** If true for a looping timer, it will run a maximum of once per frame when it expires. Otherwise, it will run as many times as can fit in
+	 *  the current frame's delta time (rounding down).
+	 */
+	uint8 bMaxOncePerFrame : 1;
+
 	/** If true, this timer was created with a delegate to call (which means if the delegate becomes invalid, we should invalidate the timer too). */
 	uint8 bRequiresDelegate : 1;
 
@@ -148,6 +153,15 @@ struct FTimerData
 	FTimerData& operator=(const FTimerData&) = delete;
 };
 
+/**
+ * Structure to support specialized looping behavior.
+ */
+struct FTimerManagerTimerParameters
+{
+	bool bLoop = false;
+	bool bMaxOncePerFrame = false;
+	float FirstDelay = -1.f;
+};
 
 /** 
  * Class to globally manage timers.
@@ -213,6 +227,50 @@ public:
 		InternalSetTimer(InOutHandle, FTimerUnifiedDelegate(MoveTemp(Callback)), InRate, InbLoop, InFirstDelay);
 	}
 
+	/* Preferred versions of the above methods that now take a more configurable structure that specifies timer looping and delay behavior. */
+
+	/**
+	 * Sets a timer to call the given native function at a set interval.  If a timer is already set
+	 * for this handle, it will replace the current timer.
+	 *
+	 * @param InOutHandle			If the passed-in handle refers to an existing timer, it will be cleared before the new timer is added. A new handle to the new timer is returned in either case.
+	 * @param InObj					Object to call the timer function on.
+	 * @param InTimerMethod			Method to call when timer fires.
+	 * @param InRate				The amount of time (in seconds) between set and firing.  If <= 0.f, clears existing timers.
+	 * @param InTimerParameters		The structure with additional timer looping and delay parameters.
+	 */
+	template <class UserClass>
+	FORCEINLINE void SetTimer(FTimerHandle& InOutHandle, UserClass* InObj, typename FTimerDelegate::TMethodPtr<UserClass> InTimerMethod, float InRate, const FTimerManagerTimerParameters& InTimerParameters)
+	{
+		InternalSetTimer(InOutHandle, FTimerUnifiedDelegate(FTimerDelegate::CreateUObject(InObj, InTimerMethod)), InRate, InTimerParameters);
+	}
+	template< class UserClass >
+	FORCEINLINE void SetTimer(FTimerHandle& InOutHandle, UserClass* InObj, typename FTimerDelegate::TConstMethodPtr<UserClass> InTimerMethod, float InRate, const FTimerManagerTimerParameters& InTimerParameters)
+	{
+		InternalSetTimer(InOutHandle, FTimerUnifiedDelegate(FTimerDelegate::CreateUObject(InObj, InTimerMethod)), InRate, InTimerParameters);
+	}
+
+	/** Version that takes any generic delegate. */
+	FORCEINLINE void SetTimer(FTimerHandle& InOutHandle, FTimerDelegate const& InDelegate, float InRate, const FTimerManagerTimerParameters& InTimerParameters)
+	{
+		InternalSetTimer(InOutHandle, FTimerUnifiedDelegate(InDelegate), InRate, InTimerParameters);
+	}
+	/** Version that takes a dynamic delegate (e.g. for UFunctions). */
+	FORCEINLINE void SetTimer(FTimerHandle& InOutHandle, FTimerDynamicDelegate const& InDynDelegate, float InRate, const FTimerManagerTimerParameters& InTimerParameters)
+	{
+		InternalSetTimer(InOutHandle, FTimerUnifiedDelegate(InDynDelegate), InRate, InTimerParameters);
+	}
+	/*** Version that doesn't take a delegate */
+	FORCEINLINE void SetTimer(FTimerHandle& InOutHandle, float InRate, const FTimerManagerTimerParameters& InTimerParameters)
+	{
+		InternalSetTimer(InOutHandle, FTimerUnifiedDelegate(), InRate, InTimerParameters);
+	}
+	/** Version that takes a TFunction */
+	FORCEINLINE void SetTimer(FTimerHandle& InOutHandle, TFunction<void(void)>&& Callback, float InRate, const FTimerManagerTimerParameters& InTimerParameters)
+	{
+		InternalSetTimer(InOutHandle, FTimerUnifiedDelegate(MoveTemp(Callback)), InRate, InTimerParameters);
+	}
+
 	/**
 	 * Sets a timer to call the given native function on the next tick.
 	 *
@@ -403,7 +461,8 @@ protected:
 	ENGINE_API FTimerData* FindTimer( FTimerHandle const& InHandle );
 
 private:
-	ENGINE_API void InternalSetTimer( FTimerHandle& InOutHandle, FTimerUnifiedDelegate&& InDelegate, float InRate, bool InbLoop, float InFirstDelay );
+	ENGINE_API void InternalSetTimer( FTimerHandle& InOutHandle, FTimerUnifiedDelegate&& InDelegate, float InRate, bool bInLoop, float InFirstDelay );
+	ENGINE_API void InternalSetTimer( FTimerHandle& InOutHandle, FTimerUnifiedDelegate&& InDelegate, float InRate, const FTimerManagerTimerParameters& InTimerParameters );
 	ENGINE_API FTimerHandle InternalSetTimerForNextTick( FTimerUnifiedDelegate&& InDelegate );
 	ENGINE_API void InternalClearTimer( FTimerHandle InDelegate );
 	ENGINE_API void InternalClearAllTimers( void const* Object );
