Made SceneColor node output scene alpha on top of scene color RGB.

diff --git a/Engine/Shaders/Private/MaterialTemplate.ush b/Engine/Shaders/Private/MaterialTemplate.ush
--- a/Engine/Shaders/Private/MaterialTemplate.ush
+++ b/Engine/Shaders/Private/MaterialTemplate.ush
@@ -2628,12 +2628,13 @@ float4 SceneTextureLookup(float2 UV, int SceneTextureIndex, bool bFiltered)
 
 #if SHADING_PATH_DEFERRED
 
-/** Applies an offset to the scene texture lookup and decodes the HDR linear space color. */
-float3 DecodeSceneColorForMaterialNode(float2 ScreenUV)
+
+/** Applies an offset to the scene texture lookup and decodes the HDR linear space color and alpha. */
+float4 DecodeSceneColorAndAlpharForMaterialNode(float2 ScreenUV)
 {
 #if !defined(SceneColorCopyTexture)
 	// Hit proxies rendering pass doesn't have access to valid render buffers
-	return float3(0.0f, 0.0f, 0.0f);
+	return float4(0.0f, 0.0f, 0.0f, 0.0f);
 #else
 	float4 EncodedSceneColor = Texture2DSample(SceneColorCopyTexture, SceneColorCopySampler, ScreenUV);
 
@@ -2642,10 +2643,16 @@ float3 DecodeSceneColorForMaterialNode(float2 ScreenUV)
 
 	SampledColor *= View.OneOverPreExposure.xxx;
 
-	return SampledColor;
+	return float4(SampledColor, EncodedSceneColor.a);
 #endif
 }
 
+/** Applies an offset to the scene texture lookup and decodes the HDR linear space color. */
+float3 DecodeSceneColorForMaterialNode(float2 ScreenUV)
+{
+	return DecodeSceneColorAndAlpharForMaterialNode(ScreenUV).rgb;
+}
+
 #endif // SHADING_PATH_DEFERRED
 
 float4 MaterialExpressionDBufferTextureLookup(FMaterialPixelParameters Parameters, float2 BufferUV, int DBufferTextureIndex)
diff --git a/Engine/Source/Runtime/Engine/Private/Materials/HLSLMaterialTranslator.cpp b/Engine/Source/Runtime/Engine/Private/Materials/HLSLMaterialTranslator.cpp
--- a/Engine/Source/Runtime/Engine/Private/Materials/HLSLMaterialTranslator.cpp
+++ b/Engine/Source/Runtime/Engine/Private/Materials/HLSLMaterialTranslator.cpp
@@ -7526,8 +7526,8 @@ int32 FHLSLMaterialTranslator::SceneColor(int32 Offset, int32 ViewportUV, bool b
 
 	int32 ScreenUVCode = GetScreenAlignedUV(Offset, ViewportUV, bUseOffset);
 	return AddCodeChunk(
-		MCT_Float3,
-		TEXT("DecodeSceneColorForMaterialNode(%s)"),
+		MCT_Float4,
+		TEXT("DecodeSceneColorAndAlpharForMaterialNode(%s)"),
 		*GetParameterCode(ScreenUVCode)
 		);
 }
diff --git a/Engine/Source/Runtime/Engine/Private/Materials/MaterialExpressions.cpp b/Engine/Source/Runtime/Engine/Private/Materials/MaterialExpressions.cpp
--- a/Engine/Source/Runtime/Engine/Private/Materials/MaterialExpressions.cpp
+++ b/Engine/Source/Runtime/Engine/Private/Materials/MaterialExpressions.cpp
@@ -11255,8 +11255,13 @@ UMaterialExpressionSceneColor::UMaterialExpressionSceneColor(const FObjectInitia
 	MenuCategories.Add(ConstructorStatics.NAME_Texture);
 
 	bShaderInputData = true;
+	bShowOutputNameOnPin = true;
 #endif
 	ConstInput = FVector2D(0.f, 0.f);
+
+	Outputs.Reset();
+	Outputs.Add(FExpressionOutput(TEXT("RGB"), 1, 1, 1, 1, 0));
+	Outputs.Add(FExpressionOutput(TEXT("A"), 1, 0, 0, 0, 1));
 }
 
 void UMaterialExpressionSceneColor::PostLoad()
diff --git a/Engine/Source/Runtime/Engine/Private/Materials/MaterialHLSLTree.cpp b/Engine/Source/Runtime/Engine/Private/Materials/MaterialHLSLTree.cpp
--- a/Engine/Source/Runtime/Engine/Private/Materials/MaterialHLSLTree.cpp
+++ b/Engine/Source/Runtime/Engine/Private/Materials/MaterialHLSLTree.cpp
@@ -2440,13 +2440,13 @@ bool FExpressionSceneColor::PrepareValue(FEmitContext& Context, FEmitScope& Scop
 		Context.MaterialCompilationOutput->SetIsSceneTextureUsed(PPI_SceneColor);
 	}
 
-	return OutResult.SetType(Context, RequestedType, EExpressionEvaluation::Shader, Shader::EValueType::Float3);
+	return OutResult.SetType(Context, RequestedType, EExpressionEvaluation::Shader, Shader::EValueType::Float4);
 }
 
 void FExpressionSceneColor::EmitValueShader(FEmitContext& Context, FEmitScope& Scope, const FRequestedType& RequestedType, FEmitValueShaderResult& OutResult) const
 {
 	FEmitShaderExpression* EmitScreenUV = ScreenUVExpression->GetValueShader(Context, Scope, Shader::EValueType::Float2);
-	OutResult.Code = Context.EmitExpression(Scope, Shader::EValueType::Float3, TEXT("DecodeSceneColorForMaterialNode(%)"), EmitScreenUV);
+	OutResult.Code = Context.EmitExpression(Scope, Shader::EValueType::Float4, TEXT("DecodeSceneColorAndAlpharForMaterialNode(%)"), EmitScreenUV);
 }
 
 bool FExpressionNoise::PrepareValue(FEmitContext& Context, FEmitScope& Scope, const FRequestedType& RequestedType, FPrepareValueResult& OutResult) const
diff --git a/Engine/Source/Runtime/Engine/Private/Materials/MaterialShared.cpp b/Engine/Source/Runtime/Engine/Private/Materials/MaterialShared.cpp
--- a/Engine/Source/Runtime/Engine/Private/Materials/MaterialShared.cpp
+++ b/Engine/Source/Runtime/Engine/Private/Materials/MaterialShared.cpp
@@ -324,28 +324,68 @@ UE_IMPLEMENT_STRUCT("/Script/Engine", Vector2MaterialInput);
 UE_IMPLEMENT_STRUCT("/Script/Engine", MaterialAttributesInput);
 
 #if WITH_EDITOR
-int32 FExpressionInput::Compile(class FMaterialCompiler* Compiler)
+
+struct FConnectionMask
 {
-	if(Expression)
+	bool Mask;
+	bool MaskR;
+	bool MaskG;
+	bool MaskB;
+	bool MaskA;
+};
+
+/** Helper function that returns the intersection of components mask between specified input
+ * and its connected output. */
+static FConnectionMask GetConnectionMask(const FExpressionInput* Input)
+{
+	FConnectionMask CM = {
+		(bool)Input->Mask,
+		(bool)Input->MaskR,
+		(bool)Input->MaskG,
+		(bool)Input->MaskB,
+		(bool)Input->MaskA,
+	};
+
+	if (Input->Expression->GetOutputs().IsValidIndex(Input->OutputIndex))
 	{
-		Expression->ValidateState();
-		
-		int32 ExpressionResult = Compiler->CallExpression(FMaterialExpressionKey(Expression, OutputIndex, Compiler->GetMaterialAttribute(), Compiler->IsCurrentlyCompilingForPreviousFrame()),Compiler);
+		FExpressionOutput& Output = Input->Expression->GetOutputs()[Input->OutputIndex];
+		CM.MaskR = (bool)Output.MaskR && (!CM.Mask || CM.MaskR);
+		CM.MaskG = (bool)Output.MaskG && (!CM.Mask || CM.MaskG);
+		CM.MaskB = (bool)Output.MaskB && (!CM.Mask || CM.MaskB);
+		CM.MaskA = (bool)Output.MaskA && (!CM.Mask || CM.MaskA);
+		CM.Mask |= (bool)Output.Mask;
+	}
 
-		if(Mask && ExpressionResult != INDEX_NONE)
-		{
-			return Compiler->ComponentMask(
-				ExpressionResult,
-				!!MaskR,!!MaskG,!!MaskB,!!MaskA
-				);
-		}
-		else
-		{
-			return ExpressionResult;
-		}
+	return CM;
+}
+
+int32 FExpressionInput::Compile(class FMaterialCompiler* Compiler)
+{
+	if (!Expression)
+	{
+		return INDEX_NONE;
 	}
-	else
+
+	Expression->ValidateState();
+	int32 ExpressionResult = Compiler->CallExpression(FMaterialExpressionKey(Expression, OutputIndex, Compiler->GetMaterialAttribute(), Compiler->IsCurrentlyCompilingForPreviousFrame()),Compiler);
+	
+	// Early out if compiling expression failed
+	if (ExpressionResult == INDEX_NONE)
+	{
 		return INDEX_NONE;
+	}
+
+	// Use the intersection of components mask between this input and connected output.
+	// We do this to make sure that an out-of-date expression input (that most likely caches the mask from its
+	// connected output) gets the correct mask if the output mask has changed (for instance because now it
+	// specifies a mask where it didn't before when the material was saved)
+	FConnectionMask CM = GetConnectionMask(this);
+	if (CM.Mask)
+	{
+		ExpressionResult = Compiler->ComponentMask(ExpressionResult, CM.MaskR, CM.MaskG, CM.MaskB, CM.MaskA);
+	}
+
+	return ExpressionResult;
 }
 
 const UE::HLSLTree::FExpression* FExpressionInput::TryAcquireHLSLExpression(FMaterialHLSLGenerator& Generator, UE::HLSLTree::FScope& Scope, int32 InputIndex) const
@@ -355,7 +395,11 @@ const UE::HLSLTree::FExpression* FExpressionInput::TryAcquireHLSLExpression(FMat
 	if (Expression)
 	{
 		Expression->ValidateState();
-		const FSwizzleParameters SwizzleParams = Mask ? MakeSwizzleMask(!!MaskR, !!MaskG, !!MaskB, !!MaskA) : FSwizzleParameters();
+
+		// Use the intersection of components mask between this input and connected output.
+		FConnectionMask CM = GetConnectionMask(this);
+		const FSwizzleParameters SwizzleParams = CM.Mask ? MakeSwizzleMask(CM.MaskR, CM.MaskG, CM.MaskB, CM.MaskA) : FSwizzleParameters();
+		
 		Result = Generator.AcquireExpression(Scope, InputIndex, Expression, OutputIndex, SwizzleParams);
 	}
 	return Result;
