[PhysicsEngine]
* Fixed compile error happening on some platforms

diff --git a/Engine/Source/Runtime/Engine/Classes/PhysicsEngine/ExternalSpatialAccelerationPayload.h b/Engine/Source/Runtime/Engine/Classes/PhysicsEngine/ExternalSpatialAccelerationPayload.h
--- a/Engine/Source/Runtime/Engine/Classes/PhysicsEngine/ExternalSpatialAccelerationPayload.h
+++ b/Engine/Source/Runtime/Engine/Classes/PhysicsEngine/ExternalSpatialAccelerationPayload.h
@@ -45,7 +45,7 @@ public:
 
 	friend ENGINE_API uint32 GetTypeHash(const FExternalSpatialAccelerationPayload& Payload);
 private:
-	TObjectKey<UPrimitiveComponent> Component = nullptr;
+	TObjectKey<UPrimitiveComponent> Component;
 	int32 BoneId = INDEX_NONE;
 	Chaos::FUniqueIdx CachedUniqueIdx;
 
