Chaos : Optimized the removal of pending async commands to be 0(N) by avoiding calling RemoveAt for each command

diff --git a/Engine/Source/Runtime/Engine/Private/PhysicsEngine/Experimental/PhysScene_Chaos.cpp b/Engine/Source/Runtime/Engine/Private/PhysicsEngine/Experimental/PhysScene_Chaos.cpp
--- a/Engine/Source/Runtime/Engine/Private/PhysicsEngine/Experimental/PhysScene_Chaos.cpp
+++ b/Engine/Source/Runtime/Engine/Private/PhysicsEngine/Experimental/PhysScene_Chaos.cpp
@@ -150,6 +150,9 @@ public:
 		const bool bAllowResim = PhysicsSettings->PhysicsPrediction.bEnablePhysicsPrediction;
 		const int32 NumFrames = PhysicsSettings->GetPhysicsHistoryCount();
 
+		TArray<int32> CommandIndicesToRemove;
+		CommandIndicesToRemove.Reserve(PendingCommands.Num());
+
 		for (int32 Idx = 0; Idx < PendingCommands.Num(); ++Idx)
 		{
 			const int32 CurrentFrame = static_cast<FPBDRigidsSolver*>(GetSolver())->GetCurrentFrame();
@@ -187,11 +190,12 @@ public:
 
 			if (bRemove)
 			{
-				PendingCommands.RemoveAt(Idx);	//Need to keep functions in order. If this is slow we could try going in reverse order, but expecting number of commands to be low per frame
-				--Idx;
+				CommandIndicesToRemove.Add(Idx);
 			}
 		}
 
+		RemoveArrayItemsAtSortedIndices(PendingCommands, CommandIndicesToRemove);
+
 		const FReal DeltaTime = GetDeltaTime_Internal();
 		const FReal SimTime = GetSimTime_Internal();
 		//TODO: handle case where callbacks modify AsyncPhysicsTickComponents or AsyncPhysicsTickActors
