[NetPhysics] Add back deprecated logic for rewind physics frame offset to keep support for deprecated API.

diff --git a/Engine/Source/Runtime/Engine/Classes/GameFramework/PlayerController.h b/Engine/Source/Runtime/Engine/Classes/GameFramework/PlayerController.h
--- a/Engine/Source/Runtime/Engine/Classes/GameFramework/PlayerController.h
+++ b/Engine/Source/Runtime/Engine/Classes/GameFramework/PlayerController.h
@@ -2242,52 +2242,253 @@ private:
 	FString CurrentInputModeDebugString;
 #endif
 
+public:
+
+	/**
+	 * DEPRECATED 5.4, physics frame offset and time dilation handled via ClientAckTimeDilation() and ClientSetupAsyncPhysicsTimestamp()
+	 * Frame number exchange. This doesn't inherently do anything but is used by the network prediction physics system.
+	 * This may be moved out at some point.
+	 *
+	 * This is meant to provide a mechanism for client side prediction to correlate client input and server frame numbers.
+	 * Frame is a loose concept here. It doesn't necessary mean GFrameNumber. Its just an arbitrary increasing sequence of numbers that is used
+	 * to label discrete units of client->Server input. For example the main thread may tick at a high variable rate but input is generated at a fixed
+	 * step interval.
+	 */
+	
+	struct UE_DEPRECATED(5.4, "Deprecated for not being used in physics frame offset calculation anymore and will get removed eventually. Use GetPhysicsTimestamp() and GetNetworkPhysicsTickOffset() for the physics frame offset which is automatically kept in sync by time dilation while Physics Prediction in Project Settings is enabled. Recommended when using the new flow: Set p.net.CmdOffsetEnabled = 0 to stop the legacy frame offset logic from running in the background.")
+		FInputCmdBuffer
+	{
+		PRAGMA_DISABLE_DEPRECATION_WARNINGS
+		int32 HeadFrame() const { return LastWritten; }
+		int32 TailFrame() const { return FMath::Max(0, LastWritten - Buffer.Num() + 1); }
+		TArray<uint8>& Write(int32 Frame) { LastWritten = FMath::Max(Frame, LastWritten); return Buffer[Frame % Buffer.Num()]; }
+		const TArray<uint8>& Get(int32 Frame) const { return Buffer[Frame % Buffer.Num()]; }
+		PRAGMA_ENABLE_DEPRECATION_WARNINGS
+
+	private:
+		PRAGMA_DISABLE_DEPRECATION_WARNINGS
+		int32 LastWritten = INDEX_NONE;
+		TStaticArray<TArray<uint8>, 16> Buffer;
+		PRAGMA_ENABLE_DEPRECATION_WARNINGS
+
+	public:
+		PRAGMA_DISABLE_DEPRECATION_WARNINGS
+		FInputCmdBuffer() = default;
+		~FInputCmdBuffer() = default;
+		FInputCmdBuffer(const FInputCmdBuffer&) = default;
+		FInputCmdBuffer(FInputCmdBuffer&&) = default;
+		FInputCmdBuffer& operator=(const FInputCmdBuffer&) = default;
+		FInputCmdBuffer& operator=(FInputCmdBuffer&&) = default;
+		PRAGMA_ENABLE_DEPRECATION_WARNINGS
+	};
+
+	PRAGMA_DISABLE_DEPRECATION_WARNINGS
+	UE_DEPRECATED(5.4, "Deprecated for not being used in physics frame offset calculation anymore and will get removed eventually. Use GetPhysicsTimestamp() and GetNetworkPhysicsTickOffset() for the physics frame offset which is automatically kept in sync by time dilation while Physics Prediction in Project Settings is enabled. Recommended when using the new flow: Set p.net.CmdOffsetEnabled = 0 to stop the legacy frame offset logic from running in the background.")
+	FInputCmdBuffer& GetInputBuffer() { return InputBuffer_DEPRECATED; }
+	PRAGMA_ENABLE_DEPRECATION_WARNINGS
+	
+	// -------------------------------------------------------------------------
+	// Client
+	// -------------------------------------------------------------------------
+	
+	/** DEPRECATED 5.4, physics frame offset and time dilation handled via ClientAckTimeDilation() and ClientSetupAsyncPhysicsTimestamp() */
+	struct UE_DEPRECATED(5.4, "Deprecated for not being used in physics frame offset calculation anymore and will get removed eventually. Use GetPhysicsTimestamp() and GetNetworkPhysicsTickOffset() for the physics frame offset which is automatically kept in sync by time dilation while Physics Prediction in Project Settings is enabled. Recommended when using the new flow: Set p.net.CmdOffsetEnabled = 0 to stop the legacy frame offset logic from running in the background.")
+		FClientFrameInfo
+	{
+		PRAGMA_DISABLE_DEPRECATION_WARNINGS
+		int32 LastRecvInputFrame = INDEX_NONE;	// The latest inputcmd that the server acknowledged receiving, but not yet processed (this is our frame number that we gave them)
+		int32 LastProcessedInputFrame = INDEX_NONE; // The latest InputCmd that the server actually processed (this is our frame number that we gave them)
+		int32 LastRecvServerFrame = INDEX_NONE; // the latest ServerFrame number that the processing of LastRecvInputFrame happened on (Server's local frame number)
+		int8 QuantizedTimeDilation = 1; // Server sent this to this client, telling them to dilate local time either catch up or slow down
+		float TargetNumBufferedCmds = 0.f;
+		int32 GetLocalFrameOffset() const { return LastProcessedInputFrame - LastRecvServerFrame; }
+		PRAGMA_ENABLE_DEPRECATION_WARNINGS
+
+		PRAGMA_DISABLE_DEPRECATION_WARNINGS
+		FClientFrameInfo() = default;
+		~FClientFrameInfo() = default;
+		FClientFrameInfo(const FClientFrameInfo&) = default;
+		FClientFrameInfo(FClientFrameInfo&&) = default;
+		FClientFrameInfo& operator=(const FClientFrameInfo&) = default;
+		FClientFrameInfo& operator=(FClientFrameInfo&&) = default;
+		PRAGMA_ENABLE_DEPRECATION_WARNINGS
+	};
+	
+	// Client pushes input data locally. RPC is sent here but also includes redundant data
+	UE_DEPRECATED(5.4, "Deprecated for not being used in physics frame offset calculation anymore and will get removed eventually. Use GetPhysicsTimestamp() and GetNetworkPhysicsTickOffset() for the physics frame offset which is automatically kept in sync by time dilation while Physics Prediction in Project Settings is enabled. Recommended when using the new flow: Set p.net.CmdOffsetEnabled = 0 to stop the legacy frame offset logic from running in the background.")
+	ENGINE_API void PushClientInput(int32 ClientInputFrame, TArray<uint8>& Data);
+	
+	// Client says "Here is input frame number X" (and then calls other RPCs to deliver InputCmd payload)
+	UE_DEPRECATED(5.4, "Deprecated for not being used in physics frame offset calculation anymore and will get removed eventually. Use GetPhysicsTimestamp() and GetNetworkPhysicsTickOffset() for the physics frame offset which is automatically kept in sync by time dilation while Physics Prediction in Project Settings is enabled. Recommended when using the new flow: Set p.net.CmdOffsetEnabled = 0 to stop the legacy frame offset logic from running in the background.")
+	UFUNCTION(Server, unreliable)
+	ENGINE_API void ServerRecvClientInputFrame(int32 RecvClientInputFrame, const TArray<uint8>& Data);
+
+	UE_DEPRECATED(5.4, "Set via NetworkPhysicsComponent to enable networking of FInputCmdBuffer which keeps the legacy GetLocalToServerAsyncPhysicsTickOffset() and GetAsyncPhysicsTimestamp() functions updated. Use GetNetworkPhysicsTickOffset() and GetPhysicsTimestamp() instead.")
+	ENGINE_API void EnableNetworkedPhysicsInputSync(bool EnablePrediction)
+	{
+		bSyncInputsForNetworkedPhysics_DEPRECATED = EnablePrediction;
+	}
+
+	PRAGMA_DISABLE_DEPRECATION_WARNINGS
+	UE_DEPRECATED(5.4, "Deprecated for not being used in physics frame offset calculation anymore and will get removed eventually. Use GetPhysicsTimestamp() and GetNetworkPhysicsTickOffset() for the physics frame offset which is automatically kept in sync by time dilation while Physics Prediction in Project Settings is enabled. Recommended when using the new flow: Set p.net.CmdOffsetEnabled = 0 to stop the legacy frame offset logic from running in the background.")
+	const FClientFrameInfo& GetClientFrameInfo() const { return ClientFrameInfo_DEPRECATED; }
+	PRAGMA_ENABLE_DEPRECATION_WARNINGS
+
+	PRAGMA_DISABLE_DEPRECATION_WARNINGS
+	UE_DEPRECATED(5.4, "Deprecated for not being used in physics frame offset calculation anymore and will get removed eventually. Use GetPhysicsTimestamp() and GetNetworkPhysicsTickOffset() for the physics frame offset which is automatically kept in sync by time dilation while Physics Prediction in Project Settings is enabled. Recommended when using the new flow: Set p.net.CmdOffsetEnabled = 0 to stop the legacy frame offset logic from running in the background.")
+	FClientFrameInfo& GetClientFrameInfo() { return ClientFrameInfo_DEPRECATED; }
+	PRAGMA_ENABLE_DEPRECATION_WARNINGS
+
+	// -------------------------------------------------------------------------
+	// Server
+	// -------------------------------------------------------------------------
+	
+	/** DEPRECATED 5.4, physics frame offset and time dilation handled via ClientAckTimeDilation() and ClientSetupAsyncPhysicsTimestamp() */
+	struct UE_DEPRECATED(5.4, "Deprecated for not being used in physics frame offset calculation anymore and will get removed eventually. Use GetPhysicsTimestamp() and GetNetworkPhysicsTickOffset() for the physics frame offset which is automatically kept in sync by time dilation while Physics Prediction in Project Settings is enabled. Recommended when using the new flow: Set p.net.CmdOffsetEnabled = 0 to stop the legacy frame offset logic from running in the background.")
+		FServerFrameInfo
+	{
+		PRAGMA_DISABLE_DEPRECATION_WARNINGS
+		int32 LastProcessedInputFrame = INDEX_NONE;	// The last client frame number we processed. "processed" is arbitrary and we are informed about when commands are processed via SetServerProessedInputFrame
+		int32 LastLocalFrame = INDEX_NONE; // The local frame number that we processed the latest client input frame on. Again, processed is arbitrary and set via SetServerProessedInputFrame
+		int32 LastSentLocalFrame = INDEX_NONE;	// Tracks the latest LastLocalFrame that we sent to the client. Just to prevent redundantly sending info via RPC
+		float TargetTimeDilation = 1.f;
+		int8 QuantizedTimeDilation = 1; // Server sets this to tell client to slowdown or speed up
+		float TargetNumBufferedCmds = 1.f; // How many buffered cmds the server thinks this client should ideally have to absorb PL and latency variance
+		bool bFault = true;
+		PRAGMA_ENABLE_DEPRECATION_WARNINGS
+
+		PRAGMA_DISABLE_DEPRECATION_WARNINGS
+		FServerFrameInfo() = default;
+		~FServerFrameInfo() = default;
+		FServerFrameInfo(const FServerFrameInfo&) = default;
+		FServerFrameInfo(FServerFrameInfo&&) = default;
+		FServerFrameInfo& operator=(const FServerFrameInfo&) = default;
+		FServerFrameInfo& operator=(FServerFrameInfo&&) = default;
+		PRAGMA_ENABLE_DEPRECATION_WARNINGS
+	};
+	
+	// We call this in ::SendClientAdjustment to tell the client what the last processed input frame was for it and on what local frame number it was processed
+	UE_DEPRECATED(5.4, "Deprecated for not being used in physics frame offset calculation anymore and will get removed eventually. Use GetPhysicsTimestamp() and GetNetworkPhysicsTickOffset() for the physics frame offset which is automatically kept in sync by time dilation while Physics Prediction in Project Settings is enabled. Recommended when using the new flow: Set p.net.CmdOffsetEnabled = 0 to stop the legacy frame offset logic from running in the background.")
+	UFUNCTION(Client, unreliable)
+	ENGINE_API void ClientRecvServerAckFrame(int32 LastProcessedInputFrame, int32 RecvServerFrameNumber, int8 TimeDilation);
+	
+	UE_DEPRECATED(5.4, "Deprecated for not being used in physics frame offset calculation anymore and will get removed eventually. Use GetPhysicsTimestamp() and GetNetworkPhysicsTickOffset() for the physics frame offset which is automatically kept in sync by time dilation while Physics Prediction in Project Settings is enabled. Recommended when using the new flow: Set p.net.CmdOffsetEnabled = 0 to stop the legacy frame offset logic from running in the background.")
+	UFUNCTION(Client, unreliable)
+	ENGINE_API void ClientRecvServerAckFrameDebug(uint8 NumBuffered, float TargetNumBufferedCmds);
+	
+	PRAGMA_DISABLE_DEPRECATION_WARNINGS
+	UE_DEPRECATED(5.4, "Deprecated for not being used in physics frame offset calculation anymore and will get removed eventually. Use GetPhysicsTimestamp() and GetNetworkPhysicsTickOffset() for the physics frame offset which is automatically kept in sync by time dilation while Physics Prediction in Project Settings is enabled. Recommended when using the new flow: Set p.net.CmdOffsetEnabled = 0 to stop the legacy frame offset logic from running in the background.")
+	FServerFrameInfo& GetServerFrameInfo() { return ServerFrameInfo_DEPRECATED; };
+	PRAGMA_ENABLE_DEPRECATION_WARNINGS
+
 private:
 
+	/** Deprecated 5.4 */
+	PRAGMA_DISABLE_DEPRECATION_WARNINGS
+	FInputCmdBuffer InputBuffer_DEPRECATED;
+	FClientFrameInfo ClientFrameInfo_DEPRECATED;
+	FServerFrameInfo ServerFrameInfo_DEPRECATED;
+	PRAGMA_ENABLE_DEPRECATION_WARNINGS
+
+	/** Deprecated 5.4, set by NetworkPhysicsComponent to sync FInputCmdBuffer used in calculating client/server physics frame offset. */
+	bool bSyncInputsForNetworkedPhysics_DEPRECATED = false;
+
 	/** The estimated offset between the local async physics tick frame number and the server's
 	*	This is used to synchronize events that happen in the async physics tick */
-	int32 LocalToServerAsyncPhysicsTickOffset = INDEX_NONE;
-	bool LocalToServerAsyncPhysicsTickOffsetAssigned = false;
+	int32 LocalToServerAsyncPhysicsTickOffset_DEPRECATED = INDEX_NONE;
 	
-	/** The latest physics step we've sent to the server. Due to async we need to avoid duplicate sends */
-	int32 ClientLatestAsyncPhysicsStepSent = INDEX_NONE;
+	/** The estimated offset between the server async physics tick frame number and the local's
+	*	This is used to synchronize events that happen in the async physics tick */
+	int32 ServerToLocalAsyncPhysicsTickOffset_DEPRECATED = INDEX_NONE;
+
+	/** The latest server step we've received an offset correction for. This allows us to ignore out of order corrections that arrive late */
+	int32 ClientLatestCorrectedOffsetServerStep_DEPRECATED = INDEX_NONE;
+	
+	/** The server records the latest timestamp it has to correct. This is used to update client (which may not happen on every physics step) */
+	FAsyncPhysicsTimestamp ServerLatestTimestampToCorrect_DEPRECATED;
+	
+	/** The latest timestamp the client has sent to the server with prediction of which server frame it corresponds to. */
+	FAsyncPhysicsTimestamp ServerPendingTimestamp_DEPRECATED;
+
+	UE_DEPRECATED(5.4, "Deprecated for not being used in physics frame offset calculation anymore and will get removed eventually. Use GetPhysicsTimestamp() and GetNetworkPhysicsTickOffset() for the physics frame offset which is automatically kept in sync by time dilation while Physics Prediction in Project Settings is enabled. Recommended when using the new flow: Set p.net.CmdOffsetEnabled = 0 to stop the legacy frame offset logic from running in the background.")
+	UFUNCTION(Client, Unreliable)
+	ENGINE_API void ClientCorrectionAsyncPhysicsTimestamp(FAsyncPhysicsTimestamp Timestamp);
+
+	/** Update the tick offset in between the local client and the server */
+	ENGINE_API virtual void AsyncPhysicsTickActor(float DeltaTime, float SimTime) override;
+
+public:
+
+	/** Generates a timestamp for the upcoming physics step (plus any pending time). Useful for synchronizing client and server events on a specific physics step */
+	UE_DEPRECATED(5.4, "Use GetPhysicsTimestamp() and GetNetworkPhysicsTickOffset() for the physics frame offset which is automatically kept in sync by time dilation while Physics Prediction in Project Settings is enabled. Recommended when using the new flow: Set p.net.CmdOffsetEnabled = 0 to stop the legacy frame offset logic from running in the background.")
+	ENGINE_API FAsyncPhysicsTimestamp GetAsyncPhysicsTimestamp(float DeltaSeconds = 0.f);
+
+	/** Returns the current estimated offset between the local async physics step and the server. This is useful for dealing with low level synchronization.
+		In general it's recommended to use GetAsyncPhysicsTimestamp which accounts for the offset automatically*/
+	UE_DEPRECATED(5.4, "Use GetPhysicsTimestamp() and GetNetworkPhysicsTickOffset() for the physics frame offset which is automatically kept in sync by time dilation while Physics Prediction in Project Settings is enabled. Recommended when using the new flow: Set p.net.CmdOffsetEnabled = 0 to stop the legacy frame offset logic from running in the background.")
+	int32 GetLocalToServerAsyncPhysicsTickOffset() const { return LocalToServerAsyncPhysicsTickOffset_DEPRECATED; }
+
+	/** Returns the current estimated offset between the server async physics step and the local one. */
+	UE_DEPRECATED(5.4, "Use GetPhysicsTimestamp() and GetNetworkPhysicsTickOffset() for the physics frame offset which is automatically kept in sync by time dilation while Physics Prediction in Project Settings is enabled. Recommended when using the new flow: Set p.net.CmdOffsetEnabled = 0 to stop the legacy frame offset logic from running in the background.")
+	int32 GetServerToLocalAsyncPhysicsTickOffset() const { return (ServerToLocalAsyncPhysicsTickOffset_DEPRECATED != INDEX_NONE) ? ServerToLocalAsyncPhysicsTickOffset_DEPRECATED : LocalToServerAsyncPhysicsTickOffset_DEPRECATED; }
+	
+	/** Set the offset between the server async physics step and the local one.*/
+	UE_DEPRECATED(5.4, "Use GetPhysicsTimestamp() and GetNetworkPhysicsTickOffset() for the physics frame offset which is automatically kept in sync by time dilation while Physics Prediction in Project Settings is enabled. Recommended when using the new flow: Set p.net.CmdOffsetEnabled = 0 to stop the legacy frame offset logic from running in the background.")
+	void SetServerToLocalAsyncPhysicsTickOffset(const int32 AsyncPhysicsTickOffset) { ServerToLocalAsyncPhysicsTickOffset_DEPRECATED = AsyncPhysicsTickOffset; }
+
+	UE_DEPRECATED(5.4, "Deprecated for not being used in physics frame offset calculation anymore and will get removed eventually. Use GetPhysicsTimestamp() and GetNetworkPhysicsTickOffset() for the physics frame offset which is automatically kept in sync by time dilation while Physics Prediction in Project Settings is enabled. Recommended when using the new flow: Set p.net.CmdOffsetEnabled = 0 to stop the legacy frame offset logic from running in the background.")
+	void UpdateServerTimestampToCorrect();
+
+private:
+
+	/** The static offset between the local async physics tick frame number and the server's, kept in sync via time-dilation
+	*	This is used to synchronize events that happen in the async physics tick */
+	int32 NetworkPhysicsTickOffset = INDEX_NONE;
+	bool bNetworkPhysicsTickOffsetAssigned = false;
 
 	/** The latest server step we've received a time dilation for. Needed for out of order updates */
 	int32 ClientLatestTimeDilationServerStep = INDEX_NONE;
 
+	/** The latest physics step we've sent to the server. Due to async we need to avoid duplicate sends */
+	int32 ClientLatestAsyncPhysicsStepSent = INDEX_NONE;
+
 	/** The latest physics step we've received from the client. */
 	int32 ServerLatestAsyncPhysicsStepReceived = INDEX_NONE;
 
-	/** Update the tick ofsset in between the local client and the server */
+public:
+
+	/** Update the tick offset in between the local client and the server */
 	ENGINE_API void UpdateServerAsyncPhysicsTickOffset();
 
+	/** Server receives the clients FAsyncPhysicsTimestamp with its predicted ServerFrame and clients LocalFrame, to updates the time dilation needed to keep the ServerFrame prediction in sync */
 	UFUNCTION(Server, Unreliable)
 	ENGINE_API void ServerSendLatestAsyncPhysicsTimestamp(FAsyncPhysicsTimestamp Timestamp);
 
+	/** Client receives the setup of the FAsyncPhysicsTimestamp ServerFrame and LocalFrame offset from the server. */
 	UFUNCTION(Client, Reliable)
-	ENGINE_API void ClientCorrectionAsyncPhysicsTimestamp(FAsyncPhysicsTimestamp Timestamp);
+	ENGINE_API void ClientSetupNetworkPhysicsTimestamp(FAsyncPhysicsTimestamp Timestamp);
 
+	/** Client receives the time dilation value it needs to use to keep its ServerFrame to LocalFrame offset in sync */
 	UFUNCTION(Client, Unreliable)
 	ENGINE_API void ClientAckTimeDilation(float TimeDilation, int32 ServerStep);
 
-	/** Update the tick offset in between the local client and the server */
-	ENGINE_API virtual void AsyncPhysicsTickActor(float DeltaTime, float SimTime) override;
-
-public:
-
 	/** Enqueues a command to run at the time specified by AsyncPhysicsTimestamp. Note that if the time specified was missed the command is triggered as soon as possible as part of the async tick.
-		These commands are all run on the game thread. If you want to run on the physics thread see FPhysicsSolverBase::RegisterSimOneShotCallback
-		If OwningObject is not null this command will only fire as long as the object is still alive. This allows a lambda to still use the owning object's data for read/write (assuming data is designed for async physics tick)
+	These commands are all run on the game thread. If you want to run on the physics thread see FPhysicsSolverBase::RegisterSimOneShotCallback
+	If OwningObject is not null this command will only fire as long as the object is still alive. This allows a lambda to still use the owning object's data for read/write (assuming data is designed for async physics tick)
 	*/
 	ENGINE_API void ExecuteAsyncPhysicsCommand(const FAsyncPhysicsTimestamp& AsyncPhysicsTimestamp, UObject* OwningObject, const TFunction<void()>& Command, const bool bEnableResim = true);
 
-	/** Generates a timestamp for the upcoming physics step (plus any pending time). Useful for synchronizing client and server events on a specific physics step */
-	ENGINE_API FAsyncPhysicsTimestamp GetAsyncPhysicsTimestamp(float DeltaSeconds = 0.f);
+	/** Generates a timestamp for the upcoming physics step (plus any pending time). Useful for synchronizing client and server events on a specific physics step 
+	* On the server LocalFrame and ServerFrame are the same, on the client LocalFrame is the current physics frame and ServerFrame is the server-side frame that the LocalFrame corresponds.
+	*/
+	ENGINE_API FAsyncPhysicsTimestamp GetPhysicsTimestamp(float DeltaSeconds = 0.0f);
+
+	/** Get the physics frame number offset between Local and Server, recommended is to use GetPhysicsTimestamp() which takes this offset into account on the client for the ServerFrame value */
+	int32 GetNetworkPhysicsTickOffset() const { return NetworkPhysicsTickOffset; }
+
+	/** True if the NetworkPhysicsTickOffset is setup, before this is setup the GetNetworkPhysicsTickOffset() and GetPhysicsTimestamp() functions will not return valid results */
+	bool GetNetworkPhysicsTickOffsetAssigned() const { return bNetworkPhysicsTickOffsetAssigned; }
 
-	/** Returns the current estimated offset between the local async physics step and the server. This is useful for dealing with low level synchronization.
-		In general it's recommended to use GetAsyncPhysicsTimestamp which accounts for the offset automatically*/
-	int32 GetLocalToServerAsyncPhysicsTickOffset() const { return LocalToServerAsyncPhysicsTickOffset; }	
-	bool GetLocalToServerAsyncPhysicsTickOffsetAssigned() const { return LocalToServerAsyncPhysicsTickOffsetAssigned; }
 };
 
 #if UE_ENABLE_INCLUDE_ORDER_DEPRECATED_IN_5_2
diff --git a/Engine/Source/Runtime/Engine/Private/Components/PawnMovementComponent.cpp b/Engine/Source/Runtime/Engine/Private/Components/PawnMovementComponent.cpp
--- a/Engine/Source/Runtime/Engine/Private/Components/PawnMovementComponent.cpp
+++ b/Engine/Source/Runtime/Engine/Private/Components/PawnMovementComponent.cpp
@@ -132,7 +132,7 @@ void UPawnMovementComponent::ApplyAsyncPhysicsStateAction(const UPrimitiveCompon
 		{
 			if (PlayerController->IsLocalController() && !PlayerController->IsNetMode(NM_DedicatedServer))
 			{
-				FAsyncPhysicsTimestamp TimeStamp = PlayerController->GetAsyncPhysicsTimestamp();
+				FAsyncPhysicsTimestamp TimeStamp = PlayerController->GetPhysicsTimestamp();
 				ExecuteAsyncPhysicsStateAction(ActionComponent, BoneName, TimeStamp, ActionType, ActionDatas, ActionPosition);
 
 				if (!PlayerController->IsNetMode(NM_Standalone))
@@ -168,7 +168,7 @@ void UPawnMovementComponent::MulticastAsyncPhysicsStateAction_Implementation(con
 			if (GetController() == nullptr)
 			{
 				FAsyncPhysicsTimestamp LocalTimeStamp = Timestamp;
-				LocalTimeStamp.LocalFrame = LocalTimeStamp.ServerFrame - PlayerController->GetLocalToServerAsyncPhysicsTickOffset();
+				LocalTimeStamp.LocalFrame = LocalTimeStamp.ServerFrame - PlayerController->GetNetworkPhysicsTickOffset();
 
 				ExecuteAsyncPhysicsStateAction(ActionComponent, BoneName, LocalTimeStamp, ActionType, ActionDatas, ActionPosition);
 			}
diff --git a/Engine/Source/Runtime/Engine/Private/PhysicsEngine/AsyncPhysicsInputComponent.cpp b/Engine/Source/Runtime/Engine/Private/PhysicsEngine/AsyncPhysicsInputComponent.cpp
--- a/Engine/Source/Runtime/Engine/Private/PhysicsEngine/AsyncPhysicsInputComponent.cpp
+++ b/Engine/Source/Runtime/Engine/Private/PhysicsEngine/AsyncPhysicsInputComponent.cpp
@@ -39,9 +39,9 @@ struct FAsyncPhysicsInputRewindCallback : public Chaos::IRewindCallback
 			if (APlayerController* PC = World->GetFirstPlayerController())
 			{
 				ensure(Chaos::FPhysicsSolverBase::IsNetworkPhysicsPredictionEnabled());
-				if (PC->GetLocalToServerAsyncPhysicsTickOffsetAssigned())
+				if (PC->GetNetworkPhysicsTickOffsetAssigned())
 				{
-					const int32 LocalToServerOffset = PC->GetLocalToServerAsyncPhysicsTickOffset();
+					const int32 LocalToServerOffset = PC->GetNetworkPhysicsTickOffset();
 					ServerFrame = PhysicsStep + LocalToServerOffset;
 				}
 			}
@@ -74,9 +74,9 @@ struct FAsyncPhysicsInputRewindCallback : public Chaos::IRewindCallback
 				// -----------------------------------------------------------------
 
 				ensure(Chaos::FPhysicsSolverBase::IsNetworkPhysicsPredictionEnabled());
-				if (PC->GetLocalToServerAsyncPhysicsTickOffsetAssigned())
+				if (PC->GetNetworkPhysicsTickOffsetAssigned())
 				{
-					const int32 LocalToServerOffset = PC->GetLocalToServerAsyncPhysicsTickOffset();
+					const int32 LocalToServerOffset = PC->GetNetworkPhysicsTickOffset();
 					CachedServerFrame = PhysicsStep + LocalToServerOffset;
 				}
 			}
diff --git a/Engine/Source/Runtime/Engine/Private/PhysicsEngine/BodyInstance.cpp b/Engine/Source/Runtime/Engine/Private/PhysicsEngine/BodyInstance.cpp
--- a/Engine/Source/Runtime/Engine/Private/PhysicsEngine/BodyInstance.cpp
+++ b/Engine/Source/Runtime/Engine/Private/PhysicsEngine/BodyInstance.cpp
@@ -3342,7 +3342,7 @@ void FBodyInstance::ApplyAsyncPhysicsCommand(FAsyncPhysicsTimestamp TimeStamp, c
 			OwnerComponent->GetWorld() ? OwnerComponent->GetWorld()->GetFirstPlayerController() : nullptr;
 		if (LocalController)
 		{
-			TimeStamp.LocalFrame = TimeStamp.ServerFrame - LocalController->GetLocalToServerAsyncPhysicsTickOffset();
+			TimeStamp.LocalFrame = TimeStamp.ServerFrame - LocalController->GetNetworkPhysicsTickOffset();
 			LocalController->ExecuteAsyncPhysicsCommand(TimeStamp, OwnerComponent.Get(), Command);
 		}
 	}
diff --git a/Engine/Source/Runtime/Engine/Private/PhysicsEngine/NetworkPhysicsComponent.cpp b/Engine/Source/Runtime/Engine/Private/PhysicsEngine/NetworkPhysicsComponent.cpp
--- a/Engine/Source/Runtime/Engine/Private/PhysicsEngine/NetworkPhysicsComponent.cpp
+++ b/Engine/Source/Runtime/Engine/Private/PhysicsEngine/NetworkPhysicsComponent.cpp
@@ -14,6 +14,61 @@
 
 #include UE_INLINE_GENERATED_CPP_BY_NAME(NetworkPhysicsComponent)
 
+/** These CVars are deprecated from UE 5.4, physics frame offset for networked physics prediction is now handled via PlayerController with automatic time dilation
+* p.net.CmdOffsetEnabled = 0 is recommended to disable the deprecated flow */
+namespace InputCmdCVars
+{
+	static bool bCmdOffsetEnabled = true;
+	static FAutoConsoleVariableRef CVarCmdOffsetEnabled(TEXT("p.net.CmdOffsetEnabled"), bCmdOffsetEnabled, TEXT("Enables deprecated (5.4) logic for legacy that handles physics frame offset. Recommended: Set this to 0 to stop the deprecated physics frame offset flow. "));
+
+	static int32 ForceFault = 0;
+	static FAutoConsoleVariableRef CVarForceFault(TEXT("p.net.ForceFault"), ForceFault, TEXT("Forces server side input fault"));
+	static int32 MaxBufferedCmds = 16;
+	static FAutoConsoleVariableRef CVarMaxBufferedCmds(TEXT("p.net.MaxBufferedCmds"), MaxBufferedCmds, TEXT("MaxNumber of buffered server side commands"));
+	static int32 TimeDilationEnabled = 0;
+	static FAutoConsoleVariableRef CVarTimeDilationEnabled(TEXT("p.net.TimeDilationEnabled"), TimeDilationEnabled, TEXT("Enable clientside TimeDilation"));
+	static float MaxTargetNumBufferedCmds = 5.0;
+	static FAutoConsoleVariableRef CVarMaxTargetNumBufferedCmds(TEXT("p.net.MaxTargetNumBufferedCmds"), MaxTargetNumBufferedCmds, TEXT("Maximum number of buffered inputs the server will target per client."));
+	static float MaxTimeDilationMag = 0.01f;
+	static FAutoConsoleVariableRef CVarMaxTimeDilationMag(TEXT("p.net.MaxTimeDilationMag"), MaxTimeDilationMag, TEXT("Maximum time dilation that client will use to slow down / catch up with server"));
+	static float TimeDilationAlpha = 0.1f;
+	static FAutoConsoleVariableRef CVarTimeDilationAlpha(TEXT("p.net.TimeDilationAlpha"), TimeDilationAlpha, TEXT("Lerp strength for sliding client time dilation"));
+	static float TargetNumBufferedCmdsDeltaOnFault = 1.0f;
+	static FAutoConsoleVariableRef CVarTargetNumBufferedCmdsDeltaOnFault(TEXT("p.net.TargetNumBufferedCmdsDeltaOnFault"), TargetNumBufferedCmdsDeltaOnFault, TEXT("How much to increase TargetNumBufferedCmds when an input fault occurs"));
+	static float TargetNumBufferedCmds = 1.9f;
+	static FAutoConsoleVariableRef CVarTargetNumBufferedCmds(TEXT("p.net.TargetNumBufferedCmds"), TargetNumBufferedCmds, TEXT("How much to increase TargetNumBufferedCmds when an input fault occurs"));
+	static float TargetNumBufferedCmdsAlpha = 0.005f;
+	static FAutoConsoleVariableRef CVarTargetNumBufferedCmdsAlpha(TEXT("p.net.TargetNumBufferedCmdsAlpha"), TargetNumBufferedCmdsAlpha, TEXT("Lerp strength for TargetNumBufferedCmds"));
+	static int32 LerpTargetNumBufferedCmdsAggresively = 0;
+	static FAutoConsoleVariableRef CVarLerpTargetNumBufferedCmdsAggresively(TEXT("p.net.LerpTargetNumBufferedCmdsAggresively"), LerpTargetNumBufferedCmdsAggresively, TEXT("Aggresively lerp towards TargetNumBufferedCmds. Reduces server side buffering but can cause more artifacts."));
+}
+// --------------------------------------------------------------------------------------------------------------------------------------------------
+//	Client InputCmd Stream stuff
+// --------------------------------------------------------------------------------------------------------------------------------------------------
+namespace
+{
+	int8 QuantizeTimeDilation(float F)
+	{
+		if (F == 1.f)
+		{
+			return 0;
+		}
+		float Normalized = FMath::Clamp<float>((F - 1.f) / InputCmdCVars::MaxTimeDilationMag, -1.f, 1.f);
+		return (int8)(Normalized * 128.f);
+	}
+
+	float DeQuantizeTimeDilation(int8 i)
+	{
+		if (i == 0)
+		{
+			return 1.f;
+		}
+		float Normalized = (float)i / 128.f;
+		float Uncompressed = 1.f + (Normalized * InputCmdCVars::MaxTimeDilationMag);
+		return Uncompressed;
+	}
+}
+
 bool FNetworkPhysicsRewindDataProxy::NetSerializeBase(FArchive& Ar, class UPackageMap* Map, bool& bOutSuccess, TUniqueFunction<TUniquePtr<Chaos::FBaseRewindHistory>()> CreateHistoryFunction)
 {
 	Ar << Owner;
@@ -131,6 +186,97 @@ int32 FNetworkPhysicsCallback::TriggerRewindIfNeeded_Internal(int32 LatestStepCo
 	return ResimFrame;
 }
 
+/* Deprecated 5.4 */
+void FNetworkPhysicsCallback::UpdateClientPlayer_External(int32 PhysicsStep)
+{
+	if (APlayerController* PC = World->GetFirstPlayerController())
+	{
+		// ------------------------------------------------
+		// Send RPC to server telling them what (client/local) physics step we are running
+		//	* Note that SendData is empty because of the existing API, should change this
+		// ------------------------------------------------	
+		PRAGMA_DISABLE_DEPRECATION_WARNINGS
+		TArray<uint8> SendData;
+		PC->PushClientInput(PhysicsStep, SendData);
+		PRAGMA_ENABLE_DEPRECATION_WARNINGS
+
+		if (InputCmdCVars::TimeDilationEnabled > 0)
+		{
+			UE_LOG(LogPhysics, Warning, TEXT("p.net.TimeDilationEnabled is set to true, this CVar is deprecated in UE5.4 and does not affect Time Dilation. Time Dilation is automatically used via the PlayerController if Physics Prediction is enabled in Project Settings. It's also recommended to disable the legacy flow that handled physics frame offset and this time dilation by setting: p.net.CmdOffsetEnabled = 0"));
+		}
+	}
+}
+
+/* Deprecated 5.4 */
+void FNetworkPhysicsCallback::UpdateServerPlayer_External(int32 PhysicsStep)
+{
+	// -----------------------------------------------
+	// Server: "consume" an InputCmd from each Player Controller
+	// All this means in this context is updating FServerFrameInfo:: LastProcessedInputFrame, LastLocalFrame
+	// (E.g, telling each client what "Input" of theirs we were processing and our local physics frame number.
+	// In cases where the buffer has a fault, we calculate a suggested time dilation to temporarily make client speed up 
+	// or slow down their input cmd production.
+	// -----------------------------------------------
+	const bool bForceFault = InputCmdCVars::ForceFault > 0;
+	InputCmdCVars::ForceFault = FMath::Max(0, InputCmdCVars::ForceFault - 1);
+	for (FConstPlayerControllerIterator Iterator = World->GetPlayerControllerIterator(); Iterator; ++Iterator)
+	{
+		PRAGMA_DISABLE_DEPRECATION_WARNINGS
+		if (APlayerController* PC = Iterator->Get())
+		{
+			PC->UpdateServerTimestampToCorrect();
+
+			APlayerController::FServerFrameInfo& FrameInfo = PC->GetServerFrameInfo();
+			APlayerController::FInputCmdBuffer& InputBuffer = PC->GetInputBuffer();
+
+			const int32 NumBufferedInputCmds = bForceFault ? 0 : (InputBuffer.HeadFrame() - FrameInfo.LastProcessedInputFrame);
+			// Check Overflow
+			if (NumBufferedInputCmds > InputCmdCVars::MaxBufferedCmds)
+			{
+				UE_LOG(LogPhysics, Warning, TEXT("[Remote.Input] overflow %d %d -> %d"), InputBuffer.HeadFrame(), FrameInfo.LastProcessedInputFrame, NumBufferedInputCmds);
+				FrameInfo.LastProcessedInputFrame = InputBuffer.HeadFrame() - InputCmdCVars::MaxBufferedCmds + 1;
+			}
+			// Check fault - we are waiting for Cmds to reach TargetNumBufferedCmds before continuing
+			if (FrameInfo.bFault)
+			{
+				if (NumBufferedInputCmds < (int32)FrameInfo.TargetNumBufferedCmds)
+				{
+					// Skip this because it is in fault. We will use the prev input for this frame.
+					UE_CLOG(FrameInfo.LastProcessedInputFrame != INDEX_NONE, LogPhysics, Warning, TEXT("[Remote.Input] in fault. Reusing Inputcmd. (Client) Input: %d. (Server) Local Frame: %d"), FrameInfo.LastProcessedInputFrame, FrameInfo.LastLocalFrame);
+					continue;
+				}
+				FrameInfo.bFault = false;
+			}
+			else if (NumBufferedInputCmds <= 0)
+			{
+				// No Cmds to process, enter fault state. Increment TargetNumBufferedCmds each time this happens.
+				// TODO: We should have something to bring this back down (which means skipping frames) we don't want temporary poor conditions to cause permanent high input buffering
+				FrameInfo.bFault = true;
+				FrameInfo.TargetNumBufferedCmds = FMath::Min(FrameInfo.TargetNumBufferedCmds + InputCmdCVars::TargetNumBufferedCmdsDeltaOnFault, InputCmdCVars::MaxTargetNumBufferedCmds);
+				UE_CLOG(FrameInfo.LastProcessedInputFrame != INDEX_NONE, LogPhysics, Warning, TEXT("[Remote.Input] ENTERING fault. New Target: %.2f. (Client) Input: %d. (Server) Local Frame: %d"), FrameInfo.TargetNumBufferedCmds, FrameInfo.LastProcessedInputFrame, FrameInfo.LastLocalFrame);
+				continue;
+			}
+			float TargetTimeDilation = 1.f;
+			if (NumBufferedInputCmds < (int32)FrameInfo.TargetNumBufferedCmds)
+			{
+				TargetTimeDilation += InputCmdCVars::MaxTimeDilationMag; // Tell client to speed up, we are starved on cmds
+			}
+			FrameInfo.TargetTimeDilation = FMath::Lerp(FrameInfo.TargetTimeDilation, TargetTimeDilation, InputCmdCVars::TimeDilationAlpha);
+				
+			FrameInfo.QuantizedTimeDilation = QuantizeTimeDilation(TargetTimeDilation);
+				
+			if (InputCmdCVars::LerpTargetNumBufferedCmdsAggresively != 0)
+			{
+				// When aggressive, always lerp towards target
+				FrameInfo.TargetNumBufferedCmds = FMath::Lerp(FrameInfo.TargetNumBufferedCmds, InputCmdCVars::TargetNumBufferedCmds, InputCmdCVars::TargetNumBufferedCmdsAlpha);
+			}
+			FrameInfo.LastProcessedInputFrame++;
+			FrameInfo.LastLocalFrame = PhysicsStep;
+		}
+		PRAGMA_ENABLE_DEPRECATION_WARNINGS
+	}
+}
+
 void FNetworkPhysicsCallback::InjectInputs_External(int32 PhysicsStep, int32 NumSteps)
 {
 	InjectInputsExternal.Broadcast(PhysicsStep, NumSteps);
@@ -145,6 +291,21 @@ void FNetworkPhysicsCallback::ProcessInputs_External(int32 PhysicsStep, const TA
 			SimCallbackObject.CallbackObject->ProcessInputs_External(PhysicsStep);
 		}
 	}
+	
+	/* Deprecated 5.4 */
+	if (InputCmdCVars::bCmdOffsetEnabled)
+	{
+		PRAGMA_DISABLE_DEPRECATION_WARNINGS
+		if (NetMode == NM_Client)
+		{
+			UpdateClientPlayer_External(PhysicsStep);
+		}
+		else
+		{
+			UpdateServerPlayer_External(PhysicsStep);
+		}
+		PRAGMA_ENABLE_DEPRECATION_WARNINGS
+	}
 }
 
 UNetworkPhysicsSystem::UNetworkPhysicsSystem()
@@ -334,9 +495,9 @@ void UNetworkPhysicsComponent::SendLocalInputsDatas()
 	if (PlayerController->IsLocalController() && InputsHistory)
 	{
 		// We just check that the local client to server offset is valid before doing something
-		if (HasServerWorld() || PlayerController->GetLocalToServerAsyncPhysicsTickOffsetAssigned())
+		if (HasServerWorld() || PlayerController->GetNetworkPhysicsTickOffsetAssigned())
 		{
-			const int32 LocalOffset = HasServerWorld() ? 0 : PlayerController->GetLocalToServerAsyncPhysicsTickOffset();
+			const int32 LocalOffset = HasServerWorld() ? 0 : PlayerController->GetNetworkPhysicsTickOffset();
 			const int32 NextIndex = (InputsIndex + 1) % (InputsRedundancy + 1);
 
 			// if on server (Listen server) we should send the inputs onto all the clients through repnotify
@@ -345,7 +506,7 @@ void UNetworkPhysicsComponent::SendLocalInputsDatas()
 			if (!HasServerWorld())
 			{
 #if DEBUG_NETWORK_PHYSICS
-				FAsyncPhysicsTimestamp Timestamp = PlayerController->GetAsyncPhysicsTimestamp();
+				FAsyncPhysicsTimestamp Timestamp = PlayerController->GetPhysicsTimestamp();
 
 				TArray<int32> LocalFrames, ServerFrames, InputFrames;
 				InputsHistory->DebugDatas(*ReplicatedInputs.History, LocalFrames, ServerFrames, InputFrames);
@@ -377,6 +538,31 @@ void UNetworkPhysicsComponent::SendLocalStatesDatas()
 	}
 }
 
+/* Deprecated 5.4 */
+void UNetworkPhysicsComponent::CorrectServerToLocalOffset(const int32 LocalToServerOffset)
+{
+	if (HasLocalController() && !HasServerWorld() && StatesHistory)
+	{
+		TArray<int32> LocalFrames, ServerFrames, InputFrames;
+		StatesHistory->DebugDatas(*ReplicatedStates.History, LocalFrames, ServerFrames, InputFrames);
+		int32 ServerToLocalOffset = LocalToServerOffset;
+		for (int32 FrameIndex = 0; FrameIndex < LocalFrames.Num(); ++FrameIndex)
+		{
+#if DEBUG_NETWORK_PHYSICS || DEBUG_REWIND_DATA
+			UE_LOG(LogTemp, Log, TEXT("CLIENT | GT | CorrectServerToLocalOffset | Server frame = %d | Client Frame = %d"), ServerFrames[FrameIndex], InputFrames[FrameIndex]);
+#endif
+			ServerToLocalOffset = FMath::Min(ServerToLocalOffset, ServerFrames[FrameIndex] - InputFrames[FrameIndex]);
+		}
+		PRAGMA_DISABLE_DEPRECATION_WARNINGS
+		GetPlayerController()->SetServerToLocalAsyncPhysicsTickOffset(ServerToLocalOffset);
+		PRAGMA_ENABLE_DEPRECATION_WARNINGS
+
+#if DEBUG_NETWORK_PHYSICS || DEBUG_REWIND_DATA
+		UE_LOG(LogTemp, Log, TEXT("CLIENT | GT | CorrectServerToLocalOffset | Server to local offset = %d | Local to server offset = %d"), ServerToLocalOffset, LocalToServerOffset);
+#endif
+	}
+}
+
 void UNetworkPhysicsComponent::OnRep_SetReplicatedStates()
 {
 	// The replicated states should only be used on the client since the server already have authoritative local ones
@@ -387,7 +573,17 @@ void UNetworkPhysicsComponent::OnRep_SetReplicatedStates()
 		{
 			PlayerController = GetWorld()->GetFirstPlayerController();
 		}
-		const int32 LocalOffset = PlayerController->GetLocalToServerAsyncPhysicsTickOffset();
+
+		const int32 LocalOffset = PlayerController->GetNetworkPhysicsTickOffset();
+
+		/* Deprecated 5.4 */
+		if (InputCmdCVars::bCmdOffsetEnabled)
+		{
+			PRAGMA_DISABLE_DEPRECATION_WARNINGS
+			const int32 LocalToServerOffset = PlayerController->GetLocalToServerAsyncPhysicsTickOffset();
+			CorrectServerToLocalOffset(LocalToServerOffset);
+			PRAGMA_ENABLE_DEPRECATION_WARNINGS
+		}
 
 		// Record the received states from the server into the history for future use
 		if (UWorld* World = GetWorld())
@@ -431,7 +627,8 @@ void UNetworkPhysicsComponent::OnRep_SetReplicatedInputs()
 		{
 			PlayerController = GetWorld()->GetFirstPlayerController();
 		}
-		const int32 LocalOffset = PlayerController->GetLocalToServerAsyncPhysicsTickOffset();
+
+		const int32 LocalOffset = PlayerController->GetNetworkPhysicsTickOffset();
 
 		// Record the received inputs from the server into the history for future use
 		if (UWorld* World = GetWorld())
@@ -593,7 +790,7 @@ void UNetworkPhysicsComponent::OnPostProcessInputsInternal(const int32 PhysicsSt
 		{
 			FNetworkPhysicsDatas* PhysicsDatas = InputsDatas.Get();
 			PhysicsDatas->LocalFrame = PhysicsStep;
-			PhysicsDatas->ServerFrame = HasServerWorld() ? PhysicsStep : PhysicsStep + GetPlayerController()->GetLocalToServerAsyncPhysicsTickOffset();
+			PhysicsDatas->ServerFrame = HasServerWorld() ? PhysicsStep : PhysicsStep + GetPlayerController()->GetNetworkPhysicsTickOffset();
 			PhysicsDatas->InputFrame = PhysicsStep;
 
 			PhysicsDatas->BuildDatas(ActorComponent);
diff --git a/Engine/Source/Runtime/Engine/Private/PhysicsEngine/PhysicsReplication.cpp b/Engine/Source/Runtime/Engine/Private/PhysicsEngine/PhysicsReplication.cpp
--- a/Engine/Source/Runtime/Engine/Private/PhysicsEngine/PhysicsReplication.cpp
+++ b/Engine/Source/Runtime/Engine/Private/PhysicsEngine/PhysicsReplication.cpp
@@ -329,7 +329,7 @@ void FPhysicsReplication::OnTick(float DeltaSeconds, TMap<TWeakObjectPtr<UPrimit
 			{
 				if (APlayerController* PlayerController = World->GetFirstPlayerController())
 				{
-					LocalFrameOffset = PlayerController->GetLocalToServerAsyncPhysicsTickOffset();
+					LocalFrameOffset = PlayerController->GetNetworkPhysicsTickOffset();
 				}
 			}
 		}
diff --git a/Engine/Source/Runtime/Engine/Private/PlayerController.cpp b/Engine/Source/Runtime/Engine/Private/PlayerController.cpp
--- a/Engine/Source/Runtime/Engine/Private/PlayerController.cpp
+++ b/Engine/Source/Runtime/Engine/Private/PlayerController.cpp
@@ -109,17 +109,30 @@ namespace PlayerControllerCVars
 
 namespace NetworkPhysicsCvars
 {
+	/* DEPRECATED 5.4 */
+	int32 NumRedundantCmds = 3;
+	FAutoConsoleVariableRef CVarNumRedundantCmds(TEXT("np2.NumRedundantCmds"), NumRedundantCmds, TEXT("(DEPRECATED 5.4, only part of the legacy physics frame offset logic) Number of redundant user cmds to send per frame"));
+
+#if (UE_BUILD_SHIPPING || UE_BUILD_TEST)
+	int32 EnableDebugRPC = 0;
+#else
+	int32 EnableDebugRPC = 1;
+#endif
+	/* DEPRECATED 5.4 */
+	FAutoConsoleVariableRef CVarEnableDebugRPC(TEXT("np2.EnableDebugRPC"), EnableDebugRPC, TEXT("(DEPRECATED 5.4, only part of the legacy physics frame offset logic) Sends extra debug information to clients about server side input buffering"));
+	
+	/* DEPRECATED 5.4 */
+	int32 NetworkPhysicsPredictionFrameOffset = 4;
+	FAutoConsoleVariableRef CVarNetworkPhysicsPredictionFrameOffset(TEXT("np2.NetworkPhysicsPredictionFrameOffset"), NetworkPhysicsPredictionFrameOffset, TEXT("(DEPRECATED 5.4, use np2.PredictionAsyncFrameBuffer instead) Additional frame offset to be added to the local to server offset used by network prediction"));
+	
 	int32 PredictionAsyncFrameBuffer = 3;
 	FAutoConsoleVariableRef CVarPredictionAsyncFrameBuffer(TEXT("np2.PredictionAsyncFrameBuffer"), PredictionAsyncFrameBuffer, TEXT("Additional frame offset to be added to the local to server offset used by network prediction"));
 
 	int32 TickOffsetUpdateInterval = 10;
 	FAutoConsoleVariableRef CVarTickOffsetUpdateInterval(TEXT("np2.TickOffsetUpdateInterval"), TickOffsetUpdateInterval, TEXT("How many physics ticks to wait between each tick offset update. Lowest viable value = 1, which means update each tick."));
-
-	bool TimeDilationEnabled = true;
-	FAutoConsoleVariableRef CVarTimeDilationEnabled(TEXT("np2.TimeDilationEnabled"), TimeDilationEnabled, TEXT("Server-side CVar, the server calculates a time dilation multiplier and sends to the client which manipulates the accumulated time on client to speed up or slow down the physics simulation in order to correct the physics frame offset between client and server."));
 	
 	float TimeDilationAmount = 0.01f;
-	FAutoConsoleVariableRef CVarTimeDilationAmount(TEXT("np2.TimeDilationAmount"), TimeDilationAmount, TEXT("Server-side CVar, Default: 0.01 | Value is in percent where 0.01 = 1% dilation. Example: 1.0/0.01 = 100, meaning that over the time it usually takes to tick 100 physics steps we will tick 99 or 101 depending on if we dilate up or down."));
+	FAutoConsoleVariableRef CVarTimeDilationAmount(TEXT("np2.TimeDilationAmount"), TimeDilationAmount, TEXT("Server-side CVar, Disable TimeDilation by setting to 0 | Default: 0.01 | Value is in percent where 0.01 = 1% dilation. Example: 1.0/0.01 = 100, meaning that over the time it usually takes to tick 100 physics steps we will tick 99 or 101 depending on if we dilate up or down."));
 
 	bool TimeDilationEscalation = true;
 	FAutoConsoleVariableRef CVarTimeDilationEscalation(TEXT("np2.TimeDilationEscalation"), TimeDilationEscalation, TEXT("Server-side CVar, Dilate the time more depending on how many ticks we need to adjust. When set to false we use the set TimeDilationAmount and wait the amount of time it takes to perform correct the offset. When set to true we multiply the TimeDilationAmount with the buffer offset count which will correct the offset in one TimeDilationAmount cycle."));
@@ -163,6 +176,11 @@ struct FForceFeedbackEffectHistoryEntry
 
 APlayerController::APlayerController(const FObjectInitializer& ObjectInitializer)
 	: Super(ObjectInitializer)
+	PRAGMA_DISABLE_DEPRECATION_WARNINGS
+	, InputBuffer_DEPRECATED(FInputCmdBuffer())
+	, ClientFrameInfo_DEPRECATED(FClientFrameInfo())
+	, ServerFrameInfo_DEPRECATED(FServerFrameInfo())
+	PRAGMA_ENABLE_DEPRECATION_WARNINGS
 {
 	NetPriority = 3.0f;
 	CheatClass = UCheatManager::StaticClass();
@@ -1684,6 +1702,34 @@ void APlayerController::ClientSetCameraFade_Implementation(bool bEnableFading, F
 
 void APlayerController::SendClientAdjustment()
 {
+	PRAGMA_DISABLE_DEPRECATION_WARNINGS
+	if (ServerFrameInfo_DEPRECATED.LastProcessedInputFrame != INDEX_NONE && ServerFrameInfo_DEPRECATED.LastProcessedInputFrame != ServerFrameInfo_DEPRECATED.LastSentLocalFrame)
+	{
+		ServerFrameInfo_DEPRECATED.LastSentLocalFrame = ServerFrameInfo_DEPRECATED.LastProcessedInputFrame;
+		ClientRecvServerAckFrame(ServerFrameInfo_DEPRECATED.LastProcessedInputFrame, ServerFrameInfo_DEPRECATED.LastLocalFrame, ServerFrameInfo_DEPRECATED.QuantizedTimeDilation);
+		
+		if (NetworkPhysicsCvars::EnableDebugRPC)
+		{
+			ClientRecvServerAckFrameDebug(InputBuffer_DEPRECATED.HeadFrame() - ServerFrameInfo_DEPRECATED.LastProcessedInputFrame, ServerFrameInfo_DEPRECATED.TargetNumBufferedCmds);
+		}
+	}
+
+	if (UPhysicsSettings::Get()->PhysicsPrediction.bEnablePhysicsPrediction)
+	{
+		if (ServerLatestTimestampToCorrect_DEPRECATED.ServerFrame != INDEX_NONE)
+		{
+			ClientCorrectionAsyncPhysicsTimestamp(ServerLatestTimestampToCorrect_DEPRECATED);
+			ServerLatestTimestampToCorrect_DEPRECATED.ServerFrame = INDEX_NONE;
+		}
+
+		if (AcknowledgedPawn != GetPawn() && !GetSpectatorPawn())
+		{
+			return;
+		}
+	}
+	PRAGMA_ENABLE_DEPRECATION_WARNINGS
+
+
 	// Server sends updates.
 	// Note: we do this for both the pawn and spectator in case an implementation has a networked spectator.
 	APawn* RemotePawn = GetPawnOrSpectator();
@@ -1697,6 +1743,70 @@ void APlayerController::SendClientAdjustment()
 	}
 }
 
+void APlayerController::PushClientInput(int32 InRecvClientInputFrame, TArray<uint8>& Data)
+{
+	PRAGMA_DISABLE_DEPRECATION_WARNINGS
+	InputBuffer_DEPRECATED.Write(InRecvClientInputFrame) = MoveTemp(Data);
+	
+	// Do the RPC right here, including the redundant send. This should probably be time based and managed somewhere else like in Tick eventually
+	for (int32 Frame = FMath::Max(1, InRecvClientInputFrame - NetworkPhysicsCvars::NumRedundantCmds + 1); Frame <= InRecvClientInputFrame; ++Frame)
+	{
+		ServerRecvClientInputFrame(Frame, InputBuffer_DEPRECATED.Get(Frame));
+	}
+	PRAGMA_ENABLE_DEPRECATION_WARNINGS
+}
+
+void APlayerController::ServerRecvClientInputFrame_Implementation(int32 InRecvClientInputFrame, const TArray<uint8>& Data)
+{
+	PRAGMA_DISABLE_DEPRECATION_WARNINGS
+	if (InRecvClientInputFrame < 0)
+	{
+		return;
+	}
+	
+	int32 ClampDroppedFrames = 30; // Temporary clamp to guard the for-loop from locking the thread on large discrepancies, function is deprecated and will get removed in UE 5.6
+	for (int32 DroppedFrame = InputBuffer_DEPRECATED.HeadFrame() + 1; DroppedFrame < InRecvClientInputFrame && DroppedFrame > 0; ++DroppedFrame)
+	{
+		UE_LOG(LogPlayerController, Warning, TEXT("ClientInput Gap in frames (Dropped). %s [%s]. DroppedFrame: %d. RecvFrame: %d. LastProcessInputFrame: %d"), *GetName(), PlayerState ? *PlayerState->GetPlayerName() : TEXT("???"), DroppedFrame, InRecvClientInputFrame, ServerFrameInfo_DEPRECATED.LastProcessedInputFrame);
+		//ServerFrameInfo.LastProcessedInputFrame++; // Ehh lets try this for now
+		InputBuffer_DEPRECATED.Write(DroppedFrame) = InputBuffer_DEPRECATED.Get(DroppedFrame - 1);
+	
+		if (--ClampDroppedFrames <= 0)
+		{
+			UE_LOG(LogPlayerController, Warning, TEXT("ClientInput Gap in frames (Dropped) reached max dropped frames, leaving a gap in the InputBuffer."));
+			break;
+		}
+	}
+	
+	InputBuffer_DEPRECATED.Write(InRecvClientInputFrame) = MoveTemp(const_cast<TArray<uint8>&>(Data));
+	
+	if (ServerFrameInfo_DEPRECATED.LastProcessedInputFrame < InputBuffer_DEPRECATED.TailFrame())
+	{
+		// At this point, things are pretty bad and we are going to drop commands. 
+		// We still need guards client side to not send too many commands
+		UE_LOG(LogPlayerController, Warning, TEXT("ClientInput buffer overflow. %s [%s]. RecvFrame: %d. LastProcessInputFrame: %d."), *GetName(), PlayerState ? *PlayerState->GetPlayerName() : TEXT("???"), InRecvClientInputFrame, ServerFrameInfo_DEPRECATED.LastProcessedInputFrame);
+		ServerFrameInfo_DEPRECATED.LastProcessedInputFrame = (InputBuffer_DEPRECATED.TailFrame() + InputBuffer_DEPRECATED.HeadFrame()) / 2;
+	}
+	PRAGMA_ENABLE_DEPRECATION_WARNINGS
+}
+
+void APlayerController::ClientRecvServerAckFrame_Implementation(int32 LastProcessedInputFrame, int32 RecvServerFrameNumber, int8 TimeDilation)
+{
+	PRAGMA_DISABLE_DEPRECATION_WARNINGS
+	ClientFrameInfo_DEPRECATED.LastRecvServerFrame = RecvServerFrameNumber;
+	ClientFrameInfo_DEPRECATED.LastProcessedInputFrame = LastProcessedInputFrame;
+	ClientFrameInfo_DEPRECATED.QuantizedTimeDilation = TimeDilation;
+	PRAGMA_ENABLE_DEPRECATION_WARNINGS
+}
+
+void APlayerController::ClientRecvServerAckFrameDebug_Implementation(uint8 NumBuffered, float TargetNumBufferedCmds)
+{
+	PRAGMA_DISABLE_DEPRECATION_WARNINGS
+	ClientFrameInfo_DEPRECATED.LastRecvInputFrame = ClientFrameInfo_DEPRECATED.LastProcessedInputFrame + NumBuffered;
+	ClientFrameInfo_DEPRECATED.TargetNumBufferedCmds = TargetNumBufferedCmds;
+	PRAGMA_ENABLE_DEPRECATION_WARNINGS
+}
+
 /// @cond DOXYGEN_WARNINGS
 
 void APlayerController::ClientCapBandwidth_Implementation(int32 Cap)
@@ -6029,6 +6139,22 @@ FAsyncPhysicsTimestamp APlayerController::GetAsyncPhysicsTimestamp(float DeltaSe
 {
 	using namespace Chaos;
 
+	FAsyncPhysicsTimestamp Timestamp = GetPhysicsTimestamp(DeltaSeconds);
+	
+	// Handle deprecated flow using LocalToServerAsyncPhysicsTickOffset_DEPRECATED
+	if (Timestamp.IsValid() && IsLocalController())
+	{
+		Timestamp.ServerFrame -= NetworkPhysicsTickOffset;
+		Timestamp.ServerFrame += LocalToServerAsyncPhysicsTickOffset_DEPRECATED;
+	}
+
+	return Timestamp;
+}
+
+FAsyncPhysicsTimestamp APlayerController::GetPhysicsTimestamp(float DeltaSeconds)
+{
+	using namespace Chaos;
+
 	FAsyncPhysicsTimestamp Timestamp;
 
 	if(UWorld* World = GetWorld())
@@ -6049,7 +6175,7 @@ FAsyncPhysicsTimestamp APlayerController::GetAsyncPhysicsTimestamp(float DeltaSe
 				if (IsLocalController())
 				{
 					//If local controller we update server frame based on our estimate
-					Timestamp.ServerFrame = LocalPhysicsStep + LocalToServerAsyncPhysicsTickOffset;
+					Timestamp.ServerFrame = LocalPhysicsStep + NetworkPhysicsTickOffset;
 				}
 			}
 		}
@@ -6077,7 +6203,7 @@ void APlayerController::UpdateServerAsyncPhysicsTickOffset()
 		}
 	}
 
-	FAsyncPhysicsTimestamp Timestamp = GetAsyncPhysicsTimestamp();
+	FAsyncPhysicsTimestamp Timestamp = GetPhysicsTimestamp();
 	if(ClientLatestAsyncPhysicsStepSent + NetworkPhysicsCvars::TickOffsetUpdateInterval > Timestamp.LocalFrame)
 	{
 		//Only send a new timestamp if enough physics ticks have passed, based on CVar.
@@ -6087,7 +6213,7 @@ void APlayerController::UpdateServerAsyncPhysicsTickOffset()
 	}
 
 	ClientLatestAsyncPhysicsStepSent = Timestamp.LocalFrame;
-	Timestamp.ServerFrame = LocalToServerAsyncPhysicsTickOffsetAssigned ? Timestamp.ServerFrame : INDEX_NONE; // If offset is not yet assigned, set an invalid ServerFrame
+	Timestamp.ServerFrame = bNetworkPhysicsTickOffsetAssigned ? Timestamp.ServerFrame : INDEX_NONE; // If offset is not yet assigned, set an invalid ServerFrame
 	ServerSendLatestAsyncPhysicsTimestamp(Timestamp);
 }
 
@@ -6103,22 +6229,30 @@ void APlayerController::ServerSendLatestAsyncPhysicsTimestamp_Implementation(FAs
 	}
 
 	ServerLatestAsyncPhysicsStepReceived = Timestamp.LocalFrame;
-		
+
+	PRAGMA_DISABLE_DEPRECATION_WARNINGS
+	// Only cache the most up to date timestamp based on LocalFrame
+	if (Timestamp.LocalFrame > ServerPendingTimestamp_DEPRECATED.LocalFrame)
+	{
+		ServerPendingTimestamp_DEPRECATED = Timestamp;
+	}
+	PRAGMA_ENABLE_DEPRECATION_WARNINGS
+
 	// Get current server timestamp and add the frame buffer to the ServerFrame
-	FAsyncPhysicsTimestamp ActualTimestamp = GetAsyncPhysicsTimestamp();
+	FAsyncPhysicsTimestamp ActualTimestamp = GetPhysicsTimestamp();
 	ActualTimestamp.ServerFrame += NetworkPhysicsCvars::PredictionAsyncFrameBuffer;
 
 	// Mark offset as assigned when we get a valid predicted server frame.
 	const int32 PredictedServerFrame = Timestamp.ServerFrame;
-	LocalToServerAsyncPhysicsTickOffsetAssigned |= PredictedServerFrame != INDEX_NONE;
+	bNetworkPhysicsTickOffsetAssigned |= PredictedServerFrame != INDEX_NONE;
 
 	// Send update to client if offset is not assigned
 	// Note that we are sending the current ServerFrame along with the frame buffer added, to the client.
-	if (!LocalToServerAsyncPhysicsTickOffsetAssigned)
+	if (!bNetworkPhysicsTickOffsetAssigned)
 	{
 		Timestamp.ServerFrame = ActualTimestamp.ServerFrame;
-		LocalToServerAsyncPhysicsTickOffset = Timestamp.ServerFrame - Timestamp.LocalFrame;
-		ClientCorrectionAsyncPhysicsTimestamp(Timestamp); /* Reliable RPC */
+		NetworkPhysicsTickOffset = Timestamp.ServerFrame - Timestamp.LocalFrame;
+		ClientSetupNetworkPhysicsTimestamp(Timestamp); /* Reliable RPC */
 	}
 
 	/*
@@ -6126,7 +6260,7 @@ void APlayerController::ServerSendLatestAsyncPhysicsTimestamp_Implementation(FAs
 	* If buffer goes low, speed up the client by raising the time dilation multiplier (each deltaTime accounts for more of the accumulated time, filling the accumulated time faster to tick the next physics step)
 	* If buffer goes high, slow down the client by lowering the time dilation multiplier (each deltaTime accounts for less of the accumulated time, taking longer to fill the accumulated time
 	*/
-	if (NetworkPhysicsCvars::TimeDilationEnabled && LocalToServerAsyncPhysicsTickOffsetAssigned)
+	if (bNetworkPhysicsTickOffsetAssigned)
 	{
 		// Get the buffer offset amount that deviates from the target buffer (Note: the buffer is already added to ActualTimestamp.ServerFrame here and in the PredictedServerFrame received from the client)
 		// 0 means buffer is perfect, positive value means the buffer is too large, negative value means the buffer is too small
@@ -6148,14 +6282,51 @@ void APlayerController::ServerSendLatestAsyncPhysicsTimestamp_Implementation(FAs
 }
 
 void APlayerController::ClientCorrectionAsyncPhysicsTimestamp_Implementation(FAsyncPhysicsTimestamp Timestamp)
+{
+	PRAGMA_DISABLE_DEPRECATION_WARNINGS
+
+	//This tells the client that a timestamp it sent out was wrong (for example we ran locally on step 5 and expected server to run on step 10, but it actually ran on step 11).
+	//The error can only be later. That is, it can never be that we expect to run on server step 10 but actually ran on sever step 9
+	//Once a timestamp has been corrected, any earlier timestamps can be ignored (these can be out of order because of networking)
+	if (Timestamp.ServerFrame <= ClientLatestCorrectedOffsetServerStep_DEPRECATED)
+	{
+		//already corrected after this so do nothing
+		return;
+	}
+	ClientLatestCorrectedOffsetServerStep_DEPRECATED = Timestamp.ServerFrame;
+
+	if (UWorld* World = GetWorld())
+	{
+		if (FPhysScene* PhysScene = World->GetPhysicsScene())
+		{
+			if (PhysScene->GetSolver()->GetEvolution()->IsResimming())
+			{
+				return;
+			}
+		}
+	}
+
+	FAsyncPhysicsTimestamp CurrentTimestamp = GetAsyncPhysicsTimestamp();
+	// We need to avoid changing this offset as much as possible since it will invalidate histories and will trigger resim
+	// To deal with that we compute a safe margin based on a user cvar + half the RTT
+	// This margin will only be applied the first time we will compute the offset 
+	const int32 FrameOffset = LocalToServerAsyncPhysicsTickOffset_DEPRECATED == 0 ? (NetworkPhysicsCvars::NetworkPhysicsPredictionFrameOffset + (CurrentTimestamp.LocalFrame - Timestamp.LocalFrame) / 2) : 0;
+
+	const int32 NewOffset = FMath::Max(LocalToServerAsyncPhysicsTickOffset_DEPRECATED, Timestamp.ServerFrame - Timestamp.LocalFrame + FrameOffset); //The new offset as reported by the server
+	LocalToServerAsyncPhysicsTickOffset_DEPRECATED = NewOffset;
+
+	PRAGMA_ENABLE_DEPRECATION_WARNINGS
+}
+
+void APlayerController::ClientSetupNetworkPhysicsTimestamp_Implementation(FAsyncPhysicsTimestamp Timestamp)
 {
 	ensure(UPhysicsSettings::Get()->PhysicsPrediction.bEnablePhysicsPrediction);
 
 	// Assign async physics tick offset
-	if (!LocalToServerAsyncPhysicsTickOffsetAssigned)
+	if (!bNetworkPhysicsTickOffsetAssigned)
 	{
-		LocalToServerAsyncPhysicsTickOffsetAssigned = true;
-		LocalToServerAsyncPhysicsTickOffset = Timestamp.ServerFrame - Timestamp.LocalFrame;
+		bNetworkPhysicsTickOffsetAssigned = true;
+		NetworkPhysicsTickOffset = Timestamp.ServerFrame - Timestamp.LocalFrame;
 	}
 }
 
@@ -6173,6 +6344,21 @@ void APlayerController::ClientAckTimeDilation_Implementation(float TimeDilation,
 	}
 }
 
+void APlayerController::UpdateServerTimestampToCorrect()
+{
+	PRAGMA_DISABLE_DEPRECATION_WARNINGS
+	// If the client has not sent the correct estimate of which server frame its local frame would correspond to, send back which server frame that local frame actually corresponded to
+	const FAsyncPhysicsTimestamp ActualTimestamp = GetAsyncPhysicsTimestamp();
+
+	if (ServerPendingTimestamp_DEPRECATED.ServerFrame != INDEX_NONE && ServerPendingTimestamp_DEPRECATED.ServerFrame != ActualTimestamp.ServerFrame)
+	{
+		ServerLatestTimestampToCorrect_DEPRECATED.ServerFrame = ActualTimestamp.ServerFrame;
+		ServerLatestTimestampToCorrect_DEPRECATED.LocalFrame = ServerPendingTimestamp_DEPRECATED.LocalFrame;
+		ServerPendingTimestamp_DEPRECATED.ServerFrame = INDEX_NONE;
+	}
+	PRAGMA_ENABLE_DEPRECATION_WARNINGS
+}
+
 void APlayerController::AsyncPhysicsTickActor(float DeltaTime, float SimTime)
 {
 	Super::AsyncPhysicsTickActor(DeltaTime, SimTime);
diff --git a/Engine/Source/Runtime/Engine/Public/Physics/NetworkPhysicsComponent.h b/Engine/Source/Runtime/Engine/Public/Physics/NetworkPhysicsComponent.h
--- a/Engine/Source/Runtime/Engine/Public/Physics/NetworkPhysicsComponent.h
+++ b/Engine/Source/Runtime/Engine/Public/Physics/NetworkPhysicsComponent.h
@@ -268,6 +268,14 @@ struct FNetworkPhysicsCallback : public Chaos::IRewindCallback
 	// but also needs to be able to access GT data (actor iterator, actor state)
 	void UpdateReplicationMap_Internal(int32 PhysicsStep);
 
+	// Update client player on GT
+	UE_DEPRECATED(5.4, "Physics frame offset is handled by the PlayerController automatically, it's recommended to use APlayerController::GetAsyncPhysicsTimestamp() to get the ServerFrame and LocalFrame on both client and server. Also disable the deprecated flow by setting p.net.CmdOffsetEnabled = 0")
+	void UpdateClientPlayer_External(int32 PhysicsStep);
+
+	// Update server player on GT
+	UE_DEPRECATED(5.4, "Physics frame offset is handled by the PlayerController automatically, it's recommended to use APlayerController::GetAsyncPhysicsTimestamp() to get the ServerFrame and LocalFrame on both client and server. Also disable the deprecated flow by setting p.net.CmdOffsetEnabled = 0")
+	void UpdateServerPlayer_External(int32 PhysicsStep);
+
 	// Cache the current netmode for use in PT
 	void UpdateNetMode()
 	{
@@ -393,6 +401,10 @@ public:
 	// Delegate linked to the physics rewind callback to send record local inputs/states
 	ENGINE_API void OnPostProcessInputsInternal(const int32 PhysicsStep);
 
+	// Correct the player controller Server to local offset based on the received replicated states
+	UE_DEPRECATED(5.4, "Physics frame offset is handled by the PlayerController automatically, it's recommended to use APlayerController::GetAsyncPhysicsTimestamp() to get the ServerFrame and LocalFrame on both client and server. Also disable the deprecated flow by setting p.net.CmdOffsetEnabled = 0")
+	ENGINE_API void CorrectServerToLocalOffset(const int32 LocalToServerOffset);
+
 	// Used to create any physics engine information for this component 
 	ENGINE_API virtual void BeginPlay() override;
 
