UWorld::BlockTillStreamingCompleted() : Do an extra loop to handle potential FlushAsyncLoading() calls happening inside the WorkToDo loop

Reported multiple times through UDN

#rnx

diff --git a/Engine/Source/Runtime/Engine/Private/World.cpp b/Engine/Source/Runtime/Engine/Private/World.cpp
--- a/Engine/Source/Runtime/Engine/Private/World.cpp
+++ b/Engine/Source/Runtime/Engine/Private/World.cpp
@@ -4177,7 +4177,9 @@ void UWorld::BlockTillLevelStreamingCompleted()
 	}
 
 	bool bIsStreamingPaused = false;
-	bool bWorkToDo = false;
+
+	// In case a FlushAsyncLoading() happens inside loop we need to tick an extra loop.
+	int32 WorkToDo = 2; 
 	do
 	{
 		// Update world's required streaming levels
@@ -4185,20 +4187,22 @@ void UWorld::BlockTillLevelStreamingCompleted()
 
 		// Probe if we have anything to do
 		UpdateLevelStreaming();
-
-		bWorkToDo = (IsVisibilityRequestPending() || IsAsyncLoading());
-		if (bWorkToDo)
+		
+		// Everytime we have work to do, add an extra loop to handle FlushAsyncLoading calls
+		if (IsVisibilityRequestPending() || IsAsyncLoading())
 		{
-			if (!bIsStreamingPaused && GEngine->GameViewport && GEngine->BeginStreamingPauseDelegate && GEngine->BeginStreamingPauseDelegate->IsBound())
-			{
-				GEngine->BeginStreamingPauseDelegate->Execute(GEngine->GameViewport->Viewport);
-				bIsStreamingPaused = true;
-			}
+			WorkToDo = 2;
+		}
 
-			// Flush level streaming requests, blocking till completion.
-			FlushLevelStreaming(EFlushLevelStreamingType::Full);
+		if (!bIsStreamingPaused && GEngine->GameViewport && GEngine->BeginStreamingPauseDelegate && GEngine->BeginStreamingPauseDelegate->IsBound())
+		{
+			GEngine->BeginStreamingPauseDelegate->Execute(GEngine->GameViewport->Viewport);
+			bIsStreamingPaused = true;
 		}
-	} while (bWorkToDo);
+
+		// Flush level streaming requests, blocking till completion.
+		FlushLevelStreaming(EFlushLevelStreamingType::Full);
+	} while (--WorkToDo > 0);
 
 	if (bIsStreamingPaused && GEngine->EndStreamingPauseDelegate && GEngine->EndStreamingPauseDelegate->IsBound())
 	{
