Iris
* Components that are created from a User Construction Script and are not stably named will use the class CDO for their default state.
* This prevents protocol mismatches while we work on a proper fix and replicate the true archetype of these components so the clients can use as a template.

#rnx

diff --git a/Engine/Source/Runtime/Engine/Private/Components/ActorComponent.cpp b/Engine/Source/Runtime/Engine/Private/Components/ActorComponent.cpp
--- a/Engine/Source/Runtime/Engine/Private/Components/ActorComponent.cpp
+++ b/Engine/Source/Runtime/Engine/Private/Components/ActorComponent.cpp
@@ -43,6 +43,8 @@
 #include "ObjectTrace.h"
 
 #if UE_WITH_IRIS
+#include "Iris/Core/IrisLog.h"
+#include "Iris/ReplicationSystem/ReplicationFragment.h"
 #include "Iris/ReplicationSystem/ReplicationFragmentUtil.h"
 #include "Net/Iris/ReplicationSystem/ReplicationSystemUtil.h"
 #endif // UE_WITH_IRIS
@@ -2284,6 +2286,18 @@ bool UActorComponent::CanEditChange(const FProperty* InProperty) const
 #if UE_WITH_IRIS
 void UActorComponent::RegisterReplicationFragments(UE::Net::FFragmentRegistrationContext& Context, UE::Net::EFragmentRegistrationFlags RegistrationFlags)
 {
+	if (CreationMethod == EComponentCreationMethod::UserConstructionScript)
+	{
+		//$IRIS TODO: UCS subobjects without a stable name need to replicate their true archetype. Currently they are instantiated on the client from the class CDO.
+		//            Prevent protocol mismatches by forcing their default state to use the CDO. But the replicated object state is probably wrong if it changed any non-replicated properties.
+		if (!IsNameStableForNetworking() && GetArchetype() != GetClass()->GetDefaultObject())
+		{
+			RegistrationFlags |= UE::Net::EFragmentRegistrationFlags::InitializeDefaultStateFromClassDefaults;
+
+			UE_LOG(LogIris, Warning, TEXT("The default state of UCS created subobject %s::%s was initialized using the wrong archetype. The object's state on client might be initialied wrong."), *GetNameSafe(GetOwner()), *GetName());
+		}
+	}
+	
 	// Build descriptors and allocate PropertyReplicationFragments for this object
 	UE::Net::FReplicationFragmentUtil::CreateAndRegisterFragmentsForObject(this, Context, RegistrationFlags);
 }
