World Partition: wp.Runtime.ShowRuntimeSpatialHashCellStreamingPriority can now be used to display the original heatmap mode (1), or the new grayscale mode (2).

diff --git a/Engine/Source/Runtime/Engine/Private/WorldPartition/RuntimeHashSet/WorldPartitionRuntimeHashSetDebugDraw.cpp b/Engine/Source/Runtime/Engine/Private/WorldPartition/RuntimeHashSet/WorldPartitionRuntimeHashSetDebugDraw.cpp
--- a/Engine/Source/Runtime/Engine/Private/WorldPartition/RuntimeHashSet/WorldPartitionRuntimeHashSetDebugDraw.cpp
+++ b/Engine/Source/Runtime/Engine/Private/WorldPartition/RuntimeHashSet/WorldPartitionRuntimeHashSetDebugDraw.cpp
@@ -143,7 +143,7 @@ bool UWorldPartitionRuntimeHashSet::Draw2D(FWorldPartitionDraw2DContext& DrawCon
 				{
 				case 0:
 					CellColors.Add(Cell->GetDebugColor(VisualizeMode));
-					CellOpacity = 0.25f / FMath::Max<float>(GShowRuntimeHashSetDebugDisplayLevelCount, 1);
+					CellOpacity = (VisualizeMode == EWorldPartitionRuntimeCellVisualizeMode::StreamingPriority) ? 0.75f : 0.25f / FMath::Max<float>(GShowRuntimeHashSetDebugDisplayLevelCount, 1);
 					break;
 				case 1:
 					if (DataLayerDebugColors.Num() && Cell->GetDataLayers().Num())
diff --git a/Engine/Source/Runtime/Engine/Private/WorldPartition/WorldPartitionDebugHelper.cpp b/Engine/Source/Runtime/Engine/Private/WorldPartition/WorldPartitionDebugHelper.cpp
--- a/Engine/Source/Runtime/Engine/Private/WorldPartition/WorldPartitionDebugHelper.cpp
+++ b/Engine/Source/Runtime/Engine/Private/WorldPartition/WorldPartitionDebugHelper.cpp
@@ -197,8 +197,8 @@ void FWorldPartitionDebugHelper::DrawLegendItem(UCanvas* Canvas, const FString&
 	Pos.Y = TextPos.Y;
 }
 
-bool FWorldPartitionDebugHelper::bShowRuntimeSpatialHashCellStreamingPriority = false;
-FAutoConsoleVariableRef FWorldPartitionDebugHelper::ShowRuntimeSpatialHashCellStreamingPriorityCommand(
+int32 FWorldPartitionDebugHelper::ShowRuntimeSpatialHashCellStreamingPriorityMode = 0;
+FAutoConsoleVariableRef FWorldPartitionDebugHelper::ShowRuntimeSpatialHashCellStreamingPriorityModeCommand(
 	TEXT("wp.Runtime.ShowRuntimeSpatialHashCellStreamingPriority"),
-	FWorldPartitionDebugHelper::bShowRuntimeSpatialHashCellStreamingPriority,
-	TEXT("Enable to show a heatmap of the runtime spatial hash grid cells based on their priority."));
\ No newline at end of file
+	FWorldPartitionDebugHelper::ShowRuntimeSpatialHashCellStreamingPriorityMode,
+	TEXT("Enable to show a heatmap of the runtime spatial hash grid cells based on their priority (0=disabled, 1=heatmap, 2=grayscale."));
\ No newline at end of file
diff --git a/Engine/Source/Runtime/Engine/Private/WorldPartition/WorldPartitionRuntimeLevelStreamingCell.cpp b/Engine/Source/Runtime/Engine/Private/WorldPartition/WorldPartitionRuntimeLevelStreamingCell.cpp
--- a/Engine/Source/Runtime/Engine/Private/WorldPartition/WorldPartitionRuntimeLevelStreamingCell.cpp
+++ b/Engine/Source/Runtime/Engine/Private/WorldPartition/WorldPartitionRuntimeLevelStreamingCell.cpp
@@ -4,6 +4,7 @@
 #include "WorldPartition/WorldPartitionLevelStreamingDynamic.h"
 #include "WorldPartition/WorldPartitionActorDescView.h"
 #include "WorldPartition/WorldPartitionLevelStreamingPolicy.h"
+#include "WorldPartition/WorldPartitionDebugHelper.h"
 #include "Engine/LevelStreaming.h"
 #include "Engine/Level.h"
 #include "Misc/HierarchicalLogArchive.h"
@@ -209,7 +210,21 @@ FLinearColor UWorldPartitionRuntimeLevelStreamingCell::GetDebugColor(EWorldParti
 			if (DebugStreamingPriority >= 0.0f && DebugStreamingPriority <= 1.0f)
 			{
 				const float PriorityGradient = FMath::Cube(1.0f - DebugStreamingPriority);
-				return FLinearColor(PriorityGradient, PriorityGradient, PriorityGradient, 1.0f);
+
+				if (FWorldPartitionDebugHelper::GetRuntimeSpatialHashCellStreamingPriorityMode() == 2)
+				{
+					// Grayscale
+					return FLinearColor(PriorityGradient, PriorityGradient, PriorityGradient, 1.0f);
+				}
+
+				// Heatmap
+				static TArray<FLinearColor, TInlineAllocator<4>> Colors { FLinearColor::Blue, FLinearColor::Green, FLinearColor::Yellow, FLinearColor::Red };
+				const float ColorGrad = FMath::Clamp(PriorityGradient, 0.0f, 1.0f) * (Colors.Num() - 1);
+				return FLinearColor::LerpUsingHSV(
+					Colors[FMath::Min<int32>(ColorGrad, Colors.Num() - 1)],
+					Colors[FMath::Min<int32>(ColorGrad + 1, Colors.Num() - 1)],
+					FMath::Frac(ColorGrad)
+				);
 			}
 			return FLinearColor::Transparent;
 		}
diff --git a/Engine/Source/Runtime/Engine/Public/WorldPartition/WorldPartitionDebugHelper.h b/Engine/Source/Runtime/Engine/Public/WorldPartition/WorldPartitionDebugHelper.h
--- a/Engine/Source/Runtime/Engine/Public/WorldPartition/WorldPartitionDebugHelper.h
+++ b/Engine/Source/Runtime/Engine/Public/WorldPartition/WorldPartitionDebugHelper.h
@@ -22,7 +22,8 @@ public:
 	static bool IsDebugCellNameShown(const FString& Name);
 	static void DrawText(UCanvas* Canvas, const FString& Text, const UFont* Font, const FColor& Color, FVector2D& Pos, float* MaxTextWidth = nullptr);
 	static void DrawLegendItem(UCanvas* Canvas, const FString& Text, const UFont* Font, const FColor& Color, const FColor& TextColor, FVector2D& Pos, float* MaxItemWidth = nullptr);
-	FORCEINLINE static bool IsRuntimeSpatialHashCellStreamingPriorityShown() { return bShowRuntimeSpatialHashCellStreamingPriority; }
+	FORCEINLINE static bool IsRuntimeSpatialHashCellStreamingPriorityShown() { return ShowRuntimeSpatialHashCellStreamingPriorityMode != 0; }
+	FORCEINLINE static int32 GetRuntimeSpatialHashCellStreamingPriorityMode() { return ShowRuntimeSpatialHashCellStreamingPriorityMode; }
 	FORCEINLINE static bool CanDrawContentBundles() { return bCanDrawContentBundles; }
 
 private:
@@ -38,8 +39,8 @@ private:
 	static FAutoConsoleCommand DebugFilterByCellNameCommand;
 	static TArray<FString> DebugCellNameFilter;
 
-	static FAutoConsoleVariableRef ShowRuntimeSpatialHashCellStreamingPriorityCommand;
-	static bool bShowRuntimeSpatialHashCellStreamingPriority;
+	static FAutoConsoleVariableRef ShowRuntimeSpatialHashCellStreamingPriorityModeCommand;
+	static int32 ShowRuntimeSpatialHashCellStreamingPriorityMode;
 
 	static FAutoConsoleVariableRef DrawContentBundlesCommand;
 	static bool bCanDrawContentBundles;
