Disable Substrate for shader pipelines that include geometry shaders for D3D11.
Cross-compiling HLSL2021 to FXC-compatible HLSL is not supported for geometry shaders.

#rnx

diff --git a/Engine/Source/Runtime/Engine/Private/ShaderCompiler/ShaderCompiler.cpp b/Engine/Source/Runtime/Engine/Private/ShaderCompiler/ShaderCompiler.cpp
--- a/Engine/Source/Runtime/Engine/Private/ShaderCompiler/ShaderCompiler.cpp
+++ b/Engine/Source/Runtime/Engine/Private/ShaderCompiler/ShaderCompiler.cpp
@@ -7382,6 +7382,17 @@ namespace
 	}
 }
 
+static bool IsSubstrateSupportForShaderPipeline(const FShaderCompilerInput& Input)
+{
+	// Substrate requires HLSL2021 which must be cross-compiled for D3D11 to be consumed by FXC compiler.
+	// This cross-compilation toolchain does not support geometry shaders.
+	bool bPipelineContainsGeometryShader = false;
+	Input.Environment.GetCompileArgument(TEXT("PIPELINE_CONTAINS_GEOMETRYSHADER"), bPipelineContainsGeometryShader);
+	const bool bCanRHICompileHlsl2021GeometryShaders = !(GetMaxSupportedFeatureLevel((EShaderPlatform)Input.Target.Platform) == ERHIFeatureLevel::SM5);
+	const bool bIsSubstrateSupportedForPipeline = !bPipelineContainsGeometryShader || bCanRHICompileHlsl2021GeometryShaders;
+	return bIsSubstrateSupportedForPipeline;
+}
+
 void GlobalBeginCompileShader(
 	const FString& DebugGroupName,
 	const FVertexFactoryType* VFType,
@@ -8079,7 +8090,7 @@ void GlobalBeginCompileShader(
 		SET_SHADER_DEFINE(Input.Environment, SUPPORT_VSM_FOWARD_QUALITY, bHighQualityShadow ? 1 : 0);
 	}
 
-	const bool bSubstrate = Substrate::IsSubstrateEnabled();
+	const bool bSubstrate = Substrate::IsSubstrateEnabled() && IsSubstrateSupportForShaderPipeline(Input);
 	{
 		SET_SHADER_DEFINE(Input.Environment, SUBSTRATE_ENABLED, bSubstrate ? 1 : 0);
 
diff --git a/Engine/Source/Runtime/Renderer/Private/VolumetricFogVoxelization.cpp b/Engine/Source/Runtime/Renderer/Private/VolumetricFogVoxelization.cpp
--- a/Engine/Source/Runtime/Renderer/Private/VolumetricFogVoxelization.cpp
+++ b/Engine/Source/Runtime/Renderer/Private/VolumetricFogVoxelization.cpp
@@ -278,6 +278,7 @@ public:
 	{
 		FMeshMaterialShader::ModifyCompilationEnvironment(Parameters, OutEnvironment);
 		OutEnvironment.SetDefine(TEXT("MAX_SLICES_PER_VOXELIZATION_PASS"), GetVoxelizationSlicesPerPass(Parameters.Platform));
+		OutEnvironment.SetCompileArgument(TEXT("PIPELINE_CONTAINS_GEOMETRYSHADER"), true);
 	}
 
 	void GetShaderBindings(
