Changed inertialization node checks to ensure.

diff --git a/Engine/Source/Runtime/Engine/Private/Animation/AnimNode_Inertialization.cpp b/Engine/Source/Runtime/Engine/Private/Animation/AnimNode_Inertialization.cpp
--- a/Engine/Source/Runtime/Engine/Private/Animation/AnimNode_Inertialization.cpp
+++ b/Engine/Source/Runtime/Engine/Private/Animation/AnimNode_Inertialization.cpp
@@ -905,10 +905,15 @@ float FInertializationPoseDiff::CalcInertialFloat(float x0, float v0, float t, f
 		v0 = 0.0f;
 	}
 
-	check(x0 >= 0.0f);
-	check(v0 <= 0.0f);
-	check(t >= 0.0f);
-	check(t1 >= 0.0f);
+	// Check for invalid values - this is only expected to occur if NaNs or other invalid values are coming into the node
+	if (!ensureMsgf(x0 >= 0.0f && v0 <= 0.0f && t >= 0.0f && t1 >= 0.0f,
+		TEXT("Invalid Value(s) in Inertialization - x0: %f, v0: %f, t: %f, t1: %f"), x0, v0, t, t1))
+	{
+		x0 = 0.0f;
+		v0 = 0.0f;
+		t = 0.0f;
+		t1 = 0.0f;
+	}
 
 	// Limit t1 such that the curve does not overshoot below zero (ensuring that x >= 0 for all t between 0 and t1).
 	//
