Use 16k chunks for mesh batches to reduce allocations.

diff --git a/Engine/Source/Runtime/Engine/Public/SceneManagement.h b/Engine/Source/Runtime/Engine/Public/SceneManagement.h
--- a/Engine/Source/Runtime/Engine/Public/SceneManagement.h
+++ b/Engine/Source/Runtime/Engine/Public/SceneManagement.h
@@ -2027,7 +2027,7 @@ protected:
 	 * Using TChunkedArray which will never realloc as new elements are added
 	 * @todo - use mem stack
 	 */
-	TChunkedArray<FMeshBatch, 1024, FConcurrentLinearArrayAllocator> MeshBatchStorage;
+	TChunkedArray<FMeshBatch, 16384, FConcurrentLinearArrayAllocator> MeshBatchStorage;
 
 	/** Meshes to render */
 	TArray<TArray<FMeshBatchAndRelevance, SceneRenderingAllocator>*, TInlineAllocator<2, SceneRenderingAllocator> > MeshBatches;
