Further refinement of scene component validation - if actor root component would be stripped, entire actor is stripped by AActor::NeedsLoadForTargetPlatform

diff --git a/Engine/Source/Runtime/Engine/Private/Components/SceneComponent.cpp b/Engine/Source/Runtime/Engine/Private/Components/SceneComponent.cpp
--- a/Engine/Source/Runtime/Engine/Private/Components/SceneComponent.cpp
+++ b/Engine/Source/Runtime/Engine/Private/Components/SceneComponent.cpp
@@ -567,6 +567,11 @@ EDataValidationResult USceneComponent::IsDataValid(FDataValidationContext& Conte
 	if (   DetailMode != EDetailMode::DM_Low 
 		&& (GValidateSceneComponentAttachmentDetailLevel_Low || GValidateSceneComponentAttachmentDetailLevel_Medium || GValidateSceneComponentAttachmentDetailLevel_High))
 	{
+		const AActor* const Owner = GetOwner();
+		const USceneComponent* const RootComponent = Owner ? Owner->GetRootComponent() : nullptr;
+		const bool bActorStrippedAtLow = RootComponent && RootComponent->DetailMode > EDetailMode::DM_Low;
+		const bool bActorStrippedAtMedium = RootComponent && RootComponent->DetailMode > EDetailMode::DM_Medium;
+		const bool bActorStrippedAtHigh = RootComponent && RootComponent->DetailMode > EDetailMode::DM_High;
 		for (USceneComponent* ChildSceneComponent : GetAttachChildren())
 		{
 			if (!ChildSceneComponent || ChildSceneComponent->IsEditorOnly())
@@ -574,9 +579,9 @@ EDataValidationResult USceneComponent::IsDataValid(FDataValidationContext& Conte
 				continue;
 			}
 			
-			bool bBrokenAtLow = DetailMode > EDetailMode::DM_Low && ChildSceneComponent->DetailMode <= EDetailMode::DM_Low;
-			bool bBrokenAtMedium = DetailMode > EDetailMode::DM_Medium && ChildSceneComponent->DetailMode <= EDetailMode::DM_Medium;
-			bool bBrokenAtHigh = DetailMode > EDetailMode::DM_High && ChildSceneComponent->DetailMode <= EDetailMode::DM_High;
+			const bool bBrokenAtLow = !bActorStrippedAtLow && DetailMode > EDetailMode::DM_Low && ChildSceneComponent->DetailMode <= EDetailMode::DM_Low;
+			const bool bBrokenAtMedium = !bActorStrippedAtMedium && DetailMode > EDetailMode::DM_Medium && ChildSceneComponent->DetailMode <= EDetailMode::DM_Medium;
+			const bool bBrokenAtHigh = !bActorStrippedAtHigh && DetailMode > EDetailMode::DM_High && ChildSceneComponent->DetailMode <= EDetailMode::DM_High;
 			if((GValidateSceneComponentAttachmentDetailLevel_Low && bBrokenAtLow)
 			|| (GValidateSceneComponentAttachmentDetailLevel_Medium && bBrokenAtMedium)
 			|| (GValidateSceneComponentAttachmentDetailLevel_High && bBrokenAtHigh))
