Make sure we respect r.ShaderCompiler.DebugDumpWorkerInputs when compiling locally.

diff --git a/Engine/Source/Runtime/Engine/Private/ShaderCompiler/ShaderCompiler.cpp b/Engine/Source/Runtime/Engine/Private/ShaderCompiler/ShaderCompiler.cpp
--- a/Engine/Source/Runtime/Engine/Private/ShaderCompiler/ShaderCompiler.cpp
+++ b/Engine/Source/Runtime/Engine/Private/ShaderCompiler/ShaderCompiler.cpp
@@ -2740,8 +2740,7 @@ FString CreateShaderCompilerWorkerDebugCommandLine(FString DebugWorkerInputFileP
 		DebugWorkerOutputFileName);
 }
 
-// Serialize Queued Job information
-bool FShaderCompileUtilities::DoWriteTasks(const TArray<FShaderCommonCompileJobPtr>& QueuedJobs, FArchive& InTransferFile, IDistributedBuildController* BuildDistributionController, bool bUseRelativePaths, bool bCompressTaskFile)
+static void DumpWorkerInputs(TConstArrayView<FShaderCommonCompileJobPtr> QueuedJobs)
 {
 	if (CVarDebugDumpWorkerInputs.GetValueOnAnyThread())
 	{
@@ -2780,6 +2779,12 @@ bool FShaderCompileUtilities::DoWriteTasks(const TArray<FShaderCommonCompileJobP
 			}
 		}
 	}
+}
+
+// Serialize Queued Job information
+bool FShaderCompileUtilities::DoWriteTasks(const TArray<FShaderCommonCompileJobPtr>& QueuedJobs, FArchive& InTransferFile, IDistributedBuildController* BuildDistributionController, bool bUseRelativePaths, bool bCompressTaskFile)
+{
+	DumpWorkerInputs(QueuedJobs);
 
 	return DoWriteTasksInner(QueuedJobs, InTransferFile, BuildDistributionController, bUseRelativePaths, bCompressTaskFile);
 }
@@ -3985,6 +3990,8 @@ void FShaderCompileThreadRunnable::CompileDirectlyThroughDll()
 
 		if (CurrentWorkerInfo.QueuedJobs.Num() > 0)
 		{
+			DumpWorkerInputs(CurrentWorkerInfo.QueuedJobs);
+
 			for (int32 JobIndex = 0; JobIndex < CurrentWorkerInfo.QueuedJobs.Num(); JobIndex++)
 			{
 				FShaderCommonCompileJob& CurrentJob = *CurrentWorkerInfo.QueuedJobs[JobIndex];
