Changed CluesterUnionComponent to store the cached bone ids in a TArray instead of a TSet as after previous changes we weren't using it for search.

Also modified the native broadcast events to pass the bone ids on a TArray (This means now both broadcast events passes the bone ids as TArray).

Lastly, made changes in listeners of the OnComponentAdded event to use the bIsNew flag so operations are only done if the component was newly added

diff --git a/Engine/Source/Runtime/Engine/Classes/PhysicsEngine/ClusterUnionComponent.h b/Engine/Source/Runtime/Engine/Classes/PhysicsEngine/ClusterUnionComponent.h
--- a/Engine/Source/Runtime/Engine/Classes/PhysicsEngine/ClusterUnionComponent.h
+++ b/Engine/Source/Runtime/Engine/Classes/PhysicsEngine/ClusterUnionComponent.h
@@ -69,9 +69,28 @@ USTRUCT()
 struct FClusteredComponentData
 {
 	GENERATED_BODY()
+	
+	FClusteredComponentData() = default;
 
+#if WITH_EDITORONLY_DATA
+	UE_DEPRECATED(5.4, "This property is deprecated. Please use BonesData instead.")
 	// Set of bone Ids that we actually added into the cluster union.
-	TSet<FClusterUnionBoneData> BonesData;
+	TSet<int32> BoneIds;
+
+	UE_DEPRECATED(5.4, "This property is deprecated and no longer used. It will be removed in future versions.")
+	// Cached acceleration structure handles - needed to properly cleanup the component from the accel structure.
+	TSet<FExternalSpatialAccelerationPayload> CachedAccelerationPayloads;
+
+PRAGMA_DISABLE_DEPRECATION_WARNINGS
+	FClusteredComponentData(const FClusteredComponentData& Other) = default;
+	FClusteredComponentData(FClusteredComponentData&& Other) = default;
+	FClusteredComponentData& operator=(const FClusteredComponentData& Other) = default;
+	FClusteredComponentData& operator=(FClusteredComponentData&& Other) = default;
+PRAGMA_ENABLE_DEPRECATION_WARNINGS
+#endif // WITH_EDITORONLY_DATA
+
+	// Array of bone Ids that we actually added into the cluster union.
+	TArray<FClusterUnionBoneData> BonesData;
 
 	// Using a TWeakObjectPtr here because the UClusterUnionReplicatedProxyComponent will have a pointer back
 	// and we don't want to get into a situation where a circular reference occurs.
@@ -116,6 +135,24 @@ struct FClusterUnionPendingAddData
 {
 	GENERATED_BODY()
 
+	FClusterUnionPendingAddData() = default;
+	
+#if WITH_EDITORONLY_DATA
+	UE_DEPRECATED(5.4, "This property is deprecated. Please use BonesData instead.")
+	TArray<int32> BoneIds;
+
+	UE_DEPRECATED(5.4, "This property is deprecated and no longer used. It will be removed in future versions.")
+	UPROPERTY()
+	TArray<FExternalSpatialAccelerationPayload> AccelerationPayloads;
+
+PRAGMA_DISABLE_DEPRECATION_WARNINGS
+	FClusterUnionPendingAddData(const FClusterUnionPendingAddData& Other) = default;
+	FClusterUnionPendingAddData(FClusterUnionPendingAddData&& Other) = default;
+	FClusterUnionPendingAddData& operator=(const FClusterUnionPendingAddData& Other) = default;
+	FClusterUnionPendingAddData& operator=(FClusterUnionPendingAddData&& Other) = default;
+PRAGMA_ENABLE_DEPRECATION_WARNINGS
+#endif // WITH_EDITORONLY_DATA
+
 	UPROPERTY()
 	TSet<FClusterUnionBoneData> BonesData;
 };
@@ -133,10 +170,10 @@ struct FClusterUnionInitializationData
 };
 
 DECLARE_DYNAMIC_MULTICAST_DELEGATE_ThreeParams(FOnClusterUnionAddedComponent, UPrimitiveComponent*, Component, const TSet<int32>&, BoneIds, bool, bIsNew);
-DECLARE_MULTICAST_DELEGATE_FourParams(FOnClusterUnionAddedComponentNative, UPrimitiveComponent*, const TSet<FClusterUnionBoneData>& /*BoneIds*/, const TArray<FClusterUnionBoneData>& /*RemovedBoneIds*/, bool /*bIsNew*/);
+DECLARE_MULTICAST_DELEGATE_FourParams(FOnClusterUnionAddedComponentNative, UPrimitiveComponent*, const TArray<FClusterUnionBoneData>& /*BoneIds*/, const TArray<FClusterUnionBoneData>& /*RemovedBoneIds*/, bool /*bIsNew*/);
 
 DECLARE_DYNAMIC_MULTICAST_DELEGATE_OneParam(FOnClusterUnionRemovedComponent, UPrimitiveComponent*, Component);
-DECLARE_MULTICAST_DELEGATE_TwoParams(FOnClusterUnionRemovedComponentNative, UPrimitiveComponent*, const TSet<FClusterUnionBoneData>& /*RemovedBonesData*/);
+DECLARE_MULTICAST_DELEGATE_TwoParams(FOnClusterUnionRemovedComponentNative, UPrimitiveComponent*, const TArray<FClusterUnionBoneData>& /*RemovedBonesData*/);
 
 DECLARE_DYNAMIC_MULTICAST_DELEGATE_TwoParams(FOnClusterUnionBoundsChanged, UClusterUnionComponent*, Component, const FBoxSphereBounds&, Bounds);
 
@@ -304,8 +341,8 @@ protected:
 	UFUNCTION()
 	ENGINE_API void ForceSetChildToParent(UPrimitiveComponent* InComponent, const TArray<int32>& BoneIds, const TArray<FTransform>& ChildToParent);
 
-	ENGINE_API void BroadcastComponentAddedEvents(UPrimitiveComponent* ChangedComponent, const TSet<FClusterUnionBoneData>& BoneIds, bool bIsNew, const TArray<FClusterUnionBoneData>& RemovedBoneIDs);
-	ENGINE_API void BroadcastComponentRemovedEvents(UPrimitiveComponent* ChangedComponent, const TSet<FClusterUnionBoneData>& InRemovedBonesData);
+	ENGINE_API void BroadcastComponentAddedEvents(UPrimitiveComponent* ChangedComponent, const TArray<FClusterUnionBoneData>& BoneIds, bool bIsNew, const TArray<FClusterUnionBoneData>& RemovedBoneIDs);
+	ENGINE_API void BroadcastComponentRemovedEvents(UPrimitiveComponent* ChangedComponent, const TArray<FClusterUnionBoneData>& InRemovedBonesData);
 
 	Chaos::FClusterUnionPhysicsProxy* GetPhysicsProxy() const { return PhysicsProxy; }
 	Chaos::FClusterUnionPhysicsProxy* GetPhysicsProxy() { return PhysicsProxy; }
diff --git a/Engine/Source/Runtime/Engine/Private/PhysicsEngine/ClusterUnionComponent.cpp b/Engine/Source/Runtime/Engine/Private/PhysicsEngine/ClusterUnionComponent.cpp
--- a/Engine/Source/Runtime/Engine/Private/PhysicsEngine/ClusterUnionComponent.cpp
+++ b/Engine/Source/Runtime/Engine/Private/PhysicsEngine/ClusterUnionComponent.cpp
@@ -373,7 +373,7 @@ void UClusterUnionComponent::RemoveComponentBonesFromCluster(UPrimitiveComponent
 	TSet<Chaos::FPhysicsObjectHandle> PhysicsObjectsToRemove;
 	if (FClusteredComponentData* ComponentData = PerComponentData.Find(ComponentKey))
 	{
-		for (TSet<FClusterUnionBoneData>::TIterator RemoveBoneDataIterator = ComponentData->BonesData.CreateIterator(); RemoveBoneDataIterator; ++RemoveBoneDataIterator)
+		for (TArray<FClusterUnionBoneData>::TIterator RemoveBoneDataIterator = ComponentData->BonesData.CreateIterator(); RemoveBoneDataIterator; ++RemoveBoneDataIterator)
 		{
 			if (AccelerationStructure)
 			{
@@ -876,7 +876,12 @@ void UClusterUnionComponent::SyncClusterUnionFromProxy()
 	const Chaos::FClusterUnionSyncedData& FullData = PhysicsProxy->GetSyncedData_External();
 
 	TArray<Chaos::FPBDRigidParticle*> ChildParticles;
-	ChildParticles.Reserve(FullData.ChildParticles.Num());
+
+	if (FullData.bDidSyncGeometry)
+	{
+		ChildParticles.Reserve(FullData.ChildParticles.Num());
+	}
+	
 	PerShapeComponentBone.Reset(FullData.ChildParticles.Num());
 
 	// Note that at the UClusterUnionComponent level we really only want to be dealing with components.
@@ -897,7 +902,11 @@ void UClusterUnionComponent::SyncClusterUnionFromProxy()
 				Chaos::FPhysicsObjectHandle Handle = Component->GetPhysicsObjectById(ChildData.BoneId);
 				
 				Chaos::FPBDRigidParticle* Particle = Interface->GetRigidParticle(Handle);
-				ChildParticles.Add(Particle);
+
+				if (FullData.bDidSyncGeometry)
+				{
+					ChildParticles.Add(Particle);
+				}
 
 				FMappedBoneData MappedBoneData(Handle, Particle, Particle ? Particle->UniqueIdx() : Chaos::FUniqueIdx(), ChildData.ChildToParent);
 				BoneIDToBoneData.Add(ChildData.BoneId, MappedBoneData);
@@ -906,7 +915,11 @@ void UClusterUnionComponent::SyncClusterUnionFromProxy()
 			}
 			else
 			{
-				ChildParticles.Add(nullptr);
+				if (FullData.bDidSyncGeometry)
+				{
+					ChildParticles.Add(nullptr);
+				}
+			
 				PerShapeComponentBone.Add({});
 			}
 		}
@@ -1023,31 +1036,28 @@ void UClusterUnionComponent::HandleAddOrModifiedClusteredComponent(const FMapped
 				ComponentData.ReplicatedProxyComponent = ReplicatedProxy;
 			}
 		}
-
 	}
 
 	TArray<FClusterUnionBoneData> RemovedBoneIDs;
 	RemovedBoneIDs.Reserve(ComponentData.BonesData.Num());
 
-	// In the case where we need to keep the acceleration structure up to date, we need to make sure old bone ids are properly
-	// removed from the acceleration structure.
-	if (AccelerationStructure)
+	for (const FClusterUnionBoneData& BoneData : ComponentData.BonesData)
 	{
-		for (const FClusterUnionBoneData& BoneData : ComponentData.BonesData)
+		if (!PerBoneChildToParent.Contains(BoneData.ID))
 		{
-			if (!PerBoneChildToParent.Contains(BoneData.ID))
+			// In the case where we need to keep the acceleration structure up to date, we need to make sure old bone ids are properly
+			// removed from the acceleration structure.
+			if (AccelerationStructure)
 			{
-				if (Chaos::FPhysicsObjectHandle PhysicsObject = ChangedComponentData.ComponentPtr->GetPhysicsObjectById(BoneData.ID))
+				FExternalSpatialAccelerationPayload Handle;
+				Handle.Initialize(ChangedComponentData.ComponentKey, BoneData.ID, BoneData.ParticleID);
+				if (ensure(Handle.IsValid()))
 				{
-					FExternalSpatialAccelerationPayload Handle;
-					Handle.Initialize(ChangedComponentData.ComponentKey, BoneData.ID, BoneData.ParticleID);
-					if (ensure(Handle.IsValid()))
-					{
-						AccelerationStructure->RemoveElement(Handle);
-					}
+					AccelerationStructure->RemoveElement(Handle);
 				}
-				RemovedBoneIDs.Emplace(FClusterUnionBoneData(BoneData.ID, BoneData.ParticleID));
 			}
+
+			RemovedBoneIDs.Emplace(FClusterUnionBoneData(BoneData.ID, BoneData.ParticleID));
 		}
 	}
 
@@ -1155,7 +1165,7 @@ void UClusterUnionComponent::HandleRemovedClusteredComponent(TObjectKey<UPrimiti
 	PendingComponentSync.Remove(RemovedComponent);
 }
 
-void UClusterUnionComponent::BroadcastComponentAddedEvents(UPrimitiveComponent* ChangedComponent, const TSet<FClusterUnionBoneData>& BoneIds, bool bIsNew, const TArray<FClusterUnionBoneData>& RemovedBoneIDs)
+void UClusterUnionComponent::BroadcastComponentAddedEvents(UPrimitiveComponent* ChangedComponent, const TArray<FClusterUnionBoneData>& BoneIds, bool bIsNew, const TArray<FClusterUnionBoneData>& RemovedBoneIDs)
 {
 	// TODO: Add a new delegate that takes a TArray and deprecate this one
 	if (OnComponentAddedEvent.IsBound())
@@ -1170,7 +1180,7 @@ void UClusterUnionComponent::BroadcastComponentAddedEvents(UPrimitiveComponent*
 	}
 }
 
-void UClusterUnionComponent::BroadcastComponentRemovedEvents(UPrimitiveComponent* ChangedComponent, const TSet<FClusterUnionBoneData>& InRemovedBonesData)
+void UClusterUnionComponent::BroadcastComponentRemovedEvents(UPrimitiveComponent* ChangedComponent, const TArray<FClusterUnionBoneData>& InRemovedBonesData)
 {
 	if (OnComponentRemovedEvent.IsBound())
 	{
