Add Projection settings accessors to PhysicsConstraintComponent

diff --git a/Engine/Source/Runtime/Engine/Classes/PhysicsEngine/PhysicsConstraintComponent.h b/Engine/Source/Runtime/Engine/Classes/PhysicsEngine/PhysicsConstraintComponent.h
--- a/Engine/Source/Runtime/Engine/Classes/PhysicsEngine/PhysicsConstraintComponent.h
+++ b/Engine/Source/Runtime/Engine/Classes/PhysicsEngine/PhysicsConstraintComponent.h
@@ -342,7 +342,28 @@ public:
 	// If true, the collision between the two rigid bodies of the constraint will be disabled.
 	UFUNCTION(BlueprintCallable, Category = "Physics|Components|PhysicsConstraint")
 	ENGINE_API void SetDisableCollision(bool bDisableCollision);
-	
+
+	// If true, joint projection is enabled. Projection is a semi-physics post-solve correction for fixing small errors, and a teleport for fixing larger errors. See SetProjectionParams
+	UFUNCTION(BlueprintCallable, Category = "Physics|Components|PhysicsConstraint")
+	ENGINE_API void SetProjectionEnabled(bool bProjectionEnabled);
+
+	// Is projection enabled. See SetProjectionEnabled
+	UFUNCTION(BlueprintCallable, Category = "Physics|Components|PhysicsConstraint")
+	bool IsProjectionEnabled() const;
+
+	/**
+	* Set the projection settings for use when projection is enabled. See SetProjectionEnabled.
+	* For ragdolls you usually require that ProjectionLinearAlpha and ProjectionAngularAlpha be zero. They are most useful for cosmetic chains etc. ProjectionLinearTolerance and ProjectionAngularTolerance
+	* are emergency error recovery settings and should only rarely affect the simulation. If the simulation is significantly affected by the Tolerance settings, it likely indicates a setup stability issue.
+	* 
+	* @param ProjectionLinearAlpha Controls the semi-physical correction of linear error remaining after the joint solve. Will add energy into the system. Best for joint chains connected to a kinematic.
+	* @param ProjectionAngularAlpha Controls the semi-physical correction of angular error remaining after the joint solve. Will add energy into the system. Best for joint chains connected to a kinematic.
+	* @param ProjectionLinearTolerance Linear errors above this will be corrected with a non-physical teleport.
+	* @param ProjectionAngularTolerance Angular errors above this will be corrected with a non-physical teleport.
+	*/
+	UFUNCTION(BlueprintCallable, Category = "Physics|Components|PhysicsConstraint")
+	ENGINE_API void SetProjectionParams(float ProjectionLinearAlpha, float ProjectionAngularAlpha, float ProjectionLinearTolerance, float ProjectionAngularTolerance);
+
 	// Retrieve the constraint force most recently applied to maintain this constraint. Returns 0 forces if the constraint is not initialized or broken.
 	UFUNCTION(BlueprintCallable, Category = "Physics|Components|PhysicsConstraint")
 	ENGINE_API void GetConstraintForce(FVector& OutLinearForce, FVector& OutAngularForce);
diff --git a/Engine/Source/Runtime/Engine/Private/PhysicsEngine/PhysicsConstraintComponent.cpp b/Engine/Source/Runtime/Engine/Private/PhysicsEngine/PhysicsConstraintComponent.cpp
--- a/Engine/Source/Runtime/Engine/Private/PhysicsEngine/PhysicsConstraintComponent.cpp
+++ b/Engine/Source/Runtime/Engine/Private/PhysicsEngine/PhysicsConstraintComponent.cpp
@@ -596,6 +596,28 @@ void UPhysicsConstraintComponent::SetDisableCollision(bool bDisableCollision)
 	ConstraintInstance.SetDisableCollision(bDisableCollision);
 }
 
+bool UPhysicsConstraintComponent::IsProjectionEnabled() const
+{
+	return ConstraintInstance.IsProjectionEnabled();
+}
+
+void UPhysicsConstraintComponent::SetProjectionEnabled(bool bInEnabled)
+{
+	if (bInEnabled)
+	{
+		ConstraintInstance.EnableProjection();
+	}
+	else
+	{
+		ConstraintInstance.DisableProjection();
+	}
+}
+
+void UPhysicsConstraintComponent::SetProjectionParams(float ProjectionLinearAlpha, float ProjectionAngularAlpha, float ProjectionLinearTolerance, float ProjectionAngularTolerance)
+{
+	ConstraintInstance.SetProjectionParams(IsProjectionEnabled(), ProjectionLinearAlpha, ProjectionAngularAlpha, ProjectionLinearTolerance, ProjectionAngularTolerance);
+}
+
 FConstraintInstanceAccessor UPhysicsConstraintComponent::GetConstraint()
 {
 	return FConstraintInstanceAccessor(this);
