Oodle Wave: moved engine dep to editor only!

diff --git a/Engine/Source/Runtime/Engine/Engine.Build.cs b/Engine/Source/Runtime/Engine/Engine.Build.cs
--- a/Engine/Source/Runtime/Engine/Engine.Build.cs
+++ b/Engine/Source/Runtime/Engine/Engine.Build.cs
@@ -165,6 +165,9 @@ public class Engine : ModuleRules
 		{
 			// The SparseVolumeTexture module containing the importer is only loaded and used in the editor.
 			DynamicallyLoadedModuleNames.Add("SparseVolumeTexture");
+			AddEngineThirdPartyPrivateStaticDependencies(Target,
+				"OodleWAV"
+				);
 		}
 
 		// Cross platform Audio Codecs:
@@ -172,8 +175,7 @@ public class Engine : ModuleRules
 			"UEOgg",
 			"Vorbis",
 			"VorbisFile",
-			"libOpus",
-			"OodleWAV"
+			"libOpus"
 			);
 
 		DynamicallyLoadedModuleNames.Add("EyeTracker");
