[Chaos] Changing UClusterUnionComponent::IsAuthority to only check for it's owner local role to determine authority.

This saves ~0.1ms on lower end platforms when we are dealing with a large number of components as getting the Netmode requires us to go through get the World and NetDriver.

A cvar was added to revert to the old behavior if needed by setting it to false, ClusterUnion.UseLocalRoleForAuthorityCheck

diff --git a/Engine/Source/Runtime/Engine/Private/PhysicsEngine/ClusterUnionComponent.cpp b/Engine/Source/Runtime/Engine/Private/PhysicsEngine/ClusterUnionComponent.cpp
--- a/Engine/Source/Runtime/Engine/Private/PhysicsEngine/ClusterUnionComponent.cpp
+++ b/Engine/Source/Runtime/Engine/Private/PhysicsEngine/ClusterUnionComponent.cpp
@@ -25,6 +25,9 @@ namespace
 	bool bUseClusterUnionAccelerationStructure = true;
 	FAutoConsoleVariableRef CVarUseClusterUnionAccelerationStructure(TEXT("ClusterUnion.UseAccelerationStructure"), bUseClusterUnionAccelerationStructure, TEXT("Whether component level sweeps and overlaps against cluster unions should use an acceleration structure instead."));
 
+	bool bUseLocalRoleForAuthorityCheck = true;
+	FAutoConsoleVariableRef CVarUseLocalRoleForAuthorityCheck(TEXT("ClusterUnion.UseLocalRoleForAuthorityCheck"), bUseLocalRoleForAuthorityCheck, TEXT("If true, we will only check this component's owner local role to determine authority"));
+
 	template<typename PayloadType>
 	struct TClusterUnionAABBTreeStorageTraits
 	{
@@ -618,6 +621,16 @@ ENGINE_API bool UClusterUnionComponent::IsAnchored() const
 
 bool UClusterUnionComponent::IsAuthority() const
 {
+	if (bUseLocalRoleForAuthorityCheck)
+	{
+		if (AActor* Owner = GetOwner())
+        {
+        	return Owner->GetLocalRole() == ROLE_Authority;
+        }
+
+		return false;
+	}
+
 	ENetMode Mode = GetNetMode();
 	if (Mode == ENetMode::NM_Standalone)
 	{
