Add animation compression memory usage estimation

Compressing animations happens in the background when a cook is happening. The asset compilation manager has a maximum amount of memory that it is willing to give out and schedules work based on estimated memory usage. For animation compression the estimates can be wrong by an order of magnitude, which is too much to not introduce problems in the system.

The biggest factors are that memory usage can depend on the compression codec in use and on the number of codecs for an asset. The compression code tries all codecs in parallel. A particularly bad case is the RemoveLinearKeys codec, which has memory usage in the order of O(NumKeys * NumBones). This CL adds a new API for estimating the memory usage of running a compression codec and it adjusts the code that estimates the memory usage to actually account for running all codecs in parallel.

diff --git a/Engine/Source/Runtime/Engine/Classes/Animation/AnimBoneCompressionCodec.h b/Engine/Source/Runtime/Engine/Classes/Animation/AnimBoneCompressionCodec.h
--- a/Engine/Source/Runtime/Engine/Classes/Animation/AnimBoneCompressionCodec.h
+++ b/Engine/Source/Runtime/Engine/Classes/Animation/AnimBoneCompressionCodec.h
@@ -9,6 +9,7 @@
 #include "AnimBoneCompressionCodec.generated.h"
 
 class UAnimBoneCompressionCodec;
+class UAnimSequence;
 struct FBlendedCurve;
 class FBoneData;
 
@@ -39,6 +40,19 @@ class UAnimBoneCompressionCodec : public UObject
 	/** Compresses the curve data from an animation sequence. */
 	ENGINE_API virtual bool Compress(const FCompressibleAnimData& CompressibleAnimData, FCompressibleAnimDataResult& OutResult) PURE_VIRTUAL(UAnimCurveCompressionCodec::Compress, return false;);
 
+	/**
+	 * Estimates the peak memory usage in bytes needed to compress the given data with this codec.
+	 * This is used to make informed scheduling decisions during asset cooking. Estimates that are too low may cause
+	 * out-of-memory conditions. Estimates that are too high can unnecessarily limit the number of concurrent cook
+	 * processes.
+	 * 
+	 * Values below zero indicate that no estimate has been given and high memory usage is hence assumed.
+	 * 
+	 * @param AnimSequence	The Animation Sequence to estimate the memory usage for.
+	 * @returns Estimated peak memory usage in bytes. Values < 0 indicate maximum memory usage.
+	 */
+	ENGINE_API virtual int64 EstimateCompressionMemoryUsage(const UAnimSequence& AnimSequence) const;
+
 	/*
 	 * Called to generate a unique DDC key for this codec instance and input anim sequence and TargetPlatform
 	 * A suitable key should be generated from: the InstanceGuid, a codec version, and all relevant properties that drive the behavior.
diff --git a/Engine/Source/Runtime/Engine/Classes/Animation/AnimCompress_RemoveLinearKeys.h b/Engine/Source/Runtime/Engine/Classes/Animation/AnimCompress_RemoveLinearKeys.h
--- a/Engine/Source/Runtime/Engine/Classes/Animation/AnimCompress_RemoveLinearKeys.h
+++ b/Engine/Source/Runtime/Engine/Classes/Animation/AnimCompress_RemoveLinearKeys.h
@@ -76,6 +76,12 @@ class UAnimCompress_RemoveLinearKeys : public UAnimCompress
 	UPROPERTY(EditAnywhere, Category=LinearKeyRemoval)
 	uint32 bActuallyFilterLinearKeys:1;
 
+	//~ Begin UAnimBoneCompressionCodec Interface
+#if WITH_EDITORONLY_DATA
+	virtual int64 EstimateCompressionMemoryUsage(const UAnimSequence& AnimSequence) const override;
+#endif // WITH_EDITORONLY_DATA
+	//~ End UAnimBoneCompressionCodec Interface
+
 protected:
 	//~ Begin UAnimCompress Interface
 #if WITH_EDITOR
diff --git a/Engine/Source/Runtime/Engine/Private/Animation/AnimBoneCompressionCodec.cpp b/Engine/Source/Runtime/Engine/Private/Animation/AnimBoneCompressionCodec.cpp
--- a/Engine/Source/Runtime/Engine/Private/Animation/AnimBoneCompressionCodec.cpp
+++ b/Engine/Source/Runtime/Engine/Private/Animation/AnimBoneCompressionCodec.cpp
@@ -3,6 +3,7 @@
 #include "Animation/AnimBoneCompressionCodec.h"
 #include "Animation/AnimBoneCompressionSettings.h"
 #include "Animation/AnimBoneDecompressionData.h"
+#include "Animation/AnimSequence.h"
 #include "Misc/MemStack.h"
 
 #include UE_INLINE_GENERATED_CPP_BY_NAME(AnimBoneCompressionCodec)
@@ -13,6 +14,18 @@ UAnimBoneCompressionCodec::UAnimBoneCompressionCodec(const FObjectInitializer& O
 {
 }
 
+#if WITH_EDITORONLY_DATA
+int64 UAnimBoneCompressionCodec::EstimateCompressionMemoryUsage(const UAnimSequence& AnimSequence) const
+{
+#if WITH_EDITOR
+	// This is a conservative estimate that gives a codec enough space to create two raw copies of the input data.
+	return 2 * AnimSequence.GetApproxRawSize();
+#else
+	return -1;
+#endif // WITH_EDITOR
+}
+#endif // WITH_EDITORONLY_DATA
+
 UAnimBoneCompressionCodec* UAnimBoneCompressionCodec::GetCodec(const FString& DDCHandle)
 {
 	const FString ThisHandle = GetCodecDDCHandle();
diff --git a/Engine/Source/Runtime/Engine/Private/Animation/AnimCompress_RemoveLinearKeys.cpp b/Engine/Source/Runtime/Engine/Private/Animation/AnimCompress_RemoveLinearKeys.cpp
--- a/Engine/Source/Runtime/Engine/Private/Animation/AnimCompress_RemoveLinearKeys.cpp
+++ b/Engine/Source/Runtime/Engine/Private/Animation/AnimCompress_RemoveLinearKeys.cpp
@@ -59,6 +59,18 @@ UAnimCompress_RemoveLinearKeys::UAnimCompress_RemoveLinearKeys(const FObjectInit
 	bActuallyFilterLinearKeys = true;
 }
 
+#if WITH_EDITORONLY_DATA
+int64 UAnimCompress_RemoveLinearKeys::EstimateCompressionMemoryUsage(const UAnimSequence& AnimSequence) const
+{
+	int64 BaseSize = AnimSequence.GetApproxRawSize();
+	if (const IAnimationDataModel* DataModel = AnimSequence.GetDataModel())
+	{
+		return BaseSize + 3 * DataModel->GetNumberOfKeys() * DataModel->GetNumBoneTracks() * sizeof(FTransform);
+	}
+	return BaseSize;
+}
+#endif // WITH_EDITORONLY_DATA
+
 #if WITH_EDITOR
 struct RotationAdapter
 {
diff --git a/Engine/Source/Runtime/Engine/Private/Animation/AnimationCompressionDerivedData.cpp b/Engine/Source/Runtime/Engine/Private/Animation/AnimationCompressionDerivedData.cpp
--- a/Engine/Source/Runtime/Engine/Private/Animation/AnimationCompressionDerivedData.cpp
+++ b/Engine/Source/Runtime/Engine/Private/Animation/AnimationCompressionDerivedData.cpp
@@ -106,11 +106,31 @@ void FAnimationSequenceAsyncCacheTask::BeginCache(const FIoHash& KeyHash)
 			return AnimSequence->GetApproxRawSize();
 		}();
 		
-		const int64 RequiredMemory = AnimSequence->GetApproxRawSize() + AdditiveAnimSize;
+		int64 RequiredMemory = AnimSequence->GetApproxRawSize() + AdditiveAnimSize;
 		if (Compression::FAnimationCompressionMemorySummaryScope::ShouldStoreCompressionResults())
 		{
 			Compression::FAnimationCompressionMemorySummaryScope::CompressionResultSummary().GatherPreCompressionStats(AnimSequence->GetApproxRawSize(), AnimSequence->GetApproxCompressedSize());
 		}
+
+		if (const UAnimBoneCompressionSettings* BoneCompressionSettings = AnimSequence->BoneCompressionSettings.Get())
+		{
+			// We try out all compression codecs in parallel to find the best one, so we need to sum up the cost for every codec.
+			for (TObjectPtr<UAnimBoneCompressionCodec> Codec : BoneCompressionSettings->Codecs)
+			{
+				if (Codec)
+				{
+					const int64 PeakMemoryEstimate = Codec->EstimateCompressionMemoryUsage(*AnimSequence);
+					if (PeakMemoryEstimate < 0)
+					{
+						// Assume the worst and default to the default behavior when no estimate is given.
+						UE_LOG(LogAnimationCompression, Warning, TEXT("Got invalid memory usage estimate from codec %s for %s. This can negatively affect the time compression takes."), *Codec->GetFullName(), *AnimSequence->GetFullName());
+						RequiredMemory = -1;
+						break;
+					}
+					RequiredMemory += PeakMemoryEstimate;
+				}
+			}
+		}
 		
 		CompressionStartTime = FPlatformTime::Seconds();
 
