Disable contact shadows UI if "Cast Dynamic Shadows" is off as this means contact shadows are disabled too.

diff --git a/Engine/Source/Runtime/Engine/Private/Components/LightComponent.cpp b/Engine/Source/Runtime/Engine/Private/Components/LightComponent.cpp
--- a/Engine/Source/Runtime/Engine/Private/Components/LightComponent.cpp
+++ b/Engine/Source/Runtime/Engine/Private/Components/LightComponent.cpp
@@ -616,6 +616,15 @@ bool ULightComponent::CanEditChange(const FProperty* InProperty) const
 				return false;
 			}
 		}
+
+		if (!CastDynamicShadows &&
+			  (PropertyName == GET_MEMBER_NAME_STRING_CHECKED(ULightComponent, ContactShadowLength)
+			|| PropertyName == GET_MEMBER_NAME_STRING_CHECKED(ULightComponent, ContactShadowLengthInWS)
+			|| PropertyName == GET_MEMBER_NAME_STRING_CHECKED(ULightComponent, ContactShadowCastingIntensity)
+			|| PropertyName == GET_MEMBER_NAME_STRING_CHECKED(ULightComponent, ContactShadowNonCastingIntensity)))
+		{
+			return false;
+		}
 		
 		const bool bIsRayStartOffset = PropertyName == GET_MEMBER_NAME_STRING_CHECKED(ULightComponent, RayStartOffsetDepthScale);
 
