Editor: Fix FActorIteratorState when iterating over specific actor class

WorldPartition can unload actors which will only be unloaded when GC runs. This can lead to actors getting collected by GetObjectsOfClass that are no longer registered to the world. Instead use the Level->Actors array in editor always.

#rnx

diff --git a/Engine/Source/Runtime/Engine/Public/EngineUtils.h b/Engine/Source/Runtime/Engine/Public/EngineUtils.h
--- a/Engine/Source/Runtime/Engine/Public/EngineUtils.h
+++ b/Engine/Source/Runtime/Engine/Public/EngineUtils.h
@@ -21,6 +21,11 @@
 #include "RendererInterface.h"
 #endif
 
+#if WITH_EDITOR
+#include "Algo/Accumulate.h"
+#include "Algo/Copy.h"
+#endif
+
 class FCanvas;
 class FViewport;
 class UCanvas;
@@ -184,39 +189,35 @@ public:
 #if WITH_EDITOR
 		// In the editor, you are more likely to have many worlds in memory at once.
 		// As an optimization to avoid iterating over many actors that are not in the world we are asking for,
-		// if the filter class is AActor, just use the actors that are in the world you asked for.
 		// This could be useful in runtime code as well if there are many worlds in memory, but for now we will leave
 		// it in editor code.
-		if (InClass == AActor::StaticClass())
-		{
-			// First determine the number of actors in the world to reduce reallocations when we append them to the array below.
-			int32 NumActors = 0;
-			for (ULevel* Level : InWorld->GetLevels())
-			{
-				if (Level)
-				{
-					NumActors += Level->Actors.Num();
-				}
-			}
+		
+		// First determine the number of actors in the world to reduce reallocations when we append them to the array below.
+		const int32 MaxActors = Algo::Accumulate(InWorld->GetLevels(), 0, [](int32 InTotalNum, const ULevel* InLevel) { return InTotalNum + (InLevel ? InLevel->Actors.Num() : 0); });
 
-			// Presize the array
-			ObjectArray.Reserve(NumActors);
+		// Presize the array
+		ObjectArray.Reserve(MaxActors);
 
-			// Fill the array
-			for (ULevel* Level : InWorld->GetLevels())
+		for (ULevel* Level : InWorld->GetLevels())
+		{
+			if (Level)
 			{
-				if (Level)
+				if(InClass == AActor::StaticClass())
 				{
 					ObjectArray.Append(Level->Actors);
 				}
+				else
+				{
+					Algo::CopyIf(Level->Actors, ObjectArray, [InClass](AActor* InActor) { return InActor && InActor->IsA(InClass); });
+				}
 			}
 		}
-		else
-#endif // WITH_EDITOR
+#else // WITH_EDITOR
 		{
 			constexpr EObjectFlags ExcludeFlags = RF_ClassDefaultObject;
 			GetObjectsOfClass(InClass, ObjectArray, true, ExcludeFlags, EInternalObjectFlags::Garbage);
 		}
+#endif
 
 		const auto ActorSpawnedDelegate = FOnActorSpawned::FDelegate::CreateRaw(this, &FActorIteratorState::OnActorSpawned);
 		ActorSpawnedDelegateHandle = CurrentWorld->AddOnActorSpawnedHandler(ActorSpawnedDelegate);
