Iris - Set pawn to be scheduled for replication before PlayerController.

diff --git a/Engine/Source/Runtime/Engine/Private/PlayerController.cpp b/Engine/Source/Runtime/Engine/Private/PlayerController.cpp
--- a/Engine/Source/Runtime/Engine/Private/PlayerController.cpp
+++ b/Engine/Source/Runtime/Engine/Private/PlayerController.cpp
@@ -4825,7 +4825,7 @@ void APlayerController::SetPawn(APawn* InPawn)
 
 				if (IsValid(InPawn))
 				{
-					UE::Net::FReplicationSystemUtil::AddDependentActor(this, InPawn);
+					UE::Net::FReplicationSystemUtil::AddDependentActor(this, InPawn, UE::Net::EDependentObjectSchedulingHint::ScheduleBeforeParent);
 				}
 			}
 		}
