FWorldPartitionUtils::FSimulateCookedSession
- Fix crash when destroying simulated cooked session object.
- Add a simple unit-test to make sure we keep this in a working state.

diff --git a/Engine/Source/Runtime/Engine/Private/Tests/WorldPartition/WorldPartitionStreamingGenerationSimulationTests.cpp b/Engine/Source/Runtime/Engine/Private/Tests/WorldPartition/WorldPartitionStreamingGenerationSimulationTests.cpp
new file mode 100644
--- /dev/null
+++ b/Engine/Source/Runtime/Engine/Private/Tests/WorldPartition/WorldPartitionStreamingGenerationSimulationTests.cpp
@@ -0,0 +1,67 @@
+// Copyright Epic Games, Inc. All Rights Reserved.
+
+#include "Misc/AutomationTest.h"
+
+#if WITH_EDITOR
+#include "WorldPartition/WorldPartition.h"
+#include "WorldPartition/WorldPartitionUtils.h"
+#include "EditorWorldUtils.h"
+#endif
+
+#if WITH_DEV_AUTOMATION_TESTS
+
+#define TEST_NAME_ROOT "System.Engine.WorldPartition"
+
+namespace WorldPartitionTests
+{
+	IMPLEMENT_SIMPLE_AUTOMATION_TEST(FWorldPartitionStreamingGenerationSimulationTest, TEST_NAME_ROOT ".StreamingGenerationSimulation", EAutomationTestFlags::EditorContext | EAutomationTestFlags::EngineFilter)
+
+	bool FWorldPartitionStreamingGenerationSimulationTest::RunTest(const FString& Parameters)
+	{
+#if WITH_EDITOR
+		ON_SCOPE_EXIT
+		{
+			CollectGarbage(GARBAGE_COLLECTION_KEEPFLAGS);
+		};
+
+		FScopedEditorWorld ScopedEditorWorld(
+			TEXTVIEW("/Engine/WorldPartition/UnitTests/WPUnitTests"),
+			UWorld::InitializationValues()
+				.RequiresHitProxies(false)
+				.ShouldSimulatePhysics(false)
+				.EnableTraceCollision(false)
+				.CreateNavigation(false)
+				.CreateAISystem(false)
+				.AllowAudioPlayback(false)
+				.CreatePhysicsScene(true)
+		);
+		
+		UWorld* World = ScopedEditorWorld.GetWorld();
+		if (!TestNotNull(TEXT("Missing World Object"), World))
+		{
+			return false;
+		}
+
+		FWorldPartitionUtils::FSimulateCookedSession SimulateCookedSession(World);		
+		if (!TestTrue(TEXT("Simulate Cooking"), SimulateCookedSession.IsValid()))
+		{
+			return false;
+		}
+
+		FWorldPartitionStreamingQuerySource QuerySource;
+		QuerySource.bSpatialQuery = true;
+		QuerySource.Radius = HALF_WORLD_MAX;
+
+		TArray<const IWorldPartitionCell*> Cells;
+		if (!TestTrue(TEXT("GetIntersectingCells"), SimulateCookedSession.GetIntersectingCells({ QuerySource }, Cells)))
+		{
+			return false;
+		}
+#endif
+		return true;
+	}
+}
+
+#undef TEST_NAME_ROOT
+
+#endif
\ No newline at end of file
diff --git a/Engine/Source/Runtime/Engine/Private/WorldPartition/WorldPartitionUtils.cpp b/Engine/Source/Runtime/Engine/Private/WorldPartition/WorldPartitionUtils.cpp
--- a/Engine/Source/Runtime/Engine/Private/WorldPartition/WorldPartitionUtils.cpp
+++ b/Engine/Source/Runtime/Engine/Private/WorldPartition/WorldPartitionUtils.cpp
@@ -14,7 +14,7 @@
 #define LOCTEXT_NAMESPACE "WorldPartition"
 
 FWorldPartitionUtils::FSimulateCookedSession::FSimulateCookedSession(UWorld* InWorld)
-: bIsValid(false)
+	: CookContext(nullptr)
 {
 	UE_CLOG(!InWorld, LogWorldPartition, Error, TEXT("FSimulateCookedSession was given an invalid world."));
 	UE_CLOG(InWorld && !InWorld->IsPartitionedWorld(), LogWorldPartition, Warning, TEXT("FSimulateCookedSession was given a non-partitioned world."));
@@ -22,10 +22,13 @@ FWorldPartitionUtils::FSimulateCookedSession::FSimulateCookedSession(UWorld* InW
 	if (InWorld)
 	{
 		WorldPartition = InWorld->GetWorldPartition();
-		bIsValid = WorldPartition.IsValid() ? SimulateCook() : false;
+		if (WorldPartition.IsValid())
+		{
+			SimulateCook();
+		}
 	}
 
-	UE_CLOG(!bIsValid, LogWorldPartition, Warning, TEXT("FSimulateCookedSession failed to generate streaming."));
+	UE_CLOG(!IsValid(), LogWorldPartition, Warning, TEXT("FSimulateCookedSession failed to generate streaming."));
 }
 
 bool FWorldPartitionUtils::FSimulateCookedSession::SimulateCook()
@@ -40,10 +43,11 @@ bool FWorldPartitionUtils::FSimulateCookedSession::SimulateCook()
 	}
 
 	// Simulate cook
-	FWorldPartitionCookPackageContext CookContext;
-	WorldPartition->BeginCook(CookContext);
+	check(!CookContext);
+	CookContext = new FWorldPartitionCookPackageContext;
+	WorldPartition->BeginCook(*CookContext);
 
-	if (!CookContext.GatherPackagesToCook())
+	if (!CookContext->GatherPackagesToCook())
 	{
 		return false;
 	}
@@ -64,7 +68,7 @@ bool FWorldPartitionUtils::FSimulateCookedSession::SimulateCook()
 
 FWorldPartitionUtils::FSimulateCookedSession::~FSimulateCookedSession()
 {
-	if (WorldPartition.IsValid())
+	if (CookContext)
 	{
 		// Simulate Content Bundles injection
 		TArray<TSharedPtr<FContentBundleEditor>> ContentBundles;
@@ -80,7 +84,10 @@ FWorldPartitionUtils::FSimulateCookedSession::~FSimulateCookedSession()
 			}
 		}
 
+		WorldPartition->EndCook(*CookContext);
 		WorldPartition->FlushStreaming();
+
+		delete CookContext;
 	}
 }
 
diff --git a/Engine/Source/Runtime/Engine/Public/WorldPartition/WorldPartitionUtils.h b/Engine/Source/Runtime/Engine/Public/WorldPartition/WorldPartitionUtils.h
--- a/Engine/Source/Runtime/Engine/Public/WorldPartition/WorldPartitionUtils.h
+++ b/Engine/Source/Runtime/Engine/Public/WorldPartition/WorldPartitionUtils.h
@@ -9,6 +9,7 @@
 class UWorld;
 class UWorldPartition;
 class IWorldPartitionCell;
+class FWorldPartitionCookPackageContext;
 struct FWorldPartitionStreamingQuerySource;
 
 struct FWorldPartitionUtils
@@ -19,14 +20,14 @@ struct FWorldPartitionUtils
 		FSimulateCookedSession(UWorld* InWorld);
 		~FSimulateCookedSession();
 
-		bool IsValid() const { return bIsValid; }
+		bool IsValid() const { return !!CookContext; }
 		bool ForEachStreamingCells(TFunctionRef<void(const IWorldPartitionCell*)> Func);
 		bool GetIntersectingCells(const TArray<FWorldPartitionStreamingQuerySource>& InSources, TArray<const IWorldPartitionCell*>& OutCells);
 
 	private:
 		bool SimulateCook();
 
-		bool bIsValid;
+		FWorldPartitionCookPackageContext* CookContext;
 		TWeakObjectPtr<UWorldPartition> WorldPartition;
 	};
 };
