Fixing attenuation override with live attenuation settings feature

diff --git a/Engine/Source/Runtime/Engine/Private/ActiveSound.cpp b/Engine/Source/Runtime/Engine/Private/ActiveSound.cpp
--- a/Engine/Source/Runtime/Engine/Private/ActiveSound.cpp
+++ b/Engine/Source/Runtime/Engine/Private/ActiveSound.cpp
@@ -1680,14 +1680,17 @@ void FActiveSound::UpdateAttenuation(float DeltaTime, FSoundParseParameters& Par
 
 	// Get the attenuation settings to use for this application to the active sound
 	// Use the passed-in attenuation settings
-	if (SettingsAttenuationNode)
+	if (!bIsAttenuationSettingsOverridden)
 	{
-		Settings = SettingsAttenuationNode;
-	}
-	// We fallback to using the asset's settings directly
-	else if (SoundAttenuation)
-	{
-		Settings = &SoundAttenuation->Attenuation;
+		if (SettingsAttenuationNode)
+		{
+			Settings = SettingsAttenuationNode;
+		}
+		// We fallback to using the asset's settings directly
+		else if (SoundAttenuation)
+		{
+			Settings = &SoundAttenuation->Attenuation;
+		}
 	}
 
 	// Reset Focus data and recompute if necessary
