Fixed UWorldPartitionLevelStreamingDynamic::RequestLevel which didn't properly handle requesting a WorldAsset already requeted by another streaming level in the world (affects PIE).
Refactored ULevelStreaming to share code with child classes.

diff --git a/Engine/Source/Runtime/Engine/Classes/Engine/LevelStreaming.h b/Engine/Source/Runtime/Engine/Classes/Engine/LevelStreaming.h
--- a/Engine/Source/Runtime/Engine/Classes/Engine/LevelStreaming.h
+++ b/Engine/Source/Runtime/Engine/Classes/Engine/LevelStreaming.h
@@ -687,6 +687,9 @@ protected:
 	/** Called by SetLoadedLevel */
 	virtual void OnLevelLoadedChanged(ULevel* Level) {}
 
+	/** Called by RequestLevel to detect existing streaming level with same world asset */
+	bool ValidateUniqueWorldAsset(UWorld* PersistentWorld);
+
 	ENGINE_API void OnLoadingStarted();
 	ENGINE_API void OnLoadingFinished();
 
diff --git a/Engine/Source/Runtime/Engine/Private/LevelStreaming.cpp b/Engine/Source/Runtime/Engine/Private/LevelStreaming.cpp
--- a/Engine/Source/Runtime/Engine/Private/LevelStreaming.cpp
+++ b/Engine/Source/Runtime/Engine/Private/LevelStreaming.cpp
@@ -1434,6 +1434,40 @@ void ULevelStreaming::PrepareLoadedLevel(ULevel* InLevel, UPackage* InLevelPacka
 	}
 }
 
+bool ULevelStreaming::ValidateUniqueWorldAsset(UWorld* PersistentWorld)
+{
+	// Validate that the streaming level is unique, check for clash with currently loaded streaming levels
+	for (ULevelStreaming* OtherLevel : PersistentWorld->GetStreamingLevels())
+	{
+		if (OtherLevel == nullptr || OtherLevel == this)
+		{
+			continue;
+		}
+
+		const ELevelStreamingState OtherState = OtherLevel->CurrentState;
+		if (OtherState == ELevelStreamingState::FailedToLoad || OtherState == ELevelStreamingState::Removed || (OtherState == ELevelStreamingState::Unloaded && (OtherLevel->TargetState == ELevelStreamingTargetState::Unloaded || OtherLevel->TargetState == ELevelStreamingTargetState::UnloadedAndRemoved)))
+		{
+			// If the other level is neither loaded nor in the process of being loaded, we don't need to consider it
+			continue;
+		}
+
+		if (OtherLevel->WorldAsset == WorldAsset)
+		{
+			if (OtherLevel->GetIsRequestingUnloadAndRemoval())
+			{
+				return false; // Cannot load now, retry until the OtherLevel is done unloading
+			}
+			else
+			{
+				UE_LOG(LogLevelStreaming, Warning, TEXT("Streaming Level '%s' uses same destination for level ('%s') as '%s'. Level cannot be loaded again and this StreamingLevel will be flagged as failed to load."), *GetPathName(), *WorldAsset.GetLongPackageName(), *OtherLevel->GetPathName());
+				SetCurrentState(ELevelStreamingState::FailedToLoad);
+				return false;
+			}
+		}
+	}
+	return true;
+}
+
 bool ULevelStreaming::RequestLevel(UWorld* PersistentWorld, bool bAllowLevelLoadRequests, EReqLevelBlock BlockPolicy)
 {
 	// Quit early in case load request already issued
@@ -1477,33 +1511,9 @@ bool ULevelStreaming::RequestLevel(UWorld* PersistentWorld, bool bAllowLevelLoad
 	}
 
 	// Validate that our new streaming level is unique, check for clash with currently loaded streaming levels
-	for (ULevelStreaming* OtherLevel : PersistentWorld->GetStreamingLevels())
+	if (!ValidateUniqueWorldAsset(PersistentWorld))
 	{
-		if (OtherLevel == nullptr || OtherLevel == this)
-		{
-			continue;
-		}
-
-		const ELevelStreamingState OtherState = OtherLevel->CurrentState;
-		if (OtherState == ELevelStreamingState::FailedToLoad || OtherState == ELevelStreamingState::Removed || (OtherState == ELevelStreamingState::Unloaded && (OtherLevel->TargetState == ELevelStreamingTargetState::Unloaded || OtherLevel->TargetState == ELevelStreamingTargetState::UnloadedAndRemoved)))
-		{
-			// If the other level isn't loaded or in the process of being loaded we don't need to consider it
-			continue;
-		}
-
-		if (OtherLevel->WorldAsset == WorldAsset)
-		{ 
-			if (OtherLevel->GetIsRequestingUnloadAndRemoval())
-			{
-				return false; // Cannot load new level now, retry until the OtherLevel is done unloading
-			}
-			else
-			{
-				UE_LOG(LogLevelStreaming, Warning, TEXT("Streaming Level '%s' uses same destination for level ('%s') as '%s'. Level cannot be loaded again and this StreamingLevel will be flagged as failed to load."), *GetPathName(), *WorldAsset.GetLongPackageName(), *OtherLevel->GetPathName());
-				SetCurrentState(ELevelStreamingState::FailedToLoad);
-				return false;
-			}
-		}
+		return false;
 	}
 
 	TRACE_LOADTIME_REQUEST_GROUP_SCOPE(TEXT("LevelStreaming - %s"), *GetPathName());
diff --git a/Engine/Source/Runtime/Engine/Private/WorldPartition/WorldPartitionLevelStreamingDynamic.cpp b/Engine/Source/Runtime/Engine/Private/WorldPartition/WorldPartitionLevelStreamingDynamic.cpp
--- a/Engine/Source/Runtime/Engine/Private/WorldPartition/WorldPartitionLevelStreamingDynamic.cpp
+++ b/Engine/Source/Runtime/Engine/Private/WorldPartition/WorldPartitionLevelStreamingDynamic.cpp
@@ -193,14 +193,13 @@ void UWorldPartitionLevelStreamingDynamic::CreateRuntimeLevel()
  */
 bool UWorldPartitionLevelStreamingDynamic::RequestLevel(UWorld* InPersistentWorld, bool bInAllowLevelLoadRequests, EReqLevelBlock InBlockPolicy)
 {
-	const FName WorldAssetPackageFName = GetWorldAssetPackageFName();
-	const FString WorldAssetPackageName = GetWorldAssetPackageName();
-
 	if (bShouldPerformStandardLevelLoading)
 	{
 		return Super::RequestLevel(InPersistentWorld, bInAllowLevelLoadRequests, InBlockPolicy);
 	}
 
+	QUICK_SCOPE_CYCLE_COUNTER(STAT_ULevelStreaming_RequestLevel);
+
 	// Quit early in case load request already issued
 	if (GetLevelStreamingState() == ELevelStreamingState::Loading)
 	{
@@ -219,6 +218,9 @@ bool UWorldPartitionLevelStreamingDynamic::RequestLevel(UWorld* InPersistentWorl
 		return false;
 	}
 
+	const FName WorldAssetPackageFName = GetWorldAssetPackageFName();
+	const FString WorldAssetPackageName = GetWorldAssetPackageName();
+
 	// Can not load new level now either, we're still processing visibility for this one
 	ULevel* PendingLevelVisOrInvis = (InPersistentWorld->GetCurrentLevelPendingVisibility() ? InPersistentWorld->GetCurrentLevelPendingVisibility() : InPersistentWorld->GetCurrentLevelPendingInvisibility());
 	if (PendingLevelVisOrInvis && PendingLevelVisOrInvis == LoadedLevel)
@@ -227,7 +229,12 @@ bool UWorldPartitionLevelStreamingDynamic::RequestLevel(UWorld* InPersistentWorl
 		return false;
 	}
 
-	QUICK_SCOPE_CYCLE_COUNTER(STAT_ULevelStreaming_RequestLevel);
+	// Validate that our new streaming level is unique, check for clash with currently loaded streaming levels
+	if (!ValidateUniqueWorldAsset(InPersistentWorld))
+	{
+		return false;
+	}
+
 	FScopeCycleCounterUObject Context(InPersistentWorld);
 
 	// Try to find the package to load
