Fix for variable seek tables checking for incorrect failure sentinel.

diff --git a/Engine/Source/Runtime/Engine/Private/StreamedAudioChunkSeekTable.cpp b/Engine/Source/Runtime/Engine/Private/StreamedAudioChunkSeekTable.cpp
--- a/Engine/Source/Runtime/Engine/Private/StreamedAudioChunkSeekTable.cpp
+++ b/Engine/Source/Runtime/Engine/Private/StreamedAudioChunkSeekTable.cpp
@@ -228,7 +228,7 @@ struct FVariableRateSeekTable : public FStreamedAudioChunkSeekTable::ISeekTableI
 	{
 		check(Offsets.Num() == Times.Num());
 		int32 Index = Algo::LowerBound(Times, InTimeInAudioFrames);
-		if (Index != INDEX_NONE)
+		if (Index < Times.Num())
 		{
 			return Offsets[Index];
 		}
@@ -237,7 +237,7 @@ struct FVariableRateSeekTable : public FStreamedAudioChunkSeekTable::ISeekTableI
 	uint32 FindTime(uint32 InOffset) const override
 	{
 		int32 Index = Algo::LowerBound(Offsets,InOffset);
-		if (Index != INDEX_NONE)
+		if (Index < Offsets.Num())
 		{
 			return Times[Index];
 		}
