Only disable DCC on gbuffers when nanite fast tile clear is enabled

diff --git a/Engine/Source/Runtime/Engine/Private/SceneTexturesConfig.cpp b/Engine/Source/Runtime/Engine/Private/SceneTexturesConfig.cpp
--- a/Engine/Source/Runtime/Engine/Private/SceneTexturesConfig.cpp
+++ b/Engine/Source/Runtime/Engine/Private/SceneTexturesConfig.cpp
@@ -298,6 +298,10 @@ void FSceneTexturesConfig::Init(const FSceneTexturesConfigInitSettings& InitSett
 				BindingCache.Bindings[Layout].GBufferD = FindGBufferBindingByName(GBufferInfo, TEXT("GBufferD"));
 				BindingCache.Bindings[Layout].GBufferE = FindGBufferBindingByName(GBufferInfo, TEXT("GBufferE"));
 				BindingCache.Bindings[Layout].GBufferVelocity = FindGBufferBindingByName(GBufferInfo, TEXT("Velocity"));
+
+				// Remove DisableDCC flag for velocity. Only Nanite fast tile clear sets this flag currently
+				// but we want to exclude velocity because it usually doesn't have many written pixels
+				EnumRemoveFlags(BindingCache.Bindings[Layout].GBufferVelocity.Flags, TexCreate_DisableDCC);
 			}
 
 			BindingCache.GBufferParams = DefaultParams;
diff --git a/Engine/Source/Runtime/RenderCore/Private/GBufferInfo.cpp b/Engine/Source/Runtime/RenderCore/Private/GBufferInfo.cpp
--- a/Engine/Source/Runtime/RenderCore/Private/GBufferInfo.cpp
+++ b/Engine/Source/Runtime/RenderCore/Private/GBufferInfo.cpp
@@ -185,7 +185,12 @@ FGBufferBinding FindGBufferBindingByName(const FGBufferInfo& GBufferInfo, const
 
 		if (NaniteComputeMaterialsSupported())
 		{
-			Binding.Flags |= TexCreate_UAV | TexCreate_DisableDCC;
+			Binding.Flags |= TexCreate_UAV;
+
+			if (UseNaniteFastTileClear())
+			{
+				Binding.Flags |= TexCreate_DisableDCC;
+			}
 		}
 	}
 
diff --git a/Engine/Source/Runtime/RenderCore/Private/RenderUtils.cpp b/Engine/Source/Runtime/RenderCore/Private/RenderUtils.cpp
--- a/Engine/Source/Runtime/RenderCore/Private/RenderUtils.cpp
+++ b/Engine/Source/Runtime/RenderCore/Private/RenderUtils.cpp
@@ -1328,6 +1328,12 @@ bool UseNaniteComputeMaterials()
 	return NaniteComputeMaterialsSupported() && !IsVulkanPlatform(GMaxRHIShaderPlatform) /* TODO: Support CS derivatives */ && (UseComputeMaterials && UseComputeMaterials->GetValueOnRenderThread() != 0);
 }
 
+bool UseNaniteFastTileClear()
+{
+	static const IConsoleVariable* CVarNaniteFastTileClear = IConsoleManager::Get().FindConsoleVariable(TEXT("r.Nanite.FastTileClear"));
+	return NaniteComputeMaterialsSupported() && (CVarNaniteFastTileClear && CVarNaniteFastTileClear->GetInt() != 0);
+}
+
 bool NaniteTessellationSupported()
 {
 	static const auto AllowTessellation = IConsoleManager::Get().FindTConsoleVariableDataInt(TEXT("r.Nanite.AllowTessellation"));
diff --git a/Engine/Source/Runtime/RenderCore/Public/RenderUtils.h b/Engine/Source/Runtime/RenderCore/Public/RenderUtils.h
--- a/Engine/Source/Runtime/RenderCore/Public/RenderUtils.h
+++ b/Engine/Source/Runtime/RenderCore/Public/RenderUtils.h
@@ -379,6 +379,7 @@ RENDERCORE_API bool NaniteTessellationSupported();
 RENDERCORE_API bool NaniteSplineMeshesSupported();
 
 RENDERCORE_API bool UseNaniteComputeMaterials();
+RENDERCORE_API bool UseNaniteFastTileClear();
 
 RENDERCORE_API bool DoesPlatformSupportNanite(EShaderPlatform Platform, bool bCheckForProjectSetting = true);
 RENDERCORE_API bool DoesRuntimeSupportNanite(EShaderPlatform ShaderPlatform, bool bCheckForAtomicSupport, bool bCheckForProjectSetting);
