#UE Define CVAR in cpp file

diff --git a/Engine/Source/Runtime/Engine/Private/Animation/AnimTrace.cpp b/Engine/Source/Runtime/Engine/Private/Animation/AnimTrace.cpp
--- a/Engine/Source/Runtime/Engine/Private/Animation/AnimTrace.cpp
+++ b/Engine/Source/Runtime/Engine/Private/Animation/AnimTrace.cpp
@@ -285,6 +285,13 @@ UE_TRACE_EVENT_BEGIN(Animation, PoseWatch2)
 	UE_TRACE_EVENT_FIELD(bool, bIsEnabled)
 UE_TRACE_EVENT_END()
 
+FAutoConsoleVariable CVarRecordExternalMorphTargets(
+	TEXT("RecordExternalMorphTargets"),
+	false,
+	TEXT("Record the external morph target weights inside animation insights. On default this is disabled, because it can slow down recording."),
+	ECVF_Default
+);
+
 // Object annotations used for tracing
 FUObjectAnnotationSparseBool GSkeletalMeshTraceAnnotations;
 
diff --git a/Engine/Source/Runtime/Engine/Public/Animation/AnimTrace.h b/Engine/Source/Runtime/Engine/Public/Animation/AnimTrace.h
--- a/Engine/Source/Runtime/Engine/Public/Animation/AnimTrace.h
+++ b/Engine/Source/Runtime/Engine/Public/Animation/AnimTrace.h
@@ -36,12 +36,7 @@ struct FAnimSyncMarker;
 struct FAnimMontageInstance;
 class UPoseWatchPoseElement;
 
-static FAutoConsoleVariable CVarRecordExternalMorphTargets(
-	TEXT("RecordExternalMorphTargets"),
-	false,
-	TEXT("Record the external morph target weights inside animation insights. On default this is disabled, because it can slow down recording."),
-	ECVF_Default
-);
+extern ENGINE_API FAutoConsoleVariable CVarRecordExternalMorphTargets;
 
 struct FAnimTrace
 {
