World Partition: added wp.Runtime.RuntimeSpatialHashSortUsingCellPriority to fix taking cell priority into account when sorting cells by importance.

diff --git a/Engine/Source/Runtime/Engine/Private/WorldPartition/RuntimeHashSet/WorldPartitionRuntimeHashSetStreamingGeneration.cpp b/Engine/Source/Runtime/Engine/Private/WorldPartition/RuntimeHashSet/WorldPartitionRuntimeHashSetStreamingGeneration.cpp
--- a/Engine/Source/Runtime/Engine/Private/WorldPartition/RuntimeHashSet/WorldPartitionRuntimeHashSetStreamingGeneration.cpp
+++ b/Engine/Source/Runtime/Engine/Private/WorldPartition/RuntimeHashSet/WorldPartitionRuntimeHashSetStreamingGeneration.cpp
@@ -184,7 +184,6 @@ bool UWorldPartitionRuntimeHashSet::GenerateStreaming(UWorldPartitionStreamingPo
 		RuntimeCell->SetIsAlwaysLoaded(!CellDescInstance.bIsSpatiallyLoaded);
 		RuntimeCell->SetDataLayers(CellDescInstance.DataLayerInstances);
 		RuntimeCell->SetContentBundleUID(CellDescInstance.ContentBundleID);
-		RuntimeCell->SetPriority(CellDescInstance.Priority);
 		RuntimeCell->SetClientOnlyVisible(CellDescInstance.bClientOnlyVisible);
 		RuntimeCell->SetBlockOnSlowLoading(CellDescInstance.bBlockOnSlowStreaming);
 		RuntimeCell->SetIsHLOD(CellDescInstance.SourcePartition->HLODIndex != INDEX_NONE);
@@ -193,6 +192,7 @@ bool UWorldPartitionRuntimeHashSet::GenerateStreaming(UWorldPartitionStreamingPo
 		UWorldPartitionRuntimeCellDataSpatialHashSet* RuntimeCellDataHashSet = CastChecked<UWorldPartitionRuntimeCellDataSpatialHashSet>(RuntimeCell->RuntimeCellData);
 		RuntimeCellDataHashSet->DebugName = CellUniqueId.Name;
 		RuntimeCellDataHashSet->Level = CellDescInstance.Level;
+		RuntimeCellDataHashSet->Priority = CellDescInstance.Priority;
 
 		return RuntimeCell;
 	};
diff --git a/Engine/Source/Runtime/Engine/Private/WorldPartition/WorldPartitionRuntimeCell.cpp b/Engine/Source/Runtime/Engine/Private/WorldPartition/WorldPartitionRuntimeCell.cpp
--- a/Engine/Source/Runtime/Engine/Private/WorldPartition/WorldPartitionRuntimeCell.cpp
+++ b/Engine/Source/Runtime/Engine/Private/WorldPartition/WorldPartitionRuntimeCell.cpp
@@ -17,7 +17,6 @@
 UWorldPartitionRuntimeCell::UWorldPartitionRuntimeCell(const FObjectInitializer& ObjectInitializer)
 	: Super(ObjectInitializer)
 	, bIsAlwaysLoaded(false)
-	, Priority(0)
 #if !UE_BUILD_SHIPPING
 	, DebugStreamingPriority(-1.f)
 #endif
@@ -98,10 +97,7 @@ void UWorldPartitionRuntimeCell::DumpStateLog(FHierarchicalLogArchive& Ar) const
 
 int32 UWorldPartitionRuntimeCell::SortCompare(const UWorldPartitionRuntimeCell* Other, bool bCanUseSortingCache) const
 {
-	const int32 Comparison = RuntimeCellData->SortCompare(Other->RuntimeCellData, bCanUseSortingCache);
-
-	// Cell priority (lower value is higher prio)
-	return (Comparison != 0) ? Comparison : (Priority - Other->Priority);
+	return RuntimeCellData->SortCompare(Other->RuntimeCellData, bCanUseSortingCache);
 }
 
 bool UWorldPartitionRuntimeCell::IsDebugShown() const
diff --git a/Engine/Source/Runtime/Engine/Private/WorldPartition/WorldPartitionRuntimeCellData.cpp b/Engine/Source/Runtime/Engine/Private/WorldPartition/WorldPartitionRuntimeCellData.cpp
--- a/Engine/Source/Runtime/Engine/Private/WorldPartition/WorldPartitionRuntimeCellData.cpp
+++ b/Engine/Source/Runtime/Engine/Private/WorldPartition/WorldPartitionRuntimeCellData.cpp
@@ -9,6 +9,7 @@ UWorldPartitionRuntimeCellData::UWorldPartitionRuntimeCellData(const FObjectInit
 	, CachedMinSourcePriority(MAX_uint8)
 	, CachedSourceInfoEpoch(MIN_int32)
 	, ContentBounds(ForceInit)
+	, Priority(0)
 {}
 
 void UWorldPartitionRuntimeCellData::Serialize(FArchive& Ar)
diff --git a/Engine/Source/Runtime/Engine/Private/WorldPartition/WorldPartitionRuntimeCellDataSpatialHash.cpp b/Engine/Source/Runtime/Engine/Private/WorldPartition/WorldPartitionRuntimeCellDataSpatialHash.cpp
--- a/Engine/Source/Runtime/Engine/Private/WorldPartition/WorldPartitionRuntimeCellDataSpatialHash.cpp
+++ b/Engine/Source/Runtime/Engine/Private/WorldPartition/WorldPartitionRuntimeCellDataSpatialHash.cpp
@@ -18,6 +18,12 @@ static FAutoConsoleVariableRef CVarRuntimeSpatialHashSortUsingCellExtent(
 	GRuntimeSpatialHashSortUsingCellExtent,
 	TEXT("Set to 1 to use cell extent instead of cell grid level when sorting cells by importance."));
 
+static bool GRuntimeSpatialHashSortUsingCellPriority = true;
+static FAutoConsoleVariableRef CVarRuntimeSpatialHashSortUsingCellPriority(
+	TEXT("wp.Runtime.RuntimeSpatialHashSortUsingCellPriority"),
+	GRuntimeSpatialHashSortUsingCellPriority,
+	TEXT("Set to 1 to use cell priority as part of the sorting criterias when sorting cells by importance."));
+
 UWorldPartitionRuntimeCellDataSpatialHash::UWorldPartitionRuntimeCellDataSpatialHash(const FObjectInitializer& ObjectInitializer)
 	: Super(ObjectInitializer)
 	, Extent(0)
@@ -146,19 +152,35 @@ int32 UWorldPartitionRuntimeCellDataSpatialHash::SortCompare(const UWorldPartiti
 		Result = GRuntimeSpatialHashSortUsingCellExtent ? int32(Other->Extent - Extent) : (Other->Level - Level);
 		if (bCanUseSortingCache && (Result == 0))
 		{
-			// Closest distance (lower value is higher prio)
-			const double Diff = CachedSourceSortingDistance - Other->CachedSourceSortingDistance;
-			if (FMath::IsNearlyZero(Diff))
+			if (GRuntimeSpatialHashSortUsingCellPriority)
 			{
-				const double RawDistanceDiff = CachedMinSquareDistanceToSource - Other->CachedMinSquareDistanceToSource;
-				Result = RawDistanceDiff < 0 ? -1 : (RawDistanceDiff > 0.f ? 1 : 0);
+				// Cell priority (lower value is higher prio)
+				Result = Priority - Other->Priority;
 			}
-			else
+
+			if (Result == 0)
 			{
-				Result = Diff < 0.f ? -1 : (Diff > 0.f ? 1 : 0);
+				// Closest distance (lower value is higher prio)
+				const double Diff = CachedSourceSortingDistance - Other->CachedSourceSortingDistance;
+				if (FMath::IsNearlyZero(Diff))
+				{
+					const double RawDistanceDiff = CachedMinSquareDistanceToSource - Other->CachedMinSquareDistanceToSource;
+					Result = RawDistanceDiff < 0 ? -1 : (RawDistanceDiff > 0.f ? 1 : 0);
+				}
+				else
+				{
+					Result = Diff < 0.f ? -1 : (Diff > 0.f ? 1 : 0);
+				}
 			}
 		}
 	}
+
+	if (!GRuntimeSpatialHashSortUsingCellPriority && Result == 0)
+	{
+		// Cell priority (lower value is higher prio)
+		Result = Priority - InOther->Priority;
+	}
+
 	return Result;
 }
 
diff --git a/Engine/Source/Runtime/Engine/Private/WorldPartition/WorldPartitionRuntimeSpatialHash.cpp b/Engine/Source/Runtime/Engine/Private/WorldPartition/WorldPartitionRuntimeSpatialHash.cpp
--- a/Engine/Source/Runtime/Engine/Private/WorldPartition/WorldPartitionRuntimeSpatialHash.cpp
+++ b/Engine/Source/Runtime/Engine/Private/WorldPartition/WorldPartitionRuntimeSpatialHash.cpp
@@ -1524,7 +1524,6 @@ bool UWorldPartitionRuntimeSpatialHash::CreateStreamingGrid(const FSpatialHashRu
 					StreamingCell->SetIsAlwaysLoaded(bIsCellAlwaysLoaded);
 					StreamingCell->SetDataLayers(GridCellDataChunk.GetDataLayers());
 					StreamingCell->SetContentBundleUID(GridCellDataChunk.GetContentBundleID());
-					StreamingCell->SetPriority(RuntimeGrid.Priority);
 					StreamingCell->SetClientOnlyVisible(CurrentStreamingGrid.bClientOnlyVisible);
 					StreamingCell->SetBlockOnSlowLoading(CurrentStreamingGrid.bBlockOnSlowStreaming);
 					StreamingCell->SetIsHLOD(RuntimeGrid.HLODLayer ? true : false);
@@ -1538,8 +1537,9 @@ bool UWorldPartitionRuntimeSpatialHash::CreateStreamingGrid(const FSpatialHashRu
 					CellDataSpatialHash->Level = Level;
 					CellDataSpatialHash->Position = FVector(Bounds.GetCenter(), 0.f);
 					CellDataSpatialHash->Extent = (float)CellExtent;
-					CellDataSpatialHash->GridName = RuntimeGrid.GridName;			
+					CellDataSpatialHash->GridName = RuntimeGrid.GridName;
 					CellDataSpatialHash->DebugName = CellName + WorldInstanceSuffix;
+					CellDataSpatialHash->Priority = RuntimeGrid.Priority;
 
 					PopulateRuntimeCell(StreamingCell, FilteredActors, OutPackagesToGenerate);
 
diff --git a/Engine/Source/Runtime/Engine/Public/WorldPartition/WorldPartitionRuntimeCell.h b/Engine/Source/Runtime/Engine/Public/WorldPartition/WorldPartitionRuntimeCell.h
--- a/Engine/Source/Runtime/Engine/Public/WorldPartition/WorldPartitionRuntimeCell.h
+++ b/Engine/Source/Runtime/Engine/Public/WorldPartition/WorldPartitionRuntimeCell.h
@@ -199,7 +199,6 @@ class UWorldPartitionRuntimeCell : public UObject, public IWorldPartitionCell
 	virtual FLinearColor GetDebugColor(EWorldPartitionRuntimeCellVisualizeMode VisualizeMode) const { static const FLinearColor DefaultColor = FLinearColor::Black.CopyWithNewOpacity(0.25f); return DefaultColor; }
 	virtual bool IsAlwaysLoaded() const { return bIsAlwaysLoaded; }
 	virtual void SetIsAlwaysLoaded(bool bInIsAlwaysLoaded) { bIsAlwaysLoaded = bInIsAlwaysLoaded; }
-	virtual void SetPriority(int32 InPriority) { Priority = InPriority; }
 	ENGINE_API virtual void SetStreamingPriority(int32 InStreamingPriority) const PURE_VIRTUAL(UWorldPartitionRuntimeCell::SetStreamingPriority,);
 	virtual EStreamingStatus GetStreamingStatus() const { return LEVEL_Unloaded; }
 	UE_DEPRECATED(5.3, "IsLoading is deprecated.")
@@ -288,10 +287,6 @@ private:
 	UPROPERTY()
 	TArray<FName> DataLayers;
 
-	// Custom Priority
-	UPROPERTY()
-	int32 Priority;
-
 	UPROPERTY()
 	bool bClientOnlyVisible;
 
diff --git a/Engine/Source/Runtime/Engine/Public/WorldPartition/WorldPartitionRuntimeCellData.h b/Engine/Source/Runtime/Engine/Public/WorldPartition/WorldPartitionRuntimeCellData.h
--- a/Engine/Source/Runtime/Engine/Public/WorldPartition/WorldPartitionRuntimeCellData.h
+++ b/Engine/Source/Runtime/Engine/Public/WorldPartition/WorldPartitionRuntimeCellData.h
@@ -48,5 +48,8 @@ class UWorldPartitionRuntimeCellData : public UObject
 	UPROPERTY()
 	FBox ContentBounds;
 
+	UPROPERTY()
+	int32 Priority;
+
 	FStringTest DebugName;
 };
