Disabled parallel GDME to investigate crashes.

diff --git a/Engine/Source/Runtime/Engine/Private/PrimitiveSceneProxy.cpp b/Engine/Source/Runtime/Engine/Private/PrimitiveSceneProxy.cpp
--- a/Engine/Source/Runtime/Engine/Private/PrimitiveSceneProxy.cpp
+++ b/Engine/Source/Runtime/Engine/Private/PrimitiveSceneProxy.cpp
@@ -106,7 +106,7 @@ static TAutoConsoleVariable<int32> CVarApproximateOcclusionQueries(
 	ECVF_RenderThreadSafe
 );
 
-bool GParallelGatherDynamicMeshElements = true;
+bool GParallelGatherDynamicMeshElements = false;
 static FAutoConsoleVariableRef CVarParallelGatherDynamicMeshElements(
 	TEXT("r.Visibility.DynamicMeshElements.Parallel"),
 	GParallelGatherDynamicMeshElements,
