Add csv profiling tags to InstanceStaticMeshComponent

diff --git a/Engine/Source/Runtime/Engine/Private/InstancedStaticMesh.cpp b/Engine/Source/Runtime/Engine/Private/InstancedStaticMesh.cpp
--- a/Engine/Source/Runtime/Engine/Private/InstancedStaticMesh.cpp
+++ b/Engine/Source/Runtime/Engine/Private/InstancedStaticMesh.cpp
@@ -2483,6 +2483,7 @@ void UInstancedStaticMeshComponent::InitInstanceBody(int32 InstanceIdx, FBodyIns
 void UInstancedStaticMeshComponent::CreateAllInstanceBodies()
 {
 	TRACE_CPUPROFILER_EVENT_SCOPE(UInstancedStaticMeshComponent::CreateAllInstanceBodies);
+	CSV_SCOPED_TIMING_STAT_EXCLUSIVE(UInstancedStaticMeshComponent_CreateAllInstanceBodies);
 	QUICK_SCOPE_CYCLE_COUNTER(STAT_UInstancedStaticMeshComponent_CreateAllInstanceBodies);
 	STAT(FScopeCycleCounter Context(StatId);)
 
@@ -3148,6 +3149,9 @@ int32 UInstancedStaticMeshComponent::AddInstance(const FTransform& InstanceTrans
 
 TArray<int32> UInstancedStaticMeshComponent::AddInstancesInternal(TConstArrayView<FTransform> InstanceTransforms, bool bShouldReturnIndices, bool bWorldSpace, bool bUpdateNavigation)
 {
+	TRACE_CPUPROFILER_EVENT_SCOPE(UInstancedStaticMeshComponent::AddInstancesInternal);
+	CSV_SCOPED_TIMING_STAT_EXCLUSIVE(UInstancedStaticMeshComponent_AddInstancesInternal);
+
 	const int32 Count = InstanceTransforms.Num();
 
 	TArray<int32> NewInstanceIndices;
