Propagate asset name all the way to base shadermap serialization for a better error message.

diff --git a/Engine/Source/Runtime/Engine/Classes/Materials/MaterialInterface.h b/Engine/Source/Runtime/Engine/Classes/Materials/MaterialInterface.h
--- a/Engine/Source/Runtime/Engine/Classes/Materials/MaterialInterface.h
+++ b/Engine/Source/Runtime/Engine/Classes/Materials/MaterialInterface.h
@@ -1202,6 +1202,7 @@ extern void SerializeInlineShaderMaps(
 	const TMap<const class ITargetPlatform*, TArray<FMaterialResource*>>* PlatformMaterialResourcesToSave,
 	FArchive& Ar,
 	TArray<FMaterialResource>& OutLoadedResources,
+	const FName& SerializingAsset = NAME_None,
 	uint32* OutOffsetToFirstResource = nullptr);
 /** Helper function to process (register) serialized inline shader maps for the given material resources. */
 extern void ProcessSerializedInlineShaderMaps(UMaterialInterface* Owner, TArray<FMaterialResource>& LoadedResources, TArray<FMaterialResource*>& OutMaterialResourcesLoaded);
diff --git a/Engine/Source/Runtime/Engine/Private/Materials/Material.cpp b/Engine/Source/Runtime/Engine/Private/Materials/Material.cpp
--- a/Engine/Source/Runtime/Engine/Private/Materials/Material.cpp
+++ b/Engine/Source/Runtime/Engine/Private/Materials/Material.cpp
@@ -720,6 +720,7 @@ void SerializeInlineShaderMaps(
 	const TMap<const ITargetPlatform*, TArray<FMaterialResource*>>* PlatformMaterialResourcesToSavePtr,
 	FArchive& Ar,
 	TArray<FMaterialResource>& OutLoadedResources,
+	const FName& SerializingAsset,
 	uint32* OutOffsetToFirstResource)
 {
 	LLM_SCOPE(ELLMTag::Shaders);
@@ -788,7 +789,7 @@ void SerializeInlineShaderMaps(
 			for (int32 ResourceIndex = 0; ResourceIndex < NumLoadedResources; ++ResourceIndex)
 			{
 				FMaterialResource& LoadedResource = OutLoadedResources[OutLoadedResources.AddDefaulted()];
-				LoadedResource.SerializeInlineShaderMap(ResourceAr);
+				LoadedResource.SerializeInlineShaderMap(ResourceAr, SerializingAsset);
 			}
 #endif
 		}
@@ -2754,7 +2755,8 @@ void UMaterial::Serialize(FArchive& Ar)
 		SerializeInlineShaderMaps(
 			NULL,
 			Ar,
-			LoadedMaterialResources
+			LoadedMaterialResources,
+			GetFName()
 #if STORE_ONLY_ACTIVE_SHADERMAPS
 			, &OffsetToFirstResource
 #endif
diff --git a/Engine/Source/Runtime/Engine/Private/Materials/MaterialInstance.cpp b/Engine/Source/Runtime/Engine/Private/Materials/MaterialInstance.cpp
--- a/Engine/Source/Runtime/Engine/Private/Materials/MaterialInstance.cpp
+++ b/Engine/Source/Runtime/Engine/Private/Materials/MaterialInstance.cpp
@@ -2909,7 +2909,8 @@ void UMaterialInstance::Serialize(FArchive& Ar)
 			SerializeInlineShaderMaps(
 				NULL,
 				Ar,
-				LoadedMaterialResources
+				LoadedMaterialResources,
+				GetFName()
 #if STORE_ONLY_ACTIVE_SHADERMAPS
 				, &OffsetToFirstResource
 #endif
diff --git a/Engine/Source/Runtime/Engine/Private/Materials/MaterialShader.cpp b/Engine/Source/Runtime/Engine/Private/Materials/MaterialShader.cpp
--- a/Engine/Source/Runtime/Engine/Private/Materials/MaterialShader.cpp
+++ b/Engine/Source/Runtime/Engine/Private/Materials/MaterialShader.cpp
@@ -3113,14 +3113,14 @@ FMaterialShaderMap* FMaterialShaderMap::GetFinalizedClone() const
 }
 #endif // WITH_EDITOR
 
-bool FMaterialShaderMap::Serialize(FArchive& Ar, bool bInlineShaderResources, bool bLoadedByCookedMaterial, bool bInlineShaderCode)
+bool FMaterialShaderMap::Serialize(FArchive& Ar, bool bInlineShaderResources, bool bLoadedByCookedMaterial, bool bInlineShaderCode, const FName& SerializingAsset)
 {
 	SCOPED_LOADTIMER(FMaterialShaderMap_Serialize);
 	// Note: This is saved to the DDC, not into packages (except when cooked)
 	// Backwards compatibility therefore will not work based on the version of Ar
 	// Instead, just bump MATERIALSHADERMAP_DERIVEDDATA_VER
 	ShaderMapId.Serialize(Ar, bLoadedByCookedMaterial);
-	bool bSerialized = Super::Serialize(Ar, bInlineShaderResources, bLoadedByCookedMaterial, bInlineShaderCode);
+	bool bSerialized = Super::Serialize(Ar, bInlineShaderResources, bLoadedByCookedMaterial, bInlineShaderCode, SerializingAsset);
 #if STATS
 	// This is an unsavory hack to repair STAT_Shaders_NumShadersLoaded not being calculated right in the superclass because the Content class isn't allowed to have virtual functions atm.
 	// A better way is tracked in UE-127112
diff --git a/Engine/Source/Runtime/Engine/Private/Materials/MaterialShared.cpp b/Engine/Source/Runtime/Engine/Private/Materials/MaterialShared.cpp
--- a/Engine/Source/Runtime/Engine/Private/Materials/MaterialShared.cpp
+++ b/Engine/Source/Runtime/Engine/Private/Materials/MaterialShared.cpp
@@ -1533,7 +1533,7 @@ void FMaterial::LegacySerialize(FArchive& Ar)
 	SerializeInlineShaderMap(Ar);
 }
 
-void FMaterial::SerializeInlineShaderMap(FArchive& Ar)
+void FMaterial::SerializeInlineShaderMap(FArchive& Ar, const FName& SerializingAsset)
 {
 	bool bCooked = Ar.IsCooking();
 	Ar << bCooked;
@@ -1588,7 +1588,7 @@ void FMaterial::SerializeInlineShaderMap(FArchive& Ar)
 			if (bValid)
 			{
 				TRefCountPtr<FMaterialShaderMap> LoadedShaderMap = new FMaterialShaderMap();
-				if (LoadedShaderMap->Serialize(Ar, true, bCooked && Ar.IsLoading()))
+				if (LoadedShaderMap->Serialize(Ar, true, bCooked && Ar.IsLoading(), false, SerializingAsset))
 				{
 					GameThreadShaderMap = MoveTemp(LoadedShaderMap);
 					GameThreadShaderMap->GetResource()->SetOwnerName(GetOwnerFName());
diff --git a/Engine/Source/Runtime/Engine/Public/MaterialShared.h b/Engine/Source/Runtime/Engine/Public/MaterialShared.h
--- a/Engine/Source/Runtime/Engine/Public/MaterialShared.h
+++ b/Engine/Source/Runtime/Engine/Public/MaterialShared.h
@@ -1338,7 +1338,7 @@ public:
 	ENGINE_API void Release();
 
 	/** Serializes the shader map. */
-	bool Serialize(FArchive& Ar, bool bInlineShaderResources=true, bool bLoadedByCookedMaterial=false, bool bInlineShaderCode=false);
+	bool Serialize(FArchive& Ar, bool bInlineShaderResources=true, bool bLoadedByCookedMaterial=false, bool bInlineShaderCode=false, const FName& SerializingAsset = NAME_None);
 
 #if WITH_EDITOR
 	/** Saves this shader map to the derived data cache. */
@@ -1857,7 +1857,7 @@ public:
 	ENGINE_API virtual void LegacySerialize(FArchive& Ar);
 
 	/** Serializes the shader map inline in this material, including any shader dependencies. */
-	void SerializeInlineShaderMap(FArchive& Ar);
+	void SerializeInlineShaderMap(FArchive& Ar, const FName& SerializingAsset = NAME_None);
 
 	/** Serializes the shader map inline in this material, including any shader dependencies. */
 	void RegisterInlineShaderMap(bool bLoadedByCookedMaterial);
diff --git a/Engine/Source/Runtime/RenderCore/Private/ShaderMap.cpp b/Engine/Source/Runtime/RenderCore/Private/ShaderMap.cpp
--- a/Engine/Source/Runtime/RenderCore/Private/ShaderMap.cpp
+++ b/Engine/Source/Runtime/RenderCore/Private/ShaderMap.cpp
@@ -123,7 +123,7 @@ void FShaderMapBase::UnfreezeContent()
 
 #define CHECK_SHADERMAP_DEPENDENCIES (WITH_EDITOR || !(UE_BUILD_SHIPPING || UE_BUILD_TEST))
 
-bool FShaderMapBase::Serialize(FArchive& Ar, bool bInlineShaderResources, bool bLoadedByCookedMaterial, bool bInlineShaderCode)
+bool FShaderMapBase::Serialize(FArchive& Ar, bool bInlineShaderResources, bool bLoadedByCookedMaterial, bool bInlineShaderCode, const FName& SerializingAsset)
 {
 	LLM_SCOPE(ELLMTag::Shaders);
 	if (Ar.IsSaving())
@@ -210,7 +210,9 @@ bool FShaderMapBase::Serialize(FArchive& Ar, bool bInlineShaderResources, bool b
 				// also do not warn for shader platforms other than current (if the game targets more than one RHI)
 				if (FApp::CanEverRender() && ShaderPlatform == GMaxRHIShaderPlatform)
 				{
-					UE_LOG(LogShaders, Error, TEXT("Missing shader resource for hash '%s' for shader platform '%s' in the shader library"), *ResourceHash.ToString(), *LexToString(ShaderPlatform));
+					UE_LOG(LogShaders, Error, TEXT("Missing shader resource for hash '%s' for shader platform '%s' in the shader library while serializing asset %s"), *ResourceHash.ToString(),
+						*LexToString(ShaderPlatform),
+						*SerializingAsset.ToString());
 				}
 			}
 		}
diff --git a/Engine/Source/Runtime/RenderCore/Public/Shader.h b/Engine/Source/Runtime/RenderCore/Public/Shader.h
--- a/Engine/Source/Runtime/RenderCore/Public/Shader.h
+++ b/Engine/Source/Runtime/RenderCore/Public/Shader.h
@@ -2274,7 +2274,7 @@ public:
 
 	RENDERCORE_API void FinalizeContent();
 	RENDERCORE_API void UnfreezeContent();
-	RENDERCORE_API bool Serialize(FArchive& Ar, bool bInlineShaderResources, bool bLoadedByCookedMaterial, bool bInlineShaderCode=false);
+	RENDERCORE_API bool Serialize(FArchive& Ar, bool bInlineShaderResources, bool bLoadedByCookedMaterial, bool bInlineShaderCode=false, const FName& SerializingAsset = NAME_None);
 
 	EShaderPermutationFlags GetPermutationFlags() const
 	{
