Networking
* Replicated components added in the ConstructionScript of a blueprint class will now be considered netaddressable and properly mapped on the client instead of instantiated twice.
* Now ensuring when we start replicating a dynamic UserConstructionScript component that is built from an archetype different from the class CDO. The client version of such components were always initialized using the class CDO and probably initialized wrong.

diff --git a/Engine/Source/Runtime/Engine/Private/ActorConstruction.cpp b/Engine/Source/Runtime/Engine/Private/ActorConstruction.cpp
--- a/Engine/Source/Runtime/Engine/Private/ActorConstruction.cpp
+++ b/Engine/Source/Runtime/Engine/Private/ActorConstruction.cpp
@@ -1400,6 +1400,8 @@ void AActor::PostCreateBlueprintComponent(UActorComponent* NewActorComp)
 			int32& Count = ComponentArchetypeCounts.FindOrAdd(NewActorComp->GetArchetype());
 			FSetUCSSerializationIndex::Set(NewActorComp, Count);
 			++Count;
+
+			NewActorComp->SetNetAddressable();
 		}
 
 		// The component may not have been added to ReplicatedComponents if it was duplicated from
diff --git a/Engine/Source/Runtime/Engine/Private/ActorReplication.cpp b/Engine/Source/Runtime/Engine/Private/ActorReplication.cpp
--- a/Engine/Source/Runtime/Engine/Private/ActorReplication.cpp
+++ b/Engine/Source/Runtime/Engine/Private/ActorReplication.cpp
@@ -598,6 +598,12 @@ void AActor::AddComponentForReplication(UActorComponent* Component)
 		return;
 	}
 
+	if (Component->CreationMethod == EComponentCreationMethod::UserConstructionScript && !Component->IsNameStableForNetworking())
+	{
+		ensureMsgf(Component->GetArchetype() == Component->GetClass()->GetDefaultObject(), TEXT("Replicated component %s::%s was added dynamically outside the construction script. This is not well supported and the component on the client will be initialized using the wrong archetype."),
+			*GetName(), *Component->GetName());
+	}
+
 	const ELifetimeCondition NetCondition = AllowActorComponentToReplicate(Component);
 
 	FReplicatedComponentInfo* ComponentInfo = ReplicatedComponentsInfo.FindByKey(Component);
diff --git a/Engine/Source/Runtime/Engine/Private/Components/ActorComponent.cpp b/Engine/Source/Runtime/Engine/Private/Components/ActorComponent.cpp
--- a/Engine/Source/Runtime/Engine/Private/Components/ActorComponent.cpp
+++ b/Engine/Source/Runtime/Engine/Private/Components/ActorComponent.cpp
@@ -2148,7 +2148,7 @@ bool UActorComponent::IsNameStableForNetworking() const
 	 * Components are net addressable if:
 	 *	-They are Default Subobjects (created in C++ constructor)
 	 *	-They were loaded directly from a package (placed in map actors)
-	 *	-They were explicitly set to bNetAddressable (blueprint components created by SCS)
+	 *	-They were explicitly set to bNetAddressable (blueprint components created by SCS or UCS executed in the ConstructionScript only)
 	 */
 
 	return bNetAddressable || (Super::IsNameStableForNetworking() && (CreationMethod != EComponentCreationMethod::UserConstructionScript));
@@ -2288,13 +2288,10 @@ void UActorComponent::RegisterReplicationFragments(UE::Net::FFragmentRegistratio
 {
 	if (CreationMethod == EComponentCreationMethod::UserConstructionScript)
 	{
-		//$IRIS TODO: UCS subobjects without a stable name need to replicate their true archetype. Currently they are instantiated on the client from the class CDO.
-		//            Prevent protocol mismatches by forcing their default state to use the CDO. But the replicated object state is probably wrong if it changed any non-replicated properties.
 		if (!IsNameStableForNetworking() && GetArchetype() != GetClass()->GetDefaultObject())
 		{
 			RegistrationFlags |= UE::Net::EFragmentRegistrationFlags::InitializeDefaultStateFromClassDefaults;
-
-			UE_LOG(LogIris, Warning, TEXT("The default state of UCS created subobject %s::%s was initialized using the wrong archetype. The object's state on client might be initialied wrong."), *GetNameSafe(GetOwner()), *GetName());
+			UE_LOG(LogIris, Warning, TEXT("The default state of replicated dynamic component %s::%s will be built using the class CDO instead of the archetype. The non-replicated properties of the component on clients may be initialized wrong."), *GetNameSafe(GetOwner()), *GetName());
 		}
 	}
 	
