[Navigation]
ISM
- Fixed partial navigation update not using up to date navigation bounds when updating a single instance
HISM
- Fixed missing scale support for filters used in overlapping queries
- Fixed tree marked as out of date when performing in place update (i.e. updating rotation or scale) resulting in accumulated navigation dirty areas until next rebuild. To minimize change on this fading out component we rebuild the tree instead for operations performed in non game worlds (where we expected most of this manipulations to be made).

diff --git a/Engine/Source/Runtime/Engine/Private/HierarchicalInstancedStaticMesh.cpp b/Engine/Source/Runtime/Engine/Private/HierarchicalInstancedStaticMesh.cpp
--- a/Engine/Source/Runtime/Engine/Private/HierarchicalInstancedStaticMesh.cpp
+++ b/Engine/Source/Runtime/Engine/Private/HierarchicalInstancedStaticMesh.cpp
@@ -2286,10 +2286,25 @@ bool UHierarchicalInstancedStaticMeshComponent::UpdateInstanceTransform(int32 In
 	const FTransform NewLocalTransform = bWorldSpace ? NewInstanceTransform.GetRelativeTransform(GetComponentTransform()) : NewInstanceTransform;
 	const FVector NewLocalLocation = NewLocalTransform.GetTranslation();
 
-	// if we are only updating rotation/scale we update the instance directly in the cluster tree
 	const bool bIsOmittedInstance = (RenderIndex == INDEX_NONE);
 	const bool bIsBuiltInstance = !bIsOmittedInstance && RenderIndex < NumBuiltRenderInstances;
-	const bool bDoInPlaceUpdate = bIsBuiltInstance && NewLocalLocation.Equals(OldTransform.GetOrigin());
+
+	bool bAllowInPlaceUpdateForRotationOrScaleChange = true;
+
+	// Code path using 'bDoInPlaceUpdate' indicates that it updates the cluster tree but
+	// bounds are not updated until next tree rebuild and some overlapping queries rely on those bounds.
+	// We want to make sure a manipulation in an Editor world will fully update the information so queries
+	// will return proper information to callers (e.g. navigation rebuild and preview)
+#if WITH_EDITOR
+	if (const UWorld* World = GetWorld())
+	{
+		const bool bIsGameWorld = World->IsGameWorld();
+		bAllowInPlaceUpdateForRotationOrScaleChange = bIsGameWorld;
+	}
+#endif // WITH_EDITOR
+
+	// if we are only updating rotation/scale then we update the instance directly in the cluster tree
+	const bool bDoInPlaceUpdate = bAllowInPlaceUpdateForRotationOrScaleChange && bIsBuiltInstance && NewLocalLocation.Equals(OldTransform.GetOrigin());
 
 	bool Result = Super::UpdateInstanceTransform(InstanceIndex, NewInstanceTransform, bWorldSpace, bMarkRenderStateDirty, bTeleport);
 	
@@ -3142,7 +3157,7 @@ int32 UHierarchicalInstancedStaticMeshComponent::GetOverlappingSphereCount(const
 	for (const FTransform& TM : Transforms)
 	{
 		const FVector Center = TM.GetLocation();
-		const FSphere InstanceSphere(Center, MeshBounds.SphereRadius);
+		const FSphere InstanceSphere(Center, MeshBounds.SphereRadius * TM.GetScale3D().GetMax());
 		
 		if (Sphere.Intersects(InstanceSphere))
 		{
@@ -3162,9 +3177,7 @@ int32 UHierarchicalInstancedStaticMeshComponent::GetOverlappingBoxCount(const FB
 	const FBoxSphereBounds MeshBounds = GetStaticMesh()->GetBounds();
 	for(FTransform& T : Transforms)
 	{
-		const FVector Center = T.GetLocation();
-		const FBox OtherBox(FVector(Center - MeshBounds.BoxExtent), FVector(Center + MeshBounds.BoxExtent));
-
+		const FBox OtherBox(MeshBounds.TransformBy(T).GetBox());
 		if(Box.Intersect(OtherBox))
 		{
 			Count++;
@@ -3183,8 +3196,7 @@ void UHierarchicalInstancedStaticMeshComponent::GetOverlappingBoxTransforms(cons
 	const FBoxSphereBounds MeshBounds = GetStaticMesh()->GetBounds();
 	OutTransforms.RemoveAllSwap([&MeshBounds, &Box](const FTransform& Transform) -> bool
 	{
-		const FVector Center = Transform.GetLocation();
-		const FBox OtherBox(FVector(Center - MeshBounds.BoxExtent), FVector(Center + MeshBounds.BoxExtent));
+		const FBox OtherBox(MeshBounds.TransformBy(Transform).GetBox());
 		return !Box.Intersect(OtherBox); 
 	});
 }
diff --git a/Engine/Source/Runtime/Engine/Private/InstancedStaticMesh.cpp b/Engine/Source/Runtime/Engine/Private/InstancedStaticMesh.cpp
--- a/Engine/Source/Runtime/Engine/Private/InstancedStaticMesh.cpp
+++ b/Engine/Source/Runtime/Engine/Private/InstancedStaticMesh.cpp
@@ -3724,19 +3724,11 @@ bool UInstancedStaticMeshComponent::UpdateInstanceTransform(const int32 Instance
 
 	FInstancedStaticMeshInstanceData& InstanceData = PerInstanceSMData[InstanceIndex];
 
+	const FTransform PreviousTransform(InstanceData.Transform);
+
 	// TODO: Computing LocalTransform is useless when we're updating the world location for the entire mesh.
 	// Should find some way around this for performance.
 
-	if (IsNavigationRelevant() && SupportsPartialNavigationUpdate())
-	{
-		// Append instance's previous and new transforms to dirty both areas
-		PartialNavigationUpdates(
-			{
-				FTransform(InstanceData.Transform) * GetComponentTransform(),
-				(bWorldSpace ? NewInstanceTransform : NewInstanceTransform * GetComponentTransform())
-			});
-	}
-
 	// Render data uses local transform of the instance
 	const FTransform LocalTransform = bWorldSpace ? NewInstanceTransform.GetRelativeTransform(GetComponentTransform()) : NewInstanceTransform;
 	InstanceData.Transform = LocalTransform.ToMatrixWithScale();
@@ -3749,6 +3741,17 @@ bool UInstancedStaticMeshComponent::UpdateInstanceTransform(const int32 Instance
 		UpdateInstanceBodyTransform(InstanceIndex, WorldTransform, bTeleport);
 	}
 
+	if (IsNavigationRelevant() && SupportsPartialNavigationUpdate())
+	{
+		// Perform partial update after instance gets updated since we need NavigationBounds using up to date instances.
+		// Append instance's previous and new transforms to dirty both areas
+		PartialNavigationUpdates(
+			{
+				PreviousTransform * GetComponentTransform(),
+				(bWorldSpace ? NewInstanceTransform : NewInstanceTransform * GetComponentTransform())
+			});
+	}
+
 	if (bMarkRenderStateDirty)
 	{
 		MarkRenderStateDirty();
