[BodyInstance] Added option 'p.SkipShapeCreationForEmptyBodySetup', disabled by default, to skip actors and shapes creation for all instances if the body setup doesn't have any geometry.

diff --git a/Engine/Source/Runtime/Engine/Private/PhysicsEngine/BodyInstance.cpp b/Engine/Source/Runtime/Engine/Private/PhysicsEngine/BodyInstance.cpp
--- a/Engine/Source/Runtime/Engine/Private/PhysicsEngine/BodyInstance.cpp
+++ b/Engine/Source/Runtime/Engine/Private/PhysicsEngine/BodyInstance.cpp
@@ -90,6 +90,12 @@ FAutoConsoleVariableRef CVarbEnableOverrideSolverDeltaTime(
 	bEnableOverrideSolverDeltaTime,
 	TEXT("If true, setting for override solver delta time can be used.  False will disable this feature."));
 
+bool bSkipShapeCreationForEmptyBodySetup = false;
+FAutoConsoleVariableRef CVarSkipShapeCreationForEmptyBodySetup(
+	TEXT("p.SkipShapeCreationForEmptyBodySetup"),
+	bSkipShapeCreationForEmptyBodySetup,
+	TEXT("If true, CreateShapesAndActors will not try to create actors and shapes for all instances if the body setup doesn't have any geometry."));
+
 using namespace PhysicsInterfaceTypes;
 
 bool IsRigidBodyKinematic_AssumesLocked(const FPhysicsActorHandle& InActorRef)
@@ -1189,7 +1195,6 @@ void FInitBodiesHelperBase::CreateActor_AssumesLocked(FBodyInstance* Instance, c
 	SCOPE_CYCLE_COUNTER(STAT_CreatePhysicsActor);
 	checkSlow(!FPhysicsInterface::IsValid(Instance->ActorHandle));
 	const ECollisionEnabled::Type CollisionType = Instance->GetCollisionEnabled();
-	const bool bDisableSim = !CollisionEnabledHasPhysics(CollisionType) && DisableQueryOnlyActors;
 
 	FActorCreationParams ActorParams;
 	ActorParams.InitialTM = Transform;
@@ -1328,6 +1333,24 @@ bool FInitBodiesHelperBase::CreateShapesAndActors()
 	// Ensure we have the AggGeom inside the body setup so we can calculate the number of shapes
 	BodySetup->CreatePhysicsMeshes();
 
+	if (bSkipShapeCreationForEmptyBodySetup)
+	{
+		if (BodySetup->TriMeshGeometries.IsEmpty() && BodySetup->AggGeom.GetElementCount() == 0)
+		{
+#if WITH_EDITOR
+			// In the editor we may have ended up here because of world trace ignoring our EnableCollision.
+			// Since we can't get at the data in that function we check for it here
+			if (PrimitiveComp && PrimitiveComp->IsCollisionEnabled())
+#endif
+			{
+				UE_LOG(LogPhysics, Log, TEXT("Init of %d instances of Primitive Component %s failed. Does it have collision data available?"),
+					NumBodies, PrimitiveComp ? *PrimitiveComp->GetReadableName() : TEXT("null"));
+			}
+
+			return false;
+		}
+	}
+
 	for (int32 BodyIdx = NumBodies - 1; BodyIdx >= 0; BodyIdx--)   // iterate in reverse since list might shrink
 	{
 		FBodyInstance* Instance = Bodies[BodyIdx];
