SoundWave Resampling works in editor. CuePoints reflect Resampling of sound wave in editor & in package builds.

- Rescale CuePoints based on CurrentPlatformSampleRate at editor time, during Serialization & during Initialization of SoundWaveDataPtr
- Update ImportedSampleRate on PostLoad of the SoundWave so that the ImportedSampleRate is correctly updated and a useable value without needing to reimport.

diff --git a/Engine/Source/Runtime/Engine/Classes/Sound/SoundWave.h b/Engine/Source/Runtime/Engine/Classes/Sound/SoundWave.h
--- a/Engine/Source/Runtime/Engine/Classes/Sound/SoundWave.h
+++ b/Engine/Source/Runtime/Engine/Classes/Sound/SoundWave.h
@@ -1384,7 +1384,27 @@ public:
 	
 	mutable TMap<FName, ISoundWaveLoadingBehaviorUtil::FClassData> OwnerLoadingBehaviorCache;
 	mutable FCriticalSection OwnerLoadingBehaviorCacheCS;
-			
+
+	/*
+	* Returns this SoundWave's CuePoints array with the frame values scaled by
+	* InSampleRate / ImportedSampleRate to account for resampling of the sound wave source data.
+	* If no resampling is necessary, returns the CuePoints as-is.
+	*
+	* @param InSampleRate	The sample rate the SoundWave 
+	* @return CuePoints array scaled if resampling occurred
+	*/
+	ENGINE_API TArray<FSoundWaveCuePoint> GetCuePointsScaledForSampleRate(const float InSampleRate) const;
+
+	/*
+	* Modifies the InOutCuePoints array with the frame values scaled by
+	* InSampleRate / ImportedSampleRate to account for resampling of the sound wave source data.
+	* Does not modify InOutCuePoints if no resampling is necessary
+	*
+	* @param InSampleRate	The sample rate the SoundWave
+	* @param InOutCuePoints	The CuePoints array to re-scale
+	*/
+	ENGINE_API void ScaleCuePointsForSampleRate(const float InSampleRate, TArray<FSoundWaveCuePoint>& InOutCuePoints) const;
+
 	/*
 	* Returns a sample rate if there is a specific sample rate override for this platform, -1.0 otherwise.
 	*/
diff --git a/Engine/Source/Runtime/Engine/Private/AudioCompressionSettingsUtils.cpp b/Engine/Source/Runtime/Engine/Private/AudioCompressionSettingsUtils.cpp
--- a/Engine/Source/Runtime/Engine/Private/AudioCompressionSettingsUtils.cpp
+++ b/Engine/Source/Runtime/Engine/Private/AudioCompressionSettingsUtils.cpp
@@ -121,67 +121,61 @@ void CacheAudioCookOverrides(FPlatformAudioCookOverrides& OutOverrides, const TC
 	// Cache size:
 	const int32 DefaultCacheSize = 64 * 1024;
 	int32 RetrievedCacheSize = DefaultCacheSize;
-	int32 RetrievedChunkSizeOverride = INDEX_NONE;
-
-	PlatformFile->GetInt(*CategoryName, TEXT("CacheSizeKB"), RetrievedCacheSize);
-	if (!RetrievedCacheSize)
+	if (!PlatformFile->GetInt(*CategoryName, TEXT("CacheSizeKB"), RetrievedCacheSize))
 	{
 		RetrievedCacheSize = DefaultCacheSize;
 	}
 
 	OutOverrides.StreamCachingSettings.CacheSizeKB = RetrievedCacheSize;
 
-	PlatformFile->GetInt(*CategoryName, TEXT("MaxChunkSizeOverrideKB"), RetrievedChunkSizeOverride);
-	OutOverrides.StreamCachingSettings.MaxChunkSizeOverrideKB = RetrievedChunkSizeOverride;
+	int32 RetrievedChunkSizeOverride = INDEX_NONE;
+	if (PlatformFile->GetInt(*CategoryName, TEXT("MaxChunkSizeOverrideKB"), RetrievedChunkSizeOverride))
+	{
+		OutOverrides.StreamCachingSettings.MaxChunkSizeOverrideKB = RetrievedChunkSizeOverride;
+	}
 
 	bool bForceLegacyStreamChunking = false;
-	PlatformFile->GetBool(*CategoryName, TEXT("bForceLegacyStreamChunking"), bForceLegacyStreamChunking);
-	OutOverrides.StreamCachingSettings.bForceLegacyStreamChunking = bForceLegacyStreamChunking;
+	if (PlatformFile->GetBool(*CategoryName, TEXT("bForceLegacyStreamChunking"), bForceLegacyStreamChunking))
+	{
+		OutOverrides.StreamCachingSettings.bForceLegacyStreamChunking = bForceLegacyStreamChunking;
+	}
 
 	int32 ZerothChunkSizeForLegacyStreamChunking = 0;
-	PlatformFile->GetInt(*CategoryName, TEXT("ZerothChunkSizeForLegacyStreamChunking"), ZerothChunkSizeForLegacyStreamChunking);
-	OutOverrides.StreamCachingSettings.ZerothChunkSizeForLegacyStreamChunkingKB = ZerothChunkSizeForLegacyStreamChunking;
+	if (PlatformFile->GetInt(*CategoryName, TEXT("ZerothChunkSizeForLegacyStreamChunking"), ZerothChunkSizeForLegacyStreamChunking))
+	{
+		OutOverrides.StreamCachingSettings.ZerothChunkSizeForLegacyStreamChunkingKB = ZerothChunkSizeForLegacyStreamChunking;
+	}
 
-	PlatformFile->GetBool(*CategoryName, TEXT("bResampleForDevice"), OutOverrides.bResampleForDevice);
+	bool bResampleForDevice = false;
+	if (PlatformFile->GetBool(*CategoryName, TEXT("bResampleForDevice"), bResampleForDevice))
+	{
+		OutOverrides.bResampleForDevice = bResampleForDevice;
+	}
 
-	PlatformFile->GetFloat(*CategoryName, TEXT("CompressionQualityModifier"), OutOverrides.CompressionQualityModifier);
+	float CompressionQualityModifier = 0.0f;
+	if (PlatformFile->GetFloat(*CategoryName, TEXT("CompressionQualityModifier"), CompressionQualityModifier))
+	{
+		OutOverrides.CompressionQualityModifier = CompressionQualityModifier;
+	}
 
-	PlatformFile->GetFloat(*CategoryName, TEXT("AutoStreamingThreshold"), OutOverrides.AutoStreamingThreshold);
+	float AutoStreamingThreshold = 0.0f;
+	if (PlatformFile->GetFloat(*CategoryName, TEXT("AutoStreamingThreshold"), AutoStreamingThreshold))
+	{
+		OutOverrides.AutoStreamingThreshold = AutoStreamingThreshold;
+	}
 
 #if 1
 	//Cache sample rate map:
 	float RetrievedSampleRate = -1.0f;
 
-	PlatformFile->GetFloat(*CategoryName, TEXT("MaxSampleRate"), RetrievedSampleRate);
-	float* FoundSampleRate = OutOverrides.PlatformSampleRates.Find(ESoundwaveSampleRateSettings::Max);
-
-	if (FoundSampleRate)
-	{
-		if (!FMath::IsNearlyEqual(*FoundSampleRate, RetrievedSampleRate))
-		{
-			*FoundSampleRate = RetrievedSampleRate;
-		}
-
-	}
-	else
+	if (PlatformFile->GetFloat(*CategoryName, TEXT("MaxSampleRate"), RetrievedSampleRate))
 	{
 		OutOverrides.PlatformSampleRates.Add(ESoundwaveSampleRateSettings::Max, RetrievedSampleRate);
 	}
 
 	RetrievedSampleRate = -1.0f;
 
-	PlatformFile->GetFloat(*CategoryName, TEXT("HighSampleRate"), RetrievedSampleRate);
-	FoundSampleRate = OutOverrides.PlatformSampleRates.Find(ESoundwaveSampleRateSettings::High);
-
-	if (FoundSampleRate)
-	{
-		if (!FMath::IsNearlyEqual(*FoundSampleRate, RetrievedSampleRate))
-		{
-			*FoundSampleRate = RetrievedSampleRate;
-		}
-
-	}
-	else
+	if (PlatformFile->GetFloat(*CategoryName, TEXT("HighSampleRate"), RetrievedSampleRate))
 	{
 		OutOverrides.PlatformSampleRates.Add(ESoundwaveSampleRateSettings::High, RetrievedSampleRate);
 	}
@@ -189,55 +183,24 @@ void CacheAudioCookOverrides(FPlatformAudioCookOverrides& OutOverrides, const TC
 
 	RetrievedSampleRate = -1.0f;
 
-	PlatformFile->GetFloat(*CategoryName, TEXT("MedSampleRate"), RetrievedSampleRate);
-	FoundSampleRate = OutOverrides.PlatformSampleRates.Find(ESoundwaveSampleRateSettings::Medium);
-
-	if (FoundSampleRate)
-	{
-		if (!FMath::IsNearlyEqual(*FoundSampleRate, RetrievedSampleRate))
-		{
-			*FoundSampleRate = RetrievedSampleRate;
-		}
-	}
-	else
+	if (PlatformFile->GetFloat(*CategoryName, TEXT("MedSampleRate"), RetrievedSampleRate))
 	{
 		OutOverrides.PlatformSampleRates.Add(ESoundwaveSampleRateSettings::Medium, RetrievedSampleRate);
 	}
 
 	RetrievedSampleRate = -1.0f;
 
-	PlatformFile->GetFloat(*CategoryName, TEXT("LowSampleRate"), RetrievedSampleRate);
-	FoundSampleRate = OutOverrides.PlatformSampleRates.Find(ESoundwaveSampleRateSettings::Low);
-
-	if (FoundSampleRate)
-	{
-		if (!FMath::IsNearlyEqual(*FoundSampleRate, RetrievedSampleRate))
-		{
-			*FoundSampleRate = RetrievedSampleRate;
-		}
-	}
-	else
+	if (PlatformFile->GetFloat(*CategoryName, TEXT("LowSampleRate"), RetrievedSampleRate))
 	{
 		OutOverrides.PlatformSampleRates.Add(ESoundwaveSampleRateSettings::Low, RetrievedSampleRate);
 	}
 
 	RetrievedSampleRate = -1.0f;
 
-	PlatformFile->GetFloat(*CategoryName, TEXT("MinSampleRate"), RetrievedSampleRate);
-	FoundSampleRate = OutOverrides.PlatformSampleRates.Find(ESoundwaveSampleRateSettings::Min);
-
-	if (FoundSampleRate)
-	{
-		if (!FMath::IsNearlyEqual(*FoundSampleRate, RetrievedSampleRate))
-		{
-			*FoundSampleRate = RetrievedSampleRate;
-		}
-	}
-	else
+	if (PlatformFile->GetFloat(*CategoryName, TEXT("MinSampleRate"), RetrievedSampleRate))
 	{
 		OutOverrides.PlatformSampleRates.Add(ESoundwaveSampleRateSettings::Min, RetrievedSampleRate);
 	}
-
 #else
 
 	//Cache sample rate map.
diff --git a/Engine/Source/Runtime/Engine/Private/SoundWave.cpp b/Engine/Source/Runtime/Engine/Private/SoundWave.cpp
--- a/Engine/Source/Runtime/Engine/Private/SoundWave.cpp
+++ b/Engine/Source/Runtime/Engine/Private/SoundWave.cpp
@@ -201,6 +201,12 @@ void FSoundWaveData::InitializeDataFromSoundWave(USoundWave& InWave)
 
 #if WITH_EDITOR
 	bLoadedFromCookedData = InWave.IsLoadedFromCookedData();
+
+	// only necessary to set cue points here in editor 
+	// when the sample rate can change due to platform settings changing
+	// or during first import of the asset.
+	// otherwise, CuePoints are always updated during serialization
+	SetAllCuePoints(InWave.GetCuePointsScaledForSampleRate(SampleRate));
 #endif //WITH_EDITOR
 }
 
@@ -1917,6 +1923,21 @@ void USoundWave::PostLoad()
 	{
 		UE_LOG(LogAudio, Warning, TEXT("Sound Wave '%s' has defined an effect chain but is not mono or stereo."), *GetName());
 	}
+
+
+	if (ImportedSampleRate == 0)
+	{
+		// update ImportedSampleRate ASAP so that it reflects the 
+		// will only be valid if this sound wave was from an imported .wav file. 
+		// otherwise we can safely ignore it
+		TArray<uint8> ImportedRawPCMData;
+		uint32 OriginalSampleRate = 0;
+		uint16 ImportedNumChannels = 0;
+		if (GetImportedSoundWaveData(ImportedRawPCMData, OriginalSampleRate, ImportedNumChannels))
+		{
+			ImportedSampleRate = OriginalSampleRate;
+		}
+	}
 #endif
 
 	// Don't need to do anything in post load if this is a source bus or procedural audio
@@ -2315,6 +2336,45 @@ ISoundWaveLoadingBehaviorUtil::FClassData USoundWave::GetOwnerLoadingBehavior(co
 	return {};
 }
 
+TArray<FSoundWaveCuePoint> USoundWave::GetCuePointsScaledForSampleRate(const float InSampleRate) const
+{
+	TArray<FSoundWaveCuePoint> CuePointsScaled = CuePoints;
+	ScaleCuePointsForSampleRate(InSampleRate, CuePointsScaled);
+	return CuePointsScaled;
+}
+
+void USoundWave::ScaleCuePointsForSampleRate(const float InSampleRate, TArray<FSoundWaveCuePoint>& InOutCuePoints) const
+{
+	if (InOutCuePoints.IsEmpty() || InSampleRate <= 0.0f)
+	{
+		return;
+	}
+
+	// if ImportedSampleRate is 0, then it probably wasn't set properly yet in editor, 
+	// this will be okay since SoundWave will update ImportedSampleRate on PostLoad
+	// and re-initialize its internals so it will be valid.
+	// ensuring or throwing an error/warning here isn't necessary
+	if (ImportedSampleRate <= 0)
+	{
+		return;
+	}
+
+	float ResampleRatio = -1.0f;
+	if (!FMath::IsNearlyEqual(InSampleRate, ImportedSampleRate))
+	{
+		ResampleRatio = InSampleRate / (float)ImportedSampleRate;
+	}
+
+	if (ResampleRatio > 0.0f)
+	{
+		for (FSoundWaveCuePoint& CuePoint : InOutCuePoints)
+		{
+			CuePoint.ScaleFrameValues(ResampleRatio);
+		}
+	}
+}
+
+
 float USoundWave::GetSampleRateForTargetPlatform(const ITargetPlatform* TargetPlatform)
 {
 	const FPlatformAudioCookOverrides* Overrides = FPlatformCompressionUtilities::GetCookOverrides(*TargetPlatform->IniPlatformName());
@@ -4404,11 +4464,11 @@ void USoundWave::CacheInheritedLoadingBehavior() const
 
 void USoundWave::SerializeCuePoints(FArchive& Ar, const bool bIsLoadingFromCookedArchive)
 {
-	TArray<FSoundWaveCuePoint> CuePointsCopy;
+	TArray<FSoundWaveCuePoint> PlatformCuePoints;
 #if WITH_EDITORONLY_DATA
 	if (Ar.IsCooking())
 	{
-		CuePointsCopy = CuePoints;
+		PlatformCuePoints = CuePoints;
 		float ResampleRatio = -1.0f;
 		const ITargetPlatform* CookingTarget = Ar.CookingTarget();
 		if (ensure(CookingTarget))
@@ -4417,45 +4477,29 @@ void USoundWave::SerializeCuePoints(FArchive& Ar, const bool bIsLoadingFromCooke
 			{
 				if (Overrides->bResampleForDevice)
 				{
+					// at this point, the ImportedSampleRate should be ready to use!
+					ensureMsgf(ImportedSampleRate > 0, TEXT("SerializeCuePoints: %s ImportedSampleRate not set: ImportedSampleRate = %d"), *GetName(), ImportedSampleRate);
 					const float TargetSampleRate = GetSampleRateForTargetPlatform(CookingTarget);
-
-					// annoyingly, read the entire imported sound wave data to successfully get the ImportedSampleRate at this point
-					TArray<uint8> RawWaveData;
-					uint32 OriginalSampleRate = 0;
-					uint16 OriginalNumChannels = 0;
-					ensureMsgf(GetImportedSoundWaveData(RawWaveData, OriginalSampleRate, OriginalNumChannels), TEXT("SerializeCuePoints: %s Failed to retrieve imported sound wave data: OriginalSampleRate = %d"), *GetName(), OriginalSampleRate);
-					
-					if (OriginalSampleRate > 0 && !FMath::IsNearlyEqual(TargetSampleRate, OriginalSampleRate))
-					{
-						ResampleRatio = TargetSampleRate / (float)OriginalSampleRate;
-					}
-
-					if (ResampleRatio > 0.0f)
-					{
-						for (FSoundWaveCuePoint& CuePoint : CuePointsCopy)
-						{
-							CuePoint.ScaleFrameValues(ResampleRatio);
-						}
-					}
+					ScaleCuePointsForSampleRate(TargetSampleRate, PlatformCuePoints);
 				}
 			}
 		}
 	}
 	else if (Ar.IsLoading() && !bIsLoadingFromCookedArchive)
 	{
-		SoundWaveDataPtr->SetAllCuePoints(CuePoints);
+		SoundWaveDataPtr->SetAllCuePoints(GetCuePointsScaledForSampleRate(GetSampleRateForCurrentPlatform()));
 		return;
 	}
 #endif
 
 	if (Ar.IsCooking() || bIsLoadingFromCookedArchive)
 	{
-		Ar << CuePointsCopy;
+		Ar << PlatformCuePoints;
 	}
 
 	if (bIsLoadingFromCookedArchive)
 	{
-		SoundWaveDataPtr->SetAllCuePoints(CuePointsCopy);
+		SoundWaveDataPtr->SetAllCuePoints(PlatformCuePoints);
 	}
 }
 
