check SoundWave RawData before getting ImportedSampleRate

diff --git a/Engine/Source/Runtime/Engine/Private/SoundWave.cpp b/Engine/Source/Runtime/Engine/Private/SoundWave.cpp
--- a/Engine/Source/Runtime/Engine/Private/SoundWave.cpp
+++ b/Engine/Source/Runtime/Engine/Private/SoundWave.cpp
@@ -1925,7 +1925,7 @@ void USoundWave::PostLoad()
 	}
 
 
-	if (ImportedSampleRate == 0)
+	if (RawData.HasPayloadData() && ImportedSampleRate == 0)
 	{
 		// update ImportedSampleRate ASAP so that it reflects the 
 		// will only be valid if this sound wave was from an imported .wav file. 
