Iris - Don't add dependent actor unless the PlayerController is replicated.

#rnx

diff --git a/Engine/Source/Runtime/Engine/Private/PlayerController.cpp b/Engine/Source/Runtime/Engine/Private/PlayerController.cpp
--- a/Engine/Source/Runtime/Engine/Private/PlayerController.cpp
+++ b/Engine/Source/Runtime/Engine/Private/PlayerController.cpp
@@ -4812,7 +4812,7 @@ void APlayerController::ReceivedSpectatorClass(TSubclassOf<ASpectatorPawn> Spect
 void APlayerController::SetPawn(APawn* InPawn)
 {
 #if UE_WITH_IRIS
-	if (GetLocalRole() == ROLE_Authority)
+	if (GetLocalRole() == ROLE_Authority && UE::Net::FReplicationSystemUtil::GetNetHandle(this).IsValid())
 	{
 		if (const UReplicationSystem* ReplicationSystem = UE::Net::FReplicationSystemUtil::GetReplicationSystem(this))
 		{
