World On Begin Play Event

Deprecating bBegunPlay flag public access inside UWorld in favour of public accessors. This allows me to raise an event when bBegunPlay is set, which I can then use later to trigger certain events instead of constantly polling for changes every frame.

diff --git a/Engine/Plugins/Experimental/Animation/PoseSearch/Source/Editor/Private/PoseSearchDatabasePreviewScene.cpp b/Engine/Plugins/Experimental/Animation/PoseSearch/Source/Editor/Private/PoseSearchDatabasePreviewScene.cpp
--- a/Engine/Plugins/Experimental/Animation/PoseSearch/Source/Editor/Private/PoseSearchDatabasePreviewScene.cpp
+++ b/Engine/Plugins/Experimental/Animation/PoseSearch/Source/Editor/Private/PoseSearchDatabasePreviewScene.cpp
@@ -45,14 +45,14 @@ namespace UE::PoseSearch
 		// Trigger Begin Play in this preview world.
 		// This is needed for the CharacterMovementComponent to be able to switch to falling mode. 
 		// See: UCharacterMovementComponent::StartFalling
-		if (PreviewWorld && !PreviewWorld->bBegunPlay)
+		if (PreviewWorld && !PreviewWorld->GetBegunPlay())
 		{
 			for (FActorIterator It(PreviewWorld); It; ++It)
 			{
 				It->DispatchBeginPlay();
 			}
 
-			PreviewWorld->bBegunPlay = true;
+			PreviewWorld->SetBegunPlay(true);
 		}
 
 		GetWorld()->Tick(LEVELTICK_All, InDeltaTime);
diff --git a/Engine/Plugins/FX/Niagara/Source/NiagaraEditor/Private/Widgets/SNiagaraSystemViewport.cpp b/Engine/Plugins/FX/Niagara/Source/NiagaraEditor/Private/Widgets/SNiagaraSystemViewport.cpp
--- a/Engine/Plugins/FX/Niagara/Source/NiagaraEditor/Private/Widgets/SNiagaraSystemViewport.cpp
+++ b/Engine/Plugins/FX/Niagara/Source/NiagaraEditor/Private/Widgets/SNiagaraSystemViewport.cpp
@@ -1181,7 +1181,7 @@ void FNiagaraBaselineViewportClient::Tick(float DeltaSeconds)
 
 	if (UWorld* World = PreviewScene->GetWorld())
 	{
-		if (!World->bBegunPlay)
+		if (!World->GetBegunPlay())
 		{
 			for (FActorIterator It(World); It; ++It)
 			{
@@ -1190,7 +1190,7 @@ void FNiagaraBaselineViewportClient::Tick(float DeltaSeconds)
 					It->DispatchBeginPlay();
 				}
 			}
-			World->bBegunPlay = true;
+			World->SetBegunPlay(true);
 
 			// Simulate behavior from GameEngine.cpp
 			World->bWorldWasLoadedThisTick = false;
diff --git a/Engine/Plugins/NetcodeUnitTest/NetcodeUnitTest/Source/NetcodeUnitTest/Private/Net/NUTUtilNet.cpp b/Engine/Plugins/NetcodeUnitTest/NetcodeUnitTest/Source/NetcodeUnitTest/Private/Net/NUTUtilNet.cpp
--- a/Engine/Plugins/NetcodeUnitTest/NetcodeUnitTest/Source/NetcodeUnitTest/Private/Net/NUTUtilNet.cpp
+++ b/Engine/Plugins/NetcodeUnitTest/NetcodeUnitTest/Source/NetcodeUnitTest/Private/Net/NUTUtilNet.cpp
@@ -505,7 +505,7 @@ UWorld* NUTNet::CreateUnitTestWorld(bool bHookTick/*=true*/)
 
 
 		// Hack-mark the world as having begun play (when it has not)
-		ReturnVal->bBegunPlay = true;
+		ReturnVal->SetBegunPlay(true);
 
 		// Hack-mark the world as having initialized actors (to allow RPC hooks)
 		ReturnVal->bActorsInitialized = true;
diff --git a/Engine/Source/Editor/Kismet/Private/SCSEditorViewportClient.cpp b/Engine/Source/Editor/Kismet/Private/SCSEditorViewportClient.cpp
--- a/Engine/Source/Editor/Kismet/Private/SCSEditorViewportClient.cpp
+++ b/Engine/Source/Editor/Kismet/Private/SCSEditorViewportClient.cpp
@@ -202,7 +202,7 @@ FSCSEditorViewportClient::FSCSEditorViewportClient(TWeakPtr<FBlueprintEditor>& I
 	PreviewScene->AddComponent(EditorFloorComp, FTransform::Identity);
 
 	// Turn off so that actors added to the world do not have a lifespan (so they will not auto-destroy themselves).
-	PreviewScene->GetWorld()->bBegunPlay = false;
+	PreviewScene->GetWorld()->SetBegunPlay(false);
 
 	PreviewScene->SetSkyCubemap(GUnrealEd->GetThumbnailManager()->AmbientCubemap);
 }
@@ -1017,7 +1017,7 @@ void FSCSEditorViewportClient::ToggleIsSimulateEnabled()
 	BlueprintEditorPtr.Pin()->DestroyPreview();
 
 	bIsSimulateEnabled = !bIsSimulateEnabled;
-	PreviewScene->GetWorld()->bBegunPlay = bIsSimulateEnabled;
+	PreviewScene->GetWorld()->SetBegunPlay(bIsSimulateEnabled);
 	PreviewScene->GetWorld()->bShouldSimulatePhysics = bIsSimulateEnabled;
 
 	TSharedPtr<SWidget> SubobjectEditor = BlueprintEditorPtr.Pin()->GetSubobjectEditor();
diff --git a/Engine/Source/Editor/UMGEditor/Private/WidgetBlueprintEditor.cpp b/Engine/Source/Editor/UMGEditor/Private/WidgetBlueprintEditor.cpp
--- a/Engine/Source/Editor/UMGEditor/Private/WidgetBlueprintEditor.cpp
+++ b/Engine/Source/Editor/UMGEditor/Private/WidgetBlueprintEditor.cpp
@@ -102,7 +102,7 @@ FWidgetBlueprintEditor::FWidgetBlueprintEditor()
 	, bUpdatingSequencerSelection(false)
 	, bUpdatingExternalSelection(false)
 {
-	PreviewScene.GetWorld()->bBegunPlay = false;
+	PreviewScene.GetWorld()->SetBegunPlay(false);
 
 	// Register sequencer menu extenders.
 	ISequencerModule& SequencerModule = FModuleManager::Get().LoadModuleChecked<ISequencerModule>( "Sequencer" );
diff --git a/Engine/Source/Runtime/Engine/Classes/Engine/World.h b/Engine/Source/Runtime/Engine/Classes/Engine/World.h
--- a/Engine/Source/Runtime/Engine/Classes/Engine/World.h
+++ b/Engine/Source/Runtime/Engine/Classes/Engine/World.h
@@ -1251,8 +1251,19 @@ public:
 	uint8 bActorsInitialized:1;
 
 	/** Whether BeginPlay has been called on actors */
+	UE_DEPRECATED(5.4, "Public access to bBegunPlay is deprecated. Please update your code to use the public accessors GetBegunPlay() & SetBegunPlay().")
 	uint8 bBegunPlay:1;
 
+	/** Set whether BeginPlay has been called on actors */
+	void SetBegunPlay(bool bHasBegunPlay);
+
+	/** Get whether BeginPlay has been called on actors */
+	bool GetBegunPlay() const;
+
+	DECLARE_EVENT_OneParam(UWorld, FOnBeginPlay, bool);
+	/** Return the event that is broadcast when bBegunPlay is changed */
+	FOnBeginPlay& GetOnBeginPlayEvent() { return OnBeginPlay; }
+
 	/** Whether the match has been started */
 	uint8 bMatchStarted:1;
 
@@ -1362,8 +1373,8 @@ private:
 
 	/** Creates the dynamic source and static level collections if they don't already exist. */
 	void ConditionallyCreateDefaultLevelCollections();
-
-
+	
+	FOnBeginPlay OnBeginPlay;
 public:
 
 	/** Handle to the active audio device for this world. */
diff --git a/Engine/Source/Runtime/Engine/Private/GameStateBase.cpp b/Engine/Source/Runtime/Engine/Private/GameStateBase.cpp
--- a/Engine/Source/Runtime/Engine/Private/GameStateBase.cpp
+++ b/Engine/Source/Runtime/Engine/Private/GameStateBase.cpp
@@ -233,7 +233,7 @@ bool AGameStateBase::HasBegunPlay() const
 	UWorld* World = GetWorld();
 	if (World)
 	{
-		return World->bBegunPlay;
+		return World->GetBegunPlay();
 	}
 
 	return false;
diff --git a/Engine/Source/Runtime/Engine/Private/World.cpp b/Engine/Source/Runtime/Engine/Private/World.cpp
--- a/Engine/Source/Runtime/Engine/Private/World.cpp
+++ b/Engine/Source/Runtime/Engine/Private/World.cpp
@@ -4253,6 +4253,28 @@ bool UWorld::CanAddLoadedLevelToWorld(ULevel* Level) const
 	return true;
 }
 
+PRAGMA_DISABLE_DEPRECATION_WARNINGS
+void UWorld::SetBegunPlay(bool bHasBegunPlay)
+{
+	if(bBegunPlay == bHasBegunPlay)
+	{
+		return;
+	}
+	
+	bBegunPlay = bHasBegunPlay;
+	if(OnBeginPlay.IsBound())
+	{
+		OnBeginPlay.Broadcast(bBegunPlay);
+	}
+}
+
+bool UWorld::GetBegunPlay() const
+{
+	return bBegunPlay;
+}
+PRAGMA_ENABLE_DEPRECATION_WARNINGS
+
+
 extern ENGINE_API bool GIsLowMemory;
 
 void UWorld::UpdateLevelStreaming()
@@ -5811,7 +5833,7 @@ ABrush* UWorld::GetDefaultBrush() const
 
 bool UWorld::HasBegunPlay() const
 {
-	return bBegunPlay && PersistentLevel && PersistentLevel->Actors.Num();
+	return GetBegunPlay() && PersistentLevel && PersistentLevel->Actors.Num();
 }
 
 bool UWorld::AreActorsInitialized() const
diff --git a/Engine/Source/Runtime/Engine/Private/WorldSettings.cpp b/Engine/Source/Runtime/Engine/Private/WorldSettings.cpp
--- a/Engine/Source/Runtime/Engine/Private/WorldSettings.cpp
+++ b/Engine/Source/Runtime/Engine/Private/WorldSettings.cpp
@@ -296,7 +296,7 @@ float AWorldSettings::SetTimeDilation(float NewTimeDilation)
 void AWorldSettings::NotifyBeginPlay()
 {
 	UWorld* World = GetWorld();
-	if (!World->bBegunPlay)
+	if (!World->GetBegunPlay())
 	{
 		for (FActorIterator It(World); It; ++It)
 		{
@@ -305,7 +305,7 @@ void AWorldSettings::NotifyBeginPlay()
 			It->DispatchBeginPlay(bFromLevelLoad);
 		}
 
-		World->bBegunPlay = true;
+		World->SetBegunPlay(true);
 	}
 }
 
diff --git a/Engine/Source/Runtime/UMG/Private/Components/Viewport.cpp b/Engine/Source/Runtime/UMG/Private/Components/Viewport.cpp
--- a/Engine/Source/Runtime/UMG/Private/Components/Viewport.cpp
+++ b/Engine/Source/Runtime/UMG/Private/Components/Viewport.cpp
@@ -115,13 +115,13 @@ void FUMGViewportClient::Tick(float InDeltaTime)
 	{
 		// Begin Play
 		UWorld* PreviewWorld = PreviewScene->GetWorld();
-		if ( !PreviewWorld->bBegunPlay )
+		if ( !PreviewWorld->GetBegunPlay() )
 		{
 			for ( FActorIterator It(PreviewWorld); It; ++It )
 			{
 				It->DispatchBeginPlay();
 			}
-			PreviewWorld->bBegunPlay = true;
+			PreviewWorld->SetBegunPlay(true);
 		}
 
 		// Tick
