Fix invalid device profiles applied when the platform profile is changed without the game's knowledge

Also fix another issue where GFPs cannot create required child profiles because the parent exists but was not created yet. We need to load parents at this point because we create the children.

#tests Device profiles are now correctly applied when switching platform profiles while an additional override is applied, still works without one.
#rnx

diff --git a/Engine/Plugins/Experimental/GameFeatures/Source/GameFeatures/Private/GameFeatureData.cpp b/Engine/Plugins/Experimental/GameFeatures/Source/GameFeatures/Private/GameFeatureData.cpp
--- a/Engine/Plugins/Experimental/GameFeatures/Source/GameFeatures/Private/GameFeatureData.cpp
+++ b/Engine/Plugins/Experimental/GameFeatures/Source/GameFeatures/Private/GameFeatureData.cpp
@@ -247,6 +247,13 @@ void UGameFeatureData::InitializeHierarchicalPluginIniFiles(const FString& Plugi
 
 					if (ParentProfileName && ProfileSuffix && (FragmentIncludes.Num() > 0 || PluginCVars.Num() > 0))
 					{
+						// We need to load all candidate device profiles here or else we won't be able to create a child for them
+						TArray<FString> LoadableProfileNames = DeviceProfileManager.GetLoadableProfileNames(*PlatformName);
+						for (const FString& ProfileName : LoadableProfileNames)
+						{
+							DeviceProfileManager.FindProfile(ProfileName, true);
+						}
+
 						for (const UDeviceProfile* Profile : DeviceProfileManager.Profiles)
 						{
 							// Check if one of the parents for this profile is the one this rule applies to
diff --git a/Engine/Source/Runtime/Engine/Classes/DeviceProfiles/DeviceProfileManager.h b/Engine/Source/Runtime/Engine/Classes/DeviceProfiles/DeviceProfileManager.h
--- a/Engine/Source/Runtime/Engine/Classes/DeviceProfiles/DeviceProfileManager.h
+++ b/Engine/Source/Runtime/Engine/Classes/DeviceProfiles/DeviceProfileManager.h
@@ -84,6 +84,15 @@ public:
 	 */
 	ENGINE_API bool HasLoadableProfileName(const FString& ProfileName, FName OptionalPlatformName = FName());
 
+	/**
+	 * Get a list of all a named device profiles that are available to call CreateProfile with.
+	 *
+	 * @param OptionalPlatformName - The platform name to use for loading.
+	 *
+	 * @return matching profiles.
+	 */
+	ENGINE_API TArray<FString> GetLoadableProfileNames(FName OptionalPlatformName = FName()) const;
+
 	/**
 	 * Delete a profile.
 	 *
diff --git a/Engine/Source/Runtime/Engine/Private/DeviceProfiles/DeviceProfileManager.cpp b/Engine/Source/Runtime/Engine/Private/DeviceProfiles/DeviceProfileManager.cpp
--- a/Engine/Source/Runtime/Engine/Private/DeviceProfiles/DeviceProfileManager.cpp
+++ b/Engine/Source/Runtime/Engine/Private/DeviceProfiles/DeviceProfileManager.cpp
@@ -803,14 +803,42 @@ bool UDeviceProfileManager::HasLoadableProfileName(const FString& ProfileName, F
 	// @todo config: we could likely cache local ini files to speed this up,
 	// along with the ones we load in LoadConfig
 	// NOTE: This happens at runtime, so maybe only do this if !RequiresCookedData()?
-	FConfigFile* PlatformConfigFile;
-	FConfigFile LocalConfigFile;
-
-	PlatformConfigFile = GConfig->Find(GDeviceProfilesIni);
+	FConfigFile* PlatformConfigFile = GConfig->Find(GDeviceProfilesIni);
 	const FString SectionName = FString::Printf(TEXT("%s %s"), *ProfileName, *UDeviceProfile::StaticClass()->GetName());
 	return PlatformConfigFile->Contains(SectionName);
 }
 
+TArray<FString> UDeviceProfileManager::GetLoadableProfileNames(FName OptionalPlatformName) const
+{
+	FConfigCacheIni* ConfigSystem = GConfig;
+
+	if (OptionalPlatformName != NAME_None)
+	{
+#if ALLOW_OTHER_PLATFORM_CONFIG
+		ConfigSystem = FConfigCacheIni::ForPlatform(OptionalPlatformName);
+#else
+
+		checkf(OptionalPlatformName == FName(FPlatformProperties::IniPlatformName()), 
+			TEXT("UDeviceProfileManager::GetLoadableProfileNames - This platform cannot load configurations for other platforms."));
+#endif
+	}
+
+	TArray<FString> Results;
+	FConfigFile* PlatformConfigFile = GConfig->Find(GDeviceProfilesIni);
+	for (const TTuple<FString, FConfigSection>& Entry : AsConst(*PlatformConfigFile))
+	{
+		FString ProfileName;
+		FString ProfileClass;
+
+		if (Entry.Key.Split(" ", &ProfileName, &ProfileClass) && ProfileClass == *UDeviceProfile::StaticClass()->GetName())
+		{
+			Results.Add(ProfileName);
+		}
+	}
+
+	return Results;
+}
+
 
 void UDeviceProfileManager::DeleteProfile( UDeviceProfile* Profile )
 {
