Iris - Properly quantize and dequantize RootMotionSourceGroup LastPreAdditiveVelocity.

#rnx

diff --git a/Engine/Source/Runtime/Engine/Private/GameFramework/RootMotionSourceGroupNetSerializer.cpp b/Engine/Source/Runtime/Engine/Private/GameFramework/RootMotionSourceGroupNetSerializer.cpp
--- a/Engine/Source/Runtime/Engine/Private/GameFramework/RootMotionSourceGroupNetSerializer.cpp
+++ b/Engine/Source/Runtime/Engine/Private/GameFramework/RootMotionSourceGroupNetSerializer.cpp
@@ -309,6 +309,7 @@ void FRootMotionSourceGroupNetSerializer::Quantize(FNetSerializationContext& Con
 
 	{
 		FNetQuantizeArgs QuantizeArgs = Args;
+		QuantizeArgs.Source = NetSerializerValuePointer(&Source.LastPreAdditiveVelocity);
 		QuantizeArgs.Target = NetSerializerValuePointer(&Target.LastPreAdditiveVelocity);
 		QuantizeArgs.NetSerializerConfig = NetSerializerConfigParam(QuantizedVectorNetSerializer->DefaultConfig);
 		QuantizedVectorNetSerializer->Quantize(Context, QuantizeArgs);
@@ -316,6 +317,7 @@ void FRootMotionSourceGroupNetSerializer::Quantize(FNetSerializationContext& Con
 
 	{
 		FNetQuantizeArgs QuantizeArgs = Args;
+		// Note that the serializer operates directly on the source type so the Source should not be changed.
 		QuantizeArgs.Target = NetSerializerValuePointer(&Target.RootMotionSources);
 		RootMotionSourcesSerializer::Quantize(Context, QuantizeArgs);
 		if (Context.HasError())
@@ -326,6 +328,7 @@ void FRootMotionSourceGroupNetSerializer::Quantize(FNetSerializationContext& Con
 
 	{
 		FNetQuantizeArgs QuantizeArgs = Args;
+		// Note that the serializer operates directly on the source type so the Source should not be changed.
 		QuantizeArgs.Target = NetSerializerValuePointer(&Target.PendingRootMotionSources);
 		PendingRootMotionSourcesSerializer::Quantize(Context, QuantizeArgs);
 		if (Context.HasError())
@@ -351,6 +354,7 @@ void FRootMotionSourceGroupNetSerializer::Dequantize(FNetSerializationContext& C
 	{
 		FNetDequantizeArgs DequantizeArgs = Args;
 		DequantizeArgs.Source = NetSerializerValuePointer(&Source.LastPreAdditiveVelocity);
+		DequantizeArgs.Target = NetSerializerValuePointer(&Target.LastPreAdditiveVelocity);
 		DequantizeArgs.NetSerializerConfig = NetSerializerConfigParam(QuantizedVectorNetSerializer->DefaultConfig);
 		QuantizedVectorNetSerializer->Dequantize(Context, DequantizeArgs);
 	}
