Add cvar net.ClientIncomingBunchFrameTimeLimitMS to amortize incoming network bunches over multiple frames, as an option to trade off some latency for frame time. Disabled by default.

diff --git a/Engine/Source/Runtime/Engine/Classes/Engine/ActorChannel.h b/Engine/Source/Runtime/Engine/Classes/Engine/ActorChannel.h
--- a/Engine/Source/Runtime/Engine/Classes/Engine/ActorChannel.h
+++ b/Engine/Source/Runtime/Engine/Classes/Engine/ActorChannel.h
@@ -480,6 +480,9 @@ private:
 		PRAGMA_ENABLE_DEPRECATION_WARNINGS
 	}
 
+	/** Used internally by ProcessQueuedBunches to check time limits */
+	bool ShouldProcessAllQueuedBunches(float CurrentTimeSeconds);
+
 private:
 
 	// TODO: It would be nice to merge the tracking of these with PendingGuidResolves, to not duplicate memory,
diff --git a/Engine/Source/Runtime/Engine/Classes/Engine/NetDriver.h b/Engine/Source/Runtime/Engine/Classes/Engine/NetDriver.h
--- a/Engine/Source/Runtime/Engine/Classes/Engine/NetDriver.h
+++ b/Engine/Source/Runtime/Engine/Classes/Engine/NetDriver.h
@@ -2086,6 +2086,18 @@ public:
 	 */
 	ENGINE_API virtual bool IsEncryptionRequired() const;
 
+	/** Returns the value of cvar net.ClientIncomingBunchFrameTimeLimitMS on clients, or 0 otherwise. 0 = no limit. */
+	ENGINE_API float GetIncomingBunchFrameProcessingTimeLimit() const;
+
+	/** Returns true if the cvar net.ClientIncomingBunchFrameTimeLimitMS is set and the limit was exceeded */
+	ENGINE_API bool HasExceededIncomingBunchFrameProcessingTime() const;
+
+	/** Called internally by channels to track processing time for net.ClientIncomingBunchFrameTimeLimitMS and HasExceededIncomingBunchFrameProcessingTime() */
+	void AddBunchProcessingFrameTimeMS(float Milliseconds) { IncomingBunchProcessingElapsedFrameTimeMS += Milliseconds; }
+
+	/** Called internally by channels to track how many hit net.QueuedBunchTimeFailsafeSeconds */
+	void AddQueuedBunchFailsafeChannel() { ++QueuedBunchFailsafeNumChannels; }
+
 protected:
 	
 	/** Stream of random numbers to be used by this instance of UNetDriver */
@@ -2186,4 +2198,17 @@ private:
 
 	/** Cached value for UEngine.GlobalNetTravelCount, at the time of NetDriver initialization */
 	uint32 CachedGlobalNetTravelCount = 0;
+	
+	/** Accumulated number of frames in the current stat gathering period */
+	uint32 StatUpdateFrames = 0;
+
+	/** Milliseconds spent processing incoming bunches in the current frame, directly from the network and queued on channels. */
+	float IncomingBunchProcessingElapsedFrameTimeMS = 0.0f;
+	
+	/** Accumulated number of frames in the current stat period for which HasExceededIncomingBunchFrameProcessingTime would return true */
+	uint32 NumFramesOverIncomingBunchTimeLimit = 0;
+
+	/** Accumulated number of channels that hit the net.QueuedBunchTimeFailsafeSeconds this frame */
+	uint32 QueuedBunchFailsafeNumChannels = 0;
+
 };
diff --git a/Engine/Source/Runtime/Engine/Private/DataChannel.cpp b/Engine/Source/Runtime/Engine/Private/DataChannel.cpp
--- a/Engine/Source/Runtime/Engine/Private/DataChannel.cpp
+++ b/Engine/Source/Runtime/Engine/Private/DataChannel.cpp
@@ -169,6 +169,12 @@ namespace UE::Net
 		bSkipDestroyNetStartupActorsOnChannelCloseDueToLevelUnloaded,
 		TEXT("Controls if Actor that is a NetStartUpActor assosciated with the channel is destroyed or not when we receive a channel close with ECloseReason::LevelUnloaded."), ECVF_Default);
 
+	static float QueuedBunchTimeFailsafeSeconds = 2.0f;
+	static FAutoConsoleVariableRef CVarNetQueuedBunchTimeFailsafeSeconds(
+		TEXT("net.QueuedBunchTimeFailsafeSeconds"),
+		QueuedBunchTimeFailsafeSeconds,
+		TEXT("Amount of time in seconds to wait with queued bunches before forcibly processing them all, ignoring the NetDriver's HasExceededIncomingBunchFrameProcessingTime."));
+
 	// bTearOff is private but we still might need to adjust the flag on clients based on the channel close reason
 	class FTearOffSetter final
 	{
@@ -2742,6 +2748,27 @@ bool UActorChannel::CanStopTicking() const
 	return Super::CanStopTicking() && PendingGuidResolves.Num() == 0 && QueuedBunches.Num() == 0;
 }
 
+bool UActorChannel::ShouldProcessAllQueuedBunches(float CurrentTimeSeconds)
+{
+	using namespace UE::Net;
+
+	if (Connection->Driver->GetIncomingBunchFrameProcessingTimeLimit() <= 0.0f)
+	{
+		return true;
+	}
+
+	const float SecondsWithQueuedBunches = static_cast<float>(CurrentTimeSeconds - QueuedBunchStartTime);
+	const bool bExceededFailsafeTime = SecondsWithQueuedBunches > QueuedBunchTimeFailsafeSeconds;
+	
+	if (bExceededFailsafeTime)
+	{
+		Connection->Driver->AddQueuedBunchFailsafeChannel();
+		UE_LOG(LogNet, Verbose, TEXT("UActorChannel::ProcessQueuedBunches hit frame time failsafe after %.3f seconds. Processing entire queue of %d bunches for channel: %s"), SecondsWithQueuedBunches, QueuedBunches.Num(), *Describe());
+	}
+
+	return bExceededFailsafeTime;
+};
+
 bool UActorChannel::ProcessQueuedBunches()
 {
 	if (PendingGuidResolves.Num() == 0 && QueuedBunches.Num() == 0)
@@ -2749,7 +2776,7 @@ bool UActorChannel::ProcessQueuedBunches()
 		return true;
 	}
 
-	const uint32 QueueBunchStartCycles = FPlatformTime::Cycles();
+	const uint64 QueueBunchStartCycles = FPlatformTime::Cycles64();
 
 	// Try to resolve any guids that are holding up the network stream on this channel
 	// TODO: This could take a non-trivial amount of time since both GetObjectFromNetGUID
@@ -2789,7 +2816,7 @@ bool UActorChannel::ProcessQueuedBunches()
 	// If we don't have any time, then don't bother doing anything (including warning) as that may make things worse.
 	if (bHasTimeToProcess)
 	{
-		// We can process all of the queued up bunches if ALL of these are true:
+		// We can process at least some of the queued up bunches if ALL of these are true:
 		//	1. We no longer have any pending guids to load
 		//	2. We aren't still processing bunches on another channel that this actor was previously on
 		//	3. We haven't spent too much time yet this frame processing queued bunches
@@ -2799,15 +2826,26 @@ bool UActorChannel::ProcessQueuedBunches()
 			&& !Connection->Driver->ShouldQueueBunchesForActorGUID(ActorNetGUID))
 		{
 			DECLARE_SCOPE_CYCLE_COUNTER(TEXT("ProcessQueuedBunches time"), STAT_ProcessQueuedBunchesTime, STATGROUP_Net);
-			for (FInBunch* QueuedInBunch : QueuedBunches)
+
+			// Check failsafe timer so we don't end up starving the queue indefinitely due to frame time limits
+			const bool bProcessAllBunches = ShouldProcessAllQueuedBunches(FPlatformTime::ToSeconds64(QueueBunchStartCycles));
+			
+			int32 BunchIndex = 0;
+			while ((BunchIndex < QueuedBunches.Num()) && (bProcessAllBunches || !Connection->Driver->HasExceededIncomingBunchFrameProcessingTime()))
 			{
+				FInBunch* QueuedInBunch = QueuedBunches[BunchIndex];
+
 				ProcessBunch(*QueuedInBunch);
 				delete QueuedInBunch;
+				QueuedBunches[BunchIndex] = nullptr;
+
+				++BunchIndex;
 			}
 
-			UE_LOG(LogNet, VeryVerbose, TEXT("UActorChannel::ProcessQueuedBunches: Flushing queued bunches. ChIndex: %i, Actor: %s, Queued: %i"), ChIndex, Actor != NULL ? *Actor->GetPathName() : TEXT("NULL"), QueuedBunches.Num());
+			UE_LOG(LogNet, VeryVerbose, TEXT("UActorChannel::ProcessQueuedBunches: Flushing %i of %i queued bunches. ChIndex: %i, Actor: %s"), BunchIndex, QueuedBunches.Num(), ChIndex, Actor != NULL ? *Actor->GetPathName() : TEXT("NULL"));
 
-			QueuedBunches.Empty();
+			// Remove processed bunches from the queue
+			QueuedBunches.RemoveAt(0, BunchIndex);
 
 			// Call any onreps that were delayed because we were queuing bunches
 			for (auto& ReplicatorPair : ReplicationMap)
@@ -2821,10 +2859,14 @@ bool UActorChannel::ProcessQueuedBunches()
 				FNetGUIDCache::FIsOwnerOrPawnHelper Helper(Connection->Driver->GuidCache.Get(), Connection->OwningActor, Actor);
 #endif
 
-				PackageMapClient->SetHasQueuedBunches(ActorNetGUID, false);
+				PackageMapClient->SetHasQueuedBunches(ActorNetGUID, !QueuedBunches.IsEmpty());
 			}
 
-			QueuedBunchObjectReferences.Empty();
+			// We don't know exactly which bunches were referring to which objects, so we have to conservatively keep them all around.
+			if (QueuedBunches.IsEmpty())
+			{
+				QueuedBunchObjectReferences.Empty();
+			}
 		}
 		else
 		{
@@ -2851,9 +2893,9 @@ bool UActorChannel::ProcessQueuedBunches()
 		}
 
 		// Update the driver with our time spent
-		const uint32 QueueBunchEndCycles = FPlatformTime::Cycles();
-		const uint32 QueueBunchDeltaCycles = QueueBunchEndCycles - QueueBunchStartCycles;
-		const float QueueBunchDeltaMilliseconds = FPlatformTime::ToMilliseconds(QueueBunchDeltaCycles);
+		const uint64 QueueBunchEndCycles = FPlatformTime::Cycles64();
+		const uint64 QueueBunchDeltaCycles = QueueBunchEndCycles - QueueBunchStartCycles;
+		const float QueueBunchDeltaMilliseconds = static_cast<float>(FPlatformTime::ToMilliseconds64(QueueBunchDeltaCycles));
 
 		Connection->Driver->ProcessQueuedBunchesCurrentFrameMilliseconds += QueueBunchDeltaMilliseconds;
 	}
@@ -2975,8 +3017,10 @@ void UActorChannel::ReceivedBunch( FInBunch & Bunch )
 		//	2. We already have queued up bunches
 		//	3. If this actor was previously on a channel that is now still processing bunches after a close
 		//	4. The driver is requesting queuing for this GUID
+		//	5. We are time-boxing incoming bunch processing and have run out of time this frame
 		if (PendingGuidResolves.Num() > 0 || QueuedBunches.Num() > 0 || Connection->KeepProcessingActorChannelBunchesMap.Contains(ActorNetGUID) ||
-			 (Connection->Driver->ShouldQueueBunchesForActorGUID(ActorNetGUID)))
+			Connection->Driver->ShouldQueueBunchesForActorGUID(ActorNetGUID) ||
+			Connection->Driver->HasExceededIncomingBunchFrameProcessingTime())
 		{
 			if (Connection->KeepProcessingActorChannelBunchesMap.Contains(ActorNetGUID))
 			{
@@ -3026,6 +3070,17 @@ void UActorChannel::ProcessBunch( FInBunch & Bunch )
 		return;
 	}
 
+	uint64 StartTimeCycles = FPlatformTime::Cycles64();
+	ON_SCOPE_EXIT
+	{
+		if (Connection->Driver->GetIncomingBunchFrameProcessingTimeLimit() > 0.0f)
+		{
+			const uint64 DeltaCycles = FPlatformTime::Cycles64() - StartTimeCycles;
+			const double DeltaMS = FPlatformTime::ToMilliseconds64(DeltaCycles);
+			Connection->Driver->AddBunchProcessingFrameTimeMS(static_cast<float>(DeltaMS));
+		}
+	};
+
 	FReplicationFlags RepFlags;
 
 	// ------------------------------------------------------------
diff --git a/Engine/Source/Runtime/Engine/Private/NetDriver.cpp b/Engine/Source/Runtime/Engine/Private/NetDriver.cpp
--- a/Engine/Source/Runtime/Engine/Private/NetDriver.cpp
+++ b/Engine/Source/Runtime/Engine/Private/NetDriver.cpp
@@ -144,6 +144,7 @@ DECLARE_CYCLE_STAT(TEXT("NetDriver TickFlush"), STAT_NetTickFlush, STATGROUP_Gam
 DECLARE_CYCLE_STAT(TEXT("NetDriver TickFlush GatherStats"), STAT_NetTickFlushGatherStats, STATGROUP_Game);
 DECLARE_CYCLE_STAT(TEXT("NetDriver TickFlush GatherStatsPerfCounters"), STAT_NetTickFlushGatherStatsPerfCounters, STATGROUP_Game);
 DECLARE_CYCLE_STAT(TEXT("ReceiveRPC_ProcessRemoteFunction"), STAT_NetReceiveRPC_ProcessRemoteFunction, STATGROUP_Game);
+DECLARE_DWORD_ACCUMULATOR_STAT(TEXT("Net bunch time % overshoot frames"), STAT_NetInBunchTimeOvershootPercent, STATGROUP_Net);
 
 DEFINE_LOG_CATEGORY_STATIC(LogNetSyncLoads, Log, All);
 
@@ -326,6 +327,15 @@ namespace UE::Net::Private
 		return false;
 	}
 
+	static float ClientIncomingBunchFrameTimeLimitMS = 0.0f;
+
+	static FAutoConsoleVariableRef CVarClientIncomingBunchFrameTimeLimitMS(
+		TEXT("net.ClientIncomingBunchFrameTimeLimitMS"),
+		ClientIncomingBunchFrameTimeLimitMS,
+		TEXT("Time in milliseconds to limit client incoming bunch processing to. If 0, no limit. As long as we're below the limit, will start processing another bunch. A single bunch that takes a while to process can overshoot the limit. ")
+		TEXT("After the limit is hit, remaining bunches in a packet are queued, and the IpNetDriver will not process any more packets in the current frame."));
+
+
 } //namespace UE::Net::Private
 
 namespace UE::Net
@@ -1243,6 +1253,12 @@ void UNetDriver::TickFlush(float DeltaSeconds)
 		LastCleanupTime = CurrentRealtimeSeconds;
 	}
 
+	if (QueuedBunchFailsafeNumChannels > 0)
+	{
+		UE_LOG(LogNet, Log, TEXT("UNetDriver::TickFlush: %u channel(s) exceeded net.QueuedBunchTimeFailsafeSeconds and flushed their entire queue(s) this frame."), QueuedBunchFailsafeNumChannels);
+		QueuedBunchFailsafeNumChannels = 0;
+	}
+
 	UpdateNetworkStats();
 
 	// Update the lag state
@@ -1906,6 +1922,8 @@ void UNetDriver::TickDispatch( float DeltaTime )
 	// Get new time.
 	ElapsedTime += DeltaTime;
 
+	IncomingBunchProcessingElapsedFrameTimeMS = 0.0f;
+
 	// Checks for standby cheats if enabled	
 	UpdateStandbyCheatStatus();
 
@@ -7335,6 +7353,24 @@ bool UNetDriver::IsEncryptionRequired() const
 			!bPIEDisableEncryptionCheck;
 }
 
+float UNetDriver::GetIncomingBunchFrameProcessingTimeLimit() const
+{
+	using namespace UE::Net::Private;
+
+	if (!IsServer())
+	{
+		return ClientIncomingBunchFrameTimeLimitMS;
+	}
+
+	return 0.0f;
+}
+
+bool UNetDriver::HasExceededIncomingBunchFrameProcessingTime() const
+{
+	const float TimeLimit = GetIncomingBunchFrameProcessingTimeLimit();
+	return (TimeLimit > 0.0f) && (IncomingBunchProcessingElapsedFrameTimeMS > TimeLimit);
+}
+
 void UNetDriver::UpdateNetworkStats()
 {
 	using namespace UE::Net::Private;
@@ -7348,6 +7384,10 @@ void UNetDriver::UpdateNetworkStats()
 	{
 		SCOPE_CYCLE_COUNTER(STAT_NetTickFlushGatherStats);
 
+		++StatUpdateFrames;
+
+		NumFramesOverIncomingBunchTimeLimit += HasExceededIncomingBunchFrameProcessingTime() ? 1 : 0;
+
 		const double CurrentRealtimeSeconds = FPlatformTime::Seconds();
 		// Update network stats (only main game net driver for now) if stats or perf counters are used
 		if (NetDriverName == NAME_GameNetDriver &&
@@ -7371,6 +7411,7 @@ void UNetDriver::UpdateNetworkStats()
 			int32 ClientsOutPacketsThisFrameMax = 0;
 			int NumClients = 0;
 			int32 MaxPacketOverhead = 0;
+			int32 InBunchTimeOvershootPercent = 0;
 
 			// these need to be updated even if we are not collecting stats, since they get reported to analytics/QoS
 			for (UNetConnection* Client : ClientConnections)
@@ -7478,6 +7519,9 @@ void UNetDriver::UpdateNetworkStats()
 				VoiceInPercent = (InBytes > 0) ? FMath::TruncToInt(100.f * (float)VoiceBytesRecv / (float)InBytes) : 0;
 				VoiceOutPercent = (OutBytes > 0) ? FMath::TruncToInt(100.f * (float)VoiceBytesSent / (float)OutBytes) : 0;
 
+				// Percent of frames this stat period that the incoming bunch processing time limit was exceeded (0-100)
+				InBunchTimeOvershootPercent = (StatUpdateFrames > 0) ? FMath::TruncToInt(100.0f * (float)NumFramesOverIncomingBunchTimeLimit / (float)StatUpdateFrames) : 0;
+
 				if (World)
 				{
 					NumActors = World->GetActorCount();
@@ -7566,6 +7610,8 @@ void UNetDriver::UpdateNetworkStats()
 				SET_DWORD_STAT(STAT_NumNetGUIDsPending, UnAckCount);
 				SET_DWORD_STAT(STAT_NumNetGUIDsUnAckd, PendingCount);
 				SET_DWORD_STAT(STAT_NetSaturated, NetSaturated);
+
+				SET_DWORD_STAT(STAT_NetInBunchTimeOvershootPercent, InBunchTimeOvershootPercent);
 			}
 
 			// If we are to replicate server stats out to an observer, then set those values
@@ -7745,7 +7791,9 @@ void UNetDriver::UpdateNetworkStats()
 			VoiceBytesRecv = 0;
 			VoiceInPercent = 0;
 			VoiceOutPercent = 0;
+			NumFramesOverIncomingBunchTimeLimit = 0;
 			StatUpdateTime = CurrentRealtimeSeconds;
+			StatUpdateFrames = 0;
 		}
 		else
 		{
