Disabled mesh material parameter caching to fix performance regression.

diff --git a/Engine/Source/Runtime/Engine/Private/Components/MeshComponent.cpp b/Engine/Source/Runtime/Engine/Private/Components/MeshComponent.cpp
--- a/Engine/Source/Runtime/Engine/Private/Components/MeshComponent.cpp
+++ b/Engine/Source/Runtime/Engine/Private/Components/MeshComponent.cpp
@@ -29,7 +29,7 @@ UMeshComponent::UMeshComponent(const FObjectInitializer& ObjectInitializer)
 	bUseAsOccluder = true;
 	bCanEverAffectNavigation = true;
 	bCachedMaterialParameterIndicesAreDirty = true;
-	bEnableMaterialParameterCaching = true;
+	bEnableMaterialParameterCaching = false;
 }
 
 UMaterialInterface* UMeshComponent::GetMaterial(int32 ElementIndex) const
