Fix for landscape layer not being recognized as such when used within material layers
#tests editor

diff --git a/Engine/Source/Runtime/Engine/Private/Materials/HLSLMaterialTranslator.cpp b/Engine/Source/Runtime/Engine/Private/Materials/HLSLMaterialTranslator.cpp
--- a/Engine/Source/Runtime/Engine/Private/Materials/HLSLMaterialTranslator.cpp
+++ b/Engine/Source/Runtime/Engine/Private/Materials/HLSLMaterialTranslator.cpp
@@ -8078,6 +8078,10 @@ int32 FHLSLMaterialTranslator::StaticTerrainLayerWeight(FName LayerName,int32 De
 		constexpr EMaterialSamplerType SamplerType = SAMPLERTYPE_Masks;
 		FString WeightmapName = FString::Printf(TEXT("Weightmap%d"), WeightmapIndex);
 		int32 TextureReferenceIndex = INDEX_NONE;
+
+		FMaterialParameterInfo GlobalParameterInfo;
+		PushParameterOwner(GlobalParameterInfo);
+
 		int32 TextureCodeIndex = TextureParameter(FName(*WeightmapName), GEngine->WeightMapPlaceholderTexture, TextureReferenceIndex, SamplerType);
 		int32 SampleCodeIndex = TextureSample(TextureCodeIndex, TextureCoordinate(3, false, false), SamplerType, /*MipValue0Index = */INDEX_NONE, /*MipValue1Index = */INDEX_NONE, /*MipValueMode = */TMVM_None, /*SamplerSource = */SSM_TerrainWeightmapGroupSettings);
 
@@ -8093,6 +8097,7 @@ int32 FHLSLMaterialTranslator::StaticTerrainLayerWeight(FName LayerName,int32 De
 			WeightmapCode = Add(WeightmapCode, CurrentWeightmapCode);
 		}
 
+		PopParameterOwner();
 		++NumWeightmapParameters;
 	}
 
