World Partition Runtime HashSet:

Implemented streaming cells sorting priority:
- Moved all cells sorting criterias into the base class, and changed the ordering: source priority > cell hierarchical level > cell custom priority > cell distance and angle from source.
- Original spatial hash implementation still use the old sorting criterias (2D): source priority > cell hierarchical level or extent (cvar) > cell distance and angle from source > cell custom priority.

Implemented cells streaming performance computations:
- Moved most of the code into the base class so all hash classes can share a single implementation.
- Original runtime spatial hash implementation still use the old (2D) performance computations for backward compatibility.

diff --git a/Engine/Source/Runtime/Engine/Private/WorldPartition/RuntimeHashSet/WorldPartitionRuntimeHashSetDebugDraw.cpp b/Engine/Source/Runtime/Engine/Private/WorldPartition/RuntimeHashSet/WorldPartitionRuntimeHashSetDebugDraw.cpp
--- a/Engine/Source/Runtime/Engine/Private/WorldPartition/RuntimeHashSet/WorldPartitionRuntimeHashSetDebugDraw.cpp
+++ b/Engine/Source/Runtime/Engine/Private/WorldPartition/RuntimeHashSet/WorldPartitionRuntimeHashSetDebugDraw.cpp
@@ -121,9 +121,9 @@ bool UWorldPartitionRuntimeHashSet::Draw2D(FWorldPartitionDraw2DContext& DrawCon
 		const FBox Region3D(FVector(Region.Min.X, Region.Min.Y, -HALF_WORLD_MAX), FVector(Region.Max.X, Region.Max.Y, HALF_WORLD_MAX));		
 		StreamingData->SpatialIndex->ForEachIntersectingElement(Region3D, [&FilteredCells](UWorldPartitionRuntimeCell* Cell)
 		{
-			UWorldPartitionRuntimeCellDataSpatialHashSet* RuntimeCellDataHashSet = CastChecked<UWorldPartitionRuntimeCellDataSpatialHashSet>(Cell->RuntimeCellData);
+			UWorldPartitionRuntimeCellData* RuntimeCellData = Cell->RuntimeCellData;
 
-			if ((RuntimeCellDataHashSet->Level >= GShowRuntimeHashSetDebugDisplayLevel) && (RuntimeCellDataHashSet->Level < (GShowRuntimeHashSetDebugDisplayLevel + GShowRuntimeHashSetDebugDisplayLevelCount)))
+			if ((RuntimeCellData->HierarchicalLevel >= GShowRuntimeHashSetDebugDisplayLevel) && (RuntimeCellData->HierarchicalLevel < (GShowRuntimeHashSetDebugDisplayLevel + GShowRuntimeHashSetDebugDisplayLevelCount)))
 			{
 				FilteredCells.Add(Cell);
 			}
diff --git a/Engine/Source/Runtime/Engine/Private/WorldPartition/RuntimeHashSet/WorldPartitionRuntimeHashSetStreamingGeneration.cpp b/Engine/Source/Runtime/Engine/Private/WorldPartition/RuntimeHashSet/WorldPartitionRuntimeHashSetStreamingGeneration.cpp
--- a/Engine/Source/Runtime/Engine/Private/WorldPartition/RuntimeHashSet/WorldPartitionRuntimeHashSetStreamingGeneration.cpp
+++ b/Engine/Source/Runtime/Engine/Private/WorldPartition/RuntimeHashSet/WorldPartitionRuntimeHashSetStreamingGeneration.cpp
@@ -189,10 +189,10 @@ bool UWorldPartitionRuntimeHashSet::GenerateStreaming(UWorldPartitionStreamingPo
 		RuntimeCell->SetIsHLOD(CellDescInstance.SourcePartition->HLODIndex != INDEX_NONE);
 		RuntimeCell->SetGuid(CellUniqueId.Guid);
 
-		UWorldPartitionRuntimeCellDataSpatialHashSet* RuntimeCellDataHashSet = CastChecked<UWorldPartitionRuntimeCellDataSpatialHashSet>(RuntimeCell->RuntimeCellData);
-		RuntimeCellDataHashSet->DebugName = CellUniqueId.Name;
-		RuntimeCellDataHashSet->Level = CellDescInstance.Level;
-		RuntimeCellDataHashSet->Priority = CellDescInstance.Priority;
+		UWorldPartitionRuntimeCellData* RuntimeCellData = RuntimeCell->RuntimeCellData;
+		RuntimeCellData->DebugName = CellUniqueId.Name;
+		RuntimeCellData->HierarchicalLevel = CellDescInstance.Level;
+		RuntimeCellData->Priority = CellDescInstance.Priority;
 
 		return RuntimeCell;
 	};
@@ -208,7 +208,7 @@ bool UWorldPartitionRuntimeHashSet::GenerateStreaming(UWorldPartitionStreamingPo
 			TArray<IStreamingGenerationContext::FActorInstance> CellActorInstances;
 			if (PopulateCellActorInstances(CellDescInstance.ActorSetInstances, bIsMainWorldPartition, bIsCellAlwaysLoaded, CellActorInstances))
 			{
-				UWorldPartitionRuntimeCell* RuntimeCell = RuntimeCells.Emplace_GetRef(CreateRuntimeCellFromCellDesc(CellDescInstance,StreamingPolicy->GetRuntimeCellClass(), UWorldPartitionRuntimeCellDataSpatialHashSet::StaticClass()));
+				UWorldPartitionRuntimeCell* RuntimeCell = RuntimeCells.Emplace_GetRef(CreateRuntimeCellFromCellDesc(CellDescInstance, StreamingPolicy->GetRuntimeCellClass(), UWorldPartitionRuntimeCellData::StaticClass()));
 				PopulateRuntimeCell(RuntimeCell, CellActorInstances, OutPackagesToGenerate);
 
 				// Override the cell bounds if the runtime partition provided one
diff --git a/Engine/Source/Runtime/Engine/Private/WorldPartition/WorldPartitionRuntimeCell.cpp b/Engine/Source/Runtime/Engine/Private/WorldPartition/WorldPartitionRuntimeCell.cpp
--- a/Engine/Source/Runtime/Engine/Private/WorldPartition/WorldPartitionRuntimeCell.cpp
+++ b/Engine/Source/Runtime/Engine/Private/WorldPartition/WorldPartitionRuntimeCell.cpp
@@ -95,9 +95,9 @@ void UWorldPartitionRuntimeCell::DumpStateLog(FHierarchicalLogArchive& Ar) const
 }
 #endif
 
-int32 UWorldPartitionRuntimeCell::SortCompare(const UWorldPartitionRuntimeCell* Other, bool bCanUseSortingCache) const
+int32 UWorldPartitionRuntimeCell::SortCompare(const UWorldPartitionRuntimeCell* Other) const
 {
-	return RuntimeCellData->SortCompare(Other->RuntimeCellData, bCanUseSortingCache);
+	return RuntimeCellData->SortCompare(Other->RuntimeCellData);
 }
 
 bool UWorldPartitionRuntimeCell::IsDebugShown() const
diff --git a/Engine/Source/Runtime/Engine/Private/WorldPartition/WorldPartitionRuntimeCellData.cpp b/Engine/Source/Runtime/Engine/Private/WorldPartition/WorldPartitionRuntimeCellData.cpp
--- a/Engine/Source/Runtime/Engine/Private/WorldPartition/WorldPartitionRuntimeCellData.cpp
+++ b/Engine/Source/Runtime/Engine/Private/WorldPartition/WorldPartitionRuntimeCellData.cpp
@@ -7,9 +7,14 @@ int32 UWorldPartitionRuntimeCellData::StreamingSourceCacheEpoch = 0;
 UWorldPartitionRuntimeCellData::UWorldPartitionRuntimeCellData(const FObjectInitializer& ObjectInitializer)
 	: Super(ObjectInitializer)
 	, CachedMinSourcePriority(MAX_uint8)
+	, bCachedWasRequestedByBlockingSource(false)
+	, CachedMinSquareDistanceToBlockingSource(MAX_dbl)
+	, CachedMinBlockOnSlowStreamingRatio(MAX_flt)
+	, CachedMinSpatialSortingPriority(MAX_dbl)
 	, CachedSourceInfoEpoch(MIN_int32)
 	, ContentBounds(ForceInit)
 	, Priority(0)
+	, HierarchicalLevel(0)
 {}
 
 void UWorldPartitionRuntimeCellData::Serialize(FArchive& Ar)
@@ -20,26 +25,81 @@ void UWorldPartitionRuntimeCellData::Serialize(FArchive& Ar)
 
 void UWorldPartitionRuntimeCellData::ResetStreamingSourceInfo() const
 {
-	check(ShouldResetStreamingSourceInfo());
-	CachedSourcePriorityWeights.Reset();
 	CachedMinSourcePriority = MAX_uint8;
-	CachedSourceInfoEpoch = StreamingSourceCacheEpoch;
+	bCachedWasRequestedByBlockingSource = false;
+	CachedMinSquareDistanceToBlockingSource = MAX_dbl;
+	CachedMinBlockOnSlowStreamingRatio = MAX_flt;
+	CachedMinSpatialSortingPriority = MAX_dbl;
+	CachedSourceInfoEpoch = StreamingSourceCacheEpoch;	
 }
 
 void UWorldPartitionRuntimeCellData::AppendStreamingSourceInfo(const FWorldPartitionStreamingSource& Source, const FSphericalSector& SourceShape) const
 {
-	static_assert((uint8)EStreamingSourcePriority::Lowest == 255);
-	CachedSourcePriorityWeights.Add(1.f - ((float)Source.Priority / 255.f));
+	if (CachedSourceInfoEpoch != StreamingSourceCacheEpoch)
+	{
+		ResetStreamingSourceInfo();
+		check(CachedSourceInfoEpoch == StreamingSourceCacheEpoch);
+	}
+
 	CachedMinSourcePriority = FMath::Min((uint8)Source.Priority, CachedMinSourcePriority);
-}
 
-void UWorldPartitionRuntimeCellData::MergeStreamingSourceInfo() const
-{}
+	if (Source.bBlockOnSlowLoading)
+	{
+		bCachedWasRequestedByBlockingSource = true;
+
+		const FVector CellToSource = SourceShape.GetCenter() - ContentBounds.GetClosestPointTo(SourceShape.GetCenter());
+		const double CellToSourceSquareDistance = CellToSource.SizeSquared();
+		CachedMinSquareDistanceToBlockingSource = FMath::Min(CellToSourceSquareDistance, CachedMinSquareDistanceToBlockingSource);
+
+		const double BlockOnSlowStreamingRatio = FMath::Sqrt(CachedMinSquareDistanceToBlockingSource) / SourceShape.GetRadius();
+		CachedMinBlockOnSlowStreamingRatio = FMath::Min(CachedMinBlockOnSlowStreamingRatio, BlockOnSlowStreamingRatio);
+	}
+
+	// Compute square distance from cell to source ratio
+	const double SoureDistanceRatio = FMath::Clamp(ContentBounds.ComputeSquaredDistanceToPoint(SourceShape.GetCenter()) / FMath::Square(SourceShape.GetRadius()), 0.0f, 1.0f);
+
+	// Compute cosine angle from cell to source direction ratio
+	const FVector CellToSource = SourceShape.GetCenter() - ContentBounds.GetClosestPointTo(SourceShape.GetCenter());
+	const float SourceCosAngle = ContentBounds.IsInsideOrOn(SourceShape.GetCenter()) ? -1.0f : (SourceShape.GetAxis() | CellToSource.GetSafeNormal());
+	const float SourceCosAngleRatio = SourceCosAngle * 0.5f + 0.5f;
 
-int32 UWorldPartitionRuntimeCellData::SortCompare(const UWorldPartitionRuntimeCellData* Other, bool bCanUseSortingCache) const
+	// Compute final cell priority for this source
+	static float CosAngleContribution = 0.4f;
+	const double SortingPriority = SoureDistanceRatio * FMath::Cube(FMath::Pow(SourceCosAngleRatio, CosAngleContribution));
+
+	// Update if lower
+	CachedMinSpatialSortingPriority = FMath::Min(CachedMinSpatialSortingPriority, SortingPriority);
+}
+
+/**
+ * Sorting criterias:
+ *	- Highest priority affecting source (lowest to highest)
+ *	- Cell hierarchical level (highest to lowest)
+ -	- Cell custom priority (lowest to highest)
+ *	- Cell distance and angle from source (lowest to highest)
+ */
+int32 UWorldPartitionRuntimeCellData::SortCompare(const UWorldPartitionRuntimeCellData* InOther) const
 {
-	// Source priority (lower value is higher prio)
-	return bCanUseSortingCache ? ((int32)CachedMinSourcePriority - (int32)Other->CachedMinSourcePriority) : 0;
+	int64 Result = (int32)CachedMinSourcePriority - (int32)InOther->CachedMinSourcePriority;
+
+	if (!Result)
+	{
+		// Cell hierarchical level (highest to lowest)
+		Result = InOther->HierarchicalLevel - HierarchicalLevel;
+
+		if (!Result)
+		{
+			// Cell priority (lower value is higher prio)
+			Result = Priority - InOther->Priority;			
+
+			if (!Result)
+			{
+				Result = (CachedMinSpatialSortingPriority < InOther->CachedMinSpatialSortingPriority) ? -1 : 1;
+			}
+		}
+	}
+
+	return (int32)FMath::Clamp(Result, -1, 1);
 }
 
 const FBox& UWorldPartitionRuntimeCellData::GetContentBounds() const
diff --git a/Engine/Source/Runtime/Engine/Private/WorldPartition/WorldPartitionRuntimeCellDataSpatialHash.cpp b/Engine/Source/Runtime/Engine/Private/WorldPartition/WorldPartitionRuntimeCellDataSpatialHash.cpp
--- a/Engine/Source/Runtime/Engine/Private/WorldPartition/WorldPartitionRuntimeCellDataSpatialHash.cpp
+++ b/Engine/Source/Runtime/Engine/Private/WorldPartition/WorldPartitionRuntimeCellDataSpatialHash.cpp
@@ -22,14 +22,11 @@ static bool GRuntimeSpatialHashSortUsingCellPriority = true;
 static FAutoConsoleVariableRef CVarRuntimeSpatialHashSortUsingCellPriority(
 	TEXT("wp.Runtime.RuntimeSpatialHashSortUsingCellPriority"),
 	GRuntimeSpatialHashSortUsingCellPriority,
-	TEXT("Set to 1 to use cell priority as part of the sorting criterias when sorting cells by importance."));
+	TEXT("Set to 1 to use cell priority before distance to/angle from source as part of the sorting criterias when sorting cells by importance."));
 
 UWorldPartitionRuntimeCellDataSpatialHash::UWorldPartitionRuntimeCellDataSpatialHash(const FObjectInitializer& ObjectInitializer)
 	: Super(ObjectInitializer)
 	, Extent(0)
-	, Level(0)
-	, bCachedIsBlockingSource(false)
-	, CachedMinSquareDistanceToBlockingSource(MAX_dbl)
 	, CachedMinSquareDistanceToSource(MAX_dbl)
 {}
 
@@ -37,9 +34,8 @@ void UWorldPartitionRuntimeCellDataSpatialHash::ResetStreamingSourceInfo() const
 {
 	Super::ResetStreamingSourceInfo();
 
-	bCachedIsBlockingSource = false;
-	CachedMinSquareDistanceToBlockingSource = MAX_dbl;
 	CachedMinSquareDistanceToSource = MAX_dbl;
+	CachedSourcePriorityWeights.Reset();
 	CachedSourceSquaredDistances.Reset();
 	CachedInstersectingShapes.Reset();
 }
@@ -89,13 +85,14 @@ void UWorldPartitionRuntimeCellDataSpatialHash::AppendStreamingSourceInfo(const
 
 	const double SquareDistance = FVector::DistSquared2D(SourceShape.GetCenter(), Position);
 
-	// Only consider blocking sources
+	// Update cached values based on the 2D distance
 	if (Source.bBlockOnSlowLoading)
 	{
-		bCachedIsBlockingSource = true;
 		CachedMinSquareDistanceToBlockingSource = FMath::Min(SquareDistance, CachedMinSquareDistanceToBlockingSource);
+		CachedMinBlockOnSlowStreamingRatio = FMath::Min(CachedMinBlockOnSlowStreamingRatio, FMath::Sqrt(CachedMinSquareDistanceToBlockingSource) / SourceShape.GetRadius());
 	}
 	CachedSourceSquaredDistances.Add(SquareDistance);
+	CachedSourcePriorityWeights.Add(1.0f - ((float)Source.Priority / (float)EStreamingSourcePriority::Lowest));
 	CachedInstersectingShapes.Add(SourceShape);
 }
 
@@ -141,16 +138,18 @@ void UWorldPartitionRuntimeCellDataSpatialHash::MergeStreamingSourceInfo() const
 	}
 }
 
-int32 UWorldPartitionRuntimeCellDataSpatialHash::SortCompare(const UWorldPartitionRuntimeCellData* InOther, bool bCanUseSortingCache) const
+int32 UWorldPartitionRuntimeCellDataSpatialHash::SortCompare(const UWorldPartitionRuntimeCellData* InOther) const
 {
-	int32 Result = Super::SortCompare(InOther, bCanUseSortingCache);
+	int32 Result = (int32)CachedMinSourcePriority - (int32)InOther->CachedMinSourcePriority;
+
 	if (Result == 0)
 	{
 		const UWorldPartitionRuntimeCellDataSpatialHash* Other = (UWorldPartitionRuntimeCellDataSpatialHash*)InOther;
-		
+
 		// By default, now compare cell's extent instead of its grid level since we compare cells across multiple WPs/grids (higher value is higher prio)
-		Result = GRuntimeSpatialHashSortUsingCellExtent ? int32(Other->Extent - Extent) : (Other->Level - Level);
-		if (bCanUseSortingCache && (Result == 0))
+		Result = GRuntimeSpatialHashSortUsingCellExtent ? int32(Other->Extent - Extent) : (InOther->HierarchicalLevel - HierarchicalLevel);
+
+		if (Result == 0)
 		{
 			if (GRuntimeSpatialHashSortUsingCellPriority)
 			{
diff --git a/Engine/Source/Runtime/Engine/Private/WorldPartition/WorldPartitionRuntimeHash.cpp b/Engine/Source/Runtime/Engine/Private/WorldPartition/WorldPartitionRuntimeHash.cpp
--- a/Engine/Source/Runtime/Engine/Private/WorldPartition/WorldPartitionRuntimeHash.cpp
+++ b/Engine/Source/Runtime/Engine/Private/WorldPartition/WorldPartitionRuntimeHash.cpp
@@ -18,6 +18,18 @@
 
 #define LOCTEXT_NAMESPACE "WorldPartition"
 
+ENGINE_API float GBlockOnSlowStreamingRatio = 0.25f;
+static FAutoConsoleVariableRef CVarBlockOnSlowStreamingRatio(
+	TEXT("wp.Runtime.BlockOnSlowStreamingRatio"),
+	GBlockOnSlowStreamingRatio,
+	TEXT("Ratio of DistanceToCell / LoadingRange to use to determine if World Partition streaming needs to block"));
+
+ENGINE_API float GBlockOnSlowStreamingWarningFactor = 2.f;
+static FAutoConsoleVariableRef CVarBlockOnSlowStreamingWarningFactor(
+	TEXT("wp.Runtime.BlockOnSlowStreamingWarningFactor"),
+	GBlockOnSlowStreamingWarningFactor,
+	TEXT("Factor of wp.Runtime.BlockOnSlowStreamingRatio we want to start notifying the user"));
+
 #if DO_CHECK
 void URuntimeHashExternalStreamingObjectBase::BeginDestroy()
 {
@@ -101,6 +113,25 @@ UWorldPartitionRuntimeCell* UWorldPartitionRuntimeHash::CreateRuntimeCell(UClass
 	return RuntimeCell;
 }
 
+EWorldPartitionStreamingPerformance UWorldPartitionRuntimeHash::GetStreamingPerformanceForCell(const UWorldPartitionRuntimeCell* Cell) const
+{
+	check(Cell->GetBlockOnSlowLoading());
+
+	if (Cell->RuntimeCellData->bCachedWasRequestedByBlockingSource)
+	{
+		if (Cell->RuntimeCellData->CachedMinBlockOnSlowStreamingRatio < GBlockOnSlowStreamingRatio)
+		{
+			return EWorldPartitionStreamingPerformance::Critical;
+		}
+		else if (Cell->RuntimeCellData->CachedMinBlockOnSlowStreamingRatio < (GBlockOnSlowStreamingRatio * GBlockOnSlowStreamingWarningFactor))
+		{
+			return EWorldPartitionStreamingPerformance::Slow;
+		}
+	}
+
+	return EWorldPartitionStreamingPerformance::Good;
+}
+
 URuntimeHashExternalStreamingObjectBase* UWorldPartitionRuntimeHash::CreateExternalStreamingObject(TSubclassOf<URuntimeHashExternalStreamingObjectBase> InClass, UObject* InOuter, FName InName, UWorld* InOwningWorld, UWorld* InOuterWorld)
 {
 	if (FindObject<URuntimeHashExternalStreamingObjectBase>(InOuter, *InName.ToString()))
@@ -467,11 +498,6 @@ bool UWorldPartitionRuntimeHash::RemoveExternalStreamingObject(URuntimeHashExter
 
 void UWorldPartitionRuntimeHash::FStreamingSourceCells::AddCell(const UWorldPartitionRuntimeCell* Cell, const FWorldPartitionStreamingSource& Source, const FSphericalSector& SourceShape)
 {
-	if (Cell->ShouldResetStreamingSourceInfo())
-	{
-		Cell->ResetStreamingSourceInfo();
-	}
-
 	Cell->AppendStreamingSourceInfo(Source, SourceShape);
 	Cells.Add(Cell);
 }
diff --git a/Engine/Source/Runtime/Engine/Private/WorldPartition/WorldPartitionRuntimeSpatialHash.cpp b/Engine/Source/Runtime/Engine/Private/WorldPartition/WorldPartitionRuntimeSpatialHash.cpp
--- a/Engine/Source/Runtime/Engine/Private/WorldPartition/WorldPartitionRuntimeSpatialHash.cpp
+++ b/Engine/Source/Runtime/Engine/Private/WorldPartition/WorldPartitionRuntimeSpatialHash.cpp
@@ -59,18 +59,6 @@ static FAutoConsoleVariableRef CVarShowRuntimeSpatialHashGridLevelCount(
 	GShowRuntimeSpatialHashGridLevelCount,
 	TEXT("Used to choose how many grid levels to display when showing world partition runtime hash."));
 
-static float GBlockOnSlowStreamingRatio = 0.25f;
-static FAutoConsoleVariableRef CVarBlockOnSlowStreamingRatio(
-	TEXT("wp.Runtime.BlockOnSlowStreamingRatio"),
-	GBlockOnSlowStreamingRatio,
-	TEXT("Ratio of DistanceToCell / LoadingRange to use to determine if World Partition streaming needs to block"));
-
-static float GBlockOnSlowStreamingWarningFactor = 2.f;
-static FAutoConsoleVariableRef CVarBlockOnSlowStreamingWarningFactor(
-	TEXT("wp.Runtime.BlockOnSlowStreamingWarningFactor"),
-	GBlockOnSlowStreamingWarningFactor,
-	TEXT("Factor of wp.Runtime.BlockOnSlowStreamingRatio we want to start notifying the user"));
-
 #if !UE_BUILD_SHIPPING
 static int32 GFilterRuntimeSpatialHashGridLevel = INDEX_NONE;
 static FAutoConsoleVariableRef CVarFilterRuntimeSpatialHashGridLevel(
@@ -1534,7 +1522,7 @@ bool UWorldPartitionRuntimeSpatialHash::CreateStreamingGrid(const FSpatialHashRu
 					const double CellExtent = Bounds.GetExtent().X;
 					check(CellExtent < MAX_flt);
 
-					CellDataSpatialHash->Level = Level;
+					CellDataSpatialHash->HierarchicalLevel = Level;
 					CellDataSpatialHash->Position = FVector(Bounds.GetCenter(), 0.f);
 					CellDataSpatialHash->Extent = (float)CellExtent;
 					CellDataSpatialHash->GridName = RuntimeGrid.GridName;
@@ -1793,12 +1781,6 @@ const FSpatialHashStreamingGrid* UWorldPartitionRuntimeSpatialHash::GetStreaming
 
 EWorldPartitionStreamingPerformance UWorldPartitionRuntimeSpatialHash::GetStreamingPerformanceForCell(const UWorldPartitionRuntimeCell* Cell) const
 {
-	// If base class already returning critical. Early out.
-	if (Super::GetStreamingPerformanceForCell(Cell) == EWorldPartitionStreamingPerformance::Critical)
-	{
-		return EWorldPartitionStreamingPerformance::Critical;
-	}
-
 	check(Cell->GetBlockOnSlowLoading());
 	const double BlockOnSlowStreamingWarningRatio = GBlockOnSlowStreamingRatio * GBlockOnSlowStreamingWarningFactor;
 	
@@ -1809,9 +1791,9 @@ EWorldPartitionStreamingPerformance UWorldPartitionRuntimeSpatialHash::GetStream
 	{
 		const float LoadingRange = StreamingGrid->LoadingRange;
 
-		if (CellDataSpatialHash->IsBlockingSource())
+		if (CellDataSpatialHash->bCachedWasRequestedByBlockingSource)
 		{
-			const double Distance = FMath::Sqrt(CellDataSpatialHash->GetMinSquareDistanceToBlockingSource()) - ((double)StreamingGrid->GetCellSize(CellDataSpatialHash->Level) / 2);
+			const double Distance = FMath::Sqrt(CellDataSpatialHash->CachedMinSquareDistanceToBlockingSource) - ((double)StreamingGrid->GetCellSize(CellDataSpatialHash->HierarchicalLevel) / 2);
 
 			const double Ratio = Distance / LoadingRange;
 
diff --git a/Engine/Source/Runtime/Engine/Private/WorldPartition/WorldPartitionStreamingGeneration.cpp b/Engine/Source/Runtime/Engine/Private/WorldPartition/WorldPartitionStreamingGeneration.cpp
--- a/Engine/Source/Runtime/Engine/Private/WorldPartition/WorldPartitionStreamingGeneration.cpp
+++ b/Engine/Source/Runtime/Engine/Private/WorldPartition/WorldPartitionStreamingGeneration.cpp
@@ -1444,6 +1444,7 @@ bool UWorldPartition::GenerateContainerStreaming(const FGenerateStreamingParams&
 
 		StreamingPolicy->SetContainerResolver(StreamingGenerator.GetContainerResolver());
 		StreamingPolicy->PrepareActorToCellRemapping();
+		StreamingPolicy->SetShouldMergeStreamingSourceInfo(RuntimeHash->GetShouldMergeStreamingSourceInfo());
 		return true;
 	}
 
diff --git a/Engine/Source/Runtime/Engine/Private/WorldPartition/WorldPartitionStreamingPolicy.cpp b/Engine/Source/Runtime/Engine/Private/WorldPartition/WorldPartitionStreamingPolicy.cpp
--- a/Engine/Source/Runtime/Engine/Private/WorldPartition/WorldPartitionStreamingPolicy.cpp
+++ b/Engine/Source/Runtime/Engine/Private/WorldPartition/WorldPartitionStreamingPolicy.cpp
@@ -68,6 +68,7 @@ UWorldPartitionStreamingPolicy::UWorldPartitionStreamingPolicy(const FObjectInit
 	, ProcessedToActivateCells(0)
 	, ProcessedToLoadCells(0)
 	, bCriticalPerformanceRequestedBlockTillOnWorld(false)
+	, bShouldMergeStreamingSourceInfo(false)
 	, CriticalPerformanceBlockTillLevelStreamingCompletedEpoch(0)
 	, ServerStreamingStateEpoch(INT_MIN)
 	, ServerStreamingEnabledEpoch(INT_MIN)
@@ -185,8 +186,7 @@ bool UWorldPartitionStreamingPolicy::GetIntersectingCells(const TArray<FWorldPar
 	TArray<const UWorldPartitionRuntimeCell*> SortedCells = Cells.Array();
 	Algo::Sort(SortedCells, [&QueryCache](const UWorldPartitionRuntimeCell* CellA, const UWorldPartitionRuntimeCell* CellB)
 	{
-		const bool bCanUseSortingCache = false;
-		int32 SortCompare = CellA->SortCompare(CellB, bCanUseSortingCache);
+		int32 SortCompare = CellA->SortCompare(CellB);
 		if (SortCompare == 0)
 		{
 			// Closest distance (lower value is higher prio)
@@ -195,10 +195,8 @@ bool UWorldPartitionStreamingPolicy::GetIntersectingCells(const TArray<FWorldPar
 			{
 				return CellA->GetLevelPackageName().LexicalLess(CellB->GetLevelPackageName());
 			}
-			else
-			{
-				return Diff < 0;
-			}
+
+			return Diff < 0;
 		}
 		return SortCompare < 0;
 	});
@@ -468,14 +466,17 @@ void UWorldPartitionStreamingPolicy::UpdateStreamingState()
 			if (ActivatedCells.Contains(Cell))
 			{
 				// Update streaming source info for pending add to world cells
-				if (ActivatedCells.GetPendingAddToWorldCells().Contains(Cell))
+				if (bShouldMergeStreamingSourceInfo && ActivatedCells.GetPendingAddToWorldCells().Contains(Cell))
 				{
 					Cell->MergeStreamingSourceInfo();
 				}
 			}
 			else if (!ShouldSkipCellForPerformance(Cell) && !ShouldSkipDisabledHLODCell(Cell))
 			{
-				Cell->MergeStreamingSourceInfo();
+				if (bShouldMergeStreamingSourceInfo)
+				{
+					Cell->MergeStreamingSourceInfo();
+				}
 				ToActivateCells.Add(Cell);
 			}
 		}
@@ -491,7 +492,7 @@ void UWorldPartitionStreamingPolicy::UpdateStreamingState()
 			if (LoadedCells.Contains(Cell))
 			{
 				// Update streaming source info for pending load cells
-				if (!Cell->GetLevel())
+				if (bShouldMergeStreamingSourceInfo && !Cell->GetLevel())
 				{
 					Cell->MergeStreamingSourceInfo();
 				}
@@ -511,7 +512,10 @@ void UWorldPartitionStreamingPolicy::UpdateStreamingState()
 					}
 					else
 					{
-						Cell->MergeStreamingSourceInfo();
+						if (bShouldMergeStreamingSourceInfo)
+						{
+							Cell->MergeStreamingSourceInfo();
+						}
 						ToLoadCells.Add(Cell);
 					}
 				}
@@ -605,9 +609,12 @@ void UWorldPartitionStreamingPolicy::UpdateDebugCellsStreamingPriority(const TSe
 		TArray<const UWorldPartitionRuntimeCell*> Cells = ActivateStreamingCells.Array();
 		Cells.Append(LoadStreamingCells.Array());
 
-		for (const UWorldPartitionRuntimeCell* Cell : Cells)
+		if (bShouldMergeStreamingSourceInfo)
 		{
-			Cell->MergeStreamingSourceInfo();
+			for (const UWorldPartitionRuntimeCell* Cell : Cells)
+			{
+				Cell->MergeStreamingSourceInfo();
+			}
 		}
 
 		SortStreamingCellsByImportance(Cells);
diff --git a/Engine/Source/Runtime/Engine/Public/WorldPartition/RuntimeHashSet/WorldPartitionRuntimeHashSet.h b/Engine/Source/Runtime/Engine/Public/WorldPartition/RuntimeHashSet/WorldPartitionRuntimeHashSet.h
--- a/Engine/Source/Runtime/Engine/Public/WorldPartition/RuntimeHashSet/WorldPartitionRuntimeHashSet.h
+++ b/Engine/Source/Runtime/Engine/Public/WorldPartition/RuntimeHashSet/WorldPartitionRuntimeHashSet.h
@@ -111,16 +111,6 @@ public:
 	TArray<FRuntimePartitionStreamingData> RuntimeStreamingData;
 };
 
-UCLASS(Within = WorldPartitionRuntimeCell, MinimalAPI)
-class UWorldPartitionRuntimeCellDataSpatialHashSet : public UWorldPartitionRuntimeCellData
-{
-	GENERATED_BODY()
-
-public:
-	UPROPERTY()
-	int32 Level;
-};
-
 UCLASS(HideDropdown, MinimalAPI)
 class UWorldPartitionRuntimeHashSet : public UWorldPartitionRuntimeHash
 {
@@ -194,4 +184,4 @@ public:
 
 	UPROPERTY()
 	TArray<FRuntimePartitionStreamingData> RuntimeStreamingData;
-};
+};
\ No newline at end of file
diff --git a/Engine/Source/Runtime/Engine/Public/WorldPartition/WorldPartitionRuntimeCell.h b/Engine/Source/Runtime/Engine/Public/WorldPartition/WorldPartitionRuntimeCell.h
--- a/Engine/Source/Runtime/Engine/Public/WorldPartition/WorldPartitionRuntimeCell.h
+++ b/Engine/Source/Runtime/Engine/Public/WorldPartition/WorldPartitionRuntimeCell.h
@@ -321,11 +321,9 @@ protected:
 
 public:
 	//~Begin UWorldPartitionRuntimeCellData Proxy
-	inline bool ShouldResetStreamingSourceInfo() const { return RuntimeCellData->ShouldResetStreamingSourceInfo(); }
-	inline void ResetStreamingSourceInfo() const { RuntimeCellData->ResetStreamingSourceInfo(); }
 	inline void AppendStreamingSourceInfo(const FWorldPartitionStreamingSource& Source, const FSphericalSector& SourceShape) const { RuntimeCellData->AppendStreamingSourceInfo(Source, SourceShape); }
 	inline void MergeStreamingSourceInfo() const { RuntimeCellData->MergeStreamingSourceInfo(); }
-	ENGINE_API int32 SortCompare(const UWorldPartitionRuntimeCell* Other, bool bCanUseSortingCache = true) const;
+	ENGINE_API int32 SortCompare(const UWorldPartitionRuntimeCell* Other) const;
 	inline const FBox& GetContentBounds() const { return RuntimeCellData->GetContentBounds(); }
 	inline FBox GetCellBounds() const { return RuntimeCellData->GetCellBounds(); }
 	ENGINE_API virtual bool IsDebugShown() const;
diff --git a/Engine/Source/Runtime/Engine/Public/WorldPartition/WorldPartitionRuntimeCellData.h b/Engine/Source/Runtime/Engine/Public/WorldPartition/WorldPartitionRuntimeCellData.h
--- a/Engine/Source/Runtime/Engine/Public/WorldPartition/WorldPartitionRuntimeCellData.h
+++ b/Engine/Source/Runtime/Engine/Public/WorldPartition/WorldPartitionRuntimeCellData.h
@@ -20,12 +20,10 @@ class UWorldPartitionRuntimeCellData : public UObject
 	ENGINE_API void Serialize(FArchive& Ar);
 	//~End UObject Interface
 
-	inline bool ShouldResetStreamingSourceInfo() const { return CachedSourceInfoEpoch != StreamingSourceCacheEpoch; }
-
 	ENGINE_API virtual void ResetStreamingSourceInfo() const;
 	ENGINE_API virtual void AppendStreamingSourceInfo(const FWorldPartitionStreamingSource& Source, const FSphericalSector& SourceShape) const;
-	ENGINE_API virtual void MergeStreamingSourceInfo() const;
-	ENGINE_API virtual int32 SortCompare(const UWorldPartitionRuntimeCellData* InOther, bool bCanUseSortingCache = true) const;
+	ENGINE_API virtual void MergeStreamingSourceInfo() const {}
+	ENGINE_API virtual int32 SortCompare(const UWorldPartitionRuntimeCellData* InOther) const;
 
 	ENGINE_API virtual const FBox& GetContentBounds() const;
 	ENGINE_API virtual FBox GetCellBounds() const;
@@ -36,11 +34,20 @@ class UWorldPartitionRuntimeCellData : public UObject
 	static ENGINE_API int32 StreamingSourceCacheEpoch;
 	static inline void DirtyStreamingSourceCacheEpoch() { ++StreamingSourceCacheEpoch; }
 
-	// Source Priority
+	// Minimum affecting source priority
 	mutable uint8 CachedMinSourcePriority;
 
-	// Source Priorities
-	mutable TArray<float> CachedSourcePriorityWeights;
+	// Determine if the cell was requested by a blocking source
+	mutable bool bCachedWasRequestedByBlockingSource;
+
+	// Square distance from the cell to the closest blocking streaming source
+	mutable double CachedMinSquareDistanceToBlockingSource;
+
+	// Ratio used to determine the cell streaming performance status
+	mutable float CachedMinBlockOnSlowStreamingRatio;
+
+	// Spatial priority based on distance and angle from source
+	mutable double CachedMinSpatialSortingPriority;
 
 	// Epoch used to dirty cache
 	mutable int32 CachedSourceInfoEpoch;
@@ -51,5 +58,8 @@ class UWorldPartitionRuntimeCellData : public UObject
 	UPROPERTY()
 	int32 Priority;
 
+	UPROPERTY()
+	int32 HierarchicalLevel;
+
 	FStringTest DebugName;
 };
diff --git a/Engine/Source/Runtime/Engine/Public/WorldPartition/WorldPartitionRuntimeCellDataSpatialHash.h b/Engine/Source/Runtime/Engine/Public/WorldPartition/WorldPartitionRuntimeCellDataSpatialHash.h
--- a/Engine/Source/Runtime/Engine/Public/WorldPartition/WorldPartitionRuntimeCellDataSpatialHash.h
+++ b/Engine/Source/Runtime/Engine/Public/WorldPartition/WorldPartitionRuntimeCellDataSpatialHash.h
@@ -17,35 +17,23 @@ class UWorldPartitionRuntimeCellDataSpatialHash : public UWorldPartitionRuntimeC
 	ENGINE_API virtual void ResetStreamingSourceInfo() const override;
 	ENGINE_API virtual void AppendStreamingSourceInfo(const FWorldPartitionStreamingSource& Source, const FSphericalSector& SourceShape) const override;
 	ENGINE_API virtual void MergeStreamingSourceInfo() const override;
-	ENGINE_API virtual int32 SortCompare(const UWorldPartitionRuntimeCellData* InOther, bool bCanUseSortingCache = true) const override;
+	ENGINE_API virtual int32 SortCompare(const UWorldPartitionRuntimeCellData* InOther) const override;
 	ENGINE_API virtual FBox GetCellBounds() const override;
 	ENGINE_API virtual bool IsDebugShown() const override;
 	//~End UWorldPartitionRuntimeCellData
 
-	bool IsBlockingSource() const { return bCachedIsBlockingSource; }
-	double GetMinSquareDistanceToBlockingSource() const { return CachedMinSquareDistanceToBlockingSource; }
-
 	UPROPERTY()
 	FVector Position;
 
 	UPROPERTY()
 	float Extent;
 
-	UPROPERTY()
-	int32 Level;
-
 	UPROPERTY()
 	FName GridName;
 
 private:
 	ENGINE_API float ComputeSourceToCellAngleFactor(const FSphericalSector& SourceShape) const;
 
-	// Used to determine if cell was requested by blocking source
-	mutable bool bCachedIsBlockingSource;
-
-	// Square distance from the cell to the closest blocking streaming source
-	mutable double CachedMinSquareDistanceToBlockingSource;
-
 	// Square distance from the cell to the closest streaming source
 	mutable double CachedMinSquareDistanceToSource;
 
@@ -57,6 +45,9 @@ private:
 	// - The angle between the cell and each source orientation
 	mutable double CachedSourceSortingDistance;
 
+	// Source Priorities
+	mutable TArray<float> CachedSourcePriorityWeights;
+
 	// Square distance from the cell to  the intersecting streaming sources
 	mutable TArray<double> CachedSourceSquaredDistances;
 
diff --git a/Engine/Source/Runtime/Engine/Public/WorldPartition/WorldPartitionRuntimeHash.h b/Engine/Source/Runtime/Engine/Public/WorldPartition/WorldPartitionRuntimeHash.h
--- a/Engine/Source/Runtime/Engine/Public/WorldPartition/WorldPartitionRuntimeHash.h
+++ b/Engine/Source/Runtime/Engine/Public/WorldPartition/WorldPartitionRuntimeHash.h
@@ -22,6 +22,9 @@
 struct FHierarchicalLogArchive;
 class FWorldPartitionDraw2DContext;
 
+extern ENGINE_API float GBlockOnSlowStreamingRatio;
+extern ENGINE_API float GBlockOnSlowStreamingWarningFactor;
+
 UENUM()
 enum class EWorldPartitionStreamingPerformance : uint8
 {
@@ -156,11 +159,12 @@ public:
 	virtual void Draw3D(const TArray<FWorldPartitionStreamingSource>& Sources) const {}
 	virtual bool ContainsRuntimeHash(const FString& Name) const { return false; }
 	virtual bool IsStreaming3D() const { return true; }
+	virtual bool GetShouldMergeStreamingSourceInfo() const { return false; }
 
 protected:
 	static ENGINE_API URuntimeHashExternalStreamingObjectBase* CreateExternalStreamingObject(TSubclassOf<URuntimeHashExternalStreamingObjectBase> InClass, UObject* InOuter, FName InName, UWorld* InOwningWorld, UWorld* InOuterWorld);
 	ENGINE_API UWorldPartitionRuntimeCell* CreateRuntimeCell(UClass* CellClass, UClass* CellDataClass, const FString& CellName, const FString& CellInstanceSuffix, UObject* InOuter = nullptr);
-	virtual EWorldPartitionStreamingPerformance GetStreamingPerformanceForCell(const UWorldPartitionRuntimeCell* Cell) const { return EWorldPartitionStreamingPerformance::Good; }
+	virtual EWorldPartitionStreamingPerformance GetStreamingPerformanceForCell(const UWorldPartitionRuntimeCell* Cell) const;
 
 #if WITH_EDITOR
 	template <class T>
diff --git a/Engine/Source/Runtime/Engine/Public/WorldPartition/WorldPartitionRuntimeSpatialHash.h b/Engine/Source/Runtime/Engine/Public/WorldPartition/WorldPartitionRuntimeSpatialHash.h
--- a/Engine/Source/Runtime/Engine/Public/WorldPartition/WorldPartitionRuntimeSpatialHash.h
+++ b/Engine/Source/Runtime/Engine/Public/WorldPartition/WorldPartitionRuntimeSpatialHash.h
@@ -369,6 +369,7 @@ private:
 	ENGINE_API virtual void Draw3D(const TArray<FWorldPartitionStreamingSource>& Sources) const override;
 	ENGINE_API virtual bool ContainsRuntimeHash(const FString& Name) const override;
 	ENGINE_API virtual bool IsStreaming3D() const override;
+	ENGINE_API virtual bool GetShouldMergeStreamingSourceInfo() const override { return true; }
 
 	ENGINE_API void GetAlwaysLoadedStreamingCells(const FSpatialHashStreamingGrid& StreamingGrid, TSet<const UWorldPartitionRuntimeCell*>& Cells) const;
 	ENGINE_API const TMap<FName, const FSpatialHashStreamingGrid*>& GetNameToGridMapping() const;
diff --git a/Engine/Source/Runtime/Engine/Public/WorldPartition/WorldPartitionStreamingPolicy.h b/Engine/Source/Runtime/Engine/Public/WorldPartition/WorldPartitionStreamingPolicy.h
--- a/Engine/Source/Runtime/Engine/Public/WorldPartition/WorldPartitionStreamingPolicy.h
+++ b/Engine/Source/Runtime/Engine/Public/WorldPartition/WorldPartitionStreamingPolicy.h
@@ -93,6 +93,8 @@ public:
 	virtual bool InjectExternalStreamingObject(URuntimeHashExternalStreamingObjectBase* ExternalStreamingObject) { return true; }
 	virtual bool RemoveExternalStreamingObject(URuntimeHashExternalStreamingObjectBase* ExternalStreamingObject) { return true; }
 
+	virtual void SetShouldMergeStreamingSourceInfo(bool bInShouldMergeStreamingSourceInfo) { bShouldMergeStreamingSourceInfo = bInShouldMergeStreamingSourceInfo; }
+
 protected:
 	virtual void SetCellStateToLoaded(const UWorldPartitionRuntimeCell* InCell, int32& InOutMaxCellsToLoad);
 	virtual void SetCellStateToActivated(const UWorldPartitionRuntimeCell* InCell, int32& InOutMaxCellsToLoad);
@@ -120,8 +122,10 @@ protected:
 	// Used by UWorldPartitionSubsystem
 	UPROPERTY(Transient)
 	TArray<TObjectPtr<const UWorldPartitionRuntimeCell>> ToActivateCells;
+	
 	UPROPERTY(Transient)
 	TArray<TObjectPtr<const UWorldPartitionRuntimeCell>> ToLoadCells;
+	
 	int32 ProcessedToActivateCells;
 	int32 ProcessedToLoadCells;
 
@@ -139,6 +143,10 @@ private:
 	static FAutoConsoleVariableRef CVarForceUpdateFrameCount;
 
 	bool bCriticalPerformanceRequestedBlockTillOnWorld;
+	
+	UPROPERTY()
+	bool bShouldMergeStreamingSourceInfo;
+
 	int32 CriticalPerformanceBlockTillLevelStreamingCompletedEpoch;
 	int32 ServerStreamingStateEpoch;
 	int32 ServerStreamingEnabledEpoch;
