Fix relevancy issue with cluster union and already placed geometry collections
- Cluster union is destroyed as the client moves away but GC stays with their replicated proxies
- Transient state of the replicated proxy was preventing the GC from being added back to the cluster union
- Make sure we reset the transient data of the replicated proxies on the client

diff --git a/Engine/Source/Runtime/Engine/Classes/PhysicsEngine/ClusterUnionReplicatedProxyComponent.h b/Engine/Source/Runtime/Engine/Classes/PhysicsEngine/ClusterUnionReplicatedProxyComponent.h
--- a/Engine/Source/Runtime/Engine/Classes/PhysicsEngine/ClusterUnionReplicatedProxyComponent.h
+++ b/Engine/Source/Runtime/Engine/Classes/PhysicsEngine/ClusterUnionReplicatedProxyComponent.h
@@ -52,6 +52,9 @@ public:
 	UFUNCTION()
 	bool IsPendingDeletion() { return bIsPendingDeletion; }
 
+	/* reinitialize the internal transient state of the replicated proxy */
+	ENGINE_API void ResetTransientState();
+
 protected:
 
 	UFUNCTION()
diff --git a/Engine/Source/Runtime/Engine/Private/PhysicsEngine/ClusterUnionComponent.cpp b/Engine/Source/Runtime/Engine/Private/PhysicsEngine/ClusterUnionComponent.cpp
--- a/Engine/Source/Runtime/Engine/Private/PhysicsEngine/ClusterUnionComponent.cpp
+++ b/Engine/Source/Runtime/Engine/Private/PhysicsEngine/ClusterUnionComponent.cpp
@@ -1125,6 +1125,25 @@ void UClusterUnionComponent::HandleAddOrModifiedClusteredComponent(const FMapped
 	}
 }
 
+static UClusterUnionReplicatedProxyComponent* GetReplicatedProxyFromComponent(const UClusterUnionComponent* ClusterUnionComponent, const UPrimitiveComponent* Component)
+{
+	if (Component)
+	{
+		if (AActor* ComponentOwner = Component->GetOwner())
+		{
+			TInlineComponentArray<UClusterUnionReplicatedProxyComponent*> ReplicatedProxyComponents(ComponentOwner);
+			for (UClusterUnionReplicatedProxyComponent* ReplicatedProxyComponent : ReplicatedProxyComponents)
+			{
+				if (ReplicatedProxyComponent->GetParentClusterUnionComponent() == ClusterUnionComponent && ReplicatedProxyComponent->GetChildClusteredComponent() == Component)
+				{
+					return ReplicatedProxyComponent;
+				}
+			}
+		}
+	}
+	return nullptr;
+}
+
 void UClusterUnionComponent::HandleRemovedClusteredComponent(TObjectKey<UPrimitiveComponent> RemovedComponent)
 {
 	if (FClusteredComponentData* Data = PerComponentData.Find(RemovedComponent))
@@ -1150,6 +1169,16 @@ void UClusterUnionComponent::HandleRemovedClusteredComponent(TObjectKey<UPrimiti
 				ProxyComponent->DestroyComponent();
 			}
 		}
+		else
+		{
+			// On the client, we need to reset the state of the replicated proxy has it may outlive the cluster union
+			// and may still be in existence when the cluster union is created again 
+			// Note: we need to get the replciated component from the actor on the client because the PerComponentData do not set the replciated proxy pointer (unlike the server )
+			if (UClusterUnionReplicatedProxyComponent* ProxyComponent = GetReplicatedProxyFromComponent(this, RemovedComponent.ResolveObjectPtr()))
+			{
+				ProxyComponent->ResetTransientState();
+			}
+		}
 
 		if (AActor* Owner = Data->Owner.Get())
 		{
diff --git a/Engine/Source/Runtime/Engine/Private/PhysicsEngine/ClusterUnionReplicatedProxyComponent.cpp b/Engine/Source/Runtime/Engine/Private/PhysicsEngine/ClusterUnionReplicatedProxyComponent.cpp
--- a/Engine/Source/Runtime/Engine/Private/PhysicsEngine/ClusterUnionReplicatedProxyComponent.cpp
+++ b/Engine/Source/Runtime/Engine/Private/PhysicsEngine/ClusterUnionReplicatedProxyComponent.cpp
@@ -199,6 +199,16 @@ void UClusterUnionReplicatedProxyComponent::FlushNetDormancyIfNeeded()
 	}
 }
 
+void UClusterUnionReplicatedProxyComponent::ResetTransientState()
+{
+	LastSyncedBoneIds.Reset();
+
+	if (AActor* Owner = GetOwner(); Owner && DeferSetChildToParentHandle.IsValid())
+	{
+		Owner->GetWorldTimerManager().ClearTimer(DeferSetChildToParentHandle);
+	}
+}
+
 void UClusterUnionReplicatedProxyComponent::AddComponentToCluster()
 {
 	if (!ParentClusterUnion.IsValid() || !ChildClusteredComponent.IsValid() || IsPendingDeletion())
@@ -261,4 +271,4 @@ void UClusterUnionReplicatedProxyComponent::DeferSetChildToParentChildUntilClust
 	{
 		DeferSetChildToParentHandle = Owner->GetWorldTimerManager().SetTimerForNextTick(this, &UClusterUnionReplicatedProxyComponent::DeferSetChildToParentChildUntilClusteredComponentInParentUnion);
 	}
-}
\ No newline at end of file
+}
