[shaders] resubmit - enable preprocessed job cache by default; change prioritization of ShaderJobTasks in non-cook cases to background to not occupy fg worker threads (for improved responsiveness in editor)
- additional fix for use after free in removal of duplicate jobs
- additional fix for case where job cache is disabled - don't set bCachePreprocessed to true if job cache itself is off
- additional fix for cloning a job - preprocess output also needs to be copied as it would have run before cloning and the job submit task is not called for these cloned jobs

diff --git a/Engine/Source/Runtime/Engine/Private/ShaderCompiler/ShaderCompiler.cpp b/Engine/Source/Runtime/Engine/Private/ShaderCompiler/ShaderCompiler.cpp
--- a/Engine/Source/Runtime/Engine/Private/ShaderCompiler/ShaderCompiler.cpp
+++ b/Engine/Source/Runtime/Engine/Private/ShaderCompiler/ShaderCompiler.cpp
@@ -122,7 +122,7 @@ static FAutoConsoleVariableRef CVarShaderCompilerJobCacheOverflowReducePercent(
 
 static TAutoConsoleVariable<bool> CVarPreprocessedJobCache(
 	TEXT("r.ShaderCompiler.PreprocessedJobCache"),
-	false,
+	true,
 	TEXT("If enabled will shader compile jobs will be preprocessed at submission time in the cook process (when the job is queued) and generate job input hashes based on preprocessed source."),
 	ECVF_Default
 );
@@ -1329,6 +1329,7 @@ static FShaderCommonCompileJob* CloneJob_Single(const FShaderCompileJob* SrcJob)
 	Job->ShaderParameters = SrcJob->ShaderParameters;
 	Job->PendingShaderMap = SrcJob->PendingShaderMap;
 	Job->Input = SrcJob->Input;
+	Job->PreprocessOutput = SrcJob->PreprocessOutput;
 	if (SrcJob->bInputHashSet)
 	{
 		Job->InputHash = SrcJob->InputHash;
@@ -1347,6 +1348,7 @@ static FShaderCommonCompileJob* CloneJob_Pipeline(const FShaderPipelineCompileJo
 	for(int32 i = 0; i < SrcJob->StageJobs.Num(); ++i)
 	{
 		Job->StageJobs[i]->Input = SrcJob->StageJobs[i]->Input;
+		Job->StageJobs[i]->PreprocessOutput = SrcJob->StageJobs[i]->PreprocessOutput;
 	}
 
 	if (SrcJob->bInputHashSet)
@@ -1540,14 +1542,14 @@ int32 FShaderJobCache::RemoveAllPendingJobsWithId(uint32 InId)
 						JobData.DuplicateJobsWaitList = JobData.DuplicateJobsWaitList->NextLink;
 					}
 
+					// This removes the current job (at DuplicateIndex), so we don't increment in this case
+					RemoveDuplicateJob(DuplicateJob);
+
 					// Duplicate jobs are in their own list, not one of the priority lists, so don't use UnlinkJobWithPriority
 					check(DuplicateJob->PendingPriority == EShaderCompileJobPriority::None);
 					Unlink(*DuplicateJob);
 					RemoveJob(DuplicateJob);
 					++NumRemoved;
-
-					// This removes the current job (at DuplicateIndex), so we don't increment in this case
-					RemoveDuplicateJob(DuplicateJob);
 				}
 				else
 				{
@@ -1729,6 +1731,7 @@ void FShaderJobCache::SubmitJobs(const TArray<FShaderCommonCompileJobPtr>& InJob
 		{
 			for (FShaderCommonCompileJobPtr Job : InJobs)
 			{
+				UE::Tasks::ETaskPriority Prio = IsRunningCookCommandlet() ? UE::Tasks::ETaskPriority::Normal : UE::Tasks::ETaskPriority::BackgroundNormal;
 				UE::Tasks::Launch(UE_SOURCE_LOCATION, [Job, this]()
 				{
 					TRACE_CPUPROFILER_EVENT_SCOPE(ShaderJobTask);
@@ -1747,7 +1750,7 @@ void FShaderJobCache::SubmitJobs(const TArray<FShaderCommonCompileJobPtr>& InJob
 					SubmitJob(Job);
 
 					Job->TimeTaskSubmitJobs = FPlatformTime::Seconds() - TimeStart;
-				});
+				}, Prio);
 			}
 		}
 		else
@@ -8308,7 +8311,7 @@ void GlobalBeginCompileShader(
 	checkf(Format, TEXT("Shader format %s cannot be found"), *ShaderFormatName.ToString());
 	Format->ModifyShaderCompilerInput(Input);
 
-	if (CVarPreprocessedJobCache.GetValueOnAnyThread())
+	if (ShaderCompiler::IsJobCacheEnabled() && CVarPreprocessedJobCache.GetValueOnAnyThread())
 	{
 		Input.bCachePreprocessed = true;
 	}
