OodleWav -> UEWavComp

diff --git a/Engine/Config/BaseEditor.ini b/Engine/Config/BaseEditor.ini
--- a/Engine/Config/BaseEditor.ini
+++ b/Engine/Config/BaseEditor.ini
@@ -255,7 +255,7 @@ bExpandClassPickerClassList=False
 Class=UnrealEd.ActorFactory
 
 [AudioImporter]
-EnableOodleWAV=True
+EnableUEWavComp=True
 
 [TextureImporter]
 AllowNonPowerOfTwoTextures=True
diff --git a/Engine/Source/Runtime/Engine/Engine.Build.cs b/Engine/Source/Runtime/Engine/Engine.Build.cs
--- a/Engine/Source/Runtime/Engine/Engine.Build.cs
+++ b/Engine/Source/Runtime/Engine/Engine.Build.cs
@@ -166,7 +166,7 @@ public class Engine : ModuleRules
 			// The SparseVolumeTexture module containing the importer is only loaded and used in the editor.
 			DynamicallyLoadedModuleNames.Add("SparseVolumeTexture");
 			AddEngineThirdPartyPrivateStaticDependencies(Target,
-				"OodleWAV"
+				"UEWavComp"
 				);
 		}
 
diff --git a/Engine/Source/Runtime/Engine/Private/SoundWave.cpp b/Engine/Source/Runtime/Engine/Private/SoundWave.cpp
--- a/Engine/Source/Runtime/Engine/Private/SoundWave.cpp
+++ b/Engine/Source/Runtime/Engine/Private/SoundWave.cpp
@@ -34,7 +34,7 @@
 #include "UObject/ObjectSaveContext.h"
 #include "ISoundWaveCloudStreaming.h"
 #if WITH_EDITORONLY_DATA
-#include "oowav.h"
+#include "uewav.h"
 #endif
 
 static int32 SoundWaveDefaultLoadingBehaviorCVar = static_cast<int32>(ESoundWaveLoadingBehavior::LoadOnDemand);
@@ -4857,14 +4857,14 @@ TFuture<FSharedBuffer> USoundWave::FEditorAudioBulkData::GetPayload() const
 	}
 	if (*WaveInfo.pFormatTag == FWaveModInfo::WAVE_INFO_FORMAT_OODLE_WAVE)
 	{
-		// Convert OodleWave data back to a WAV file
+		// Convert UEWavComp data back to a WAV file
 		*WaveInfo.pFormatTag = FWaveModInfo::WAVE_INFO_FORMAT_PCM;
 		int16* samples = (int16*)WaveInfo.SampleDataStart;
 		int64 num_samples = WaveInfo.GetNumSamples();
 		int64 num_channels = *WaveInfo.pChannels;
 		TArray<int16> scratch_buffer;
 		scratch_buffer.AddUninitialized(num_samples);
-		oowav_decode16(samples, scratch_buffer.GetData(), num_samples, num_channels);
+		uewav_decode16(samples, scratch_buffer.GetData(), num_samples, num_channels);
 	}
 	TPromise<FSharedBuffer> promise;
 	promise.EmplaceValue(Buffer);
@@ -4885,16 +4885,16 @@ void USoundWave::FEditorAudioBulkData::UpdatePayload(FSharedBuffer InPayload, UO
 	}
 	if (*WaveInfo.pFormatTag == FWaveModInfo::WAVE_INFO_FORMAT_PCM)
 	{
-		bool bEnableOodleWAV = false;
-		GConfig->GetBool(TEXT("AudioImporter"), TEXT("EnableOodleWAV"), bEnableOodleWAV, GEditorIni);
-		if (bEnableOodleWAV)
+		bool bEnableUEWavComp = false;
+		GConfig->GetBool(TEXT("AudioImporter"), TEXT("EnableUEWavComp"), bEnableUEWavComp, GEditorIni);
+		if (bEnableUEWavComp)
 		{
 			int16* samples = (int16*)WaveInfo.SampleDataStart;
 			int64 num_samples = WaveInfo.GetNumSamples();
 			int64 num_channels = *WaveInfo.pChannels;
 			TArray<int16> scratch_buffer;
 			scratch_buffer.AddUninitialized(num_samples);
-			oowav_encode16(samples, scratch_buffer.GetData(), num_samples, num_channels);
+			uewav_encode16(samples, scratch_buffer.GetData(), num_samples, num_channels);
 			*WaveInfo.pFormatTag = FWaveModInfo::WAVE_INFO_FORMAT_OODLE_WAVE;
 		}
 	}
diff --git a/Engine/Source/Runtime/OodleWAV/OodleWAV.Build.cs b/Engine/Source/Runtime/OodleWAV/OodleWAV.Build.cs
deleted file mode 100644
diff --git a/Engine/Source/Runtime/OodleWAV/Source/oowav.h b/Engine/Source/Runtime/OodleWAV/Source/oowav.h
deleted file mode 100644
diff --git a/Engine/Source/Runtime/UEWavComp/Source/uewav.h b/Engine/Source/Runtime/UEWavComp/Source/uewav.h
new file mode 100644
--- /dev/null
+++ b/Engine/Source/Runtime/UEWavComp/Source/uewav.h
@@ -0,0 +1,21 @@
+// Copyright Epic Games, Inc. All Rights Reserved.
+/*
+    Oodle WAVE is a simple data transform on PCM audio data that is modifies the data to be more compressible by Oodle. 
+*/
+#pragma once
+
+#ifdef __cplusplus
+extern "C" {
+#endif
+
+// Encodes/decodes 16-bit samples. samples is the input/output buffer, scratch_buffer is a temporary buffer.  
+// operates in place. samples is both input and output.
+// scratch_buffer should be the same size as samples.
+// num_samples is the number of samples in the buffer
+// num_channels is the number of channels in the buffer.
+void uewav_encode16(short *samples, short *scratch_buffer, long long num_samples, long long num_channels);
+void uewav_decode16(short *samples, short *scratch_buffer, long long num_samples, long long num_channels);
+
+#ifdef __cplusplus
+}
+#endif
\ No newline at end of file
diff --git a/Engine/Source/Runtime/UEWavComp/UEWavComp.Build.cs b/Engine/Source/Runtime/UEWavComp/UEWavComp.Build.cs
new file mode 100644
--- /dev/null
+++ b/Engine/Source/Runtime/UEWavComp/UEWavComp.Build.cs
@@ -0,0 +1,38 @@
+// Copyright Epic Games, Inc. All Rights Reserved.
+using UnrealBuildTool;
+using System.IO;
+
+public class UEWavComp : ModuleRules
+{
+	public UEWavComp(ReadOnlyTargetRules Target) : base(Target)
+	{
+		Type = ModuleType.External;
+
+		bool bWithLib = false;
+		if (Target.Platform == UnrealTargetPlatform.Win64 && Target.WindowsPlatform.Architecture != UnrealArch.Arm64)
+		{
+			string LibPath = Path.Combine(ModuleDirectory, "Lib");
+			PublicAdditionalLibraries.Add(Path.Combine(LibPath, "uewav_w64.lib"));
+			bWithLib = true;
+		}
+		else if (Target.Platform == UnrealTargetPlatform.Mac)
+		{
+			string LibPath = Path.Combine(ModuleDirectory, "Lib");
+			PublicAdditionalLibraries.Add(Path.Combine(LibPath, "uewav_mac.a"));
+			bWithLib = true;
+		}
+		else if (Target.IsInPlatformGroup(UnrealPlatformGroup.Unix))
+		{
+			string LibPath = Path.Combine(ModuleDirectory, "Lib/Unix", Target.Architecture.LinuxName);
+			PublicAdditionalLibraries.Add(Path.Combine(LibPath, "uewav_linux.a"));
+			bWithLib = true;
+		}
+
+		if (bWithLib)
+		{
+			PublicSystemIncludePaths.Add(Path.Combine(ModuleDirectory, "Source"));
+		}
+
+		PublicDefinitions.Add("WITH_OOWAV=" + (bWithLib ? '1' : '0'));
+	}
+}
