Fix potential issue with unused virtuals being called from constructor/destructor

#rnx

diff --git a/Engine/Source/Runtime/Engine/Private/Tests/Loading/AsyncLoadingTests_Shared.h b/Engine/Source/Runtime/Engine/Private/Tests/Loading/AsyncLoadingTests_Shared.h
--- a/Engine/Source/Runtime/Engine/Private/Tests/Loading/AsyncLoadingTests_Shared.h
+++ b/Engine/Source/Runtime/Engine/Private/Tests/Loading/AsyncLoadingTests_Shared.h
@@ -85,11 +85,11 @@ public:
 	UAsyncLoadingTests_Shared* Object1 = nullptr;
 	UAsyncLoadingTests_Shared* Object2 = nullptr;
 
-	virtual void CreateObjects();
-	virtual void MutateObjects();
-	virtual void SaveObjects();
-	virtual void LoadObjects();
-	virtual void CleanupObjects();
+	void CreateObjects();
+	void MutateObjects();
+	void SaveObjects();
+	void LoadObjects();
+	void CleanupObjects();
 
 	FLoadingTestsScope(FAutomationTestBase* InAutomationTest)
 		: AutomationTest(*InAutomationTest)
