Fixed memory stomp in CPU particles due to bad malloc size.

diff --git a/Engine/Source/Runtime/Engine/Private/Particles/ParticleSystemRender.cpp b/Engine/Source/Runtime/Engine/Private/Particles/ParticleSystemRender.cpp
--- a/Engine/Source/Runtime/Engine/Private/Particles/ParticleSystemRender.cpp
+++ b/Engine/Source/Runtime/Engine/Private/Particles/ParticleSystemRender.cpp
@@ -929,7 +929,7 @@ void FDynamicSpriteEmitterData::GetDynamicMeshElementsEmitter(const FParticleSys
 
 						if (bSort)
 						{
-							ParticleOrder = (FParticleOrder*)FConcurrentLinearAllocator::Malloc(ParticleCount, 16);
+							ParticleOrder = (FParticleOrder*)FConcurrentLinearAllocator::Malloc(sizeof(FParticleOrder) * ParticleCount, 16);
 							SortSpriteParticles(SourceData->SortMode, SourceData->bUseLocalSpace, SourceData->ActiveParticleCount, 
 								SourceData->DataContainer.ParticleData, SourceData->ParticleStride, SourceData->DataContainer.ParticleIndices,
 								View, Proxy->GetLocalToWorld(), ParticleOrder);
