Replace an ensure with an error log in CameraShakeSourceComponent::Start

- an ensure just of a user error not setting up a CameraShakeSourceActor properly seems a bit too much. An error log should be more appropriate.

#rnx

diff --git a/Engine/Source/Runtime/Engine/Private/Camera/CameraShakeSourceComponent.cpp b/Engine/Source/Runtime/Engine/Private/Camera/CameraShakeSourceComponent.cpp
--- a/Engine/Source/Runtime/Engine/Private/Camera/CameraShakeSourceComponent.cpp
+++ b/Engine/Source/Runtime/Engine/Private/Camera/CameraShakeSourceComponent.cpp
@@ -17,6 +17,7 @@
 #include "Widgets/Notifications/SNotificationList.h"
 #endif
 
+DEFINE_LOG_CATEGORY_STATIC(LogCameraShakeSourceComponent, Log, All);
 
 UCameraShakeSourceComponent::UCameraShakeSourceComponent(const FObjectInitializer& ObjectInitializer)
 	: Super(ObjectInitializer)
@@ -74,10 +75,14 @@ void UCameraShakeSourceComponent::EndPlay(const EEndPlayReason::Type EndPlayReas
 
 void UCameraShakeSourceComponent::Start()
 {
-	if (ensureMsgf(CameraShake.Get() != nullptr, TEXT("No camera shake was specified on this source!")))
+	if (CameraShake.Get() != nullptr)
 	{
 		StartCameraShake(CameraShake);
 	}
+	else
+	{
+		UE_LOG(LogCameraShakeSourceComponent, Error, TEXT("%s: No camera shake was specified on this source. Me = %s, Owner = %s"), ANSI_TO_TCHAR(__FUNCTION__), *this->GetFullName(), *GetOwner()->GetFullName());
+	}
 }
 
 void UCameraShakeSourceComponent::StartCameraShake(TSubclassOf<UCameraShakeBase> InCameraShake, float Scale, ECameraShakePlaySpace PlaySpace, FRotator UserPlaySpaceRot)
