- Add optional CSV memory usage output for FX, mode is controlled with fx.DetailedCSVStats.MemoryMode and disabled by default as it impacts performance
- Implement GetApproxMemoryUsage for Niagara and use this function in the debug HUD so we display Cascade memory usage

diff --git a/Engine/Plugins/FX/Niagara/Source/Niagara/Private/NiagaraComponent.cpp b/Engine/Plugins/FX/Niagara/Source/Niagara/Private/NiagaraComponent.cpp
--- a/Engine/Plugins/FX/Niagara/Source/Niagara/Private/NiagaraComponent.cpp
+++ b/Engine/Plugins/FX/Niagara/Source/Niagara/Private/NiagaraComponent.cpp
@@ -761,8 +761,12 @@ void UNiagaraComponent::ReleaseToPool()
 
 uint32 UNiagaraComponent::GetApproxMemoryUsage() const
 {
-	// TODO: implement memory usage for the component pool statistics
-	return 1;
+	uint32 MemoryBytes = 0;
+	if (SystemInstanceController)
+	{
+		MemoryBytes += uint32(SystemInstanceController->GetTotalBytesUsed());
+	}
+	return MemoryBytes;
 }
 
 void UNiagaraComponent::ActivateSystem(bool bFlagAsJustAttached)
diff --git a/Engine/Plugins/FX/Niagara/Source/Niagara/Private/NiagaraDebugHud.cpp b/Engine/Plugins/FX/Niagara/Source/Niagara/Private/NiagaraDebugHud.cpp
--- a/Engine/Plugins/FX/Niagara/Source/Niagara/Private/NiagaraDebugHud.cpp
+++ b/Engine/Plugins/FX/Niagara/Source/Niagara/Private/NiagaraDebugHud.cpp
@@ -1116,7 +1116,7 @@ void FNiagaraDebugHud::GatherSystemInfo()
 		}
 
 		// Track rough memory usage
-		const int64 BytesUsed = SystemInstanceController.IsValid() ? SystemInstanceController->GetTotalBytesUsed() : 0;
+		const int64 BytesUsed = FXComponent->GetApproxMemoryUsage();
 		SystemDebugInfo.TotalBytes += BytesUsed;
 		GlobalTotalBytes += BytesUsed;
 
diff --git a/Engine/Source/Runtime/Engine/Classes/Particles/ParticleSystem.h b/Engine/Source/Runtime/Engine/Classes/Particles/ParticleSystem.h
--- a/Engine/Source/Runtime/Engine/Classes/Particles/ParticleSystem.h
+++ b/Engine/Source/Runtime/Engine/Classes/Particles/ParticleSystem.h
@@ -151,6 +151,7 @@ public:
 	FName CSVStat_Activation = NAME_None;
 	FName CSVStat_Waits = NAME_None;
 	FName CSVStat_Culled = NAME_None;
+	FName CSVStat_MemoryKB = NAME_None;
 #endif
 #endif
 
diff --git a/Engine/Source/Runtime/Engine/Private/Particles/ParticleComponents.cpp b/Engine/Source/Runtime/Engine/Private/Particles/ParticleComponents.cpp
--- a/Engine/Source/Runtime/Engine/Private/Particles/ParticleComponents.cpp
+++ b/Engine/Source/Runtime/Engine/Private/Particles/ParticleComponents.cpp
@@ -142,6 +142,7 @@ void UFXSystemAsset::PostInitProperties()
 	CSVStat_Activation = *FString::Printf(TEXT("Activation/%s"), *GetFName().ToString());
 	CSVStat_Waits = *FString::Printf(TEXT("Waits/%s"), *GetFName().ToString());
 	CSVStat_Culled = *FString::Printf(TEXT("Culled/%s"), *GetFName().ToString());
+	CSVStat_MemoryKB = *FString::Printf(TEXT("MemoryKB/%s"), *GetFName().ToString());
 #endif
 }
 
@@ -8434,6 +8435,7 @@ uint32 UParticleSystemComponent::GetApproxMemoryUsage()const
 		if (FParticleDynamicData* DynamicData = PSysSceneProxy->GetDynamicData())
 		{
 			MemUsage += DynamicData->GetMemoryFootprint();
+		#if 0
 			for (FDynamicEmitterDataBase* DynEmitterData : DynamicData->DynamicEmitterDataArray)
 			{
 				if (DynEmitterData)
@@ -8447,6 +8449,7 @@ uint32 UParticleSystemComponent::GetApproxMemoryUsage()const
 					MemUsage += MemCounter.Max;
 				}
 			}
+		#endif
 		}
 	}
 
diff --git a/Engine/Source/Runtime/Engine/Private/Particles/ParticlePerfStats.cpp b/Engine/Source/Runtime/Engine/Private/Particles/ParticlePerfStats.cpp
--- a/Engine/Source/Runtime/Engine/Private/Particles/ParticlePerfStats.cpp
+++ b/Engine/Source/Runtime/Engine/Private/Particles/ParticlePerfStats.cpp
@@ -21,7 +21,7 @@
 #include "CanvasTypes.h"
 #include "Engine/Font.h"
 #include "Particles/ParticleSystemComponent.h"
-
+#include "Serialization/ArchiveCountMem.h"
 
 DECLARE_STATS_GROUP(TEXT("ParticleStats"), STATGROUP_ParticleStats, STATCAT_Advanced);
 DECLARE_CYCLE_STAT(TEXT("Particle Stats Tick [GT]"), STAT_ParticleStats_TickGT, STATGROUP_ParticleStats);
@@ -1115,6 +1115,17 @@ static FAutoConsoleVariable CVarWriteDetailedCSVStats(
 	ECVF_Default | ECVF_RenderThreadSafe
 );
 
+int32 GFXDetailedCSVMemorySMode = 0;
+static FAutoConsoleVariableRef CVarFXDetailedCSVMemoryStats(
+	TEXT("fx.DetailedCSVStats.MemoryMode"),
+	GFXDetailedCSVMemorySMode,
+	TEXT("Gathers approximate memory information depending on the mode.\n")
+	TEXT("0 = Disabled (default).\n")
+	TEXT("1 = Minimal information (small performance impact).\n")
+	TEXT("2 = Full information (large performance impact)."),
+	ECVF_Default
+);
+
 void OnDetailedCSVStatsEnabledChanged(IConsoleVariable* Variable)
 {
 	const bool bDetailedCSVStats = CVarWriteDetailedCSVStats->GetBool();
@@ -1170,6 +1181,51 @@ bool FParticlePerfStatsListener_CSVProfiler::Tick()
 	if (FCsvProfiler* CSVProfiler = FCsvProfiler::Get())
 	{
 #if WITH_PER_SYSTEM_PARTICLE_PERF_STATS
+		// Very slow, this gives coverage over everything but is bad for performance
+		if (GFXDetailedCSVMemorySMode != 0)
+		{
+			TMap<UFXSystemAsset*, uint64> MemoryUsage;
+			for (TObjectIterator<UFXSystemComponent> It; It; ++It)
+			{
+				UFXSystemComponent* FXComponent = *It;
+				UFXSystemAsset* FXAsset = FXComponent ? FXComponent->GetFXSystemAsset() : nullptr;
+				if (!IsValidChecked(FXComponent) || !IsValid(FXAsset) || FXComponent->IsUnreachable() || FXComponent->HasAnyFlags(EObjectFlags::RF_ClassDefaultObject))
+				{
+					continue;
+				}
+
+				if (FXAsset->CSVStat_MemoryKB.IsNone())
+				{
+					continue;
+				}
+				const bool bFullMemory = GFXDetailedCSVMemorySMode == 2;
+
+				uint64& MemoryBytes = MemoryUsage.FindOrAdd(FXAsset);
+				if (bFullMemory)
+				{
+					if (MemoryBytes == 0)
+					{
+						MemoryBytes += FArchiveCountMem(FXAsset).GetMax();
+					}
+					MemoryBytes += FArchiveCountMem(FXComponent).GetMax();
+
+					FResourceSizeEx CompResSize = FResourceSizeEx(EResourceSizeMode::EstimatedTotal);
+					FXComponent->GetResourceSizeEx(CompResSize);
+					MemoryBytes += CompResSize.GetTotalMemoryBytes();
+				}
+				else
+				{
+					MemoryBytes += FXComponent->GetApproxMemoryUsage();
+				}
+			}
+
+			for (auto OutputIt=MemoryUsage.CreateConstIterator(); OutputIt; ++OutputIt)
+			{
+				const int32 MemoryKB = (int32)FMath::DivideAndRoundUp(OutputIt.Value(), 1024ull);
+				CSVProfiler->RecordCustomStat(OutputIt.Key()->CSVStat_MemoryKB, CSV_CATEGORY_INDEX(Particles), MemoryKB, ECsvCustomStatOp::Set);
+			}
+		}
+
 		FParticlePerfStatsManager::ForAllSystemStats(
 			[&](TWeakObjectPtr<const UFXSystemAsset>& WeakSystem, TUniquePtr<FParticlePerfStats>& Stats)
 			{
@@ -1177,12 +1233,12 @@ bool FParticlePerfStatsListener_CSVProfiler::Tick()
 				{
 					if (const UFXSystemAsset* System = WeakSystem.Get())
 					{
-						float TotalTime = FPlatformTime::ToMilliseconds64(Stats->GetGameThreadStats().GetTotalCycles()) * 1000.0f;
-						float GTTime = FPlatformTime::ToMilliseconds64(Stats->GetGameThreadStats().GetTotalCycles_GTOnly()) * 1000.0f;
-						float AvgTime = FPlatformTime::ToMilliseconds64(Stats->GetGameThreadStats().GetPerInstanceAvgCycles()) * 1000.0f;
-						int32 Count = (int32)Stats->GetGameThreadStats().NumInstances;
-						float Activation = FPlatformTime::ToMilliseconds64(Stats->GetGameThreadStats().ActivationCycles) * 1000.0f;
-						float Wait = FPlatformTime::ToMilliseconds64(Stats->GetGameThreadStats().WaitCycles) * 1000.0f;
+						const float TotalTime = FPlatformTime::ToMilliseconds64(Stats->GetGameThreadStats().GetTotalCycles()) * 1000.0f;
+						const float GTTime = FPlatformTime::ToMilliseconds64(Stats->GetGameThreadStats().GetTotalCycles_GTOnly()) * 1000.0f;
+						const float AvgTime = FPlatformTime::ToMilliseconds64(Stats->GetGameThreadStats().GetPerInstanceAvgCycles()) * 1000.0f;
+						const int32 Count = (int32)Stats->GetGameThreadStats().NumInstances;
+						const float Activation = FPlatformTime::ToMilliseconds64(Stats->GetGameThreadStats().ActivationCycles) * 1000.0f;
+						const float Wait = FPlatformTime::ToMilliseconds64(Stats->GetGameThreadStats().WaitCycles) * 1000.0f;
 
 						CSVProfiler->RecordCustomStat(System->CSVStat_Total, CSV_CATEGORY_INDEX(Particles), TotalTime, ECsvCustomStatOp::Set);
 						CSVProfiler->RecordCustomStat(System->CSVStat_GTOnly, CSV_CATEGORY_INDEX(Particles), GTTime, ECsvCustomStatOp::Set);
