nop comments only

#rnx

diff --git a/Engine/Source/Runtime/Engine/Private/BatchedElements.cpp b/Engine/Source/Runtime/Engine/Private/BatchedElements.cpp
--- a/Engine/Source/Runtime/Engine/Private/BatchedElements.cpp
+++ b/Engine/Source/Runtime/Engine/Private/BatchedElements.cpp
@@ -581,6 +581,9 @@ void FBatchedElements::PrepareShaders(
 	// used to mask individual channels and desaturate
 	FMatrix ColorWeights( FPlane(1, 0, 0, 0), FPlane(0, 1, 0, 0), FPlane(0, 0, 1, 0), FPlane(0, 0, 0, 0) );
 
+	// this is the inverse of the gamma of the target; 
+	//= 1.0 if "NoGamma" (EnableGammaCorrection(false))
+	//= 1.0/2.2 to output with LinearToSRGB
 	float GammaToUse = Gamma;
 
 	ESimpleElementBlendMode MaskedBlendMode = SE_BLEND_Opaque;
@@ -787,6 +790,7 @@ void FBatchedElements::PrepareShaders(
 	
 				if (FMath::Abs(Gamma - 1.0f) < UE_KINDA_SMALL_NUMBER)
 				{
+					// runs "Main"
 					TShaderMapRef<FSimpleElementPS> PixelShader(GetGlobalShaderMap(FeatureLevel));
 					GraphicsPSOInit.BoundShaderState.PixelShaderRHI = PixelShader.GetPixelShader();
 
@@ -796,7 +800,11 @@ void FBatchedElements::PrepareShaders(
 				}
 				else
 				{
+					// runs "GammaMain"
 					TShaderRef<FSimpleElementGammaBasePS> BasePixelShader;
+
+					// these shaders differ in setting SRGB_INPUT_TEXTURE, which is ignored
+					//  so they are in fact the same
 					if (Texture->bSRGB)
 					{
 						TShaderMapRef<FSimpleElementGammaPS_SRGB> PixelShader_SRGB(GetGlobalShaderMap(FeatureLevel));
diff --git a/Engine/Source/Runtime/Engine/Public/SimpleElementShaders.h b/Engine/Source/Runtime/Engine/Public/SimpleElementShaders.h
--- a/Engine/Source/Runtime/Engine/Public/SimpleElementShaders.h
+++ b/Engine/Source/Runtime/Engine/Public/SimpleElementShaders.h
@@ -117,6 +117,8 @@ public:
 
 	static void ModifyCompilationEnvironment(const FGlobalShaderPermutationParameters& Parameters, FShaderCompilerEnvironment& OutEnvironment)
 	{
+		// SRGB_INPUT_TEXTURE is now ignored, this has no effect
+		// input texture gamma correction is done automatically by the sampler, controlled by TexCreate_SRGB
 		OutEnvironment.SetDefine(TEXT("SRGB_INPUT_TEXTURE"), bSRGBTexture);
 	}
 };
