Some simple fixes to PCM/ADPCM decoder.

- process seek requests before checking for end of file and earlying out. Prevents external users from seeking if the decoder has reached the end of the file.
- Properly memzero the decode buffer when earlying out. Was memzeroing from the start rather than from the current decode offset

diff --git a/Engine/Source/Runtime/Engine/Private/ADPCMAudioInfo.cpp b/Engine/Source/Runtime/Engine/Private/ADPCMAudioInfo.cpp
--- a/Engine/Source/Runtime/Engine/Private/ADPCMAudioInfo.cpp
+++ b/Engine/Source/Runtime/Engine/Private/ADPCMAudioInfo.cpp
@@ -316,12 +316,6 @@ bool FADPCMAudioInfo::ReadCompressedInfo(const uint8* InSrcBufferData, uint32 In
 
 bool FADPCMAudioInfo::ReadCompressedData(uint8* Destination, bool bLooping, uint32 BufferSize)
 {
-	// If we've already read through this asset and we are not looping, memzero and early out.
-	if (TotalSamplesStreamed >= TotalSamplesPerChannel && !bLooping)
-	{
-		FMemory::Memzero(Destination, BufferSize);
-		return true;
-	}
 
 	const uint32 ChannelSampleSize = sizeof(uint16) * NumChannels;
 
@@ -331,6 +325,13 @@ bool FADPCMAudioInfo::ReadCompressedData(uint8* Destination, bool bLooping, uint
 
 	ProcessSeekRequest();
 
+	// If we've already read through this asset and we are not looping, memzero and early out.
+	if (TotalSamplesStreamed >= TotalSamplesPerChannel && !bLooping)
+	{
+		FMemory::Memzero(Destination, BufferSize);
+		return true;
+	}
+
 	int16* OutData = (int16*)Destination;
 	bool ReachedEndOfSamples = false;
 	if(Format == WAVE_FORMAT_ADPCM)
@@ -430,7 +431,7 @@ bool FADPCMAudioInfo::ReadCompressedData(uint8* Destination, bool bLooping, uint
 				if (!bLooping)
 				{
 					// Zero remaining buffer
-					FMemory::Memzero(OutData, BufferSize);
+					FMemory::Memzero(OutData + OutDataOffset, BufferSize);
 
 					ResetSeekState();
 
