Warn on load/save when Blueprint CDOs contain serialized NULL members that should otherwise be referencing a default subobject instance assigned by the closest native parent class's C++ constructor.

This will patch those members back to reference the instance that was created at CDO construction time so that the asset functions correctly in the editor, but it will not restore any non-default values that may have been set prior to when the asset is/was last saved. This was likely the result of a regression that has since been addressed (not specifically identified).

Re-saving the asset will return the asset to a working state and dismiss the warning on load. If we happen to see it during a save, this is indicative of a problem that will cause data loss.

Currently disabled by default for the initial submission, in order to iteratively fix up assets prior to enabling it by default. Can set 'bp.DefaultSubobjectValidationDisabled 0' at the console to locally enable.

diff --git a/Engine/Source/Runtime/Engine/Private/BlueprintGeneratedClass.cpp b/Engine/Source/Runtime/Engine/Private/BlueprintGeneratedClass.cpp
--- a/Engine/Source/Runtime/Engine/Private/BlueprintGeneratedClass.cpp
+++ b/Engine/Source/Runtime/Engine/Private/BlueprintGeneratedClass.cpp
@@ -36,6 +36,7 @@
 #include "Kismet2/BlueprintEditorUtils.h"
 #include "BlueprintCompilationManager.h"
 #include "Engine/LevelScriptBlueprint.h"
+#include "ProfilingDebugging/CpuProfilerTrace.h"
 #include "WorldPartition/WorldPartitionActorDescUtils.h"
 #include "UObject/CookedMetaData.h"
 #include "UObject/UObjectGlobals.h"
@@ -55,7 +56,7 @@ static FAutoConsoleVariableRef CVarBlueprintNativePropertyInitFastPathDisabled(
 	ECVF_Default
 );
 
-int32 GBlueprintComponentInstancingFastPathDisabled = 0;
+static int32 GBlueprintComponentInstancingFastPathDisabled = 0;
 static FAutoConsoleVariableRef CVarBlueprintComponentInstancingFastPathDisabled(
 	TEXT("bp.ComponentInstancingFastPathDisabled"),
 	GBlueprintComponentInstancingFastPathDisabled,
@@ -63,6 +64,16 @@ static FAutoConsoleVariableRef CVarBlueprintComponentInstancingFastPathDisabled(
 	ECVF_Default
 );
 
+#if WITH_EDITOR
+static int32 GBlueprintDefaultSubobjectValidationDisabled = 1;
+static FAutoConsoleVariableRef CVarBlueprintDefaultSubobjectValidationDisabled(
+	TEXT("bp.DefaultSubobjectValidationDisabled"),
+	GBlueprintDefaultSubobjectValidationDisabled,
+	TEXT("Disable Blueprint class default subobject validation at editor load/save time."),
+	ECVF_Default
+);
+#endif	// WITH_EDITOR
+
 #if WITH_ADDITIONAL_CRASH_CONTEXTS
 struct BPGCBreadcrumbsParams
 {
@@ -194,6 +205,185 @@ namespace UE::Runtime::Engine::Private
 
 			return false;
 		}
+
+#if WITH_EDITOR
+	protected:
+		static void ValidateObjectPropertyValue(UObject* InOuter, const FObjectProperty* InProperty, void* InValuePtr, const void* InDefValuePtr)
+		{
+			check(InProperty);
+
+			// Get the reference assigned to the value address for the given property.
+			UObject* ObjValue = InProperty->GetObjectPropertyValue(InValuePtr);
+			if (!ObjValue)
+			{
+				// If the current reference value is NULL, grab the reference at the default value address for the same property.
+				ObjValue = InProperty->GetObjectPropertyValue(InDefValuePtr);
+				if (ObjValue && ObjValue->IsDefaultSubobject() && ObjValue->HasAllFlags(RF_DefaultSubObject))
+				{
+					check(InOuter);
+
+					// Attempt to find a matching instanced DSO within the current outer scope.
+					UObject* CurrentValue = InOuter->GetDefaultSubobjectByName(ObjValue->GetFName());
+					if (CurrentValue)
+					{
+						// In some cases, we might find a matching subobject instance that doesn't have the flag set to indicate that it was also
+						// instanced at construction time as a default subobject. Only fix up the value here if the instance also has that flag.
+						if (CurrentValue->HasAllFlags(RF_DefaultSubObject))
+						{
+							UE_LOG(LogBlueprint, Warning, TEXT("%s: Detected a NULL reference value for the class member named (%s). Changes to this property may not be restored on load. Check to see if any changes need to be re-applied, then re-save the asset to fix this warning."), *InOuter->GetPathName(), *InProperty->GetName());
+
+							// If the default reference is a non-NULL DSO, then the current container's reference should also be non-NULL. However,
+							// we want a reference to a matching subobject that exists within the current container. For DSOs, this should have
+							// been instanced at construction time (because we will have also run the container type's native constructor), but it's
+							// possible that this field has lost the reference somewhere along the way (e.g. at serialization time). So in order to
+							// ensure that we at least have a valid instanced DSO referenced by the property, reassign it to the current instance.
+							// Note that the current instance may not have been serialized if the reference was already NULL at save time, so this
+							// may result in data loss on load. However, this at least allows the object to be fixed up
+							InProperty->SetObjectPropertyValue(InValuePtr, CurrentValue);
+						}
+
+						// No need to validate nested DSOs here - we've simply returned the field to the initialized state of the container object.
+					}
+					else
+					{
+						// Could not find a matching DSO instance within the current container's scope; warn about this, but leave it set to NULL.
+						// We're not going to create a new instance here, because if we're in this situation, then it means the DSO was not
+						// instanced for the container object at construction time, which would occur for example if we started allowing users
+						// to mark inherited components as optional at the Blueprint editor level. So we'd want to determine why that occurred.
+						// @todo - If Blueprints ever add support to mark DSOs as optional at the editor level, we'll then need to revisit this.
+						UE_LOG(LogBlueprint, Warning, TEXT("%s: Missing a default subobject instance named \'%s\'. This should have been instanced at construction time for \'%s\'."), *InOuter->GetPathName(), *ObjValue->GetName(), *InProperty->GetOwnerStruct()->GetName());
+					}
+				}
+			}
+			else if(ObjValue->IsDefaultSubobject())
+			{
+				// If the current value is a default subobject, recursively validate any nested DSOs.
+				ValidateDefaultSubobjects(ObjValue);
+			}
+		}
+
+		static void ValidateInstancedObjectProperty(UObject* InOuter, const FProperty* InProperty, void* InDataPtr, const void* InDefaultDataPtr)
+		{
+			// It's possible for reference properties to be declared as a fixed array, so iterate over the fixed size (generally just one).
+			for (int32 ArrayIdx = 0; ArrayIdx < InProperty->ArrayDim; ++ArrayIdx)
+			{
+				if (const FObjectProperty* ObjProp = CastField<FObjectProperty>(InProperty))
+				{
+					// For a scalar reference type, validate the current value against the default value, and ensure that they match.
+					UObject* ObjValuePtr = ObjProp->ContainerPtrToValuePtr<UObject>(InDataPtr, ArrayIdx);
+					const UObject* DefValuePtr = ObjProp->ContainerPtrToValuePtr<UObject>(InDefaultDataPtr, ArrayIdx);
+					ValidateObjectPropertyValue(InOuter, ObjProp, ObjValuePtr, DefValuePtr);
+				}
+				else if (const FStructProperty* StructProp = CastField<FStructProperty>(InProperty))
+				{
+					// Recurse into struct properties, in case any members are assigned to a top-level DSO owned by the input object.
+					ValidateDefaultSubobjects_Inner(
+						InOuter,
+						StructProp->Struct,
+						StructProp->ContainerPtrToValuePtr<void>(InDataPtr, ArrayIdx),
+						StructProp->ContainerPtrToValuePtr<void>(InDefaultDataPtr, ArrayIdx));
+				}
+				else if (const FArrayProperty* ArrProp = CastField<FArrayProperty>(InProperty))
+				{
+					// For array types, validate each element's value against the default value, and ensure that all DSO elements match up.
+					// Note that it's possible for the default value to be larger or smaller than the current value in terms of the element
+					// count; in that case, we assume DSOs can't be removed at the Blueprint level, so the default will always include them.
+					FScriptArrayHelper_InContainer ArrValue(ArrProp, InDataPtr, ArrayIdx);
+					FScriptArrayHelper_InContainer DefValue(ArrProp, InDefaultDataPtr, ArrayIdx);
+					for (int32 ValueIdx = 0; ValueIdx < ArrValue.Num() && ValueIdx < DefValue.Num(); ++ValueIdx)
+					{
+						ValidateInstancedObjectProperty(InOuter, ArrProp->Inner, ArrValue.GetRawPtr(ValueIdx), DefValue.GetRawPtr(ValueIdx));
+					}
+				}
+				else if (const FSetProperty* SetProp = CastField<FSetProperty>(InProperty))
+				{
+					// For set containers, validate each element's value against the default value, and ensure that all DSO elements match up.
+					// As with arrays, we must also consider that the element counts may differ between current and default, but that all DSOs
+					// are at least always present in the default container, and match up with the elements in the current set container value.
+					FScriptSetHelper_InContainer SetValue(SetProp, InDataPtr, ArrayIdx);
+					FScriptSetHelper_InContainer DefValue(SetProp, InDefaultDataPtr, ArrayIdx);
+					for (int32 ValueIdx = 0; ValueIdx < SetValue.Num() && ValueIdx < DefValue.Num(); ++ValueIdx)
+					{
+						ValidateInstancedObjectProperty(InOuter, SetProp->ElementProp, SetValue.GetElementPtr(ValueIdx), DefValue.GetElementPtr(ValueIdx));
+					}
+				}
+				else if (const FMapProperty* MapProp = CastField<FMapProperty>(InProperty))
+				{
+					// For map containers, validate each pair's value against the default value, and ensure that all DSO elements match up.
+					// As above, we consider that the number of pairs in the map may differ between current and default, but that all DSOs
+					// are at least always present in the default container, and match up with the pairs in the current map container value.
+					FScriptMapHelper_InContainer MapValue(MapProp, InDataPtr, ArrayIdx);
+					FScriptMapHelper_InContainer DefValue(MapProp, InDefaultDataPtr, ArrayIdx);
+					for (int32 ValueIdx = 0; ValueIdx < MapValue.Num() && ValueIdx < DefValue.Num(); ++ValueIdx)
+					{
+						ValidateInstancedObjectProperty(InOuter, MapProp->ValueProp, MapValue.GetValuePtr(ValueIdx), DefValue.GetValuePtr(ValueIdx));
+					}
+				}
+			}
+		}
+
+		static void ValidateDefaultSubobjects_Inner(UObject* InOuter, UStruct* InStruct, void* InDataPtr, const void* InDefaultDataPtr)
+		{
+			// Iterate over all reference properties, including those inherited from the parent class hierarchy.
+			for (const FProperty* RefProp = InStruct->RefLink; RefProp; RefProp = RefProp->NextRef)
+			{
+				// We only need to consider 'Instanced' reference properties here (e.g. instanced subobjects created at construction time).
+				if (RefProp->ContainsInstancedObjectProperty())
+				{
+					ValidateInstancedObjectProperty(InOuter, RefProp, InDataPtr, InDefaultDataPtr);
+				}
+			}
+		}
+		
+	public:
+		/**
+		 * Utility method that iterates the reference property chain for a given object's underlying (non-native) type, and validates all
+		 * references to any instanced default subobjects (DSOs), or subobjects that are instanced natively when we first construct and
+		 * initialize the object through its native super class hierarchy. Current validation steps include:
+		 *
+		 *	a) Ensure that the given object has a non-NULL value for each reference property inherited from its native super chain, when
+		 *	   compared to its closest native parent class default object (CDO). Blueprints cannot mark default subobjects as "optional,"
+		 *	   which means that if the closest native parent CDO has a non-NULL value for a property that's referencing a default subobject
+		 *	   instance, then we expect that the Blueprint's CDO should also have a non-NULL value for the same property, but referencing
+		 *	   its own unique instance of a subobject with the same name. In some cases (due to reinstancing bugs perhaps), these references
+		 *	   can unintentionally be serialized as NULL to the Blueprint asset, which can lead to data loss and corrupted assets that are
+		 *	   otherwise unrecoverable. This both emits a warning to the log when validation fails, as well as attempts to restore these
+		 *	   references back to the unique instance that was constructed/initialized for the object, in order to allow for data recovery.
+		 *
+		 * Any additional validation steps that are implemented as part of this path in the future should be noted above for completeness.
+		 */
+		static void ValidateDefaultSubobjects(UObject* InObject)
+		{
+			TRACE_CPUPROFILER_EVENT_SCOPE(FBlueprintGeneratedClassUtils::ValidateDefaultSubobjects);
+
+			if (GBlueprintDefaultSubobjectValidationDisabled)
+			{
+				return;
+			}
+
+			check(InObject);
+
+			UClass* ObjClass = InObject->GetClass();
+			check(ObjClass);
+
+			// No need to validate non-native class types.
+			if (ObjClass->IsNative())
+			{
+				return;
+			}
+
+			// Find the closest native super class in the inheritance hierarchy.
+			UClass* NativeParentClass = FindFirstNativeClassInHierarchy(ObjClass);
+			check(NativeParentClass);
+
+			// Grab a reference to its default object. We'll use it as the basis for comparison below.
+			UObject* NativeParentCDO = NativeParentClass->GetDefaultObject(false);
+			check(NativeParentCDO);
+
+			// Validate this object's DSO member references against its closest native super class defaults.
+			ValidateDefaultSubobjects_Inner(InObject, NativeParentClass, InObject, NativeParentCDO);
+		}
+#endif	// WITH_EDITOR
 	};
 }
 
@@ -569,10 +759,24 @@ void UBlueprintGeneratedClass::SerializeDefaultObject(UObject* Object, FStructur
 	FScopeLock SerializeAndPostLoadLock(&SerializeAndPostLoadCritical);
 	FArchive& UnderlyingArchive = Slot.GetUnderlyingArchive();
 
+#if WITH_EDITOR
+	using namespace UE::Runtime::Engine::Private;
+	if (UnderlyingArchive.IsSaving() && !UnderlyingArchive.IsObjectReferenceCollector() && Object == ClassDefaultObject)
+	{
+		// Validate/fix up default subobjects prior to saving the CDO.
+		FBlueprintGeneratedClassUtils::ValidateDefaultSubobjects(Object);
+	}
+#endif	// WITH_EDITOR
+
 	Super::SerializeDefaultObject(Object, Slot);
 
 	if (UnderlyingArchive.IsLoading() && !UnderlyingArchive.IsObjectReferenceCollector() && Object == ClassDefaultObject)
 	{
+#if WITH_EDITOR
+		// Validate/fix up default subobjects after serializing the CDO (e.g. fix up any unexpected NULL refs).
+		FBlueprintGeneratedClassUtils::ValidateDefaultSubobjects(Object);
+#endif	// WITH_EDITOR
+
 		CreatePersistentUberGraphFrame(Object, true);
 
 		// On load, build the custom property list used in post-construct initialization logic. Note that in the editor, this will be refreshed during compile-on-load.
