Update TotalSamples of the SoundWave asset on PostImport

diff --git a/Engine/Source/Runtime/Engine/Private/SoundWave.cpp b/Engine/Source/Runtime/Engine/Private/SoundWave.cpp
--- a/Engine/Source/Runtime/Engine/Private/SoundWave.cpp
+++ b/Engine/Source/Runtime/Engine/Private/SoundWave.cpp
@@ -1938,6 +1938,11 @@ void USoundWave::PostLoad()
 			ImportedSampleRate = OriginalSampleRate;
 		}
 	}
+	
+	if (TotalSamples == 0 && ImportedSampleRate != 0)
+	{
+		TotalSamples = Duration * ImportedSampleRate;
+	}
 #endif
 
 	// Don't need to do anything in post load if this is a source bus or procedural audio
